[Lime]

Integrating code that allows for geometry collections to actually be hidden when you ask them to.
Modifying minor BP flow funcitonality so that we don't hide the Geometry Collection for no reason, and we clean up the tempororay static mesh component that we created in order to generate the SM previews.

diff --git a/Engine/Source/Runtime/Experimental/GeometryCollectionEngine/Private/GeometryCollection/GeometryCollectionComponent.cpp b/Engine/Source/Runtime/Experimental/GeometryCollectionEngine/Private/GeometryCollection/GeometryCollectionComponent.cpp
--- a/Engine/Source/Runtime/Experimental/GeometryCollectionEngine/Private/GeometryCollection/GeometryCollectionComponent.cpp
+++ b/Engine/Source/Runtime/Experimental/GeometryCollectionEngine/Private/GeometryCollection/GeometryCollectionComponent.cpp
@@ -750,6 +750,13 @@ void UGeometryCollectionComponent::EndPlay(const EEndPlayReason::Type ReasonEnd)
 	CurrentCacheTime = MAX_flt;
 }
 
+void UGeometryCollectionComponent::OnVisibilityChanged()
+{
+	Super::OnVisibilityChanged();
+
+	RefreshCustomRenderer();
+}
+
 void UGeometryCollectionComponent::GetLifetimeReplicatedProps(TArray<FLifetimeProperty>& OutLifetimeProps) const
 {
 	Super::GetLifetimeReplicatedProps(OutLifetimeProps);
@@ -6246,7 +6253,7 @@ void UGeometryCollectionComponent::RefreshCustomRenderer()
 					const bool bRenderRootProxy = bEnableRootProxyForCustomRenderer && !bIsBroken && (RestCollection->RootProxyData.ProxyMeshes.Num() > 0);
 
 					uint32 StateFlags = 0;
-					StateFlags |= bHiddenInGame ? 0 : IGeometryCollectionExternalRenderInterface::EState_Visible;
+					StateFlags |= bHiddenInGame || !IsVisible() ? 0 : IGeometryCollectionExternalRenderInterface::EState_Visible;
 					StateFlags |= bRenderRootProxy ? 0 : IGeometryCollectionExternalRenderInterface::EState_Broken;
 					StateFlags |= bEnableRootProxyForCustomRenderer ? 0 : IGeometryCollectionExternalRenderInterface::EState_ForcedBroken;
 
diff --git a/Engine/Source/Runtime/Experimental/GeometryCollectionEngine/Public/GeometryCollection/GeometryCollectionComponent.h b/Engine/Source/Runtime/Experimental/GeometryCollectionEngine/Public/GeometryCollection/GeometryCollectionComponent.h
--- a/Engine/Source/Runtime/Experimental/GeometryCollectionEngine/Public/GeometryCollection/GeometryCollectionComponent.h
+++ b/Engine/Source/Runtime/Experimental/GeometryCollectionEngine/Public/GeometryCollection/GeometryCollectionComponent.h
@@ -609,6 +609,7 @@ public:
 	GEOMETRYCOLLECTIONENGINE_API virtual void OnActorEnableCollisionChanged() override;
 	GEOMETRYCOLLECTIONENGINE_API virtual void BeginPlay() override;
 	GEOMETRYCOLLECTIONENGINE_API virtual void EndPlay(const EEndPlayReason::Type ReasonEnd) override;
+	GEOMETRYCOLLECTIONENGINE_API virtual void OnVisibilityChanged() override;
 	GEOMETRYCOLLECTIONENGINE_API virtual void GetLifetimeReplicatedProps(TArray<FLifetimeProperty>& OutLifetimeProps) const override;
 	GEOMETRYCOLLECTIONENGINE_API virtual void GetResourceSizeEx(FResourceSizeEx& CumulativeResourceSize) override;
 
