Removed unused PushActivity() method

#rnx

(git:9e9356c - ridgers - [REDACTED_PHONE]:01:32 +0100)

diff --git a/Engine/Source/Runtime/Experimental/IoStoreOnDemand/Private/CoreHttp/Http.cpp b/Engine/Source/Runtime/Experimental/IoStoreOnDemand/Private/CoreHttp/Http.cpp
--- a/Engine/Source/Runtime/Experimental/IoStoreOnDemand/Private/CoreHttp/Http.cpp
+++ b/Engine/Source/Runtime/Experimental/IoStoreOnDemand/Private/CoreHttp/Http.cpp
@@ -2560,7 +2560,6 @@ public:
 						~FWorkQueue();
 	bool				HasWork() const { return List != nullptr; }
 	void				AddActivity(FActivity* Activity);
-	void				PushActivity(FActivity* Activity);
 	FActivity*			PopActivity();
 	void				TickCancels(FTickState& State);
 
@@ -2596,14 +2595,6 @@ void FWorkQueue::AddActivity(FActivity* Activity)
 	ActiveSlots |= (1ull << Activity->Slot);
 }
 
-////////////////////////////////////////////////////////////////////////////////
-void FWorkQueue::PushActivity(FActivity* Activity)
-{
-	Activity->Next = List;
-	List = Activity;
-	ListTail = (ListTail != nullptr) ? ListTail : Activity;
-}
-
 ////////////////////////////////////////////////////////////////////////////////
 FActivity* FWorkQueue::PopActivity()
 {
