Measure request timings by sampling instead of accumulation

The previous approach to keeping an eye on request time was overly
complex and hid some of the wall time making logs using the timings to
misrepresent what really happened over time. We can derive a better
understanding of how a request went by sampling the times at which it
reaches different states.

#rnx
(git:7f1741d - ridgers - [REDACTED_PHONE]:56:26 +0100)

diff --git a/Engine/Source/Runtime/Experimental/IoStoreOnDemand/Private/CoreHttp/Client.h b/Engine/Source/Runtime/Experimental/IoStoreOnDemand/Private/CoreHttp/Client.h
--- a/Engine/Source/Runtime/Experimental/IoStoreOnDemand/Private/CoreHttp/Client.h
+++ b/Engine/Source/Runtime/Experimental/IoStoreOnDemand/Private/CoreHttp/Client.h
@@ -138,19 +138,21 @@ private:
 ////////////////////////////////////////////////////////////////////////////////
 class UE_API FTicketPerf
 {
-#if IAS_HTTP_WITH_PERF
 public:
 	struct FSample
 	{
-		uint32			TotalMs;
+		uint32			SendMs;
 		uint32			WaitMs;
+		uint32			RecvMs;
 	};
 
-	FSample				GetSendSample() const;
-	FSample				GetRecvSample() const;
-#endif // IAS_HTTP_WITH_PERF
+	FSample				GetSample() const;
 };
 
+#if !IAS_HTTP_WITH_PERF
+inline FTicketPerf::FSample FTicketPerf::GetSample() const { /* nop */ }
+#endif
+
 ////////////////////////////////////////////////////////////////////////////////
 class UE_API FTicketStatus
 {
diff --git a/Engine/Source/Runtime/Experimental/IoStoreOnDemand/Private/CoreHttp/Http.cpp b/Engine/Source/Runtime/Experimental/IoStoreOnDemand/Private/CoreHttp/Http.cpp
--- a/Engine/Source/Runtime/Experimental/IoStoreOnDemand/Private/CoreHttp/Http.cpp
+++ b/Engine/Source/Runtime/Experimental/IoStoreOnDemand/Private/CoreHttp/Http.cpp
@@ -1521,50 +1521,36 @@ bool FConnectionPool::IsValidHostUrl(FAnsiStringView Url)
 class FStopwatch
 {
 public:
-	struct FInterval
-	{
-					FInterval() : Elapsed(0), Counter(0) {}
-		int64		Elapsed : 48;
-		int64		Counter : 16;
-	};
-
-	struct FLapTime
-	{
-		FInterval	Total;
-		FInterval	Wait;
-	};
-
-	void		Start()		{ Laps[Index].Total.Elapsed -= Sample(); }
-	void		Stop()		{ Laps[Index].Total.Elapsed += Sample(); }
-	void		Wait()		{ Laps[Index].Wait.Elapsed -= Sample(); Laps[Index].Wait.Counter++; }
-	void		Unwait()	{ Laps[Index].Wait.Elapsed += Sample(); }
-	void		Lap()		{ ++Index; check(Index < UE_ARRAY_COUNT(Laps)); }
-	void		AddCount()	{ Laps[Index].Total.Counter++; }
-	int64		Sample();
-
-	const FLapTime&	GetLap(uint32 i) const
-	{
-		check(i < UE_ARRAY_COUNT(Laps));
-		return Laps[i];
-	}
+	uint64	GetInterval(uint32 i) const;
+	void	SendStart()		{ Impl(0); }
+	void	SendEnd()		{ Impl(1); }
+	void	RecvStart()		{ Impl(2); }
+	void	RecvEnd()		{ Impl(3); }
 
 private:
-	FLapTime	Laps[2];
-	uint32		Index = 0;
+	void	Impl(uint32 Index);
+	uint64	Samples[4] = {};
+	uint32	Counts[2] = {};
 };
 
 ////////////////////////////////////////////////////////////////////////////////
-int64 FStopwatch::Sample()
+uint64 FStopwatch::GetInterval(uint32 i) const
 {
-	int64 Value = FPlatformTime::Cycles64();
-	static int64 Base = 0;
-	if (Base == 0)
+	if (i >= UE_ARRAY_COUNT(Samples) - 1)
 	{
-		Base = Value;
 		return 0;
 	}
+	return Samples[i + 1] - Samples[i];
+}
 
-	return Value - Base;
+////////////////////////////////////////////////////////////////////////////////
+void FStopwatch::Impl(uint32 Index)
+{
+	if (uint64& Out = Samples[Index]; Out == 0)
+	{
+		Out = FPlatformTime::Cycles64();
+	}
+	Counts[Index >> 1] += !(Index & 1);
 }
 
 #endif // IAS_HTTP_WITH_PERF
@@ -1635,33 +1621,6 @@ static void Activity_ChangeState(FActivity* Activity, FActivity::EState InState,
 {
 	Trace(Activity, ETrace::StateChange, InState);
 
-#if IAS_HTTP_WITH_PERF
-	using EState = FActivity::EState;
-
-	FStopwatch& Stopwatch = Activity->Stopwatch;
-	if (InState == EState::Send)
-	{
-		if (Activity->State == EState::RecvMessage)
-		{
-			Stopwatch = FStopwatch();
-		}
-		Stopwatch.Start();
-	}
-	else if (Activity->State == EState::Send)
-	{
-		Stopwatch.Stop();
-		Stopwatch.Lap();
-	}
-	else if (InState == EState::RecvContent || InState == EState::RecvStream)
-	{
-		Stopwatch.Start();
-	}
-	else if (Activity->State == EState::RecvContent || Activity->State == EState::RecvStream)
-	{
-		Stopwatch.Stop();
-	}
-#endif // IAS_HTTP_WITH_PERF
-
 	check(Activity->State != InState);
 	Activity->State = InState;
 	Activity->StateParam = Param;
@@ -1992,8 +1951,10 @@ const char* FTicketStatus::GetErrorReason() const
 #if IAS_HTTP_WITH_PERF
 
 ////////////////////////////////////////////////////////////////////////////////
-static FTicketPerf::FSample GetPerfSample(const FActivity* Activity, uint32 Index)
+FTicketPerf::FSample FTicketPerf::GetSample() const
 {
+	const auto* Activity = (FActivity*)this;
+
 	static uint64 Freq;
 	if (Freq == 0)
 	{
@@ -2002,27 +1963,14 @@ static FTicketPerf::FSample GetPerfSample(const FActivity* Activity, uint32 Inde
 
 	auto ToMs = [] (uint64 Value) { return uint32((Value * 1000ull) / Freq); };
 
-	const FStopwatch::FLapTime& LapTime = Activity->Stopwatch.GetLap(Index);
+	const FStopwatch& Stopwatch = Activity->Stopwatch;
 	return {
-		ToMs(LapTime.Total.Elapsed),
-		ToMs(LapTime.Wait.Elapsed),
+		ToMs(Stopwatch.GetInterval(0)),
+		ToMs(Stopwatch.GetInterval(1)),
+		ToMs(Stopwatch.GetInterval(2)),
 	};
 }
 
-////////////////////////////////////////////////////////////////////////////////
-FTicketPerf::FSample FTicketPerf::GetSendSample() const
-{
-	const auto* Activity = (FActivity*)this;
-	return GetPerfSample(Activity, 0);
-}
-
-////////////////////////////////////////////////////////////////////////////////
-FTicketPerf::FSample FTicketPerf::GetRecvSample() const
-{
-	const auto* Activity = (FActivity*)this;
-	return GetPerfSample(Activity, 1);
-}
-
 #endif // IAS_HTTP_WITH_PERF
 
 
@@ -2034,6 +1982,10 @@ static int32 DoSend(FActivity* Activity, FSocket& Socket)
 {
 	Trace(Activity, ETrace::StateChange, Activity->State);
 
+#if IAS_HTTP_WITH_PERF
+	Activity->Stopwatch.SendStart();
+#endif
+
 	TRACE_CPUPROFILER_EVENT_SCOPE(IasHttp::DoSend);
 
 	FBuffer& Buffer = Activity->Buffer;
@@ -2055,10 +2007,6 @@ static int32 DoSend(FActivity* Activity, FSocket& Socket)
 	case FSocket::EResult::Wait:		return Result;
 	}
 
-#if IAS_HTTP_WITH_PERF
-	Activity->Stopwatch.AddCount();
-#endif
-
 	checkf(Result > 0, TEXT("Result wasn't caught by switch statement so it is expected to be a positive amount of bytes sent"));
 	Activity->StateParam += Result;
 	if (Activity->StateParam < Buffer.GetSize())
@@ -2066,6 +2014,10 @@ static int32 DoSend(FActivity* Activity, FSocket& Socket)
 		return DoSend(Activity, Socket);
 	}
 
+#if IAS_HTTP_WITH_PERF
+	Activity->Stopwatch.SendEnd();
+#endif
+
 	Activity_ChangeState(Activity, FActivity::EState::RecvMessage, Buffer.GetSize());
 	return Result;
 }
@@ -2084,8 +2036,11 @@ static int32 DoRecvMessage(FActivity* Activity, FSocket& Socket)
 	{
 		Trace(Activity, ETrace::StateChange, Activity->State);
 
-		auto [Dest, DestSize] = Buffer.GetMutableFree(0, PageSize);
+#if IAS_HTTP_WITH_PERF
+		Activity->Stopwatch.RecvStart();
+#endif
 
+		auto [Dest, DestSize] = Buffer.GetMutableFree(0, PageSize);
 		int32 Result = Socket.Recv(Dest, DestSize);
 
 		if (Result == int32(FSocket::EResult::Wait))
@@ -2319,10 +2274,6 @@ static int32 DoRecvContent(FActivity* Activity, FSocket& Socket, int32& MaxRecvS
 			return -1;
 		}
 
-#if IAS_HTTP_WITH_PERF
-		Activity->Stopwatch.AddCount();
-#endif
-
 		check(Result <= MaxRecvSize);
 		Activity->StateParam += Result;
 		MaxRecvSize -= Result;
@@ -2334,6 +2285,10 @@ static int32 DoRecvContent(FActivity* Activity, FSocket& Socket, int32& MaxRecvS
 		return -1;
 	}
 
+#if IAS_HTTP_WITH_PERF
+	Activity->Stopwatch.RecvEnd();
+#endif
+
 	Activity_ChangeState(Activity, FActivity::EState::RecvDone);
 	return 0;
 }
