Linux server shipping fix.

#rnx

diff --git a/Engine/Source/Runtime/Experimental/IoStoreOnDemand/Private/CoreHttp/Client.h b/Engine/Source/Runtime/Experimental/IoStoreOnDemand/Private/CoreHttp/Client.h
--- a/Engine/Source/Runtime/Experimental/IoStoreOnDemand/Private/CoreHttp/Client.h
+++ b/Engine/Source/Runtime/Experimental/IoStoreOnDemand/Private/CoreHttp/Client.h
@@ -150,7 +150,7 @@ public:
 };
 
 #if !IAS_HTTP_WITH_PERF
-inline FTicketPerf::FSample FTicketPerf::GetSample() const { /* nop */ }
+inline FTicketPerf::FSample FTicketPerf::GetSample() const { return {}; }
 #endif
 
 ////////////////////////////////////////////////////////////////////////////////
