Chaos - fix cluster union bounds when one of the child particle transforms changes

This caused a mismatch between geometry and the bounds, leading to missed collisions.

diff --git a/Engine/Source/Runtime/Experimental/Chaos/Private/Chaos/ClusterUnionManager.cpp b/Engine/Source/Runtime/Experimental/Chaos/Private/Chaos/ClusterUnionManager.cpp
--- a/Engine/Source/Runtime/Experimental/Chaos/Private/Chaos/ClusterUnionManager.cpp
+++ b/Engine/Source/Runtime/Experimental/Chaos/Private/Chaos/ClusterUnionManager.cpp
@@ -1169,6 +1169,7 @@ namespace Chaos
 			const FShapesArray& ShapesArray = ClusterUnion.InternalCluster->ShapesArray();
 			check(ClusterUnion.ChildParticles.Num() == ShapesArray.Num());
 
+			bool bAnyTransformChanged = false;
 			for (FPBDRigidParticleHandle* Particle : PendingGeometryRefresh)
 			{
 				if (!Particle)
@@ -1196,10 +1197,12 @@ namespace Chaos
 						if (!Transformed->GetTransform().Equals(Frame))
 						{
 							Transformed->SetTransform(Frame);
+							bAnyTransformChanged = true;
 						}
 					}
 				}
 
+				// @todo(chaos): we should probably rebuild the bounds if the sim and query filters change as well
 				TransferClusterUnionShapeData(
 					ShapesArray[Index],
 					Particle,
@@ -1210,6 +1213,12 @@ namespace Chaos
 			}
 
 			ClusterUnion.ClearPendingGeometryOperations(EClusterUnionGeometryOperation::Refresh);
+
+			if (bAnyTransformChanged)
+			{
+				// We have changed the transforms of some children so we must update the bounds
+				ClusterUnion.InternalCluster->UpdateWorldSpaceState(ClusterUnion.InternalCluster->GetTransformPQ(), FVec3(0));
+			}
 		}
 
 		// Re-enable the BVH to rebuild it
