Filtering out geometry collection editor data from memory calculations

diff --git a/Engine/Source/Runtime/Experimental/GeometryCollectionEngine/Private/GeometryCollection/GeometryCollectionObject.cpp b/Engine/Source/Runtime/Experimental/GeometryCollectionEngine/Private/GeometryCollection/GeometryCollectionObject.cpp
--- a/Engine/Source/Runtime/Experimental/GeometryCollectionEngine/Private/GeometryCollection/GeometryCollectionObject.cpp
+++ b/Engine/Source/Runtime/Experimental/GeometryCollectionEngine/Private/GeometryCollection/GeometryCollectionObject.cpp
@@ -882,18 +882,21 @@ void UGeometryCollection::Serialize(FArchive& Ar)
 	TObjectPtr<UDataflow> StrippedDataflowAsset = nullptr;
 
 	bool bIsCookedOrCooking = Ar.IsCooking();
-	if (bIsCookedOrCooking && Ar.IsSaving())
+	if ((bIsCookedOrCooking && Ar.IsSaving()) || (Ar.IsCountingMemory() && Ar.IsFilterEditorOnly()))
 	{
 #if WITH_EDITOR
-		// if we have a valid selection material, let's make sure we replace it with one that will be cooked
-		// this avoid getting warning about the selected material being reference but not cooked
-		const int32 SelectedMaterialIndex = GetBoneSelectedMaterialIndex();
-		if (!Materials.IsEmpty() && Materials.IsValidIndex(SelectedMaterialIndex))
+		if (bIsCookedOrCooking && Ar.IsSaving())
 		{
-			Materials[SelectedMaterialIndex] = Materials[0];
+			// if we have a valid selection material, let's make sure we replace it with one that will be cooked
+			// this avoid getting warning about the selected material being reference but not cooked
+			const int32 SelectedMaterialIndex = GetBoneSelectedMaterialIndex();
+			if (!Materials.IsEmpty() && Materials.IsValidIndex(SelectedMaterialIndex))
+			{
+				Materials[SelectedMaterialIndex] = Materials[0];
+			}
+			// Likewise remove the direct reference to the BoneSelectedMaterial on cook
+			BoneSelectedMaterial = nullptr;
 		}
-		// Likewise remove the direct reference to the BoneSelectedMaterial on cook
-		BoneSelectedMaterial = nullptr;
 
 		if (bStripOnCook)
 		{
diff --git a/Engine/Source/Runtime/Experimental/GeometryCollectionEngine/Private/GeometryCollection/GeometryCollectionRenderData.cpp b/Engine/Source/Runtime/Experimental/GeometryCollectionEngine/Private/GeometryCollection/GeometryCollectionRenderData.cpp
--- a/Engine/Source/Runtime/Experimental/GeometryCollectionEngine/Private/GeometryCollection/GeometryCollectionRenderData.cpp
+++ b/Engine/Source/Runtime/Experimental/GeometryCollectionEngine/Private/GeometryCollection/GeometryCollectionRenderData.cpp
@@ -724,7 +724,8 @@ FGeometryCollectionRenderData::~FGeometryCollectionRenderData()
 
 void FGeometryCollectionRenderData::Serialize(FArchive& Ar, UGeometryCollection& Owner)
 {
-	if (Owner.bStripRenderDataOnCook && Ar.IsCooking())
+	const bool bIsArchiveValidCandidateForStrip = (Ar.IsCooking() || (Ar.IsCountingMemory() && Ar.IsFilterEditorOnly()));
+	if (Owner.bStripRenderDataOnCook && bIsArchiveValidCandidateForStrip)
 	{
 		// Don't cook rendering data.
 		// This is used if we expect to use a custom GC render path such as the ISM pool.
