Chaos - support runtime change of collision cvars for testing/debugging by destroying existing collisions
- fix debug draw client/server disable (previously only worked if shape debug draw was enabled)

diff --git a/Engine/Source/Runtime/Experimental/Chaos/Private/Chaos/Collision/ParticlePairMidPhase.cpp b/Engine/Source/Runtime/Experimental/Chaos/Private/Chaos/Collision/ParticlePairMidPhase.cpp
--- a/Engine/Source/Runtime/Experimental/Chaos/Private/Chaos/Collision/ParticlePairMidPhase.cpp
+++ b/Engine/Source/Runtime/Experimental/Chaos/Private/Chaos/Collision/ParticlePairMidPhase.cpp
@@ -64,7 +64,7 @@ namespace Chaos
 
 		extern int32 ChaosOneWayInteractionPairCollisionMode;
 
-		extern bool bChaosUseMACD;
+		extern bool bChaosForceMACD;
 	}
 
 	using namespace CVars;
@@ -882,7 +882,7 @@ namespace Chaos
 
 		// Enable Motion-Aware Collision Detection if either particle requests it
 		// (NOTE: MACD also affects how the particle world-space bounds is expanded in the broadphase)
-		const bool bIsMACD = Context.GetSettings().bAllowMACD && (P0->MACDEnabled() || P1->MACDEnabled());
+		const bool bIsMACD = Context.GetSettings().bAllowMACD && (CVars::bChaosForceMACD || P0->MACDEnabled() || P1->MACDEnabled());
 		Flags.bIsMACD = bIsMACD;
 
 		// Initially we allow for convex optimization where available
diff --git a/Engine/Source/Runtime/Experimental/Chaos/Private/Chaos/PBDRigidsEvolutionGBF.cpp b/Engine/Source/Runtime/Experimental/Chaos/Private/Chaos/PBDRigidsEvolutionGBF.cpp
--- a/Engine/Source/Runtime/Experimental/Chaos/Private/Chaos/PBDRigidsEvolutionGBF.cpp
+++ b/Engine/Source/Runtime/Experimental/Chaos/Private/Chaos/PBDRigidsEvolutionGBF.cpp
@@ -38,6 +38,10 @@ namespace Chaos
 
 	namespace CVars
 	{
+		// Debug/testing MACD
+		extern bool bChaosUseMACD;
+		extern bool bChaosForceMACD;
+
 		FRealSingle HackMaxAngularVelocity = 1000.f;
 		FAutoConsoleVariableRef CVarHackMaxAngularVelocity(TEXT("p.HackMaxAngularVelocity"), HackMaxAngularVelocity, TEXT("Max cap on angular velocity: rad/s. This is only a temp solution and should not be relied on as a feature. -1.f to disable"));
 
@@ -733,6 +737,8 @@ void FPBDRigidsEvolutionGBF::Integrate(FReal Dt)
 	const FReal MaxVelocityBoundsExpansionMACD = GetCollisionConstraints().GetDetectorSettings().MaxVelocityBoundsExpansionMACD;
 	const FReal HackMaxAngularSpeedSq = CVars::HackMaxAngularVelocity * CVars::HackMaxAngularVelocity;
 	const FReal HackMaxLinearSpeedSq = CVars::HackMaxVelocity * CVars::HackMaxVelocity;
+	const bool bAllowMACD = CVars::bChaosUseMACD;
+	const bool bForceMACD = CVars::bChaosForceMACD;
 
 	FChaosVDContextWrapper CVDContext;
 	CVD_GET_WRAPPED_CURRENT_CONTEXT(CVDContext);
@@ -824,21 +830,13 @@ void FPBDRigidsEvolutionGBF::Integrate(FReal Dt)
 				// separated by more than Cull Distance at collision detection time.
 				// NOTE: We use a different (larger) max bounds expansion when MACD is enabled.
 				FVec3 VelocityBoundsDelta = FVec3(0);
-				if (!Particle.MACDEnabled())
-				{
-					if ((VelocityBoundsMultiplier > 0) && (MaxVelocityBoundsExpansion > 0))
-					{
-						VelocityBoundsDelta = (-VelocityBoundsMultiplier * Dt) * V;
-						VelocityBoundsDelta = VelocityBoundsDelta.BoundToCube(MaxVelocityBoundsExpansion);
-					}
-				}
-				else
+				const bool bIsMACD = bForceMACD || (bAllowMACD && Particle.MACDEnabled());
+				const FReal ParticleVelocityBoundsMultiplier = bIsMACD ? VelocityBoundsMultiplierMACD : VelocityBoundsMultiplier;
+				const FReal ParticleMaxVelocityBoundsExpansion = bIsMACD ? MaxVelocityBoundsExpansionMACD : MaxVelocityBoundsExpansion;
+				if ((ParticleVelocityBoundsMultiplier > 0) && (ParticleMaxVelocityBoundsExpansion > 0))
 				{
-					if ((VelocityBoundsMultiplierMACD > 0) && (MaxVelocityBoundsExpansionMACD > 0))
-					{
-						VelocityBoundsDelta = (-VelocityBoundsMultiplierMACD * Dt) * V;
-						VelocityBoundsDelta = VelocityBoundsDelta.BoundToCube(MaxVelocityBoundsExpansionMACD);
-					}
+					VelocityBoundsDelta = (-ParticleVelocityBoundsMultiplier * Dt) * V;
+					VelocityBoundsDelta = VelocityBoundsDelta.BoundToCube(ParticleMaxVelocityBoundsExpansion);
 				}
 
 				if (!Particle.CCDEnabled())
diff --git a/Engine/Source/Runtime/Experimental/Chaos/Private/PBDRigidsSolver.cpp b/Engine/Source/Runtime/Experimental/Chaos/Private/PBDRigidsSolver.cpp
--- a/Engine/Source/Runtime/Experimental/Chaos/Private/PBDRigidsSolver.cpp
+++ b/Engine/Source/Runtime/Experimental/Chaos/Private/PBDRigidsSolver.cpp
@@ -332,25 +332,31 @@ namespace Chaos
 		FAutoConsoleVariableRef CVarChaosSolverJointMaxInertiaRatio(TEXT("p.Chaos.Solver.Joint.MaxInertiaRatio"), ChaosSolverJointMaxInertiaRatio, TEXT("6Dof joint MaxInertiaRatio (if > 0)"));
 
 		// Collision detection cvars
+
+		// Utility to support runtime changes to some high-level collision configuration that requires we update all existing collisions (actually we just destroy them)
+		// This is only intended for testing and debugging and handling the change is not fast.
+		bool bChaosCollisionConfigChanged = false;
+		FConsoleVariableDelegate OnCollisionConfigCVarChanged = FConsoleVariableDelegate::CreateLambda([](IConsoleVariable* CVar) -> void { bChaosCollisionConfigChanged = true; });
+
 		// These override the engine config if >= 0
 		FRealSingle ChaosSolverCullDistance = -1.0f;
-		FAutoConsoleVariableRef CVarChaosSolverCullDistance(TEXT("p.Chaos.Solver.Collision.CullDistance"), ChaosSolverCullDistance, TEXT("Override cull distance (if >= 0)"));
+		FAutoConsoleVariableRef CVarChaosSolverCullDistance(TEXT("p.Chaos.Solver.Collision.CullDistance"), ChaosSolverCullDistance, TEXT("Override cull distance (if >= 0)"), OnCollisionConfigCVarChanged);
 
 		// @todo(chaos): move to physics project settings and set these to -1 when we are settled on values...
 		FRealSingle ChaosSolverVelocityBoundsMultiplier = 1.0f;
 		FRealSingle ChaosSolverMaxVelocityBoundsExpansion = 3.0f;			// This should probably be a fraction of object size (see FParticlePairMidPhase::GenerateCollisions)
 		FRealSingle ChaosSolverVelocityBoundsMultiplierMACD = 1.0f;
 		FRealSingle ChaosSolverMaxVelocityBoundsExpansionMACD = 1000.0f;	// For use when Movement-Aware Collision Detection (MACD) is enabled
-		FAutoConsoleVariableRef CVarChaosSolverVelocityBoundsMultiplier(TEXT("p.Chaos.Solver.Collision.VelocityBoundsMultiplier"), ChaosSolverVelocityBoundsMultiplier, TEXT("Override velocity bounds multiplier (if >= 0)"));
-		FAutoConsoleVariableRef CVarChaosSolverMaxVelocityBoundsExpansion(TEXT("p.Chaos.Solver.Collision.MaxVelocityBoundsExpansion"), ChaosSolverMaxVelocityBoundsExpansion, TEXT("Override max velocity bounds expansion (if >= 0)"));
-		FAutoConsoleVariableRef CVarChaosSolverVelocityBoundsMultiplierMACD(TEXT("p.Chaos.Solver.Collision.VelocityBoundsMultiplierMACD"), ChaosSolverVelocityBoundsMultiplierMACD, TEXT("Override velocity bounds multiplier for MACD (if >= 0)"));
-		FAutoConsoleVariableRef CVarChaosSolverMaxVelocityBoundsExpansionMACD(TEXT("p.Chaos.Solver.Collision.MaxVelocityBoundsExpansionMACD"), ChaosSolverMaxVelocityBoundsExpansionMACD, TEXT("Override max velocity bounds expansion for MACD (if >= 0)"));
+		FAutoConsoleVariableRef CVarChaosSolverVelocityBoundsMultiplier(TEXT("p.Chaos.Solver.Collision.VelocityBoundsMultiplier"), ChaosSolverVelocityBoundsMultiplier, TEXT("Override velocity bounds multiplier (if >= 0)"), OnCollisionConfigCVarChanged);
+		FAutoConsoleVariableRef CVarChaosSolverMaxVelocityBoundsExpansion(TEXT("p.Chaos.Solver.Collision.MaxVelocityBoundsExpansion"), ChaosSolverMaxVelocityBoundsExpansion, TEXT("Override max velocity bounds expansion (if >= 0)"), OnCollisionConfigCVarChanged);
+		FAutoConsoleVariableRef CVarChaosSolverVelocityBoundsMultiplierMACD(TEXT("p.Chaos.Solver.Collision.VelocityBoundsMultiplierMACD"), ChaosSolverVelocityBoundsMultiplierMACD, TEXT("Override velocity bounds multiplier for MACD (if >= 0)"), OnCollisionConfigCVarChanged);
+		FAutoConsoleVariableRef CVarChaosSolverMaxVelocityBoundsExpansionMACD(TEXT("p.Chaos.Solver.Collision.MaxVelocityBoundsExpansionMACD"), ChaosSolverMaxVelocityBoundsExpansionMACD, TEXT("Override max velocity bounds expansion for MACD (if >= 0)"), OnCollisionConfigCVarChanged);
 
 		FRealSingle ChaosSolverMaxPushOutVelocity = -1.0f;
-		FAutoConsoleVariableRef CVarChaosSolverMaxPushOutVelocity(TEXT("p.Chaos.Solver.Collision.MaxPushOutVelocity"), ChaosSolverMaxPushOutVelocity, TEXT("Override max pushout velocity (if >= 0)"));
+		FAutoConsoleVariableRef CVarChaosSolverMaxPushOutVelocity(TEXT("p.Chaos.Solver.Collision.MaxPushOutVelocity"), ChaosSolverMaxPushOutVelocity, TEXT("Override max pushout velocity (if >= 0)"), OnCollisionConfigCVarChanged);
 
 		FRealSingle ChaosSolverDepenetrationVelocity = -1.0f;
-		FAutoConsoleVariableRef CVarChaosSolverInitialOverlapDepentrationVelocity(TEXT("p.Chaos.Solver.Collision.DepenetrationVelocity"), ChaosSolverDepenetrationVelocity, TEXT("Override initial overlap depenetration velocity (if >= 0)"));
+		FAutoConsoleVariableRef CVarChaosSolverInitialOverlapDepentrationVelocity(TEXT("p.Chaos.Solver.Collision.DepenetrationVelocity"), ChaosSolverDepenetrationVelocity, TEXT("Override initial overlap depenetration velocity (if >= 0)"), OnCollisionConfigCVarChanged);
 
 		int32 ChaosSolverCleanupCommandsOnDestruction = 1;
 		FAutoConsoleVariableRef CVarChaosSolverCleanupCommandsOnDestruction(TEXT("p.Chaos.Solver.CleanupCommandsOnDestruction"), ChaosSolverCleanupCommandsOnDestruction, TEXT("Whether or not to run internal command queue cleanup on solver destruction (0 = no cleanup, >0 = cleanup all commands)"));
@@ -368,11 +374,15 @@ namespace Chaos
 
 		// Enable/Disable CCD. Set to false to disable the system, regardless of particle settings
 		bool bChaosUseCCD = true;
-		FAutoConsoleVariableRef  CVarChaosUseCCD(TEXT("p.Chaos.Solver.UseCCD"), bChaosUseCCD, TEXT("Global flag to turn CCD on or off. Default is true (on)"));
+		FAutoConsoleVariableRef  CVarChaosUseCCD(TEXT("p.Chaos.Solver.UseCCD"), bChaosUseCCD, TEXT("Global flag to turn CCD on or off. Default is true (on)"), OnCollisionConfigCVarChanged);
 
 		// Enable/Disable MACD (Motion-Aware Collision Detection). Set to false to disable the system, regardless of particle settings
 		bool bChaosUseMACD = true;
-		FAutoConsoleVariableRef CVarChaos_Collision_UseMACD(TEXT("p.Chaos.Solver.UseMACD"), bChaosUseMACD, TEXT("Global flag to turn Movement-Aware Collision Detection (MACD) on or off. Default is true (on)"));
+		FAutoConsoleVariableRef CVarChaos_Collision_UseMACD(TEXT("p.Chaos.Solver.UseMACD"), bChaosUseMACD, TEXT("Global flag to turn Movement-Aware Collision Detection (MACD) on or off. Default is true (on)"), OnCollisionConfigCVarChanged);
+
+		// Use to force all collisions to use MACD for testing (must also have bChaosUseMACD enabled)
+		bool bChaosForceMACD = false;
+		FAutoConsoleVariableRef CVarChaos_Collision_ForceMACD(TEXT("p.Chaos.Solver.bChaosForceMACD"), bChaosForceMACD, TEXT("Force all collisions to use MACD for testing"), OnCollisionConfigCVarChanged);
 
 		// Joint cvars
 		float ChaosSolverJointMinSolverStiffness = 1.0f;
@@ -1253,24 +1263,12 @@ namespace Chaos
 
 	void FPBDRigidsSolver::SetVelocityBoundsExpansion(const FReal BoundsVelocityMultiplier, const FReal MaxBoundsVelocityExpansion)
 	{ 
-		FPBDCollisionConstraints& CollisionContainer = GetEvolution()->GetCollisionConstraints();
-		if ((CollisionContainer.GetDetectorSettings().BoundsVelocityInflation != BoundsVelocityMultiplier) || (CollisionContainer.GetDetectorSettings().MaxVelocityBoundsExpansion != MaxBoundsVelocityExpansion))
-		{
-			// If the settings change we must recreate collisions
-			GetEvolution()->DestroyTransientConstraints();
-		}
-		CollisionContainer.SetVelocityBoundsExpansion(BoundsVelocityMultiplier, MaxBoundsVelocityExpansion);
+		GetEvolution()->GetCollisionConstraints().SetVelocityBoundsExpansion(BoundsVelocityMultiplier, MaxBoundsVelocityExpansion);
 	}
 
 	void FPBDRigidsSolver::SetVelocityBoundsExpansionMACD(const FReal BoundsVelocityMultiplier, const FReal MaxBoundsVelocityExpansion)
 	{
-		FPBDCollisionConstraints& CollisionContainer = GetEvolution()->GetCollisionConstraints();
-		if ((CollisionContainer.GetDetectorSettings().BoundsVelocityInflationMACD != BoundsVelocityMultiplier) || (CollisionContainer.GetDetectorSettings().MaxVelocityBoundsExpansionMACD != MaxBoundsVelocityExpansion))
-		{
-			// If the settings change we must recreate collisions
-			GetEvolution()->DestroyTransientConstraints();
-		}
-		CollisionContainer.SetVelocityBoundsExpansionMACD(BoundsVelocityMultiplier, MaxBoundsVelocityExpansion);
+		GetEvolution()->GetCollisionConstraints().SetVelocityBoundsExpansionMACD(BoundsVelocityMultiplier, MaxBoundsVelocityExpansion);
 	}
 
 	void FPBDRigidsSolver::PrepareAdvanceBy(const FReal DeltaTime)
@@ -1303,6 +1301,14 @@ namespace Chaos
 
 		// Apply CVAR overrides if set
 		{
+			// To enable runtime support for switching collision features on/off we need to update existing constraints when config changes.
+			if (bChaosCollisionConfigChanged)
+			{
+				// For now destroy the collisions. This is a bit over the top and causes problems for sleeping islands, but it's only for debugging/testing.
+				GetEvolution()->DestroyTransientConstraints();
+				bChaosCollisionConfigChanged = false;
+			}
+
 			if (ChaosSolverCollisionPositionFrictionIterations >= 0)
 			{
 				MEvolution->GetCollisionConstraints().SetPositionFrictionIterations(ChaosSolverCollisionPositionFrictionIterations);
@@ -1330,14 +1336,13 @@ namespace Chaos
 			{
 				SetCollisionCullDistance(ChaosSolverCullDistance);
 			}
-			SetVelocityBoundsExpansion(FMath::Max(0.0f, ChaosSolverVelocityBoundsMultiplier), FMath::Max(0.0f, ChaosSolverMaxVelocityBoundsExpansion));
-			if (bChaosUseMACD)
+			if ((ChaosSolverVelocityBoundsMultiplier >= 0.0f) && (ChaosSolverMaxVelocityBoundsExpansion >= 0.0f))
 			{
-				SetVelocityBoundsExpansionMACD(FMath::Max(0.0f, ChaosSolverVelocityBoundsMultiplierMACD), FMath::Max(0.0f, ChaosSolverMaxVelocityBoundsExpansionMACD));
+				SetVelocityBoundsExpansion(ChaosSolverVelocityBoundsMultiplier, ChaosSolverMaxVelocityBoundsExpansion);
 			}
-			else
+			if ((ChaosSolverVelocityBoundsMultiplierMACD >= 0.0f) && (ChaosSolverMaxVelocityBoundsExpansionMACD >= 0.0f))
 			{
-				SetVelocityBoundsExpansionMACD(FMath::Max(0.0f, ChaosSolverVelocityBoundsMultiplier), FMath::Max(0.0f, ChaosSolverMaxVelocityBoundsExpansion));
+				SetVelocityBoundsExpansionMACD(ChaosSolverVelocityBoundsMultiplierMACD, ChaosSolverMaxVelocityBoundsExpansion);
 			}
 			if (ChaosSolverMaxPushOutVelocity >= 0.0f)
 			{
@@ -2504,18 +2509,18 @@ TRACE_COUNTER_SET(ChaosTraceCounter_##Name, Value)
 #if CHAOS_DEBUG_DRAW
 		QUICK_SCOPE_CYCLE_COUNTER(SolverDebugDraw);
 
-		if (ChaosSolverDebugDrawShapes == 1)
+		const bool bIsServer = GetDebugName().ToString().StartsWith(TEXT("Server"));
+		if (bIsServer && !ChaosSolverDebugDrawShowServer)
 		{
-			const bool bIsServer = GetDebugName().ToString().StartsWith(TEXT("Server"));
-			if (bIsServer && !ChaosSolverDebugDrawShowServer)
-			{
-				return;
-			}
-			if (!bIsServer && !ChaosSolverDebugDrawShowClient)
-			{
-				return;
-			}
+			return;
+		}
+		if (!bIsServer && !ChaosSolverDebugDrawShowClient)
+		{
+			return;
+		}
 
+		if (ChaosSolverDebugDrawShapes == 1)
+		{
 			if (ChaosSolverDebugDrawColorShapeByClientServer)
 			{
 				if (bIsServer)
@@ -2531,7 +2536,6 @@ TRACE_COUNTER_SET(ChaosTraceCounter_##Name, Value)
 			{
 				ChaosSolverDebugDebugDrawSettings.ShapesColorsPerState = DebugDraw::GetDefaultShapesColorsByState();
 			}
-
 			DebugDrawShapes(!!ChaosSolverDrawShapesShowStatic, !!ChaosSolverDrawShapesShowKinematic, !!ChaosSolverDrawShapesShowDynamic);
 		}
 		if (ChaosSolverDebugDrawCollisions == 1)
