[Geometry Collection] Adding support for ISM custom data in UGeometryCollectionISMPoolRenderer.

#localization none
#tests Tested changes in PIE.

diff --git a/Engine/Source/Runtime/Experimental/GeometryCollectionEngine/Private/GeometryCollection/GeometryCollectionISMPoolRenderer.cpp b/Engine/Source/Runtime/Experimental/GeometryCollectionEngine/Private/GeometryCollection/GeometryCollectionISMPoolRenderer.cpp
--- a/Engine/Source/Runtime/Experimental/GeometryCollectionEngine/Private/GeometryCollection/GeometryCollectionISMPoolRenderer.cpp
+++ b/Engine/Source/Runtime/Experimental/GeometryCollectionEngine/Private/GeometryCollection/GeometryCollectionISMPoolRenderer.cpp
@@ -147,14 +147,14 @@ void UGeometryCollectionISMPoolRenderer::InitInstancesFromGeometryCollection(UGe
 			}
 			FGeometryCollectionStaticMeshInstance StaticMeshInstance;
 			StaticMeshInstance.StaticMesh = const_cast<UStaticMesh*>(StaticMesh);
+			StaticMeshInstance.Desc.NumCustomDataFloats = AutoInstanceMesh.GetNumDataPerInstance();
 			if (bMaterialOverride)
 			{
 				StaticMeshInstance.MaterialsOverrides.Reset();
 				StaticMeshInstance.MaterialsOverrides.Append(AutoInstanceMesh.Materials);
 			}
 
-			TArray<float> DummyCustomData;
-			InstancesGroup.MeshIds.Add(ISMPoolComponent->AddMeshToGroup(InstancesGroup.GroupIndex, StaticMeshInstance, AutoInstanceMesh.NumInstances, DummyCustomData));
+			InstancesGroup.MeshIds.Add(ISMPoolComponent->AddMeshToGroup(InstancesGroup.GroupIndex, StaticMeshInstance, AutoInstanceMesh.NumInstances, AutoInstanceMesh.CustomData));
 		}
 	}
 }
