Removed unused "pipeline length" parameter

 Pipelining of requests was changed to be automatic, to distribute work
 cross all open connections to the host, to facilitate cancel requests
better, and to mitigate head-of-line stalls while still trying to cover
up request latency. As such this pipeline length parameter became
unused. It may return in the future if we adopt h2 as they allows
requests to be cancelled.

#rnx

(git:2e9c84d - ridgers - [REDACTED_PHONE]:55:14 +0100)

diff --git a/Engine/Source/Runtime/Experimental/IoStoreOnDemand/Private/CoreHttp/Client.h b/Engine/Source/Runtime/Experimental/IoStoreOnDemand/Private/CoreHttp/Client.h
--- a/Engine/Source/Runtime/Experimental/IoStoreOnDemand/Private/CoreHttp/Client.h
+++ b/Engine/Source/Runtime/Experimental/IoStoreOnDemand/Private/CoreHttp/Client.h
@@ -57,7 +57,6 @@ public:
 	{
 		int32				SetHostFromUrl(FAnsiStringView Url);
 		uint16				ConnectionCount = 1;
-		uint16				PipelineLength = 1;
 		int32				SendBufSize = -1;
 		int32				RecvBufSize = -1;
 		struct {
diff --git a/Engine/Source/Runtime/Experimental/IoStoreOnDemand/Private/CoreHttp/Http.cpp b/Engine/Source/Runtime/Experimental/IoStoreOnDemand/Private/CoreHttp/Http.cpp
--- a/Engine/Source/Runtime/Experimental/IoStoreOnDemand/Private/CoreHttp/Http.cpp
+++ b/Engine/Source/Runtime/Experimental/IoStoreOnDemand/Private/CoreHttp/Http.cpp
@@ -1237,14 +1237,13 @@ public:
 	enum class EDirection : uint8 { Send, Recv };
 	static const uint32 InvalidIp = 0x00ff'ffff;
 
-					FHost(const ANSICHAR* InHostName, uint32 InPort, uint32 InMaxConn, uint32 PipeLength=1);
+					FHost(const ANSICHAR* InHostName, uint32 InPort, uint32 InMaxConn);
 	void			SetBufferSize(EDirection Dir, int32 Size);
 	int32			GetBufferSize(EDirection Dir) const;
 	FResult			Connect(FSocket& Socket);
 	int32			IsResolved() const;
 	FResult			ResolveHostName();
 	uint32			GetMaxConnections() const	{ return MaxConnections; }
-	uint32			GetPipelineLength() const	{ return PipelineLength; }
 	uint32			GetIpAddress() const		{ return IpAddresses[0]; }
 	FAnsiStringView	GetHostName() const			{ return HostName; }
 	uint32			GetPort() const				{ return Port; }
@@ -1256,15 +1255,13 @@ private:
 	int16			RecvBufKb = -1;
 	uint16			Port;
 	uint8			MaxConnections;
-	uint8			PipelineLength;
 };
 
 ////////////////////////////////////////////////////////////////////////////////
-FHost::FHost(const ANSICHAR* InHostName, uint32 InPort, uint32 InMaxConn, uint32 PipeLength)
+FHost::FHost(const ANSICHAR* InHostName, uint32 InPort, uint32 InMaxConn)
 : HostName(InHostName)
 , Port(uint16(InPort))
 , MaxConnections(uint8(InMaxConn))
-, PipelineLength(uint8(PipeLength))
 {
 	check(MaxConnections && MaxConnections == InMaxConn);
 }
@@ -1461,8 +1458,7 @@ FConnectionPool::FConnectionPool(const FParams& Params)
 	new (Internal) FHost(
 		HostDest,
 		Params.Host.Port,
-		Params.ConnectionCount,
-		Params.PipelineLength
+		Params.ConnectionCount
 	);
 	Internal->SetBufferSize(FHost::EDirection::Send, Params.SendBufSize);
 	Internal->SetBufferSize(FHost::EDirection::Recv, Params.RecvBufSize);
@@ -3798,7 +3794,6 @@ IOSTOREONDEMAND_API void IasHttpTest(const ANSICHAR* TestHost="localhost")
 		FConnectionPool::FParams Params;
 		Params.SetHostFromUrl(BuildUrl());
 		Params.ConnectionCount = (i % 2) + 1;
-		Params.PipelineLength = (i % 5) + 1;
 		FConnectionPool Pool(Params);
 		for (int32 j = 0; j < i; ++j)
 		{
diff --git a/Engine/Source/Runtime/Experimental/IoStoreOnDemand/Private/OnDemandHttpClient.cpp b/Engine/Source/Runtime/Experimental/IoStoreOnDemand/Private/OnDemandHttpClient.cpp
--- a/Engine/Source/Runtime/Experimental/IoStoreOnDemand/Private/OnDemandHttpClient.cpp
+++ b/Engine/Source/Runtime/Experimental/IoStoreOnDemand/Private/OnDemandHttpClient.cpp
@@ -243,7 +243,6 @@ TUniquePtr<HTTP::FConnectionPool> FHttpClient::CreateConnection(const FStringVie
 		Params.RecvBufSize = Config.ReceiveBufferSize;
 	}
 	Params.ConnectionCount = uint16(Config.MaxConnectionCount);
-	Params.PipelineLength = uint16(Config.PipelineLength);
 
 	return MakeUnique<HTTP::FConnectionPool>(Params);
 }
diff --git a/Engine/Source/Runtime/Experimental/IoStoreOnDemand/Private/OnDemandHttpClient.h b/Engine/Source/Runtime/Experimental/IoStoreOnDemand/Private/OnDemandHttpClient.h
--- a/Engine/Source/Runtime/Experimental/IoStoreOnDemand/Private/OnDemandHttpClient.h
+++ b/Engine/Source/Runtime/Experimental/IoStoreOnDemand/Private/OnDemandHttpClient.h
@@ -19,7 +19,6 @@ struct FHttpClientConfig
 	TArray<FString> Endpoints;
 	int32 PrimaryEndpoint = 0;
 	int32 MaxConnectionCount = 8;
-	int32 PipelineLength = 1;
 	int32 MaxRetryCount = 1;
 	int32 ReceiveBufferSize = -1;
 	bool bChangeEndpointAfterSuccessfulRetry = true;
diff --git a/Engine/Source/Runtime/Experimental/IoStoreOnDemand/Private/OnDemandIoDispatcherBackend.cpp b/Engine/Source/Runtime/Experimental/IoStoreOnDemand/Private/OnDemandIoDispatcherBackend.cpp
--- a/Engine/Source/Runtime/Experimental/IoStoreOnDemand/Private/OnDemandIoDispatcherBackend.cpp
+++ b/Engine/Source/Runtime/Experimental/IoStoreOnDemand/Private/OnDemandIoDispatcherBackend.cpp
@@ -112,13 +112,6 @@ static FAutoConsoleVariableRef CVar_IasHttpConnectionCount(
 	TEXT("Number of open HTTP connections to the on demand endpoint(s).")
 );
 
-static int32 GIasHttpPipelineLength = 2;
-static FAutoConsoleVariableRef CVar_GIasHttpPipelineLength(
-	TEXT("ias.HttpPipelineLength"),
-	GIasHttpPipelineLength,
-	TEXT("Number of concurrent requests on one connection")
-);
-
 /**
  *This is only applied when the connection was made to a single ServiceUrl rather than a DistributedUrl.
  * In the latter case we will make two attempts on the primary CDN followed by a single attempt for the
@@ -2502,7 +2495,6 @@ uint32 FOnDemandIoBackend::Run()
 		.Endpoints = AvailableEps.Urls,
 		.PrimaryEndpoint = FMath::Min(GIasHttpPrimaryEndpoint, AvailableEps.Urls.Num() -1),
 		.MaxConnectionCount = GIasHttpConnectionCount,
-		.PipelineLength = GIasHttpPipelineLength,
 		.MaxRetryCount = FMath::Max(AvailableEps.Urls.Num() + 1, GIasHttpRetryCount),
 		.ReceiveBufferSize = GIasHttpRecvBufKiB >= 0 ? GIasHttpRecvBufKiB << 10 : -1,
 		.bChangeEndpointAfterSuccessfulRetry = GIasHttpChangeEndpointAfterSuccessfulRetry,
