Deprecated FFastArraySerializer::FPostReplicatedReceiveParameters::bHasMoreUnmappedReferences as its behavior is inconsistent and code should not rely on it.
- Fixed issue with OldArraySize parameter not always propagated correctly when using iris

diff --git a/Engine/Plugins/Runtime/GameplayAbilities/Source/GameplayAbilities/Private/GameplayEffect.cpp b/Engine/Plugins/Runtime/GameplayAbilities/Source/GameplayAbilities/Private/GameplayEffect.cpp
--- a/Engine/Plugins/Runtime/GameplayAbilities/Source/GameplayAbilities/Private/GameplayEffect.cpp
+++ b/Engine/Plugins/Runtime/GameplayAbilities/Source/GameplayAbilities/Private/GameplayEffect.cpp
@@ -4879,6 +4879,8 @@ void FActiveGameplayEffectsContainer::PostReplicatedReceive(const FFastArraySeri
 		}
 		else
 		{
+// Keep this until we have actually deprecated the parameter just in case.
+PRAGMA_DISABLE_DEPRECATION_WARNINGS
 			if (!Parameters.bHasMoreUnmappedReferences) // Do not invoke GCs when we have missing information (like AActor*s in EffectContext)
 			{
 				NumConsecutiveUnmappedReferencesDebug = 0;
@@ -4895,6 +4897,7 @@ void FActiveGameplayEffectsContainer::PostReplicatedReceive(const FFastArraySeri
 				ensureMsgf(NumConsecutiveUnmappedReferencesDebug < HighNumberOfConsecutiveUnmappedRefs, TEXT("%hs: bHasMoreUnmappedReferences is preventing GameplayCues from firing"), __func__);
 				UE_CLOG((NumConsecutiveUnmappedReferencesDebug % HighNumberOfConsecutiveUnmappedRefs) == 0, LogAbilitySystem, Error, TEXT("%hs: bHasMoreUnmappedReferences is preventing GameplayCues from firing (%u consecutive misses)"), __func__, NumConsecutiveUnmappedReferencesDebug);
 			}
+PRAGMA_ENABLE_DEPRECATION_WARNINGS
 		}
 	}
 }
diff --git a/Engine/Source/Runtime/Experimental/Iris/Core/Public/Iris/ReplicationSystem/Private/FastArrayReplicationFragmentInternal.h b/Engine/Source/Runtime/Experimental/Iris/Core/Public/Iris/ReplicationSystem/Private/FastArrayReplicationFragmentInternal.h
--- a/Engine/Source/Runtime/Experimental/Iris/Core/Public/Iris/ReplicationSystem/Private/FastArrayReplicationFragmentInternal.h
+++ b/Engine/Source/Runtime/Experimental/Iris/Core/Public/Iris/ReplicationSystem/Private/FastArrayReplicationFragmentInternal.h
@@ -81,15 +81,19 @@ struct FFastArrayReplicationFragmentHelper
 	 * We only want to do this for FastArrays that define PostReplicatedReceive since it might require extra work to calculate the required parameters
 	 */
 	template<typename FastArrayType>
-	static inline typename TEnableIf<TModels_V<FFastArraySerializer::CPostReplicatedReceiveFuncable, FastArrayType, const FFastArraySerializer::FPostReplicatedReceiveParameters&>, void>::Type CallPostReplicatedReceiveOrNot(FastArrayType& ArraySerializer, bool bHasUnresolvedReferences)
+	static inline typename TEnableIf<TModels_V<FFastArraySerializer::CPostReplicatedReceiveFuncable, FastArrayType, const FFastArraySerializer::FPostReplicatedReceiveParameters&>, void>::Type CallPostReplicatedReceiveOrNot(FastArrayType& ArraySerializer, int32 OldArraySize, bool bHasUnresolvedReferences)
 	{
 		FFastArraySerializer::FPostReplicatedReceiveParameters PostReceivedParameters;
+PRAGMA_DISABLE_DEPRECATION_WARNINGS
 		PostReceivedParameters.bHasMoreUnmappedReferences = bHasUnresolvedReferences;
+PRAGMA_ENABLE_DEPRECATION_WARNINGS
+		PostReceivedParameters.OldArraySize = OldArraySize;
+
 		ArraySerializer.PostReplicatedReceive(PostReceivedParameters);
 	}
 
 	template<typename FastArrayType>
-	static inline typename TEnableIf<!TModels_V<FFastArraySerializer::CPostReplicatedReceiveFuncable, FastArrayType, const FFastArraySerializer::FPostReplicatedReceiveParameters&>, void>::Type CallPostReplicatedReceiveOrNot(FastArrayType& ArraySerializer, bool bHasUnresolvedReferences) {}
+	static inline typename TEnableIf<!TModels_V<FFastArraySerializer::CPostReplicatedReceiveFuncable, FastArrayType, const FFastArraySerializer::FPostReplicatedReceiveParameters&>, void>::Type CallPostReplicatedReceiveOrNot(FastArrayType& ArraySerializer, int32 OldArraySize, bool bHasUnresolvedReferences) {}
 };
 
 class FFastArrayReplicationFragmentBase : public FReplicationFragment
@@ -217,6 +221,9 @@ void FFastArrayReplicationFragmentHelper::ApplyReplicatedState(FastArrayType* Ds
 	ConditionalRebuildItemMap(*DstArraySerializer, *DstWrappedArray, false);
 	ConditionalRebuildItemMap(*SrcArraySerializer, *SrcWrappedArray, bForceRebuildItemMap);
 
+	// We need this for callback
+	const int32 OriginalSize = DstWrappedArray->Num();
+
 	// Find removed elements in received data, that is elements that exist in old map but not in new map
 	TArray<int32> RemovedIndices;
 	{
@@ -341,7 +348,7 @@ void FFastArrayReplicationFragmentHelper::ApplyReplicatedState(FastArrayType* Ds
 	}
 
 	// Invoke PostReplicatedReceive if is defined by the serializer
-	CallPostReplicatedReceiveOrNot(*DstArraySerializer, Context.bHasUnresolvableReferences);
+	CallPostReplicatedReceiveOrNot(*DstArraySerializer, OriginalSize, Context.bHasUnresolvableReferences);
 }
 
 }} // End of namespaces
diff --git a/Engine/Source/Runtime/Net/Core/Classes/Net/Serialization/FastArraySerializer.h b/Engine/Source/Runtime/Net/Core/Classes/Net/Serialization/FastArraySerializer.h
--- a/Engine/Source/Runtime/Net/Core/Classes/Net/Serialization/FastArraySerializer.h
+++ b/Engine/Source/Runtime/Net/Core/Classes/Net/Serialization/FastArraySerializer.h
@@ -520,6 +520,7 @@ struct FFastArraySerializer
 	{
 		// This is the size that the array had before the receive.
 		int32 OldArraySize;
+		UE_DEPRECATED(5.4, "This is unsafe to use and will be removed.")
 		uint32 bHasMoreUnmappedReferences : 1U;
 	};
 
@@ -693,6 +694,7 @@ private:
 			TArray<int32, TInlineAllocator<8>>& AddedIndices,
 			GuidMapType& GuidMap);
 	
+PRAGMA_DISABLE_DEPRECATION_WARNINGS
 		/** Conditionally invoke PostReplicatedReceive method depending on if is defined or not */
 		template<typename FastArrayType = SerializerType>
 		inline typename TEnableIf<TModels_V<CPostReplicatedReceiveFuncable, FastArrayType, const FFastArraySerializer::FPostReplicatedReceiveParameters>, void>::Type CallPostReplicatedReceiveOrNot(int32 OldArraySize)
@@ -702,6 +704,7 @@ private:
 			PostReceivedParameters.bHasMoreUnmappedReferences = Parms.bOutHasMoreUnmapped;
 			ArraySerializer.PostReplicatedReceive(PostReceivedParameters);
 		}
+PRAGMA_ENABLE_DEPRECATION_WARNINGS
 
 		template<typename FastArrayType = SerializerType>
 		inline typename TEnableIf<!TModels_V<CPostReplicatedReceiveFuncable, FastArrayType, const FFastArraySerializer::FPostReplicatedReceiveParameters>, void>::Type CallPostReplicatedReceiveOrNot(int32) {}
