Iris - Keep reliability as is when splitting a NetBlob. Split unreliable multicast RPCs will now also keep their normal execution order.

diff --git a/Engine/Plugins/Runtime/ReplicationSystemTestPlugin/Source/Private/Tests/ReplicationSystem/NetBlob/TestPartialNetBlob.cpp b/Engine/Plugins/Runtime/ReplicationSystemTestPlugin/Source/Private/Tests/ReplicationSystem/NetBlob/TestPartialNetBlob.cpp
--- a/Engine/Plugins/Runtime/ReplicationSystemTestPlugin/Source/Private/Tests/ReplicationSystem/NetBlob/TestPartialNetBlob.cpp
+++ b/Engine/Plugins/Runtime/ReplicationSystemTestPlugin/Source/Private/Tests/ReplicationSystem/NetBlob/TestPartialNetBlob.cpp
@@ -82,7 +82,7 @@ UE_NET_TEST_FIXTURE(FPartialNetBlobTestFixture, MissingFirstPartialNetBlobCauses
 	SetupTestContext(ServerContext, Server);
 
 	constexpr uint32 PayloadBitCount = 44444U;
-	const TRefCountPtr<FNetBlob>& Blob = CreateUnreliableMockNetBlob(PayloadBitCount);
+	const TRefCountPtr<FNetBlob>& Blob = CreateReliableMockNetBlob(PayloadBitCount);
 
 	TArray<TRefCountPtr<FNetBlob>> PartialNetBlobs;
 	MockSequentialPartialNetBlobHandler->SplitNetBlob(Blob, PartialNetBlobs);
@@ -98,7 +98,7 @@ UE_NET_TEST_FIXTURE(FPartialNetBlobTestFixture, MissingArbitraryPartialNetBlobCa
 	SetupTestContext(ServerContext, Server);
 
 	constexpr uint32 PayloadBitCount = 44444U;
-	const TRefCountPtr<FNetBlob>& Blob = CreateUnreliableMockNetBlob(PayloadBitCount);
+	const TRefCountPtr<FNetBlob>& Blob = CreateReliableMockNetBlob(PayloadBitCount);
 
 	TArray<TRefCountPtr<FNetBlob>> PartialNetBlobs;
 	MockSequentialPartialNetBlobHandler->SplitNetBlob(Blob, PartialNetBlobs);
@@ -129,7 +129,7 @@ UE_NET_TEST_FIXTURE(FPartialNetBlobTestFixture, TestSplitBlobIsReceivedDespitePa
 	// Create attachment
 	{
 		constexpr uint32 PayloadBitCount = 16000;
-		const TRefCountPtr<FNetObjectAttachment>& Attachment = MockNetObjectAttachmentHandler->CreateUnreliableNetObjectAttachment(PayloadBitCount);
+		const TRefCountPtr<FNetObjectAttachment>& Attachment = MockNetObjectAttachmentHandler->CreateReliableNetObjectAttachment(PayloadBitCount);
 		FNetObjectReference AttachmentTarget = FObjectReferenceCache::MakeNetObjectReference(ServerObject->NetRefHandle);
 		Server->GetReplicationSystem()->QueueNetObjectAttachment(Client->ConnectionIdOnServer, AttachmentTarget, Attachment);
 	}
@@ -161,6 +161,7 @@ UE_NET_TEST_FIXTURE(FPartialNetBlobTestFixture, TestSplitBlobIsReceivedDespitePa
 	Server->DestroyObject(ServerObject);
 }
 
+#if 0 // Temporary disabled because it will fail for the wrong reason.
 UE_NET_TEST_FIXTURE(FPartialNetBlobTestFixture, PartOfUnexpectedSizeCausesError)
 {
 	FTestContext ServerContext;
@@ -171,7 +172,7 @@ UE_NET_TEST_FIXTURE(FPartialNetBlobTestFixture, PartOfUnexpectedSizeCausesError)
 	{
 		const USequentialPartialNetBlobHandlerConfig* Config = MockSequentialPartialNetBlobHandler->GetConfig();
 
-		// Guesstimate on blob overhead. The important thing is we get exactly as many pars as we wish.
+		// Guesstimate on blob overhead. The important thing is we get exactly as many parts as we wish.
 		constexpr uint32 BlobOverheadBitCount = 128U;
 
 		constexpr int32 PartCount = 3U;
@@ -196,5 +197,38 @@ UE_NET_TEST_FIXTURE(FPartialNetBlobTestFixture, PartOfUnexpectedSizeCausesError)
 	ServerContext.SerializationContext.GetNetBlobReceiver()->OnNetBlobReceived(ServerContext.SerializationContext, WrongSizeSecondBlob);
 	UE_NET_ASSERT_TRUE(ServerContext.SerializationContext.HasError());
 }
+#endif
+
+UE_NET_TEST_FIXTURE(FPartialNetBlobTestFixture, UnexpectedPartSequenceCausesError)
+{
+	FTestContext ServerContext;
+	SetupTestContext(ServerContext, Server);
+
+	TArray<TRefCountPtr<FNetBlob>> PartialNetBlobs;
+	TRefCountPtr<FNetBlob> WrongSecondBlob;
+	{
+		const USequentialPartialNetBlobHandlerConfig* Config = MockSequentialPartialNetBlobHandler->GetConfig();
+
+		// Guesstimate on blob overhead. The important thing is we get exactly as many parts as we wish.
+		constexpr uint32 BlobOverheadBitCount = 128U;
+
+		constexpr int32 PartCount = 3U;
+		uint32 TotalPayloadBitCount = PartCount * Config->GetMaxPartBitCount() - BlobOverheadBitCount - 1U;
+		const TRefCountPtr<FNetBlob>& Blob = CreateReliableMockNetBlob(TotalPayloadBitCount);
+		MockSequentialPartialNetBlobHandler->SplitNetBlob(Blob, PartialNetBlobs);
+		UE_NET_ASSERT_EQ(PartialNetBlobs.Num(), PartCount);
+
+		TArray<TRefCountPtr<FNetBlob>> OtherPartialNetBlobs;
+		const TRefCountPtr<FNetBlob>& OtherBlob = CreateReliableMockNetBlob(TotalPayloadBitCount);
+		MockSequentialPartialNetBlobHandler->SplitNetBlob(OtherBlob, OtherPartialNetBlobs);
+	
+		WrongSecondBlob = MoveTemp(OtherPartialNetBlobs[1]);
+	}
+
+
+	ServerContext.SerializationContext.GetNetBlobReceiver()->OnNetBlobReceived(ServerContext.SerializationContext, PartialNetBlobs[0]);
+	ServerContext.SerializationContext.GetNetBlobReceiver()->OnNetBlobReceived(ServerContext.SerializationContext, WrongSecondBlob);
+	UE_NET_ASSERT_TRUE(ServerContext.SerializationContext.HasError());
+}
 
 }
diff --git a/Engine/Source/Runtime/Experimental/Iris/Core/Private/Iris/ReplicationSystem/AttachmentReplication.cpp b/Engine/Source/Runtime/Experimental/Iris/Core/Private/Iris/ReplicationSystem/AttachmentReplication.cpp
--- a/Engine/Source/Runtime/Experimental/Iris/Core/Private/Iris/ReplicationSystem/AttachmentReplication.cpp
+++ b/Engine/Source/Runtime/Experimental/Iris/Core/Private/Iris/ReplicationSystem/AttachmentReplication.cpp
@@ -634,21 +634,35 @@ public:
 	struct FInitParams
 	{
 		const UPartialNetObjectAttachmentHandler* PartialNetObjectAttachmentHandler = nullptr;
+		bool bOnlyProcessUnreliable = false;
 	};
 
 	void Init(const FInitParams& InitParams)
 	{
 		PartialNetObjectAttachmentHandler = InitParams.PartialNetObjectAttachmentHandler;
+		bOnlyProcessUnreliable = InitParams.bOnlyProcessUnreliable;
 	}
 
 	bool IsEmpty() const { return Queue.IsEmpty(); }
 	bool HasUnprocessed() const { return !Queue.IsEmpty(); }
-	bool IsSafeToDestroy() const { return Queue.IsEmpty() && !NetBlobAssembler.IsValid(); }
+	bool IsSafeToDestroy() const
+	{ 
+		return Queue.IsEmpty() && !ReliableNetBlobAssembler.IsValid() && !UnreliableNetBlobAssembler.IsValid();
+	}
 
 	void Enqueue(FNetSerializationContext& Context, FNetRefHandle RefHandle, const TRefCountPtr<FNetBlob>& NetBlob, bool bIsPartialNetBlob)
 	{
+		const bool bIsBlobReliable = NetBlob->IsReliable();
+		if (bIsBlobReliable && bOnlyProcessUnreliable)
+		{
+			UE_LOG(LogIris, Error, TEXT("Received reliable blob when only unreliable is supported."));
+			Context.SetError(GNetError_UnsupportedNetBlob);
+			return;
+		}
+
 		if (bIsPartialNetBlob)
 		{
+			TUniquePtr<FNetBlobAssembler>& NetBlobAssembler = bIsBlobReliable ? ReliableNetBlobAssembler : UnreliableNetBlobAssembler;
 			if (!NetBlobAssembler.IsValid())
 			{
 				FNetBlobAssemblerInitParams InitParams;
@@ -658,18 +672,27 @@ public:
 			}
 
 			NetBlobAssembler->AddPartialNetBlob(Context, RefHandle, reinterpret_cast<const TRefCountPtr<FPartialNetBlob>&>(NetBlob));
-			if (NetBlobAssembler->IsReadyToAssemble())
+			if (NetBlobAssembler->IsReadyToAssemble() || NetBlobAssembler->IsSequenceBroken())
 			{
-				const TRefCountPtr<FNetBlob>& AssembledBlob = NetBlobAssembler->Assemble(Context);
-				if (AssembledBlob.IsValid())
+				if (NetBlobAssembler->IsReadyToAssemble())
 				{
-					Queue.Enqueue(AssembledBlob);
+					const TRefCountPtr<FNetBlob>& AssembledBlob = NetBlobAssembler->Assemble(Context);
+					if (AssembledBlob.IsValid())
+					{
+						Queue.Enqueue(AssembledBlob);
+					}
 				}
 				NetBlobAssembler.Reset();
 			}
 		}
 		else
 		{
+			// If we're set to only process unreliable blobs and we're not done assembling a full set of unreliable partial blobs at this point we know we will not succeed in doing so. We want to allow the queue to be destroyed as soon as possible so let's reset the assembler.
+			if (bOnlyProcessUnreliable)
+			{
+				UnreliableNetBlobAssembler.Reset();
+			}
+
 			Queue.Enqueue(NetBlob);
 		}
 	}
@@ -698,39 +721,40 @@ public:
 		return Queue.Pop();
 	}
 
-	void DequeueUnreliable(TArray<TRefCountPtr<FNetBlob>>& Unreliable)
+	SIZE_T GetQueueCount() const
 	{
-		for (SIZE_T It = 0, EndIt = Queue.Count(); It < EndIt; ++It)
-		{
-			TRefCountPtr<FNetBlob>& RefCntBlob = Queue.PokeAtOffset(It);
-			if (RefCntBlob.GetRefCount() > 0)
-			{
-				if (const bool bIsUnReliable = !EnumHasAnyFlags(RefCntBlob.GetReference()->GetCreationInfo().Flags, ENetBlobFlags::Reliable))
-				{
-					Unreliable.Emplace(MoveTemp(RefCntBlob));
-				}
-			}
-		}
+		return Queue.Count();
+	}
+
+	void SetQueueCapacity(uint32 QueueCapacity)
+	{
+		const SIZE_T QueueCount = Queue.Count();
+		const SIZE_T DropCount = QueueCount - QueueCapacity;
+		UE_LOG(LogIris, Warning, TEXT("Dropping %u attachments to due to change in unreliable queue capacity to %u"), DropCount, QueueCapacity);
+		Queue.Pop(DropCount);
 	}
 
 private:
-	TUniquePtr<FNetBlobAssembler> NetBlobAssembler;
+	TUniquePtr<FNetBlobAssembler> ReliableNetBlobAssembler;
+	TUniquePtr<FNetBlobAssembler> UnreliableNetBlobAssembler;
 	TResizableCircularQueue<TRefCountPtr<FNetBlob>> Queue;
 	const UPartialNetObjectAttachmentHandler* PartialNetObjectAttachmentHandler = nullptr;
+	/* Whether this queue is only expecting unrealiable NetBlobs or not. */
+	bool bOnlyProcessUnreliable = false;
 };
 
 FNetObjectAttachmentReceiveQueue::FNetObjectAttachmentReceiveQueue()
-: ReliableQueue(nullptr)
-, DeferredProcessingQueue(nullptr)
-, MaxUnreliableCount(AttachmentReplicationCVars::UnreliableRPCQueueSize)
-, PartialNetBlobType(InvalidNetBlobType)
+: MaxUnreliableCount(AttachmentReplicationCVars::UnreliableRPCQueueSize)
 {
 }
 
 FNetObjectAttachmentReceiveQueue::~FNetObjectAttachmentReceiveQueue()
 {
 	delete ReliableQueue;
-	delete DeferredProcessingQueue;
+	for (FDeferredProcessingQueue* Queue : MakeArrayView(DeferredProcessingQueues))
+	{
+		delete Queue;
+	}
 }
 
 void FNetObjectAttachmentReceiveQueue::Init(const FNetObjectAttachmentReceiveQueueInitParams& InitParams)
@@ -741,89 +765,71 @@ void FNetObjectAttachmentReceiveQueue::Init(const FNetObjectAttachmentReceiveQue
 
 bool FNetObjectAttachmentReceiveQueue::IsSafeToDestroy() const
 {
-	return UnreliableQueue.IsEmpty() && IsDeferredProcessingQueueSafeToDestroy() && (ReliableQueue == nullptr || ReliableQueue->IsSafeToDestroy());
+	return IsDeferredProcessingQueueSafeToDestroy(EDeferredProcessingQueue::Unreliable) && IsDeferredProcessingQueueSafeToDestroy(EDeferredProcessingQueue::Reliable) && (ReliableQueue == nullptr || ReliableQueue->IsSafeToDestroy());
 }
 
 bool FNetObjectAttachmentReceiveQueue::HasUnprocessed() const
 {
-	return !UnreliableQueue.IsEmpty() || HasDeferredProcessingQueueUnprocessed();
+	return HasDeferredProcessingQueueUnprocessed(EDeferredProcessingQueue::Unreliable) || HasDeferredProcessingQueueUnprocessed(EDeferredProcessingQueue::Reliable);
 }
 
 const TRefCountPtr<FNetBlob>* FNetObjectAttachmentReceiveQueue::PeekReliable()
 {
-	if (IsDeferredProcessingQueueEmpty())
+	if (IsDeferredProcessingQueueEmpty(EDeferredProcessingQueue::Reliable))
 	{
 		return nullptr;
 	}
 	else
 	{
-		return DeferredProcessingQueue->Peek();
+		return DeferredProcessingQueues[EDeferredProcessingQueue::Reliable]->Peek();
 	}
 }
 
 void FNetObjectAttachmentReceiveQueue::PopReliable()
 {
-	DeferredProcessingQueue->Pop();
+	DeferredProcessingQueues[EDeferredProcessingQueue::Reliable]->Pop();
 }
 
 const TRefCountPtr<FNetBlob>* FNetObjectAttachmentReceiveQueue::PeekUnreliable() const
 {
-	if (UnreliableQueue.IsEmpty())
+	if (IsDeferredProcessingQueueEmpty(EDeferredProcessingQueue::Unreliable))
 	{
 		return nullptr;
 	}
 	else
 	{
-		return &UnreliableQueue.Peek();
+		return DeferredProcessingQueues[EDeferredProcessingQueue::Unreliable]->Peek();
 	}
 }
 
 void FNetObjectAttachmentReceiveQueue::PopUnreliable()
 {
-	UnreliableQueue.Pop();
-}
-
-void FNetObjectAttachmentReceiveQueue::GetOrderedUnreliable(TArray<TRefCountPtr<FNetBlob>>& OrderedUnreliable)
-{
-	if (DeferredProcessingQueue != nullptr)
-	{
-		DeferredProcessingQueue->DequeueUnreliable(OrderedUnreliable);
-	}
-	
-	if (ReliableQueue != nullptr)
-	{
-		ReliableQueue->DequeueUnreliable(OrderedUnreliable);
-	}
+	return DeferredProcessingQueues[EDeferredProcessingQueue::Unreliable]->Pop();
 }
 
 void FNetObjectAttachmentReceiveQueue::SetUnreliableQueueCapacity(uint32 QueueCapacity)
 {
 	MaxUnreliableCount = QueueCapacity;
 	
-	const SIZE_T UnreliableCount = UnreliableQueue.Count();
-	if (QueueCapacity >= UnreliableCount)
+	if (FDeferredProcessingQueue* UnreliableProcessingQueue = DeferredProcessingQueues[EDeferredProcessingQueue::Unreliable])
 	{
-		return;
+		UnreliableProcessingQueue->SetQueueCapacity(QueueCapacity);
 	}
-
-	const SIZE_T DropCount = UnreliableCount - QueueCapacity;
-	UE_LOG(LogIris, Warning, TEXT("Dropping %u attachments to due to change in unreliable queue capacity to %u"), DropCount, QueueCapacity);
-	UnreliableQueue.Pop(DropCount);
 }
 
-bool FNetObjectAttachmentReceiveQueue::IsDeferredProcessingQueueEmpty() const
+bool FNetObjectAttachmentReceiveQueue::IsDeferredProcessingQueueEmpty(EDeferredProcessingQueue Queue) const
 {
-	return DeferredProcessingQueue == nullptr || DeferredProcessingQueue->IsEmpty();
+	return DeferredProcessingQueues[Queue] == nullptr || DeferredProcessingQueues[Queue]->IsEmpty();
 }
 
-bool FNetObjectAttachmentReceiveQueue::IsDeferredProcessingQueueSafeToDestroy() const
+bool FNetObjectAttachmentReceiveQueue::IsDeferredProcessingQueueSafeToDestroy(EDeferredProcessingQueue Queue) const
 {
-	return DeferredProcessingQueue == nullptr || DeferredProcessingQueue->IsSafeToDestroy();
+	return DeferredProcessingQueues[Queue] == nullptr || DeferredProcessingQueues[Queue]->IsSafeToDestroy();
 }
 
-bool FNetObjectAttachmentReceiveQueue::HasDeferredProcessingQueueUnprocessed() const
+bool FNetObjectAttachmentReceiveQueue::HasDeferredProcessingQueueUnprocessed(EDeferredProcessingQueue Queue) const
 {
-	return DeferredProcessingQueue != nullptr && DeferredProcessingQueue->HasUnprocessed();
+	return DeferredProcessingQueues[Queue] != nullptr && DeferredProcessingQueues[Queue]->HasUnprocessed();
 }
 
 bool FNetObjectAttachmentReceiveQueue::IsPartialNetBlob(const TRefCountPtr<FNetBlob>& Blob) const
@@ -871,11 +877,11 @@ uint32 FNetObjectAttachmentReceiveQueue::DeserializeReliable(FNetSerializationCo
 	if (ReliableQueue == nullptr)
 	{
 		ReliableQueue = new FReliableNetBlobQueue();
-		checkSlow(DeferredProcessingQueue == nullptr);
-		DeferredProcessingQueue = new FDeferredProcessingQueue();
-		FDeferredProcessingQueue::FInitParams InitParams;
-		InitParams.PartialNetObjectAttachmentHandler = PartialNetObjectAttachmentHandler;
-		DeferredProcessingQueue->Init(InitParams);
+
+		checkSlow(DeferredProcessingQueues[EDeferredProcessingQueue::Reliable] == nullptr);
+		DeferredProcessingQueues[EDeferredProcessingQueue::Reliable] = new FDeferredProcessingQueue();
+		FDeferredProcessingQueue::FInitParams InitParams = { .PartialNetObjectAttachmentHandler = PartialNetObjectAttachmentHandler, .bOnlyProcessUnreliable = false };
+		DeferredProcessingQueues[EDeferredProcessingQueue::Reliable]->Init(InitParams);
 	}
 
 	uint32 DeserializedReliableCount = 0;
@@ -893,6 +899,7 @@ uint32 FNetObjectAttachmentReceiveQueue::DeserializeReliable(FNetSerializationCo
 		return DeserializedReliableCount;
 	}
 
+	FDeferredProcessingQueue* DeferredProcessingQueue = DeferredProcessingQueues[EDeferredProcessingQueue::Reliable];
 	while (const TRefCountPtr<FNetBlob>* Attachment = ReliableQueue->Peek())
 	{
 		DeferredProcessingQueue->Enqueue(Context, RefHandle, *Attachment, IsPartialNetBlob(*Attachment));
@@ -904,11 +911,35 @@ uint32 FNetObjectAttachmentReceiveQueue::DeserializeReliable(FNetSerializationCo
 		}
 	}
 
+	// Add all received unreliable attachments to prevent their processing from being blocked by reliable ones.
+	{
+		TArray<TRefCountPtr<FNetBlob>> UnreliableAttachments;
+		UnreliableAttachments.Reserve(16);
+		ReliableQueue->DequeueUnreliable(UnreliableAttachments);
+		for (TRefCountPtr<FNetBlob>& UnreliableAttachment : UnreliableAttachments)
+		{
+			DeferredProcessingQueue->Enqueue(Context, RefHandle, UnreliableAttachment, IsPartialNetBlob(UnreliableAttachment));
+			if (Context.HasErrorOrOverflow())
+			{
+				return DeserializedReliableCount;
+			}
+		}
+	}
+
 	return DeserializedReliableCount;
 }
 
 uint32 FNetObjectAttachmentReceiveQueue::DeserializeUnreliable(FNetSerializationContext& Context, FNetRefHandle RefHandle)
 {
+	if (DeferredProcessingQueues[EDeferredProcessingQueue::Unreliable] == nullptr)
+	{
+		DeferredProcessingQueues[EDeferredProcessingQueue::Unreliable] = new FDeferredProcessingQueue();
+		FDeferredProcessingQueue::FInitParams InitParams = { .PartialNetObjectAttachmentHandler = PartialNetObjectAttachmentHandler, .bOnlyProcessUnreliable = true };
+		DeferredProcessingQueues[EDeferredProcessingQueue::Unreliable]->Init(InitParams);
+	}
+
+	FDeferredProcessingQueue* DeferredProcessingQueue = DeferredProcessingQueues[EDeferredProcessingQueue::Unreliable];
+
 	INetBlobReceiver* BlobReceiver = Context.GetNetBlobReceiver();
 	checkSlow(BlobReceiver != nullptr);
 
@@ -920,7 +951,7 @@ uint32 FNetObjectAttachmentReceiveQueue::DeserializeUnreliable(FNetSerialization
 	do
 	{
 		// If the unreliable queue overflows this could be a carefully crafted malicious packet.
-		if (UnreliableQueue.Count() == MaxUnreliableCount)
+		if (DeferredProcessingQueue->GetQueueCount() == MaxUnreliableCount)
 		{
 			UE_LOG(LogIris, Error, TEXT("Unreliable queue is full for %s"), *RefHandle.ToString());
 			Context.SetError(NetError_UnreliableQueueFull);
@@ -954,7 +985,13 @@ uint32 FNetObjectAttachmentReceiveQueue::DeserializeUnreliable(FNetSerialization
 			break;
 		}
 
-		UnreliableQueue.Enqueue(Attachment);
+		DeferredProcessingQueue->Enqueue(Context, RefHandle, Attachment, IsPartialNetBlob(Attachment));
+		if (Context.HasErrorOrOverflow())
+		{
+			UE_LOG(LogIris, Error, TEXT("Failed to deserialize unreliable attachments for %s"), *RefHandle.ToString());
+			break;
+		}
+
 		++DeserializedUnreliableCount;
 	} while (bHasMoreAttachments);
 
diff --git a/Engine/Source/Runtime/Experimental/Iris/Core/Private/Iris/ReplicationSystem/AttachmentReplication.h b/Engine/Source/Runtime/Experimental/Iris/Core/Private/Iris/ReplicationSystem/AttachmentReplication.h
--- a/Engine/Source/Runtime/Experimental/Iris/Core/Private/Iris/ReplicationSystem/AttachmentReplication.h
+++ b/Engine/Source/Runtime/Experimental/Iris/Core/Private/Iris/ReplicationSystem/AttachmentReplication.h
@@ -172,29 +172,33 @@ public:
 	const TRefCountPtr<FNetBlob>* PeekUnreliable() const;
 	void PopUnreliable();
 
-	void GetOrderedUnreliable(TArray<TRefCountPtr<FNetBlob>>& OrderedUnreliable);
-
 	void SetUnreliableQueueCapacity(uint32 QueueCapacity);
 
 private:
 	friend FNetObjectAttachmentsReader;
 	class FDeferredProcessingQueue;
 
-	bool IsDeferredProcessingQueueEmpty() const;
-	bool IsDeferredProcessingQueueSafeToDestroy() const;
-	bool HasDeferredProcessingQueueUnprocessed() const;
+	enum EDeferredProcessingQueue : unsigned
+	{
+		Unreliable,
+		Reliable,
+	};
+
+	bool IsDeferredProcessingQueueEmpty(EDeferredProcessingQueue Queue) const;
+	bool IsDeferredProcessingQueueSafeToDestroy(EDeferredProcessingQueue Queue) const;
+	bool HasDeferredProcessingQueueUnprocessed(EDeferredProcessingQueue Queue) const;
+
 	bool IsPartialNetBlob(const TRefCountPtr<FNetBlob>& Blob) const;
 
 	void Deserialize(FNetSerializationContext& Context, FNetRefHandle RefHandle);
 	uint32 DeserializeReliable(FNetSerializationContext& Context, FNetRefHandle RefHandle);
 	uint32 DeserializeUnreliable(FNetSerializationContext& Context, FNetRefHandle RefHandle);
 
-	TResizableCircularQueue<TRefCountPtr<FNetBlob>> UnreliableQueue;
-	FReliableNetBlobQueue* ReliableQueue;
-	FDeferredProcessingQueue* DeferredProcessingQueue;
-	uint32 MaxUnreliableCount;
-	FNetBlobType PartialNetBlobType;
+	FReliableNetBlobQueue* ReliableQueue = nullptr;
+	FDeferredProcessingQueue* DeferredProcessingQueues[2] = {};
 	const UPartialNetObjectAttachmentHandler* PartialNetObjectAttachmentHandler = nullptr;
+	uint32 MaxUnreliableCount = 0;
+	FNetBlobType PartialNetBlobType = InvalidNetBlobType;
 };
 
 struct FNetObjectAttachmentsReaderInitParams
diff --git a/Engine/Source/Runtime/Experimental/Iris/Core/Private/Iris/ReplicationSystem/NetBlob/NetBlobAssembler.cpp b/Engine/Source/Runtime/Experimental/Iris/Core/Private/Iris/ReplicationSystem/NetBlob/NetBlobAssembler.cpp
--- a/Engine/Source/Runtime/Experimental/Iris/Core/Private/Iris/ReplicationSystem/NetBlob/NetBlobAssembler.cpp
+++ b/Engine/Source/Runtime/Experimental/Iris/Core/Private/Iris/ReplicationSystem/NetBlob/NetBlobAssembler.cpp
@@ -31,24 +31,60 @@ void FNetBlobAssembler::AddPartialNetBlob(FNetSerializationContext& Context, FNe
 		return;
 	}
 
-	const uint32 PartIndex = PartialNetBlob->GetPartIndex();
-	if (PartIndex != NextPartIndex)
+	// We expect broken sequences to be handled as soon as they're reported. We've already reported the error once.
+	if (bIsBrokenSequence)
 	{
 		Context.SetError(NetError_PartialNetBlobSequenceError);
 		return;
 	}
 
-	if (PartIndex == 0U)
+	const uint32 SequenceNumber = PartialNetBlob->GetSequenceNumber();
+	const bool bIsFirstPart = PartialNetBlob->IsFirstPart();
+
+	// Broken sequence detection. For reliable blobs the NextPartIndex and NextSequenceSumber must match. If reliability changes we need some additional validation.
+	const bool bIsReliable = PartialNetBlob->IsReliable();
+	if (SequenceNumber != NextSequenceSumber || bIsReliable != bIsProcessingReliable)
+	{
+		// Reliable sequences are expected to be fully processed before proceeding to a new sequence. If we're in the middle of processing everything needs to match.
+		if (bIsProcessingReliable)
+		{
+			bIsBrokenSequence = true;
+			Context.SetError(NetError_PartialNetBlobSequenceError);
+			return;
+		}
+
+		// Gracefully handle going from unreliable to first part of blob, regardless of its reliability.
+		if (bIsFirstPart)
+		{
+			NextSequenceSumber = SequenceNumber;
+			LastPartSequenceNumber = SequenceNumber + PartialNetBlob->GetPartCount() - 1U;
+		}
+		else if (SequenceNumber != NextSequenceSumber)
+		{
+			bIsBrokenSequence = true;
+			if (bIsReliable)
+			{
+				Context.SetError(NetError_PartialNetBlobSequenceError);
+			}
+			return;
+		}
+	}
+
+	bIsProcessingReliable = bIsReliable;
+
+	if (bIsFirstPart)
 	{
-		PartCount = PartialNetBlob->GetPartCount();
+		const uint32 PartCount = PartialNetBlob->GetPartCount();
 		if (PartCount == 0)
 		{
+			bIsBrokenSequence = true;
 			Context.SetError(GNetError_InvalidValue);
 			return;
 		}
 
 		if (PartCount > PartialNetBlobHandlerConfig->GetMaxPartCount())
 		{
+			bIsBrokenSequence = true;
 			Context.SetError(GNetError_InvalidValue);
 			return;
 		}
@@ -57,6 +93,7 @@ void FNetBlobAssembler::AddPartialNetBlob(FNetSerializationContext& Context, FNe
 		// We allow part sizes to be lower than the config value in case it was hotfixed.
 		if (PayloadBitCount > PartialNetBlobHandlerConfig->GetMaxPartBitCount())
 		{
+			bIsBrokenSequence = true;
 			Context.SetError(GNetError_InvalidValue);
 			return;
 		}
@@ -80,6 +117,7 @@ void FNetBlobAssembler::AddPartialNetBlob(FNetSerializationContext& Context, FNe
 	{
 		if (InRefHandle != RefHandle)
 		{
+			bIsBrokenSequence = true;
 			Context.SetError(NetError_PartialNetBlobSequenceError);
 			return;
 		}
@@ -87,8 +125,9 @@ void FNetBlobAssembler::AddPartialNetBlob(FNetSerializationContext& Context, FNe
 		const uint32 PayloadBitCount = PartialNetBlob->GetPayloadBitCount();
 		const uint32 PayloadByteCount = (PayloadBitCount + 7U)/8U;
 		// All parts except the last one is expected to match the first part. The last part may be smaller.
-		if (((PayloadBitCount != FirstPayloadBitCount) && ((PartIndex + 1U) < PartCount)) || (PayloadBitCount > FirstPayloadBitCount))
+		if (((PayloadBitCount != FirstPayloadBitCount) && (NextSequenceSumber != LastPartSequenceNumber)) || (PayloadBitCount > FirstPayloadBitCount))
 		{
+			bIsBrokenSequence = true;
 			Context.SetError(GNetError_InvalidValue);
 			return;
 		}
@@ -97,11 +136,13 @@ void FNetBlobAssembler::AddPartialNetBlob(FNetSerializationContext& Context, FNe
 		BitWriter.WriteBitStream(PartialNetBlob->GetPayload(), 0, PayloadBitCount);
 	}
 
-	++NextPartIndex;
-	if (NextPartIndex == PartCount)
+	if (NextSequenceSumber == LastPartSequenceNumber)
 	{
 		bIsReadyToAssemble = true;
-		NextPartIndex = 0;
+	}
+	else
+	{
+		++NextSequenceSumber;
 	}
 }
 
diff --git a/Engine/Source/Runtime/Experimental/Iris/Core/Private/Iris/ReplicationSystem/NetBlob/NetRPC.cpp b/Engine/Source/Runtime/Experimental/Iris/Core/Private/Iris/ReplicationSystem/NetBlob/NetRPC.cpp
--- a/Engine/Source/Runtime/Experimental/Iris/Core/Private/Iris/ReplicationSystem/NetBlob/NetRPC.cpp
+++ b/Engine/Source/Runtime/Experimental/Iris/Core/Private/Iris/ReplicationSystem/NetBlob/NetRPC.cpp
@@ -353,6 +353,12 @@ bool FNetRPC::ResolveFunctionAndObject(FNetSerializationContext& Context)
 		CreationInfo.Flags |= ENetBlobFlags::Reliable;
 	}
 
+	// The sending side will set Ordered on unicast/reliable RPCs so we're restoring that flag. Unicast RPCs are ordered with respect to other reliable and unicast RPCs whereas multicast RPCs are not.
+	if ((Function->FunctionFlags & FUNC_NetMulticast) == 0)
+	{
+		CreationInfo.Flags |= UE::Net::ENetBlobFlags::Ordered;
+	}
+
 	// Set the BlobDescriptor even if it has zero size so that we can trace with a meaningful name.
 	BlobDescriptor = FunctionDescriptor->Descriptor;
 
diff --git a/Engine/Source/Runtime/Experimental/Iris/Core/Private/Iris/ReplicationSystem/NetBlob/PartialNetBlob.cpp b/Engine/Source/Runtime/Experimental/Iris/Core/Private/Iris/ReplicationSystem/NetBlob/PartialNetBlob.cpp
--- a/Engine/Source/Runtime/Experimental/Iris/Core/Private/Iris/ReplicationSystem/NetBlob/PartialNetBlob.cpp
+++ b/Engine/Source/Runtime/Experimental/Iris/Core/Private/Iris/ReplicationSystem/NetBlob/PartialNetBlob.cpp
@@ -7,6 +7,18 @@
 #include "Iris/Serialization/InternalNetSerializationContext.h"
 #include "Iris/Serialization/NetSerializationContext.h"
 #include "Net/Core/Trace/NetTrace.h"
+#include <atomic>
+
+namespace UE::Net::Private
+{
+	/*
+	 * We need a unique sequence number per split blob to detect out of order sequences. In theory this requirement is only needed for split blobs sent via the same queue or whatever mechanism is used.
+	 * As the overhead of large sequence numbers should be relatively low compared to the split payload size and sharing a global sequence number avoids bloating the splitting API we have this one shared atomic.
+	 * Each split NetBlob will reserve unique sequence numbers for all of its parts. If there are multiple ReplicationSystems sharing the atomic adds relatively little growing of the sequence numbers.
+	 * Splitting blobs is a special case mainly used to enable replicating very large objects.
+	 */
+	static std::atomic<uint32> PartialNetBlobGlobalSequenceNumber;
+}
 
 namespace UE::Net
 {
@@ -14,9 +26,6 @@ namespace UE::Net
 FPartialNetBlob::FPartialNetBlob(const FNetBlobCreationInfo& CreationInfo)
 : FNetBlob(CreationInfo)
 , OriginalCreationInfo({})
-, PartIndex(0)
-, PartCount(0)
-, PayloadBitCount(0)
 {
 }
 
@@ -63,25 +72,25 @@ void FPartialNetBlob::InternalSerialize(FNetSerializationContext& Context) const
 		UE_NET_TRACE_EXIT_NAMED_SCOPE(UserProvidedScope);
 	}
 	UE_NET_TRACE_SCOPE(PartialNetBlob, *Writer, Context.GetTraceCollector(), ENetTraceVerbosity::Trace);
-	UE_NET_TRACE_NAMED_DYNAMIC_NAME_SCOPE(PartIndexScope, static_cast<const TCHAR*>(nullptr), *Writer, Context.GetTraceCollector(), ENetTraceVerbosity::VeryVerbose);
+	UE_NET_TRACE_NAMED_DYNAMIC_NAME_SCOPE(SequenceScope, static_cast<const TCHAR*>(nullptr), *Writer, Context.GetTraceCollector(), ENetTraceVerbosity::VeryVerbose);
 #if UE_NET_TRACE_ENABLED
 	if (FNetTrace::GetNetTraceVerbosityEnabled(ENetTraceVerbosity::VeryVerbose))
 	{
-		TStringBuilder<32> Builder;
+		TStringBuilder<64> Builder;
 		if (IsFirstPart())
 		{
-			Builder.Appendf(TEXT("PartIndex %u/%u"), PartIndex, PartCount);
+			Builder.Appendf(TEXT("Seq %u First part of %u"), SequenceNumber, PartCount);
 		}
 		else
 		{
-			Builder.Appendf(TEXT("PartIndex %u"), PartIndex);
+			Builder.Appendf(TEXT("Seq %u"), SequenceNumber);
 		}
-		UE_NET_TRACE_SET_SCOPE_NAME(PartIndexScope, Builder.ToString());
+		UE_NET_TRACE_SET_SCOPE_NAME(SequenceScope, Builder.ToString());
 	}
 #endif // UE_NET_TRACE_ENABLED
 
-	WritePackedUint16(Writer, PartIndex);
-	if (IsFirstPart())
+	WritePackedUint32(Writer, SequenceNumber);
+	if (Writer->WriteBool(IsFirstPart()))
 	{
 		WritePackedUint16(Writer, PartCount - 1U);
 
@@ -96,11 +105,13 @@ void FPartialNetBlob::InternalDeserialize(FNetSerializationContext& Context)
 {
 	FNetBitStreamReader* Reader = Context.GetBitStreamReader();
 	UE_NET_TRACE_SCOPE(PartialNetBlob, *Reader, Context.GetTraceCollector(), ENetTraceVerbosity::Trace);
-	UE_NET_TRACE_NAMED_DYNAMIC_NAME_SCOPE(PartIndexScope, static_cast<const TCHAR*>(nullptr), *Reader, Context.GetTraceCollector(), ENetTraceVerbosity::VeryVerbose);
+	UE_NET_TRACE_NAMED_DYNAMIC_NAME_SCOPE(SequenceScope, static_cast<const TCHAR*>(nullptr), *Reader, Context.GetTraceCollector(), ENetTraceVerbosity::VeryVerbose);
 
-	PartIndex = ReadPackedUint16(Reader);
-	if (IsFirstPart())
+	SequenceNumber = ReadPackedUint32(Reader);
+	SequenceFlags = ESequenceFlags::None;
+	if (Reader->ReadBool())
 	{
+		SequenceFlags = ESequenceFlags::IsFirstPart;
 		PartCount = ReadPackedUint16(Reader) + 1U;
 
 		UE_NET_TRACE_SCOPE(OriginalCreationInfo, *Reader, Context.GetTraceCollector(), ENetTraceVerbosity::Trace);
@@ -112,16 +123,16 @@ void FPartialNetBlob::InternalDeserialize(FNetSerializationContext& Context)
 #if UE_NET_TRACE_ENABLED
 	if (FNetTrace::GetNetTraceVerbosityEnabled(ENetTraceVerbosity::VeryVerbose))
 	{
-		TStringBuilder<32> Builder;
+		TStringBuilder<64> Builder;
 		if (IsFirstPart())
 		{
-			Builder.Appendf(TEXT("PartIndex %u/%u"), PartIndex, PartCount);
+			Builder.Appendf(TEXT("Seq %u First part of %u"), SequenceNumber, PartCount);
 		}
 		else
 		{
-			Builder.Appendf(TEXT("PartIndex %u"), PartIndex);
+			Builder.Appendf(TEXT("Seq %u"), SequenceNumber);
 		}
-		UE_NET_TRACE_SET_SCOPE_NAME(PartIndexScope, Builder.ToString());
+		UE_NET_TRACE_SET_SCOPE_NAME(SequenceScope, Builder.ToString());
 	}
 #endif // UE_NET_TRACE_ENABLED
 }
@@ -253,6 +264,9 @@ void FPartialNetBlob::SplitPayload(const FPartialNetBlob::FPayloadSplitParams& S
 {
 	const uint32 PartialBlobCount = (SplitParams.PayloadBitCount + SplitParams.PartBitCount - 1U)/SplitParams.PartBitCount;
 	OutPartialBlobs.Reserve(OutPartialBlobs.Num() + int32(PartialBlobCount));
+	
+	// Reserve sequence numbers for all parts. 
+	uint32 SequenceNumber = Private::PartialNetBlobGlobalSequenceNumber.fetch_add(PartialBlobCount, std::memory_order_relaxed);
 
 	uint32 PayloadBitOffset = 0U;
 	for (uint32 PartIt = 0, PartEndIt = PartialBlobCount; PartIt != PartEndIt; ++PartIt)
@@ -262,8 +276,9 @@ void FPartialNetBlob::SplitPayload(const FPartialNetBlob::FPayloadSplitParams& S
 		FPartialNetBlob* PartialBlob = new FPartialNetBlob(SplitParams.CreationInfo);
 		PartialBlob->SetDebugName(SplitParams.DebugName);
 		PartialBlob->OriginalCreationInfo = SplitParams.OriginalCreationInfo;
-		PartialBlob->PartIndex = static_cast<uint16>(PartIt);
 		PartialBlob->PartCount = static_cast<uint16>(PartialBlobCount);
+		PartialBlob->SequenceFlags = (bIsFirstPart ? ESequenceFlags::IsFirstPart : ESequenceFlags::None);
+		PartialBlob->SequenceNumber = SequenceNumber++;
 		if (bIsFirstPart && EnumHasAnyFlags(SplitParams.OriginalCreationInfo.Flags, ENetBlobFlags::HasExports))
 		{
 			PartialBlob->OriginalBlob = SplitParams.OriginalBlob;
diff --git a/Engine/Source/Runtime/Experimental/Iris/Core/Private/Iris/ReplicationSystem/NetBlob/ReliableNetBlobQueue.cpp b/Engine/Source/Runtime/Experimental/Iris/Core/Private/Iris/ReplicationSystem/NetBlob/ReliableNetBlobQueue.cpp
--- a/Engine/Source/Runtime/Experimental/Iris/Core/Private/Iris/ReplicationSystem/NetBlob/ReliableNetBlobQueue.cpp
+++ b/Engine/Source/Runtime/Experimental/Iris/Core/Private/Iris/ReplicationSystem/NetBlob/ReliableNetBlobQueue.cpp
@@ -1,6 +1,7 @@
 // Copyright Epic Games, Inc. All Rights Reserved.
 
 #include "Iris/ReplicationSystem/NetBlob/ReliableNetBlobQueue.h"
+#include "Iris/IrisConfigInternal.h"
 #include "Iris/ReplicationSystem/NetBlob/NetBlobHandlerManager.h"
 #include "Iris/ReplicationSystem/ObjectReferenceCache.h"
 #include "Iris/Core/IrisLog.h"
@@ -10,6 +11,7 @@
 #include "Iris/Serialization/NetSerializationContext.h"
 #include "Iris/Serialization/InternalNetSerializationContext.h"
 #include "Iris/Serialization/NetExportContext.h"
+#include "Net/Core/Trace/NetTrace.h"
 
 namespace UE::Net::Private
 {
@@ -17,6 +19,8 @@ namespace UE::Net::Private
 static const FName NetError_ReliableQueueFull("Reliable attachment queue full");
 static const FName NetError_InvalidSequence("Invalid sequence number");
 
+static FNetDebugName const* ReliabilityNetDebugNames[2];
+
 // ReliableNetBlobQueue
 FReliableNetBlobQueue::FReliableNetBlobQueue()
 : Sent{}
@@ -25,6 +29,10 @@ FReliableNetBlobQueue::FReliableNetBlobQueue()
 , LastSeq(0)
 , UnsentBlobCount(0)
 {
+#if UE_NET_TRACE_ENABLED
+	Private::ReliabilityNetDebugNames[0] = CreatePersistentNetDebugName(TEXT("Unreliable"));
+	Private::ReliabilityNetDebugNames[1] = CreatePersistentNetDebugName(TEXT("Reliable"));
+#endif
 }
 
 FReliableNetBlobQueue::~FReliableNetBlobQueue()
@@ -108,6 +116,13 @@ uint32 FReliableNetBlobQueue::SerializeInternal(FNetSerializationContext& Contex
 		// GetRefCount() is not const.
 		TRefCountPtr<FNetBlob>& Attachment = NetBlobs[Index];
 
+#if UE_NET_USE_READER_WRITER_SENTINEL
+		{
+			UE_NET_TRACE_SCOPE(Sentintel, *Writer, Context.GetTraceCollector(), ENetTraceVerbosity::Verbose);
+			Writer->WriteBits(0x5e4714e1, 32);
+		}
+#endif
+
 		// If this sequence is disjoint from the previous sequence we need to serialize the full index.
 		// It's important that the sequence number is sent first so we can validate it before receiving exports and payload.
 		if (Writer->WriteBool(Seq != PrevWrittenSeq + 1U))
@@ -119,6 +134,8 @@ uint32 FReliableNetBlobQueue::SerializeInternal(FNetSerializationContext& Contex
 		const bool bHasData = Attachment.GetRefCount() > 0;
 		if (Writer->WriteBool(bHasData))
 		{
+			UE_NET_TRACE_NAMED_DYNAMIC_NAME_SCOPE(ReliabilityScope, Private::ReliabilityNetDebugNames[Attachment->IsReliable()], *Writer, Context.GetTraceCollector(), ENetTraceVerbosity::Verbose);
+
 			// If we have exports, append them, if attachment is rolled back we will roll back any appended exports as well.
 			ObjectReferenceCache->AddPendingExports(Context, Attachment->CallGetExports());
 
@@ -194,7 +211,7 @@ void FReliableNetBlobQueue::CommitReplicationRecord(const FReliableNetBlobQueue:
 			TRefCountPtr<FNetBlob>& RefCountBlob = NetBlobs[BlobIndex];
 			if (const FNetBlob* Blob = RefCountBlob.GetReference())
 			{
-				if (!EnumHasAnyFlags(Blob->GetCreationInfo().Flags, ENetBlobFlags::Reliable))
+				if (!Blob->IsReliable())
 				{
 					RefCountBlob.SafeRelease();
 				}
@@ -227,6 +244,20 @@ uint32 FReliableNetBlobQueue::DeserializeInternal(FNetSerializationContext& Cont
 	bool bHasMoreBlobs = false;
 	do
 	{
+
+#if UE_NET_USE_READER_WRITER_SENTINEL
+		{
+			UE_NET_TRACE_SCOPE(Sentintel, *Reader, Context.GetTraceCollector(), ENetTraceVerbosity::Verbose);
+			const uint32 Sentinel = Reader->ReadBits(32);
+			if (Sentinel != 0x5e4714e1U)
+			{
+				UE_LOG(LogIris, Error, TEXT("Wrong sentinel %u != %u"), Sentinel, 0x5e4714e1U);
+				Context.SetError(GNetError_BitStreamError);
+				return 0;
+			}
+		}
+#endif
+
 		if (const bool bIsNewSequence = Reader->ReadBool())
 		{
 			Index = Reader->ReadBits(IndexBitCount);
@@ -253,6 +284,8 @@ uint32 FReliableNetBlobQueue::DeserializeInternal(FNetSerializationContext& Cont
 		TRefCountPtr<FNetBlob> Blob;
 		if (const bool bHasData = Reader->ReadBool())
 		{
+			UE_NET_TRACE_NAMED_DYNAMIC_NAME_SCOPE(ReliabilityScope, Private::ReliabilityNetDebugNames[0], *Reader, Context.GetTraceCollector(), ENetTraceVerbosity::Verbose);
+
 			FNetBlobCreationInfo CreationInfo;
 			FNetBlob::DeserializeCreationInfo(Context, CreationInfo);
 			Blob = BlobReceiver->CreateNetBlob(CreationInfo);
@@ -272,6 +305,8 @@ uint32 FReliableNetBlobQueue::DeserializeInternal(FNetSerializationContext& Cont
 			{
 				Blob->Deserialize(Context);
 			}
+
+			UE_NET_TRACE_SET_SCOPE_NAME(ReliabilityScope, Private::ReliabilityNetDebugNames[Blob->IsReliable()]);
 		}
 		
 		bHasMoreBlobs = Reader->ReadBool();
@@ -356,7 +391,7 @@ void FReliableNetBlobQueue::DequeueUnreliable(TArray<TRefCountPtr<FNetBlob>>& Un
 		TRefCountPtr<FNetBlob>& RefCntBlob = NetBlobs[Index];
 		if (RefCntBlob.GetRefCount() > 0)
 		{
-			if (const bool bIsUnReliable = !EnumHasAnyFlags(RefCntBlob.GetReference()->GetCreationInfo().Flags, ENetBlobFlags::Reliable))
+			if (const bool bIsUnReliable = !RefCntBlob->IsReliable())
 			{
 				Unreliable.Emplace(MoveTemp(RefCntBlob));
 
diff --git a/Engine/Source/Runtime/Experimental/Iris/Core/Private/Iris/ReplicationSystem/NetBlob/SequentialPartialNetBlobHandler.cpp b/Engine/Source/Runtime/Experimental/Iris/Core/Private/Iris/ReplicationSystem/NetBlob/SequentialPartialNetBlobHandler.cpp
--- a/Engine/Source/Runtime/Experimental/Iris/Core/Private/Iris/ReplicationSystem/NetBlob/SequentialPartialNetBlobHandler.cpp
+++ b/Engine/Source/Runtime/Experimental/Iris/Core/Private/Iris/ReplicationSystem/NetBlob/SequentialPartialNetBlobHandler.cpp
@@ -30,7 +30,7 @@ bool USequentialPartialNetBlobHandler::SplitNetBlob(const TRefCountPtr<FNetBlob>
 
 	FNetBlobCreationInfo CreationInfo = {};
 	CreationInfo.Type = GetNetBlobType();
-	CreationInfo.Flags = ENetBlobFlags::Reliable;
+	CreationInfo.Flags = Blob->GetCreationInfo().Flags & (ENetBlobFlags::Ordered | ENetBlobFlags::Reliable);
 
 	FPartialNetBlob::FSplitParams SplitParams = {};
 	SplitParams.MaxPartBitCount = Config->GetMaxPartBitCount();
@@ -61,7 +61,7 @@ bool USequentialPartialNetBlobHandler::SplitNetBlob(const UE::Net::FNetObjectRef
 
 	FNetBlobCreationInfo CreationInfo = {};
 	CreationInfo.Type = GetNetBlobType();
-	CreationInfo.Flags = ENetBlobFlags::Reliable;
+	CreationInfo.Flags = Blob->GetCreationInfo().Flags & (ENetBlobFlags::Ordered | ENetBlobFlags::Reliable);
 
 	FPartialNetBlob::FSplitParams SplitParams = {};
 	SplitParams.MaxPartBitCount = Config->GetMaxPartBitCount();
diff --git a/Engine/Source/Runtime/Experimental/Iris/Core/Private/Iris/ReplicationSystem/ReplicationReader.cpp b/Engine/Source/Runtime/Experimental/Iris/Core/Private/Iris/ReplicationSystem/ReplicationReader.cpp
--- a/Engine/Source/Runtime/Experimental/Iris/Core/Private/Iris/ReplicationSystem/ReplicationReader.cpp
+++ b/Engine/Source/Runtime/Experimental/Iris/Core/Private/Iris/ReplicationSystem/ReplicationReader.cpp
@@ -2231,21 +2231,6 @@ void FReplicationReader::ResolveAndDispatchAttachments(FNetSerializationContext&
 					return;
 				}
 			}
-
-			// Dispatch remaining ordered unreliable attachments. The side effects of for example blocking move RPCs are significant.
-			{
-				TArray<TRefCountPtr<FNetBlob>> OrderedUnreliableAttachments;
-				OrderedUnreliableAttachments.Reserve(16);
-				AttachmentQueue->GetOrderedUnreliable(OrderedUnreliableAttachments);
-				for (TRefCountPtr<FNetBlob>& Attachment : OrderedUnreliableAttachments)
-				{
-					NetBlobHandlerManager->OnNetBlobReceived(Context, Attachment);
-					if (Context.HasError())
-					{
-						return;
-					}
-				}
-			}
 		}
 
 		if (bDispatchUnreliableAttachments)
diff --git a/Engine/Source/Runtime/Experimental/Iris/Core/Public/Iris/ReplicationSystem/NetBlob/NetBlob.h b/Engine/Source/Runtime/Experimental/Iris/Core/Public/Iris/ReplicationSystem/NetBlob/NetBlob.h
--- a/Engine/Source/Runtime/Experimental/Iris/Core/Public/Iris/ReplicationSystem/NetBlob/NetBlob.h
+++ b/Engine/Source/Runtime/Experimental/Iris/Core/Public/Iris/ReplicationSystem/NetBlob/NetBlob.h
@@ -70,6 +70,9 @@ public:
 	/** Returns the FNetBlobCreationInfo. */
 	const FNetBlobCreationInfo& GetCreationInfo() const { return CreationInfo; }
 
+	/* Whether the blob is reliably sent or not. */
+	bool IsReliable() const { return EnumHasAnyFlags(GetCreationInfo().Flags, ENetBlobFlags::Reliable); }
+
 	/** Returns the FReplicationStateDescriptor if there is one. It's recommended to use a descriptor instead of overriding the serialization methods. */
 	const FReplicationStateDescriptor* GetReplicationStateDescriptor() const { return BlobDescriptor.GetReference(); }
 
diff --git a/Engine/Source/Runtime/Experimental/Iris/Core/Public/Iris/ReplicationSystem/NetBlob/NetBlobAssembler.h b/Engine/Source/Runtime/Experimental/Iris/Core/Public/Iris/ReplicationSystem/NetBlob/NetBlobAssembler.h
--- a/Engine/Source/Runtime/Experimental/Iris/Core/Public/Iris/ReplicationSystem/NetBlob/NetBlobAssembler.h
+++ b/Engine/Source/Runtime/Experimental/Iris/Core/Public/Iris/ReplicationSystem/NetBlob/NetBlobAssembler.h
@@ -36,6 +36,9 @@ public:
 	/** Returns true if all parts of the split blob have been added and is ready to be assembled. */
 	bool IsReadyToAssemble() const { return bIsReadyToAssemble; }
 
+	/** Returns true if the sequence order is broken. */
+	bool IsSequenceBroken() const { return bIsBrokenSequence; }
+
 	/**
 	 * Assemble all parts of the split blob.
 	 * @param Context FNetSerializationContext for error reporting. Call HasError() on it afterwards to check whether something went wrong.
@@ -50,11 +53,16 @@ private:
 	TArray<uint32> Payload;
 	FNetBitStreamWriter BitWriter;
 	FNetRefHandle RefHandle;
+	const USequentialPartialNetBlobHandlerConfig* PartialNetBlobHandlerConfig = nullptr;
 	uint32 NextPartIndex = 0;
-	uint32 PartCount = 0;
+	// Initialized to an arbitrary number. If we receive the first part of a sequence it doesn't matter. If the next received blob isn't the first part it's a broken sequence.
+	uint32 NextSequenceSumber = 0U;
+	// The sequence number of the last part expected to be received before being ready to assemble.
+	uint32 LastPartSequenceNumber = ~0U;
 	uint32 FirstPayloadBitCount = 0;
 	bool bIsReadyToAssemble = false;
-	const USequentialPartialNetBlobHandlerConfig* PartialNetBlobHandlerConfig = nullptr;
+	bool bIsBrokenSequence = false;
+	bool bIsProcessingReliable = false;
 };
 
 }
diff --git a/Engine/Source/Runtime/Experimental/Iris/Core/Public/Iris/ReplicationSystem/NetBlob/PartialNetBlob.h b/Engine/Source/Runtime/Experimental/Iris/Core/Public/Iris/ReplicationSystem/NetBlob/PartialNetBlob.h
--- a/Engine/Source/Runtime/Experimental/Iris/Core/Public/Iris/ReplicationSystem/NetBlob/PartialNetBlob.h
+++ b/Engine/Source/Runtime/Experimental/Iris/Core/Public/Iris/ReplicationSystem/NetBlob/PartialNetBlob.h
@@ -28,12 +28,21 @@ public:
 	IRISCORE_API static bool SplitNetBlob(const FNetBlobCreationInfo& CreationInfo, const FSplitParams& SplitParams, const TRefCountPtr<FRawDataNetBlob>& Blob, TArray<TRefCountPtr<FNetBlob>>& OutPartialBlobs);
 
 public:
-	IRISCORE_API FPartialNetBlob(const FNetBlobCreationInfo& CreationInfo);
+	enum class ESequenceFlags : uint32
+	{
+		None = 0,
+		IsFirstPart = 1U,
+	};
+	FRIEND_ENUM_CLASS_FLAGS(ESequenceFlags);
 
-	uint32 GetPartIndex() const { return PartIndex; }
+	IRISCORE_API FPartialNetBlob(const FNetBlobCreationInfo& CreationInfo);
 
 	// This will only return legible data for the first part on the receiving end.
-	uint32 GetPartCount() const { return (PartIndex == 0U ? PartCount : 0U); }
+	uint32 GetPartCount() const { return (IsFirstPart() ? PartCount : 0U); }
+
+	bool IsFirstPart() const { return EnumHasAnyFlags(SequenceFlags, ESequenceFlags::IsFirstPart); }
+
+	uint32 GetSequenceNumber() const { return SequenceNumber; }
 
 	uint32 GetPayloadBitCount() const { return PayloadBitCount; }
 	const uint32* GetPayload() const { return Payload.GetData(); }
@@ -44,9 +53,6 @@ public:
 
 private:
 
-	bool IsFirstPart() const { return PartIndex == 0; }
-	bool IsLastPart() const { return PartIndex+1 == PartCount; }
-
 	virtual TArrayView<const FNetObjectReference> GetExports() const override final;
 	virtual void SerializeWithObject(FNetSerializationContext& Context, FNetRefHandle RefHandle) const override;
 	virtual void DeserializeWithObject(FNetSerializationContext& Context, FNetRefHandle RefHandle) override;
@@ -77,10 +83,11 @@ private:
 	FNetBlobCreationInfo OriginalCreationInfo;
 	TRefCountPtr<FNetBlob> OriginalBlob;
 
-	uint16 PartIndex;
+	ESequenceFlags SequenceFlags = ESequenceFlags::None;
+	uint32 SequenceNumber = 0;
 	// PartCount is only valid if it's the first part.
-	uint16 PartCount;
-	uint16 PayloadBitCount;
+	uint16 PartCount = 0;
+	uint16 PayloadBitCount = 0;
 	// Use uint32 for guaranteed FNetBitStreamReader/Writer compatibility
 	TArray<uint32> Payload;
 
