Iris
* Added Net.Iris.PrintAlwaysRelevantObjects :  Print the list of root objects considered always relevant
* Added Net.Iris.PrintRelevantObjectsToConnection : Print the list of root objects relevant to a given connection. Default to print all connections. Use ConnectionId=X,Y,Z to print specific connections.

#rnx

diff --git a/Engine/Source/Runtime/Experimental/Iris/Core/Private/Iris/ReplicationSystem/Filtering/ReplicationFiltering.cpp b/Engine/Source/Runtime/Experimental/Iris/Core/Private/Iris/ReplicationSystem/Filtering/ReplicationFiltering.cpp
--- a/Engine/Source/Runtime/Experimental/Iris/Core/Private/Iris/ReplicationSystem/Filtering/ReplicationFiltering.cpp
+++ b/Engine/Source/Runtime/Experimental/Iris/Core/Private/Iris/ReplicationSystem/Filtering/ReplicationFiltering.cpp
@@ -340,23 +340,11 @@ void FReplicationFiltering::FilterNonRelevantObjects()
 		return;
 	}
 
-	const uint32* const CurrentFrameScopeData = NetRefHandleManager->GetCurrentFrameScopableInternalIndices().GetData();
-	const uint32* const WithOwnerData = ObjectsWithOwnerFilter.GetData();
-	const uint32* const ConnectionFiltersData = AllConnectionFilteredObjects.GetData();
-	const uint32* const DynamicFilteredData = DynamicFilterEnabledObjects.GetData();
-	const uint32* const GroupFilteredOutData = Groups->GetGroupFilteredOutObjects().GetData();
-
+	// Start by filling the relevant object list with those considered AlwaysRelevant.
 	FNetBitArrayView GlobalRelevantObjects = NetRefHandleManager->GetRelevantObjectsInternalIndices();
-	
-	uint32* GlobalRelevantData = GlobalRelevantObjects.GetData();
+	BuildAlwaysRelevantList(GlobalRelevantObjects, NetRefHandleManager->GetCurrentFrameScopableInternalIndices());
 
-	//$IRIS TODO: Could be worth sampling if it's faster for the AlwaysRelevant list to be kept cached instead of recalculating it here every frame.
-	const uint32 MaxWords = GlobalRelevantObjects.GetNumWords();
-	for (uint32 WordIndex = 0; WordIndex < MaxWords; ++WordIndex)
-	{
-		// Build the list of always relevant objects (objects that have no filters)
-		GlobalRelevantData[WordIndex] = CurrentFrameScopeData[WordIndex] & ~(WithOwnerData[WordIndex] | ConnectionFiltersData[WordIndex] | DynamicFilteredData[WordIndex] | GroupFilteredOutData[WordIndex]);
-	}
+	//$IRIS TODO: Should sample if it's faster for the AlwaysRelevant list to be kept cached instead of recalculating it here every frame.
 
 	// Build the list of currently relevant objects. e.g. always relevant objects + filterable objects relevant to at least one connection.	
 	auto MergeConnectionScopes = [this, &GlobalRelevantObjects](uint32 ConnectionId)
@@ -370,6 +358,27 @@ void FReplicationFiltering::FilterNonRelevantObjects()
 	//$IRIS TODO: Need to ensure newly relevant objects are immediately polled similar to calling ForceNetUpdate.
 }
 
+void FReplicationFiltering::BuildAlwaysRelevantList(FNetBitArrayView OutAlwaysRelevantList, const FNetBitArrayView ScopeList) const
+{
+	// The list of all replicated objects
+	const uint32* const ScopeListData = ScopeList.GetData();
+
+	// The different list of filtered objects
+	const uint32* const WithOwnerData = ObjectsWithOwnerFilter.GetData();
+	const uint32* const ConnectionFiltersData = AllConnectionFilteredObjects.GetData();
+	const uint32* const DynamicFilteredData = DynamicFilterEnabledObjects.GetData();
+	const uint32* const GroupFilteredOutData = Groups->GetGroupFilteredOutObjects().GetData();
+
+	uint32* OutAlwaysRelevantListData = OutAlwaysRelevantList.GetData();
+
+	const uint32 MaxWords = OutAlwaysRelevantList.GetNumWords();
+	for (uint32 WordIndex = 0; WordIndex < MaxWords; ++WordIndex)
+	{
+		// Build the list of always relevant objects, e.g. objects that have no filters
+		OutAlwaysRelevantListData[WordIndex] = ScopeListData[WordIndex] & ~(WithOwnerData[WordIndex] | ConnectionFiltersData[WordIndex] | DynamicFilteredData[WordIndex] | GroupFilteredOutData[WordIndex]);
+	}
+}
+
 void FReplicationFiltering::UpdateDynamicFilters(ENetFilterType FilterPass)
 {
 	/**
diff --git a/Engine/Source/Runtime/Experimental/Iris/Core/Private/Iris/ReplicationSystem/Filtering/ReplicationFiltering.h b/Engine/Source/Runtime/Experimental/Iris/Core/Private/Iris/ReplicationSystem/Filtering/ReplicationFiltering.h
--- a/Engine/Source/Runtime/Experimental/Iris/Core/Private/Iris/ReplicationSystem/Filtering/ReplicationFiltering.h
+++ b/Engine/Source/Runtime/Experimental/Iris/Core/Private/Iris/ReplicationSystem/Filtering/ReplicationFiltering.h
@@ -97,6 +97,9 @@ public:
 	/** Returns the name of the Filter represented by the handle. */
 	FName GetFilterName(FNetObjectFilterHandle Filter) const;
 
+	/** Fill the passed in list with root objects considered always relevant */
+	void BuildAlwaysRelevantList(FNetBitArrayView OutAlwaysRelevantList, const FNetBitArrayView ScopeList) const;
+
 	// Connection handling
 	void AddConnection(uint32 ConnectionId);
 	void RemoveConnection(uint32 ConnectionId);
diff --git a/Engine/Source/Runtime/Experimental/Iris/Core/Private/Iris/ReplicationSystem/ObjectReplicationBridge.cpp b/Engine/Source/Runtime/Experimental/Iris/Core/Private/Iris/ReplicationSystem/ObjectReplicationBridge.cpp
--- a/Engine/Source/Runtime/Experimental/Iris/Core/Private/Iris/ReplicationSystem/ObjectReplicationBridge.cpp
+++ b/Engine/Source/Runtime/Experimental/Iris/Core/Private/Iris/ReplicationSystem/ObjectReplicationBridge.cpp
@@ -2005,3 +2005,36 @@ void UObjectReplicationBridge::OnErrorWithNetRefHandleReported(uint32 ErrorType,
 		*NetRefHandleManager->PrintObjectFromIndex(ObjectInternalIndex)
 	);
 }
+
+TArray<uint32> UObjectReplicationBridge::FindConnectionsFromArgs(const TArray<FString>& Args) const
+{
+	using namespace UE::Net::Private;
+
+	TArray<uint32> ConnectionList;
+
+	// If ConnectionId=XX was specified
+	if (const FString* ArgConnectionIds = Args.FindByPredicate([](const FString& Str) { return Str.Contains(TEXT("ConnectionId=")); }))
+	{
+		constexpr bool bIgnoreSeperators = false;
+
+		// Find all the Ids passed in the argument
+		FString StrConnectionIds;
+		if (FParse::Value(**ArgConnectionIds, TEXT("ConnectionId="), StrConnectionIds, bIgnoreSeperators))
+		{
+			TArray<FString> StrIds;
+			StrConnectionIds.ParseIntoArray(StrIds, TEXT(","));
+
+			for (const FString& StrId : StrIds)
+			{
+				int32 Id = INDEX_NONE;
+				LexFromString(Id, *StrId);
+				if (Id > INDEX_NONE)
+				{
+					ConnectionList.AddUnique((uint32)Id);
+				}
+			}
+		}
+	}
+
+	return ConnectionList;
+}
\ No newline at end of file
diff --git a/Engine/Source/Runtime/Experimental/Iris/Core/Private/Iris/ReplicationSystem/ObjectReplicationBridgeDebugging.cpp b/Engine/Source/Runtime/Experimental/Iris/Core/Private/Iris/ReplicationSystem/ObjectReplicationBridgeDebugging.cpp
--- a/Engine/Source/Runtime/Experimental/Iris/Core/Private/Iris/ReplicationSystem/ObjectReplicationBridgeDebugging.cpp
+++ b/Engine/Source/Runtime/Experimental/Iris/Core/Private/Iris/ReplicationSystem/ObjectReplicationBridgeDebugging.cpp
@@ -12,6 +12,8 @@
 #include "Iris/ReplicationSystem/ReplicationSystemInternal.h"
 #include "Iris/ReplicationSystem/ReplicationOperations.h"
 
+#include "Iris/ReplicationSystem/Filtering/ReplicationFiltering.h"
+
 #include "Net/Core/NetBitArrayPrinter.h"
 #include "Net/Core/Trace/NetDebugName.h"
 
@@ -60,7 +62,7 @@ void RemoteProtocolMismatchDetected(UReplicationSystem* ReplicationSystem, uint3
 	}
 }
 
-UReplicationSystem* FindReplicationSystem(const TArray<FString>& Args)
+UReplicationSystem* FindReplicationSystemFromArg(const TArray<FString>& Args)
 {
 	uint32 RepSystemId = 0;
 
@@ -122,11 +124,14 @@ void LogRootObjectListWithSubObjects(FNetRefHandleManager* NetRefHandleManager,
 // --------------------------------------------------------------------------------------------------------------------------------------------
 
 //-----------------------------------------------
-FAutoConsoleCommand ObjectBridgePrintDynamicFilter(TEXT("Net.Iris.PrintDynamicFilterClassConfig"), TEXT("Prints the dynamic filter configured to be assigned to specific classes."), FConsoleCommandWithArgsDelegate::CreateLambda([](const TArray< FString >& Args)
+FAutoConsoleCommand ObjectBridgePrintDynamicFilter(
+	TEXT("Net.Iris.PrintDynamicFilterClassConfig"), 
+	TEXT("Prints the dynamic filter configured to be assigned to specific classes."), 
+	FConsoleCommandWithArgsDelegate::CreateLambda([](const TArray< FString >& Args)
 {
 	using namespace UE::Net::Private::ObjectBridgeDebugging;
 
-	UReplicationSystem* RepSystem = FindReplicationSystem(Args);
+	UReplicationSystem* RepSystem = FindReplicationSystemFromArg(Args);
 	if (!RepSystem)
 	{
 		UE_LOG(LogIrisBridge, Error, TEXT("Could not find ReplicationSystem."));
@@ -171,7 +176,7 @@ FAutoConsoleCommand ObjectBridgePrintReplicatedObjects(
 {
 	using namespace UE::Net::Private::ObjectBridgeDebugging;
 
-	UReplicationSystem* RepSystem = FindReplicationSystem(Args);
+	UReplicationSystem* RepSystem = FindReplicationSystemFromArg(Args);
 	if (RepSystem)
 	{
 		if (UObjectReplicationBridge* ObjectBridge = CastChecked<UObjectReplicationBridge>(RepSystem->GetReplicationBridge()))
@@ -187,7 +192,7 @@ void UObjectReplicationBridge::PrintReplicatedObjects() const
 	using namespace UE::Net::Private;
 	using namespace UE::Net::Private::ObjectBridgeDebugging;
 
-	UE_LOG(LogIrisBridge, Display, TEXT("------------  Start Printing ALL Replicated Objects ------------"));
+	UE_LOG(LogIrisBridge, Display, TEXT("################ Start Printing ALL Replicated Objects ################"));
 	UE_LOG(LogIrisBridge, Display, TEXT(""));
 
 	uint32 TotalRootObjects = 0;
@@ -210,7 +215,7 @@ void UObjectReplicationBridge::PrintReplicatedObjects() const
 
 	UE_LOG(LogIrisBridge, Display, TEXT(""));
 	UE_LOG(LogIrisBridge, Display, TEXT("Printed %u root objects and %u sub objects"), Stats.Get<0>(), Stats.Get<1>());
-	UE_LOG(LogIrisBridge, Display, TEXT("------------  Stop Printing ALL Replicated Objects ------------"));
+	UE_LOG(LogIrisBridge, Display, TEXT("################ Stop Printing ALL Replicated Objects ################"));
 }
 
 //-----------------------------------------------
@@ -221,8 +226,7 @@ FAutoConsoleCommand ObjectBridgePrintRelevantObjects(
 {
 	using namespace UE::Net::Private::ObjectBridgeDebugging;
 
-	UReplicationSystem* RepSystem = FindReplicationSystem(Args);
-	if (RepSystem)
+	if (UReplicationSystem* RepSystem = FindReplicationSystemFromArg(Args))
 	{
 		if (UObjectReplicationBridge* ObjectBridge = CastChecked<UObjectReplicationBridge>(RepSystem->GetReplicationBridge()))
 		{
@@ -237,7 +241,7 @@ void UObjectReplicationBridge::PrintRelevantObjects() const
 	using namespace UE::Net::Private;
 	using namespace UE::Net::Private::ObjectBridgeDebugging;
 
-	UE_LOG(LogIrisBridge, Display, TEXT("------------  Start Printing Relevant Objects ------------"));
+	UE_LOG(LogIrisBridge, Display, TEXT("################ Start Printing Relevant Objects ################"));
 	UE_LOG(LogIrisBridge, Display, TEXT(""));
 
 	uint32 TotalRootObjects = 0;
@@ -253,5 +257,161 @@ void UObjectReplicationBridge::PrintRelevantObjects() const
 
 	UE_LOG(LogIrisBridge, Display, TEXT(""));
 	UE_LOG(LogIrisBridge, Display, TEXT("Printed %u root objects and %u sub objects"), Stats.Get<0>(), Stats.Get<1>());
-	UE_LOG(LogIrisBridge, Display, TEXT("------------  Stop Printing Relevant Objects ------------"));
+	UE_LOG(LogIrisBridge, Display, TEXT("################ Stop Printing Relevant Objects ################"));
+}
+
+//-----------------------------------------------
+FAutoConsoleCommand ObjectBridgePrintAlwaysRelevantObjects(
+	TEXT("Net.Iris.PrintAlwaysRelevantObjects"),
+	TEXT("Prints the list of netobjects always relevant to every connection"),
+	FConsoleCommandWithArgsDelegate::CreateLambda([](const TArray< FString >& Args)
+{
+	using namespace UE::Net::Private::ObjectBridgeDebugging;
+
+	if (UReplicationSystem* RepSystem = FindReplicationSystemFromArg(Args))
+	{
+		if (UObjectReplicationBridge* ObjectBridge = CastChecked<UObjectReplicationBridge>(RepSystem->GetReplicationBridge()))
+		{
+			ObjectBridge->PrintAlwaysRelevantObjects();
+		}
+	}
+}));
+
+void UObjectReplicationBridge::PrintAlwaysRelevantObjects() const
+{
+	using namespace UE::Net;
+	using namespace UE::Net::Private;
+	using namespace UE::Net::Private::ObjectBridgeDebugging;
+
+	FReplicationSystemInternal* ReplicationSystemInternal = GetReplicationSystem()->GetReplicationSystemInternal();
+
+	UE_LOG(LogIrisBridge, Display, TEXT("################ Start Printing Always Relevant Objects ################"));
+	UE_LOG(LogIrisBridge, Display, TEXT(""));
+
+	FNetBitArray AlwaysRelevantList;
+	AlwaysRelevantList.Init(NetRefHandleManager->GetMaxActiveObjectCount());
+	
+	ReplicationSystemInternal->GetFiltering().BuildAlwaysRelevantList(MakeNetBitArrayView(AlwaysRelevantList), ReplicationSystemInternal->GetNetRefHandleManager().GetGlobalScopableInternalIndices());
+
+	// Remove subobjects from the list.
+	MakeNetBitArrayView(AlwaysRelevantList).Combine(NetRefHandleManager->GetSubObjectInternalIndicesView(), FNetBitArrayView::AndNotOp);
+
+	TTuple<uint32, uint32> Stats(0, 0);
+	LogRootObjectListWithSubObjects(NetRefHandleManager, MakeNetBitArrayView(AlwaysRelevantList), Stats, nullptr);
+
+	UE_LOG(LogIrisBridge, Display, TEXT(""));
+	UE_LOG(LogIrisBridge, Display, TEXT("Printed %u root objects and %u sub objects"), Stats.Get<0>(), Stats.Get<1>());
+	UE_LOG(LogIrisBridge, Display, TEXT("################ Stop Printing Always Relevant Objects ################"));
+}
+
+//-----------------------------------------------
+FAutoConsoleCommand ObjectBridgePrintRelevantObjectsToConnection(
+TEXT("Net.Iris.PrintRelevantObjectsToConnection"),
+TEXT("Prints the list of replicated objects relevant to a specific connection"),
+FConsoleCommandWithArgsDelegate::CreateLambda([](const TArray< FString >& Args)
+{
+	using namespace UE::Net;
+	using namespace UE::Net::Private;
+	using namespace UE::Net::Private::ObjectBridgeDebugging;
+
+	UReplicationSystem* RepSystem = FindReplicationSystemFromArg(Args);
+	if (RepSystem)
+	{
+		if (UObjectReplicationBridge* ObjectBridge = CastChecked<UObjectReplicationBridge>(RepSystem->GetReplicationBridge()))
+		{
+			FReplicationSystemInternal* ReplicationSystemInternal = RepSystem->GetReplicationSystemInternal();
+
+			ObjectBridge->PrintRelevantObjectsForConnections(Args);
+		}
+	}
+}));
+
+void UObjectReplicationBridge::PrintRelevantObjectsForConnections(const TArray<FString>& Args) const
+{
+	using namespace UE::Net;
+	using namespace UE::Net::Private;
+	using namespace UE::Net::Private::ObjectBridgeDebugging;
+
+	UE_LOG(LogIrisBridge, Display, TEXT("################ Start Printing Relevant Objects of %d Connections ################"), 1);
+	UE_LOG(LogIrisBridge, Display, TEXT(""));
+
+	FReplicationSystemInternal* ReplicationSystemInternal = GetReplicationSystem()->GetReplicationSystemInternal();
+	FReplicationConnections& Connections = ReplicationSystemInternal->GetConnections();
+	const FNetBitArray& ValidConnections = Connections.GetValidConnections();
+
+	// Default to all connections
+	FNetBitArray ConnectionsToPrint;
+	ConnectionsToPrint.InitAndCopy(ValidConnections);
+
+	// Filter down the list if users wanted specific connections
+	TArray<uint32> RequestedConnectionList = FindConnectionsFromArgs(Args);
+	if (RequestedConnectionList.Num())
+	{
+		ConnectionsToPrint.Reset();
+		for (uint32 ConnectionId : RequestedConnectionList)
+		{
+			if (ValidConnections.IsBitSet(ConnectionId))
+			{
+				ConnectionsToPrint.SetBit(ConnectionId);
+			}
+			else
+			{
+				UE_LOG(LogIris, Warning, TEXT("UObjectReplicationBridge::PrintRelevantObjectsForConnections ConnectionId: %u is not valid"), ConnectionId);
+			}
+		}
+	}
+
+	ConnectionsToPrint.ForAllSetBits([&](uint32 ConnectionId)
+	{
+		UE_LOG(LogIrisBridge, Display, TEXT(""));
+		UE_LOG(LogIrisBridge, Display, TEXT("###### Begin Relevant list of Connection:%u ######"), ConnectionId);
+		UE_LOG(LogIrisBridge, Display, TEXT(""));
+
+		struct FRootObjectData
+		{
+			FInternalNetRefIndex ObjectIndex = 0;
+			UObject* Instance = nullptr;
+			UClass* Class = nullptr;
+
+			// Sort by class names
+			bool operator<(const FRootObjectData& rhs) const 
+			{ 
+				if (Class==rhs.Class) { return false; }
+				if (!Class) { return false; }
+				if (!rhs.Class) { return true; }
+				return Class->GetName() < rhs.Class->GetName();
+			}
+		};
+		TArray<FRootObjectData> RelevantObjects;
+
+		FNetBitArray RootObjects;
+		RootObjects.Init(NetRefHandleManager->GetMaxActiveObjectCount());
+		MakeNetBitArrayView(RootObjects).Set(GetReplicationSystem()->GetReplicationSystemInternal()->GetFiltering().GetRelevantObjectsInScope(ConnectionId), FNetBitArrayView::AndNotOp, NetRefHandleManager->GetSubObjectInternalIndicesView());
+		
+		RootObjects.ForAllSetBits([&](uint32 RootObjectIndex)
+		{
+			FRootObjectData Data;
+			Data.ObjectIndex = RootObjectIndex;
+			Data.Instance = NetRefHandleManager->GetReplicatedObjectInstance(RootObjectIndex);
+			Data.Class = Data.Instance ? Data.Instance->GetClass() : nullptr;
+
+			RelevantObjects.Emplace(MoveTemp(Data));
+		});
+
+		RelevantObjects.Sort();
+
+		for (const FRootObjectData& ObjectData : RelevantObjects)
+		{
+			UE_LOG(LogIrisBridge, Display, TEXT("%s"), *PrintNetObject(NetRefHandleManager, ObjectData.ObjectIndex));
+
+			// TODO: Output the subobjects and tell if they are relevant or not to the connection.
+		}
+
+		UE_LOG(LogIrisBridge, Display, TEXT(""));
+		UE_LOG(LogIrisBridge, Display, TEXT("###### Stop Relevant list of Connection:%u | Total: %u root objects relevant ######"), ConnectionId, RelevantObjects.Num());
+		UE_LOG(LogIrisBridge, Display, TEXT(""));
+	});
+
+	UE_LOG(LogIrisBridge, Display, TEXT(""));
+	UE_LOG(LogIrisBridge, Display, TEXT("################ Stop Printing Relevant Objects of %d Connections ################"), 1);
 }
\ No newline at end of file
diff --git a/Engine/Source/Runtime/Experimental/Iris/Core/Public/Iris/ReplicationSystem/ObjectReplicationBridge.h b/Engine/Source/Runtime/Experimental/Iris/Core/Public/Iris/ReplicationSystem/ObjectReplicationBridge.h
--- a/Engine/Source/Runtime/Experimental/Iris/Core/Public/Iris/ReplicationSystem/ObjectReplicationBridge.h
+++ b/Engine/Source/Runtime/Experimental/Iris/Core/Public/Iris/ReplicationSystem/ObjectReplicationBridge.h
@@ -184,7 +184,9 @@ public:
 	// Debug functions exposed via console commands
 	void PrintDynamicFilterClassConfig() const;
 	void PrintReplicatedObjects() const;
+	void PrintAlwaysRelevantObjects() const;
 	void PrintRelevantObjects() const;
+	void PrintRelevantObjectsForConnections(const TArray<FString>& Args) const;
 
 protected:
 	IRISCORE_API virtual ~UObjectReplicationBridge();
@@ -312,6 +314,12 @@ protected:
 	/** Change the max tick rate to match the one from the engine */
 	void SetMaxTickRate(float InMaxTickRate) { MaxTickRate = InMaxTickRate; }
 
+	/**
+	 * Parses a list of arguments and returns a list of Connection's that match them
+	 * Ex: ConnectionId=1 or ConnectionId=1,5,7
+	 */
+	IRISCORE_API virtual TArray<uint32> FindConnectionsFromArgs(const TArray<FString>& Args) const;
+
 private:
 
 	/** Forcibly poll a single replicated object */
@@ -339,7 +347,6 @@ private:
 
 	void SetNetPushIdOnInstance(UE::Net::FReplicationInstanceProtocol* InstanceProtocol, FNetHandle NetHandle);
 
-
 	/** Tries to load the classes used in poll period overrides. */
 	void FindClassesInPollPeriodOverrides();
 
