Chaos - Motion-Aware Collision Detection (MACD) imporovements
- improved boundary edge rejection for MACD
- option to have MACD fall back to non-MACD when the convex center is outside the triangle plane (disabled by default). This can be used to minimize impact without throwing away the tunnelling benefits of MACD.

diff --git a/Engine/Source/Runtime/Experimental/Chaos/Private/Chaos/Collision/ConvexTriangleContactPoint.cpp b/Engine/Source/Runtime/Experimental/Chaos/Private/Chaos/Collision/ConvexTriangleContactPoint.cpp
--- a/Engine/Source/Runtime/Experimental/Chaos/Private/Chaos/Collision/ConvexTriangleContactPoint.cpp
+++ b/Engine/Source/Runtime/Experimental/Chaos/Private/Chaos/Collision/ConvexTriangleContactPoint.cpp
@@ -579,7 +579,7 @@ namespace Chaos::Private
 {
 	// Generate a contact manifold between a convex and a triangle, given the closest feature (i.e., single Contact point)
 	template <typename ConvexType>
-	void ConvexTriangleManifoldFromContact(const ConvexType& Convex, const FTriangle& Triangle, const FConvexContactPoint& Contact, const FReal CullDistance, FContactPointLargeManifold& OutManifold)
+	void ConvexTriangleManifoldFromContact(const ConvexType& Convex, const FTriangle& Triangle, const FVec3& TriangleNormal, const FConvexContactPoint& Contact, const FReal CullDistance, FContactPointLargeManifold& OutManifold)
 	{
 		// Convex plane
 		const int32 ConvexPlaneIndex = Contact.Features[0].PlaneIndex;
@@ -589,7 +589,7 @@ namespace Chaos::Private
 
 		// Triangle plane
 		check(Contact.Features[1].PlaneIndex == 0);
-		const FVec3 TriN = Triangle.GetNormal();
+		const FVec3 TriN = TriangleNormal;
 		const FVec3 TriC = Triangle.GetCentroid();
 
 		const FVec3& SeparatingAxis = Contact.ShapeContactNormal;
@@ -711,12 +711,12 @@ namespace Chaos::Private
 				AddContact(ConvexX, TriX, Distance);
 			}
 		}
-
 	}
 
 	template void ConvexTriangleManifoldFromContact(
 		const FImplicitConvex3& Convex, 
 		const FTriangle& Triangle, 
+		const FVec3& TriangleNormal, 
 		const FConvexContactPoint& Contact,
 		const FReal CullDistance, 
 		FContactPointLargeManifold& OutManifold);
@@ -724,6 +724,7 @@ namespace Chaos::Private
 	template void ConvexTriangleManifoldFromContact(
 		const TImplicitObjectInstanced<FImplicitConvex3>& Convex,
 		const FTriangle& Triangle, 
+		const FVec3& TriangleNormal,
 		const FConvexContactPoint& Contact,
 		const FReal CullDistance, 
 		FContactPointLargeManifold& OutManifold);
@@ -731,6 +732,7 @@ namespace Chaos::Private
 	template void ConvexTriangleManifoldFromContact(
 		const TImplicitObjectScaled<FImplicitConvex3>& Convex,
 		const FTriangle& Triangle,
+		const FVec3& TriangleNormal,
 		const FConvexContactPoint& Contact,
 		const FReal CullDistance,
 		FContactPointLargeManifold& OutManifold);
@@ -738,6 +740,7 @@ namespace Chaos::Private
 	template void ConvexTriangleManifoldFromContact(
 		const FImplicitBox3& Convex,
 		const FTriangle& Triangle,
+		const FVec3& TriangleNormal,
 		const FConvexContactPoint& Contact,
 		const FReal CullDistance,
 		FContactPointLargeManifold& OutManifold);
@@ -745,6 +748,7 @@ namespace Chaos::Private
 	template void ConvexTriangleManifoldFromContact(
 		const TImplicitObjectScaled<FImplicitBox3>& Convex,
 		const FTriangle& Triangle,
+		const FVec3& TriangleNormal,
 		const FConvexContactPoint& Contact,
 		const FReal CullDistance,
 		FContactPointLargeManifold& OutManifold);
@@ -752,6 +756,7 @@ namespace Chaos::Private
 	template void ConvexTriangleManifoldFromContact(
 		const TImplicitObjectInstanced<FImplicitBox3>& Convex,
 		const FTriangle& Triangle,
+		const FVec3& TriangleNormal,
 		const FConvexContactPoint& Contact,
 		const FReal CullDistance,
 		FContactPointLargeManifold& OutManifold);
diff --git a/Engine/Source/Runtime/Experimental/Chaos/Private/Chaos/Collision/MeshContactGenerator.cpp b/Engine/Source/Runtime/Experimental/Chaos/Private/Chaos/Collision/MeshContactGenerator.cpp
--- a/Engine/Source/Runtime/Experimental/Chaos/Private/Chaos/Collision/MeshContactGenerator.cpp
+++ b/Engine/Source/Runtime/Experimental/Chaos/Private/Chaos/Collision/MeshContactGenerator.cpp
@@ -352,9 +352,10 @@ namespace Chaos::Private
 		}
 	}
 
-	bool FMeshContactGenerator::FixFeature(const int32 LocalTriangleIndex, Private::EConvexFeatureType& InOutFeatureType, int32& InOutFeatureIndex, FVec3& InOutPlaneNormal, FVec3& InOutPlanePosition)
+	bool FMeshContactGenerator::FixFeature(const int32 LocalTriangleIndex, Private::EConvexFeatureType& InOutFeatureType, int32& InOutFeatureIndex, FVec3& InOutPlaneNormal)
 	{
 		const FTriangleExt& Triangle = Triangles[LocalTriangleIndex];
+		const FVec3& TriangleNormal = Triangle.GetNormal();
 
 		// For convex edges, we ensure that the normal is between the normals of the adjacent faces
 		// For concave edges, we replace the normal with the triangle face normal
@@ -369,15 +370,18 @@ namespace Chaos::Private
 			const FContactEdgeID EdgeID = FContactEdgeID(VertexIndex0, VertexIndex1);
 
 			const int32 OtherLocalTriangleIndex = GetOtherTriangleIndexForEdge(LocalTriangleIndex, EdgeID);
+
+			// If there is no other triangle we have a boundary edge
 			if (OtherLocalTriangleIndex == INDEX_NONE)
 			{
-				// We don't have the other triangle. Maybe we should use the face plane rather than leaving the edge?
-				return false;
+				// Boundary edge - use face normal
+				InOutFeatureType = Private::EConvexFeatureType::Plane;
+				InOutFeatureIndex = 0;
+				InOutPlaneNormal = TriangleNormal;
+				return true;
 			}
 
 			const FTriangleExt& OtherTriangle = Triangles[OtherLocalTriangleIndex];
-
-			const FVec3& TriangleNormal = Triangle.GetNormal();
 			const FVec3& OtherTriangleNormal = OtherTriangle.GetNormal();
 
 			// Common case - both triangles have the same normal - treat as concave
@@ -415,6 +419,7 @@ namespace Chaos::Private
 				return true;
 			}
 
+			// Same as above but against the other triangle sharing the edge
 			FVec3 OtherVertex0, OtherVertex1;
 			if (OtherTriangle.GetVertexWithID(VertexIndex0, OtherVertex0) && OtherTriangle.GetVertexWithID(VertexIndex1, OtherVertex1))
 			{
@@ -445,10 +450,27 @@ namespace Chaos::Private
 			for (int32 OtherLocalTriangleIndex = 0; OtherLocalTriangleIndex < Triangles.Num(); ++OtherLocalTriangleIndex)
 			{
 				const FTriangleExt& OtherTriangle = Triangles[OtherLocalTriangleIndex];
+
+				// We don't collide with boundary vertices. The vertex is a boundary vertex if any of the edges including
+				// the vertex are boundary edges (i.e., not shared between two triangles)
+				int32 OtherVertexIndexB, OtherVertexIndexC;
+				if (OtherTriangle.GetOtherVertexIDs(VertexIndexA, OtherVertexIndexB, OtherVertexIndexC))
+				{
+					const FContactEdgeID EdgeAB = FContactEdgeID(VertexIndexA, OtherVertexIndexB);
+					const FContactEdgeID EdgeCA = FContactEdgeID(OtherVertexIndexC, VertexIndexA);
+					if (!IsSharedEdge(EdgeAB) || !IsSharedEdge(EdgeCA))
+					{
+						InOutFeatureType = Private::EConvexFeatureType::Plane;
+						InOutFeatureIndex = 0;
+						InOutPlaneNormal = TriangleNormal;
+						return true;
+					}
+				}
+
+				// Does the contact normal point into the infinite prism formed by extruding the triangle along the face normal?
 				FVec3 OtherVertexB, OtherVertexC;
 				if (OtherTriangle.GetOtherVerticesFromID(VertexIndexA, OtherVertexB, OtherVertexC))
 				{
-					// Does the contact normal point into the infinite prism formed by extruding the triangle along the face normal?
 					const FVec3& OtherTriangleNormal = OtherTriangle.GetNormal();
 					const FVec3 OtherEdge0 = OtherVertexB - VertexA;
 					const FVec3 OtherEdge1 = VertexA - OtherVertexC;
diff --git a/Engine/Source/Runtime/Experimental/Chaos/Private/Chaos/CollisionOneShotManifolds.cpp b/Engine/Source/Runtime/Experimental/Chaos/Private/Chaos/CollisionOneShotManifolds.cpp
--- a/Engine/Source/Runtime/Experimental/Chaos/Private/Chaos/CollisionOneShotManifolds.cpp
+++ b/Engine/Source/Runtime/Experimental/Chaos/Private/Chaos/CollisionOneShotManifolds.cpp
@@ -106,6 +106,10 @@ namespace Chaos
 	FAutoConsoleVariableRef CVarChaos_Collision_EnableMeshManifoldOptimizedLoop(TEXT("p.Chaos.Collision.EnableMeshManifoldOptimizedLoop"), bChaos_Collision_EnableMeshManifoldOptimizedLoop, TEXT(""));
 	FAutoConsoleVariableRef CVarChaos_Collision_EnableMeshManifoldOptimizedLoop_TriMesh(TEXT("p.Chaos.Collision.EnableMeshManifoldOptimizedLoopTriMesh"), bChaos_Collision_EnableMeshManifoldOptimizedLoop_TriMesh, TEXT(""));
 
+	// MACD uses the non-MACD path when the shape is outside the triangle plane
+	bool bChaos_Collision_EnableMACDFallback = false;
+	FAutoConsoleVariableRef CVarChaos_Collision_EnableMACDFallback(TEXT("p.Chaos.Collision.EnableMACDFallback"), bChaos_Collision_EnableMACDFallback, TEXT(""));
+
 	// Whether to use the new index-less GJK. 
 	// @todo(chaos): This should be removed once soaked for a bit (enabled 7 June 2022)
 	bool bChaos_Collision_UseGJK2 = false;
diff --git a/Engine/Source/Runtime/Experimental/Chaos/Private/Chaos/CollisionOneShotManifoldsMeshShapes.cpp b/Engine/Source/Runtime/Experimental/Chaos/Private/Chaos/CollisionOneShotManifoldsMeshShapes.cpp
--- a/Engine/Source/Runtime/Experimental/Chaos/Private/Chaos/CollisionOneShotManifoldsMeshShapes.cpp
+++ b/Engine/Source/Runtime/Experimental/Chaos/Private/Chaos/CollisionOneShotManifoldsMeshShapes.cpp
@@ -36,6 +36,7 @@ namespace Chaos
 	extern int32 Chaos_Collision_MeshManifoldHashSize;
 	extern bool bChaos_Collision_EnableMeshManifoldOptimizedLoop;
 	extern bool bChaos_Collision_EnableMeshManifoldOptimizedLoop_TriMesh;
+	extern bool bChaos_Collision_EnableMACDFallback;
 
 	extern bool bChaos_Collision_UseCapsuleTriMesh2;
 	extern bool bChaos_Collision_UseConvexTriMesh2;
@@ -160,7 +161,17 @@ namespace Chaos
 		template<typename ConvexType>
 		void GenerateConvexTriangleOneShotManifoldMACD(const ConvexType& Convex, const FRigidTransform3& ConvexTransform, const FVec3& ConvexRelativeMovement, Private::FMeshContactGenerator& ContactGenerator, const int32 TriangleIndex, const FReal CullDistance, FContactPointLargeManifold& OutContactPoints)
 		{
-			const FTriangle& Triangle = ContactGenerator.GetTriangle(TriangleIndex);
+			const FTriangle TriangleAtP = ContactGenerator.GetTriangle(TriangleIndex);
+			const FVec3 TriangleNormal = ContactGenerator.GetTriangleNormal(TriangleIndex);
+
+			FTriangle Triangle = FTriangle(TriangleAtP.GetVertex(0) + ConvexRelativeMovement, TriangleAtP.GetVertex(1) + ConvexRelativeMovement, TriangleAtP.GetVertex(2) + ConvexRelativeMovement);
+
+			// If we started inside the triangle we ignore this triangle
+			const FReal ConvexTriangleDistance = FVec3::DotProduct(Convex.GetCenterOfMass() - Triangle.GetVertex(0), TriangleNormal);
+			if (ConvexTriangleDistance < 0)
+			{
+				return;
+			}
 
 			// Find the closest feature pair on the triangle and convex
 			Private::FConvexContactPoint ClosestContact;
@@ -169,12 +180,27 @@ namespace Chaos
 				ClosestContact.Features[0].ObjectIndex = 0;
 				ClosestContact.Features[1].ObjectIndex = TriangleIndex;
 
+				// Back face culling
+				const FReal TriangleDotNormal = FVec3::DotProduct(TriangleNormal, ClosestContact.ShapeContactNormal);
+				if (TriangleDotNormal < 0)
+				{
+#if CHAOS_DEBUG_DRAW
+					if (CVars::ChaosSolverDebugDrawMeshContacts && FDebugDrawQueue::GetInstance().IsDebugDrawingEnabled())
+					{
+						const FVec3 P = ConvexTransform.TransformPositionNoScale(ClosestContact.ShapeContactPoints[1] - ConvexRelativeMovement);
+						const FVec3 N = ConvexTransform.TransformVectorNoScale(ClosestContact.ShapeContactNormal);
+						FDebugDrawQueue::GetInstance().DrawDebugLine(P, P + 10.0f * N, FColor::Black, false, CVars::ChaosSolverDebugDebugDrawSettings.DrawDuration, (uint8)CVars::ChaosSolverDebugDebugDrawSettings.DrawPriority, 1.25f * CVars::ChaosSolverDebugDebugDrawSettings.LineThickness);
+					}
+#endif
+					return;
+				}
+
 				// Generate a manifold based on the closest features
 				// NOTE: normal points from triangle to convex
 				const FReal ConvexMotionAlongNormal = FVec3::DotProduct(ConvexRelativeMovement, ClosestContact.ShapeContactNormal);
 				const FReal CullDistancePadding = FMath::Max(0, -ConvexMotionAlongNormal);
 				const FReal NetCullDistance = CullDistance + CullDistancePadding;
-				Private::ConvexTriangleManifoldFromContact(Convex, Triangle, ClosestContact, NetCullDistance, OutContactPoints);
+				Private::ConvexTriangleManifoldFromContact(Convex, Triangle, TriangleNormal, ClosestContact, NetCullDistance, OutContactPoints);
 
 				if (OutContactPoints.Num() > 0)
 				{
@@ -184,7 +210,7 @@ namespace Chaos
 						for (int32 ContactIndex = 0; ContactIndex < OutContactPoints.Num(); ++ContactIndex)
 						{
 							FContactPoint& ContactPoint = OutContactPoints[ContactIndex];
-							const FVec3 P = ConvexTransform.TransformPositionNoScale(ContactPoint.ShapeContactPoints[1]);
+							const FVec3 P = ConvexTransform.TransformPositionNoScale(ContactPoint.ShapeContactPoints[1] - ConvexRelativeMovement);
 							const FVec3 N = ConvexTransform.TransformVectorNoScale(ContactPoint.ShapeContactNormal);
 							FColor Color = FColor::Black;
 							switch (ClosestContact.Features[1].FeatureType)
@@ -206,7 +232,7 @@ namespace Chaos
 
 					// Adjust the feature if it is invalid. E.g., if we collide with a triangle edge and the normal is outside the range allowed by
 					// the triangles sharing the edge we will project the normal into the valid range
-					const bool bFeatureChanged = ContactGenerator.FixFeature(TriangleIndex, ClosestContact.Features[1].FeatureType, ClosestContact.Features[1].PlaneFeatureIndex, ClosestContact.ShapeContactNormal, ClosestContact.ShapeContactPoints[1]);
+					const bool bFeatureChanged = ContactGenerator.FixFeature(TriangleIndex, ClosestContact.Features[1].FeatureType, ClosestContact.Features[1].PlaneFeatureIndex, ClosestContact.ShapeContactNormal);
 					if (bFeatureChanged)
 					{
 						for (int32 ContactIndex = 0; ContactIndex < OutContactPoints.Num(); ++ContactIndex)
@@ -229,7 +255,10 @@ namespace Chaos
 					for (int32 ContactIndex = 0; ContactIndex < OutContactPoints.Num(); ++ContactIndex)
 					{
 						FContactPoint& ContactPoint = OutContactPoints[ContactIndex];
-						Private::ApplyContactPointShift(ContactPoint, ConvexRelativeMovement);
+						
+						const FReal ShiftDotNormal = FVec3::DotProduct(ConvexRelativeMovement, ContactPoint.ShapeContactNormal);
+						ContactPoint.ShapeContactPoints[1] += -ShiftDotNormal * ContactPoint.ShapeContactNormal;
+						ContactPoint.Phi += ShiftDotNormal;
 					}
 				}
 			}
@@ -245,37 +274,54 @@ namespace Chaos
 			}
 			else
 			{
-				// We detect collisions as if both shapes were in the pre-movement locations
-				// GenerateConvexTriangleOneShotManifoldWithMotion will correct the contacts to the actual locations
-				const FRigidTransform3 ShiftedConvexTransform = FRigidTransform3(ConvexTransform.GetTranslation() - RelativeMovement, ConvexTransform.GetRotation());
-				FRigidTransform3 ShiftedMeshToConvexTransform = MeshTransform.GetRelativeTransformNoScale(ShiftedConvexTransform);
-				ShiftedMeshToConvexTransform.SetScale3D(MeshScale);
+				FRigidTransform3 MeshToConvexTransform = MeshTransform.GetRelativeTransformNoScale(ConvexTransform);
+				MeshToConvexTransform.SetScale3D(MeshScale);
 
-				// @todo(chaos): add Convex.CalculateInverseTransformed bounds with scale support (to optimize sphere and capsule)
+				// NOTE: Convex bounds is extended backwards to encompass the pre-movement position
 				const FVec3 ConvexRelativeMovement = ConvexTransform.InverseTransformVectorNoScale(RelativeMovement);
-				const FAABB3 ConvexBounds = FAABB3(Convex.BoundingBox()).GrowByVector(ConvexRelativeMovement).Thicken(InCullDistance);
-				const FAABB3 MeshQueryBounds = ConvexBounds.InverseTransformedAABB(ShiftedMeshToConvexTransform);
+				const FAABB3 ConvexBounds = FAABB3(Convex.BoundingBox()).GrowByVector(-ConvexRelativeMovement).Thicken(InCullDistance);
+				const FAABB3 MeshQueryBounds = ConvexBounds.InverseTransformedAABB(MeshToConvexTransform);
 				const FReal CullDistance = InCullDistance;
 
 				// Collect all the triangles that overlap our convex. Triangles will be in Convex space
-				Mesh.CollectTriangles(MeshQueryBounds, ShiftedMeshToConvexTransform, ConvexBounds, ContactGenerator);
+				Mesh.CollectTriangles(MeshQueryBounds, MeshToConvexTransform, ConvexBounds, ContactGenerator);
 
-				FContactPointLargeManifold Contacts;
-				Contacts.Reserve(6);
+				FContactPointManifold Manifold;
+				FContactPointLargeManifold LargeManifold;
+				LargeManifold.Reserve(6);
 
 				const auto& GenerateConvexTriangleContacts =
-					[&Convex, &ShiftedConvexTransform, &ConvexRelativeMovement, CullDistance, &Contacts](Private::FMeshContactGenerator& ContactGenerator, const int32 TriangleIndex)
+					[&Convex, &ConvexTransform, &ConvexRelativeMovement, CullDistance, &Manifold, &LargeManifold](Private::FMeshContactGenerator& ContactGenerator, const int32 TriangleIndex)
 				{
-					Contacts.Reset();
-					GenerateConvexTriangleOneShotManifoldMACD(Convex, ShiftedConvexTransform, ConvexRelativeMovement, ContactGenerator, TriangleIndex, CullDistance, Contacts);
+					// If we are outside the plane of the triangle, just use non-MACD collision
+					const FTriangle& Triangle = ContactGenerator.GetTriangle(TriangleIndex);
+					bool bUseNonMACDPath = false;
+					if (bChaos_Collision_EnableMACDFallback)
+					{
+						const FVec3 TriangleNormal = ContactGenerator.GetTriangleNormal(TriangleIndex);
+						const FReal Distance = FVec3::DotProduct(Convex.GetCenterOfMass() - Triangle.GetVertex(0), TriangleNormal);
+						bUseNonMACDPath = Distance > 0;
+					}
 
-					ContactGenerator.AddTriangleContacts(TriangleIndex, MakeArrayView(Contacts));
+					if (bUseNonMACDPath)
+					{
+						Manifold.Reset();
+						GenerateConvexTriangleOneShotManifold(Convex, Triangle, CullDistance, Manifold);
+						ContactGenerator.AddTriangleContacts(TriangleIndex, MakeArrayView(Manifold));
+						return;
+					}
+					else
+					{
+						LargeManifold.Reset();
+						GenerateConvexTriangleOneShotManifoldMACD(Convex, ConvexTransform, ConvexRelativeMovement, ContactGenerator, TriangleIndex, CullDistance, LargeManifold);
+						ContactGenerator.AddTriangleContacts(TriangleIndex, MakeArrayView(LargeManifold));
+					}
 				};
 
 				ContactGenerator.GenerateMeshContacts(GenerateConvexTriangleContacts);
 
 				// Process the contacts to minimize manifold etc
-				ContactGenerator.ProcessGeneratedContacts(ShiftedConvexTransform, ShiftedMeshToConvexTransform);
+				ContactGenerator.ProcessGeneratedContacts(ConvexTransform, MeshToConvexTransform);
 			}
 		}
 
@@ -444,7 +490,11 @@ namespace Chaos
 			{
 				const FVec3 RelativeMovement = FVec3(Constraint.GetRelativeMovement());
 				Private::FMeshContactGeneratorSettings ContactGeneratorSettings;
-				ContactGeneratorSettings.bFixNormals = RelativeMovement.IsNearlyZero();
+
+				// @todo(chaos): we don't need to enable normal fixup for MACD mode, except when 
+				// fall back to non-MACD (which is a per-triangle decision) so for now it must be enabled.
+				ContactGeneratorSettings.bFixNormals = true;//RelativeMovement.IsNearlyZero();
+
 				Private::FMeshContactGenerator ContactGenerator(ContactGeneratorSettings);
 
 				if (const FImplicitBox3* RawBox = Convex.template GetObject<FImplicitBox3>())
diff --git a/Engine/Source/Runtime/Experimental/Chaos/Public/Chaos/Collision/ConvexContactPointUtilities.h b/Engine/Source/Runtime/Experimental/Chaos/Public/Chaos/Collision/ConvexContactPointUtilities.h
--- a/Engine/Source/Runtime/Experimental/Chaos/Public/Chaos/Collision/ConvexContactPointUtilities.h
+++ b/Engine/Source/Runtime/Experimental/Chaos/Public/Chaos/Collision/ConvexContactPointUtilities.h
@@ -75,14 +75,6 @@ namespace Chaos:: Private
 			PMax = D0;
 		}
 	}
-
-	// Update contact as if the second shape has moved. Shift must be in the space of the second object
-	inline void ApplyContactPointShift(FContactPoint& InOutContactPoint, const FVec3& Shift)
-	{
-		const FReal ShiftDotNormal = FVec3::DotProduct(Shift, InOutContactPoint.ShapeContactNormal);
-		InOutContactPoint.ShapeContactPoints[1] += (Shift - ShiftDotNormal * InOutContactPoint.ShapeContactNormal);
-		InOutContactPoint.Phi += ShiftDotNormal;
-	}
 }
 
 namespace Chaos
diff --git a/Engine/Source/Runtime/Experimental/Chaos/Public/Chaos/Collision/ConvexTriangleContactPoint.h b/Engine/Source/Runtime/Experimental/Chaos/Public/Chaos/Collision/ConvexTriangleContactPoint.h
--- a/Engine/Source/Runtime/Experimental/Chaos/Public/Chaos/Collision/ConvexTriangleContactPoint.h
+++ b/Engine/Source/Runtime/Experimental/Chaos/Public/Chaos/Collision/ConvexTriangleContactPoint.h
@@ -14,7 +14,7 @@ namespace Chaos
 		extern CHAOS_API bool FindClosestFeatures(const ConvexType& Convex, const FRigidTransform3& ConvexTransform, const FTriangle& Triangle, const FVec3& ConvexRelativeMovement, const FReal CullDistance, FConvexContactPoint& OutContact);
 
 		template <typename ConvexType>
-		extern CHAOS_API void ConvexTriangleManifoldFromContact(const ConvexType& Convex, const FTriangle& Triangle, const FConvexContactPoint& Contact, const FReal CullDistance, FContactPointLargeManifold& OutManifold);
+		extern CHAOS_API void ConvexTriangleManifoldFromContact(const ConvexType& Convex, const FTriangle& Triangle, const FVec3& TriangleNormal, const FConvexContactPoint& Contact, const FReal CullDistance, FContactPointLargeManifold& OutManifold);
 	}
 
 	template <typename ConvexType>
diff --git a/Engine/Source/Runtime/Experimental/Chaos/Public/Chaos/Collision/MeshContactGenerator.h b/Engine/Source/Runtime/Experimental/Chaos/Public/Chaos/Collision/MeshContactGenerator.h
--- a/Engine/Source/Runtime/Experimental/Chaos/Public/Chaos/Collision/MeshContactGenerator.h
+++ b/Engine/Source/Runtime/Experimental/Chaos/Public/Chaos/Collision/MeshContactGenerator.h
@@ -161,22 +161,45 @@ namespace Chaos::Private
 				return false;
 			}
 
+			inline bool GetOtherVertexIDs(const int32 VertexID, int32& OutVertexID0, int32& OutVertexID1) const
+			{
+				if (VertexID == VertexIndices[0])
+				{
+					OutVertexID0 = VertexIndices[1];
+					OutVertexID1 = VertexIndices[2];
+					return true;
+				}
+				else if (VertexID == VertexIndices[1])
+				{
+					OutVertexID0 = VertexIndices[2];
+					OutVertexID1 = VertexIndices[0];
+					return true;
+				}
+				else if (VertexID == VertexIndices[2])
+				{
+					OutVertexID0 = VertexIndices[0];
+					OutVertexID1 = VertexIndices[1];
+					return true;
+				}
+				return false;
+			}
+
 			// Get the positions of the other two vertices in the triangle. (VertexIndex is an index into the owning mesh's vertices)
-			inline bool GetOtherVerticesFromID(const int32 VertexIndex, FVec3& OutVertex0, FVec3& OutVertex1) const
+			inline bool GetOtherVerticesFromID(const int32 VertexID, FVec3& OutVertex0, FVec3& OutVertex1) const
 			{
-				if (VertexIndex == VertexIndices[0])
+				if (VertexID == VertexIndices[0])
 				{
 					OutVertex0 = Triangle.GetVertex(1);
 					OutVertex1 = Triangle.GetVertex(2);
 					return true;
 				}
-				else if (VertexIndex == VertexIndices[1])
+				else if (VertexID == VertexIndices[1])
 				{
 					OutVertex0 = Triangle.GetVertex(2);
 					OutVertex1 = Triangle.GetVertex(0);
 					return true;
 				}
-				else if (VertexIndex == VertexIndices[2])
+				else if (VertexID == VertexIndices[2])
 				{
 					OutVertex0 = Triangle.GetVertex(0);
 					OutVertex1 = Triangle.GetVertex(1);
@@ -370,6 +393,16 @@ namespace Chaos::Private
 			}
 		}
 
+		bool IsSharedEdge(const FContactEdgeID& EdgeID) const
+		{
+			const FEdgeTriangleIndices* EdgeTriangleIndices = EdgeTriangleIndicesMap.Find(EdgeID);
+			if (EdgeTriangleIndices != nullptr)
+			{
+				return (EdgeTriangleIndices->LocalTriangleIndices[0] != INDEX_NONE) && (EdgeTriangleIndices->LocalTriangleIndices[1] != INDEX_NONE);
+			}
+			return false;
+		}
+
 		int32 GetOtherTriangleIndexForEdge(const int32 LocalTriangleIndex, const FContactEdgeID& EdgeID)
 		{
 			FEdgeTriangleIndices* EdgeTriangleIndices = EdgeTriangleIndicesMap.Find(EdgeID);
@@ -422,7 +455,12 @@ namespace Chaos::Private
 			return Triangles[LocalTriangleIndex].GetTriangle();
 		}
 
-		bool FixFeature(const int32 LocalTriangleIndex, Private::EConvexFeatureType& InOutFeatureType, int32& InOutFeatureIndex, FVec3& InOutPlaneNormal, FVec3& InOutPlanePosition);
+		const FVec3& GetTriangleNormal(const int32 LocalTriangleIndex) const
+		{
+			return Triangles[LocalTriangleIndex].GetNormal();
+		}
+
+		bool FixFeature(const int32 LocalTriangleIndex, Private::EConvexFeatureType& InOutFeatureType, int32& InOutFeatureIndex, FVec3& InOutPlaneNormal);
 
 		void AddTriangleContacts(const int32 LocalTriangleIndex, const TArrayView<FContactPoint>& TriangleContactPoints);
 
