Chaos cloth: kinematic collider improvements to help with stability. Also more debug draw.

I realized I hadn't re-implemented the tangential falloff when I refactored the code to use ISPC. That's now added.  Also added a max depth test so we don't push through the wrong direction when things slide around the legs.

Added a map for friction against kinematic colliders.

diff --git a/Engine/Plugins/Experimental/ChaosCloth/Source/Chaos/Private/ChaosCloth/ChaosClothConstraints.cpp b/Engine/Plugins/Experimental/ChaosCloth/Source/Chaos/Private/ChaosCloth/ChaosClothConstraints.cpp
--- a/Engine/Plugins/Experimental/ChaosCloth/Source/Chaos/Private/ChaosCloth/ChaosClothConstraints.cpp
+++ b/Engine/Plugins/Experimental/ChaosCloth/Source/Chaos/Private/ChaosCloth/ChaosClothConstraints.cpp
@@ -803,7 +803,7 @@ void FClothConstraints::AddRules(
 	Softs::FSolverReal MeshScale, bool bEnabled)
 {
 	// Self collisions
-	CreateSelfCollisionConstraints(ConfigProperties, VertexSets, FaceSets, FaceIntMaps, TriangleMesh);
+	CreateSelfCollisionConstraints(ConfigProperties, WeightMaps, VertexSets, FaceSets, FaceIntMaps, TriangleMesh);
 
 	// Edge constraints
 	CreateStretchConstraints(ConfigProperties, WeightMaps, TriangleMesh, PatternData);
@@ -868,7 +868,8 @@ void FClothConstraints::AddRules(
 	}
 }
 
-void FClothConstraints::CreateSelfCollisionConstraints(const Softs::FCollectionPropertyConstFacade& ConfigProperties, 
+void FClothConstraints::CreateSelfCollisionConstraints(const Softs::FCollectionPropertyConstFacade& ConfigProperties,
+	const TMap<FString, TConstArrayView<FRealSingle>>& WeightMaps,
 	const TMap<FString, const TSet<int32>*>& VertexSets,
 	const TMap<FString, const TSet<int32>*>& FaceSets,
 	const TMap<FString, TConstArrayView<int32>>& FaceIntMaps, const FTriangleMesh& TriangleMesh)
@@ -890,6 +891,7 @@ void FClothConstraints::CreateSelfCollisionConstraints(const Softs::FCollectionP
 			NumParticles,
 			TriangleMesh,
 			AnimationPositions,
+			WeightMaps,
 			FaceIntMaps,
 			ConfigProperties);
 		++NumPostCollisionConstraintRules;
@@ -2166,7 +2168,7 @@ void FClothConstraints::Update(
 	}
 	if (SelfCollisionConstraints)
 	{
-		SelfCollisionConstraints->SetProperties(ConfigProperties, FaceIntMaps);
+		SelfCollisionConstraints->SetProperties(ConfigProperties, WeightMaps, FaceIntMaps);
 	}
 	if (SelfCollisionInit)
 	{
diff --git a/Engine/Plugins/Experimental/ChaosCloth/Source/Chaos/Private/ChaosCloth/ChaosClothVisualization.cpp b/Engine/Plugins/Experimental/ChaosCloth/Source/Chaos/Private/ChaosCloth/ChaosClothVisualization.cpp
--- a/Engine/Plugins/Experimental/ChaosCloth/Source/Chaos/Private/ChaosCloth/ChaosClothVisualization.cpp
+++ b/Engine/Plugins/Experimental/ChaosCloth/Source/Chaos/Private/ChaosCloth/ChaosClothVisualization.cpp
@@ -2153,6 +2153,38 @@ static FAutoConsoleVariableRef CVarClothVizWeightMapName(TEXT("p.ChaosClothVisua
 					DrawPoint(PDI, P, Orange, nullptr, 2.f);
 					DrawLine(PDI, Pos1, P, Orange);
 				}
+
+				const TArray<int32>& KinematicCollidingParticles = SelfCollisionConstraints->GetKinematicCollidingParticles();
+				const TArray<TMap<int32, Softs::FSolverReal>>& KinematicColliderTimers = SelfCollisionConstraints->GetKinematicColliderTimers();
+				const FTriangleMesh& TriangleMesh = SelfCollisionConstraints->GetTriangleMesh();
+				for (int32 Index1 : KinematicCollidingParticles)
+				{
+					const FVector P = LocalSpaceLocation + FVector(Positions[Index1]);
+
+					static const FLinearColor Orange(0.3f, 0.15f, 0.f);
+					DrawPoint(PDI, P, Orange, nullptr, 2.f);
+
+					const TMap<int32, Softs::FSolverReal>& Timers = KinematicColliderTimers[Index1 - Offset];
+					for (const TPair<int32, Softs::FSolverReal>& ElemAndTimer : Timers)
+					{
+						const int32 Index2 = TriangleMesh.GetElements()[ElemAndTimer.Get<0>()][0];
+						const int32 Index3 = TriangleMesh.GetElements()[ElemAndTimer.Get<0>()][1];
+						const int32 Index4 = TriangleMesh.GetElements()[ElemAndTimer.Get<0>()][2];
+
+						const Softs::FSolverVec3& P1 = Positions[Index1];
+						const Softs::FSolverVec3& P2 = Positions[Index2];
+						const Softs::FSolverVec3& P3 = Positions[Index3];
+						const Softs::FSolverVec3& P4 = Positions[Index4];
+						Softs::FSolverVec3 Bary;
+						const FVector Pos1 = LocalSpaceLocation + FVector(FindClosestPointAndBaryOnTriangle(P2, P3, P4, P1, Bary));
+
+						static const FLinearColor LtRed(0.6f, 0.f, 0.f);
+						static const FLinearColor DkRed(0.3f, 0.f, 0.f);
+						const FLinearColor& Color = ElemAndTimer.Get<1>() > 0.f ? LtRed : DkRed;
+						DrawPoint(PDI, Pos1, Color, nullptr, 2.f);
+						DrawLine(PDI, Pos1, P, Color);
+					}
+				}
 			}
 
 			if (const Softs::FPBDSelfCollisionSphereConstraints* const SelfCollisionSphereConstraints =
diff --git a/Engine/Plugins/Experimental/ChaosCloth/Source/Chaos/Public/ChaosCloth/ChaosClothConstraints.h b/Engine/Plugins/Experimental/ChaosCloth/Source/Chaos/Public/ChaosCloth/ChaosClothConstraints.h
--- a/Engine/Plugins/Experimental/ChaosCloth/Source/Chaos/Public/ChaosCloth/ChaosClothConstraints.h
+++ b/Engine/Plugins/Experimental/ChaosCloth/Source/Chaos/Public/ChaosCloth/ChaosClothConstraints.h
@@ -151,6 +151,7 @@ PRAGMA_ENABLE_DEPRECATION_WARNINGS
 	private:
 		void CreateSelfCollisionConstraints(
 			const Softs::FCollectionPropertyConstFacade& ConfigProperties,
+			const TMap<FString, TConstArrayView<FRealSingle>>& WeightMaps,
 			const TMap<FString, const TSet<int32>*>& VertexSets,
 			const TMap<FString, const TSet<int32>*>& FaceSets,
 			const TMap<FString, TConstArrayView<int32>>& FaceIntMaps,
diff --git a/Engine/Plugins/Experimental/ChaosClothAssetEditor/Source/ChaosClothAssetDataflowNodes/Private/ChaosClothAsset/SimulationSelfCollisionConfigNode.cpp b/Engine/Plugins/Experimental/ChaosClothAssetEditor/Source/ChaosClothAssetDataflowNodes/Private/ChaosClothAsset/SimulationSelfCollisionConfigNode.cpp
--- a/Engine/Plugins/Experimental/ChaosClothAssetEditor/Source/ChaosClothAssetDataflowNodes/Private/ChaosClothAsset/SimulationSelfCollisionConfigNode.cpp
+++ b/Engine/Plugins/Experimental/ChaosClothAssetEditor/Source/ChaosClothAssetDataflowNodes/Private/ChaosClothAsset/SimulationSelfCollisionConfigNode.cpp
@@ -13,6 +13,7 @@ FChaosClothAssetSimulationSelfCollisionConfigNode::FChaosClothAssetSimulationSel
 	RegisterInputConnection(&SelfCollisionLayers.StringValue, GET_MEMBER_NAME_CHECKED(FChaosClothAssetConnectableIStringValue, StringValue));
 	RegisterInputConnection(&SelfCollisionDisabledFaces.StringValue, GET_MEMBER_NAME_CHECKED(FChaosClothAssetConnectableIStringValue, StringValue));
 	RegisterInputConnection(&SelfCollisionEnabledKinematicFaces.StringValue, GET_MEMBER_NAME_CHECKED(FChaosClothAssetConnectableIStringValue, StringValue));
+	RegisterInputConnection(&SelfCollisionKinematicColliderFrictionWeighted.WeightMap);
 }
 
 void FChaosClothAssetSimulationSelfCollisionConfigNode::AddProperties(FPropertyHelper& PropertyHelper) const
@@ -28,7 +29,7 @@ void FChaosClothAssetSimulationSelfCollisionConfigNode::AddProperties(FPropertyH
 	PropertyHelper.SetPropertyString(this, &SelfCollisionEnabledKinematicFaces);
 	PropertyHelper.SetProperty(this, &SelfCollisionKinematicColliderThickness);
 	PropertyHelper.SetProperty(this, &SelfCollisionKinematicColliderStiffness);
-	PropertyHelper.SetProperty(this, &SelfCollisionKinematicColliderFriction);
+	PropertyHelper.SetPropertyWeighted(TEXT("SelfCollisionKinematicColliderFriction"), SelfCollisionKinematicColliderFrictionWeighted);
 
 	PropertyHelper.SetPropertyBool(this, &bUseSelfIntersections);
 	PropertyHelper.SetPropertyBool(this, &bUseGlobalIntersectionAnalysis);
@@ -37,3 +38,18 @@ void FChaosClothAssetSimulationSelfCollisionConfigNode::AddProperties(FPropertyH
 	PropertyHelper.SetPropertyBool(this, &bUseGlobalPostStepContours);
 	PropertyHelper.SetProperty(this, &SelfCollisionProximityStiffness);
 }
+
+void FChaosClothAssetSimulationSelfCollisionConfigNode::Serialize(FArchive& Ar)
+{
+	Super::Serialize(Ar);
+	if (Ar.IsLoading())
+	{
+#if WITH_EDITORONLY_DATA
+		if (SelfCollisionKinematicColliderFriction_DEPRECATED != FrictionDeprecatedValue)
+		{
+			SelfCollisionKinematicColliderFrictionWeighted.Low = SelfCollisionKinematicColliderFrictionWeighted.High = SelfCollisionKinematicColliderFriction_DEPRECATED;
+			SelfCollisionKinematicColliderFriction_DEPRECATED = FrictionDeprecatedValue;
+		}
+#endif
+	}
+}
\ No newline at end of file
diff --git a/Engine/Plugins/Experimental/ChaosClothAssetEditor/Source/ChaosClothAssetDataflowNodes/Public/ChaosClothAsset/SimulationSelfCollisionConfigNode.h b/Engine/Plugins/Experimental/ChaosClothAssetEditor/Source/ChaosClothAssetDataflowNodes/Public/ChaosClothAsset/SimulationSelfCollisionConfigNode.h
--- a/Engine/Plugins/Experimental/ChaosClothAssetEditor/Source/ChaosClothAssetDataflowNodes/Public/ChaosClothAsset/SimulationSelfCollisionConfigNode.h
+++ b/Engine/Plugins/Experimental/ChaosClothAssetEditor/Source/ChaosClothAssetDataflowNodes/Public/ChaosClothAsset/SimulationSelfCollisionConfigNode.h
@@ -57,9 +57,9 @@ public:
 	UPROPERTY(EditAnywhere, Category = "Self-Collision Properties - Kinematic Colliders", meta = (UIMin = "0", UIMax = "1", ClampMin = "0", ClampMax = "1"))
 	float SelfCollisionKinematicColliderStiffness = 1.f;
 
-	/** Friction coefficient for cloth - kinematic cloth interaction. */
+	/** Friction coefficient for cloth - kinematic cloth interaction. Weight map is on the dynamic cloth, not the collider.*/
 	UPROPERTY(EditAnywhere, Category = "Self-Collision Properties - Kinematic Colliders", meta = (UIMin = "0", UIMax = "1", ClampMin = "0", ClampMax = "1"))
-	float SelfCollisionKinematicColliderFriction = 0.0f;
+	FChaosClothAssetWeightedValue SelfCollisionKinematicColliderFrictionWeighted = {true, 0.0f, 0.f, TEXT("SelfCollisionKinematicColliderFriction")};
 
 	/** Enable self intersection resolution. This will try to fix any cloth intersections that are not handled by collision repulsions. */
 	UPROPERTY(EditAnywhere, Category = "Experimental")
@@ -87,6 +87,13 @@ public:
 
 	FChaosClothAssetSimulationSelfCollisionConfigNode(const Dataflow::FNodeParameters& InParam, FGuid InGuid = FGuid::NewGuid());
 
+	virtual void Serialize(FArchive& Ar) override;
 private:
+	// Deprecated properties
+#if WITH_EDITORONLY_DATA
+	static constexpr float FrictionDeprecatedValue = -1.f;
+	UPROPERTY()
+	float SelfCollisionKinematicColliderFriction_DEPRECATED = FrictionDeprecatedValue;
+#endif
 	virtual void AddProperties(FPropertyHelper& PropertyHelper) const override;
 };
diff --git a/Engine/Source/Runtime/Experimental/Chaos/Private/Chaos/PBDCollisionSpringConstraints.ispc b/Engine/Source/Runtime/Experimental/Chaos/Private/Chaos/PBDCollisionSpringConstraints.ispc
--- a/Engine/Source/Runtime/Experimental/Chaos/Private/Chaos/PBDCollisionSpringConstraints.ispc
+++ b/Engine/Source/Runtime/Experimental/Chaos/Private/Chaos/PBDCollisionSpringConstraints.ispc
@@ -182,6 +182,8 @@ export void ApplyKinematicCollisionSpringConstraints(
 			const uniform FIntVector KinematicColliderElements[],
 			const uniform FIntVector TriangleMeshElements[],
 			const uniform float Height,
+			const uniform float OneOverTangentialFalloffDist,
+			const uniform float MaxDepth,
 			const uniform float KinematicColliderStiffness, 
 			const uniform float KinematicColliderFrictionCoefficient, 
 			const uniform int32 NumConstraints)
@@ -228,21 +230,21 @@ export void ApplyKinematicCollisionSpringConstraints(
 			const varying FVector3f P = FindClosestPointAndBaryOnTriangle(P2, P3, P4, P1, Bary);
 			const varying FVector3f Difference = P1 - P;
 			const varying float NormalDifference = VectorDot(Difference, Normal);
-			if(NormalDifference > Height)
+			if(NormalDifference >= Height || NormalDifference < MaxDepth)
 			{
 				continue;
 			}
 
 			const varying float TangentialDifference = VectorSize(Difference - NormalDifference * Normal);
 			const uniform float TangentialFalloffMultiplier = 1.5f;
-			const varying float TangentialFalloff = 1.f - TangentialDifference / max(TangentialFalloffMultiplier * Height, FLOAT_KINDA_SMALL_NUMBER);
+			const varying float TangentialFalloff = 1.f - TangentialDifference * OneOverTangentialFalloffDist;
 			if(TangentialFalloff <= 0.f)
 			{
 				continue;
 			}
 
 			const float NormalDelta = Height - NormalDifference;
-			const FVector3f RepulsionDelta = KinematicColliderStiffness * NormalDelta * Normal;
+			const FVector3f RepulsionDelta = KinematicColliderStiffness * TangentialFalloff * NormalDelta *  Normal;
 
 			P1 = P1 + RepulsionDelta;
 
@@ -267,3 +269,104 @@ export void ApplyKinematicCollisionSpringConstraints(
 		VectorScatter(&PandInvM[Index1], SetVector4(P1, IM));
 	}
 }
+
+export void ApplyKinematicCollisionSpringConstraintsWithMaps(
+			uniform FVector4f PandInvM[],
+			const uniform FVector3f ParticlesX[],
+			const uniform int32 KinematicCollidingParticles[],
+			const uniform FIntVector KinematicColliderElements[],
+			const uniform FIntVector TriangleMeshElements[],
+			const uniform float Height,
+			const uniform float OneOverTangentialFalloffDist,
+			const uniform float MaxDepth,
+			const uniform float KinematicColliderStiffness, 
+			const uniform FVector2f& KinematicColliderFrictionCoefficientOffsetRange, 
+			const uniform float KinematicColliderFrictionCoefficientMapValues[],
+			const uniform int32 NumConstraints)
+{
+	foreach(Index = 0...NumConstraints)
+	{
+		const int32 Index1 = KinematicCollidingParticles[Index];
+
+		const FVector4f PandInvM1 = VectorGather(&PandInvM[Index1]);
+		varying FVector3f P1;
+		varying float IM;
+		UnzipPandInvM(PandInvM1, P1, IM);
+
+		const varying FIntVector Elements = VectorLoad(&KinematicColliderElements[extract(Index,0)]);
+		
+		#pragma ignore warning(perf)
+		const varying float FrictionMapValue = KinematicColliderFrictionCoefficientMapValues[Index1];
+		const varying float FrictionCoefficient = KinematicColliderFrictionCoefficientOffsetRange.V[0] + FrictionMapValue * KinematicColliderFrictionCoefficientOffsetRange.V[1];
+
+		for(uniform int32 EIndex = 0; EIndex < 3; ++EIndex)
+		{
+			const varying int32 ElementIndex = Elements.V[EIndex];
+			if(ElementIndex == -1)
+			{
+				break;
+			}
+
+			const FIntVector Element = VectorGather(&TriangleMeshElements[ElementIndex]);
+
+			const varying int32 Index2 = Element.V[0];
+			const varying int32 Index3 = Element.V[1];
+			const varying int32 Index4 = Element.V[2];
+		
+			const varying FVector4f PandInvM2 = VectorGather(&PandInvM[Index2]);
+			const varying FVector4f PandInvM3 = VectorGather(&PandInvM[Index3]);
+			const varying FVector4f PandInvM4 = VectorGather(&PandInvM[Index4]);
+
+			varying FVector3f P2, P3, P4;
+			varying float IM2, IM3, IM4;
+		
+			UnzipPandInvM(PandInvM2, P2, IM2);
+			UnzipPandInvM(PandInvM3, P3, IM3);
+			UnzipPandInvM(PandInvM4, P4, IM4);
+
+			const varying FVector3f Normal = GetTriangleNormalFlipped(P2, P3, P4);
+
+			varying FVector3f Bary;
+			const varying FVector3f P = FindClosestPointAndBaryOnTriangle(P2, P3, P4, P1, Bary);
+			const varying FVector3f Difference = P1 - P;
+			const varying float NormalDifference = VectorDot(Difference, Normal);
+			if(NormalDifference >= Height || NormalDifference < MaxDepth)
+			{
+				continue;
+			}
+
+			const varying float TangentialDifference = VectorSize(Difference - NormalDifference * Normal);
+			const uniform float TangentialFalloffMultiplier = 1.5f;
+			const varying float TangentialFalloff = 1.f - TangentialDifference * OneOverTangentialFalloffDist;
+			if(TangentialFalloff <= 0.f)
+			{
+				continue;
+			}
+
+			const float NormalDelta = Height - NormalDifference;
+			const FVector3f RepulsionDelta = KinematicColliderStiffness * TangentialFalloff * NormalDelta *  Normal;
+
+			P1 = P1 + RepulsionDelta;
+
+
+			if (FrictionCoefficient > 0)
+			{
+				const FVector3f X1 = VectorGather(&ParticlesX[Index1]);
+				const FVector3f X2 = VectorGather(&ParticlesX[Index2]);
+				const FVector3f X3 = VectorGather(&ParticlesX[Index3]);
+				const FVector3f X4 = VectorGather(&ParticlesX[Index4]);
+				const FVector3f X = Bary.V[0] * X2 + Bary.V[1] * X3 + Bary.V[2] * X4;
+				const FVector3f RelativeDisplacement = (P1 - X1) - (P - X);
+				const FVector3f RelativeDisplacementTangent = RelativeDisplacement - VectorDot(RelativeDisplacement,Normal) * Normal;
+				const float RelativeDisplacementTangentLength = VectorSize(RelativeDisplacementTangent);
+				const float PositionCorrection = min(NormalDelta * FrictionCoefficient, RelativeDisplacementTangentLength);
+				const float CorrectionRatio = RelativeDisplacementTangentLength < FLOAT_SMALL_NUMBER ? 0.f : PositionCorrection / RelativeDisplacementTangentLength;
+				const FVector3f FrictionDelta = -CorrectionRatio * RelativeDisplacementTangent;
+
+				P1 = P1 + FrictionDelta;
+			}
+		}
+
+		VectorScatter(&PandInvM[Index1], SetVector4(P1, IM));
+	}
+}
\ No newline at end of file
diff --git a/Engine/Source/Runtime/Experimental/Chaos/Private/Chaos/PBDCollisionSpringConstraintsBase.cpp b/Engine/Source/Runtime/Experimental/Chaos/Private/Chaos/PBDCollisionSpringConstraintsBase.cpp
--- a/Engine/Source/Runtime/Experimental/Chaos/Private/Chaos/PBDCollisionSpringConstraintsBase.cpp
+++ b/Engine/Source/Runtime/Experimental/Chaos/Private/Chaos/PBDCollisionSpringConstraintsBase.cpp
@@ -26,9 +26,17 @@ FAutoConsoleVariableRef CVarChaosCollisionSpringsISPCEnabled(TEXT("p.Chaos.Colli
 
 #if UE_BUILD_SHIPPING
 static constexpr Chaos::Softs::FSolverReal KinematicColliderMaxTimer = (Chaos::Softs::FSolverReal)0.1f;
+static constexpr Chaos::Softs::FSolverReal KinematicColliderFalloffMultiplier = (Chaos::Softs::FSolverReal)1.f;
+static constexpr Chaos::Softs::FSolverReal KinematicColliderMaxDepthMultiplier = (Chaos::Softs::FSolverReal)10.f;
 #else
-Chaos::Softs::FSolverReal KinematicColliderMaxTimer = (Chaos::Softs::FSolverReal)0.1f;
+Chaos::Softs::FSolverReal KinematicColliderMaxTimer = (Chaos::Softs::FSolverReal)0.1f; 
 FAutoConsoleVariableRef CVarChaosCollisionSpringMaxTimer(TEXT("p.Chaos.CollisionSpring.MaxTimer"), KinematicColliderMaxTimer, TEXT("Amount of time (in seconds) to remember a kinematic collision connection after it has moved more than Thickness away. Increasing this can reduce jitter at the cost of more computation."));
+
+Chaos::Softs::FSolverReal KinematicColliderFalloffMultiplier = (Chaos::Softs::FSolverReal)1.f;
+FAutoConsoleVariableRef CVarChaosCollisionSpringFalloffMultiplier(TEXT("p.Chaos.CollisionSpring.FalloffMultiplier"), KinematicColliderFalloffMultiplier, TEXT("Tangential distance away from a triangle (scaled by thickness) beyond which a point isn't considered to be kinematically colliding"));
+
+Chaos::Softs::FSolverReal KinematicColliderMaxDepthMultiplier = (Chaos::Softs::FSolverReal)10.f;
+FAutoConsoleVariableRef CVarChaosCollisionSpringMaxDepthMultiplier(TEXT("p.Chaos.CollisionSpring.MaxDepthMultiplier"), KinematicColliderMaxDepthMultiplier, TEXT("Penetration depth beyond which we ignore the kinematic collision (so you don't push through the wrong side)"));
 #endif
 
 namespace Chaos::Softs {
@@ -57,20 +65,21 @@ FPBDCollisionSpringConstraintsBase::FPBDCollisionSpringConstraintsBase(
 	const FTriangleMesh& InTriangleMesh,
 	const TArray<FSolverVec3>* InReferencePositions,
 	TSet<TVec2<int32>>&& InDisabledCollisionElements,
+	const TConstArrayView<FRealSingle>& InKinematicColliderFrictionMultipliers,
 	const TConstArrayView<int32>& InSelfCollisionLayers,
 	const FSolverReal InThickness,
 	const FSolverReal InStiffness,
 	const FSolverReal InFrictionCoefficient,
 	const FSolverReal InKinematicColliderThickness,
 	const FSolverReal InKinematicColliderStiffness,
-	const FSolverReal InKinematicColliderFrictionCoefficient,
+	const FSolverVec2 InKinematicColliderFrictionCoefficient,
 	const FSolverReal InProximityStiffness)
 	: Thickness(InThickness)
 	, Stiffness(InStiffness)
 	, FrictionCoefficient(InFrictionCoefficient)
 	, KinematicColliderThickness(InKinematicColliderThickness)
 	, KinematicColliderStiffness(InKinematicColliderStiffness)
-	, KinematicColliderFrictionCoefficient(InKinematicColliderFrictionCoefficient)
+	, KinematicColliderFrictionCoefficient(InKinematicColliderFrictionCoefficient, InKinematicColliderFrictionMultipliers, InNumParticles)
 	, ProximityStiffness(InProximityStiffness)
 	, TriangleMesh(InTriangleMesh)
 	, ReferencePositions(InReferencePositions)
@@ -585,21 +594,46 @@ void FPBDCollisionSpringConstraintsBase::ApplyKinematicConstraints(SolverParticl
 	TRACE_CPUPROFILER_EVENT_SCOPE(ChaosPBDCollisionSpring_ApplyKinematicConstraints);
 
 	const FSolverReal Height = Thickness + KinematicColliderThickness;
+	const FSolverReal OneOverTangentialFalloffDist = (FSolverReal)1.f / FMath::Max(KinematicColliderFalloffMultiplier * Height, UE_KINDA_SMALL_NUMBER);
+	const FSolverReal MaxDepth = -Height * KinematicColliderMaxDepthMultiplier;
 #if INTEL_ISPC
 	static_assert(sizeof(ispc::FIntVector) == sizeof(TVector<int32, MaxKinematicConnectionsPerPoint>), "sizeof(ispc::FIntVector) != sizeof(TVector<int32, MaxKinematicConnectionsPerPoint>)");
 	if (bRealTypeCompatibleWithISPC && bChaos_CollisionSpring_ISPC_Enabled)
 	{
-		ispc::ApplyKinematicCollisionSpringConstraints(
-			(ispc::FVector4f*)Particles.GetPAndInvM().GetData(),
-			(const ispc::FVector3f*)Particles.XArray().GetData(),
-			KinematicCollidingParticles.GetData(),
-			(const ispc::FIntVector*)KinematicColliderElements.GetData(),
-			(const ispc::FIntVector*)TriangleMesh.GetElements().GetData(),
-			Height,
-			KinematicColliderStiffness,
-			KinematicColliderFrictionCoefficient,
-			KinematicCollidingParticles.Num()
-		);
+		if (KinematicColliderFrictionCoefficient.HasWeightMap())
+		{
+			ispc::ApplyKinematicCollisionSpringConstraintsWithMaps(
+				(ispc::FVector4f*)Particles.GetPAndInvM().GetData(),
+				(const ispc::FVector3f*)Particles.XArray().GetData(),
+				KinematicCollidingParticles.GetData(),
+				(const ispc::FIntVector*)KinematicColliderElements.GetData(),
+				(const ispc::FIntVector*)TriangleMesh.GetElements().GetData(),
+				Height,
+				OneOverTangentialFalloffDist,
+				MaxDepth,
+				KinematicColliderStiffness,
+				reinterpret_cast<const ispc::FVector2f&>(KinematicColliderFrictionCoefficient.GetOffsetRange()),
+				KinematicColliderFrictionCoefficient.GetMapValues().GetData(),
+				KinematicCollidingParticles.Num()
+			);
+
+		}
+		else
+		{
+			ispc::ApplyKinematicCollisionSpringConstraints(
+				(ispc::FVector4f*)Particles.GetPAndInvM().GetData(),
+				(const ispc::FVector3f*)Particles.XArray().GetData(),
+				KinematicCollidingParticles.GetData(),
+				(const ispc::FIntVector*)KinematicColliderElements.GetData(),
+				(const ispc::FIntVector*)TriangleMesh.GetElements().GetData(),
+				Height,
+				OneOverTangentialFalloffDist,
+				MaxDepth,
+				KinematicColliderStiffness,
+				(FSolverReal)KinematicColliderFrictionCoefficient,
+				KinematicCollidingParticles.Num()
+			);
+		}
 	}
 	else
 #endif
@@ -632,32 +666,32 @@ void FPBDCollisionSpringConstraintsBase::ApplyKinematicConstraints(SolverParticl
 				const FSolverVec3 Difference = P1 - P;
 				const FSolverReal NormalDifference = Difference.Dot(Normal);
 
-				if (NormalDifference > Height)
+				if (NormalDifference >= Height || NormalDifference < MaxDepth)
 				{
 					continue;
 				}
 
 				const FSolverReal TangentialDifference = (Difference - NormalDifference * Normal).Size();
-				constexpr FSolverReal TangentialFalloffMultiplier(1.5f);
-				const FSolverReal TangentialFalloff = (FSolverReal)1.f - TangentialDifference / FMath::Max(TangentialFalloffMultiplier * Height, UE_KINDA_SMALL_NUMBER);
+				const FSolverReal TangentialFalloff = (FSolverReal)1.f - TangentialDifference * OneOverTangentialFalloffDist;
 				if (TangentialFalloff <= 0.f)
 				{
 					continue;
 				}
 
 				const FSolverReal NormalDelta = Height - NormalDifference;
-				const FSolverVec3 RepulsionDelta = KinematicColliderStiffness * NormalDelta * Normal;
+				const FSolverVec3 RepulsionDelta = KinematicColliderStiffness * TangentialFalloff * NormalDelta * Normal;
 
 				P1 += RepulsionDelta;
 
-				if (KinematicColliderFrictionCoefficient > 0)
+				const FSolverReal KinematicFrictionCoefficient = KinematicColliderFrictionCoefficient.GetValue(Index1);
+				if (KinematicFrictionCoefficient > 0)
 				{
 					const FSolverVec3& X1 = Particles.X(Index1);
 					const FSolverVec3 X = Bary[0] * Particles.X(Index2) + Bary[1] * Particles.X(Index3) + Bary[2] * Particles.X(Index4);
 					const FSolverVec3 RelativeDisplacement = (P1 - X1) - (P - X);
 					const FSolverVec3 RelativeDisplacementTangent = RelativeDisplacement - RelativeDisplacement.Dot(Normal) * Normal;
 					const FSolverReal RelativeDisplacementTangentLength = RelativeDisplacementTangent.Length();
-					const FSolverReal PositionCorrection = FMath::Min(NormalDelta * KinematicColliderFrictionCoefficient, RelativeDisplacementTangentLength);
+					const FSolverReal PositionCorrection = FMath::Min(NormalDelta * KinematicFrictionCoefficient, RelativeDisplacementTangentLength);
 					const FSolverReal CorrectionRatio = RelativeDisplacementTangentLength < UE_SMALL_NUMBER ? 0.f : PositionCorrection / RelativeDisplacementTangentLength;
 					const FSolverVec3 FrictionDelta = -CorrectionRatio * RelativeDisplacementTangent;
 
diff --git a/Engine/Source/Runtime/Experimental/Chaos/Public/Chaos/PBDCollisionSpringConstraints.h b/Engine/Source/Runtime/Experimental/Chaos/Public/Chaos/PBDCollisionSpringConstraints.h
--- a/Engine/Source/Runtime/Experimental/Chaos/Public/Chaos/PBDCollisionSpringConstraints.h
+++ b/Engine/Source/Runtime/Experimental/Chaos/Public/Chaos/PBDCollisionSpringConstraints.h
@@ -26,6 +26,7 @@ public:
 		const int32 InNumParticles,
 		const FTriangleMesh& InTriangleMesh,
 		const TArray<FSolverVec3>* InRestPositions,
+		const TMap<FString, TConstArrayView<FRealSingle>>& WeightMaps,
 		const TMap<FString, TConstArrayView<int32>>& FaceIntMaps,
 		const FCollectionPropertyConstFacade& PropertyCollection)
 		: Base(
@@ -34,13 +35,14 @@ public:
 			InTriangleMesh,
 			InRestPositions,
 			GenerateDisabledCollisionElements(InTriangleMesh, PropertyCollection),
+			WeightMaps.FindRef(GetSelfCollisionKinematicColliderFrictionString(PropertyCollection, SelfCollisionKinematicColliderFrictionName.ToString())),
 			FaceIntMaps.FindRef(GetSelfCollisionLayersString(PropertyCollection, SelfCollisionLayersName.ToString())),
 			(FSolverReal)FMath::Max(GetSelfCollisionThickness(PropertyCollection, Base::BackCompatThickness), 0.f),
 			(FSolverReal)FMath::Clamp(GetSelfCollisionStiffness(PropertyCollection, Base::BackCompatStiffness), 0.f, 1.f),
 			FMath::Clamp((FSolverReal)GetSelfCollisionFriction(PropertyCollection, Base::BackCompatFrictionCoefficient), MinFrictionCoefficient, MaxFrictionCoefficient),
 			(FSolverReal)FMath::Max(GetSelfCollisionKinematicColliderThickness(PropertyCollection, Base::DefaultKinematicColliderThickness), 0.f),
 			(FSolverReal)FMath::Clamp((FSolverReal)GetSelfCollisionKinematicColliderStiffness(PropertyCollection, Base::DefaultKinematicColliderStiffness), 0.f, 1.f),
-			FMath::Clamp((FSolverReal)GetSelfCollisionKinematicColliderFriction(PropertyCollection, Base::DefaultKinematicColliderFrictionCoefficient), MinFrictionCoefficient, MaxFrictionCoefficient),
+			FSolverVec2(GetWeightedFloatSelfCollisionKinematicColliderFriction(PropertyCollection, Base::DefaultKinematicColliderFrictionCoefficient)).ClampAxes(MinFrictionCoefficient, MaxFrictionCoefficient),
 			(FSolverReal)GetSelfCollisionProximityStiffness(PropertyCollection, Base::DefaultProximityStiffness))
 		, SelfCollisionThicknessIndex(PropertyCollection)
 		, SelfCollisionStiffnessIndex(PropertyCollection)
@@ -65,6 +67,7 @@ public:
 			InNumParticles,
 			InTriangleMesh,
 			InRestPositions,
+			TMap<FString, TConstArrayView<FRealSingle>>(),
 			TMap<FString, TConstArrayView<int32>>(),
 			PropertyCollection
 		)
@@ -85,6 +88,7 @@ public:
 			InTriangleMesh,
 			InRestPositions,
 			MoveTemp(InDisabledCollisionElements),
+			TConstArrayView<FRealSingle>(),
 			TConstArrayView<int32>(),
 			InThickness,
 			InStiffness,
@@ -104,6 +108,7 @@ public:
 	using Base::Init;
 
 	void SetProperties(const FCollectionPropertyConstFacade& PropertyCollection,
+		const TMap<FString, TConstArrayView<FRealSingle>>& WeightMaps,
 		const TMap<FString, TConstArrayView<int32>>& FaceIntMaps)
 	{
 		if (IsSelfCollisionThicknessMutable(PropertyCollection))
@@ -130,16 +135,33 @@ public:
 				Base::UpdateCollisionLayers(FaceIntMaps.FindRef(FaceIntMapName));
 			}
 		}
+		if (IsSelfCollisionKinematicColliderThicknessMutable(PropertyCollection))
+		{
+			KinematicColliderThickness = FMath::Max(GetSelfCollisionKinematicColliderThickness(PropertyCollection), 0.f);
+		}
 		if (IsSelfCollisionKinematicColliderStiffnessMutable(PropertyCollection))
 		{
 			KinematicColliderStiffness = FMath::Clamp((FSolverReal)GetSelfCollisionKinematicColliderStiffness(PropertyCollection), 0.f, 1.f);
 		}
+		if (IsSelfCollisionKinematicColliderFrictionMutable(PropertyCollection))
+		{
+			const FSolverVec2 WeightedValue = FSolverVec2(GetWeightedFloatSelfCollisionKinematicColliderFriction(PropertyCollection)).ClampAxes(MinFrictionCoefficient, MaxFrictionCoefficient);
+			if (IsSelfCollisionKinematicColliderFrictionStringDirty(PropertyCollection))
+			{
+				const FString& WeightMapName = GetSelfCollisionKinematicColliderFrictionString(PropertyCollection);
+				KinematicColliderFrictionCoefficient = FPBDFlatWeightMap(WeightedValue, WeightMaps.FindRef(WeightMapName), GetNumParticles());
+			}
+			else
+			{
+				KinematicColliderFrictionCoefficient.SetWeightedValue(WeightedValue);
+			}
+		}
 	}
 
 	UE_DEPRECATED(5.4, "Use SetProperties with FaceIntMaps")
 	void SetProperties(const FCollectionPropertyConstFacade& PropertyCollection)
 	{
-		SetProperties(PropertyCollection, TMap<FString, TConstArrayView<int32>>());
+		SetProperties(PropertyCollection, TMap<FString, TConstArrayView<FRealSingle>>(), TMap<FString, TConstArrayView<int32>>());
 	}
 
 private:
diff --git a/Engine/Source/Runtime/Experimental/Chaos/Public/Chaos/PBDCollisionSpringConstraintsBase.h b/Engine/Source/Runtime/Experimental/Chaos/Public/Chaos/PBDCollisionSpringConstraintsBase.h
--- a/Engine/Source/Runtime/Experimental/Chaos/Public/Chaos/PBDCollisionSpringConstraintsBase.h
+++ b/Engine/Source/Runtime/Experimental/Chaos/Public/Chaos/PBDCollisionSpringConstraintsBase.h
@@ -3,6 +3,7 @@
 
 #if !COMPILE_WITHOUT_UNREAL_SUPPORT
 #include "Chaos/Core.h"
+#include "Chaos/PBDFlatWeightMap.h"
 #include "Chaos/PBDSoftsEvolutionFwd.h"
 #include "Chaos/PBDSoftsSolverParticles.h"
 #include "Chaos/PBDTriangleMeshCollisions.h"
@@ -35,13 +36,14 @@ public:
 		const FTriangleMesh& InTriangleMesh,
 		const TArray<FSolverVec3>* InReferencePositions,
 		TSet<TVec2<int32>>&& InDisabledCollisionElements,
+		const TConstArrayView<FRealSingle>& InKinematicColliderFrictionMultipliers,
 		const TConstArrayView<int32>& InSelfCollisionLayers,
 		const FSolverReal InThickness = BackCompatThickness,
 		const FSolverReal InStiffness = BackCompatStiffness,
 		const FSolverReal InFrictionCoefficient = BackCompatFrictionCoefficient,
 		const FSolverReal InKinematicColliderThickness = DefaultKinematicColliderThickness,
 		const FSolverReal InKinematicColliderStiffness = DefaultKinematicColliderStiffness,
-		const FSolverReal InKinematicColliderFrictionCoefficient = DefaultKinematicColliderFrictionCoefficient,
+		const FSolverVec2 InKinematicColliderFrictionCoefficient = FSolverVec2(DefaultKinematicColliderFrictionCoefficient),
 		const FSolverReal InProximityStiffness = DefaultProximityStiffness);
 
 	UE_DEPRECATED(5.4, "Use constructor with SelfCollisionLayers")
@@ -55,7 +57,7 @@ public:
 		const FSolverReal InStiffness = BackCompatStiffness,
 		const FSolverReal InFrictionCoefficient = BackCompatFrictionCoefficient)
 		: FPBDCollisionSpringConstraintsBase(InOffset, InNumParticles, InTriangleMesh, InReferencePositions, MoveTemp(InDisabledCollisionElements),
-			TConstArrayView<int32>(), InThickness, InStiffness, InFrictionCoefficient)
+			TConstArrayView<FRealSingle>(), TConstArrayView<int32>(), InThickness, InStiffness, InFrictionCoefficient)
 	{}
 
 	PRAGMA_DISABLE_DEPRECATION_WARNINGS
@@ -95,6 +97,10 @@ public:
 		PRAGMA_ENABLE_DEPRECATION_WARNINGS
 	}
 
+	const TArray<int32>& GetKinematicCollidingParticles() const { return KinematicCollidingParticles; }
+	const TArray<TMap<int32, FSolverReal>>& GetKinematicColliderTimers() const { return KinematicColliderTimers; }
+	const FTriangleMesh& GetTriangleMesh() const { return TriangleMesh; }
+
 	void SetThickness(FSolverReal InThickness) { Thickness = FMath::Max(InThickness, (FSolverReal)0.);  }
 	void SetFrictionCoefficient(FSolverReal InFrictionCoefficient) { FrictionCoefficient = InFrictionCoefficient; }
 
@@ -151,8 +157,6 @@ public:
 		case EConstraintType::Default:
 		case EConstraintType::GIAFlipped:
 			return 2.f * Thickness;
-		case EConstraintType::Kinematic:
-			return Thickness + KinematicColliderThickness;
 		}		
 	}
 
@@ -164,8 +168,6 @@ public:
 		case EConstraintType::Default:
 		case EConstraintType::GIAFlipped:
 			return Stiffness;
-		case EConstraintType::Kinematic:
-			return KinematicColliderStiffness;
 		}
 	}
 
@@ -176,8 +178,6 @@ public:
 		default:
 		case EConstraintType::Default:
 			return FrictionCoefficient;
-		case EConstraintType::Kinematic:
-			return KinematicColliderFrictionCoefficient;
 		case EConstraintType::GIAFlipped:
 			return (FSolverReal)0.f;
 		}
@@ -189,7 +189,7 @@ protected:
 	FSolverReal FrictionCoefficient;
 	FSolverReal KinematicColliderThickness;
 	FSolverReal KinematicColliderStiffness;
-	FSolverReal KinematicColliderFrictionCoefficient;
+	FPBDFlatWeightMap KinematicColliderFrictionCoefficient;
 	FSolverReal ProximityStiffness; // (actual spring stiffness for force-based solver)
 
 	UE_DEPRECATED(5.4, "Constraints will be made private")
@@ -201,6 +201,8 @@ protected:
 
 	CHAOS_API void UpdateCollisionLayers(const TConstArrayView<int32>& InFaceCollisionLayers);
 
+	int32 GetNumParticles() const { return NumParticles; }
+
 private:
 
 	template<typename SolverParticlesOrRange>
@@ -215,7 +217,6 @@ private:
 	enum struct EConstraintType : uint8
 	{
 		Default,
-		Kinematic,
 		GIAFlipped,
 	};
 	TArray<EConstraintType> ConstraintTypes;
diff --git a/Engine/Source/Runtime/Experimental/Chaos/Public/Chaos/PBDFlatWeightMap.h b/Engine/Source/Runtime/Experimental/Chaos/Public/Chaos/PBDFlatWeightMap.h
--- a/Engine/Source/Runtime/Experimental/Chaos/Public/Chaos/PBDFlatWeightMap.h
+++ b/Engine/Source/Runtime/Experimental/Chaos/Public/Chaos/PBDFlatWeightMap.h
@@ -159,6 +159,8 @@ public:
 	/** Return the table of stiffnesses as a read only array. */
 	TConstArrayView<FSolverReal> GetMapValues() const { return MapView.GetMapValues(); }
 
+	FSolverReal GetValue(int32 Index) const { return MapView.GetValue(Index); }
+
 	/** Reorder Indices based on Constraint reordering. */
 	inline void ReorderIndices(const TArray<int32>& OrigToReorderedIndices);
 
