Removing ensure() that is being hit on editor startup for the time being

diff --git a/Engine/Source/Runtime/Experimental/Chaos/Private/PBDRigidsSolver.cpp b/Engine/Source/Runtime/Experimental/Chaos/Private/PBDRigidsSolver.cpp
--- a/Engine/Source/Runtime/Experimental/Chaos/Private/PBDRigidsSolver.cpp
+++ b/Engine/Source/Runtime/Experimental/Chaos/Private/PBDRigidsSolver.cpp
@@ -1812,10 +1812,6 @@ namespace Chaos
 				{
 					MRewindCallback->RegisterRewindableSimCallback_Internal(SimCallbackObject);
 				}
-				else
-				{
-					ensure(!IsNetworkPhysicsPredictionEnabled());
-				}
 			}
 
 			if (SimCallbackObject->HasOption(ESimCallbackOptions::PhysicsObjectUnregister))
