Enable rendering root proxy meshes in editor  when available for geometry collections
- change back to normal rendering when asset is selected in fracture editor

diff --git a/Engine/Plugins/Experimental/ChaosEditor/Source/FractureEditor/Private/FractureEditorMode.cpp b/Engine/Plugins/Experimental/ChaosEditor/Source/FractureEditor/Private/FractureEditorMode.cpp
--- a/Engine/Plugins/Experimental/ChaosEditor/Source/FractureEditor/Private/FractureEditorMode.cpp
+++ b/Engine/Plugins/Experimental/ChaosEditor/Source/FractureEditor/Private/FractureEditorMode.cpp
@@ -564,6 +564,7 @@ void UFractureEditorMode::OnActorSelectionChanged(const TArray<UObject*>& NewSel
 		for(UGeometryCollectionComponent* GeometryCollectionComponent : GeometryCollectionComponents)
 		{
 			GeometryCollectionComponent->SetEmbeddedGeometrySelectable(true);
+			GeometryCollectionComponent->EnableRootProxyStaticMeshComponents(false);
 			
 			FGeometryCollectionEdit RestCollectionEdit = GeometryCollectionComponent->EditRestCollection(GeometryCollection::EEditUpdate::None);
 			if (!RestCollectionEdit.GetRestCollection())
@@ -602,6 +603,9 @@ void UFractureEditorMode::OnActorSelectionChanged(const TArray<UObject*>& NewSel
 
 			ExistingSelection->SetEmbeddedGeometrySelectable(false);
 
+			// should we restore the previous state ? 
+			ExistingSelection->EnableRootProxyStaticMeshComponents(true);
+
 			// If we have a Hide array on the collection, remove it.
 			if (const UGeometryCollection* RestCollection = ExistingSelection->GetRestCollection())
 			{
diff --git a/Engine/Source/Runtime/Experimental/GeometryCollectionEngine/Private/GeometryCollection/GeometryCollectionComponent.cpp b/Engine/Source/Runtime/Experimental/GeometryCollectionEngine/Private/GeometryCollection/GeometryCollectionComponent.cpp
--- a/Engine/Source/Runtime/Experimental/GeometryCollectionEngine/Private/GeometryCollection/GeometryCollectionComponent.cpp
+++ b/Engine/Source/Runtime/Experimental/GeometryCollectionEngine/Private/GeometryCollection/GeometryCollectionComponent.cpp
@@ -1077,7 +1077,7 @@ void UGeometryCollectionComponent::UpdateCachedBounds()
 
 bool UGeometryCollectionComponent::ShouldCreateRenderState() const
 {
-	return !CanUseCustomRenderer() && RootProxyStaticMeshComponents.IsEmpty();
+	return !CanUseCustomRenderer();
 }
 
 void UGeometryCollectionComponent::CreateRenderState_Concurrent(FRegisterComponentContext* Context)
@@ -1092,7 +1092,8 @@ FPrimitiveSceneProxy* UGeometryCollectionComponent::CreateSceneProxy()
 
 	FPrimitiveSceneProxy* LocalSceneProxy = nullptr;
 
-	if (RestCollection && !CanUseCustomRenderer())
+	const bool bNotUsingRootProxyComponents = RootProxyStaticMeshComponents.IsEmpty();
+	if (RestCollection && !CanUseCustomRenderer() && bNotUsingRootProxyComponents)
 	{
 		if (UseNanite(GetScene()->GetShaderPlatform()) &&
 			RestCollection->EnableNanite &&
@@ -3547,12 +3548,30 @@ void UGeometryCollectionComponent::OnUnregister()
 	}
 }
 
+void UGeometryCollectionComponent::EnableRootProxyStaticMeshComponents(bool bEnabled)
+{
+	if (bEnabled != bEnableRootProxyStaticMeshComponents)
+	{
+		bEnableRootProxyStaticMeshComponents = bEnabled;
+
+		if (bEnabled)
+		{
+			CreateRootProxyComponentsIfNeeded();
+		}
+		else
+		{
+			ClearRootProxyComponents();
+		}	
+		MarkRenderStateDirty();
+	}
+	
+}
+
 bool UGeometryCollectionComponent::ShouldCreateRootProxyComponents() const
 {
 	const bool bHasRootProxyMeshes = RestCollection && RestCollection->RootProxyData.ProxyMeshes.Num() > 0;
-	const bool bHasMeshData = RestCollection && (RestCollection->HasNaniteData() || RestCollection->HasMeshData());
 	const bool bHasCustomRenderer = CanUseCustomRenderer();
-	return bHasRootProxyMeshes && !bHasCustomRenderer && !bHasMeshData;
+	return bHasRootProxyMeshes && !bHasCustomRenderer && bEnableRootProxyStaticMeshComponents;
 }
 
 void UGeometryCollectionComponent::CreateRootProxyComponentsIfNeeded()
diff --git a/Engine/Source/Runtime/Experimental/GeometryCollectionEngine/Public/GeometryCollection/GeometryCollectionComponent.h b/Engine/Source/Runtime/Experimental/GeometryCollectionEngine/Public/GeometryCollection/GeometryCollectionComponent.h
--- a/Engine/Source/Runtime/Experimental/GeometryCollectionEngine/Public/GeometryCollection/GeometryCollectionComponent.h
+++ b/Engine/Source/Runtime/Experimental/GeometryCollectionEngine/Public/GeometryCollection/GeometryCollectionComponent.h
@@ -1327,6 +1327,10 @@ public:
 	/** Get any custom renderer. Returns nullptr if none is set. */
 	GEOMETRYCOLLECTIONENGINE_API IGeometryCollectionExternalRenderInterface* GetCustomRenderer() { return CustomRenderer.GetInterface(); }
 
+	
+	/** Enable or disable root proxy component creation when not using a custom renderer - this can be set at runtime */
+	GEOMETRYCOLLECTIONENGINE_API void EnableRootProxyStaticMeshComponents(bool bEnabled);
+
 	/** Enable or disable root proxy for custom rendering - this can be set at runtime */
 	UFUNCTION(BlueprintCallable, Category = "Physics")
 	GEOMETRYCOLLECTIONENGINE_API void EnableRootProxyForCustomRenderer(bool bEnable);
@@ -1831,6 +1835,7 @@ private:
 	void ClearRootProxyComponents();
 
 	TArray<TObjectPtr<UStaticMeshComponent>> RootProxyStaticMeshComponents;
+	bool bEnableRootProxyStaticMeshComponents = true;
 
 private:
 
