Don't dereference null pointer.

#rnx

diff --git a/Engine/Source/Runtime/Experimental/Iris/Core/Private/Iris/ReplicationSystem/NetBlob/NetRPC.cpp b/Engine/Source/Runtime/Experimental/Iris/Core/Private/Iris/ReplicationSystem/NetBlob/NetRPC.cpp
--- a/Engine/Source/Runtime/Experimental/Iris/Core/Private/Iris/ReplicationSystem/NetBlob/NetRPC.cpp
+++ b/Engine/Source/Runtime/Experimental/Iris/Core/Private/Iris/ReplicationSystem/NetBlob/NetRPC.cpp
@@ -348,15 +348,18 @@ bool FNetRPC::ResolveFunctionAndObject(FNetSerializationContext& Context)
 	Function = FunctionDescriptor->Function;
 
 	// Patch up NetBlobFlags based on function flags.
-	if (Function && ((Function->FunctionFlags & FUNC_NetReliable) != 0))
+	if (Function)
 	{
-		CreationInfo.Flags |= ENetBlobFlags::Reliable;
-	}
+		if ((Function->FunctionFlags & FUNC_NetReliable) != 0)
+		{
+			CreationInfo.Flags |= ENetBlobFlags::Reliable;
+		}
 
-	// The sending side will set Ordered on unicast/reliable RPCs so we're restoring that flag. Unicast RPCs are ordered with respect to other reliable and unicast RPCs whereas multicast RPCs are not.
-	if ((Function->FunctionFlags & FUNC_NetMulticast) == 0)
-	{
-		CreationInfo.Flags |= UE::Net::ENetBlobFlags::Ordered;
+		// The sending side will set Ordered on unicast/reliable RPCs so we're restoring that flag. Unicast RPCs are ordered with respect to other reliable and unicast RPCs whereas multicast RPCs are not.
+		if ((Function->FunctionFlags & FUNC_NetMulticast) == 0)
+		{
+			CreationInfo.Flags |= UE::Net::ENetBlobFlags::Ordered;
+		}
 	}
 
 	// Set the BlobDescriptor even if it has zero size so that we can trace with a meaningful name.
