Proper variant of deprecated FFastArraySerializer::FPostReplicatedReceiveParameters::bHasMoreUnmappedReferences as its behavior is inconsistent and code should not rely on it.

diff --git a/Engine/Plugins/Experimental/UIFramework/Source/Private/Types/UIFWidgetTree.cpp b/Engine/Plugins/Experimental/UIFramework/Source/Private/Types/UIFWidgetTree.cpp
--- a/Engine/Plugins/Experimental/UIFramework/Source/Private/Types/UIFWidgetTree.cpp
+++ b/Engine/Plugins/Experimental/UIFramework/Source/Private/Types/UIFWidgetTree.cpp
@@ -167,9 +167,6 @@ void FUIFrameworkWidgetTree::PostReplicatedChange(const TArrayView<int32>& Chang
 	}
 }
 
-void FUIFrameworkWidgetTree::PostReplicatedReceive(FPostReplicatedReceiveParameters Param)
-{}
-
 bool FUIFrameworkWidgetTree::ReplicateSubWidgets(UActorChannel* Channel, FOutBunch* Bunch, FReplicationFlags* RepFlags)
 {
 	bool bWroteSomething = false;
diff --git a/Engine/Plugins/Experimental/UIFramework/Source/Public/Types/UIFWidgetTree.h b/Engine/Plugins/Experimental/UIFramework/Source/Public/Types/UIFWidgetTree.h
--- a/Engine/Plugins/Experimental/UIFramework/Source/Public/Types/UIFWidgetTree.h
+++ b/Engine/Plugins/Experimental/UIFramework/Source/Public/Types/UIFWidgetTree.h
@@ -75,7 +75,6 @@ public:
 	void PreReplicatedRemove(const TArrayView<int32> RemovedIndices, int32 FinalSize);
 	void PostReplicatedAdd(const TArrayView<int32> AddedIndices, int32 FinalSize);
 	void PostReplicatedChange(const TArrayView<int32>& ChangedIndices, int32 FinalSize);
-	void PostReplicatedReceive(FPostReplicatedReceiveParameters Param);
 	//~ End of FFastArraySerializer
 
 	bool NetDeltaSerialize(FNetDeltaSerializeInfo& DeltaParms)
diff --git a/Engine/Source/Runtime/Experimental/Iris/Core/Public/Iris/ReplicationSystem/Private/FastArrayReplicationFragmentInternal.h b/Engine/Source/Runtime/Experimental/Iris/Core/Public/Iris/ReplicationSystem/Private/FastArrayReplicationFragmentInternal.h
--- a/Engine/Source/Runtime/Experimental/Iris/Core/Public/Iris/ReplicationSystem/Private/FastArrayReplicationFragmentInternal.h
+++ b/Engine/Source/Runtime/Experimental/Iris/Core/Public/Iris/ReplicationSystem/Private/FastArrayReplicationFragmentInternal.h
@@ -83,12 +83,14 @@ struct FFastArrayReplicationFragmentHelper
 	template<typename FastArrayType>
 	static inline typename TEnableIf<TModels_V<FFastArraySerializer::CPostReplicatedReceiveFuncable, FastArrayType, const FFastArraySerializer::FPostReplicatedReceiveParameters&>, void>::Type CallPostReplicatedReceiveOrNot(FastArrayType& ArraySerializer, int32 OldArraySize, bool bHasUnresolvedReferences)
 	{
-		FFastArraySerializer::FPostReplicatedReceiveParameters PostReceivedParameters;
 PRAGMA_DISABLE_DEPRECATION_WARNINGS
-		PostReceivedParameters.bHasMoreUnmappedReferences = bHasUnresolvedReferences;
+		FFastArraySerializer::FPostReplicatedReceiveParameters PostReceivedParameters =
+		{
+			.OldArraySize = OldArraySize,
+			.bHasMoreUnmappedReferences = bHasUnresolvedReferences,
+		};
 PRAGMA_ENABLE_DEPRECATION_WARNINGS
-		PostReceivedParameters.OldArraySize = OldArraySize;
-
+		
 		ArraySerializer.PostReplicatedReceive(PostReceivedParameters);
 	}
 
diff --git a/Engine/Source/Runtime/Net/Core/Classes/Net/Serialization/FastArraySerializer.h b/Engine/Source/Runtime/Net/Core/Classes/Net/Serialization/FastArraySerializer.h
--- a/Engine/Source/Runtime/Net/Core/Classes/Net/Serialization/FastArraySerializer.h
+++ b/Engine/Source/Runtime/Net/Core/Classes/Net/Serialization/FastArraySerializer.h
@@ -520,7 +520,7 @@ struct FFastArraySerializer
 	{
 		// This is the size that the array had before the receive.
 		int32 OldArraySize;
-		//UE_DEPRECATED(5.4, "This is unsafe to use and will be removed.")
+		UE_DEPRECATED(5.4, "This is unsafe to use and will be removed.")
 		uint32 bHasMoreUnmappedReferences : 1U;
 	};
 
@@ -699,9 +699,11 @@ private:
 		inline typename TEnableIf<TModels_V<CPostReplicatedReceiveFuncable, FastArrayType, const FFastArraySerializer::FPostReplicatedReceiveParameters>, void>::Type CallPostReplicatedReceiveOrNot(int32 OldArraySize)
 		{
 PRAGMA_DISABLE_DEPRECATION_WARNINGS
-			FFastArraySerializer::FPostReplicatedReceiveParameters PostReceivedParameters;
-			PostReceivedParameters.OldArraySize = OldArraySize;
-			PostReceivedParameters.bHasMoreUnmappedReferences = Parms.bOutHasMoreUnmapped;
+			FFastArraySerializer::FPostReplicatedReceiveParameters PostReceivedParameters =
+			{
+				.OldArraySize = OldArraySize,
+				.bHasMoreUnmappedReferences = Parms.bOutHasMoreUnmapped,
+			};
 			ArraySerializer.PostReplicatedReceive(PostReceivedParameters);
 PRAGMA_ENABLE_DEPRECATION_WARNINGS
 		}
