ChaosFlesh: Adding HeadlessChaosFleshCollisionTest

diff --git a/Engine/Source/Runtime/Experimental/Chaos/Private/Chaos/GraphColoring.cpp b/Engine/Source/Runtime/Experimental/Chaos/Private/Chaos/GraphColoring.cpp
--- a/Engine/Source/Runtime/Experimental/Chaos/Private/Chaos/GraphColoring.cpp
+++ b/Engine/Source/Runtime/Experimental/Chaos/Private/Chaos/GraphColoring.cpp
@@ -358,12 +358,9 @@ static bool VerifyNodalColoring(const TArray<TArray<int32>>& Graph, const Chaos:
 
 	TArray<int32> Particle2Incident;
 	Particle2Incident.Init(INDEX_NONE, GraphParticlesEnd - GraphParticlesStart);
-	for (int32 i = 0; i < IncidentElements.Num(); i++)
+	for (int32 i = 0; i < GraphParticlesEnd - GraphParticlesStart; i++)
 	{
-		if (IncidentElements[i].Num() > 0)
-		{
-			Particle2Incident[Graph[IncidentElements[i][0]][IncidentElementsLocalIndex[i][0]]] = i;
-		}
+		Particle2Incident[i] = i + GraphParticlesStart;
 	}
 
 	for (int32 i = 0; i < ParticlesPerColor.Num(); i++)
@@ -1604,17 +1601,15 @@ TArray<TArray<int32>> Chaos::ComputeNodalColoring(const TArray<TArray<int32>>& G
 
 	checkSlow(GraphParticlesStart <= GraphParticlesEnd);
 	checkSlow(GraphParticlesEnd <= (int32)InParticles.Size());
+	checkSlow(InParticles.Size() == IncidentElements.Num());
 	TArray<TArray<int32>> ParticlesPerColor;
 
 	TArray<int32> Particle2Incident;
 	Particle2Incident.Init(INDEX_NONE, GraphParticlesEnd - GraphParticlesStart);
 	//Assuming that offset of Graph is GraphParticlesStart
-	for (int32 i = 0; i < IncidentElements.Num(); i++) 
+	for (int32 i = 0; i < GraphParticlesEnd-GraphParticlesStart; i++)
 	{
-		if (IncidentElements[i].Num() > 0)
-		{
-			Particle2Incident[Graph[IncidentElements[i][0]][IncidentElementsLocalIndex[i][0]] - GraphParticlesStart] = i;
-		}
+		Particle2Incident[i] = i+GraphParticlesStart;
 	}
 
 	TArray<TSet<int32>*> ElementColorsSet;
@@ -1774,9 +1769,9 @@ void Chaos::ComputeExtraNodalColoring(const TArray<TArray<int32>>& Graph, const
 				}
 			}
 			ParticlesPerColor[i].SetNum(CurrentIndex);
-		}, ParticlesPerColor[0].Num() < 1000);
+		}, ParticlesPerColor.Num() > 0 && ParticlesPerColor[0].Num() < 1000);
 
-	ParticlesPerColor.SetNum(ParticleColors.Max() + 1);
+	ParticlesPerColor.SetNum(FMath::Max<int32>(ParticleColors) + 1);
 	for (int32 i = 0; i < ExtraIncidentElements.Num(); i++)
 	{
 		if (ParticleIsAffected[i] && ParticleColors[i] != INDEX_NONE)
