Chaos - Fix character ground constraint drift code to avoid simulated proxies jittering

diff --git a/Engine/Source/Runtime/Experimental/Chaos/Private/Chaos/Character/CharacterGroundConstraintContainer.cpp b/Engine/Source/Runtime/Experimental/Chaos/Private/Chaos/Character/CharacterGroundConstraintContainer.cpp
--- a/Engine/Source/Runtime/Experimental/Chaos/Private/Chaos/Character/CharacterGroundConstraintContainer.cpp
+++ b/Engine/Source/Runtime/Experimental/Chaos/Private/Chaos/Character/CharacterGroundConstraintContainer.cpp
@@ -7,6 +7,15 @@
 
 namespace Chaos
 {
+	namespace CVars
+	{
+		float Chaos_CharacterGroundConstraint_InputMovementThreshold = 0.1f;
+		FAutoConsoleVariableRef CVarChaos_CharacterGroundConstraint_InputMovementThreshold(TEXT("p.Chaos.CharacterGroundConstraint.InputMovementThreshold"), Chaos_CharacterGroundConstraint_InputMovementThreshold, TEXT("Minimum per frame input movement distance in cm."));
+		
+		float Chaos_CharacterGroundConstraint_ExternalMovementThreshold = 2.0f;
+		FAutoConsoleVariableRef CVarChaos_CharacterGroundConstraint_ExternalMovementThreshold(TEXT("p.Chaos.CharacterGroundConstraint.ExternalMovementThreshold"), Chaos_CharacterGroundConstraint_ExternalMovementThreshold, TEXT("If distance moved is less than this then retain current movement target relative to ground."));
+	}
+	
 	FCharacterGroundConstraintContainer::FCharacterGroundConstraintContainer()
 		: Base(FCharacterGroundConstraintHandle::StaticType())
 		, ConstraintPool(16, 0)
diff --git a/Engine/Source/Runtime/Experimental/Chaos/Public/Chaos/Character/CharacterGroundConstraintContainer.h b/Engine/Source/Runtime/Experimental/Chaos/Public/Chaos/Character/CharacterGroundConstraintContainer.h
--- a/Engine/Source/Runtime/Experimental/Chaos/Public/Chaos/Character/CharacterGroundConstraintContainer.h
+++ b/Engine/Source/Runtime/Experimental/Chaos/Public/Chaos/Character/CharacterGroundConstraintContainer.h
@@ -16,6 +16,12 @@ namespace ChaosTest
 
 namespace Chaos
 {
+	namespace CVars
+	{
+		CHAOS_API extern float Chaos_CharacterGroundConstraint_InputMovementThreshold;
+		CHAOS_API extern float Chaos_CharacterGroundConstraint_ExternalMovementThreshold;
+	}
+
 	class FCharacterGroundConstraintContainer;
 	namespace Private
 	{
@@ -66,12 +72,16 @@ namespace Chaos
 				return;
 			}
 
-			// If the movement target is close to zero then it gets clamped to zero and we
-			// recompute the delta position based on the previous target to avoid drift
-			constexpr float MovementThresholdSq = 0.01f;
-			if (InData.TargetDeltaPosition.SizeSquared() > MovementThresholdSq)
+			FVec3 NewLocalCharacterPosition = ComputeLocalCharacterPosition();
+
+			// If the movement target is close to zero and the character position has not
+			// changed by much then it gets clamped to zero and we recompute the delta
+			// position based on the previous target to avoid drift
+			const float InputMovementThresholdSq = CVars::Chaos_CharacterGroundConstraint_InputMovementThreshold * CVars::Chaos_CharacterGroundConstraint_InputMovementThreshold;
+			const float MovementThresholdSq = CVars::Chaos_CharacterGroundConstraint_ExternalMovementThreshold * CVars::Chaos_CharacterGroundConstraint_ExternalMovementThreshold;
+			if ((InData.TargetDeltaPosition.SizeSquared() > InputMovementThresholdSq) || ((NewLocalCharacterPosition - LocalCharacterPosition).SizeSquared() > MovementThresholdSq))
 			{
-				ComputeLocalCharacterPosition();
+				LocalCharacterPosition = NewLocalCharacterPosition;
 			}
 			else
 			{
@@ -104,7 +114,7 @@ namespace Chaos
 				{
 					GroundParticle->AddConstraintHandle(this);
 				}
-				ComputeLocalCharacterPosition();
+				LocalCharacterPosition = ComputeLocalCharacterPosition();
 				bGroundParticleChanged = true;
 			}
 		}
@@ -121,19 +131,21 @@ namespace Chaos
 		friend class Private::FCharacterGroundConstraintContainerSolver; // For setting the solver force and torque
 		friend class ChaosTest::CharacterGroundConstraintContainerTest; // For testing internals
 
-		void ComputeLocalCharacterPosition()
+		FVec3 ComputeLocalCharacterPosition()
 		{
 			if (!CharacterParticle || bDisabled)
 			{
-				return;
+				return FVec3::ZeroVector;
 			}
 
-			LocalCharacterPosition = CharacterParticle->GetX();
+			FVec3 LocalPos = CharacterParticle->GetX();
 
 			if (GroundParticle)
 			{
-				LocalCharacterPosition = GroundParticle->GetR().Inverse() * (LocalCharacterPosition - GroundParticle->GetX());
+				LocalPos = GroundParticle->GetR().Inverse() * (LocalPos - GroundParticle->GetX());
 			}
+
+			return LocalPos;
 		}
 
 		FCharacterGroundConstraintSettings Settings;
