Fix scene query crash when rewind/resim is in use

During rewind/resim we run the solver again to resimulate the required frames. This was invoking the game thread only callbacks on a worker thread. This led to a game thread only flag being set and queries from the game thread being led to believe they were supposed to query against physics thread data.
For now this change disables the ability for the resim step to call game thread callbacks and added some extra safety to ensure we never flip the frozen GT flag on a worker thread in future.

diff --git a/Engine/Source/Runtime/Engine/Private/PhysicsEngine/Experimental/PhysScene_Chaos.cpp b/Engine/Source/Runtime/Engine/Private/PhysicsEngine/Experimental/PhysScene_Chaos.cpp
--- a/Engine/Source/Runtime/Engine/Private/PhysicsEngine/Experimental/PhysScene_Chaos.cpp
+++ b/Engine/Source/Runtime/Engine/Private/PhysicsEngine/Experimental/PhysScene_Chaos.cpp
@@ -152,6 +152,9 @@ public:
 			const FPendingAsyncPhysicsCommand& PendingCommand = PendingCommands[Idx];
 			bool bRemove = PendingCommand.OwningObject.IsStale() || !PendingCommand.Command;
 
+			/* #TODO implement and re-enable resim commands. This callback must run on the main thread and resim currently does
+			 * not defer its callbacks to the main thread making its execution unsafe. 
+
 			if (!bRemove && bAllowResim && PendingCommand.bEnableResim && PendingCommand.PhysicsStep > (CurrentFrame - NumFrames))
 			{
 				if (PendingCommand.PhysicsStep < CurrentFrame)
@@ -172,6 +175,7 @@ public:
 				}
 			}
 			else
+			*/
 			{
 				if (!bRemove && PendingCommand.PhysicsStep <= CurrentFrame)
 				{
diff --git a/Engine/Source/Runtime/Experimental/Chaos/Private/Chaos/Framework/PhysicsSolverBase.cpp b/Engine/Source/Runtime/Experimental/Chaos/Private/Chaos/Framework/PhysicsSolverBase.cpp
--- a/Engine/Source/Runtime/Experimental/Chaos/Private/Chaos/Framework/PhysicsSolverBase.cpp
+++ b/Engine/Source/Runtime/Experimental/Chaos/Private/Chaos/Framework/PhysicsSolverBase.cpp
@@ -562,4 +562,10 @@ namespace Chaos
 		GTPreSimCallbacks.GTPreSimCallbacks();
 		AdvanceTask.AdvanceSolver();
 	}
+
+	void FSolverTasksPTOnly::AdvanceSolver()
+	{
+		ProcessPushData.ProcessPushData();
+		AdvanceTask.AdvanceSolver();
+	}
 }
diff --git a/Engine/Source/Runtime/Experimental/Chaos/Private/PBDRigidsSolver.cpp b/Engine/Source/Runtime/Experimental/Chaos/Private/PBDRigidsSolver.cpp
--- a/Engine/Source/Runtime/Experimental/Chaos/Private/PBDRigidsSolver.cpp
+++ b/Engine/Source/Runtime/Experimental/Chaos/Private/PBDRigidsSolver.cpp
@@ -2030,8 +2030,9 @@ namespace Chaos
 							}
 #endif
 							MRewindCallback->PreResimStep_Internal(Step, bFirst);
-							
-							FAllSolverTasks ImmediateTask(*this, PushData);
+
+							// Run the advance but omit the game thread callbacks as we're executing during the physics tick
+							FSolverTasksPTOnly ImmediateTask(*this, PushData);
 							//ensure(bSolverHasFrozenGameThreadCallbacks == false);	//We don't support this for resim as it's very expensive and difficult to schedule
 							ImmediateTask.AdvanceSolver();
 							MRewindCallback->PostResimStep_Internal(Step);
diff --git a/Engine/Source/Runtime/Experimental/Chaos/Public/Chaos/Framework/PhysicsSolverBase.h b/Engine/Source/Runtime/Experimental/Chaos/Public/Chaos/Framework/PhysicsSolverBase.h
--- a/Engine/Source/Runtime/Experimental/Chaos/Public/Chaos/Framework/PhysicsSolverBase.h
+++ b/Engine/Source/Runtime/Experimental/Chaos/Public/Chaos/Framework/PhysicsSolverBase.h
@@ -168,6 +168,7 @@ namespace Chaos
 #endif
 	};
 
+	// Container for all steps required to fully update a solver
 	struct FAllSolverTasks
 	{
 		FAllSolverTasks(FPhysicsSolverBase& InSolver, FPushPhysicsData* PushData)
@@ -187,6 +188,30 @@ namespace Chaos
 		FPhysicsSolverBase& Solver;
 	};
 
+	// Container for all physics-thread steps required to update a solver
+	// This leaves out the game thread callbacks for situations that require and update only in a physics-thread context
+	struct FSolverTasksPTOnly
+	{
+		FSolverTasksPTOnly() = delete;
+		FSolverTasksPTOnly(const FSolverTasksPTOnly&) = delete;
+		FSolverTasksPTOnly(FSolverTasksPTOnly&&) = delete;
+		FSolverTasksPTOnly& operator=(const FSolverTasksPTOnly&) = delete;
+		FSolverTasksPTOnly& operator=(FSolverTasksPTOnly&&) = delete;
+
+		FSolverTasksPTOnly(FPhysicsSolverBase& InSolver, FPushPhysicsData* InPushData)
+			: ProcessPushData(InSolver, InPushData)
+			, AdvanceTask(InSolver, InPushData)
+			, Solver(InSolver)
+		{}
+
+		FPhysicsSolverProcessPushDataTask ProcessPushData;
+		FPhysicsSolverAdvanceTask AdvanceTask;
+
+		CHAOS_API void AdvanceSolver();
+
+		FPhysicsSolverBase& Solver;
+	};
+
 	class FPersistentPhysicsTask;
 
 	enum class ELockType: uint8;
@@ -530,6 +555,7 @@ namespace Chaos
 
 		void SetGameThreadFrozen(bool InGameThreadFrozen)
 		{
+			check(IsInGameThread());
 			bGameThreadFrozen = InGameThreadFrozen;
 		}
 
