Chaos Visual Debugger : Fix concurrent access to CachedGeometryHashes structure
- make sure the right write lockls are used

diff --git a/Engine/Source/Runtime/Experimental/Chaos/Private/ChaosVisualDebugger/ChaosVisualDebuggerTrace.cpp b/Engine/Source/Runtime/Experimental/Chaos/Private/ChaosVisualDebugger/ChaosVisualDebuggerTrace.cpp
--- a/Engine/Source/Runtime/Experimental/Chaos/Private/ChaosVisualDebugger/ChaosVisualDebuggerTrace.cpp
+++ b/Engine/Source/Runtime/Experimental/Chaos/Private/ChaosVisualDebugger/ChaosVisualDebuggerTrace.cpp
@@ -93,7 +93,7 @@ struct FChaosVDGeometryTraceContext
 			return;
 		}
 
-		FWriteScopeLock WriteLock(TracedGeometrySetLock);
+		FWriteScopeLock WriteLock(CachedGeometryHashesLock);
 		CachedGeometryHashes.Remove((void*)Implicit);
 	}
 
@@ -819,6 +819,8 @@ void FChaosVisualDebuggerTrace::Reset()
 	{
 		FWriteScopeLock GeometryWriteLock(GeometryTracerObject.TracedGeometrySetLock);
 		GeometryTracerObject.GeometryTracedIDs.Reset();
+
+		FWriteScopeLock GeometryHashWriteLock(GeometryTracerObject.CachedGeometryHashesLock);
 		GeometryTracerObject.CachedGeometryHashes.Reset();
 	}
 }
