Add OnActorVisibilityChanged virtual function for UActorComponent.
SetActorHiddenInGame() now ends up calling this.
Override on GeometryCollectionComponent to update any custom renderer.

diff --git a/Engine/Source/Runtime/Engine/Classes/Components/ActorComponent.h b/Engine/Source/Runtime/Engine/Classes/Components/ActorComponent.h
--- a/Engine/Source/Runtime/Engine/Classes/Components/ActorComponent.h
+++ b/Engine/Source/Runtime/Engine/Classes/Components/ActorComponent.h
@@ -1017,6 +1017,9 @@ public:
 	/** See if this component is in the persistent level */
 	ENGINE_API bool ComponentIsInPersistentLevel(bool bIncludeLevelStreamingPersistent) const;
 
+	/** Called on each component when the Actor's visibility state changes */
+	virtual void OnActorVisibilityChanged() { MarkRenderStateDirty(); }
+
 	/** Called on each component when the Actor's bEnableCollisionChanged flag changes */
 	virtual void OnActorEnableCollisionChanged() {}
 
diff --git a/Engine/Source/Runtime/Engine/Classes/GameFramework/Actor.h b/Engine/Source/Runtime/Engine/Classes/GameFramework/Actor.h
--- a/Engine/Source/Runtime/Engine/Classes/GameFramework/Actor.h
+++ b/Engine/Source/Runtime/Engine/Classes/GameFramework/Actor.h
@@ -3088,6 +3088,9 @@ public:
 	/** Update all components transforms */
 	ENGINE_API void UpdateComponentTransforms();
 
+	/** Update all components visibility state */
+	ENGINE_API void UpdateComponentVisibility();
+
 	/** Iterate over components array and call InitializeComponent, which happens once per actor */
 	ENGINE_API void InitializeComponents();
 
diff --git a/Engine/Source/Runtime/Engine/Private/Actor.cpp b/Engine/Source/Runtime/Engine/Private/Actor.cpp
--- a/Engine/Source/Runtime/Engine/Private/Actor.cpp
+++ b/Engine/Source/Runtime/Engine/Private/Actor.cpp
@@ -4720,7 +4720,7 @@ void AActor::SetActorHiddenInGame( bool bNewHidden )
 	if (IsHidden() != bNewHidden)
 	{
 		SetHidden(bNewHidden);
-		MarkComponentsRenderStateDirty();
+		UpdateComponentVisibility();
 	}
 }
 
@@ -5540,6 +5540,24 @@ void AActor::UpdateComponentTransforms()
 	}
 }
 
+void AActor::UpdateComponentVisibility()
+{
+	for (UActorComponent* ActorComp : GetComponents())
+	{
+		if (ActorComp && ActorComp->IsRegistered())
+		{
+			ActorComp->OnActorVisibilityChanged();
+			if (UChildActorComponent* ChildActorComponent = Cast<UChildActorComponent>(ActorComp))
+			{
+				if (ChildActorComponent->GetChildActor())
+				{
+					ChildActorComponent->GetChildActor()->UpdateComponentVisibility();
+				}
+			}
+		}
+	}
+}
+
 void AActor::MarkComponentsRenderStateDirty()
 {
 	for (UActorComponent* ActorComp : GetComponents())
diff --git a/Engine/Source/Runtime/Experimental/GeometryCollectionEngine/Private/GeometryCollection/GeometryCollectionComponent.cpp b/Engine/Source/Runtime/Experimental/GeometryCollectionEngine/Private/GeometryCollection/GeometryCollectionComponent.cpp
--- a/Engine/Source/Runtime/Experimental/GeometryCollectionEngine/Private/GeometryCollection/GeometryCollectionComponent.cpp
+++ b/Engine/Source/Runtime/Experimental/GeometryCollectionEngine/Private/GeometryCollection/GeometryCollectionComponent.cpp
@@ -753,7 +753,18 @@ void UGeometryCollectionComponent::EndPlay(const EEndPlayReason::Type ReasonEnd)
 void UGeometryCollectionComponent::OnVisibilityChanged()
 {
 	Super::OnVisibilityChanged();
+	RefreshCustomRenderer();
+}
+
+void UGeometryCollectionComponent::OnActorVisibilityChanged()
+{
+	Super::OnVisibilityChanged();
+	RefreshCustomRenderer();
+}
 
+void UGeometryCollectionComponent::OnHiddenInGameChanged()
+{
+	Super::OnHiddenInGameChanged();
 	RefreshCustomRenderer();
 }
 
@@ -3497,12 +3508,6 @@ void UGeometryCollectionComponent::AsyncPhysicsTickComponent(float DeltaTime, fl
 	UpdateRepData();
 }
 
-void UGeometryCollectionComponent::OnHiddenInGameChanged()
-{
-	Super::OnHiddenInGameChanged();
-	RefreshCustomRenderer();
-}
-
 void UGeometryCollectionComponent::OnRegister()
 {
 	//UE_LOG(UGCC_LOG, Log, TEXT("GeometryCollectionComponent[%p]::OnRegister()[%p]"), this,RestCollection );
@@ -6272,7 +6277,7 @@ void UGeometryCollectionComponent::RefreshCustomRenderer()
 					const bool bRenderRootProxy = bEnableRootProxyForCustomRenderer && !bIsBroken && (RestCollection->RootProxyData.ProxyMeshes.Num() > 0);
 
 					uint32 StateFlags = 0;
-					StateFlags |= bHiddenInGame || !IsVisible() ? 0 : IGeometryCollectionExternalRenderInterface::EState_Visible;
+					StateFlags |= bHiddenInGame || !IsVisible() || GetOwner()->IsHidden() ? 0 : IGeometryCollectionExternalRenderInterface::EState_Visible;
 					StateFlags |= bRenderRootProxy ? 0 : IGeometryCollectionExternalRenderInterface::EState_Broken;
 					StateFlags |= bEnableRootProxyForCustomRenderer ? 0 : IGeometryCollectionExternalRenderInterface::EState_ForcedBroken;
 
diff --git a/Engine/Source/Runtime/Experimental/GeometryCollectionEngine/Public/GeometryCollection/GeometryCollectionComponent.h b/Engine/Source/Runtime/Experimental/GeometryCollectionEngine/Public/GeometryCollection/GeometryCollectionComponent.h
--- a/Engine/Source/Runtime/Experimental/GeometryCollectionEngine/Public/GeometryCollection/GeometryCollectionComponent.h
+++ b/Engine/Source/Runtime/Experimental/GeometryCollectionEngine/Public/GeometryCollection/GeometryCollectionComponent.h
@@ -610,6 +610,8 @@ public:
 	GEOMETRYCOLLECTIONENGINE_API virtual void BeginPlay() override;
 	GEOMETRYCOLLECTIONENGINE_API virtual void EndPlay(const EEndPlayReason::Type ReasonEnd) override;
 	GEOMETRYCOLLECTIONENGINE_API virtual void OnVisibilityChanged() override;
+	GEOMETRYCOLLECTIONENGINE_API virtual void OnActorVisibilityChanged() override;
+	GEOMETRYCOLLECTIONENGINE_API virtual void OnHiddenInGameChanged() override;
 	GEOMETRYCOLLECTIONENGINE_API virtual void GetLifetimeReplicatedProps(TArray<FLifetimeProperty>& OutLifetimeProps) const override;
 	GEOMETRYCOLLECTIONENGINE_API virtual void GetResourceSizeEx(FResourceSizeEx& CumulativeResourceSize) override;
 
@@ -642,8 +644,6 @@ public:
 	
 	GEOMETRYCOLLECTIONENGINE_API virtual void TickComponent(float DeltaTime, enum ELevelTick TickType, FActorComponentTickFunction *ThisTickFunction) override;
 	GEOMETRYCOLLECTIONENGINE_API virtual void AsyncPhysicsTickComponent(float DeltaTime, float SimTime) override;
-
-	GEOMETRYCOLLECTIONENGINE_API virtual void OnHiddenInGameChanged() override;
 	//~ End USceneComponent Interface.
 
 
