Geometry Collection : Fix potential crash in FillAutoInstanceMeshesInstancesIfNeeded if the index attribute is not fopund but the asset has auto instance meshes assigned

diff --git a/Engine/Source/Runtime/Experimental/GeometryCollectionEngine/Private/GeometryCollection/GeometryCollectionObject.cpp b/Engine/Source/Runtime/Experimental/GeometryCollectionEngine/Private/GeometryCollection/GeometryCollectionObject.cpp
--- a/Engine/Source/Runtime/Experimental/GeometryCollectionEngine/Private/GeometryCollection/GeometryCollectionObject.cpp
+++ b/Engine/Source/Runtime/Experimental/GeometryCollectionEngine/Private/GeometryCollection/GeometryCollectionObject.cpp
@@ -1285,7 +1285,6 @@ void UGeometryCollection::FillAutoInstanceMeshesInstancesIfNeeded()
 		}
 
 		const GeometryCollection::Facades::FCollectionInstancedMeshFacade InstancedMeshFacade(*GeometryCollection);
-
 		if (InstancedMeshFacade.IsValid())
 		{
 			const int32 NumTransforms = GeometryCollection->Children.Num();
@@ -1304,7 +1303,7 @@ void UGeometryCollection::FillAutoInstanceMeshesInstancesIfNeeded()
 		}
 		else
 		{
-			UE_LOG(LogGeometryCollectionInternal, Warning, TEXT("Geometry collection has AutoInstanceMeshes meshes but no matching AutoInstanceMeshIndex, remove the instanced meshes or regenerate the asset"), *GetPathName());
+			UE_LOG(LogGeometryCollectionInternal, Warning, TEXT("[%s] Could not find AutoInstanceMeshIndex attribute but the asset as instanced meshes assigned, you may need to regenerate this asset"), *GetPathName());
 		}
 	}
 }
