Fix crash with geometry collection component trying to create their static mesh component for root proxies while not having a an owner yet

diff --git a/Engine/Source/Runtime/Experimental/GeometryCollectionEngine/Private/GeometryCollection/GeometryCollectionComponent.cpp b/Engine/Source/Runtime/Experimental/GeometryCollectionEngine/Private/GeometryCollection/GeometryCollectionComponent.cpp
--- a/Engine/Source/Runtime/Experimental/GeometryCollectionEngine/Private/GeometryCollection/GeometryCollectionComponent.cpp
+++ b/Engine/Source/Runtime/Experimental/GeometryCollectionEngine/Private/GeometryCollection/GeometryCollectionComponent.cpp
@@ -3583,7 +3583,7 @@ void UGeometryCollectionComponent::CreateRootProxyComponentsIfNeeded()
 {
 	ClearRootProxyComponents();
 
-	if (ShouldCreateRootProxyComponents())
+	if (ShouldCreateRootProxyComponents() && GetOwner())
 	{
 		for (const TObjectPtr<UStaticMesh>& ProxyMesh : RestCollection->RootProxyData.ProxyMeshes)
 		{
