Fix geometry collection not showing destruction
Fallout from 31782118

diff --git a/Engine/Source/Runtime/Experimental/GeometryCollectionEngine/Private/GeometryCollection/GeometryCollectionComponent.cpp b/Engine/Source/Runtime/Experimental/GeometryCollectionEngine/Private/GeometryCollection/GeometryCollectionComponent.cpp
--- a/Engine/Source/Runtime/Experimental/GeometryCollectionEngine/Private/GeometryCollection/GeometryCollectionComponent.cpp
+++ b/Engine/Source/Runtime/Experimental/GeometryCollectionEngine/Private/GeometryCollection/GeometryCollectionComponent.cpp
@@ -6277,7 +6277,7 @@ void UGeometryCollectionComponent::RefreshCustomRenderer()
 					const bool bRenderRootProxy = bEnableRootProxyForCustomRenderer && !bIsBroken && (RestCollection->RootProxyData.ProxyMeshes.Num() > 0);
 
 					uint32 StateFlags = 0;
-					StateFlags |= bHiddenInGame || !IsVisible() || GetOwner()->IsHidden() ? 0 : IGeometryCollectionExternalRenderInterface::EState_Visible;
+					StateFlags |= bHiddenInGame || !IsVisible() ? 0 : IGeometryCollectionExternalRenderInterface::EState_Visible;
 					StateFlags |= bRenderRootProxy ? 0 : IGeometryCollectionExternalRenderInterface::EState_Broken;
 					StateFlags |= bEnableRootProxyForCustomRenderer ? 0 : IGeometryCollectionExternalRenderInterface::EState_ForcedBroken;
 
