Iris - Fixed first partial blob bug causing expected sequences to not be ready to assemble.

#rnx

diff --git a/Engine/Source/Runtime/Experimental/Iris/Core/Private/Iris/ReplicationSystem/NetBlob/NetBlobAssembler.cpp b/Engine/Source/Runtime/Experimental/Iris/Core/Private/Iris/ReplicationSystem/NetBlob/NetBlobAssembler.cpp
--- a/Engine/Source/Runtime/Experimental/Iris/Core/Private/Iris/ReplicationSystem/NetBlob/NetBlobAssembler.cpp
+++ b/Engine/Source/Runtime/Experimental/Iris/Core/Private/Iris/ReplicationSystem/NetBlob/NetBlobAssembler.cpp
@@ -54,12 +54,7 @@ void FNetBlobAssembler::AddPartialNetBlob(FNetSerializationContext& Context, FNe
 		}
 
 		// Gracefully handle going from unreliable to first part of blob, regardless of its reliability.
-		if (bIsFirstPart)
-		{
-			NextSequenceSumber = SequenceNumber;
-			LastPartSequenceNumber = SequenceNumber + PartialNetBlob->GetPartCount() - 1U;
-		}
-		else if (SequenceNumber != NextSequenceSumber)
+		if (!bIsFirstPart && SequenceNumber != NextSequenceSumber)
 		{
 			bIsBrokenSequence = true;
 			if (bIsReliable)
@@ -75,6 +70,9 @@ void FNetBlobAssembler::AddPartialNetBlob(FNetSerializationContext& Context, FNe
 	if (bIsFirstPart)
 	{
 		const uint32 PartCount = PartialNetBlob->GetPartCount();
+		NextSequenceSumber = SequenceNumber;
+		LastPartSequenceNumber = SequenceNumber + PartialNetBlob->GetPartCount() - 1U;
+
 		if (PartCount == 0)
 		{
 			bIsBrokenSequence = true;
