Iris
* Added and moved some CPU scopes to reduce verbosity and breakdown some phases.

#rnx

diff --git a/Engine/Source/Runtime/Engine/Private/NetDriver.cpp b/Engine/Source/Runtime/Engine/Private/NetDriver.cpp
--- a/Engine/Source/Runtime/Engine/Private/NetDriver.cpp
+++ b/Engine/Source/Runtime/Engine/Private/NetDriver.cpp
@@ -6812,6 +6812,8 @@ void UNetDriver::UpdateIrisReplicationViews() const
 {
 	using namespace UE::Net;
 
+	QUICK_SCOPE_CYCLE_COUNTER(STAT_UpdateIrisReplicationViews);
+
 	FReplicationView::FView DefaultView;
 	{
 		const APlayerCameraManager* CameraManager = GetDefault<APlayerCameraManager>();
@@ -6901,6 +6903,8 @@ void UNetDriver::UpdateIrisReplicationViews() const
 
 void UNetDriver::SendClientMoveAdjustments()
 {
+	QUICK_SCOPE_CYCLE_COUNTER(STAT_SendClientMoveAdjustments);
+
 	for (UNetConnection* Connection : ClientConnections)
 	{
 		if (Connection == nullptr || Connection->ViewTarget == nullptr)
diff --git a/Engine/Source/Runtime/Experimental/Iris/Core/Private/Iris/ReplicationSystem/ReplicationSystem.cpp b/Engine/Source/Runtime/Experimental/Iris/Core/Private/Iris/ReplicationSystem/ReplicationSystem.cpp
--- a/Engine/Source/Runtime/Experimental/Iris/Core/Private/Iris/ReplicationSystem/ReplicationSystem.cpp
+++ b/Engine/Source/Runtime/Experimental/Iris/Core/Private/Iris/ReplicationSystem/ReplicationSystem.cpp
@@ -330,24 +330,31 @@ public:
 
 	void UpdateFilterPostPoll()
 	{
-		IRIS_PROFILER_SCOPE(FReplicationSystem_UpdateFilterPostPoll);
 		LLM_SCOPE_BYTAG(Iris);
 
 		FReplicationFiltering& Filtering = ReplicationSystemInternal.GetFiltering();
-		Filtering.FilterPostPoll();
 
-		// Iterate over all valid connections and propagate updated scopes
-		FReplicationConnections& Connections = ReplicationSystemInternal.GetConnections();
-		
-		auto UpdateConnectionScope = [&Filtering, &Connections](uint32 ConnectionId)
 		{
-			FReplicationConnection* Conn = Connections.GetConnection(ConnectionId);
-			const FNetBitArrayView ObjectsInScope = Filtering.GetRelevantObjectsInScope(ConnectionId);
-			Conn->ReplicationWriter->UpdateScope(ObjectsInScope);
-		};
+			IRIS_PROFILER_SCOPE(FReplicationSystem_UpdateFilterPostPoll);
+			Filtering.FilterPostPoll();
+		}
 
-		const FNetBitArray& ValidConnections = Connections.GetValidConnections();
-		ValidConnections.ForAllSetBits(UpdateConnectionScope);
+		{
+			IRIS_PROFILER_SCOPE(FReplicationSystem_UpdateConnectionsScope);
+		
+			// Iterate over all valid connections and propagate updated scopes
+			FReplicationConnections& Connections = ReplicationSystemInternal.GetConnections();
+		
+			auto UpdateConnectionScope = [&Filtering, &Connections](uint32 ConnectionId)
+			{
+				FReplicationConnection* Conn = Connections.GetConnection(ConnectionId);
+				const FNetBitArrayView ObjectsInScope = Filtering.GetRelevantObjectsInScope(ConnectionId);
+				Conn->ReplicationWriter->UpdateScope(ObjectsInScope);
+			};
+
+			const FNetBitArray& ValidConnections = Connections.GetValidConnections();
+			ValidConnections.ForAllSetBits(UpdateConnectionScope);
+		}
 	}
 
 	// Can run at any time between scoping and replication.
diff --git a/Engine/Source/Runtime/Experimental/Iris/Core/Private/Iris/ReplicationSystem/ReplicationWriter.cpp b/Engine/Source/Runtime/Experimental/Iris/Core/Private/Iris/ReplicationSystem/ReplicationWriter.cpp
--- a/Engine/Source/Runtime/Experimental/Iris/Core/Private/Iris/ReplicationSystem/ReplicationWriter.cpp
+++ b/Engine/Source/Runtime/Experimental/Iris/Core/Private/Iris/ReplicationSystem/ReplicationWriter.cpp
@@ -632,7 +632,7 @@ void FReplicationWriter::SetPendingDestroyOrSubObjectPendingDestroyState(uint32
 
 void FReplicationWriter::UpdateScope(const FNetBitArrayView& UpdatedScope)
 {
-	IRIS_PROFILER_SCOPE(FReplicationWriter_ScopeUpdate);
+	//IRIS_PROFILER_SCOPE(FReplicationWriter_ScopeUpdate);
 
 	auto NewObjectFunctor = [this](uint32 Index)
 	{
@@ -646,7 +646,7 @@ void FReplicationWriter::UpdateScope(const FNetBitArrayView& UpdatedScope)
 		}
 		else if (State == EReplicatedObjectState::WaitOnFlush)
 		{			
-			if (ensureAlwaysMsgf(!Info.TearOff, TEXT("We cannot cancel flush for an object pending tearoff ( InternalIndex: %u )"), Index))
+			if (ensureMsgf(!Info.TearOff, TEXT("We cannot cancel flush for an object pending tearoff ( InternalIndex: %u )"), Index))
 			{
 				// If we are waiting on flush but are re-added to scope we reset flush flags to default.
 				ObjectsPendingDestroy.ClearBit(Index);
@@ -687,7 +687,7 @@ void FReplicationWriter::UpdateScope(const FNetBitArrayView& UpdatedScope)
 				if (OwnerInfo.GetState() < EReplicatedObjectState::PendingDestroy)
 				{
 					ensureMsgf(!bValidateObjectsWithDirtyChanges || OwnerInfo.GetState() != EReplicatedObjectState::Invalid, TEXT("Object ( InternalIndex: %u ) with Invalid state potentially marked dirty."), ObjectData.SubObjectRootIndex);
-					ensureAlwaysMsgf(!OwnerInfo.TearOff, TEXT("Parent is tearing off ( InternalIndex: %u ) currently in State: %s "), ObjectData.SubObjectRootIndex, LexToString(OwnerInfo.GetState()));
+					ensureMsgf(!OwnerInfo.TearOff, TEXT("Parent is tearing off ( InternalIndex: %u ) currently in State: %s "), ObjectData.SubObjectRootIndex, LexToString(OwnerInfo.GetState()));
 					OwnerInfo.HasDirtySubObjects |= Info.HasDirtyChangeMask;
 					ObjectsWithDirtyChanges.SetBitValue(ObjectData.SubObjectRootIndex, ObjectsWithDirtyChanges.GetBit(ObjectData.SubObjectRootIndex) || Info.HasDirtyChangeMask);
 				}
@@ -721,8 +721,8 @@ void FReplicationWriter::UpdateScope(const FNetBitArrayView& UpdatedScope)
 		{
 			UE_LOG_REPLICATIONWRITER_CONN(TEXT("New object added to scope, Waiting to start replication for ( InternalIndex: %u ) currently in State: %s "), Index, LexToString(State));
 
-			ensureAlwaysMsgf(!ObjectsWithDirtyChanges.GetBit(Index) , TEXT("New object added to scope, Waiting to start replication for ( InternalIndex: %u ) currently in State: %s "), Index, LexToString(State));
-			ensureAlwaysMsgf(!Info.HasDirtyChangeMask, TEXT("New object added to scope, Waiting to start replication for ( InternalIndex: %u ) currently in State: %s "), Index, LexToString(State));
+			ensureMsgf(!ObjectsWithDirtyChanges.GetBit(Index) , TEXT("New object added to scope, Waiting to start replication for ( InternalIndex: %u ) currently in State: %s "), Index, LexToString(State));
+			ensureMsgf(!Info.HasDirtyChangeMask, TEXT("New object added to scope, Waiting to start replication for ( InternalIndex: %u ) currently in State: %s "), Index, LexToString(State));
 		}
 	};
 
@@ -787,7 +787,7 @@ void FReplicationWriter::UpdateScope(const FNetBitArrayView& UpdatedScope)
 
 void FReplicationWriter::InternalUpdateDirtyChangeMasks(const FChangeMaskCache& CachedChangeMasks, EFlushFlags ExtraFlushFlags, bool bMarkForTearOff)
 {
-	IRIS_PROFILER_SCOPE(FReplicationWriter_UpdateDirtyChangeMasks);
+	//IRIS_PROFILER_SCOPE(FReplicationWriter_UpdateDirtyChangeMasks);
 
 	const uint32 MarkForTearOff = bMarkForTearOff ? 1U : 0U;
 	const ChangeMaskStorageType* StoragePtr = CachedChangeMasks.Storage.GetData();
@@ -2556,7 +2556,7 @@ int FReplicationWriter::PrepareAndSendHugeObjectPayload(FNetSerializationContext
 			// Need to call this in order to cleanup data associated with batch
 			HandleObjectBatchFailure(HugeObjectStatus, HugeObjectBatchInfo, WriteBitStreamInfo);
 
-			ensureAlwaysMsgf(HugeObjectStatus == EWriteObjectStatus::BitStreamOverflow, TEXT("Expected split payload to not be able to generate other errors than overflow. Got %u"), unsigned(HugeObjectStatus));
+			ensureMsgf(HugeObjectStatus == EWriteObjectStatus::BitStreamOverflow, TEXT("Expected split payload to not be able to generate other errors than overflow. Got %u"), unsigned(HugeObjectStatus));
 			// It's unexpected, but not a critical error, if no part of the payload could be sent.
 			// We do expect a smaller object to be sent though so that's why 0 is returned.
 
@@ -2848,7 +2848,7 @@ uint32 FReplicationWriter::WriteObjects(FNetSerializationContext& Context)
 
 				const uint32 InternalIndex = WriteContext.DependentObjectsPendingSend.Pop();
 				checkSlow(InternalIndex != ObjectIndexForOOBAttachment);
-				ensureAlwaysMsgf(GetReplicationInfo(InternalIndex).GetState() != EReplicatedObjectState::Invalid, TEXT("DependentObject with InternalIndex %u is not in scope"), InternalIndex);
+				ensureMsgf(GetReplicationInfo(InternalIndex).GetState() != EReplicatedObjectState::Invalid, TEXT("DependentObject with InternalIndex %u is not in scope"), InternalIndex);
 				if (!SendObjectFunction(InternalIndex))
 				{
 					// If we fail, we put the object back on the pending send stack and try again in the next packet of the batch
