increase maximum value contained by FImplicitObjectUnion::NumLeafObjects from uint16 to int32

diff --git a/Engine/Source/Runtime/Experimental/Chaos/Private/Chaos/ImplicitObjectUnion.cpp b/Engine/Source/Runtime/Experimental/Chaos/Private/Chaos/ImplicitObjectUnion.cpp
--- a/Engine/Source/Runtime/Experimental/Chaos/Private/Chaos/ImplicitObjectUnion.cpp
+++ b/Engine/Source/Runtime/Experimental/Chaos/Private/Chaos/ImplicitObjectUnion.cpp
@@ -160,7 +160,7 @@ void FImplicitObjectUnion::SetNumLeafObjects(int32 InNumLeafObjects)
 	ensure(InNumLeafObjects <= MaxNumLeafObjects);
 	check(InNumLeafObjects >= 0);
 
-	NumLeafObjects = uint16(FMath::Min(InNumLeafObjects, MaxNumLeafObjects));
+	NumLeafObjects = FMath::Min(InNumLeafObjects, MaxNumLeafObjects);
 }
 
 void FImplicitObjectUnion::CreateBVH()
diff --git a/Engine/Source/Runtime/Experimental/Chaos/Public/Chaos/ImplicitObjectUnion.h b/Engine/Source/Runtime/Experimental/Chaos/Public/Chaos/ImplicitObjectUnion.h
--- a/Engine/Source/Runtime/Experimental/Chaos/Public/Chaos/ImplicitObjectUnion.h
+++ b/Engine/Source/Runtime/Experimental/Chaos/Public/Chaos/ImplicitObjectUnion.h
@@ -57,7 +57,7 @@ class FImplicitObjectUnion : public FImplicitObject
 	// The total number of leaf objects in the hierarchy
 	int32 GetNumLeafObjects() const
 	{
-		return int32(NumLeafObjects);
+		return NumLeafObjects;
 	}
 
 	// Enable BVH suport for this Union. This should only be done for the root Union in a hierarchy
@@ -296,7 +296,7 @@ protected:
 	// BVH is only created when there are many objects.
 	// @todo(chaos): consider registering particles that may need BVH updated in evolution instead
 	TUniquePtr<Private::FImplicitBVH> BVH;
-	uint16 NumLeafObjects;
+	int32 NumLeafObjects;
 	FFLags Flags;
 };
 
