[ChaosVD] Properly initializing properties in FChaosVDJointSolverSettingsDataWrapper

diff --git a/Engine/Source/Runtime/Experimental/ChaosVisualDebugger/Public/DataWrappers/ChaosVDJointDataWrappers.h b/Engine/Source/Runtime/Experimental/ChaosVisualDebugger/Public/DataWrappers/ChaosVDJointDataWrappers.h
--- a/Engine/Source/Runtime/Experimental/ChaosVisualDebugger/Public/DataWrappers/ChaosVDJointDataWrappers.h
+++ b/Engine/Source/Runtime/Experimental/ChaosVisualDebugger/Public/DataWrappers/ChaosVDJointDataWrappers.h
@@ -118,25 +118,25 @@ public:
 	CHAOSVDRUNTIME_API bool Serialize(FArchive& Ar);
 
 	UPROPERTY(VisibleAnywhere, Category=Tolerances)
-	double SwingTwistAngleTolerance;
+	double SwingTwistAngleTolerance = 0.0;
 	UPROPERTY(VisibleAnywhere, Category=Tolerances)
-	double PositionTolerance;
+	double PositionTolerance = 0.0;
 	UPROPERTY(VisibleAnywhere, Category=Tolerances)
-	double AngleTolerance;
+	double AngleTolerance = 0.0;
 
 	UPROPERTY(VisibleAnywhere, Category="Stability Control")
-	double MinParentMassRatio;
+	double MinParentMassRatio = 0.0;
 	UPROPERTY(VisibleAnywhere, Category="Stability Control")
-	double MaxInertiaRatio;
+	double MaxInertiaRatio = 0.0;
 
 	UPROPERTY(VisibleAnywhere, Category="Solver Stiffness")
-	double MinSolverStiffness;
+	double MinSolverStiffness = 0.0;
 	UPROPERTY(VisibleAnywhere, Category="Solver Stiffness")
-	double MaxSolverStiffness;
+	double MaxSolverStiffness = 0.0;
 	UPROPERTY(VisibleAnywhere, Category="Solver Stiffness")
-	int32 NumIterationsAtMaxSolverStiffness;
+	int32 NumIterationsAtMaxSolverStiffness = INDEX_NONE;
 	UPROPERTY(VisibleAnywhere, Category="Solver Stiffness")
-	int32 NumShockPropagationIterations;
+	int32 NumShockPropagationIterations = INDEX_NONE;
 
 	UPROPERTY(VisibleAnywhere, Category="General")
 	uint16 bUseLinearSolver : 1 = false;
@@ -245,19 +245,19 @@ public:
 	UPROPERTY(VisibleAnywhere, Category="General")
 	FTransform ConnectorTransforms[2];
 	UPROPERTY(VisibleAnywhere, Category="General")
-	double Stiffness;
+	double Stiffness = 0.0;
 	UPROPERTY(VisibleAnywhere, Category="General")
-	double LinearProjection;
+	double LinearProjection = 0.0;
 	UPROPERTY(VisibleAnywhere, Category="General")
-	double AngularProjection;
+	double AngularProjection = 0.0;
 	UPROPERTY(VisibleAnywhere, Category="General")
-	double ShockPropagation;
+	double ShockPropagation = 0.0;
 	UPROPERTY(VisibleAnywhere, Category="General")
-	double TeleportDistance;
+	double TeleportDistance = 0.0;
 	UPROPERTY(VisibleAnywhere, Category="General")
-	double TeleportAngle;
+	double TeleportAngle = 0.0;
 	UPROPERTY(VisibleAnywhere, Category="General")
-	double ParentInvMassScale;
+	double ParentInvMassScale = 0.0;
 
 	UPROPERTY(VisibleAnywhere, Category="General")
 	uint16 bCollisionEnabled : 1 = false;
@@ -283,48 +283,48 @@ public:
 	uint16 bAngularSwingVelocityDriveEnabled : 1 = false;
 
 	UPROPERTY(VisibleAnywhere, Category="General")
-	EChaosVDJointMotionType LinearMotionTypes[3];
+	EChaosVDJointMotionType LinearMotionTypes[3] = { EChaosVDJointMotionType::Free, EChaosVDJointMotionType::Free, EChaosVDJointMotionType::Free };
 	UPROPERTY(VisibleAnywhere, Category="General")
-	double LinearLimit;
+	double LinearLimit = 0.0;
 	UPROPERTY(VisibleAnywhere, Category="General")
-	EChaosVDJointMotionType AngularMotionTypes[3];
+	EChaosVDJointMotionType AngularMotionTypes[3] = { EChaosVDJointMotionType::Free, EChaosVDJointMotionType::Free, EChaosVDJointMotionType::Free };
 	UPROPERTY(VisibleAnywhere, Category="General")
-	FVector AngularLimits;
+	FVector AngularLimits = FVector::ZeroVector;
 
 	UPROPERTY(VisibleAnywhere, Category="General")
-	EChaosVDJointForceMode LinearSoftForceMode;
+	EChaosVDJointForceMode LinearSoftForceMode = EChaosVDJointForceMode::Acceleration;
 	UPROPERTY(VisibleAnywhere, Category="General")
-	EChaosVDJointForceMode AngularSoftForceMode;
+	EChaosVDJointForceMode AngularSoftForceMode = EChaosVDJointForceMode::Acceleration;
 	UPROPERTY(VisibleAnywhere, Category="General")
-	double SoftLinearStiffness;
+	double SoftLinearStiffness = 0.0;
 	UPROPERTY(VisibleAnywhere, Category="General")
-	double SoftLinearDamping;
+	double SoftLinearDamping = 0.0;
 	UPROPERTY(VisibleAnywhere, Category="General")
-	double SoftTwistStiffness;
+	double SoftTwistStiffness = 0.0;
 	UPROPERTY(VisibleAnywhere, Category="General")
-	double SoftTwistDamping;
+	double SoftTwistDamping = 0.0;
 	UPROPERTY(VisibleAnywhere, Category="General")
-	double SoftSwingStiffness;
+	double SoftSwingStiffness = 0.0;
 	UPROPERTY(VisibleAnywhere, Category="General")
-	double SoftSwingDamping;
+	double SoftSwingDamping = 0.0;
 	UPROPERTY(VisibleAnywhere, Category="General")
-	double LinearRestitution;
+	double LinearRestitution = 0.0;
 	UPROPERTY(VisibleAnywhere, Category="General")
-	double TwistRestitution;
+	double TwistRestitution = 0.0;
 	UPROPERTY(VisibleAnywhere, Category="General")
-	double SwingRestitution;
+	double SwingRestitution = 0.0;
 
 	UPROPERTY(VisibleAnywhere, Category="General")
-	double LinearContactDistance;
+	double LinearContactDistance = 0.0;
 	UPROPERTY(VisibleAnywhere, Category="General")
-	double TwistContactDistance;
+	double TwistContactDistance = 0.0;
 	UPROPERTY(VisibleAnywhere, Category="General")
-	double SwingContactDistance;
+	double SwingContactDistance = 0.0;
 
 	UPROPERTY(VisibleAnywhere, Category="General")
-	FVector LinearDrivePositionTarget;
+	FVector LinearDrivePositionTarget = FVector::ZeroVector;
 	UPROPERTY(VisibleAnywhere, Category="General")
-	FVector LinearDriveVelocityTarget;
+	FVector LinearDriveVelocityTarget = FVector::ZeroVector;
 	
 	UPROPERTY(VisibleAnywhere, Category="General")
 	uint8 bLinearPositionDriveEnabled0 : 1 = false;
@@ -340,39 +340,39 @@ public:
 	uint8 bLinearVelocityDriveEnabled2 : 1 = false;
 	
 	UPROPERTY(VisibleAnywhere, Category="General")
-	EChaosVDJointForceMode LinearDriveForceMode;
+	EChaosVDJointForceMode LinearDriveForceMode = EChaosVDJointForceMode::Acceleration;
 	UPROPERTY(VisibleAnywhere, Category="General")
-	FVector LinearDriveStiffness;
+	FVector LinearDriveStiffness = FVector::ZeroVector;
 	UPROPERTY(VisibleAnywhere, Category="General")
-	FVector LinearDriveDamping;
+	FVector LinearDriveDamping = FVector::ZeroVector;
 	UPROPERTY(VisibleAnywhere, Category="General")
-	FVector LinearDriveMaxForce;
+	FVector LinearDriveMaxForce = FVector::ZeroVector;
 	UPROPERTY(VisibleAnywhere, Category="General")
-	FQuat AngularDrivePositionTarget;
+	FQuat AngularDrivePositionTarget = FQuat::Identity;
 	UPROPERTY(VisibleAnywhere, Category="General")
-	FVector AngularDriveVelocityTarget;
+	FVector AngularDriveVelocityTarget = FVector::ZeroVector;
 	UPROPERTY(VisibleAnywhere, Category="General")
-	EChaosVDJointForceMode AngularDriveForceMode;
+	EChaosVDJointForceMode AngularDriveForceMode = EChaosVDJointForceMode::Acceleration;
 	UPROPERTY(VisibleAnywhere, Category="General")
-	FVector AngularDriveStiffness;
+	FVector AngularDriveStiffness = FVector::ZeroVector;
 	UPROPERTY(VisibleAnywhere, Category="General")
-	FVector AngularDriveDamping;
+	FVector AngularDriveDamping = FVector::ZeroVector;
 	UPROPERTY(VisibleAnywhere, Category="General")
-	FVector AngularDriveMaxTorque;
+	FVector AngularDriveMaxTorque = FVector::ZeroVector;
 	UPROPERTY(VisibleAnywhere, Category="General")
-	double LinearBreakForce;
+	double LinearBreakForce = 0.0;
 	UPROPERTY(VisibleAnywhere, Category="General")
-	double LinearPlasticityLimit;
+	double LinearPlasticityLimit = 0.0;
 	UPROPERTY(VisibleAnywhere, Category="General")
-	EChaosVDPlasticityType LinearPlasticityType;
+	EChaosVDPlasticityType LinearPlasticityType = EChaosVDPlasticityType::Free;
 	UPROPERTY(VisibleAnywhere, Category="General")
-	double LinearPlasticityInitialDistanceSquared;
+	double LinearPlasticityInitialDistanceSquared = 0.0;
 	UPROPERTY(VisibleAnywhere, Category="General")
-	double AngularBreakTorque;
+	double AngularBreakTorque = 0.0;
 	UPROPERTY(VisibleAnywhere, Category="General")
-	double AngularPlasticityLimit;
+	double AngularPlasticityLimit = 0.0;
 	UPROPERTY(VisibleAnywhere, Category="General")
-	double ContactTransferScale;
+	double ContactTransferScale = 0.0;
 };
 
 CVD_IMPLEMENT_SERIALIZER(FChaosVDJointSettingsDataWrapper)
