Geometry Collection : static mesh root proxies component now updates with root transform and broken state

diff --git a/Engine/Source/Runtime/Experimental/GeometryCollectionEngine/Private/GeometryCollection/GeometryCollectionComponent.cpp b/Engine/Source/Runtime/Experimental/GeometryCollectionEngine/Private/GeometryCollection/GeometryCollectionComponent.cpp
--- a/Engine/Source/Runtime/Experimental/GeometryCollectionEngine/Private/GeometryCollection/GeometryCollectionComponent.cpp
+++ b/Engine/Source/Runtime/Experimental/GeometryCollectionEngine/Private/GeometryCollection/GeometryCollectionComponent.cpp
@@ -3606,6 +3606,32 @@ void UGeometryCollectionComponent::CreateRootProxyComponentsIfNeeded()
 	}
 }
 
+void UGeometryCollectionComponent::UpdateRootProxyComponentsIfNeeded()
+{
+	if (RootProxyStaticMeshComponents.Num() > 0 && !bUpdateComponentTransformToRootBone)
+	{
+		if (RestCollection && DynamicCollection && !IsRootBroken())
+		{
+			const TSharedPtr<FGeometryCollection, ESPMode::ThreadSafe> AssetCollection = RestCollection->GetGeometryCollection();
+			if (AssetCollection)
+			{
+				const int32 RootIndex = RestCollection->GetRootIndex();
+				const FTransform3f RootRestTransform = AssetCollection->Transform[RootIndex];
+				const FTransform3f RootCurrentTransform = ComponentSpaceTransforms.RequestRootTransform();
+				const FTransform RootTransformFromRestPose = FTransform(RootCurrentTransform.GetRelativeTransform(RootRestTransform));
+
+				for (TObjectPtr<UStaticMeshComponent> StaticMeshComponent : RootProxyStaticMeshComponents)
+				{
+					if (StaticMeshComponent)
+					{
+						StaticMeshComponent->SetRelativeTransform(RootTransformFromRestPose);
+					}
+				}
+			}
+		}
+	}
+}
+
 void UGeometryCollectionComponent::ClearRootProxyComponents()
 {
 	for (TObjectPtr<UStaticMeshComponent> StaticMeshComponent : RootProxyStaticMeshComponents)
@@ -4314,6 +4340,12 @@ void UGeometryCollectionComponent::OnPostPhysicsSync()
 			InitializeRemovalDynamicAttributesIfNeeded();
 		}
 
+		if (RootProxyStaticMeshComponents.Num() > 0)
+		{
+			ClearRootProxyComponents();
+			RecreateRenderState_Concurrent();
+		}
+
 		UpdateRemovalIfNeeded();
 
 		CheckFullyDecayed();
@@ -4411,6 +4443,8 @@ void UGeometryCollectionComponent::UpdateRenderSystemsIfNeeded(bool bDynamicColl
 			RefreshCustomRenderer();
 		}
 
+		UpdateRootProxyComponentsIfNeeded();
+
 		if (SceneProxy != nullptr)
 		{
 			if (SceneProxy->IsNaniteMesh())
diff --git a/Engine/Source/Runtime/Experimental/GeometryCollectionEngine/Public/GeometryCollection/GeometryCollectionComponent.h b/Engine/Source/Runtime/Experimental/GeometryCollectionEngine/Public/GeometryCollection/GeometryCollectionComponent.h
--- a/Engine/Source/Runtime/Experimental/GeometryCollectionEngine/Public/GeometryCollection/GeometryCollectionComponent.h
+++ b/Engine/Source/Runtime/Experimental/GeometryCollectionEngine/Public/GeometryCollection/GeometryCollectionComponent.h
@@ -1832,6 +1832,7 @@ private:
 
 	bool ShouldCreateRootProxyComponents() const;
 	void CreateRootProxyComponentsIfNeeded();
+	void UpdateRootProxyComponentsIfNeeded();
 	void ClearRootProxyComponents();
 
 	TArray<TObjectPtr<UStaticMeshComponent>> RootProxyStaticMeshComponents;
