Fix constraints not adjusting connector transform when root particle is different from original one

diff --git a/Engine/Source/Runtime/Experimental/Chaos/Private/PhysicsProxy/JointConstraintProxy.cpp b/Engine/Source/Runtime/Experimental/Chaos/Private/PhysicsProxy/JointConstraintProxy.cpp
--- a/Engine/Source/Runtime/Experimental/Chaos/Private/PhysicsProxy/JointConstraintProxy.cpp
+++ b/Engine/Source/Runtime/Experimental/Chaos/Private/PhysicsProxy/JointConstraintProxy.cpp
@@ -23,6 +23,7 @@ FJointConstraintPhysicsProxy::FJointConstraintPhysicsProxy(FJointConstraint* InC
 	: Base(EPhysicsProxyType::JointConstraintType, InOwner, MakeShared<FProxyTimestampBase>())
 	, Constraint_GT(InConstraint) // This proxy assumes ownership of the Constraint, and will free it during DestroyOnPhysicsThread
 	, Constraint_PT(InHandle)
+	, OriginalParticleHandles_PT(nullptr, nullptr)
 {
 	check(Constraint_GT !=nullptr);
 	Constraint_GT->SetProxy(this);
@@ -76,6 +77,19 @@ bool FJointConstraintPhysicsProxy::PullFromPhysicsState(const FDirtyJointConstra
 	return true;
 }
 
+template <typename TransformType>
+static void FixConnectorTransformsForRoot(const FParticlePair& OriginalHandles, const FParticlePair& RootHandles, TransformType& InOutTransformsToFix)
+{
+	for (int32 Index = 0; Index < 2; Index++)
+	{
+		if (OriginalHandles[Index] && RootHandles[Index] && RootHandles[Index] != OriginalHandles[Index])
+		{
+			const FTransform TransformOffset = OriginalHandles[Index]->GetTransformXR().GetRelativeTransform(RootHandles[Index]->GetTransformXR());
+			InOutTransformsToFix[Index] *= TransformOffset;
+		}
+	}
+}
+
 void FJointConstraintPhysicsProxy::InitializeOnPhysicsThread(FPBDRigidsSolver* InSolver, FDirtyPropertiesManager& Manager, int32 DataIdx, FDirtyChaosProperties& RemoteData)
 {
 	auto& Handles = InSolver->GetParticles().GetParticleHandles();
@@ -84,21 +98,36 @@ void FJointConstraintPhysicsProxy::InitializeOnPhysicsThread(FPBDRigidsSolver* I
 		auto& JointConstraints = InSolver->GetJointConstraints();
 		if(const FPhysicsObjectPairProperty* BodyPairs = RemoteData.FindJointPhysicsObjects(Manager, DataIdx))
 		{
-			FGeometryParticleHandle* Handle0 = BodyPairs->PhysicsBodies[0]->GetRootParticle<Chaos::EThreadContext::Internal>();
-			FGeometryParticleHandle* Handle1 = BodyPairs->PhysicsBodies[1]->GetRootParticle<Chaos::EThreadContext::Internal>();
+			FParticlePair RootHandles
+			{
+				BodyPairs->PhysicsBodies[0]->GetRootParticle<Chaos::EThreadContext::Internal>(),
+				BodyPairs->PhysicsBodies[1]->GetRootParticle<Chaos::EThreadContext::Internal>(),
+			};
 
-			if (Handle0 && Handle1)
+			if (RootHandles[0] && RootHandles[1])
 			{
 				if (const FPBDJointSettings* JointSettings = RemoteData.FindJointSettings(Manager, DataIdx))
 				{
-					Constraint_PT = JointConstraints.AddConstraint({ Handle0,Handle1 }, JointSettings->ConnectorTransforms);
+					FTransformPair ConnectorTransforms{ JointSettings->ConnectorTransforms };
+
+					// if the root particles do not match the actual particles
+					// we need to adjust the frames to be in the root particle space 
+					OriginalParticleHandles_PT =
+					{
+						BodyPairs->PhysicsBodies[0]->GetParticle<Chaos::EThreadContext::Internal>(),
+						BodyPairs->PhysicsBodies[1]->GetParticle<Chaos::EThreadContext::Internal>(),
+					};
+
+					FixConnectorTransformsForRoot(OriginalParticleHandles_PT, RootHandles, ConnectorTransforms);
 
-					Handle0->AddConstraintHandle(Constraint_PT);
-					Handle1->AddConstraintHandle(Constraint_PT);
+					Constraint_PT = JointConstraints.AddConstraint(RootHandles, ConnectorTransforms);
+
+					RootHandles[0]->AddConstraintHandle(Constraint_PT);
+					RootHandles[1]->AddConstraintHandle(Constraint_PT);
 
 					// We added a joint to the particles, so we need to modify the inertia (See FPBDRigidsEvolutionGBF::UpdateInertiaConditioning)
-					FGenericParticleHandle(Handle0)->SetInertiaConditioningDirty();
-					FGenericParticleHandle(Handle1)->SetInertiaConditioningDirty();
+					FGenericParticleHandle(RootHandles[0])->SetInertiaConditioningDirty();
+					FGenericParticleHandle(RootHandles[1])->SetInertiaConditioningDirty();
 				}
 			}
 		}
@@ -140,6 +169,7 @@ void FJointConstraintPhysicsProxy::DestroyOnPhysicsThread(FPBDRigidsSolver* InSo
 		FPBDRigidsSolver::FJointConstraints& JointConstraints = InSolver->GetJointConstraints();
 		JointConstraints.RemoveConstraint(Constraint_PT->GetConstraintIndex());
 		Constraint_PT = nullptr;
+		OriginalParticleHandles_PT = { nullptr, nullptr };
 	}
 }
 
@@ -182,8 +212,14 @@ void FJointConstraintPhysicsProxy::PushStateOnPhysicsThread(FPBDRigidsSolver* In
 				}
 			}
 
+			FPBDJointSettings JointSettingsToSet{ JointSettingsBuffer };
+
+			// PushStateOnPhysicsThread is always called right after InitializeOnPhysicsThread
+			// so we need to make sure the root space correction logic is also run here to avoid resetting the transform in the wrong space
+			FixConnectorTransformsForRoot(OriginalParticleHandles_PT, Constraint_PT->GetConstrainedParticles(), JointSettingsToSet.ConnectorTransforms);
+
 			// Update the joint settings
-			Constraint_PT->SetSettings(JointSettingsBuffer);
+			Constraint_PT->SetSettings(JointSettingsToSet);
 		}
 	}
 }
diff --git a/Engine/Source/Runtime/Experimental/Chaos/Public/PhysicsProxy/JointConstraintProxy.h b/Engine/Source/Runtime/Experimental/Chaos/Public/PhysicsProxy/JointConstraintProxy.h
--- a/Engine/Source/Runtime/Experimental/Chaos/Public/PhysicsProxy/JointConstraintProxy.h
+++ b/Engine/Source/Runtime/Experimental/Chaos/Public/PhysicsProxy/JointConstraintProxy.h
@@ -103,6 +103,7 @@ public:
 private:
 	FJointConstraint* Constraint_GT;
 	FPBDJointConstraintHandle* Constraint_PT;
+	FParticlePair OriginalParticleHandles_PT;
 	bool bInitialized = false;
 };
 
