Iris
- Fixed issue with posting rpcs to out of scope subobjects being scheduled with internal index and transmitted with invalid object.

diff --git a/Engine/Plugins/Runtime/ReplicationSystemTestPlugin/Source/Private/Tests/ReplicationSystem/RPC/ReplicatedTestObjectWithRPC.cpp b/Engine/Plugins/Runtime/ReplicationSystemTestPlugin/Source/Private/Tests/ReplicationSystem/RPC/ReplicatedTestObjectWithRPC.cpp
--- a/Engine/Plugins/Runtime/ReplicationSystemTestPlugin/Source/Private/Tests/ReplicationSystem/RPC/ReplicatedTestObjectWithRPC.cpp
+++ b/Engine/Plugins/Runtime/ReplicationSystemTestPlugin/Source/Private/Tests/ReplicationSystem/RPC/ReplicatedTestObjectWithRPC.cpp
@@ -148,3 +148,114 @@ void UTestReplicatedObjectWithRPC::NetMulticast_MultiCastRPC_Implementation()
 {
 	NetMulticast_MultiCastRPCCallOrder = ++CallOrder;
 }
+
+void UTestReplicatedObjectWithRPC::NetMulticast_ReliableMultiCastRPC_Implementation()
+{
+	NetMulticast_ReliableMultiCastRPCCallOrder = ++CallOrder;
+}
+
+
+//////////////////////////////////////////////////////////////////////////
+// Implementation for UTestReplicatedObjectWithSingleRPC
+//////////////////////////////////////////////////////////////////////////
+UTestReplicatedObjectWithSingleRPC::UTestReplicatedObjectWithSingleRPC()
+	: UReplicatedTestObject()
+{
+}
+
+void UTestReplicatedObjectWithSingleRPC::RegisterReplicationFragments(UE::Net::FFragmentRegistrationContext& Context, UE::Net::EFragmentRegistrationFlags RegistrationFlags)
+{
+	this->ReplicationFragments.Reset();
+	UE::Net::FReplicationFragmentUtil::CreateAndRegisterFragmentsForObject(this, Context, RegistrationFlags, &this->ReplicationFragments);
+}
+
+void UTestReplicatedObjectWithSingleRPC::Init(UReplicationSystem* InRepSystem)
+{
+	bIsServerObject = InRepSystem->IsServer();
+	ReplicationSystem = InRepSystem;
+}
+
+void UTestReplicatedObjectWithSingleRPC::SetRootObject(UTestReplicatedObjectWithSingleRPC* InRootObject)
+{
+	check(InRootObject);
+	RootObject = InRootObject;
+}
+
+int32 UTestReplicatedObjectWithSingleRPC::GetFunctionCallspace(UFunction* Function, FFrame* Stack)
+{
+	check(!(Function->FunctionFlags & FUNC_Static));
+	check(Function->FunctionFlags & FUNC_Net);
+
+	const bool bIsOnServer = bIsServerObject;
+
+	// get the top most function
+	while (Function->GetSuperFunction() != nullptr)
+	{
+		Function = Function->GetSuperFunction();
+	}
+
+	// Multicast RPCs
+	if ((Function->FunctionFlags & FUNC_NetMulticast))
+	{
+		if (bIsOnServer)
+		{
+			// Server should execute locally and call remotely
+			return (FunctionCallspace::Local | FunctionCallspace::Remote);
+		}
+		else
+		{
+			return FunctionCallspace::Local;
+		}
+	}
+
+	// if we are the authority
+	if (bIsOnServer)
+	{
+		if (Function->FunctionFlags & FUNC_NetClient)
+		{
+			return FunctionCallspace::Remote;
+		}
+		else
+		{
+			return FunctionCallspace::Local;
+		}
+
+	}
+	// if we are not the authority
+	else
+	{
+		if (Function->FunctionFlags & FUNC_NetServer)
+		{
+			return FunctionCallspace::Remote;
+		}
+		else
+		{
+			// don't replicate
+			return FunctionCallspace::Local;
+		}
+	}
+}
+
+bool UTestReplicatedObjectWithSingleRPC::CallRemoteFunction(UFunction* Function, void* Parameters, FOutParmRec* OutParms, FFrame* Stack)
+{
+	if (bIsSubObject)
+	{
+		return ReplicationSystem->SendRPC(RootObject, this, Function, Parameters);
+	}
+	else
+	{
+		return ReplicationSystem->SendRPC(this, nullptr, Function, Parameters);
+	}
+}
+
+void UTestReplicatedObjectWithSingleRPC::NetMulticast_ReliableMultiCastRPC_Implementation(int32 Int32Param)
+{
+	NetMulticast_ReliableMultiCastRPCCallOrder = ++CallOrder;
+}
+
+
+void UTestReplicatedObjectWithSingleRPC::NetMulticast_AnotherReliableMultiCastRPC_Implementation(double DoubleParam)
+{
+	NetMulticast_AnotherReliableMultiCastRPCCallOrder = ++CallOrder;
+}
+
diff --git a/Engine/Plugins/Runtime/ReplicationSystemTestPlugin/Source/Private/Tests/ReplicationSystem/RPC/ReplicatedTestObjectWithRPC.h b/Engine/Plugins/Runtime/ReplicationSystemTestPlugin/Source/Private/Tests/ReplicationSystem/RPC/ReplicatedTestObjectWithRPC.h
--- a/Engine/Plugins/Runtime/ReplicationSystemTestPlugin/Source/Private/Tests/ReplicationSystem/RPC/ReplicatedTestObjectWithRPC.h
+++ b/Engine/Plugins/Runtime/ReplicationSystemTestPlugin/Source/Private/Tests/ReplicationSystem/RPC/ReplicatedTestObjectWithRPC.h
@@ -75,6 +75,53 @@ public:
 	void NetMulticast_MultiCastRPC();
 	int32 NetMulticast_MultiCastRPCCallOrder = 0;
 
+	UFUNCTION(NetMulticast, reliable)
+	void NetMulticast_ReliableMultiCastRPC();
+	int32 NetMulticast_ReliableMultiCastRPCCallOrder = 0;
+
+	int32 CallOrder = 0;
+};
+
+/**
+ *  A test class for testing RPC replication
+ */
+UCLASS()
+class UTestReplicatedObjectWithSingleRPC : public UReplicatedTestObject
+{
+	GENERATED_BODY()
+
+public:
+	UTestReplicatedObjectWithSingleRPC();
+
+	void Init(UReplicationSystem* InRepSystem);
+	void SetRootObject(UTestReplicatedObjectWithSingleRPC* InRootObject);
+
+	virtual void RegisterReplicationFragments(UE::Net::FFragmentRegistrationContext& Context, UE::Net::EFragmentRegistrationFlags RegistrationFlags) override;
+
+	virtual int32 GetFunctionCallspace(UFunction* Function, FFrame* Stack) override;
+	virtual bool CallRemoteFunction(UFunction* Function, void* Parameters, FOutParmRec* OutParms, FFrame* Stack) override;
+
+	UReplicationSystem* ReplicationSystem = nullptr;
+
+	// Network data only for test
+	TArray<UE::Net::FReplicationFragment*> ReplicationFragments;
+
+	// To determine if this object is located on the server or client
+	bool bIsServerObject = false;
+
+	// Our owner when the object is a subobject
+	UTestReplicatedObjectWithSingleRPC* RootObject = nullptr;
+
+	// RPC test functions
+	UFUNCTION(NetMulticast, reliable)
+	void NetMulticast_ReliableMultiCastRPC(int32 Int32Param);
+	int32 NetMulticast_ReliableMultiCastRPCCallOrder = 0;
+
+	UFUNCTION(NetMulticast, reliable)
+	void NetMulticast_AnotherReliableMultiCastRPC(double DoubleParam);
+	int32 NetMulticast_AnotherReliableMultiCastRPCCallOrder = 0;
+
 	int32 CallOrder = 0;
 };
 
+
diff --git a/Engine/Plugins/Runtime/ReplicationSystemTestPlugin/Source/Private/Tests/ReplicationSystem/RPC/TestObjectRPC.cpp b/Engine/Plugins/Runtime/ReplicationSystemTestPlugin/Source/Private/Tests/ReplicationSystem/RPC/TestObjectRPC.cpp
--- a/Engine/Plugins/Runtime/ReplicationSystemTestPlugin/Source/Private/Tests/ReplicationSystem/RPC/TestObjectRPC.cpp
+++ b/Engine/Plugins/Runtime/ReplicationSystemTestPlugin/Source/Private/Tests/ReplicationSystem/RPC/TestObjectRPC.cpp
@@ -328,4 +328,47 @@ namespace UE::Net::Private
 		UE_NET_ASSERT_LT(ClientObject->ClientRPCCallOrder, ClientObject->ClientRPCWithParamCallOrder);
 	}
 
+	// This test is specifically written to exercise a path where we would cause bitstream errors by
+	// posting rpc on an object only temporarily in scope, which would be assigned an internal index.
+	// This same internal index would then be reused and would inherit previously posted attachments/rpcs possibly for a different type.
+	// If rpc index would be valid for the new object and it did not have the same signature it would result in a bitstream error.
+	UE_NET_TEST_FIXTURE(FRPCTestFixture, TestShortLivedSubObjectReliableRPC)
+	{
+		// Add a client
+		FReplicationSystemTestClient* Client = CreateClient();
+
+		// Spawn object on server
+		UTestReplicatedObjectWithRPC* ServerRootObject = Server->CreateObject<UTestReplicatedObjectWithRPC>();
+		ServerRootObject->Init(Server->GetReplicationSystem());
+
+		const FNetRefHandle ServerRootObjectHandle = ServerRootObject->NetRefHandle;
+
+		// Send and deliver packet
+		Server->UpdateAndSend({Client});
+
+		// Create subobject
+		UTestReplicatedObjectWithRPC* ServerSubObject = Server->CreateSubObject<UTestReplicatedObjectWithRPC>(ServerRootObjectHandle);
+		ServerSubObject->Init(Server->GetReplicationSystem());
+		ServerSubObject->SetRootObject(ServerRootObject);
+
+		const FNetRefHandle ServerSubObjectHandle = ServerSubObject->NetRefHandle;
+
+		// Call rpc
+		ServerSubObject->ClientRPC();
+
+		// Destroy subobject
+		Server->DestroyObject(ServerSubObject, EEndReplicationFlags::Destroy);
+
+		// Send and deliver packet
+		Server->UpdateAndSend({Client});
+
+		// Spawn object on server
+		UTestReplicatedObjectWithSingleRPC* ServerRootObject2 = Server->CreateObject<UTestReplicatedObjectWithSingleRPC>();
+		ServerRootObject2->Init(Server->GetReplicationSystem());
+		ServerRootObject2->NetMulticast_ReliableMultiCastRPC(3);
+
+		// Send and deliver packet
+		Server->UpdateAndSend({Client});
+	}
+
 }
\ No newline at end of file
diff --git a/Engine/Source/Runtime/Experimental/Iris/Core/Private/Iris/ReplicationSystem/NetBlob/NetRPC.cpp b/Engine/Source/Runtime/Experimental/Iris/Core/Private/Iris/ReplicationSystem/NetBlob/NetRPC.cpp
--- a/Engine/Source/Runtime/Experimental/Iris/Core/Private/Iris/ReplicationSystem/NetBlob/NetRPC.cpp
+++ b/Engine/Source/Runtime/Experimental/Iris/Core/Private/Iris/ReplicationSystem/NetBlob/NetRPC.cpp
@@ -148,6 +148,7 @@ void FNetRPC::DeserializeWithObject(FNetSerializationContext& Context, FNetRefHa
 		// Just because the serialization didn't detect an error doesn't mean everything is ok. Validate stream position.
 		if (PostNetRPCPos != Context.GetBitStreamReader()->GetPosBits())
 		{
+			UE_LOG(LogIrisRpc, Error, TEXT("Bitstream mismatch while deserializing function %s. Actual stream position: %u Expected stream position: %u"), ToCStr(BlobDescriptor->DebugName), Context.GetBitStreamReader()->GetPosBits(), PostNetRPCPos);
 			ensureMsgf(PostNetRPCPos == Context.GetBitStreamReader()->GetPosBits(), TEXT("Bitstream mismatch while deserializing function %s. Actual stream position: %u Expected stream position: %u"), ToCStr(BlobDescriptor->DebugName), Context.GetBitStreamReader()->GetPosBits(), PostNetRPCPos);
 			Context.GetBitStreamReader()->Seek(PostNetRPCPos);
 			Context.SetError(GNetError_BitStreamError);
diff --git a/Engine/Source/Runtime/Experimental/Iris/Core/Private/Iris/ReplicationSystem/ReplicationWriter.cpp b/Engine/Source/Runtime/Experimental/Iris/Core/Private/Iris/ReplicationSystem/ReplicationWriter.cpp
--- a/Engine/Source/Runtime/Experimental/Iris/Core/Private/Iris/ReplicationSystem/ReplicationWriter.cpp
+++ b/Engine/Source/Runtime/Experimental/Iris/Core/Private/Iris/ReplicationSystem/ReplicationWriter.cpp
@@ -65,7 +65,7 @@ static int32 GReplicationWriterMaxHugeObjectsInTransit = 16;
 static FAutoConsoleVariableRef CVarReplicationWriterMaxHugeObjectsInTransit(TEXT("net.Iris.ReplicationWriterMaxHugeObjectsInTransit"), GReplicationWriterMaxHugeObjectsInTransit,
 	TEXT("How many very large objects, one whose payload doesn't fit in a single packet, is allowed to be scheduled for send. Needs to be at least 1."));
 
-static bool bValidateObjectsWithDirtyChanges = false;
+static bool bValidateObjectsWithDirtyChanges = true;
 static FAutoConsoleVariableRef CvarValidateObjectsWithDirtyChanges(TEXT("net.Iris.ReplicationWriter.ValidateObjectsWithDirtyChanges"), bValidateObjectsWithDirtyChanges, TEXT("Ensure that we don't try to mark invalid objects as dirty when they shouldn't."));
 
 static const FName NetError_ObjectStateTooLarge("Object state is too large to be split.");
@@ -299,44 +299,47 @@ bool FReplicationWriter::QueueNetObjectAttachments(FInternalNetRefIndex OwnerInt
 		return false;
 	}
 
-	const bool bObjectInScope = ObjectsInScope.GetBit(OwnerInternalIndex);
-	if (!bObjectInScope && !Parameters.bAllowSendingAttachmentsToObjectsNotInScope)
+	const uint32 TargetIndex = SubObjectInternalIndex != FNetRefHandleManager::InvalidInternalIndex ? SubObjectInternalIndex : OwnerInternalIndex;
+	const bool bTargetObjectInScope = ObjectsInScope.GetBit(TargetIndex);
+	if (!bTargetObjectInScope && !Parameters.bAllowSendingAttachmentsToObjectsNotInScope)
 	{
-		UE_CLOG_REPLICATIONWRITER_WARNING(bWarnAboutDroppedAttachmentsToObjectsNotInScope, TEXT("Dropping %s attachment due to object ( InternalIndex: %u ) not in scope."), (EnumHasAnyFlags(InAttachments[0]->GetCreationInfo().Flags, ENetBlobFlags::Reliable) ? TEXT("reliable") : TEXT("unreliable")), OwnerInternalIndex);
+		UE_CLOG_REPLICATIONWRITER_WARNING(bWarnAboutDroppedAttachmentsToObjectsNotInScope, TEXT("Dropping %s attachment due to object ( InternalIndex: %u ) not in scope."), (EnumHasAnyFlags(InAttachments[0]->GetCreationInfo().Flags, ENetBlobFlags::Reliable) ? TEXT("reliable") : TEXT("unreliable")), TargetIndex);
 		return false;
 	}
 	
-	// Route attachments flagged with ScheduleAsOOB through OOB channel if we have started replicating the owner.
 	const bool bScheduleUsingOOBChannel = EnumHasAnyFlags(SendFlags, ENetObjectAttachmentSendPolicyFlags::ScheduleAsOOB);
-	if (bScheduleUsingOOBChannel && (GetReplicationInfo(OwnerInternalIndex).GetState() < EReplicatedObjectState::WaitOnCreateConfirmation || GetReplicationInfo(OwnerInternalIndex).GetState() >= EReplicatedObjectState::PendingDestroy))
+	if (bScheduleUsingOOBChannel)
 	{
-		UE_CLOG_REPLICATIONWRITER_WARNING(bWarnAboutDroppedAttachmentsToObjectsNotInScope, TEXT("Dropping attachment scheduled as ScheduleAsOOB due to object ( InternalIndex: %u ) not in replicated state."),  OwnerInternalIndex);
-		return false;
+		// Route attachments flagged with ScheduleAsOOB through OOB channel only if we have started replicating the target.
+		const EReplicatedObjectState ReplicationState = GetReplicationInfo(TargetIndex).GetState();
+		if (ReplicationState < EReplicatedObjectState::WaitOnCreateConfirmation || ReplicationState >= EReplicatedObjectState::PendingDestroy)
+		{
+			UE_CLOG_REPLICATIONWRITER_WARNING(bWarnAboutDroppedAttachmentsToObjectsNotInScope, TEXT("Dropping attachment scheduled as ScheduleAsOOB due to object ( InternalIndex: %u ) not in replicated state."),  OwnerInternalIndex);
+			return false;
+		}
 	}
 
-	const uint32 TargetIndex = (bObjectInScope && !bScheduleUsingOOBChannel) ? (SubObjectInternalIndex != FNetRefHandleManager::InvalidInternalIndex ? SubObjectInternalIndex : OwnerInternalIndex) : ObjectIndexForOOBAttachment;
-	ENetObjectAttachmentType AttachmentType = ((bObjectInScope && !bScheduleUsingOOBChannel) ? ENetObjectAttachmentType::Normal : ENetObjectAttachmentType::OutOfBand);
-	if (!Attachments.Enqueue(AttachmentType, TargetIndex, InAttachments))
+	const uint32 AttachmentQueueIndex = (bTargetObjectInScope && !bScheduleUsingOOBChannel) ? TargetIndex : ObjectIndexForOOBAttachment;
+	const ENetObjectAttachmentType AttachmentType = ((bTargetObjectInScope && !bScheduleUsingOOBChannel) ? ENetObjectAttachmentType::Normal : ENetObjectAttachmentType::OutOfBand);
+	if (!Attachments.Enqueue(AttachmentType, AttachmentQueueIndex, InAttachments))
 	{
 		return false;
 	}
 
-	// There's a special case for out of band attachments, we don't need to mark anything dirty.
-	if (IsObjectIndexForOOBAttachment(TargetIndex))
+	// We do not have to mark anything dirty as there's a special case for out of band attachments
+	if (!IsObjectIndexForOOBAttachment(AttachmentQueueIndex))
 	{
-		return true;
-	}
-
-	FReplicationInfo& TargetInfo = GetReplicationInfo(TargetIndex);
-	TargetInfo.HasAttachments = 1;
+		FReplicationInfo& TargetInfo = GetReplicationInfo(AttachmentQueueIndex);
+		TargetInfo.HasAttachments = 1;
 
-	MarkObjectDirty(TargetIndex, "QueueAttachment");
+		MarkObjectDirty(AttachmentQueueIndex, "QueueAttachment");
 
-	if (OwnerInternalIndex != TargetIndex)
-	{
-		MarkObjectDirty(OwnerInternalIndex, "QueueAttachment2");
-		FReplicationInfo& OwnerInfo = GetReplicationInfo(OwnerInternalIndex);
-		OwnerInfo.HasDirtySubObjects = 1;
+		if (OwnerInternalIndex != AttachmentQueueIndex)
+		{
+			MarkObjectDirty(OwnerInternalIndex, "QueueAttachment2");
+			FReplicationInfo& OwnerInfo = GetReplicationInfo(OwnerInternalIndex);
+			OwnerInfo.HasDirtySubObjects = 1;
+		}
 	}
 
 	return true;
@@ -397,6 +400,12 @@ void FReplicationWriter::StartReplication(uint32 InternalIndex)
 	FReplicationInfo& Info = GetReplicationInfo(InternalIndex);
 
 	ensureMsgf(Info.GetState() == EReplicatedObjectState::Invalid, TEXT("Object ( InternalIndex: %u ) is in state %s in StartReplication."), InternalIndex, LexToString(Info.GetState()));
+	if (InternalIndex != ObjectIndexForOOBAttachment && Attachments.HasUnsentAttachments(ENetObjectAttachmentType::Normal, InternalIndex))
+	{
+		UE_LOG(LogIris, Error, TEXT("FReplicationWriter::StartReplication - Expected object %s to not to have any queued up attachments"), *NetRefHandleManager->PrintObjectFromIndex(InternalIndex));
+		ensure(false);
+		Attachments.DropAllAttachments(ENetObjectAttachmentType::Normal, InternalIndex);
+	}
 
 	// Reset info
 	Info = FReplicationInfo();
