[NetPhysics] Fix bug with render interpolation after resimulation which would accumulate instead of reset the correction distance if too far for interpolation.
The total accumulated distance would trigger on the next valid correction and result in a large visual desync.

diff --git a/Engine/Source/Runtime/Experimental/Chaos/Public/Chaos/Framework/PhysicsProxyBase.h b/Engine/Source/Runtime/Experimental/Chaos/Public/Chaos/Framework/PhysicsProxyBase.h
--- a/Engine/Source/Runtime/Experimental/Chaos/Public/Chaos/Framework/PhysicsProxyBase.h
+++ b/Engine/Source/Runtime/Experimental/Chaos/Public/Chaos/Framework/PhysicsProxyBase.h
@@ -224,12 +224,21 @@ struct FProxyInterpolationError : FProxyInterpolationBase
 
 	void AccumlateErrorXR(const Chaos::FVec3 X, const FQuat R, const int32 CurrentSimTick, const int32 ErrorSmoothDuration)
 	{
-		ErrorX += X;
-		ErrorXPrev = ErrorX;
-		ErrorR *= R;
-		ErrorRPrev = ErrorR;
 		ErrorSmoothingCount = ErrorSmoothDuration; // How many simulation ticks to correct error over
 		LastSimTick = CurrentSimTick - 1; // Error is from the previous simulation tick, not the current
+		SimTicks = 0;
+
+		if (IsErrorSmoothing())
+		{
+			ErrorX += X;
+			ErrorXPrev = ErrorX;
+			ErrorR *= R;
+			ErrorRPrev = ErrorR;
+		}
+		else
+		{
+			Reset();
+		}
 	}
 
 	virtual bool UpdateError(const int32 CurrentSimTick, const Chaos::FReal AsyncFixedTimeStep)
@@ -238,10 +247,9 @@ struct FProxyInterpolationError : FProxyInterpolationBase
 		SimTicks = CurrentSimTick - LastSimTick;
 		LastSimTick = CurrentSimTick;
 
-		if (IsErrorSmoothing() && SimTicks > 0)
+		if (SimTicks > 0)
 		{
-			DecayError();
-			return true;
+			return DecayError();
 		}
 		return false;
 	}
@@ -263,8 +271,14 @@ struct FProxyInterpolationError : FProxyInterpolationBase
 	}
 
 protected:
-	void DecayError()
+	bool DecayError()
 	{
+		if (!IsErrorSmoothing())
+		{
+			Reset();
+			return false;
+		}
+
 		// Linear decay
 		// Example: If we want to decay an error of 100 over 10ticks (i.e. 10% each tick)
 		// First step:  9/10 = 0.9   |  100 * 0.9  = 90 error
@@ -281,8 +295,19 @@ protected:
 		ErrorR = FMath::Lerp(FQuat::Identity, ErrorR, Alpha);
 
 		ErrorSmoothingCount = FMath::Max(ErrorSmoothingCount - SimTicks, 0);
+		return true;
 	}
 
+	void Reset()
+	{
+		ErrorX = Chaos::FVec3::ZeroVector;
+		ErrorXPrev = Chaos::FVec3::ZeroVector;
+		ErrorR = FQuat::Identity;
+		ErrorRPrev = FQuat::Identity;
+		ErrorSmoothingCount = 0;
+		LastSimTick = 0;
+		SimTicks = 0;
+	}
 
 protected:
 	int32 LastSimTick = 0;
