[NetPhysics] Perform null checks to guard from possible crash during resimulation.

diff --git a/Engine/Source/Runtime/Experimental/Chaos/Private/RewindData.cpp b/Engine/Source/Runtime/Experimental/Chaos/Private/RewindData.cpp
--- a/Engine/Source/Runtime/Experimental/Chaos/Private/RewindData.cpp
+++ b/Engine/Source/Runtime/Experimental/Chaos/Private/RewindData.cpp
@@ -620,6 +620,10 @@ void FRewindData::PushGTDirtyData(const FDirtyPropertiesManager& SrcManager,cons
 	bNeedsSave = true;
 
 	IPhysicsProxyBase* Proxy = Dirty.Proxy;
+	if (Proxy == nullptr)
+	{
+		return;
+	}
 
 	//Helper to group most of the common logic about push data recording
 	//NOTE: when possible use passed in CopyFunc to do work, if lambda returns false you cannot record to history buffer
@@ -677,18 +681,22 @@ void FRewindData::PushGTDirtyData(const FDirtyPropertiesManager& SrcManager,cons
 		}
 	};
 
-	switch(Dirty.Proxy->GetType())
+	switch(Proxy->GetType())
 	{
 	case EPhysicsProxyType::SingleParticleProxy:
 	{
-		auto ParticleProxy = static_cast<FSingleParticlePhysicsProxy*>(Dirty.Proxy);
+		FSingleParticlePhysicsProxy* ParticleProxy = static_cast<FSingleParticlePhysicsProxy*>(Proxy);
+		if (ParticleProxy == nullptr)
+		{
+			break;
+		}
 
 		FGeometryParticleHandle* PTParticle = ParticleProxy->GetHandle_LowLevel();
-
 		if (PTParticle == nullptr)
 		{
 			break;
 		}
+
 		const bool bKeepRecording = CopyHelper(PTParticle, [PTParticle, &DirtyPropHelper](FGeometryParticleStateBase& Latest)
 		{
 			DirtyPropHelper(Latest.ParticlePositionRotation, EChaosPropertyFlags::XR, *PTParticle);
@@ -725,8 +733,16 @@ void FRewindData::PushGTDirtyData(const FDirtyPropertiesManager& SrcManager,cons
 	}
 	case EPhysicsProxyType::JointConstraintType:
 	{
-		auto JointProxy = static_cast<FJointConstraintPhysicsProxy*>(Dirty.Proxy);
+		FJointConstraintPhysicsProxy* JointProxy = static_cast<FJointConstraintPhysicsProxy*>(Proxy);
+		if (JointProxy == nullptr)
+		{
+			break;
+		}
 		FPBDJointConstraintHandle* Joint = JointProxy->GetHandle();
+		if (Joint == nullptr)
+		{
+			break;
+		}
 
 		CopyHelper(Joint, [Joint, &DirtyPropHelper](FJointStateBase& Latest)
 		{
