Iris - Don't try to remove too many elements from the queue.

#rnx

diff --git a/Engine/Source/Runtime/Experimental/Iris/Core/Private/Iris/ReplicationSystem/AttachmentReplication.cpp b/Engine/Source/Runtime/Experimental/Iris/Core/Private/Iris/ReplicationSystem/AttachmentReplication.cpp
--- a/Engine/Source/Runtime/Experimental/Iris/Core/Private/Iris/ReplicationSystem/AttachmentReplication.cpp
+++ b/Engine/Source/Runtime/Experimental/Iris/Core/Private/Iris/ReplicationSystem/AttachmentReplication.cpp
@@ -201,8 +201,9 @@ bool FNetObjectAttachmentSendQueue::Enqueue(TArrayView<const TRefCountPtr<FNetBl
 		const SIZE_T TotalCountNeeded = UnreliableQueue.Count() + Attachments.Num();
 		if (TotalCountNeeded > MaxUnreliableCount)
 		{
-			UE_LOG(LogIris, Verbose, TEXT("Dropping old RPC due to too many unreliable Attachments: %d max: %u"), UnreliableQueue.Count(), MaxUnreliableCount);
-			UnreliableQueue.PopNoCheck(TotalCountNeeded - MaxUnreliableCount);
+			UE_LOG(LogIris, Verbose, TEXT("Dropping old RPCs due to too many unreliable attachments: %u max: %u"), UnreliableQueue.Count(), MaxUnreliableCount);
+			const SIZE_T PopCount = FPlatformMath::Min(TotalCountNeeded - MaxUnreliableCount, UnreliableQueue.Count());
+			UnreliableQueue.Pop(PopCount);
 		}
 
 		for (const TRefCountPtr<FNetBlob>& Attachment : Attachments)
@@ -263,7 +264,7 @@ void FNetObjectAttachmentSendQueue::SetUnreliableQueueCapacity(uint32 QueueCapac
 
 	const SIZE_T DropCount = UnreliableCount - QueueCapacity;
 	UE_LOG(LogIris, Warning, TEXT("Dropping %u attachments due to change in unreliable queue capacity to %u"), DropCount, QueueCapacity);
-	UnreliableQueue.PopNoCheck(DropCount);
+	UnreliableQueue.Pop(DropCount);
 }
 
 EAttachmentWriteStatus FNetObjectAttachmentSendQueue::Serialize(FNetSerializationContext& Context, FNetRefHandle RefHandle, FNetObjectAttachmentSendQueue::FCommitRecord& OutRecord, bool& bOutHasUnsentAttachments)
