Geometry Collection : fix negative inertia tensor and mass when mirroring

diff --git a/Engine/Source/Runtime/Experimental/Chaos/Private/PhysicsProxy/GeometryCollectionPhysicsProxy.cpp b/Engine/Source/Runtime/Experimental/Chaos/Private/PhysicsProxy/GeometryCollectionPhysicsProxy.cpp
--- a/Engine/Source/Runtime/Experimental/Chaos/Private/PhysicsProxy/GeometryCollectionPhysicsProxy.cpp
+++ b/Engine/Source/Runtime/Experimental/Chaos/Private/PhysicsProxy/GeometryCollectionPhysicsProxy.cpp
@@ -1893,7 +1893,7 @@ float FGeometryCollectionPhysicsProxy::AdjustMassForScale(float Mass) const
 	const FVector WorldScale = Parameters.WorldTransform.GetScale3D();
 	float MassScale = (float)(WorldScale.X * WorldScale.Y * WorldScale.Z);
 	MassScale *= Parameters.MaterialOverrideMassScaleMultiplier;
-	return (Mass * MassScale);
+	return FMath::Abs(Mass * MassScale);
 }
 
 Chaos::FVec3f FGeometryCollectionPhysicsProxy::AdjustInertiaForScale(const Chaos::FVec3f& Inertia) const
diff --git a/Engine/Source/Runtime/Experimental/Chaos/Public/Chaos/Utilities.h b/Engine/Source/Runtime/Experimental/Chaos/Public/Chaos/Utilities.h
--- a/Engine/Source/Runtime/Experimental/Chaos/Public/Chaos/Utilities.h
+++ b/Engine/Source/Runtime/Experimental/Chaos/Public/Chaos/Utilities.h
@@ -464,12 +464,15 @@ namespace Chaos
 		template<typename T>
 		inline TVec3<T> ScaleInertia(const TVec3<T>& Inertia, const TVec3<T>& Scale, const bool bScaleMass)
 		{
-			TVec3<T> XYZSq = (TVec3<T>(0.5f * (Inertia.X + Inertia.Y + Inertia.Z)) - Inertia) * Scale * Scale;
+			// support for negative scale 
+			const TVec3<T> AbsScale = Scale.GetAbs();
+
+			TVec3<T> XYZSq = (TVec3<T>(0.5f * (Inertia.X + Inertia.Y + Inertia.Z)) - Inertia) * AbsScale * AbsScale;
 			T XX = XYZSq.Y + XYZSq.Z;
 			T YY = XYZSq.X + XYZSq.Z;
 			T ZZ = XYZSq.X + XYZSq.Y;
 			TVec3<T> ScaledInertia = TVec3<T>(XX, YY, ZZ);
-			T MassScale = (bScaleMass) ? Scale.X * Scale.Y * Scale.Z : 1.0f;
+			T MassScale = (bScaleMass) ? AbsScale.X * AbsScale.Y * AbsScale.Z : 1.0f;
 			return MassScale * ScaledInertia;
 		}
 
@@ -479,16 +482,19 @@ namespace Chaos
 		 */
 		inline FMatrix33 ScaleInertia(const FMatrix33& Inertia, const FVec3& Scale, const bool bScaleMass)
 		{
+			// support for negative scale 
+			const FVec3 AbsScale = Scale.GetAbs();
+
 			// @todo(chaos): do we need to support a rotation of the scale axes?
 			FVec3 D = Inertia.GetDiagonal();
-			FVec3 XYZSq = (FVec3(0.5f * (D.X + D.Y + D.Z)) - D) * Scale * Scale;
+			FVec3 XYZSq = (FVec3(0.5f * (D.X + D.Y + D.Z)) - D) * AbsScale * AbsScale;
 			FReal XX = XYZSq.Y + XYZSq.Z;
 			FReal YY = XYZSq.X + XYZSq.Z;
 			FReal ZZ = XYZSq.X + XYZSq.Y;
-			FReal XY = Inertia.M[0][1] * Scale.X * Scale.Y;
-			FReal XZ = Inertia.M[0][2] * Scale.X * Scale.Z;
-			FReal YZ = Inertia.M[1][2] * Scale.Y * Scale.Z;
-			FReal MassScale = (bScaleMass) ? Scale.X * Scale.Y * Scale.Z : 1.0f;
+			FReal XY = Inertia.M[0][1] * AbsScale.X * AbsScale.Y;
+			FReal XZ = Inertia.M[0][2] * AbsScale.X * AbsScale.Z;
+			FReal YZ = Inertia.M[1][2] * AbsScale.Y * AbsScale.Z;
+			FReal MassScale = (bScaleMass) ? AbsScale.X * AbsScale.Y * AbsScale.Z : 1.0f;
 			FMatrix33 ScaledInertia = FMatrix33(
 				MassScale * XX, MassScale * XY, MassScale * XZ,
 				MassScale * XY, MassScale * YY, MassScale * YZ,
