Dataflow:
- Added Copy/Cut/Paste
- Copy/Cut puts the selection into the Copy/Paste buffer in a human readable format
- Bug fixes

diff --git a/Engine/Plugins/Experimental/Dataflow/Source/DataflowEditor/Private/Dataflow/AssetDefinition_DataflowAsset.cpp b/Engine/Plugins/Experimental/Dataflow/Source/DataflowEditor/Private/Dataflow/AssetDefinition_DataflowAsset.cpp
--- a/Engine/Plugins/Experimental/Dataflow/Source/DataflowEditor/Private/Dataflow/AssetDefinition_DataflowAsset.cpp
+++ b/Engine/Plugins/Experimental/Dataflow/Source/DataflowEditor/Private/Dataflow/AssetDefinition_DataflowAsset.cpp
@@ -192,5 +192,31 @@ EAssetCommandResult UAssetDefinition_DataflowAsset::OpenAssets(const FAssetOpenA
 	return EAssetCommandResult::Unhandled;
 }
 
+FString FDataflowConnectionData::GetNode(const FString InConnection)
+{
+	FString Left, Right;
+
+	if (InConnection.Split(TEXT(":"), &Left, nullptr, ESearchCase::IgnoreCase, ESearchDir::FromEnd))
+	{
+		if (Left.Split(TEXT("/"), nullptr, &Right, ESearchCase::IgnoreCase, ESearchDir::FromEnd))
+		{
+			return Right;
+		}
+	}
+
+	return FString("");
+}
+
+FString FDataflowConnectionData::GetProperty(const FString InConnection)
+{
+	FString Right;
+
+	if (InConnection.Split(TEXT(":"), nullptr, &Right, ESearchCase::IgnoreCase, ESearchDir::FromEnd))
+	{
+		return Right;
+	}
+
+	return FString("");
+}
 
 #undef LOCTEXT_NAMESPACE
diff --git a/Engine/Plugins/Experimental/Dataflow/Source/DataflowEditor/Private/Dataflow/DataflowEditorCommands.cpp b/Engine/Plugins/Experimental/Dataflow/Source/DataflowEditor/Private/Dataflow/DataflowEditorCommands.cpp
--- a/Engine/Plugins/Experimental/Dataflow/Source/DataflowEditor/Private/Dataflow/DataflowEditorCommands.cpp
+++ b/Engine/Plugins/Experimental/Dataflow/Source/DataflowEditor/Private/Dataflow/DataflowEditorCommands.cpp
@@ -2,6 +2,7 @@
 
 #include "Dataflow/DataflowEditorCommands.h"
 
+#include "Dataflow/AssetDefinition_DataflowAsset.h"
 #include "Dataflow/DataflowEdNode.h"
 #include "Dataflow/DataflowContent.h"
 #include "Dataflow/DataflowEditorStyle.h"
@@ -17,6 +18,15 @@
 #include "EdGraph/EdGraphNode.h"
 #include "Editor.h"
 #include "IStructureDetailsView.h"
+#include "Serialization/ObjectWriter.h"
+#include "Serialization/ObjectReader.h"
+#include "Dataflow/DataflowGraph.h"
+#include "Framework/Notifications/NotificationManager.h"
+#include "Widgets/Notifications/SNotificationList.h"
+
+#if WITH_EDITOR
+#include "HAL/PlatformApplicationMisc.h"
+#endif
 
 #define LOCTEXT_NAMESPACE "DataflowEditorCommands"
 
@@ -434,6 +444,20 @@ static UEdGraphPin* GetPin(const UDataflowEdNode* Node, const EEdGraphPinDirecti
 	return nullptr;
 }
 
+static void ShowNotificationMessage(const FText& Message, const SNotificationItem::ECompletionState CompletionState)
+{
+	FNotificationInfo Info(Message);
+	Info.ExpireDuration = 5.0f;
+	Info.bUseLargeFont = false;
+	Info.bUseThrobber = false;
+	Info.bUseSuccessFailIcons = false;
+	TSharedPtr<SNotificationItem> Notification = FSlateNotificationManager::Get().AddNotification(Info);
+	if (Notification.IsValid())
+	{
+		Notification->SetCompletionState(CompletionState);
+	}
+}
+
 void FDataflowEditorCommands::DuplicateNodes(UDataflow* Graph, const TSharedPtr<SDataflowGraphEditor>& DataflowGraphEditor, const FGraphPanelSelectionSet& SelectedNodes)
 {
 	if (ensureMsgf(Graph != nullptr, TEXT("Warning : Failed to find valid graph.")))
@@ -542,6 +566,19 @@ void FDataflowEditorCommands::DuplicateNodes(UDataflow* Graph, const TSharedPtr<
 						}
 					}
 				}
+
+				// Display message stating that nodes were duplicated
+				const int32 NumDuplicatedNodes = SelectedEdNodes.Num();
+				FText MessageFormat;
+				if (NumDuplicatedNodes == 1)
+				{
+					MessageFormat = LOCTEXT("DataflowDuplicatedNodesSingleNode", "{0} node was duplicated");
+				}
+				else
+				{
+					MessageFormat = LOCTEXT("DataflowDuplicatedNodesMultipleNodes", "{0} nodes were duplicated");
+				}
+				ShowNotificationMessage(FText::Format(MessageFormat, NumDuplicatedNodes), SNotificationItem::CS_Success);
 			}
 
 			// Process Comment nodes
@@ -598,4 +635,304 @@ void FDataflowEditorCommands::DuplicateNodes(UDataflow* Graph, const TSharedPtr<
 	}
 }
 
+void FDataflowEditorCommands::CopyNodes(UDataflow* InGraph, const TSharedPtr<SDataflowGraphEditor>& DataflowGraphEditor, const FGraphPanelSelectionSet& InSelectedNodes)
+{
+	if (ensureMsgf(InGraph != nullptr, TEXT("Warning : Failed to find valid graph.")))
+	{
+		// 
+		if (InSelectedNodes.Num() > 0)
+		{
+			// Separate selected nodes into an array of UDataflowEdNodes and
+			// an array of UEdGraphNode_Comments
+			TSet<UDataflowEdNode*> SelectedEdNodes;
+			TSet<UEdGraphNode_Comment*> SelectedEdCommentNodes;
+
+			for (UObject* Node : InSelectedNodes)
+			{
+				if (UDataflowEdNode* EdNode = dynamic_cast<UDataflowEdNode*>(Node))
+				{
+					SelectedEdNodes.Add(EdNode);
+				}
+				else if (UEdGraphNode_Comment* EdCommentNode = dynamic_cast<UEdGraphNode_Comment*>(Node))
+				{
+					SelectedEdCommentNodes.Add(EdCommentNode);
+				}
+			}
+
+			FDataflowAssetEdit Edit = InGraph->EditDataflow();
+			if (Dataflow::FGraph* DataflowGraph = Edit.GetGraph())
+			{
+				FDataflowCopyPasteContent CopyPasteContent;
+
+				TSet<FGuid> NodeGuids;
+				TArray<FDataflowInput*> NodeInputsToSave;
+
+				// Build node data
+				if (SelectedEdNodes.Num() > 0)
+				{
+					for (UDataflowEdNode* EdNode : SelectedEdNodes)
+					{
+						if (TSharedPtr<FDataflowNode> DataflowNode = DataflowGraph->FindBaseNode(EdNode->DataflowNodeGuid))
+						{
+							NodeGuids.Add(DataflowNode->GetGuid());
+							NodeInputsToSave.Append(DataflowNode->GetInputs());
+
+							FName NodeType = DataflowNode->GetType();
+							FName NodeName = DataflowNode->GetName();
+
+							FDataflowNodeData NodeData;
+
+							NodeData.Type = NodeType.ToString();
+							NodeData.Name = NodeName.ToString();
+							NodeData.Position.X = EdNode->NodePosX;
+							NodeData.Position.Y = EdNode->NodePosY;
+
+							FString ContentString;
+
+							TUniquePtr<FDataflowNode> DefaultElement;
+							DataflowNode->TypedScriptStruct()->ExportText(ContentString, DataflowNode.Get(), DataflowNode.Get(), nullptr, PPF_None, nullptr);
+
+							NodeData.Properties = ContentString;
+
+							CopyPasteContent.NodeData.Add(NodeData);
+						}
+					}				
+				}
+
+				// Build comment node data
+				if (SelectedEdCommentNodes.Num() > 0)
+				{
+					for (UEdGraphNode_Comment* CommentEdNode : SelectedEdCommentNodes)
+					{
+						FDataflowCommentNodeData CommentNodeData;
+
+						CommentNodeData.Name = CommentEdNode->NodeComment;
+						CommentNodeData.Size.X = CommentEdNode->NodeWidth;
+						CommentNodeData.Size.Y = CommentEdNode->NodeHeight;
+						CommentNodeData.Position.X = CommentEdNode->NodePosX;
+						CommentNodeData.Position.Y = CommentEdNode->NodePosY;
+						CommentNodeData.Color = CommentEdNode->CommentColor;
+
+						CopyPasteContent.CommentNodeData.Add(CommentNodeData);
+					}
+				}
+
+				// Build connection data
+				for (FDataflowInput* Input : NodeInputsToSave)
+				{
+					if (!Input) continue;
+
+					const FDataflowNode* InputNode = Input->GetOwningNode();
+					if (!InputNode) continue;
+
+					const FDataflowOutput* Output = Input->GetConnection();
+					if (!Output) continue;
+
+					const FDataflowNode* OutputNode = Output->GetOwningNode();
+					if (!OutputNode) continue;
+
+					if (!NodeGuids.Contains(OutputNode->GetGuid()))
+						continue;
+
+					FString InConnection = FString::Format(TEXT("/{0}:{1}"), {InputNode->GetName().ToString(), Input->GetName().ToString()});
+					FString OutConnection = FString::Format(TEXT("/{0}:{1}"), { OutputNode->GetName().ToString(), Output->GetName().ToString() });
+
+					FDataflowConnectionData DataflowConnectionData;
+
+					DataflowConnectionData.In = InConnection;
+					DataflowConnectionData.Out = OutConnection;
+
+					CopyPasteContent.ConnectionData.Add(DataflowConnectionData);
+				}
+
+				FString ClipboardContent;
+				FDataflowCopyPasteContent DefaultContent;
+
+				FDataflowCopyPasteContent::StaticStruct()->ExportText(ClipboardContent, &CopyPasteContent, &DefaultContent, nullptr, PPF_None, nullptr);
+
+				// Save to clipboard
+				FPlatformApplicationMisc::ClipboardCopy(*ClipboardContent);
+
+				// Display message stating that nodes were copied to clipboard
+				const int32 NumCopiedNodes = CopyPasteContent.NodeData.Num();
+				FText MessageFormat;
+				if (NumCopiedNodes == 1)
+				{
+					MessageFormat = LOCTEXT("DataflowCopiedNodesToClipboardSingleNode", "{0} node was copied to clipboard");
+				}
+				else
+				{
+					MessageFormat = LOCTEXT("DataflowCopiedNodesToClipboardMultipleNodes", "{0} nodes were copied to clipboard");
+				}
+				ShowNotificationMessage(FText::Format(MessageFormat, NumCopiedNodes), SNotificationItem::CS_Success);
+			}
+		}
+	}
+}
+
+void FDataflowEditorCommands::PasteNodes(UDataflow* Graph, const TSharedPtr<SDataflowGraphEditor>& DataflowGraphEditor)
+{
+	FString ClipboardContent;
+	FPlatformApplicationMisc::ClipboardPaste(ClipboardContent);
+
+	if (!ClipboardContent.IsEmpty())
+	{
+		FDataflowCopyPasteContent DefaultContent;
+		FDataflowCopyPasteContent CopyPasteContent;
+
+		FDataflowCopyPasteContent::StaticStruct()->ImportText(*ClipboardContent, &CopyPasteContent, nullptr, EPropertyPortFlags::PPF_None, nullptr, FDataflowCopyPasteContent::StaticStruct()->GetName(), true);
+
+		FVector2D AppliedTranslation; // The translation that being applied to everything. for placing comment nodes
+
+		// Paste nodes
+		TSet<UDataflowEdNode*> PastedEdNodes;
+		TSet<UEdGraphNode*> PastedEdCommentNodes;
+
+		TMap<FString, UDataflowEdNode*> EdNodeMap;	// [NmaeOfNode -> DuplicatedEdNode]
+
+		if (CopyPasteContent.NodeData.Num() > 0)
+		{
+			// Store the location of the first selected node for recreating spatial relationships
+			FVector2D RefLocation;
+			RefLocation.X = CopyPasteContent.NodeData[0].Position.X; RefLocation.Y = CopyPasteContent.NodeData[0].Position.Y;
+
+			int32 Idx = 0;
+			for (FDataflowNodeData NodeData : CopyPasteContent.NodeData)
+			{
+				FName NodeType = *NodeData.Type;
+
+				if (TSharedPtr<FAssetSchemaAction_Dataflow_PasteNode_DataflowEdNode> PasteNodeAction = FAssetSchemaAction_Dataflow_PasteNode_DataflowEdNode::CreateAction(Graph, NodeType))
+				{
+					PasteNodeAction->NodeName = *NodeData.Name;
+					PasteNodeAction->NodeProperties = NodeData.Properties;
+
+					FVector2D NodeLocation(NodeData.Position.X, NodeData.Position.Y);
+					FVector2D DeltaLocation = NodeLocation - RefLocation;
+
+					if (Idx == 0)
+					{
+						AppliedTranslation = DataflowGraphEditor->GetPasteLocation() + DeltaLocation - NodeData.Position;
+					}
+					Idx++;
+
+					if (UDataflowEdNode* NewEdNode = (UDataflowEdNode*)PasteNodeAction->PerformAction(Graph, nullptr, DataflowGraphEditor->GetPasteLocation() + DeltaLocation, false))
+					{
+						PastedEdNodes.Add(NewEdNode);
+
+						EdNodeMap.Add(NodeData.Name, NewEdNode);
+					}
+				}
+			}
+		}
+
+		// Paste Comment nodes
+		if (CopyPasteContent.CommentNodeData.Num() > 0)
+		{ 
+			const TSharedPtr<SGraphEditor>& InGraphEditor = (TSharedPtr<SGraphEditor>)DataflowGraphEditor;
+
+			int32 Idx = 0;
+
+			for (FDataflowCommentNodeData CommentNodeData : CopyPasteContent.CommentNodeData)
+			{
+				if (TSharedPtr<FAssetSchemaAction_Dataflow_PasteCommentNode_DataflowEdNode> PasteCommentNodeAction = FAssetSchemaAction_Dataflow_PasteCommentNode_DataflowEdNode::CreateAction(Graph, InGraphEditor))
+				{
+					if (CopyPasteContent.NodeData.Num() == 0)
+					{
+						if (Idx == 0)
+						{
+							AppliedTranslation = DataflowGraphEditor->GetPasteLocation() - CommentNodeData.Position;
+						}
+					}
+					Idx++;
+
+					PasteCommentNodeAction->NodeName = *CommentNodeData.Name;
+					PasteCommentNodeAction->Size.X = CommentNodeData.Size.X + 50; // Make it longer, because the nodes are longer after copying ('_copy' in their name)
+					PasteCommentNodeAction->Size.Y = CommentNodeData.Size.Y + 30;
+					PasteCommentNodeAction->Color = CommentNodeData.Color;
+
+					FVector2D CommentNodeLocation(CommentNodeData.Position);
+					FVector2D NewLocation = CommentNodeLocation + AppliedTranslation;
+
+					if (UEdGraphNode* NewCommentNode = PasteCommentNodeAction->PerformAction(Graph, nullptr, NewLocation, false))
+					{
+						PastedEdCommentNodes.Add(NewCommentNode);
+					}
+				}
+			}
+		}
+
+		// Recreate connections
+		if (CopyPasteContent.ConnectionData.Num() > 0)
+		{
+			for (FDataflowConnectionData Connection : CopyPasteContent.ConnectionData)
+			{
+				FString NodeIn = FDataflowConnectionData::GetNode(Connection.In);
+				FGuid GuidIn = (EdNodeMap[NodeIn])->DataflowNodeGuid;
+
+				FString PropertyIn = FDataflowConnectionData::GetProperty(Connection.In);
+				const FName InputputName = *PropertyIn;
+
+				FString NodeOut = FDataflowConnectionData::GetNode(Connection.Out);
+				FGuid GuidOut = (EdNodeMap[NodeOut])->DataflowNodeGuid;
+
+				FString PropertyOut = FDataflowConnectionData::GetProperty(Connection.Out);
+				const FName OutputputName = *PropertyOut;
+
+				FDataflowAssetEdit Edit = Graph->EditDataflow();
+				if (Dataflow::FGraph* DataflowGraph = Edit.GetGraph())
+				{
+					if (TSharedPtr<FDataflowNode> DataflowNodeFrom = DataflowGraph->FindBaseNode(GuidOut))
+					{
+						if (TSharedPtr<FDataflowNode> DataflowNodeTo = DataflowGraph->FindBaseNode(GuidIn))
+						{
+							FDataflowInput* InputConnection = DataflowNodeTo->FindInput(InputputName);
+							FDataflowOutput* OutputConnection = DataflowNodeFrom->FindOutput(OutputputName);
+
+							DataflowGraph->Connect(OutputConnection, InputConnection);
+							
+							if (UEdGraphPin* OutputPin = GetPin(EdNodeMap[NodeOut], EEdGraphPinDirection::EGPD_Output, OutputputName))
+							{
+								if (UEdGraphPin* InputPin = GetPin(EdNodeMap[NodeIn], EEdGraphPinDirection::EGPD_Input, InputputName))
+								{
+									OutputPin->MakeLinkTo(InputPin);
+								}
+							}
+						}
+					}
+				}
+			}
+		}
+
+		// Update the selection in the Editor
+		if (PastedEdNodes.Num() > 0 || PastedEdCommentNodes.Num() > 0)
+		{
+			DataflowGraphEditor->ClearSelectionSet();
+
+			for (UDataflowEdNode* Node : PastedEdNodes)
+			{
+				DataflowGraphEditor->SetNodeSelection(Node, true);
+			}
+
+			for (UEdGraphNode* Node : PastedEdCommentNodes)
+			{
+				DataflowGraphEditor->SetNodeSelection(Node, true);
+			}
+		}
+
+		// Display message stating that nodes were pasted from clipboard
+		const int32 NumPastedNodes = CopyPasteContent.NodeData.Num();
+		FText MessageFormat;
+		if (NumPastedNodes == 1)
+		{
+			MessageFormat = LOCTEXT("DataflowPastedNodesFromClipboardSingleNode", "{0} node was pasted from clipboard");
+		}
+		else
+		{
+			MessageFormat = LOCTEXT("DataflowPastedNodesFromClipboardMultipleNodes", "{0} nodes were pasted from clipboard");
+		}
+		ShowNotificationMessage(FText::Format(MessageFormat, NumPastedNodes), SNotificationItem::CS_Success);
+	}
+}
+
+
 #undef LOCTEXT_NAMESPACE
diff --git a/Engine/Plugins/Experimental/Dataflow/Source/DataflowEditor/Private/Dataflow/DataflowGraphEditor.cpp b/Engine/Plugins/Experimental/Dataflow/Source/DataflowEditor/Private/Dataflow/DataflowGraphEditor.cpp
--- a/Engine/Plugins/Experimental/Dataflow/Source/DataflowEditor/Private/Dataflow/DataflowGraphEditor.cpp
+++ b/Engine/Plugins/Experimental/Dataflow/Source/DataflowEditor/Private/Dataflow/DataflowGraphEditor.cpp
@@ -125,6 +125,18 @@ void SDataflowGraphEditor::Construct(const FArguments& InArgs, UObject* InAssetO
 				FCanExecuteAction(),
 				FIsActionChecked::CreateSP(this, &SDataflowGraphEditor::GetPinVisibility, SGraphEditor::Pin_HideNoConnectionNoDefault)
 			);
+			GraphEditorCommands->MapAction(
+				FGenericCommands::Get().Copy,
+				FExecuteAction::CreateSP(this, &SDataflowGraphEditor::CopySelectedNodes)
+			);
+			GraphEditorCommands->MapAction(
+				FGenericCommands::Get().Cut,
+				FExecuteAction::CreateSP(this, &SDataflowGraphEditor::CutSelectedNodes)
+			);
+			GraphEditorCommands->MapAction(
+				FGenericCommands::Get().Paste,
+				FExecuteAction::CreateSP(this, &SDataflowGraphEditor::PasteSelectedNodes)
+			);
 		}
 	}
 
@@ -171,9 +183,12 @@ void SDataflowGraphEditor::DeleteNode()
 		}
 
 		const FGraphPanelSelectionSet& SelectedNodes = GetSelectedNodes();
-		FDataflowEditorCommands::DeleteNodes(DataflowAsset.Get(), SelectedNodes);
+		if (SelectedNodes.Num() > 0)
+		{
+			FDataflowEditorCommands::DeleteNodes(DataflowAsset.Get(), SelectedNodes);
 
-		OnNodeDeletedMulticast.Broadcast(SelectedNodes);
+			OnNodeDeletedMulticast.Broadcast(SelectedNodes);
+		}
 	}
 }
 
@@ -565,5 +580,44 @@ void SDataflowGraphEditor::AddReferencedObjects(FReferenceCollector& Collector)
 	Collector.AddReferencedObject(AssetOwner);
 }
 
+void SDataflowGraphEditor::CopySelectedNodes()
+{
+	if (UDataflow* Graph = DataflowAsset.Get())
+	{
+		const TSharedPtr<SDataflowGraphEditor>& DataflowGraphEditor = SharedThis(this);
+		const FGraphPanelSelectionSet& SelectedNodes = GetSelectedNodes();
+
+		if (SelectedNodes.Num() > 0)
+		{
+			FDataflowEditorCommands::CopyNodes(Graph, DataflowGraphEditor, SelectedNodes);
+		}
+	}
+}
+
+void SDataflowGraphEditor::CutSelectedNodes()
+{
+	if (UDataflow* Graph = DataflowAsset.Get())
+	{
+		const TSharedPtr<SDataflowGraphEditor>& DataflowGraphEditor = SharedThis(this);
+		const FGraphPanelSelectionSet& SelectedNodes = GetSelectedNodes();
+
+		if (SelectedNodes.Num() > 0)
+		{
+			FDataflowEditorCommands::CopyNodes(Graph, DataflowGraphEditor, SelectedNodes);
+
+			FDataflowEditorCommands::DeleteNodes(DataflowAsset.Get(), SelectedNodes);
+		}
+	}
+}
+
+void SDataflowGraphEditor::PasteSelectedNodes()
+{
+	if (UDataflow* Graph = DataflowAsset.Get())
+	{
+		const TSharedPtr<SDataflowGraphEditor>& DataflowGraphEditor = SharedThis(this);
+
+		FDataflowEditorCommands::PasteNodes(Graph, DataflowGraphEditor);
+	}
+}
 
 #undef LOCTEXT_NAMESPACE
diff --git a/Engine/Plugins/Experimental/Dataflow/Source/DataflowEditor/Private/Dataflow/DataflowSCommentNode.cpp b/Engine/Plugins/Experimental/Dataflow/Source/DataflowEditor/Private/Dataflow/DataflowSCommentNode.cpp
--- a/Engine/Plugins/Experimental/Dataflow/Source/DataflowEditor/Private/Dataflow/DataflowSCommentNode.cpp
+++ b/Engine/Plugins/Experimental/Dataflow/Source/DataflowEditor/Private/Dataflow/DataflowSCommentNode.cpp
@@ -14,6 +14,8 @@
 
 #define LOCTEXT_NAMESPACE "SDataflowEdNodeComment"
 
+constexpr const char* CopyPostfix = "_copy";
+
 //
 // Add a menu option to create a graph node.
 //
@@ -82,8 +84,13 @@ TSharedPtr<FAssetSchemaAction_Dataflow_DuplicateCommentNode_DataflowEdNode> FAss
 //
 UEdGraphNode* FAssetSchemaAction_Dataflow_DuplicateCommentNode_DataflowEdNode::PerformAction(class UEdGraph* ParentGraph, UEdGraphPin* FromPin, const FVector2D Location, bool bSelectNewNode)
 {
-	FString NewNodeName = CommentNodeToDuplicate->GetFName().ToString();
-	NewNodeName.Append("_copy");
+	FString NewNodeName = CommentNodeToDuplicate->NodeComment;
+
+	// Append "_copy" to selected node's name if it doesn't have it
+	if (!NewNodeName.Contains(CopyPostfix))
+	{
+		NewNodeName.Append(CopyPostfix);
+	}
 
 	UEdGraphNode_Comment* CommentTemplate = NewObject<UEdGraphNode_Comment>(ParentGraph, FName(*NewNodeName));
 
@@ -111,8 +118,9 @@ UEdGraphNode* FAssetSchemaAction_Dataflow_DuplicateCommentNode_DataflowEdNode::P
 	CommentTemplate->NodeWidth = CommentNodeToDuplicate->NodeWidth;
 	CommentTemplate->NodeHeight = CommentNodeToDuplicate->NodeHeight;
 	CommentTemplate->SnapToGrid(GetDefault<UEditorStyleSettings>()->GridSnapSize);
+	CommentTemplate->CommentColor = CommentNodeToDuplicate->CommentColor;
 
-	CommentTemplate->NodeComment = CommentNodeToDuplicate->NodeComment;
+	CommentTemplate->NodeComment = NewNodeName;
 
 	ParentGraph->NotifyGraphChanged();
 
@@ -125,5 +133,54 @@ void FAssetSchemaAction_Dataflow_DuplicateCommentNode_DataflowEdNode::AddReferen
 	Collector.AddReferencedObject(CommentNodeToDuplicate);
 }
 
+TSharedPtr<FAssetSchemaAction_Dataflow_PasteCommentNode_DataflowEdNode> FAssetSchemaAction_Dataflow_PasteCommentNode_DataflowEdNode::CreateAction(UEdGraph* ParentGraph, const TSharedPtr<SGraphEditor>& GraphEditor)
+{
+	return MakeShared<FAssetSchemaAction_Dataflow_PasteCommentNode_DataflowEdNode>(GraphEditor);
+}
+
+//
+//  Paste comment node
+//
+UEdGraphNode* FAssetSchemaAction_Dataflow_PasteCommentNode_DataflowEdNode::PerformAction(class UEdGraph* ParentGraph, UEdGraphPin* FromPin, const FVector2D Location, bool bSelectNewNode)
+{
+	// Append "_copy" to selected node's name if it doesn't have it
+	FString NewNodeName = NodeName.ToString();
+	if (!NewNodeName.Contains(CopyPostfix))
+	{
+		NewNodeName.Append(CopyPostfix);
+	}
+
+	UEdGraphNode_Comment* CommentTemplate = NewObject<UEdGraphNode_Comment>(ParentGraph, FName(*NewNodeName));
+
+	CommentTemplate->bCommentBubbleVisible_InDetailsPanel = false;
+	CommentTemplate->bCommentBubbleVisible = false;
+	CommentTemplate->bCommentBubblePinned = false;
+
+	ParentGraph->Modify();
+
+	CommentTemplate->SetFlags(RF_Transactional);
+
+	// set outer to be the graph so it doesn't go away
+	CommentTemplate->Rename(NULL, ParentGraph, REN_NonTransactional);
+	ParentGraph->AddNode(CommentTemplate, true, bSelectNewNode);
+
+	CommentTemplate->CreateNewGuid();
+	CommentTemplate->PostPlacedNewNode();
+	CommentTemplate->AllocateDefaultPins();
+	CommentTemplate->AutowireNewNode(FromPin);
+
+	CommentTemplate->NodePosX = Location.X;
+	CommentTemplate->NodePosY = Location.Y;
+	CommentTemplate->NodeWidth = Size.X;
+	CommentTemplate->NodeHeight = Size.Y;
+	CommentTemplate->SnapToGrid(GetDefault<UEditorStyleSettings>()->GridSnapSize);
+	CommentTemplate->CommentColor = Color;
+
+	CommentTemplate->NodeComment = NewNodeName;
+
+	ParentGraph->NotifyGraphChanged();
+
+	return CommentTemplate;
+}
 
 #undef LOCTEXT_NAMESPACE
diff --git a/Engine/Plugins/Experimental/Dataflow/Source/DataflowEditor/Private/Dataflow/DataflowSNode.cpp b/Engine/Plugins/Experimental/Dataflow/Source/DataflowEditor/Private/Dataflow/DataflowSNode.cpp
--- a/Engine/Plugins/Experimental/Dataflow/Source/DataflowEditor/Private/Dataflow/DataflowSNode.cpp
+++ b/Engine/Plugins/Experimental/Dataflow/Source/DataflowEditor/Private/Dataflow/DataflowSNode.cpp
@@ -223,7 +223,7 @@ static UDataflowEdNode* CreateNode(UDataflow* Dataflow, UEdGraphPin* FromPin, co
 				{
 					SDataflowEdNode::CopyDataflowNodeSettings(DataflowNodeToDuplicate, DataflowNode);
 				}
-
+				
 				EdNode->CreateNewGuid();
 				EdNode->PostPlacedNewNode();
 
@@ -328,5 +328,97 @@ UEdGraphNode* FAssetSchemaAction_Dataflow_DuplicateNode_DataflowEdNode::PerformA
 	return nullptr;
 }
 
+//
+// 
+//
+TSharedPtr<FAssetSchemaAction_Dataflow_PasteNode_DataflowEdNode> FAssetSchemaAction_Dataflow_PasteNode_DataflowEdNode::CreateAction(UEdGraph* ParentGraph, const FName& InNodeTypeName)
+{
+	if (Dataflow::FNodeFactory* Factory = Dataflow::FNodeFactory::GetInstance())
+	{
+		const Dataflow::FFactoryParameters& Param = Factory->GetParameters(InNodeTypeName);
+		if (Param.IsValid())
+		{
+			const FText ToolTip = FText::FromString(Param.ToolTip.IsEmpty() ? FString("Add a Dataflow node.") : Param.ToolTip);
+			const FText NodeName = FText::FromString(Param.DisplayName.ToString());
+			const FText Category = FText::FromString(Param.Category.ToString().IsEmpty() ? FString("Dataflow") : Param.Category.ToString());
+			const FText Tags = FText::FromString(Param.Tags);
+			TSharedPtr<FAssetSchemaAction_Dataflow_PasteNode_DataflowEdNode> NewNodeAction(
+				new FAssetSchemaAction_Dataflow_PasteNode_DataflowEdNode(InNodeTypeName, Category, NodeName, ToolTip, Tags));
+			return NewNodeAction;
+		}
+	}
+	return TSharedPtr<FAssetSchemaAction_Dataflow_PasteNode_DataflowEdNode>(nullptr);
+}
+
+static UDataflowEdNode* CreateNodeFromPaste(UDataflow* Dataflow, UEdGraphPin* FromPin, const FVector2D Location, bool bSelectNewNode, const FName NodeUniqueName, const FName NodeTypeName, FString NodeProperties)
+{
+	if (Dataflow::FNodeFactory* Factory = Dataflow::FNodeFactory::GetInstance())
+	{
+		if (TSharedPtr<FDataflowNode> DataflowNode =
+			Factory->NewNodeFromRegisteredType(
+				*Dataflow->GetDataflow(),
+				{ FGuid::NewGuid(), NodeTypeName, NodeUniqueName, Dataflow }))
+		{
+			if (UDataflowEdNode* EdNode = NewObject<UDataflowEdNode>(Dataflow, UDataflowEdNode::StaticClass(), NodeUniqueName))
+			{
+				Dataflow->Modify();
+
+				Dataflow->AddNode(EdNode, true, bSelectNewNode);
+
+				// Copy properties to DataflowNode
+				if (!NodeProperties.IsEmpty())
+				{
+					DataflowNode->TypedScriptStruct()->ImportText(*NodeProperties, DataflowNode.Get(), nullptr, EPropertyPortFlags::PPF_None, nullptr, DataflowNode->TypedScriptStruct()->GetName(), true);
+				}
+
+				EdNode->CreateNewGuid();
+				EdNode->PostPlacedNewNode();
+
+				EdNode->SetDataflowGraph(Dataflow->GetDataflow());
+				EdNode->SetDataflowNodeGuid(DataflowNode->GetGuid());
+				EdNode->AllocateDefaultPins();
+
+				EdNode->NodePosX = Location.X;
+				EdNode->NodePosY = Location.Y;
+
+				EdNode->SetFlags(RF_Transactional);
+
+				return EdNode;
+			}
+		}
+	}
+
+	return nullptr;
+}
+
+//
+//   
+//
+UEdGraphNode* FAssetSchemaAction_Dataflow_PasteNode_DataflowEdNode::PerformAction(class UEdGraph* ParentGraph, UEdGraphPin* FromPin, const FVector2D Location, bool bSelectNewNode)
+{
+	if (UDataflow* Dataflow = Cast<UDataflow>(ParentGraph))
+	{
+		// Append "_copy" to selected node's name if it doesn't have it
+		FString NodeToDuplicateName = NodeName.ToString();
+		if (!NodeToDuplicateName.Contains("_copy"))
+		{
+			NodeToDuplicateName.Append("_copy");
+		}
+
+		// Check if that is unique, if not then make it unique with an index postfix
+		const FString NodeBaseName = NodeToDuplicateName;
+		FName NodeUniqueName{ NodeBaseName };
+		int32 NameIndex = 0;
+		while (Dataflow->GetDataflow()->FindBaseNode(FName(NodeUniqueName)) != nullptr)
+		{
+			NodeUniqueName = FName(NodeBaseName + FString::Printf(TEXT("_%d"), NameIndex));
+			NameIndex++;
+		}
+
+		return CreateNodeFromPaste(Dataflow, FromPin, Location, bSelectNewNode, NodeUniqueName, NodeTypeName, NodeProperties);
+	}
+
+	return nullptr;
+}
 #undef LOCTEXT_NAMESPACE
 
diff --git a/Engine/Plugins/Experimental/Dataflow/Source/DataflowEditor/Public/Dataflow/AssetDefinition_DataflowAsset.h b/Engine/Plugins/Experimental/Dataflow/Source/DataflowEditor/Public/Dataflow/AssetDefinition_DataflowAsset.h
--- a/Engine/Plugins/Experimental/Dataflow/Source/DataflowEditor/Public/Dataflow/AssetDefinition_DataflowAsset.h
+++ b/Engine/Plugins/Experimental/Dataflow/Source/DataflowEditor/Public/Dataflow/AssetDefinition_DataflowAsset.h
@@ -22,6 +22,74 @@ namespace DataflowAssetDefinitionHelpers
 	DATAFLOWEDITOR_API UObject* NewOrOpenDataflowAsset(const UObject* Asset);
 }
 
+//
+// Copy/Paste structs
+//
+USTRUCT()
+struct DATAFLOWEDITOR_API FDataflowNodeData
+{
+	GENERATED_BODY()
+
+	UPROPERTY()
+	FString Type;
+
+	UPROPERTY()
+	FString Name;
+
+	UPROPERTY()
+	FString Properties;
+
+	UPROPERTY()
+	FVector2D Position = FVector2D::ZeroVector;
+};
+
+USTRUCT()
+struct DATAFLOWEDITOR_API FDataflowCommentNodeData
+{
+	GENERATED_BODY()
+
+	UPROPERTY()
+	FString Name;
+
+	UPROPERTY()
+	FVector2D Size = FVector2D::ZeroVector;
+
+	UPROPERTY()
+	FLinearColor Color = FLinearColor::White;
+
+	UPROPERTY()
+	FVector2D Position = FVector2D::ZeroVector;
+};
+
+USTRUCT()
+struct DATAFLOWEDITOR_API FDataflowConnectionData
+{
+	GENERATED_BODY()
+
+	UPROPERTY()
+	FString Out;
+
+	UPROPERTY()
+	FString In;
+
+	static FString GetNode(const FString InConnection);
+	static FString GetProperty(const FString InConnection);
+};
+
+USTRUCT()
+struct DATAFLOWEDITOR_API FDataflowCopyPasteContent
+{
+	GENERATED_BODY()
+
+	UPROPERTY()
+	TArray<FDataflowNodeData> NodeData;
+
+	UPROPERTY()
+	TArray<FDataflowCommentNodeData> CommentNodeData;
+
+	UPROPERTY()
+	TArray<FDataflowConnectionData> ConnectionData;
+};
 
 
 UCLASS()
diff --git a/Engine/Plugins/Experimental/Dataflow/Source/DataflowEditor/Public/Dataflow/DataflowEditorCommands.h b/Engine/Plugins/Experimental/Dataflow/Source/DataflowEditor/Public/Dataflow/DataflowEditorCommands.h
--- a/Engine/Plugins/Experimental/Dataflow/Source/DataflowEditor/Public/Dataflow/DataflowEditorCommands.h
+++ b/Engine/Plugins/Experimental/Dataflow/Source/DataflowEditor/Public/Dataflow/DataflowEditorCommands.h
@@ -134,6 +134,14 @@ public:
 	*/
 	static void DuplicateNodes(UDataflow* Graph, const TSharedPtr<SDataflowGraphEditor>& DataflowGraphEditor, const FGraphPanelSelectionSet& SelectedNodes);
 
+	/*
+	*  CopyNodes
+	*/
+	static void CopyNodes(UDataflow* Graph, const TSharedPtr<SDataflowGraphEditor>& DataflowGraphEditor, const FGraphPanelSelectionSet& SelectedNodes);
 
-
+	/*
+	*  PasteSelectedNodes
+	*/
+	static void PasteNodes(UDataflow* Graph, const TSharedPtr<SDataflowGraphEditor>& DataflowGraphEditor);
 };
+
diff --git a/Engine/Plugins/Experimental/Dataflow/Source/DataflowEditor/Public/Dataflow/DataflowGraphEditor.h b/Engine/Plugins/Experimental/Dataflow/Source/DataflowEditor/Public/Dataflow/DataflowGraphEditor.h
--- a/Engine/Plugins/Experimental/Dataflow/Source/DataflowEditor/Public/Dataflow/DataflowGraphEditor.h
+++ b/Engine/Plugins/Experimental/Dataflow/Source/DataflowEditor/Public/Dataflow/DataflowGraphEditor.h
@@ -122,6 +122,15 @@ public:
 	/** */
 	void ZoomToFitGraph();
 
+	/** */
+	void CopySelectedNodes();
+
+	/** */
+	void CutSelectedNodes();
+
+	/** */
+	void PasteSelectedNodes();
+
 	SGraphEditor* GetGraphEditor() { return (SGraphEditor*)this; }
 
 	/** FGCObject interface */
diff --git a/Engine/Plugins/Experimental/Dataflow/Source/DataflowEditor/Public/Dataflow/DataflowSCommentNode.h b/Engine/Plugins/Experimental/Dataflow/Source/DataflowEditor/Public/Dataflow/DataflowSCommentNode.h
--- a/Engine/Plugins/Experimental/Dataflow/Source/DataflowEditor/Public/Dataflow/DataflowSCommentNode.h
+++ b/Engine/Plugins/Experimental/Dataflow/Source/DataflowEditor/Public/Dataflow/DataflowSCommentNode.h
@@ -73,3 +73,28 @@ public:
 
 	TObjectPtr<UEdGraphNode_Comment> CommentNodeToDuplicate;
 };
+
+//
+// Action to paste a comment node in the graph
+//
+USTRUCT()
+struct DATAFLOWEDITOR_API FAssetSchemaAction_Dataflow_PasteCommentNode_DataflowEdNode : public FEdGraphSchemaAction
+{
+	GENERATED_USTRUCT_BODY();
+
+public:
+	FAssetSchemaAction_Dataflow_PasteCommentNode_DataflowEdNode(const TSharedPtr<SGraphEditor>& InGraphEditor) : FEdGraphSchemaAction()
+		, GraphEditor(InGraphEditor)
+	{}
+
+	FAssetSchemaAction_Dataflow_PasteCommentNode_DataflowEdNode() {}
+
+	static TSharedPtr<FAssetSchemaAction_Dataflow_PasteCommentNode_DataflowEdNode> CreateAction(UEdGraph* ParentGraph, const TSharedPtr<SGraphEditor>& GraphEditor);
+
+	virtual UEdGraphNode* PerformAction(class UEdGraph* ParentGraph, UEdGraphPin* FromPin, const FVector2D Location, bool bSelectNewNode = true) override;
+
+	TSharedPtr<SGraphEditor> GraphEditor;
+	FName NodeName;
+	FVector2D Size;
+	FLinearColor Color;
+};
diff --git a/Engine/Plugins/Experimental/Dataflow/Source/DataflowEditor/Public/Dataflow/DataflowSNode.h b/Engine/Plugins/Experimental/Dataflow/Source/DataflowEditor/Public/Dataflow/DataflowSNode.h
--- a/Engine/Plugins/Experimental/Dataflow/Source/DataflowEditor/Public/Dataflow/DataflowSNode.h
+++ b/Engine/Plugins/Experimental/Dataflow/Source/DataflowEditor/Public/Dataflow/DataflowSNode.h
@@ -99,3 +99,26 @@ public:
 	FName NodeTypeName;
 	TSharedPtr<FDataflowNode> DataflowNodeToDuplicate;
 };
+
+//
+// Action to paste a set of nodes in the graph
+//
+USTRUCT()
+struct DATAFLOWEDITOR_API FAssetSchemaAction_Dataflow_PasteNode_DataflowEdNode : public FEdGraphSchemaAction
+{
+	GENERATED_USTRUCT_BODY();
+
+public:
+	FAssetSchemaAction_Dataflow_PasteNode_DataflowEdNode() : FEdGraphSchemaAction() {}
+
+	FAssetSchemaAction_Dataflow_PasteNode_DataflowEdNode(const FName& InType, FText InNodeCategory, FText InMenuDesc, FText InToolTip, FText InKeywords)
+		: FEdGraphSchemaAction(InNodeCategory, InMenuDesc, InToolTip, 0, InKeywords), NodeTypeName(InType) {}
+
+	static TSharedPtr<FAssetSchemaAction_Dataflow_PasteNode_DataflowEdNode> CreateAction(UEdGraph* ParentGraph, const FName& NodeTypeName);
+
+	virtual UEdGraphNode* PerformAction(class UEdGraph* ParentGraph, UEdGraphPin* FromPin, const FVector2D Location, bool bSelectNewNode = true) override;
+
+	FName NodeTypeName;
+	FName NodeName;
+	FString NodeProperties;
+};
diff --git a/Engine/Source/Runtime/Experimental/Dataflow/Core/DataflowCore.Build.cs b/Engine/Source/Runtime/Experimental/Dataflow/Core/DataflowCore.Build.cs
--- a/Engine/Source/Runtime/Experimental/Dataflow/Core/DataflowCore.Build.cs
+++ b/Engine/Source/Runtime/Experimental/Dataflow/Core/DataflowCore.Build.cs
@@ -14,6 +14,7 @@ namespace UnrealBuildTool.Rules
 					"Chaos",
 					"DeveloperSettings",
 					"InputCore",
+					"ApplicationCore"
 				}
 			);
 		}
diff --git a/Engine/Source/Runtime/Experimental/Dataflow/Core/Private/Dataflow/DataflowGraph.cpp b/Engine/Source/Runtime/Experimental/Dataflow/Core/Private/Dataflow/DataflowGraph.cpp
--- a/Engine/Source/Runtime/Experimental/Dataflow/Core/Private/Dataflow/DataflowGraph.cpp
+++ b/Engine/Source/Runtime/Experimental/Dataflow/Core/Private/Dataflow/DataflowGraph.cpp
@@ -134,209 +134,218 @@ namespace Dataflow
 		Ar << Guid;
 		if (Ar.IsSaving())
 		{
-			FGuid ArGuid;
-			FName ArType, ArName;
-			int32 ArNum = Nodes.Num();
+			FGraph::SerializeForSaving(Ar, this, Nodes, Connections);
+		}
+		else if( Ar.IsLoading())
+		{
+			FGraph::SerializeForLoading(Ar, this, OwningObject);
+		}
+	}
 
-			Ar << ArNum;
-			for (TSharedPtr<FDataflowNode> Node : Nodes)
-			{
-				ArGuid = Node->GetGuid();
-				ArType = Node->GetType();
-				ArName = Node->GetName();
-				Ar << ArGuid << ArType << ArName;
+	void FGraph::SerializeForSaving(FArchive& Ar, FGraph* InGraph, TArray<TSharedPtr<FDataflowNode>>& InNodes, TArray<FLink>& InConnections)
+	{
+		FGuid ArGuid;
+		FName ArType, ArName;
+		int32 ArNum = InNodes.Num();
 
-				DATAFLOW_OPTIONAL_BLOCK_WRITE_BEGIN()
-				{
-					// Node needs to be serialized first to make sure it registers all the dynamic input/output for when input and output will be deserialized
-					Node->SerializeInternal(Ar);
+		Ar << ArNum;
+		for (TSharedPtr<FDataflowNode> Node : InNodes)
+		{
+			ArGuid = Node->GetGuid();
+			ArType = Node->GetType();
+			ArName = Node->GetName();
+			Ar << ArGuid << ArType << ArName;
 
-					// keep outputs and inputs separated even though their serialization code looks almost identical
-					// this is to make sure we can handle when number of inputs or outputs have changed on the node
-					int32 ArNumOutputs = Node->GetOutputs().Num();
-					Ar << ArNumOutputs;
-					for (FDataflowConnection* Output : Node->GetOutputs())
-					{
-						ArGuid = Output->GetGuid();
-						ArType = Output->GetType();
-						ArName = Output->GetName();
-						Ar << ArGuid << ArType << ArName;
-					}
+			DATAFLOW_OPTIONAL_BLOCK_WRITE_BEGIN()
+			{
+				// Node needs to be serialized first to make sure it registers all the dynamic input/output for when input and output will be deserialized
+				Node->SerializeInternal(Ar);
+
+				// keep outputs and inputs separated even though their serialization code looks almost identical
+				// this is to make sure we can handle when number of inputs or outputs have changed on the node
+				int32 ArNumOutputs = Node->GetOutputs().Num();
+				Ar << ArNumOutputs;
+				for (FDataflowConnection* Output : Node->GetOutputs())
+				{
+					ArGuid = Output->GetGuid();
+					ArType = Output->GetType();
+					ArName = Output->GetName();
+					Ar << ArGuid << ArType << ArName;
+				}
 
-					int32 ArNumInputs = Node->GetInputs().Num();
-					Ar << ArNumInputs;
-					for (FDataflowConnection* Input: Node->GetInputs())
-					{
-						ArGuid = Input->GetGuid();
-						ArType = Input->GetType();
-						ArName = Input->GetName();
-						Ar << ArGuid << ArType << ArName;
-					}
+				int32 ArNumInputs = Node->GetInputs().Num();
+				Ar << ArNumInputs;
+				for (FDataflowConnection* Input : Node->GetInputs())
+				{
+					ArGuid = Input->GetGuid();
+					ArType = Input->GetType();
+					ArName = Input->GetName();
+					Ar << ArGuid << ArType << ArName;
 				}
-				DATAFLOW_OPTIONAL_BLOCK_WRITE_END();
 			}
+			DATAFLOW_OPTIONAL_BLOCK_WRITE_END();
+		}
 
-			Ar << Connections;
+		Ar << InConnections;
+	}
 
-		}
-		else if( Ar.IsLoading())
-		{
-			FGuid ArGuid;
-			FName ArType, ArName;
-			int32 ArNum = 0;
+	void FGraph::SerializeForLoading(FArchive& Ar, FGraph* InGraph, UObject* OwningObject)
+	{
+		FGuid ArGuid;
+		FName ArType, ArName;
+		int32 ArNum = 0;
+
+		TMap<FGuid, TSharedPtr<FDataflowNode> > NodeGuidMap;
+		TMap<FGuid, FDataflowConnection* > ConnectionGuidMap;
 
-			TMap<FGuid, TSharedPtr<FDataflowNode> > NodeGuidMap;
-			TMap<FGuid, FDataflowConnection* > ConnectionGuidMap;
+		Ar << ArNum;
+		for (int32 Ndx = ArNum; Ndx > 0; Ndx--)
+		{
+			FName ArNodeName;
+			Ar << ArGuid << ArType << ArNodeName;
 
-			Ar << ArNum;
-			for (int32 Ndx = ArNum; Ndx > 0; Ndx--)
+			TSharedPtr<FDataflowNode> Node = FNodeFactory::GetInstance()->NewNodeFromRegisteredType(*InGraph, { ArGuid, ArType, ArNodeName, OwningObject });
+			DATAFLOW_OPTIONAL_BLOCK_READ_BEGIN(Node != nullptr)
 			{
-				FName ArNodeName;
-				Ar << ArGuid << ArType << ArNodeName;
+				ensure(!NodeGuidMap.Contains(ArGuid));
+				NodeGuidMap.Add(ArGuid, Node);
 
-				TSharedPtr<FDataflowNode> Node = FNodeFactory::GetInstance()->NewNodeFromRegisteredType(*this, { ArGuid, ArType, ArNodeName, OwningObject });
-				DATAFLOW_OPTIONAL_BLOCK_READ_BEGIN(Node != nullptr)
+				if ((Ar.CustomVer(FUE5MainStreamObjectVersion::GUID) < FUE5MainStreamObjectVersion::DataflowSeparateInputOutputSerialization))
 				{
-					ensure(!NodeGuidMap.Contains(ArGuid));
-					NodeGuidMap.Add(ArGuid, Node);
 
-					if ((Ar.CustomVer(FUE5MainStreamObjectVersion::GUID) < FUE5MainStreamObjectVersion::DataflowSeparateInputOutputSerialization))
+					// former input / output serialization method where we only store aggregate number of inputs and outputs
+					// this has limitation when adding more outputs or inputs
+					int ArNumInputsOutputs;
+					Ar << ArNumInputsOutputs;
+					TArray< FDataflowConnection* > InputsOutputs;
+					InputsOutputs.Append(Node->GetOutputs());
+					InputsOutputs.Append(Node->GetInputs());
+
+					// skip offset is use to correct the mismatch of outputs have been added
+					int SkipOffset = 0;
+					for (int ConnectionIndex = 0; ConnectionIndex < ArNumInputsOutputs; ConnectionIndex++)
 					{
+						Ar << ArGuid << ArType << ArName;
 
-						// former input / output serialization method where we only store aggregate number of inputs and outputs
-						// this has limitation when adding more outputs or inputs
-						int ArNumInputsOutputs;
-						Ar << ArNumInputsOutputs;
-						TArray< FDataflowConnection* > InputsOutputs;
-						InputsOutputs.Append(Node->GetOutputs());
-						InputsOutputs.Append(Node->GetInputs());
-
-						// skip offset is use to correct the mismatch of outputs have been added
-						int SkipOffset = 0;
-						for (int ConnectionIndex = 0; ConnectionIndex < ArNumInputsOutputs; ConnectionIndex++)
+						int AdjustedConnectionIndex = ConnectionIndex + SkipOffset;
+						if (InputsOutputs.IsValidIndex(AdjustedConnectionIndex))
 						{
-							Ar << ArGuid << ArType << ArName;
-
-							int AdjustedConnectionIndex = ConnectionIndex + SkipOffset;
-							if (InputsOutputs.IsValidIndex(AdjustedConnectionIndex))
+							FDataflowConnection* Connection = InputsOutputs[AdjustedConnectionIndex];
+
+							// if the name does not match this means the node has changed since the last serialization 
+							// ( added outputs for example that shift the index )
+							// in that case we try to recover by finding the next good node
+							// note we cannot just find by name as some nodes have inputs and outputs named the same 
+							// todo: implement a better way to serialize inputs and outputs seperately to avoid this case
+							while (Connection && Connection->GetName() != ArName)
 							{
-								FDataflowConnection* Connection = InputsOutputs[AdjustedConnectionIndex];
-
-								// if the name does not match this means the node has changed since the last serialization 
-								// ( added outputs for example that shift the index )
-								// in that case we try to recover by finding the next good node
-								// note we cannot just find by name as some nodes have inputs and outputs named the same 
-								// todo: implement a better way to serialize inputs and outputs seperately to avoid this case
-								while (Connection && Connection->GetName() != ArName)
+								SkipOffset++;
+								AdjustedConnectionIndex = ConnectionIndex + SkipOffset;
+								if (InputsOutputs.IsValidIndex(AdjustedConnectionIndex))
 								{
-									SkipOffset++;
-									AdjustedConnectionIndex = ConnectionIndex + SkipOffset;
-									if (InputsOutputs.IsValidIndex(AdjustedConnectionIndex))
-									{
-										Connection = InputsOutputs[AdjustedConnectionIndex];
-									}
-									else
-									{
-										Connection = nullptr;
-									}
+									Connection = InputsOutputs[AdjustedConnectionIndex];
 								}
-								if (Connection)
+								else
 								{
-									check(Connection->GetType() == ArType);
-									Connection->SetGuid(ArGuid);
-									ensure(!ConnectionGuidMap.Contains(ArGuid));
-									ConnectionGuidMap.Add(ArGuid, Connection);
+									Connection = nullptr;
 								}
 							}
+							if (Connection)
+							{
+								check(Connection->GetType() == ArType);
+								Connection->SetGuid(ArGuid);
+								ensure(!ConnectionGuidMap.Contains(ArGuid));
+								ConnectionGuidMap.Add(ArGuid, Connection);
+							}
 						}
-
-						Node->SerializeInternal(Ar);
 					}
-					else
+
+					Node->SerializeInternal(Ar);
+				}
+				else
+				{
+					// we need to deserilaize the node first because if it may add more inputs that may 
+					// be referenced when deserializing them below ( see Dataflow Node AddPin method )
+					Node->SerializeInternal(Ar);
+
+					// Outputs deserialization
 					{
-						// we need to deserilaize the node first because if it may add more inputs that may 
-						// be referenced when deserializing them below ( see Dataflow Node AddPin method )
-						Node->SerializeInternal(Ar);
+						int32 ArNumOutputs;
+						Ar << ArNumOutputs;
 
-						// Outputs deserialization
+						for (int32 OutputIndex = 0; OutputIndex < ArNumOutputs; OutputIndex++)
 						{
-							int32 ArNumOutputs;
-							Ar << ArNumOutputs;
+							Ar << ArGuid << ArType << ArName;
 
-							for (int32 OutputIndex = 0; OutputIndex < ArNumOutputs; OutputIndex++)
+							if (FDataflowOutput* Output = Node->FindOutput(ArName))
 							{
-								Ar << ArGuid << ArType << ArName;
-
-								if (FDataflowOutput* Output = Node->FindOutput(ArName))
-								{
-									check(Output->GetType() == ArType);
-									Output->SetGuid(ArGuid);
-									ensure(!ConnectionGuidMap.Contains(ArGuid));
-									ConnectionGuidMap.Add(ArGuid, Output);
-								}
-								else
-								{
-									// output has been serialized but cannot be found
-									// this means the definition of the node has changed and the output is no longer registered
-									UE_LOG(LogChaos, Display, TEXT("Cannot find registered output (%s) in Dataflow node (%s) - this may result in missing connection(s).")
-										, *ArName.ToString(), *ArNodeName.ToString());
-								}
+								check(Output->GetType() == ArType);
+								Output->SetGuid(ArGuid);
+								ensure(!ConnectionGuidMap.Contains(ArGuid));
+								ConnectionGuidMap.Add(ArGuid, Output);
+							}
+							else
+							{
+								// output has been serialized but cannot be found
+								// this means the definition of the node has changed and the output is no longer registered
+								UE_LOG(LogChaos, Display, TEXT("Cannot find registered output (%s) in Dataflow node (%s) - this may result in missing connection(s).")
+									, *ArName.ToString(), *ArNodeName.ToString());
 							}
 						}
+					}
+
+					// Inputs deserialization
+					{
+						int32 ArNumInputs;
+						Ar << ArNumInputs;
 
-						// Inputs deserialization
+						for (int32 InputIndex = 0; InputIndex < ArNumInputs; InputIndex++)
 						{
-							int32 ArNumInputs;
-							Ar << ArNumInputs;
+							Ar << ArGuid << ArType << ArName;
 
-							for (int32 InputIndex = 0; InputIndex < ArNumInputs; InputIndex++)
+							if (FDataflowInput* Input = Node->FindInput(ArName))
 							{
-								Ar << ArGuid << ArType << ArName;
-
-								if (FDataflowInput* Input = Node->FindInput(ArName))
-								{
-									check(Input->GetType() == ArType);
-									Input->SetGuid(ArGuid);
-									ensure(!ConnectionGuidMap.Contains(ArGuid));
-									ConnectionGuidMap.Add(ArGuid, Input);
-								}
-								else
-								{
-									// input has been serialized but cannot be found
-									// this means the definition of the node has changed and the input is no longer registered
-									UE_LOG(LogChaos, Display, TEXT("Cannot find registered input (%s) in Dataflow node (%s) - this may result in missing connection(s).")
-										, *ArName.ToString(), *ArNodeName.ToString());
-								}
+								check(Input->GetType() == ArType);
+								Input->SetGuid(ArGuid);
+								ensure(!ConnectionGuidMap.Contains(ArGuid));
+								ConnectionGuidMap.Add(ArGuid, Input);
+							}
+							else
+							{
+								// input has been serialized but cannot be found
+								// this means the definition of the node has changed and the input is no longer registered
+								UE_LOG(LogChaos, Display, TEXT("Cannot find registered input (%s) in Dataflow node (%s) - this may result in missing connection(s).")
+									, *ArName.ToString(), *ArNodeName.ToString());
 							}
 						}
 					}
 				}
-				DATAFLOW_OPTIONAL_BLOCK_READ_ELSE()
-				{
-					DisabledNodes.Add(ArNodeName);
-					ensureMsgf(false,
-						TEXT("Error: Missing registered node type (%s) will be removed from graph on load. Graph will fail to evaluate due to missing node (%s).")
-						, *ArType.ToString(), *ArName.ToString());
-				}
-				DATAFLOW_OPTIONAL_BLOCK_READ_END();
 			}
+			DATAFLOW_OPTIONAL_BLOCK_READ_ELSE()
+			{
+				InGraph->DisabledNodes.Add(ArNodeName);
+				ensureMsgf(false,
+					TEXT("Error: Missing registered node type (%s) will be removed from graph on load. Graph will fail to evaluate due to missing node (%s).")
+					, *ArType.ToString(), *ArName.ToString());
+			}
+			DATAFLOW_OPTIONAL_BLOCK_READ_END();
+		}
 
-			TArray< FLink > LocalConnections;
-			Ar << LocalConnections;
-			for (const FLink& Con : LocalConnections)
+		TArray< FLink > LocalConnections;
+		Ar << LocalConnections;
+		for (const FLink& Con : LocalConnections)
+		{
+			if (NodeGuidMap.Contains(Con.InputNode) && NodeGuidMap.Contains(Con.OutputNode))
 			{
-				if (NodeGuidMap.Contains(Con.InputNode) && NodeGuidMap.Contains(Con.OutputNode))
+				if (ConnectionGuidMap.Contains(Con.Input) && ConnectionGuidMap.Contains(Con.Output))
 				{
-					if (ConnectionGuidMap.Contains(Con.Input) && ConnectionGuidMap.Contains(Con.Output))
+					if (ConnectionGuidMap[Con.Input]->GetType() == ConnectionGuidMap[Con.Output]->GetType())
 					{
-						if (ConnectionGuidMap[Con.Input]->GetType() == ConnectionGuidMap[Con.Output]->GetType())
-						{
-							Connect(static_cast<FDataflowOutput*>(ConnectionGuidMap[Con.Output]), static_cast<FDataflowInput*>(ConnectionGuidMap[Con.Input]));
-						}
+						InGraph->Connect(static_cast<FDataflowOutput*>(ConnectionGuidMap[Con.Output]), static_cast<FDataflowInput*>(ConnectionGuidMap[Con.Input]));
 					}
 				}
 			}
-
 		}
 	}
 }
+
diff --git a/Engine/Source/Runtime/Experimental/Dataflow/Core/Private/Dataflow/DataflowNode.cpp b/Engine/Source/Runtime/Experimental/Dataflow/Core/Private/Dataflow/DataflowNode.cpp
--- a/Engine/Source/Runtime/Experimental/Dataflow/Core/Private/Dataflow/DataflowNode.cpp
+++ b/Engine/Source/Runtime/Experimental/Dataflow/Core/Private/Dataflow/DataflowNode.cpp
@@ -85,6 +85,19 @@ FDataflowInput* FDataflowNode::FindInput(void* Reference)
 	return nullptr;
 }
 
+const FDataflowInput* FDataflowNode::FindInput(const FGuid& InGuid) const
+{
+	for (TPair<uint32, FDataflowInput*> Elem : Inputs)
+	{
+		FDataflowInput* Con = Elem.Value;
+		if (Con->GetGuid() == InGuid)
+		{
+			return Con;
+		}
+	}
+	return nullptr;
+}
+
 TArray< FDataflowInput* > FDataflowNode::GetInputs() const
 {
 	TArray< FDataflowInput* > Result;
@@ -209,6 +222,19 @@ FDataflowOutput* FDataflowNode::FindOutput(void* Reference)
 	return nullptr;
 }
 
+const FDataflowOutput* FDataflowNode::FindOutput(const FGuid& InGuid) const
+{
+	for (TPair<uint32, FDataflowOutput*> Elem : Outputs)
+	{
+		FDataflowOutput* Con = Elem.Value;
+		if (Con->GetGuid() == InGuid)
+		{
+			return Con;
+		}
+	}
+	return nullptr;
+}
+
 int32 FDataflowNode::NumOutputs() const
 {
 	return Outputs.Num();
diff --git a/Engine/Source/Runtime/Experimental/Dataflow/Core/Public/Dataflow/DataflowGraph.h b/Engine/Source/Runtime/Experimental/Dataflow/Core/Public/Dataflow/DataflowGraph.h
--- a/Engine/Source/Runtime/Experimental/Dataflow/Core/Public/Dataflow/DataflowGraph.h
+++ b/Engine/Source/Runtime/Experimental/Dataflow/Core/Public/Dataflow/DataflowGraph.h
@@ -158,6 +158,8 @@ namespace Dataflow
 		DATAFLOWCORE_API virtual void Serialize(FArchive& Ar, UObject* OwningObject);
 		const TSet<FName>& GetDisabledNodes() const { return DisabledNodes; }
 
+		DATAFLOWCORE_API static void SerializeForSaving(FArchive& Ar, FGraph* InGraph, TArray<TSharedPtr<FDataflowNode>>& InNodes, TArray<FLink>& InConnections);
+		DATAFLOWCORE_API static void SerializeForLoading(FArchive& Ar, FGraph* InGraph, UObject* OwningObject);
 	};
 }
 
diff --git a/Engine/Source/Runtime/Experimental/Dataflow/Core/Public/Dataflow/DataflowNode.h b/Engine/Source/Runtime/Experimental/Dataflow/Core/Public/Dataflow/DataflowNode.h
--- a/Engine/Source/Runtime/Experimental/Dataflow/Core/Public/Dataflow/DataflowNode.h
+++ b/Engine/Source/Runtime/Experimental/Dataflow/Core/Public/Dataflow/DataflowNode.h
@@ -115,7 +115,7 @@ struct FDataflowNode
 	DATAFLOWCORE_API FDataflowInput* FindInput(FName Name);
 	DATAFLOWCORE_API FDataflowInput* FindInput(void* Reference);
 	DATAFLOWCORE_API const FDataflowInput* FindInput(const void* Reference) const;
-
+	DATAFLOWCORE_API const FDataflowInput* FindInput(const FGuid& InGuid) const;
 
 	DATAFLOWCORE_API virtual void AddOutput(FDataflowOutput* InPtr);
 	DATAFLOWCORE_API int NumOutputs() const;
@@ -128,6 +128,7 @@ struct FDataflowNode
 	DATAFLOWCORE_API const FDataflowOutput* FindOutput(FName Name) const;
 	DATAFLOWCORE_API const FDataflowOutput* FindOutput(uint32 GuidHash) const;
 	DATAFLOWCORE_API const FDataflowOutput* FindOutput(const void* Reference) const;
+	DATAFLOWCORE_API const FDataflowOutput* FindOutput(const FGuid& InGuid) const;
 
 	/** Return a property's byte offset from the dataflow base node address using the full property name (must includes its parent struct property names). */
 	uint32 GetPropertyOffset(const FName& PropertyFullName) const;
