Constraints: update dependencies when resolving sequencer or control rig bindings
    - FConstraintDependencyScope provides a way to build constraint dependencies when the constraint is not valid when added to the subsystem but after (when resolving sequencer or control rig bindings)

#rnx

diff --git a/Engine/Plugins/Animation/ControlRig/Source/ControlRig/Private/Constraints/ControlRigTransformableHandle.cpp b/Engine/Plugins/Animation/ControlRig/Source/ControlRig/Private/Constraints/ControlRigTransformableHandle.cpp
--- a/Engine/Plugins/Animation/ControlRig/Source/ControlRig/Private/Constraints/ControlRigTransformableHandle.cpp
+++ b/Engine/Plugins/Animation/ControlRig/Source/ControlRig/Private/Constraints/ControlRigTransformableHandle.cpp
@@ -487,8 +487,16 @@ void UTransformableControlHandle::OnObjectBoundToControlRig(UObject* InObject)
 	const UObject* CurrentObject = Binding ? Binding->GetBoundObject() : nullptr;
 	if (CurrentObject == InObject)
 	{
-		const UWorld* ThisWorld = GetWorld();
-		if (ThisWorld && InObject->GetWorld() == ThisWorld)
+		const UWorld* World = GetWorld();
+		if (!World)
+		{
+			if (const USceneComponent* BoundComponent = GetBoundComponent())
+			{
+				World = BoundComponent->GetWorld();
+			}
+		}
+		
+		if (World && InObject->GetWorld() == World)
 		{
 			Notify(EHandleEvent::ComponentUpdated);
 		}
diff --git a/Engine/Source/Runtime/Experimental/Animation/Constraints/Private/TransformConstraint.cpp b/Engine/Source/Runtime/Experimental/Animation/Constraints/Private/TransformConstraint.cpp
--- a/Engine/Source/Runtime/Experimental/Animation/Constraints/Private/TransformConstraint.cpp
+++ b/Engine/Source/Runtime/Experimental/Animation/Constraints/Private/TransformConstraint.cpp
@@ -561,6 +561,9 @@ bool UTickableTransformConstraint::HasBoundObjects() const
 
 void UTickableTransformConstraint::ResolveBoundObjects(FMovieSceneSequenceID LocalSequenceID, IMovieScenePlayer& Player, UObject* SubObject)
 {
+	// update dependencies if the constraint becomes valid once resolved 
+	FConstraintDependencyScope Scope(this);
+	
 	if (ChildTRSHandle && ChildTRSHandle->HasBoundObjects())
 	{
 		ChildTRSHandle->ResolveBoundObjects(LocalSequenceID, Player, SubObject);
@@ -659,6 +662,8 @@ void UTickableTransformConstraint::OnHandleModified(UTransformableHandle* InHand
 		if (InNotification == EHandleEvent::ComponentUpdated)
 		{
 			SetupDependencies(World);
+			// update dependencies now the component has been updated. 
+			FTransformConstraintUtils::BuildDependencies(World, this);
 			return;
 		}
 	}
@@ -724,6 +729,14 @@ void UTickableTransformConstraint::OnHandleModified(UTransformableHandle* InHand
 		}
 		return;
 	}
+
+	if (InHandle == ChildTRSHandle && InNotification == EHandleEvent::GlobalTransformUpdated)
+	{
+		if (FConstraintsEvaluationGraph::UseEvaluationGraph())
+		{
+			return MarkForEvaluation();
+		}
+	}
 }
 
 bool UTickableTransformConstraint::IsValid(const bool bDeepCheck) const
@@ -1893,6 +1906,11 @@ bool FTransformConstraintUtils::BuildDependencies(UWorld* InWorld, UTickableTran
 		return false;
 	}
 
+	if (bDebugDependencies)
+	{
+		UE_LOG(LogTemp, Warning, TEXT("Building dependencies for '%s' ..."), *GetConstraintLabel(InConstraint));
+	}
+
 	const UTransformableHandle* ParentHandle = InConstraint->ParentTRSHandle.Get();
 	const UTransformableHandle* ChildHandle = InConstraint->ChildTRSHandle.Get();
 	
@@ -2318,3 +2336,27 @@ void FTransformConstraintUtils::GetChildrenConstraints(
 	OutConstraints.Append(FilteredConstraints);
 }
 
+FConstraintDependencyScope::FConstraintDependencyScope(UTickableTransformConstraint* InConstraint, UWorld* InWorld)
+	: WeakConstraint(InConstraint)
+	, WeakWorld(InWorld)
+	, bPreviousValidity(InConstraint ? InConstraint->IsValid() : false)
+{}
+
+FConstraintDependencyScope::~FConstraintDependencyScope()
+{
+	if (!bPreviousValidity)
+	{
+		if (UTickableTransformConstraint* Constraint = WeakConstraint.IsValid() ? WeakConstraint.Get() : nullptr)
+		{
+			if (Constraint->IsValid())
+			{
+				const UObject* Target = ConstraintLocals::GetHandleTarget(Constraint->ChildTRSHandle);
+				UWorld* World = WeakWorld.IsValid() ? WeakWorld.Get() : Target ? Target->GetWorld() : nullptr;
+				if (::IsValid(World))
+				{
+					FTransformConstraintUtils::BuildDependencies(World, Constraint);
+				}
+			}
+		}
+	}
+}
diff --git a/Engine/Source/Runtime/Experimental/Animation/Constraints/Public/TransformConstraint.h b/Engine/Source/Runtime/Experimental/Animation/Constraints/Public/TransformConstraint.h
--- a/Engine/Source/Runtime/Experimental/Animation/Constraints/Public/TransformConstraint.h
+++ b/Engine/Source/Runtime/Experimental/Animation/Constraints/Public/TransformConstraint.h
@@ -44,7 +44,7 @@ public:
 	CONSTRAINTS_API virtual bool HasBoundObjects() const override;
 	
 	/** Resolve the bound objects so that any object it references are resovled and correctly set up*/
-	CONSTRAINTS_API virtual void ResolveBoundObjects(FMovieSceneSequenceID LocalSequenceID, IMovieScenePlayer& Player,UObject* SubObject) override;
+	CONSTRAINTS_API virtual void ResolveBoundObjects(FMovieSceneSequenceID LocalSequenceID, IMovieScenePlayer& Player, UObject* SubObject = nullptr) override;
 
 	/** Whether or not it's valid for example it may not be fully loaded*/
 	virtual bool IsValid(const bool bDeepCheck = true) const override;
@@ -506,3 +506,19 @@ struct FTransformConstraintUtils
 	/** Ensure default dependencies between constraints. */
 	static CONSTRAINTS_API bool BuildDependencies(UWorld* InWorld, UTickableTransformConstraint* Constraint);
 };
+
+/**
+ * FConstraintDependencyScope provides a way to build constraint dependencies when the constraint is not valid when added to the subsystem
+ * but after (when resolving sequencer or control rig bindings).
+ * The dependencies will be built on destruction if the constraint's validity changed within the lifetime of that object.
+ */
+
+struct FConstraintDependencyScope
+{
+	FConstraintDependencyScope(UTickableTransformConstraint* InConstraint, UWorld* InWorld = nullptr);
+	~FConstraintDependencyScope();
+private:
+	TWeakObjectPtr<UTickableTransformConstraint> WeakConstraint = nullptr;
+	TWeakObjectPtr<UWorld> WeakWorld = nullptr;
+	bool bPreviousValidity = false;
+};
\ No newline at end of file
diff --git a/Engine/Source/Runtime/MovieSceneTracks/Private/Systems/MovieSceneConstraintSystem.cpp b/Engine/Source/Runtime/MovieSceneTracks/Private/Systems/MovieSceneConstraintSystem.cpp
--- a/Engine/Source/Runtime/MovieSceneTracks/Private/Systems/MovieSceneConstraintSystem.cpp
+++ b/Engine/Source/Runtime/MovieSceneTracks/Private/Systems/MovieSceneConstraintSystem.cpp
@@ -176,10 +176,13 @@ struct FEvaluateConstraintChannels
 		bool Result = false;
 		ConstraintAndActiveChannel->ActiveChannel.Evaluate(FrameTime, Result);
 		Constraint->SetActive(Result);
-		Constraint->ResolveBoundObjects(TargetInstance.GetSequenceID(), *TargetInstance.GetPlayer());
+		
 		if (UTickableTransformConstraint* TransformConstraint = Cast< UTickableTransformConstraint>(Constraint))
 		{
 			TransformConstraint->InitConstraint(World);
+
+			// this has to be done once the constraint initialized
+			TransformConstraint->ResolveBoundObjects(TargetInstance.GetSequenceID(), *TargetInstance.GetPlayer());
 			
 			if (UTransformableComponentHandle* ComponentHandle = Cast<UTransformableComponentHandle>(TransformConstraint->ChildTRSHandle))
 			{
@@ -190,9 +193,12 @@ struct FEvaluateConstraintChannels
 				UpdateHandle.TransformHandle = ComponentHandle;
 				System->DynamicOffsets.Add(UpdateHandle);
 				TransformConstraint->EnsurePrimaryDependency(BoundObject->GetWorld());
-
 			}
 		}
+		else
+		{
+			Constraint->ResolveBoundObjects(TargetInstance.GetSequenceID(), *TargetInstance.GetPlayer());	
+		}
 	}
 
 	UWorld* World;
