Chaos - ClusterUnions update the mass and inertia on the game thread (previously it always reported 1.0)

This has a major downside - we have to send the mass back to the GT every time the cluster changes.

diff --git a/Engine/Source/Runtime/Engine/Private/PrimitiveComponentPhysics.cpp b/Engine/Source/Runtime/Engine/Private/PrimitiveComponentPhysics.cpp
--- a/Engine/Source/Runtime/Engine/Private/PrimitiveComponentPhysics.cpp
+++ b/Engine/Source/Runtime/Engine/Private/PrimitiveComponentPhysics.cpp
@@ -556,6 +556,23 @@ float UPrimitiveComponent::GetMass() const
 		WarnInvalidPhysicsOperations(LOCTEXT("GetMass", "GetMass"), BI, NAME_None);
 		return BI->GetBodyMass();
 	}
+	else
+	{
+		TArray<Chaos::FPhysicsObject*> PhysicsObjects = GetAllPhysicsObjects();
+		if (!PhysicsObjects.IsEmpty())
+		{
+			FLockedWritePhysicsObjectExternalInterface Interface = FPhysicsObjectExternalInterface::LockWrite(PhysicsObjects);
+
+			// Filter out inactive particles (E.g., inactive children in a geometry collection)
+			PhysicsObjects = PhysicsObjects.FilterByPredicate(
+				[&Interface](Chaos::FPhysicsObject* Object) 
+				{
+					return !Interface->AreAllDisabled({ &Object, 1 });
+				});
+
+			return Interface->GetMass(PhysicsObjects);
+		}
+	}
 
 	return 0.0f;
 }
diff --git a/Engine/Source/Runtime/Experimental/Chaos/Private/PhysicsProxy/ClusterUnionPhysicsProxy.cpp b/Engine/Source/Runtime/Experimental/Chaos/Private/PhysicsProxy/ClusterUnionPhysicsProxy.cpp
--- a/Engine/Source/Runtime/Experimental/Chaos/Private/PhysicsProxy/ClusterUnionPhysicsProxy.cpp
+++ b/Engine/Source/Runtime/Experimental/Chaos/Private/PhysicsProxy/ClusterUnionPhysicsProxy.cpp
@@ -46,6 +46,8 @@ namespace Chaos
 			BufferData.R = Particle->GetR();
 			BufferData.V = Particle->GetV();
 			BufferData.W = Particle->GetW();
+			BufferData.Mass = FRealSingle(Particle->M());
+			BufferData.Inertia = FVec3f(Particle->I());
 			BufferData.ObjectState = Particle->ObjectState();
 			BufferData.Geometry = nullptr;
 
@@ -461,6 +463,17 @@ namespace Chaos
 			Particle_External->SetGeometry(CurrentPullData.Geometry);
 			SyncedData_External.bDidSyncGeometry = true;
 		}
+
+		FReal M = FReal(CurrentPullData.Mass);
+		FReal InvM = (M > 0.0) ? (1.0 / M) : 0.0;
+		Particle_External->SetM(M, /*bInvalidate=*/false);
+		Particle_External->SetInvM(InvM, /*bInvalidate=*/false);
+
+		FVec3 I = FVec3(CurrentPullData.Inertia);
+		FVec3 InvI = (!I.IsZero()) ? FVec3(1) / I : FVec3(0);
+		Particle_External->SetI(I, /*bInvalidate=*/false);
+		Particle_External->SetInvI(InvI, /*bInvalidate=*/false);
+
 		Particle_External->SetObjectState(CurrentPullData.ObjectState, true, /*bInvalidate=*/false);
 
 		const FShapesArray& ShapeArray = Particle_External->ShapesArray();
diff --git a/Engine/Source/Runtime/Experimental/Chaos/Public/Chaos/ParticleHandle.h b/Engine/Source/Runtime/Experimental/Chaos/Public/Chaos/ParticleHandle.h
--- a/Engine/Source/Runtime/Experimental/Chaos/Public/Chaos/ParticleHandle.h
+++ b/Engine/Source/Runtime/Experimental/Chaos/Public/Chaos/ParticleHandle.h
@@ -3469,27 +3469,27 @@ public:
 	}
 
 	const TVec3<FRealSingle>& I() const { return MMassProps.Read().I(); }
-	void SetI(const TVec3<FRealSingle>& InI)
+	void SetI(const TVec3<FRealSingle>& InI, bool bInvalidate = true)
 	{
-		MMassProps.Modify(true,MDirtyFlags,Proxy,[&InI](auto& Data){ Data.SetI(InI);});
+		MMassProps.Modify(bInvalidate,MDirtyFlags,Proxy,[&InI](auto& Data){ Data.SetI(InI);});
 	}
 
 	const TVec3<FRealSingle>& InvI() const { return MMassProps.Read().InvI(); }
-	void SetInvI(const TVec3<FRealSingle>& InInvI)
+	void SetInvI(const TVec3<FRealSingle>& InInvI, bool bInvalidate = true)
 	{
-		MMassProps.Modify(true,MDirtyFlags,Proxy,[&InInvI](auto& Data){ Data.SetInvI(InInvI);});
+		MMassProps.Modify(bInvalidate,MDirtyFlags,Proxy,[&InInvI](auto& Data){ Data.SetInvI(InInvI);});
 	}
 
 	T M() const { return MMassProps.Read().M(); }
-	void SetM(const T& InM)
+	void SetM(const T& InM, bool bInvalidate = true)
 	{
-		MMassProps.Modify(true,MDirtyFlags,Proxy,[InM](auto& Data){ Data.SetM(InM);});
+		MMassProps.Modify(bInvalidate,MDirtyFlags,Proxy,[InM](auto& Data){ Data.SetM(InM);});
 	}
 
 	T InvM() const { return MMassProps.Read().InvM(); }
-	void SetInvM(const T& InInvM)
+	void SetInvM(const T& InInvM, bool bInvalidate = true)
 	{
-		MMassProps.Modify(true,MDirtyFlags,Proxy,[InInvM](auto& Data){ Data.SetInvM(InInvM);});
+		MMassProps.Modify(bInvalidate,MDirtyFlags,Proxy,[InInvM](auto& Data){ Data.SetInvM(InInvM);});
 	}
 	
 	const TVector<T,d>& CenterOfMass() const { return MMassProps.Read().CenterOfMass(); }
diff --git a/Engine/Source/Runtime/Experimental/Chaos/Public/Chaos/PullPhysicsDataImp.h b/Engine/Source/Runtime/Experimental/Chaos/Public/Chaos/PullPhysicsDataImp.h
--- a/Engine/Source/Runtime/Experimental/Chaos/Public/Chaos/PullPhysicsDataImp.h
+++ b/Engine/Source/Runtime/Experimental/Chaos/Public/Chaos/PullPhysicsDataImp.h
@@ -109,6 +109,8 @@ struct FDirtyClusterUnionData : public TBasePullData<FClusterUnionPhysicsProxy,
 	FQuat R;
 	FVec3 V;
 	FVec3 W;
+	FRealSingle Mass;
+	FVec3f Inertia;
 	EObjectStateType ObjectState = EObjectStateType::Dynamic;
 	bool bIsAnchored = false;
 	TArray<FDirtyClusterUnionParticleData> ChildParticles;
