Chaos - spheres and capsules default to PBD dynamic friction

Spheres and capsules were using velocity based dynamic friction for reasons that are no longer relevant. Switching to PBD gives better results because we run more PBD iterations.

diff --git a/Engine/Source/Runtime/Experimental/Chaos/Private/Chaos/ChaosDebugDraw.cpp b/Engine/Source/Runtime/Experimental/Chaos/Private/Chaos/ChaosDebugDraw.cpp
--- a/Engine/Source/Runtime/Experimental/Chaos/Private/Chaos/ChaosDebugDraw.cpp
+++ b/Engine/Source/Runtime/Experimental/Chaos/Private/Chaos/ChaosDebugDraw.cpp
@@ -115,6 +115,9 @@ namespace Chaos
 		float ChaosDebugDrawCollisionDuration = 0.0f;
 		FAutoConsoleVariableRef CVarChaosDebugDrawCollisionDuration(TEXT("p.Chaos.DebugDraw.CollisionDuration"), ChaosDebugDrawCollisionDuration, TEXT("How long Collision debug draw should remain on screen in seconds. 0 for 1 frame."));
 
+		bool bChaosDebugDrawCollisionAnchors = 0.0f;
+		FAutoConsoleVariableRef CVarChaosDebugDrawCollisionAnchors(TEXT("p.Chaos.DebugDraw.ShowCollisionAnchors"), bChaosDebugDrawCollisionAnchors, TEXT(""));
+
 		int32 ChaosConnectionGraphDrawLevelOffset = 0;
 		FAutoConsoleVariableRef CVarChaosConnectionGraphDrawLevelOffset(TEXT("p.Chaos.DebugDraw.ConnectionGraphLevelOffset"), ChaosConnectionGraphDrawLevelOffset, TEXT("If 0, draws the connection graph between children particles of active particles. If 1, draws the connection grpah between grand-children particles of active particles, etc."));
 
@@ -1365,6 +1368,24 @@ namespace Chaos
 					FDebugDrawQueue::GetInstance().DrawDebugLine(ImpulsePos, ImpulsePos + Settings.DrawScale * Settings.ImpulseScale * SpaceTransform.TransformVectorNoScale(Contact.GetAccumulatedImpulse()), Color, false, Duration, uint8(Settings.DrawPriority), Settings.LineThickness);
 				}
 
+				// Static Friction Anchors (bets viewed with DebugDrawPostIntegrateCollisions)
+				if (bChaosDebugDrawCollisionAnchors)
+				{
+					for (int32 PointIndex = 0; PointIndex < Contact.NumSavedManifoldPoints(); ++PointIndex)
+					{
+						const FSavedManifoldPoint& SavedManifoldPoint = Contact.GetSavedManifoldPoint(PointIndex);
+						const int32 ContactPlaneOwner = 1;
+						const int32 ContactPointOwner = 1 - ContactPlaneOwner;
+						const FRigidTransform3& PlaneTransform = (ContactPlaneOwner == 0) ? Contact.GetShapeRelativeTransform0() * WorldActorTransform0 : Contact.GetShapeRelativeTransform1() * WorldActorTransform1;
+						const FRigidTransform3& PointTransform = (ContactPlaneOwner == 0) ? Contact.GetShapeRelativeTransform1() * WorldActorTransform1 : Contact.GetShapeRelativeTransform0() * WorldActorTransform0;
+						const FVec3 PointLocation = PointTransform.TransformPosition(FVec3(SavedManifoldPoint.ShapeContactPoints[ContactPointOwner]));
+						const FVec3 PlaneLocation = PlaneTransform.TransformPosition(FVec3(SavedManifoldPoint.ShapeContactPoints[ContactPlaneOwner]));
+				
+						FColor Color = FColor::Black;
+						FDebugDrawQueue::GetInstance().DrawDebugLine(PointLocation, PlaneLocation, Color, false, Duration, uint8(Settings.DrawPriority), Settings.LineThickness);
+					}
+				}
+
 				// Show the sphere approximation if enabled
 				//if ((CVars::ChaosOneWayInteractionPairCollisionMode == (int32)EOneWayInteractionPairCollisionMode::SphereCollision) && Particle0->OneWayInteraction() && Particle1->OneWayInteraction())
 				//{
diff --git a/Engine/Source/Runtime/Experimental/Chaos/Private/Chaos/Collision/PBDCollisionConstraint.cpp b/Engine/Source/Runtime/Experimental/Chaos/Private/Chaos/Collision/PBDCollisionConstraint.cpp
--- a/Engine/Source/Runtime/Experimental/Chaos/Private/Chaos/Collision/PBDCollisionConstraint.cpp
+++ b/Engine/Source/Runtime/Experimental/Chaos/Private/Chaos/Collision/PBDCollisionConstraint.cpp
@@ -1115,7 +1115,7 @@ namespace Chaos
 
 		// When only one shape is quadratic, we only look at the quadratic contact point so we don't identify
 		// a sphere spinning on the spot as a stationary contact
-		// @todo(chaos): handle quadratic shapes better with static friction
+		// @todo(chaos): we don't really need the quadratic special case any more
 		if (IsQuadratic0() && !IsQuadratic1())
 		{
 			return DP0.SizeSquared();
diff --git a/Engine/Source/Runtime/Experimental/Chaos/Private/Chaos/Collision/PBDCollisionContainerSolver.cpp b/Engine/Source/Runtime/Experimental/Chaos/Private/Chaos/Collision/PBDCollisionContainerSolver.cpp
--- a/Engine/Source/Runtime/Experimental/Chaos/Private/Chaos/Collision/PBDCollisionContainerSolver.cpp
+++ b/Engine/Source/Runtime/Experimental/Chaos/Private/Chaos/Collision/PBDCollisionContainerSolver.cpp
@@ -18,8 +18,6 @@
 
 #include "ChaosLog.h"
 
-//PRAGMA_DISABLE_OPTIMIZATION
-
 namespace Chaos
 {
 	namespace CVars
@@ -137,7 +135,6 @@ namespace Chaos
 		}
 		else
 		{
-			// @todo(chaos): consider adding a multiplier to the initial contact friction
 			WorldFrictionDelta = ContactVel * Dt;
 		}
 
@@ -286,20 +283,22 @@ namespace Chaos
 		if (SolverSettings.NumPositionFrictionIterations > 0)
 		{
 			PositionStaticFriction = StaticFriction;
-			if (!Constraint->HasQuadraticShape())
+
+			// We have an option to apply dynamic friction in the velocity solver phase for spheres and capsules. 
+			// @todo(chaos): UE5.4: Remove quadratic special case when PBD static friction is enabled
+			const bool bUsePBDDynamicFriction = !Constraint->HasQuadraticShape() || (SolverSettings.NumVelocityFrictionIterations == 0);
+			if (bUsePBDDynamicFriction)
 			{
 				PositionDynamicFriction = DynamicFriction;
 			}
 			else
 			{
-				// Quadratic shapes don't use PBD dynamic friction - it has issues at slow speeds where the WxR is
-				// less than the position tolerance for friction point matching
-				// @todo(chaos): fix PBD dynamic friction on quadratic shapes
 				VelocityDynamicFriction = DynamicFriction;
 			}
 		}
-		else
+		else if (SolverSettings.NumVelocityFrictionIterations > 0)
 		{
+			// We have an option to run without static friction and velocity-based dynamic friction (for RBAN)
 			VelocityDynamicFriction = DynamicFriction;
 		}
 
diff --git a/Engine/Source/Runtime/Experimental/Chaos/Public/Chaos/Collision/PBDCollisionConstraint.h b/Engine/Source/Runtime/Experimental/Chaos/Public/Chaos/Collision/PBDCollisionConstraint.h
--- a/Engine/Source/Runtime/Experimental/Chaos/Public/Chaos/Collision/PBDCollisionConstraint.h
+++ b/Engine/Source/Runtime/Experimental/Chaos/Public/Chaos/Collision/PBDCollisionConstraint.h
@@ -714,6 +714,16 @@ namespace Chaos
 			SavedManifoldPoints.Reset();
 		}
 
+		int32 NumSavedManifoldPoints() const
+		{
+			return SavedManifoldPoints.Num();
+		}
+
+		const FSavedManifoldPoint& GetSavedManifoldPoint(const int32 PointIndex) const
+		{
+			return SavedManifoldPoints[PointIndex];
+		}
+
 		/**
 		 * @brief Time of impact from CCD sweep test if CCD is activate.Otherwise undefined.
 		*/
diff --git a/Engine/Source/Runtime/Experimental/Chaos/Public/Chaos/Collision/PBDCollisionSolverSettings.h b/Engine/Source/Runtime/Experimental/Chaos/Public/Chaos/Collision/PBDCollisionSolverSettings.h
--- a/Engine/Source/Runtime/Experimental/Chaos/Public/Chaos/Collision/PBDCollisionSolverSettings.h
+++ b/Engine/Source/Runtime/Experimental/Chaos/Public/Chaos/Collision/PBDCollisionSolverSettings.h
@@ -17,7 +17,7 @@ namespace Chaos
 			: MaxPushOutVelocity(0)
 			, DepenetrationVelocity(0)
 			, NumPositionFrictionIterations(4)
-			, NumVelocityFrictionIterations(1)
+			, NumVelocityFrictionIterations(0)
 			, NumPositionShockPropagationIterations(3)
 			, NumVelocityShockPropagationIterations(2)
 		{
@@ -34,9 +34,11 @@ namespace Chaos
 		// How many of the position iterations should run static/dynamic friction
 		int32 NumPositionFrictionIterations;
 
-		// How many of the velocity iterations should run dynamic friction
-		// @todo(chaos): if NumVelocityFrictionIterations > 1, then dynamic friction in the velocity phase will be iteration 
-		// count dependent (velocity-solve friction is currentlyused by quadratic shapes and RBAN)
+		// How many of the velocity iterations should run dynamic friction 
+		// This only applies to quadratic shapes if we have static friction enabled (NumPositionFrictionIterations > 0)
+		// RBAN is the only system that uses velocity-based friction by default (it also disables static friction)
+		// NOTE: velocity-based dynamic friction behaviour is iteration count dependent so generally want (NumVelocityFrictionIterations <= 1)
+		// @todo(chaos): fix iteration count dependence of velocity-based friction, or remove it as a feature
 		int32 NumVelocityFrictionIterations;
 
 		// How many position iterations should have shock propagation enabled
