Add Anim Geometry Collection component
- This component allow binding a geometry collection with a skeletal mesh for animation fo the root proxy meshes

diff --git a/Engine/Source/Runtime/Experimental/GeometryCollectionEngine/Private/GeometryCollection/GeometryCollectionComponent.cpp b/Engine/Source/Runtime/Experimental/GeometryCollectionEngine/Private/GeometryCollection/GeometryCollectionComponent.cpp
--- a/Engine/Source/Runtime/Experimental/GeometryCollectionEngine/Private/GeometryCollection/GeometryCollectionComponent.cpp
+++ b/Engine/Source/Runtime/Experimental/GeometryCollectionEngine/Private/GeometryCollection/GeometryCollectionComponent.cpp
@@ -3620,11 +3620,18 @@ void UGeometryCollectionComponent::UpdateRootProxyComponentsIfNeeded()
 				const FTransform3f RootCurrentTransform = ComponentSpaceTransforms.RequestRootTransform();
 				const FTransform RootTransformFromRestPose = FTransform(RootCurrentTransform.GetRelativeTransform(RootRestTransform));
 
-				for (TObjectPtr<UStaticMeshComponent> StaticMeshComponent : RootProxyStaticMeshComponents)
+				for (int32 MeshIndex = 0; MeshIndex < RootProxyStaticMeshComponents.Num(); MeshIndex++)
 				{
-					if (StaticMeshComponent)
+					if (TObjectPtr<UStaticMeshComponent> StaticMeshComponent = RootProxyStaticMeshComponents[MeshIndex])
 					{
-						StaticMeshComponent->SetRelativeTransform(RootTransformFromRestPose);
+						if (RootProxyLocalTransforms.IsValidIndex(MeshIndex))
+						{
+							StaticMeshComponent->SetRelativeTransform(FTransform(RootProxyLocalTransforms[MeshIndex]) * RootTransformFromRestPose);
+						}
+						else
+						{
+							StaticMeshComponent->SetRelativeTransform(RootTransformFromRestPose);
+						}
 					}
 				}
 			}
@@ -5021,6 +5028,8 @@ void UGeometryCollectionComponent::SetRestCollection(const UGeometryCollection*
 
 		CreateRootProxyComponentsIfNeeded();
 
+		ClearRootProxyLocalTransforms();
+
 		if (bApplyAssetDefaults)
 		{
 			ApplyAssetDefaults();
@@ -6278,6 +6287,30 @@ void UGeometryCollectionComponent::EnableRootProxyForCustomRenderer(bool bEnable
 	RefreshCustomRenderer();
 }
 
+void UGeometryCollectionComponent::SetRootProxyLocalTransform(int32 Index, const FTransform3f& RootProxyTransform)
+{
+	if (RootProxyLocalTransforms.IsEmpty() && RestCollection)
+	{
+		const int32 NumProxyMeshes = RestCollection->RootProxyData.ProxyMeshes.Num();
+		RootProxyLocalTransforms.Init(FTransform3f::Identity, NumProxyMeshes);
+	}
+	if (RootProxyLocalTransforms.IsValidIndex(Index))
+	{
+		RootProxyLocalTransforms[Index] = RootProxyTransform;
+	}
+}
+
+void UGeometryCollectionComponent::ClearRootProxyLocalTransforms()
+{
+	RootProxyLocalTransforms.Empty();
+}
+
+void UGeometryCollectionComponent::RefreshRootProxies()
+{
+	RefreshCustomRenderer();
+	UpdateRootProxyComponentsIfNeeded();
+}
+
 bool UGeometryCollectionComponent::CanUseCustomRenderer() const 
 {
 	return bChaos_GC_UseCustomRenderer && CustomRenderer != nullptr && GetWorld()->IsGameWorld();
@@ -6320,7 +6353,14 @@ void UGeometryCollectionComponent::RefreshCustomRenderer()
 					if (bRenderRootProxy)
 					{
 						const FTransform CompSpaceRootTransform(ComponentSpaceTransforms.RequestRootTransform());
-						RendererInterface->UpdateRootTransform(*RestCollection, CompSpaceRootTransform);
+						if (RootProxyLocalTransforms.IsEmpty())
+						{
+							RendererInterface->UpdateRootTransform(*RestCollection, CompSpaceRootTransform);
+						}
+						else
+						{
+							RendererInterface->UpdateRootTransforms(*RestCollection, CompSpaceRootTransform, MakeArrayView(RootProxyLocalTransforms));
+						}
 					}
 					else
 					{
diff --git a/Engine/Source/Runtime/Experimental/GeometryCollectionEngine/Private/GeometryCollection/GeometryCollectionISMPoolRenderer.cpp b/Engine/Source/Runtime/Experimental/GeometryCollectionEngine/Private/GeometryCollection/GeometryCollectionISMPoolRenderer.cpp
--- a/Engine/Source/Runtime/Experimental/GeometryCollectionEngine/Private/GeometryCollection/GeometryCollectionISMPoolRenderer.cpp
+++ b/Engine/Source/Runtime/Experimental/GeometryCollectionEngine/Private/GeometryCollection/GeometryCollectionISMPoolRenderer.cpp
@@ -62,7 +62,12 @@ void UGeometryCollectionISMPoolRenderer::UpdateState(UGeometryCollection const&
 
 void UGeometryCollectionISMPoolRenderer::UpdateRootTransform(UGeometryCollection const& InGeometryCollection, FTransform const& InRootTransform)
 {
-	UpdateMergedMeshTransforms(InRootTransform * ComponentTransform);
+	UpdateMergedMeshTransforms(InRootTransform * ComponentTransform, {});
+}
+
+void UGeometryCollectionISMPoolRenderer::UpdateRootTransforms(UGeometryCollection const& InGeometryCollection, FTransform const& InRootTransform, TArrayView<const FTransform3f> InRootTransforms)
+{
+	UpdateMergedMeshTransforms(InRootTransform * ComponentTransform, InRootTransforms);
 }
 
 void UGeometryCollectionISMPoolRenderer::UpdateTransforms(UGeometryCollection const& InGeometryCollection, TArrayView<const FTransform3f> InTransforms)
@@ -159,7 +164,7 @@ void UGeometryCollectionISMPoolRenderer::InitInstancesFromGeometryCollection(UGe
 	}
 }
 
-void UGeometryCollectionISMPoolRenderer::UpdateMergedMeshTransforms(FTransform const& InBaseTransform)
+void UGeometryCollectionISMPoolRenderer::UpdateMergedMeshTransforms(FTransform const& InBaseTransform, TArrayView<const FTransform3f> LocalTransforms)
 {
 	if (MergedMeshGroup.GroupIndex == INDEX_NONE)
 	{
@@ -175,7 +180,15 @@ void UGeometryCollectionISMPoolRenderer::UpdateMergedMeshTransforms(FTransform c
 	TArrayView<const FTransform> InstanceTransforms(&InBaseTransform, 1);
 	for (int32 MeshIndex = 0; MeshIndex < MergedMeshGroup.MeshIds.Num(); MeshIndex++)
 	{
-		ISMPoolComponent->BatchUpdateInstancesTransforms(MergedMeshGroup.GroupIndex, MergedMeshGroup.MeshIds[MeshIndex], 0, InstanceTransforms, true/*bWorldSpace*/, false/*bMarkRenderStateDirty*/, false/*bTeleport*/);
+		if (LocalTransforms.IsValidIndex(MeshIndex))
+		{
+			const FTransform CombinedTransform{ FTransform(LocalTransforms[MeshIndex]) * InBaseTransform };
+			ISMPoolComponent->BatchUpdateInstancesTransforms(MergedMeshGroup.GroupIndex, MergedMeshGroup.MeshIds[MeshIndex], 0, MakeArrayView(&CombinedTransform, 1), true/*bWorldSpace*/, false/*bMarkRenderStateDirty*/, false/*bTeleport*/);
+		}
+		else
+		{
+			ISMPoolComponent->BatchUpdateInstancesTransforms(MergedMeshGroup.GroupIndex, MergedMeshGroup.MeshIds[MeshIndex], 0, InstanceTransforms, true/*bWorldSpace*/, false/*bMarkRenderStateDirty*/, false/*bTeleport*/);
+		}
 	}
 }
 
diff --git a/Engine/Source/Runtime/Experimental/GeometryCollectionEngine/Private/GeometryCollection/GeometryCollectionISMPoolRenderer.h b/Engine/Source/Runtime/Experimental/GeometryCollectionEngine/Private/GeometryCollection/GeometryCollectionISMPoolRenderer.h
--- a/Engine/Source/Runtime/Experimental/GeometryCollectionEngine/Private/GeometryCollection/GeometryCollectionISMPoolRenderer.h
+++ b/Engine/Source/Runtime/Experimental/GeometryCollectionEngine/Private/GeometryCollection/GeometryCollectionISMPoolRenderer.h
@@ -23,6 +23,7 @@ public:
 	virtual void OnUnregisterGeometryCollection() override;
 	virtual void UpdateState(UGeometryCollection const& InGeometryCollection, FTransform const& InComponentTransform, uint32 InStateFlags) override;
 	virtual void UpdateRootTransform(UGeometryCollection const& InGeometryCollection, FTransform const& InRootTransform) override;
+	virtual void UpdateRootTransforms(UGeometryCollection const& InGeometryCollection, FTransform const& InRootTransform, TArrayView<const FTransform3f> InRootLocalTransforms) override;
 	virtual void UpdateTransforms(UGeometryCollection const& InGeometryCollection, TArrayView<const FTransform3f> InTransforms) override;
 	//~ End IGeometryCollectionExternalRenderInterface Interface.
 
@@ -52,7 +53,7 @@ private:
 	UGeometryCollectionISMPoolComponent* GetOrCreateISMPoolComponent();
 	void InitMergedMeshFromGeometryCollection(UGeometryCollection const& InGeometryCollection);
 	void InitInstancesFromGeometryCollection(UGeometryCollection const& InGeometryCollection);
-	void UpdateMergedMeshTransforms(FTransform const& InBaseTransform);
+	void UpdateMergedMeshTransforms(FTransform const& InBaseTransform, TArrayView<const FTransform3f> LocalTransforms);
 	void UpdateInstanceTransforms(UGeometryCollection const& InGeometryCollection, FTransform const& InBaseTransform, TArrayView<const FTransform3f> InTransforms);
 	void ReleaseGroup(FISMPoolGroup& InOutGroup);
 };
diff --git a/Engine/Source/Runtime/Experimental/GeometryCollectionEngine/Public/GeometryCollection/GeometryCollectionComponent.h b/Engine/Source/Runtime/Experimental/GeometryCollectionEngine/Public/GeometryCollection/GeometryCollectionComponent.h
--- a/Engine/Source/Runtime/Experimental/GeometryCollectionEngine/Public/GeometryCollection/GeometryCollectionComponent.h
+++ b/Engine/Source/Runtime/Experimental/GeometryCollectionEngine/Public/GeometryCollection/GeometryCollectionComponent.h
@@ -1335,6 +1335,12 @@ public:
 	UFUNCTION(BlueprintCallable, Category = "Physics")
 	GEOMETRYCOLLECTIONENGINE_API void EnableRootProxyForCustomRenderer(bool bEnable);
 
+	/** Set a specific root proxy local transform */
+	GEOMETRYCOLLECTIONENGINE_API void SetRootProxyLocalTransform(int32 Index, const FTransform3f& RootProxyTransform);
+
+	/** clear all the root proxies local transforms - all proxies will now have the same default transform */
+	GEOMETRYCOLLECTIONENGINE_API void ClearRootProxyLocalTransforms();
+
 	/** Force all GC components to reregister their custom renderer objects. */
 	static GEOMETRYCOLLECTIONENGINE_API void ReregisterAllCustomRenderers();
 
@@ -1748,6 +1754,9 @@ public:
 	/** Updates the custom renderer to the reflect the current state of the Geometry Collection */
 	GEOMETRYCOLLECTIONENGINE_API void RefreshCustomRenderer();
 
+	/** Refresh root proxies whether they are drawn through the custom renderer or normal static mesh components */
+	GEOMETRYCOLLECTIONENGINE_API void RefreshRootProxies();
+
 private:
 
 	GEOMETRYCOLLECTIONENGINE_API void IncrementSleepTimer(float DeltaTime);
@@ -1838,6 +1847,8 @@ private:
 	TArray<TObjectPtr<UStaticMeshComponent>> RootProxyStaticMeshComponents;
 	bool bEnableRootProxyStaticMeshComponents = true;
 
+	TArray<FTransform3f> RootProxyLocalTransforms;
+
 private:
 
 	enum class ENetAwakeningMode
diff --git a/Engine/Source/Runtime/Experimental/GeometryCollectionEngine/Public/GeometryCollection/GeometryCollectionExternalRenderInterface.h b/Engine/Source/Runtime/Experimental/GeometryCollectionEngine/Public/GeometryCollection/GeometryCollectionExternalRenderInterface.h
--- a/Engine/Source/Runtime/Experimental/GeometryCollectionEngine/Public/GeometryCollection/GeometryCollectionExternalRenderInterface.h
+++ b/Engine/Source/Runtime/Experimental/GeometryCollectionEngine/Public/GeometryCollection/GeometryCollectionExternalRenderInterface.h
@@ -29,8 +29,30 @@ public:
 
 	virtual void OnRegisterGeometryCollection(UGeometryCollectionComponent const& InComponent) = 0;
 	virtual void OnUnregisterGeometryCollection() = 0;
+
+	/** 
+	* Set the state of the geometry collection 
+	* this is used by the renderer to managed resources with regards to the state ( see EStateFlags )
+	*/
 	virtual void UpdateState(UGeometryCollection const& InGeometryCollection, FTransform const& InComponentTransform, uint32 InStateFlags) = 0;
+
+	/**
+	* Update the root bone transform of the geometry collection 
+	* if the geometry collection is using of multiple proxy root meshes this transform applies to all of them 
+	*/
 	virtual void UpdateRootTransform(UGeometryCollection const& InGeometryCollection, FTransform const& InRootTransform) = 0;
+
+	/**
+	* Update the root proxy transforms of the geometry collection
+	* if the geometry collection is using of multiple proxy root meshes, InRootTransforms is expected to contain an entry for each of them 
+	* @param InRootTransform		Component space root transform
+	* @param InRootLocalTransforms	Root space local transforms
+	*/
+	virtual void UpdateRootTransforms(UGeometryCollection const& InGeometryCollection, FTransform const& InRootTransform, TArrayView<const FTransform3f> InRootLocalTransforms) = 0;
+
+	/**
+	* Update all the bones transforms 
+	*/
 	virtual void UpdateTransforms(UGeometryCollection const& InGeometryCollection, TArrayView<const FTransform3f> InTransforms) = 0;
 
 	UE_DEPRECATED(5.4, "Use flags version of UpdateState instead")
