Fix build issue related to CL 31864591
( virtual pure method was added on the engine side, but some builds do not sync the game code , causing abstract class allocation errors)

preflight none

diff --git a/Engine/Source/Runtime/Experimental/GeometryCollectionEngine/Public/GeometryCollection/GeometryCollectionExternalRenderInterface.h b/Engine/Source/Runtime/Experimental/GeometryCollectionEngine/Public/GeometryCollection/GeometryCollectionExternalRenderInterface.h
--- a/Engine/Source/Runtime/Experimental/GeometryCollectionEngine/Public/GeometryCollection/GeometryCollectionExternalRenderInterface.h
+++ b/Engine/Source/Runtime/Experimental/GeometryCollectionEngine/Public/GeometryCollection/GeometryCollectionExternalRenderInterface.h
@@ -48,7 +48,7 @@ public:
 	* @param InRootTransform		Component space root transform
 	* @param InRootLocalTransforms	Root space local transforms
 	*/
-	virtual void UpdateRootTransforms(UGeometryCollection const& InGeometryCollection, FTransform const& InRootTransform, TArrayView<const FTransform3f> InRootLocalTransforms) = 0;
+	virtual void UpdateRootTransforms(UGeometryCollection const& InGeometryCollection, FTransform const& InRootTransform, TArrayView<const FTransform3f> InRootLocalTransforms) { check(false); };
 
 	/**
 	* Update all the bones transforms 
