increase maximum value contained by FImplicitObjectUnion::NumLeafObjects from uint16 to int32

diff --git a/Engine/Source/Runtime/Core/Public/UObject/FortniteSeasonBranchObjectVersion.h b/Engine/Source/Runtime/Core/Public/UObject/FortniteSeasonBranchObjectVersion.h
--- a/Engine/Source/Runtime/Core/Public/UObject/FortniteSeasonBranchObjectVersion.h
+++ b/Engine/Source/Runtime/Core/Public/UObject/FortniteSeasonBranchObjectVersion.h
@@ -45,6 +45,9 @@ struct CORE_API FFortniteSeasonBranchObjectVersion
 		// Serialize actor transform for actor descriptors
 		WorldPartitionActorDescActorTransformSerialization,
 
+		// Changing Chaos::FImplicitObjectUnion to store an int32 vs a uint16 for NumLeafObjects.
+		ChaosImplicitObjectUnionLeafObjectsToInt32,
+
 		// -----<new versions can be added above this line>-------------------------------------------------
 		VersionPlusOne,
 		LatestVersion = VersionPlusOne - 1
diff --git a/Engine/Source/Runtime/Experimental/Chaos/Private/Chaos/ImplicitObjectUnion.cpp b/Engine/Source/Runtime/Experimental/Chaos/Private/Chaos/ImplicitObjectUnion.cpp
--- a/Engine/Source/Runtime/Experimental/Chaos/Private/Chaos/ImplicitObjectUnion.cpp
+++ b/Engine/Source/Runtime/Experimental/Chaos/Private/Chaos/ImplicitObjectUnion.cpp
@@ -160,7 +160,7 @@ void FImplicitObjectUnion::SetNumLeafObjects(int32 InNumLeafObjects)
 	ensure(InNumLeafObjects <= MaxNumLeafObjects);
 	check(InNumLeafObjects >= 0);
 
-	NumLeafObjects = uint16(FMath::Min(InNumLeafObjects, MaxNumLeafObjects));
+	NumLeafObjects = FMath::Min(InNumLeafObjects, MaxNumLeafObjects);
 }
 
 void FImplicitObjectUnion::CreateBVH()
@@ -420,6 +420,7 @@ void FImplicitObjectUnion::Serialize(FChaosArchive& Ar)
 {
 	Ar.UsingCustomVersion(FExternalPhysicsCustomObjectVersion::GUID);
 	Ar.UsingCustomVersion(FFortniteMainBranchObjectVersion::GUID);
+	Ar.UsingCustomVersion(FFortniteSeasonBranchObjectVersion::GUID);
 
 	FChaosArchiveScopedMemory ScopedMemory(Ar, GetTypeName(), false);
 	FImplicitObject::SerializeImp(Ar);
@@ -443,7 +444,18 @@ void FImplicitObjectUnion::Serialize(FChaosArchive& Ar)
 	else
 	{
 		Ar << Flags.Bits;
-		Ar << NumLeafObjects;
+
+		if (Ar.CustomVer(FFortniteSeasonBranchObjectVersion::GUID) < FFortniteSeasonBranchObjectVersion::ChaosImplicitObjectUnionLeafObjectsToInt32)
+		{
+			uint16 DummyNumLeafObjects = uint16(NumLeafObjects);
+			Ar << DummyNumLeafObjects;
+			NumLeafObjects = int32(DummyNumLeafObjects);
+		}
+		else
+		{
+			Ar << NumLeafObjects;
+		}
+
 		if (Flags.bHasBVH)
 		{
 			if (Ar.IsLoading())
diff --git a/Engine/Source/Runtime/Experimental/Chaos/Public/Chaos/ImplicitObjectUnion.h b/Engine/Source/Runtime/Experimental/Chaos/Public/Chaos/ImplicitObjectUnion.h
--- a/Engine/Source/Runtime/Experimental/Chaos/Public/Chaos/ImplicitObjectUnion.h
+++ b/Engine/Source/Runtime/Experimental/Chaos/Public/Chaos/ImplicitObjectUnion.h
@@ -57,7 +57,7 @@ class FImplicitObjectUnion : public FImplicitObject
 	// The total number of leaf objects in the hierarchy
 	int32 GetNumLeafObjects() const
 	{
-		return int32(NumLeafObjects);
+		return NumLeafObjects;
 	}
 
 	// Enable BVH suport for this Union. This should only be done for the root Union in a hierarchy
@@ -296,7 +296,7 @@ protected:
 	// BVH is only created when there are many objects.
 	// @todo(chaos): consider registering particles that may need BVH updated in evolution instead
 	TUniquePtr<Private::FImplicitBVH> BVH;
-	uint16 NumLeafObjects;
+	int32 NumLeafObjects;
 	FFLags Flags;
 };
 
diff --git a/Engine/Source/Runtime/Experimental/ChaosCore/Public/Chaos/Core.h b/Engine/Source/Runtime/Experimental/ChaosCore/Public/Chaos/Core.h
--- a/Engine/Source/Runtime/Experimental/ChaosCore/Public/Chaos/Core.h
+++ b/Engine/Source/Runtime/Experimental/ChaosCore/Public/Chaos/Core.h
@@ -51,7 +51,7 @@ namespace Chaos
 	// NOTE: if you get a merge conflict on the GUID, you must replace it with a new GUID - do not accept the source or target
 	// or you will likely get DDC version conflicts resulting in crashes during load.
 	// Core version string for Chaos data. Any DDC builder dependent on Chaos for serialization should depend on this version
-	inline const TCHAR* const ChaosVersionGUID = TEXT("5F81677D-8347-4E79-85D0-9F71C9A7DF03");
+	inline const TCHAR* const ChaosVersionGUID = TEXT("5E07152B-EF71-47B8-B477-0CE28888D459");
 
 	inline FString GetChaosVersionStringInner()
 	{
