Geometry collection Animation fixes
- fixed anim component to properly play animation in blueprint viewports
- fixed root proxy local transform using the parent space reference pose instead of component space ones

diff --git a/Engine/Source/Runtime/Experimental/GeometryCollectionEngine/Private/GeometryCollection/GeometryCollectionComponent.cpp b/Engine/Source/Runtime/Experimental/GeometryCollectionEngine/Private/GeometryCollection/GeometryCollectionComponent.cpp
--- a/Engine/Source/Runtime/Experimental/GeometryCollectionEngine/Private/GeometryCollection/GeometryCollectionComponent.cpp
+++ b/Engine/Source/Runtime/Experimental/GeometryCollectionEngine/Private/GeometryCollection/GeometryCollectionComponent.cpp
@@ -3610,7 +3610,7 @@ void UGeometryCollectionComponent::UpdateRootProxyComponentsIfNeeded()
 {
 	if (RootProxyStaticMeshComponents.Num() > 0 && !bUpdateComponentTransformToRootBone)
 	{
-		if (RestCollection && DynamicCollection && !IsRootBroken())
+		if (RestCollection && !IsRootBroken())
 		{
 			const TSharedPtr<FGeometryCollection, ESPMode::ThreadSafe> AssetCollection = RestCollection->GetGeometryCollection();
 			if (AssetCollection)
