Iris
* Improve logs  for Protocol mismatches and when creating ReplicationSystem

#rnx

diff --git a/Engine/Source/Runtime/Experimental/Iris/Core/Private/Iris/ReplicationSystem/ObjectReplicationBridge.cpp b/Engine/Source/Runtime/Experimental/Iris/Core/Private/Iris/ReplicationSystem/ObjectReplicationBridge.cpp
--- a/Engine/Source/Runtime/Experimental/Iris/Core/Private/Iris/ReplicationSystem/ObjectReplicationBridge.cpp
+++ b/Engine/Source/Runtime/Experimental/Iris/Core/Private/Iris/ReplicationSystem/ObjectReplicationBridge.cpp
@@ -1932,21 +1932,21 @@ void UObjectReplicationBridge::OnProtocolMismatchReported(FNetRefHandle RefHandl
 	// Ensure at the end so the log contains all the relevant information
 	ON_SCOPE_EXIT
 	{
-		ensureMsgf(false, TEXT("Protocol mismatch detected. Compare the CDO state in the server and client logs to find the source of the issue."));
+		ensureMsgf(false, TEXT("Protocol mismatch detected from %s. Compare the CDO state in the server and client logs to find the source of the issue."), *PrintConnectionInfo(ConnectionId));
 	};
 	
 
 	const FInternalNetRefIndex ObjectInternalIndex = NetRefHandleManager->GetInternalIndex(RefHandle);
 	if (ObjectInternalIndex == FNetRefHandleManager::InvalidInternalIndex)
 	{
-		UE_LOG(LogIris, Warning, TEXT("OnProtocolMismatchReported from Connection:%u for %s. But object has no InternalIndex."), ConnectionId, *RefHandle.ToString());
+		UE_LOG(LogIris, Warning, TEXT("OnProtocolMismatchReported from Connection:%s for %s. But object has no InternalIndex."), *PrintConnectionInfo(ConnectionId), *RefHandle.ToString());
 		return;
 	}
 
 	UObject* ObjInstance = NetRefHandleManager->GetReplicatedObjectInstance(ObjectInternalIndex);
 	UObject* ObjArchetype = ObjInstance ? ObjInstance->GetArchetype() : nullptr;
 
-	UE_LOG(LogIris, Error, TEXT("OnProtocolMismatchReported from client:%u when instancing %s. CDO:%s ReplicatedObject:%s NetObject:%s"), ConnectionId, *RefHandle.ToString(), *GetNameSafe(ObjArchetype), *GetNameSafe(ObjInstance), *NetRefHandleManager->PrintObjectFromIndex(ObjectInternalIndex));
+	UE_LOG(LogIris, Error, TEXT("OnProtocolMismatchReported from client:%s when instancing %s. CDO:%s ReplicatedObject:%s NetObject:%s"), *PrintConnectionInfo(ConnectionId), *RefHandle.ToString(), *GetNameSafe(ObjArchetype), *GetNameSafe(ObjInstance), *NetRefHandleManager->PrintObjectFromIndex(ObjectInternalIndex));
 
 	if (UE_LOG_ACTIVE(LogIris, Error))
 	{
@@ -1954,14 +1954,14 @@ void UObjectReplicationBridge::OnProtocolMismatchReported(FNetRefHandle RefHandl
 		const FReplicationInstanceProtocol* InstanceProtocol = ObjectData.InstanceProtocol;
 		if (!InstanceProtocol)
 		{
-			UE_LOG(LogIris, Warning, TEXT("OnProtocolMismatchReported from Connection:%u for %s. But object %s has no InstanceProtocol."), ConnectionId, *RefHandle.ToString(), *GetNameSafe(ObjInstance));
+			UE_LOG(LogIris, Warning, TEXT("OnProtocolMismatchReported from Connection:%s for %s. But object %s has no InstanceProtocol."), *PrintConnectionInfo(ConnectionId), *RefHandle.ToString(), *GetNameSafe(ObjInstance));
 			return;
 		}
 
 		const FReplicationProtocol* Protocol = ObjectData.Protocol;
 		if (!Protocol)
 		{
-			UE_LOG(LogIris, Warning, TEXT("OnProtocolMismatchReported from Connection:%u for %s. But object %s has no Protocol."), ConnectionId, *RefHandle.ToString(), *GetNameSafe(ObjInstance));
+			UE_LOG(LogIris, Warning, TEXT("OnProtocolMismatchReported from Connection:%s for %s. But object %s has no Protocol."), *PrintConnectionInfo(ConnectionId), *RefHandle.ToString(), *GetNameSafe(ObjInstance));
 			return;
 		}
 
diff --git a/Engine/Source/Runtime/Experimental/Iris/Core/Private/Iris/ReplicationSystem/ReplicationSystem.cpp b/Engine/Source/Runtime/Experimental/Iris/Core/Private/Iris/ReplicationSystem/ReplicationSystem.cpp
--- a/Engine/Source/Runtime/Experimental/Iris/Core/Private/Iris/ReplicationSystem/ReplicationSystem.cpp
+++ b/Engine/Source/Runtime/Experimental/Iris/Core/Private/Iris/ReplicationSystem/ReplicationSystem.cpp
@@ -1710,6 +1710,8 @@ UReplicationSystem* FReplicationSystemFactory::CreateReplicationSystem(const URe
 			MaxReplicationSystemId = ReplicationSystemId;
 		}
 
+		UE_LOG(LogIris, Display, TEXT("Iris ReplicationSystem[%i] is created"), ReplicationSystemId);
+
 		ReplicationSystem->Init(ReplicationSystemId, Params);
 
 		if (GetReplicationSystemCreatedDelegate().IsBound())
@@ -1732,6 +1734,9 @@ void FReplicationSystemFactory::DestroyReplicationSystem(UReplicationSystem* Sys
 	}
 
 	const uint32 Id = System->GetId();
+
+	UE_LOG(LogIris, Display, TEXT("Iris ReplicationSystem[%i] is about to be destroyed"), Id);
+
 	if (Id < MaxReplicationSystemCount)
 	{
 		ReplicationSystems[Id] = nullptr;
