Iris
* Added csv stats tracking how long the culltesting takes every frame and how many objects were cull tested

#rnx

diff --git a/Engine/Source/Runtime/Experimental/Iris/Core/Private/Iris/IrisConfig.cpp b/Engine/Source/Runtime/Experimental/Iris/Core/Private/Iris/IrisConfig.cpp
--- a/Engine/Source/Runtime/Experimental/Iris/Core/Private/Iris/IrisConfig.cpp
+++ b/Engine/Source/Runtime/Experimental/Iris/Core/Private/Iris/IrisConfig.cpp
@@ -1,6 +1,7 @@
 // Copyright Epic Games, Inc. All Rights Reserved.
 
 #include "Iris/IrisConfig.h"
+#include "Iris/Core/IrisCsv.h"
 #include "HAL/IConsoleManager.h"
 #include "Misc/CommandLine.h"
 #include "Misc/Parse.h"
@@ -35,3 +36,7 @@ EReplicationSystem GetUseIrisReplicationCmdlineValue()
 }
 
 }
+
+
+// Enable Iris category by default on servers
+CSV_DEFINE_CATEGORY(Iris, WITH_SERVER_CODE);
\ No newline at end of file
diff --git a/Engine/Source/Runtime/Experimental/Iris/Core/Private/Iris/ReplicationSystem/Filtering/NetObjectGridFilter.cpp b/Engine/Source/Runtime/Experimental/Iris/Core/Private/Iris/ReplicationSystem/Filtering/NetObjectGridFilter.cpp
--- a/Engine/Source/Runtime/Experimental/Iris/Core/Private/Iris/ReplicationSystem/Filtering/NetObjectGridFilter.cpp
+++ b/Engine/Source/Runtime/Experimental/Iris/Core/Private/Iris/ReplicationSystem/Filtering/NetObjectGridFilter.cpp
@@ -6,6 +6,7 @@
 #include "Iris/ReplicationSystem/ReplicationProtocol.h"
 #include "Iris/ReplicationSystem/ReplicationSystem.h"
 #include "Iris/ReplicationSystem/WorldLocations.h"
+#include "Iris/Core/IrisCsv.h"
 #include "Iris/Core/IrisLog.h"
 #include "Iris/Core/IrisProfiler.h"
 
@@ -78,6 +79,9 @@ void UNetObjectGridFilter::RemoveObject(uint32 ObjectIndex, const FNetObjectFilt
 void UNetObjectGridFilter::PreFilter(FNetObjectPreFilteringParams&)
 {
 	++FrameIndex;
+#if UE_NET_IRIS_CSV_STATS
+	Stats.Reset();
+#endif
 }
 
 void UNetObjectGridFilter::Filter(FNetObjectFilteringParams& Params)
@@ -139,13 +143,18 @@ void UNetObjectGridFilter::Filter(FNetObjectFilteringParams& Params)
 	 * object indices and then iterate over that. In both cases one do need to iterate over all
 	 * objects in all relevant cells anyway and setting a bit should be faster than inserting into a set.
 	 */
-	for (const FCellAndTimestamp& CellAndTimestamp : NewCells)
+
+	if (Config->bUseExactCullDistance)
 	{
-		if (FCellObjects* Objects = Cells.Find(CellAndTimestamp.Cell))
+#if UE_NET_IRIS_CSV_STATS
+		const uint64 StartTimeInCycles = FPlatformTime::Cycles64();
+		uint32 CullTestedObjects = 0;
+#endif
+		for (const FCellAndTimestamp& CellAndTimestamp : NewCells)
 		{
-			for (const uint32 ObjectIndex : Objects->ObjectIndices)
+			if (FCellObjects* Objects = Cells.Find(CellAndTimestamp.Cell))
 			{
-				if (Config->bUseExactCullDistance)
+				for (const uint32 ObjectIndex : Objects->ObjectIndices)
 				{
 					const FObjectLocationInfo& ObjectLocationInfo = static_cast<const FObjectLocationInfo&>(Params.FilteringInfos[ObjectIndex]);
 					const FPerObjectInfo& PerObjectInfo = ObjectInfos[ObjectLocationInfo.GetInfoIndex()];
@@ -158,19 +167,21 @@ void UNetObjectGridFilter::Filter(FNetObjectFilteringParams& Params)
 						if (ObjectToViewDistSq <= DistSq)
 						{
 							ConnectionInfo.RecentObjectFrameCount.Add(ObjectIndex, Config->ViewPosRelevancyFrameCount);
+							break; // Don't need to test all connections once one is within Distance
 						}
 					}
 				}
-				else
-				{
-					AllowedObjects.SetBit(ObjectIndex);
-				}
+#if UE_NET_IRIS_CSV_STATS
+				Stats.CullTestedObjects += Objects->ObjectIndices.Num();
+#endif
 			}
 		}
-	}
 
-	if (Config->bUseExactCullDistance)
-	{
+#if UE_NET_IRIS_CSV_STATS
+		Stats.CullTestingTimeInCycles += (FPlatformTime::Cycles64() - StartTimeInCycles);
+#endif
+
+		// Set the AllowedObjects and decrease their frame count
 		for (TMap<uint32, uint32>::TIterator It = ConnectionInfo.RecentObjectFrameCount.CreateIterator(); It; ++It)
 		{
 			if (It->Value > 0)
@@ -184,6 +195,27 @@ void UNetObjectGridFilter::Filter(FNetObjectFilteringParams& Params)
 			}
 		}
 	}
+	else
+	{
+		for (const FCellAndTimestamp& CellAndTimestamp : NewCells)
+		{
+			if (FCellObjects* Objects = Cells.Find(CellAndTimestamp.Cell))
+			{
+				for (const uint32 ObjectIndex : Objects->ObjectIndices)
+				{
+					AllowedObjects.SetBit(ObjectIndex);
+				}
+			}
+		}
+	}
+}
+
+void UNetObjectGridFilter::PostFilter(FNetObjectPostFilteringParams&)
+{
+#if UE_NET_IRIS_CSV_STATS
+	CSV_CUSTOM_STAT(Iris, CullTestingTimeInMS, FPlatformTime::ToMilliseconds64(Stats.CullTestingTimeInCycles), ECsvCustomStatOp::Set);
+	CSV_CUSTOM_STAT(Iris, CullTestedObjectsCount, (int32)Stats.CullTestedObjects, ECsvCustomStatOp::Set);
+#endif
 }
 
 uint32 UNetObjectGridFilter::AllocObjectInfo()
diff --git a/Engine/Source/Runtime/Experimental/Iris/Core/Private/Iris/ReplicationSystem/Filtering/ReplicationFiltering.cpp b/Engine/Source/Runtime/Experimental/Iris/Core/Private/Iris/ReplicationSystem/Filtering/ReplicationFiltering.cpp
--- a/Engine/Source/Runtime/Experimental/Iris/Core/Private/Iris/ReplicationSystem/Filtering/ReplicationFiltering.cpp
+++ b/Engine/Source/Runtime/Experimental/Iris/Core/Private/Iris/ReplicationSystem/Filtering/ReplicationFiltering.cpp
@@ -6,6 +6,7 @@
 
 #include "Iris/IrisConfigInternal.h"
 #include "Iris/IrisConstants.h"
+#include "Iris/Core/IrisCsv.h"
 #include "Iris/Core/IrisLog.h"
 #include "Iris/Core/IrisProfiler.h"
 #include "Iris/ReplicationSystem/NetRefHandleManager.h"
@@ -297,6 +298,9 @@ void FReplicationFiltering::Init(FReplicationFilteringInitParams& Params)
 
 void FReplicationFiltering::FilterPrePoll()
 {
+#if UE_NET_IRIS_CSV_STATS
+	CSV_SCOPED_TIMING_STAT(Iris, Filter_PrePoll);
+#endif
 	ResetRemovedConnections();
 
 	InitNewConnections();
@@ -395,6 +399,9 @@ void FReplicationFiltering::FilterPostPoll()
 {
 	if (HasFragmentFilters())
 	{
+#if UE_NET_IRIS_CSV_STATS
+		CSV_SCOPED_TIMING_STAT(Iris, Filter_PostPoll);
+#endif
 		UpdateDynamicFilters(ENetFilterType::PostPoll_FragmentBased);
 	}
 }
diff --git a/Engine/Source/Runtime/Experimental/Iris/Core/Private/Iris/Stats/NetStats.cpp b/Engine/Source/Runtime/Experimental/Iris/Core/Private/Iris/Stats/NetStats.cpp
--- a/Engine/Source/Runtime/Experimental/Iris/Core/Private/Iris/Stats/NetStats.cpp
+++ b/Engine/Source/Runtime/Experimental/Iris/Core/Private/Iris/Stats/NetStats.cpp
@@ -2,6 +2,7 @@
 
 #include "Iris/Stats/NetStats.h"
 #include "Iris/Stats/NetStatsContext.h"
+#include "Iris/Core/IrisCsv.h"
 #include "Iris/Core/IrisProfiler.h"
 #include "HAL/IConsoleManager.h"
 #include "Misc/ScopeLock.h"
@@ -17,9 +18,6 @@ static FAutoConsoleVariableRef CShouldIncludeSubObjectWithRoot(
 	TEXT("If enabled SubObjects will reports stats with RootObject, if set to false SubObjects will be treated as separate objects."
 	));
 
-// Enable Iris category by default on servers
-CSV_DEFINE_CATEGORY(Iris, WITH_SERVER_CODE);
-
 // Per type stats
 CSV_DEFINE_CATEGORY(IrisPreUpdateMS, WITH_SERVER_CODE);
 CSV_DEFINE_CATEGORY(IrisPreUpdateCount, WITH_SERVER_CODE);
@@ -41,6 +39,7 @@ CSV_DEFINE_CATEGORY(IrisWriteExportsCount, WITH_SERVER_CODE);
 
 void FNetSendStats::Accumulate(const FNetSendStats& Other)
 {
+#if UE_NET_IRIS_CSV_STATS
 	FScopeLock Lock(&CS);
 
 	Stats.ScheduledForReplicationRootObjectCount += Other.Stats.ScheduledForReplicationRootObjectCount;
@@ -55,12 +54,15 @@ void FNetSendStats::Accumulate(const FNetSendStats& Other)
 
 	Stats.HugeObjectWaitingForAckTimeInSeconds += Other.Stats.HugeObjectWaitingForAckTimeInSeconds;
 	Stats.HugeObjectStallingTimeInSeconds += Other.Stats.HugeObjectStallingTimeInSeconds;
+#endif 
 }
 
 void FNetSendStats::Reset()
 {
+#if UE_NET_IRIS_CSV_STATS
 	FScopeLock Lock(&CS);
 	Stats = FStats();
+#endif
 }
 
 void FNetSendStats::ReportCsvStats()
@@ -243,4 +245,4 @@ void FNetTypeStats::ReportCSVStats()
 #undef UE_NET_STATS_RECORD_TYPESTATS_COUNT
 #undef UE_NET_STATS_RECORD_TYPESTATS_BITS
 
-}
+} // end namespace UE::Net::Private
diff --git a/Engine/Source/Runtime/Experimental/Iris/Core/Private/Iris/Stats/NetStats.h b/Engine/Source/Runtime/Experimental/Iris/Core/Private/Iris/Stats/NetStats.h
--- a/Engine/Source/Runtime/Experimental/Iris/Core/Private/Iris/Stats/NetStats.h
+++ b/Engine/Source/Runtime/Experimental/Iris/Core/Private/Iris/Stats/NetStats.h
@@ -208,4 +208,4 @@ private:
 	bool bIsEnabled = false;
 };
 
-}
+} // end namespace UE::Net::Private
diff --git a/Engine/Source/Runtime/Experimental/Iris/Core/Public/Iris/Core/IrisCsv.h b/Engine/Source/Runtime/Experimental/Iris/Core/Public/Iris/Core/IrisCsv.h
new file mode 100644
--- /dev/null
+++ b/Engine/Source/Runtime/Experimental/Iris/Core/Public/Iris/Core/IrisCsv.h
@@ -0,0 +1,7 @@
+// Copyright Epic Games, Inc. All Rights Reserved.
+
+#pragma once
+
+#include "ProfilingDebugging/CsvProfiler.h"
+
+CSV_DECLARE_CATEGORY_EXTERN(Iris);
diff --git a/Engine/Source/Runtime/Experimental/Iris/Core/Public/Iris/ReplicationSystem/Filtering/NetObjectGridFilter.h b/Engine/Source/Runtime/Experimental/Iris/Core/Public/Iris/ReplicationSystem/Filtering/NetObjectGridFilter.h
--- a/Engine/Source/Runtime/Experimental/Iris/Core/Public/Iris/ReplicationSystem/Filtering/NetObjectGridFilter.h
+++ b/Engine/Source/Runtime/Experimental/Iris/Core/Public/Iris/ReplicationSystem/Filtering/NetObjectGridFilter.h
@@ -3,6 +3,7 @@
 #pragma once
 
 #include "Iris/ReplicationSystem/Filtering/NetObjectFilter.h"
+#include "Iris/IrisConfig.h"
 #include "Net/Core/NetBitArray.h"
 #include "Containers/ChunkedArray.h"
 #include "Math/Vector.h"
@@ -67,6 +68,7 @@ protected:
 	IRISCORE_API virtual void RemoveObject(uint32 ObjectIndex, const FNetObjectFilteringInfo&) override;
 	IRISCORE_API virtual void PreFilter(FNetObjectPreFilteringParams&) override;
 	IRISCORE_API virtual void Filter(FNetObjectFilteringParams&) override;
+	IRISCORE_API virtual void PostFilter(FNetObjectPostFilteringParams&) override;
 
 protected:
 	struct FObjectLocationInfo : public FNetObjectFilteringInfo
@@ -196,6 +198,18 @@ private:
 		TMap<uint32, uint32> RecentObjectFrameCount;
 	};
 
+	/** Aggregator for stats */
+	struct FNetGridFilterStats
+	{
+		FNetGridFilterStats() = default;
+
+		void Reset() { *this = FNetGridFilterStats(); }
+		
+		/** GridFilter stats */
+		uint64 CullTestingTimeInCycles = 0;
+		uint32 CullTestedObjects = 0;
+	};
+
 private:
 	uint32 AllocObjectInfo();
 	void FreeObjectInfo(uint32 Index);
@@ -212,6 +226,10 @@ private:
 	TChunkedArray<FPerObjectInfo, ObjectInfosChunkSize> ObjectInfos;
 	UE::Net::FNetBitArray AssignedObjectInfoIndices;
 
+#if UE_NET_IRIS_CSV_STATS
+	FNetGridFilterStats Stats;
+#endif
+
 	TMap<FCellCoord, FCellObjects> Cells;
 	uint32 FrameIndex = 0;
 
