Chaos - fix capsule-triangle collision corner case

Fix an issue with elonganted triangle vertices. The signed distance from the triangle vertex to the capsule was being incorrectly calculated because the test to see which direction was positive was flawed (it was using a dot product with the centroid which is not correct - it should be using the edge plane normal).

diff --git a/Engine/Source/Runtime/Experimental/Chaos/Private/Chaos/Collision/CapsuleTriangleContactPoint.cpp b/Engine/Source/Runtime/Experimental/Chaos/Private/Chaos/Collision/CapsuleTriangleContactPoint.cpp
--- a/Engine/Source/Runtime/Experimental/Chaos/Private/Chaos/Collision/CapsuleTriangleContactPoint.cpp
+++ b/Engine/Source/Runtime/Experimental/Chaos/Private/Chaos/Collision/CapsuleTriangleContactPoint.cpp
@@ -102,6 +102,41 @@
 			return;
 		}
 
+		// Edge plane normals and signed distances to each segment point
+		FVec3 EdgeNs[3];
+		FReal EdgeD0s[3];
+		FReal EdgeD1s[3];
+
+		EdgeNs[0] = FVec3::CrossProduct(V0 - V2, FaceN);
+		EdgeNs[0].Normalize(NormalToleranceSq);
+		EdgeD0s[0] = FVec3::DotProduct(P0 - V0, EdgeNs[0]);
+		EdgeD1s[0] = FVec3::DotProduct(P1 - V0, EdgeNs[0]);
+		if (((EdgeD0s[0] > RejectDistance) & (EdgeD1s[0] > RejectDistance)) != 0)
+		{
+			// Separated from triangle
+			return;
+		}
+
+		EdgeNs[1] = FVec3::CrossProduct(V1 - V0, FaceN);
+		EdgeNs[1].Normalize(NormalToleranceSq);
+		EdgeD0s[1] = FVec3::DotProduct(P0 - V1, EdgeNs[1]);
+		EdgeD1s[1] = FVec3::DotProduct(P1 - V1, EdgeNs[1]);
+		if (((EdgeD0s[1] > RejectDistance) & (EdgeD1s[1] > RejectDistance)) != 0)
+		{
+			// Separated from triangle
+			return;
+		}
+
+		EdgeNs[2] = FVec3::CrossProduct(V2 - V1, FaceN);
+		EdgeNs[2].Normalize(NormalToleranceSq);
+		EdgeD0s[2] = FVec3::DotProduct(P0 - V2, EdgeNs[2]);
+		EdgeD1s[2] = FVec3::DotProduct(P1 - V2, EdgeNs[2]);
+		if (((EdgeD0s[2] > RejectDistance) & (EdgeD1s[2] > RejectDistance)) != 0)
+		{
+			// Separated from triangle
+			return;
+		}
+
 		// Cull based on signed distance to capsule segment
 		// The edge-segment data is saved for used later on to generate contacts if necessary
 		FVec3 EdgeSegmentDeltas[3];
@@ -131,7 +166,8 @@
 			const FReal SegmentEdgeDistSq = SegmentEdgeN.SizeSquared();
 
 			// Separating axis always points away from the triangle
-			if (FVec3::DotProduct(EdgeP1 - Centroid, SegmentEdgeN) < FReal(0))
+			const FReal DotEdge = FVec3::DotProduct(SegmentEdgeN, EdgeNs[EdgeIndex]);
+			if (DotEdge < FReal(0))
 			{
 				SegmentEdgeN = -SegmentEdgeN;
 				SegmentEdgeDistSign = FReal(-1);
@@ -165,41 +201,6 @@
 			EdgeDotFace[EdgeIndex] = DotFace;
 		}
 
-		// Edge plane normals and signed distances to each segment point
-		FVec3 EdgeNs[3];
-		FReal EdgeD0s[3];
-		FReal EdgeD1s[3];
-
-		EdgeNs[0] = FVec3::CrossProduct(V0 - V2, FaceN);
-		EdgeNs[0].Normalize(NormalToleranceSq);
-		EdgeD0s[0] = FVec3::DotProduct(P0 - V0, EdgeNs[0]);
-		EdgeD1s[0] = FVec3::DotProduct(P1 - V0, EdgeNs[0]);
-		if (((EdgeD0s[0] > RejectDistance) & (EdgeD1s[0] > RejectDistance)) != 0)
-		{
-			// Separated from triangle
-			return;
-		}
-
-		EdgeNs[1] = FVec3::CrossProduct(V1 - V0, FaceN);
-		EdgeNs[1].Normalize(NormalToleranceSq);
-		EdgeD0s[1] = FVec3::DotProduct(P0 - V1, EdgeNs[1]);
-		EdgeD1s[1] = FVec3::DotProduct(P1 - V1, EdgeNs[1]);
-		if (((EdgeD0s[1] > RejectDistance) & (EdgeD1s[1] > RejectDistance)) != 0)
-		{
-			// Separated from triangle
-			return;
-		}
-
-		EdgeNs[2] = FVec3::CrossProduct(V2 - V1, FaceN);
-		EdgeNs[2].Normalize(NormalToleranceSq);
-		EdgeD0s[2] = FVec3::DotProduct(P0 - V2, EdgeNs[2]);
-		EdgeD1s[2] = FVec3::DotProduct(P1 - V2, EdgeNs[2]);
-		if (((EdgeD0s[2] > RejectDistance) & (EdgeD1s[2] > RejectDistance)) != 0)
-		{
-			// Separated from triangle
-			return;
-		}
-
 		// Handle the case where the end cap(s) are inside all the edge planes. This should be fairly common unless the triangles are
 		// much smaller than the capsule.
 		// NOTE: Here we can only collide with the end-cap that is closest to the triangle (unless they are both exactly the same distance)
@@ -436,7 +437,7 @@
 					ContactPoint.ContactType = EContactPointType::VertexPlane;
 					ContactPoint.FaceIndex = INDEX_NONE;
 				}
-				else 
+				else
 				{
 					FContactPoint& ContactPoint = OutContactPoints[OutContactPoints.Add()];
 					ContactPoint.ShapeContactPoints[0] = SegmentP - R * SegmentEdgeN;
