Add cvar to disable root proxy component creation

diff --git a/Engine/Source/Runtime/Experimental/GeometryCollectionEngine/Private/GeometryCollection/GeometryCollectionComponent.cpp b/Engine/Source/Runtime/Experimental/GeometryCollectionEngine/Private/GeometryCollection/GeometryCollectionComponent.cpp
--- a/Engine/Source/Runtime/Experimental/GeometryCollectionEngine/Private/GeometryCollection/GeometryCollectionComponent.cpp
+++ b/Engine/Source/Runtime/Experimental/GeometryCollectionEngine/Private/GeometryCollection/GeometryCollectionComponent.cpp
@@ -131,6 +131,9 @@ FAutoConsoleVariableRef CVarGeometryCollectionEmitRootBreakingEvent(TEXT("p.Chao
 bool GeometryCollectionCreatePhysicsStateInEditor = false;
 FAutoConsoleVariableRef CVarGeometryCollectionCreatePhysicsStateInEditor(TEXT("p.Chaos.GC.CreatePhysicsStateInEditor"), GeometryCollectionCreatePhysicsStateInEditor, TEXT("when on , physics state for a GC will be create in editor ( non PIE )"));
 
+bool GeometryCollectionEnableRootProxyComponents = true;
+FAutoConsoleVariableRef CVarGeometryCollectionRootProxyComponents(TEXT("p.Chaos.GC.EnableRootProxyComponents"), GeometryCollectionEnableRootProxyComponents, TEXT("when on ( by default ) , create root proxy components"));
+
 bool GeometryCollectionUseReplicationV2 = true;
 FAutoConsoleVariableRef CVarGeometryCollectionUseReplicationV2(TEXT("p.Chaos.GC.UseReplicationV2"), GeometryCollectionUseReplicationV2, TEXT("When true use new replication data model"));
 
@@ -3576,7 +3579,7 @@ bool UGeometryCollectionComponent::ShouldCreateRootProxyComponents() const
 {
 	const bool bHasRootProxyMeshes = RestCollection && RestCollection->RootProxyData.ProxyMeshes.Num() > 0;
 	const bool bHasCustomRenderer = CanUseCustomRenderer();
-	return bHasRootProxyMeshes && !bHasCustomRenderer && bEnableRootProxyStaticMeshComponents;
+	return bHasRootProxyMeshes && !bHasCustomRenderer && bEnableRootProxyStaticMeshComponents && GeometryCollectionEnableRootProxyComponents;
 }
 
 void UGeometryCollectionComponent::CreateRootProxyComponentsIfNeeded()
