Iris
* Added CSV exclusive scope to measure RepNotify and PostRepNotify user callbacks

#rnx

diff --git a/Engine/Source/Runtime/Experimental/Iris/Core/Private/Iris/ReplicationSystem/DequantizeAndApplyHelper.cpp b/Engine/Source/Runtime/Experimental/Iris/Core/Private/Iris/ReplicationSystem/DequantizeAndApplyHelper.cpp
--- a/Engine/Source/Runtime/Experimental/Iris/Core/Private/Iris/ReplicationSystem/DequantizeAndApplyHelper.cpp
+++ b/Engine/Source/Runtime/Experimental/Iris/Core/Private/Iris/ReplicationSystem/DequantizeAndApplyHelper.cpp
@@ -13,6 +13,7 @@
 #include "Iris/Serialization/NetBitStreamReader.h"
 #include "Iris/Serialization/NetSerializer.h"
 #include "Misc/MemStack.h"
+#include "ProfilingDebugging/CsvProfiler.h"
 
 namespace UE::Net::Private
 {
@@ -238,6 +239,7 @@ void FDequantizeAndApplyHelper::ApplyAndCallLegacyFunctions(FContext* Context, F
 void FDequantizeAndApplyHelper::CallLegacyPostApplyFunctions(FContext* Context, FNetSerializationContext& NetSerializationContext)
 {
 	IRIS_PROFILER_SCOPE(FDequantizeAndApplyHelper_CallLegacyPostApplyFunctions);
+	CSV_SCOPED_TIMING_STAT_EXCLUSIVE(RepNotifies);
 
 	checkSlow(Context);
 
