Constraints: stronger dependency checks when creating stacked constraints

#rnx

diff --git a/Engine/Source/Runtime/Experimental/Animation/Constraints/Private/TransformConstraint.cpp b/Engine/Source/Runtime/Experimental/Animation/Constraints/Private/TransformConstraint.cpp
--- a/Engine/Source/Runtime/Experimental/Animation/Constraints/Private/TransformConstraint.cpp
+++ b/Engine/Source/Runtime/Experimental/Animation/Constraints/Private/TransformConstraint.cpp
@@ -93,8 +93,6 @@ public:
 		}
 	}
 	
-private:
-
 	// ensure that InPossiblePrimary is not depending on InPossibleSecondary to avoid creating cycles 
 	static bool HasPrerequisiteDependencyWith(const FTickFunction* InSecondary, const FTickFunction* InPrimary, TSet<const FTickFunction*>& InOutVisitedFunctions)
 	{
@@ -135,6 +133,7 @@ private:
 		return false;
 	}
 
+private:
 	/**
 	 * Manage tick dependencies if needed to avoid cycles from a tick dependency pov.
 	 * Both InConstraintToUpdate and its parent handle are supposed valid at this point
@@ -2012,7 +2011,7 @@ bool FTransformConstraintUtils::BuildDependencies(UWorld* InWorld, UTickableTran
 
 	// if child handle is the parent of some other constraints, ensure they will tick after that new one
 	TArray<TWeakObjectPtr<UTickableConstraint>> ChildChildConstraints;
-	GetChildrenConstraints(InWorld, ChildHandle, ChildChildConstraints, bIncludeTarget && !bSelf);
+	GetChildrenConstraints(InWorld, InConstraint, ChildChildConstraints, bIncludeTarget && !bSelf);
 	for (const TWeakObjectPtr<UTickableConstraint>& ChildConstraint: ChildChildConstraints)
 	{
 		Controller.SetConstraintsDependencies(InConstraint->ConstraintID, ChildConstraint->ConstraintID);
@@ -2280,20 +2279,28 @@ int32 FTransformConstraintUtils::GetLastActiveConstraintIndex(const TArray< TWea
 
 void FTransformConstraintUtils::GetChildrenConstraints(
 	UWorld* World,
-	const UTransformableHandle* InHandle,
+	const UTickableTransformConstraint* InConstraint,
 	TArray< TWeakObjectPtr<UTickableConstraint> >& OutConstraints,
 	const bool bIncludeTarget)
 {
+	if (!InConstraint || !InConstraint->IsValid())
+	{
+		// this probably has been checked before but we want to make sure the data is safe to use
+		return;
+	}
+	
+	const UTransformableHandle* Handle = InConstraint->ChildTRSHandle.Get();
+	
 	using ConstraintPtr = TWeakObjectPtr<UTickableConstraint>;
 	
 	// filter for transform constraints where the InHandle is the parent (based on its hash value)
 	// and also has the same target if bIncludeTarget is true
-	const uint32 ParentHash = InHandle->GetHash();
-	const UObject* ParentTarget = InHandle->GetTarget().Get();
-	auto Predicate = [InHandle, ParentHash, bIncludeTarget, ParentTarget, World](const ConstraintPtr& Constraint)
+	const uint32 ParentHash = Handle->GetHash();
+	const UObject* ParentTarget = Handle->GetTarget().Get();
+	auto Predicate = [InConstraint, Handle, ParentHash, bIncludeTarget, ParentTarget, World](const ConstraintPtr& Constraint)
 	{
 		const UTickableTransformConstraint* TransformConstraint = Cast<UTickableTransformConstraint>(Constraint.Get());
-		if (!TransformConstraint)
+		if (!TransformConstraint || TransformConstraint == InConstraint)
 		{
 			return false;
 		}
@@ -2312,16 +2319,41 @@ void FTransformConstraintUtils::GetChildrenConstraints(
 				if (Target == ParentTarget)
 				{
 					// check direct dependencies to avoid evaluation order issues
-					if (InHandle->HasDirectDependencyWith(*OtherParentHandle))
+					if (Handle->HasDirectDependencyWith(*OtherParentHandle))
 					{
 						return false;
 					}
 
 					// check constraints dependencies to avoid cycles
-					if (HasConstraintDependencyWith(World, OtherParentHandle, InHandle))
+					if (HasConstraintDependencyWith(World, OtherParentHandle, Handle))
 					{
 						return false;
 					}
+
+					// check TransformConstraint's ChildHandle
+					if (const UTransformableHandle* OtherChildHandle = TransformConstraint->ChildTRSHandle)
+					{
+						const UTransformableHandle* ParentHandle = InConstraint->ParentTRSHandle.Get();
+						
+						// if TransformConstraint's child is InConstraint's parent then avoid cycles
+						if (OtherChildHandle->GetHash() == ParentHandle->GetHash())
+						{
+							return false;
+						}
+
+						// check dependencies with OtherChildHandle to avoid cycles
+						const FHandleDependencyChecker Checker(World);
+						if (Checker.HasDependency(*ParentHandle, *OtherChildHandle))
+						{
+							return false;
+						}
+						
+						// check dependencies with OtherParentHandle to avoid cycles
+						if (Checker.HasDependency(*ParentHandle, *OtherParentHandle))
+						{
+							return false;
+						}
+					}
 					
 					return true;
 				}
@@ -2360,3 +2392,39 @@ FConstraintDependencyScope::~FConstraintDependencyScope()
 		}
 	}
 }
+
+FHandleDependencyChecker::FHandleDependencyChecker(UWorld* InWorld)
+	: WeakWorld(InWorld)
+{}
+
+bool FHandleDependencyChecker::HasDependency(const UTransformableHandle& InHandle, const UTransformableHandle& InParentToCheck) const
+{
+	// check direct dependency
+	if (InHandle.HasDirectDependencyWith(InParentToCheck))
+	{
+		return true;
+	}
+
+	UWorld* World = WeakWorld.IsValid() ? WeakWorld.Get() : nullptr;
+	if (::IsValid(World))
+	{
+		// check constraints dependency
+		if (HasConstraintDependencyWith(World, &InHandle, &InParentToCheck))
+		{
+			return true;
+		}
+
+		// check any existing tick dependency
+		{
+			TSet<const FTickFunction*> VisitedFunctions;
+			const FTickFunction* TickFunction = InHandle.GetTickFunction();
+			const FTickFunction* ParentTickFunctionToCheck = InParentToCheck.GetTickFunction();
+			if (FConstraintCycleChecker::HasPrerequisiteDependencyWith(TickFunction, ParentTickFunctionToCheck, VisitedFunctions))
+			{
+				return true;
+			}
+		}
+	}
+	
+	return false;
+}
\ No newline at end of file
diff --git a/Engine/Source/Runtime/Experimental/Animation/Constraints/Public/TransformConstraint.h b/Engine/Source/Runtime/Experimental/Animation/Constraints/Public/TransformConstraint.h
--- a/Engine/Source/Runtime/Experimental/Animation/Constraints/Public/TransformConstraint.h
+++ b/Engine/Source/Runtime/Experimental/Animation/Constraints/Public/TransformConstraint.h
@@ -491,12 +491,12 @@ struct FTransformConstraintUtils
 	static CONSTRAINTS_API int32 GetLastActiveConstraintIndex(const TArray< TWeakObjectPtr<UTickableConstraint> >& InConstraints);
 
 	/**
-	 * Fills a constraint array that InHandle is the parent of.
+	 * Fills a constraint array that InConstraint->ChildHandle is the parent of.
 	 * If bIncludeTarget is true, we also get the other constraints that act on the same target.
 	 */
 	static CONSTRAINTS_API void GetChildrenConstraints(
 		UWorld* World,
-		const UTransformableHandle* InHandle,
+		const UTickableTransformConstraint* InConstraint,
 		TArray< TWeakObjectPtr<UTickableConstraint> >& OutConstraints,
 		const bool bIncludeTarget = false);
 	
@@ -521,4 +521,17 @@ private:
 	TWeakObjectPtr<UTickableTransformConstraint> WeakConstraint = nullptr;
 	TWeakObjectPtr<UWorld> WeakWorld = nullptr;
 	bool bPreviousValidity = false;
+};
+
+/**
+ * FHandleDependencyChecker provides a way to check (direct + constraints + tick) dependencies between two UTransformableHandle
+ * HasDependency will return true if InHandle depends on InParentToCheck.
+ */
+
+struct FHandleDependencyChecker
+{
+	FHandleDependencyChecker(UWorld* InWorld = nullptr);
+	bool HasDependency(const UTransformableHandle& InHandle, const UTransformableHandle& InParentToCheck) const;		
+private:
+	TWeakObjectPtr<UWorld> WeakWorld = nullptr;
 };
\ No newline at end of file
