Remove ensureAlways

#rnx

diff --git a/Engine/Source/Runtime/Experimental/Iris/Core/Private/Iris/Serialization/StringNetSerializers.cpp b/Engine/Source/Runtime/Experimental/Iris/Core/Private/Iris/Serialization/StringNetSerializers.cpp
--- a/Engine/Source/Runtime/Experimental/Iris/Core/Private/Iris/Serialization/StringNetSerializers.cpp
+++ b/Engine/Source/Runtime/Experimental/Iris/Core/Private/Iris/Serialization/StringNetSerializers.cpp
@@ -221,7 +221,7 @@ void FNameNetSerializer::Quantize(FNetSerializationContext& Context, const FNetQ
 			}
 			else
 			{
-				ensureAlwaysMsgf(bEncodingSuccess, TEXT("Failed to encode string '%s'"), TempWideBuffer);
+				ensureMsgf(bEncodingSuccess, TEXT("Failed to encode string '%s'"), TempWideBuffer);
 				TargetName.ElementCount = 1;
 				EncodingBuffer[0] = 0;
 			}
