Make SetCustomInstanceData() be part of generic IGeometryCollectionExternalRenderInterface.
It's not truly generic so might to move out to it's own interface?
Add geometry collection blueprint library with function to call this.

diff --git a/Engine/Source/Runtime/Experimental/GeometryCollectionEngine/Private/GeometryCollection/GeometryCollectionBlueprintLibrary.cpp b/Engine/Source/Runtime/Experimental/GeometryCollectionEngine/Private/GeometryCollection/GeometryCollectionBlueprintLibrary.cpp
new file mode 100644
--- /dev/null
+++ b/Engine/Source/Runtime/Experimental/GeometryCollectionEngine/Private/GeometryCollection/GeometryCollectionBlueprintLibrary.cpp
@@ -0,0 +1,20 @@
+// Copyright Epic Games, Inc. All Rights Reserved.
+
+#include "GeometryCollection/GeometryCollectionBlueprintLibrary.h"
+
+#include "GeometryCollection/GeometryCollectionComponent.h"
+#include "GeometryCollection/GeometryCollectionExternalRenderInterface.h"
+
+void UGeometryCollectionBlueprintLibrary::SetISMPoolCustomInstanceData(UGeometryCollectionComponent* GeometryCollectionComponent, int32 CustomFloatIndex, float CustomFloatValue)
+{
+	if (GeometryCollectionComponent == nullptr)
+	{
+		return;
+	}
+	IGeometryCollectionExternalRenderInterface* CustomRenderer = GeometryCollectionComponent->GetCustomRenderer();
+	if (CustomRenderer == nullptr)
+	{
+		return;
+	}
+	CustomRenderer->SetCustomInstanceData(CustomFloatIndex, CustomFloatValue);
+}
diff --git a/Engine/Source/Runtime/Experimental/GeometryCollectionEngine/Public/GeometryCollection/GeometryCollectionBlueprintLibrary.h b/Engine/Source/Runtime/Experimental/GeometryCollectionEngine/Public/GeometryCollection/GeometryCollectionBlueprintLibrary.h
new file mode 100644
--- /dev/null
+++ b/Engine/Source/Runtime/Experimental/GeometryCollectionEngine/Public/GeometryCollection/GeometryCollectionBlueprintLibrary.h
@@ -0,0 +1,25 @@
+// Copyright Epic Games, Inc. All Rights Reserved.
+
+#pragma once
+
+#include "Kismet/BlueprintFunctionLibrary.h"
+
+#include "GeometryCollectionBlueprintLibrary.generated.h"
+
+/** Blueprint library for Geometry Collections. */
+UCLASS(meta = (ScriptName = "GeometryCollectionLibrary"))
+class GEOMETRYCOLLECTIONENGINE_API UGeometryCollectionBlueprintLibrary : public UBlueprintFunctionLibrary
+{
+	GENERATED_BODY()
+
+public:
+	/** 
+	 * Set a custom instance data value for all instances associated with a geometry collection. 
+	 * This assumes that the geometry collection is using a custom instanced renderer.
+	 * @param GeometryCollectionComponent	The Geometry Collection Component that we want to set custom instance data on.
+	 * @param CustomDataIndex	The custom instance data slot that we want to set.
+	 * @param CustomDataValue	The value to set to the custom instance data slot.
+	 */
+	UFUNCTION(BlueprintCallable, Category = AtomMaterial)
+	static void SetISMPoolCustomInstanceData(UGeometryCollectionComponent* GeometryCollectionComponent, int32 CustomDataIndex, float CustomDataValue);
+};
diff --git a/Engine/Source/Runtime/Experimental/GeometryCollectionEngine/Public/GeometryCollection/GeometryCollectionExternalRenderInterface.h b/Engine/Source/Runtime/Experimental/GeometryCollectionEngine/Public/GeometryCollection/GeometryCollectionExternalRenderInterface.h
--- a/Engine/Source/Runtime/Experimental/GeometryCollectionEngine/Public/GeometryCollection/GeometryCollectionExternalRenderInterface.h
+++ b/Engine/Source/Runtime/Experimental/GeometryCollectionEngine/Public/GeometryCollection/GeometryCollectionExternalRenderInterface.h
@@ -55,6 +55,9 @@ public:
 	*/
 	virtual void UpdateTransforms(UGeometryCollection const& InGeometryCollection, TArrayView<const FTransform3f> InTransforms) = 0;
 
+	// todo: Maybe move this to a separate ISMPool renderer interface? But maybe that is overengineering?
+	virtual void SetCustomInstanceData(int32 CustomFloatIndex, float CustomFloatValue) {}
+
 	UE_DEPRECATED(5.4, "Use flags version of UpdateState instead")
 	virtual void UpdateState(UGeometryCollection const& InGeometryCollection, FTransform const& InComponentTransform, bool bInIsBroken, bool bInIsVisible) 
 	{
