ChaosSofts : Skip particles of 0 stiffness on multires springs.

diff --git a/Engine/Source/Runtime/Experimental/Chaos/Private/Chaos/SoftsMultiResConstraints.cpp b/Engine/Source/Runtime/Experimental/Chaos/Private/Chaos/SoftsMultiResConstraints.cpp
--- a/Engine/Source/Runtime/Experimental/Chaos/Private/Chaos/SoftsMultiResConstraints.cpp
+++ b/Engine/Source/Runtime/Experimental/Chaos/Private/Chaos/SoftsMultiResConstraints.cpp
@@ -21,6 +21,8 @@ static_assert(sizeof(ispc::FIntVector) == sizeof(Chaos::TVec3<int32>), "sizeof(i
 
 bool bChaos_MultiRes_ISPC_Enabled = true;
 FAutoConsoleVariableRef CVarChaosMultiResISPCEnabled(TEXT("p.Chaos.MultiRes.ISPC"), bChaos_MultiRes_ISPC_Enabled, TEXT("Whether to use ISPC optimizations in MultiRes constraints"));
+bool bChaos_MultiRes_SparseWeightMap_Enabled = false;
+FAutoConsoleVariableRef CVarChaosMultiResSparseWeightMapEnabled(TEXT("p.Chaos.MultiRes.SparseWeightMap"), bChaos_MultiRes_SparseWeightMap_Enabled, TEXT("Exploit the sparse weight map structure and skip the particles with 0 stiffness at the beginning and at the end"));
 #endif
 
 namespace Chaos::Softs {
@@ -71,6 +73,32 @@ void FMultiResConstraints::SetProperties(const FCollectionPropertyConstFacade& P
 			VelocityTargetStiffness.SetWeightedValue(WeightedValue, bUseXPBD ? MaxStiffness : (FSolverReal)1.f);
 		}
 	}
+
+	for (int32 Index = 0; Index < NumParticles; ++Index)
+	{
+		if (IsConstraintActive(Index))
+		{
+			NonZeroStiffnessMin = Index;
+			break;
+		}
+	}
+	for (int32 Index = NumParticles - 1; Index > INDEX_NONE; --Index)
+	{
+		if (IsConstraintActive(Index))
+		{
+			NonZeroStiffnessMax = Index;
+			break;
+		}
+	}
+	check(NonZeroStiffnessMax >= NonZeroStiffnessMin)
+	if (NonZeroStiffnessMax != INDEX_NONE && NonZeroStiffnessMin != INDEX_NONE)
+	{
+		bStiffnessEntriesInitialized = true;
+	}
+	else
+	{
+		bStiffnessEntriesInitialized = false;
+	}
 }
 
 void FMultiResConstraints::ApplyHelper(FSolverParticlesRange& Particles, const FSolverReal Dt, const int32 ConstraintIndex, const FSolverReal ExpStiffnessValue, const FSolverReal ExpVelocityTargetStiffnessValue) const
@@ -130,20 +158,44 @@ void FMultiResConstraints::Apply(FSolverParticlesRange& Particles, const FSolver
 		}
 		else
 		{
-			ispc::ApplyMultiResConstraintsWithWeightMaps(
-				(ispc::FVector4f*)Particles.GetPAndInvM().GetData(),
-				(const ispc::FVector3f*)Particles.XArray().GetData(),
-				(const ispc::FVector3f*)FineTargetPositions.GetData(),
-				(const ispc::FVector3f*)FineTargetVelocities.GetData(),
-				Dt,
-				bStiffnessHasWeightMap,
-				Stiffness.GetIndices().GetData(),
-				Stiffness.GetTable().GetData(),
-				bVelocityStiffnessHasWeightMap,
-				VelocityTargetStiffness.GetIndices().GetData(),
-				VelocityTargetStiffness.GetTable().GetData(),
-				Particles.Size()
-			);
+#if !UE_BUILD_SHIPPING
+			if (bChaos_MultiRes_SparseWeightMap_Enabled && bStiffnessEntriesInitialized)
+			{
+				check(NonZeroStiffnessMax != INDEX_NONE && NonZeroStiffnessMin != INDEX_NONE)
+				ispc::ApplyMultiResConstraintsWithWeightMaps(
+					(ispc::FVector4f*)&Particles.GetPAndInvM().GetData()[NonZeroStiffnessMin],
+					(const ispc::FVector3f*)&Particles.XArray().GetData()[NonZeroStiffnessMin],
+					(const ispc::FVector3f*)&FineTargetPositions.GetData()[NonZeroStiffnessMin],
+					(const ispc::FVector3f*)&FineTargetVelocities.GetData()[NonZeroStiffnessMin],
+					Dt,
+					bStiffnessHasWeightMap,
+					&Stiffness.GetIndices().GetData()[NonZeroStiffnessMin],
+					&Stiffness.GetTable().GetData()[0],
+					bVelocityStiffnessHasWeightMap,
+					&VelocityTargetStiffness.GetIndices().GetData()[NonZeroStiffnessMin],
+					&VelocityTargetStiffness.GetTable().GetData()[0],
+					NonZeroStiffnessMax - NonZeroStiffnessMin 
+				);
+			}
+			else
+#endif
+			{
+				ispc::ApplyMultiResConstraintsWithWeightMaps(
+					(ispc::FVector4f*)Particles.GetPAndInvM().GetData(),
+					(const ispc::FVector3f*)Particles.XArray().GetData(),
+					(const ispc::FVector3f*)FineTargetPositions.GetData(),
+					(const ispc::FVector3f*)FineTargetVelocities.GetData(),
+					Dt,
+					bStiffnessHasWeightMap,
+					Stiffness.GetIndices().GetData(),
+					Stiffness.GetTable().GetData(),
+					bVelocityStiffnessHasWeightMap,
+					VelocityTargetStiffness.GetIndices().GetData(),
+					VelocityTargetStiffness.GetTable().GetData(),
+					Particles.Size()
+				);
+			}
+
 		}
 	}
 	else
@@ -223,16 +275,34 @@ void FMultiResConstraints::UpdateFineTargets(const FSolverParticlesRange& Coarse
 #if INTEL_ISPC
 		if (!bUseXPBD && bChaos_MultiRes_ISPC_Enabled)
 		{
-			ispc::MultiResUpdateFineTargets(
-				(ispc::FVector3f*)FineTargetPositions.GetData(),
-				(ispc::FVector3f*)FineTargetVelocities.GetData(),
-				(const ispc::FVector3f*)CoarseX.GetData(),
-				(const ispc::FVector3f*)Normals.GetData(),
-				(const ispc::FVector3f*)CoarseVelocity.GetData(),
-				(const ispc::FVector4f*)CoarseToFinePositionBaryCoordsAndDist.GetData(),
-				(const ispc::FIntVector*)CoarseToFineSourceMeshVertIndices.GetData(),
-				NumParticles
-			);
+#if !UE_BUILD_SHIPPING
+			if (Stiffness.HasWeightMap() && bChaos_MultiRes_SparseWeightMap_Enabled && bStiffnessEntriesInitialized)
+			{
+				ispc::MultiResUpdateFineTargets(
+					(ispc::FVector3f*)&FineTargetPositions.GetData()[NonZeroStiffnessMin],
+					(ispc::FVector3f*)&FineTargetVelocities.GetData()[NonZeroStiffnessMin],
+					(const ispc::FVector3f*)CoarseX.GetData(),
+					(const ispc::FVector3f*)Normals.GetData(),
+					(const ispc::FVector3f*)CoarseVelocity.GetData(),
+					(const ispc::FVector4f*)&CoarseToFinePositionBaryCoordsAndDist.GetData()[NonZeroStiffnessMin],
+					(const ispc::FIntVector*)&CoarseToFineSourceMeshVertIndices.GetData()[NonZeroStiffnessMin],
+					NonZeroStiffnessMax - NonZeroStiffnessMin
+				);
+			}
+			else
+#endif
+			{
+				ispc::MultiResUpdateFineTargets(
+					(ispc::FVector3f*)FineTargetPositions.GetData(),
+					(ispc::FVector3f*)FineTargetVelocities.GetData(),
+					(const ispc::FVector3f*)CoarseX.GetData(),
+					(const ispc::FVector3f*)Normals.GetData(),
+					(const ispc::FVector3f*)CoarseVelocity.GetData(),
+					(const ispc::FVector4f*)CoarseToFinePositionBaryCoordsAndDist.GetData(),
+					(const ispc::FIntVector*)CoarseToFineSourceMeshVertIndices.GetData(),
+					NumParticles
+				);
+			}
 		}
 		else
 #endif
diff --git a/Engine/Source/Runtime/Experimental/Chaos/Public/Chaos/SoftsMultiResConstraints.h b/Engine/Source/Runtime/Experimental/Chaos/Public/Chaos/SoftsMultiResConstraints.h
--- a/Engine/Source/Runtime/Experimental/Chaos/Public/Chaos/SoftsMultiResConstraints.h
+++ b/Engine/Source/Runtime/Experimental/Chaos/Public/Chaos/SoftsMultiResConstraints.h
@@ -118,6 +118,10 @@ private:
 	TArray<FSolverVec3> FineTargetVelocities;
 	TArray<FSolverReal> CoarseBarycentricMass;
 
+	int32 NonZeroStiffnessMin = INDEX_NONE;
+	int32 NonZeroStiffnessMax = INDEX_NONE;
+	bool bStiffnessEntriesInitialized = false;
+
 	UE_CHAOS_DECLARE_PROPERTYCOLLECTION_NAME(MultiResUseXPBD, bool);
 	UE_CHAOS_DECLARE_PROPERTYCOLLECTION_NAME(MultiResStiffness, float);
 	UE_CHAOS_DECLARE_PROPERTYCOLLECTION_NAME(MultiResVelocityTargetStiffness, float);
