Iris - Fix mismatching alloc/free.

#rnx
#prefl;ight 65e5a8d71e0f4ee02147d4d2

diff --git a/Engine/Source/Runtime/Experimental/Iris/Core/Private/Iris/ReplicationSystem/NetBlob/NetBlob.cpp b/Engine/Source/Runtime/Experimental/Iris/Core/Private/Iris/ReplicationSystem/NetBlob/NetBlob.cpp
--- a/Engine/Source/Runtime/Experimental/Iris/Core/Private/Iris/ReplicationSystem/NetBlob/NetBlob.cpp
+++ b/Engine/Source/Runtime/Experimental/Iris/Core/Private/Iris/ReplicationSystem/NetBlob/NetBlob.cpp
@@ -28,7 +28,7 @@ FNetBlob::~FNetBlob()
 	{
 		if (EnumHasAnyFlags(Descriptor->Traits, EReplicationStateTraits::HasDynamicState))
 		{
-			if (uint8* StateBuffer = QuantizedBlobState.Get())
+			if (uint8* StateBuffer = QuantizedBlobState.GetStateBuffer())
 			{
 				Private::FReplicationStateOperationsInternal::FreeDynamicState(StateBuffer, Descriptor);
 			}
@@ -36,7 +36,7 @@ FNetBlob::~FNetBlob()
 	}
 }
 
-void FNetBlob::SetState(const TRefCountPtr<const FReplicationStateDescriptor>& InBlobDescriptor, TUniquePtr<uint8> InQuantizedBlobState)
+void FNetBlob::SetState(const TRefCountPtr<const FReplicationStateDescriptor>& InBlobDescriptor, FQuantizedBlobState&& InQuantizedBlobState)
 {
 	BlobDescriptor = InBlobDescriptor;
 	QuantizedBlobState = MoveTemp(InQuantizedBlobState);
@@ -95,26 +95,26 @@ void FNetBlob::Deserialize(FNetSerializationContext& Context)
 
 void FNetBlob::CollectObjectReferences(FNetSerializationContext& Context, FNetReferenceCollector& Collector) const
 {
-	if (BlobDescriptor.IsValid() && QuantizedBlobState.IsValid())
+	if (BlobDescriptor.IsValid() && QuantizedBlobState.GetStateBuffer())
 	{
 		const FNetSerializerChangeMaskParam InitStateChangeMaskInfo = { 0 };
-		Private::FReplicationStateOperationsInternal::CollectReferences(Context, Collector, InitStateChangeMaskInfo, QuantizedBlobState.Get(), BlobDescriptor);
+		Private::FReplicationStateOperationsInternal::CollectReferences(Context, Collector, InitStateChangeMaskInfo, QuantizedBlobState.GetStateBuffer(), BlobDescriptor);
 	}
 }
 
 void FNetBlob::SerializeBlob(FNetSerializationContext& Context) const
 {
-	if (BlobDescriptor.IsValid() && QuantizedBlobState.IsValid())
+	if (BlobDescriptor.IsValid() && QuantizedBlobState.GetStateBuffer())
 	{
-		FReplicationStateOperations::Serialize(Context, QuantizedBlobState.Get(), BlobDescriptor);
+		FReplicationStateOperations::Serialize(Context, QuantizedBlobState.GetStateBuffer(), BlobDescriptor);
 	}
 }
 
 void FNetBlob::DeserializeBlob(FNetSerializationContext& Context)
 {
-	if (BlobDescriptor.IsValid() && QuantizedBlobState.IsValid())
+	if (BlobDescriptor.IsValid() && QuantizedBlobState.GetStateBuffer())
 	{
-		FReplicationStateOperations::Deserialize(Context, QuantizedBlobState.Get(), BlobDescriptor);
+		FReplicationStateOperations::Deserialize(Context, QuantizedBlobState.GetStateBuffer(), BlobDescriptor);
 	}
 }
 
@@ -164,4 +164,13 @@ void FNetObjectAttachment::DeserializeSubObjectReference(FNetSerializationContex
 	ReadFullNetObjectReference(Context, TargetObjectReference);
 }
 
+FNetBlob::FQuantizedBlobState::FQuantizedBlobState(uint32 Size, uint32 Alignment)
+: StateBuffer(static_cast<uint8*>(GMalloc->Malloc(Size, Alignment)))
+{
+	if (StateBuffer)
+	{
+		FMemory::Memzero(StateBuffer, Size);
+	}
+}
+
 }
diff --git a/Engine/Source/Runtime/Experimental/Iris/Core/Private/Iris/ReplicationSystem/NetBlob/NetRPC.cpp b/Engine/Source/Runtime/Experimental/Iris/Core/Private/Iris/ReplicationSystem/NetBlob/NetRPC.cpp
--- a/Engine/Source/Runtime/Experimental/Iris/Core/Private/Iris/ReplicationSystem/NetBlob/NetRPC.cpp
+++ b/Engine/Source/Runtime/Experimental/Iris/Core/Private/Iris/ReplicationSystem/NetBlob/NetRPC.cpp
@@ -369,9 +369,7 @@ bool FNetRPC::ResolveFunctionAndObject(FNetSerializationContext& Context)
 
 	if (FunctionDescriptor->Descriptor->InternalSize)
 	{
-		uint8* StateBuffer = static_cast<uint8*>(GMalloc->Malloc(FunctionDescriptor->Descriptor->InternalSize, FunctionDescriptor->Descriptor->InternalAlignment));
-		FMemory::Memzero(StateBuffer, FunctionDescriptor->Descriptor->InternalSize);
-		QuantizedBlobState.Reset(StateBuffer);
+		QuantizedBlobState = FQuantizedBlobState(FunctionDescriptor->Descriptor->InternalSize, FunctionDescriptor->Descriptor->InternalAlignment);
 	}
 
 	return true;
@@ -392,14 +390,12 @@ FNetRPC* FNetRPC::Create(UReplicationSystem* ReplicationSystem, const FNetBlobCr
 	}
 
 	const FReplicationStateDescriptor* BlobDescriptor = FunctionDescriptor->Descriptor;
-	uint8* QuantizedBlobState = nullptr;
+	FQuantizedBlobState QuantizedBlobState;
 
 	// Don't spend CPU cycles on quantizing zero parameters
 	if (BlobDescriptor != nullptr && BlobDescriptor->InternalSize)
 	{
-		uint8* StateBuffer = static_cast<uint8*>(GMalloc->Malloc(BlobDescriptor->InternalSize, BlobDescriptor->InternalAlignment));
-		FMemory::Memzero(StateBuffer, BlobDescriptor->InternalSize);
-		QuantizedBlobState = StateBuffer;
+		QuantizedBlobState = FQuantizedBlobState(BlobDescriptor->InternalSize, BlobDescriptor->InternalAlignment);
 
 		// Setup Context
 		FNetSerializationContext Context;
@@ -408,7 +404,7 @@ FNetRPC* FNetRPC::Create(UReplicationSystem* ReplicationSystem, const FNetBlobCr
 		Context.SetInternalContext(&InternalContext);
 
 		// Quantize the function parameters
-		FReplicationStateOperations::Quantize(Context, StateBuffer, static_cast<const uint8*>(FunctionParameters), BlobDescriptor);
+		FReplicationStateOperations::Quantize(Context, QuantizedBlobState.GetStateBuffer(), static_cast<const uint8*>(FunctionParameters), BlobDescriptor);
 	}
 
 	FNetRPC* NetRPC = new FNetRPC(CreationInfo);
@@ -424,7 +420,7 @@ FNetRPC* FNetRPC::Create(UReplicationSystem* ReplicationSystem, const FNetBlobCr
 				FNetSerializationContext LocalContext;
 				FNetReferenceCollector Collector(ENetReferenceCollectorTraits::OnlyCollectReferencesThatCanBeExported);
 				const FNetSerializerChangeMaskParam InitStateChangeMaskInfo = { 0 };
-				FReplicationStateOperationsInternal::CollectReferences(LocalContext, Collector, InitStateChangeMaskInfo, QuantizedBlobState, BlobDescriptor);
+				FReplicationStateOperationsInternal::CollectReferences(LocalContext, Collector, InitStateChangeMaskInfo, QuantizedBlobState.GetStateBuffer(), BlobDescriptor);
 
 				if (Collector.GetCollectedReferences().Num())
 				{
@@ -439,7 +435,7 @@ FNetRPC* FNetRPC::Create(UReplicationSystem* ReplicationSystem, const FNetBlobCr
 			}
 		}
 
-		NetRPC->SetState(BlobDescriptor, TUniquePtr<uint8>(QuantizedBlobState));
+		NetRPC->SetState(BlobDescriptor, MoveTemp(QuantizedBlobState));
 	}
 
 	return NetRPC;
@@ -555,7 +551,7 @@ void FNetRPC::CallFunction(FNetRPCCallContext& CallContext)
 			}
 		}
 		
-		FReplicationStateOperations::Dequantize(Context, FunctionParameters, QuantizedBlobState.Get(), BlobDescriptor);
+		FReplicationStateOperations::Dequantize(Context, FunctionParameters, QuantizedBlobState.GetStateBuffer(), BlobDescriptor);
 	}
 
 	// Forward function
diff --git a/Engine/Source/Runtime/Experimental/Iris/Core/Public/Iris/ReplicationSystem/NetBlob/NetBlob.h b/Engine/Source/Runtime/Experimental/Iris/Core/Public/Iris/ReplicationSystem/NetBlob/NetBlob.h
--- a/Engine/Source/Runtime/Experimental/Iris/Core/Public/Iris/ReplicationSystem/NetBlob/NetBlob.h
+++ b/Engine/Source/Runtime/Experimental/Iris/Core/Public/Iris/ReplicationSystem/NetBlob/NetBlob.h
@@ -61,11 +61,29 @@ struct FNetBlobCreationInfo
 class FNetBlob
 {
 public:
+	struct FQuantizedBlobState
+	{
+		FQuantizedBlobState() = default;
+		FQuantizedBlobState(const FQuantizedBlobState&) = delete;
+		FQuantizedBlobState(FQuantizedBlobState&&);
+		IRISCORE_API FQuantizedBlobState(uint32 Size, uint32 Alignment);
+		~FQuantizedBlobState();
+
+		FQuantizedBlobState& operator=(const FQuantizedBlobState&) = delete;
+		FQuantizedBlobState& operator=(FQuantizedBlobState&&);
+
+		uint8* GetStateBuffer() { return StateBuffer; }
+		const uint8* GetStateBuffer() const { return StateBuffer; }
+
+	private:
+		uint8* StateBuffer = nullptr;
+	};
+
 	/** Construct a NetBlob with reference count zero. */
 	IRISCORE_API FNetBlob(const FNetBlobCreationInfo&);
 
 	/** Set the blob state. Use when there's a descriptor to avoid having to override serialization functions. */
-	IRISCORE_API void SetState(const TRefCountPtr<const FReplicationStateDescriptor>& BlobDescriptor, TUniquePtr<uint8> QuantizedBlobState);
+	IRISCORE_API void SetState(const TRefCountPtr<const FReplicationStateDescriptor>& BlobDescriptor, FQuantizedBlobState&& QuantizedBlobState);
 
 	/** Returns the FNetBlobCreationInfo. */
 	const FNetBlobCreationInfo& GetCreationInfo() const { return CreationInfo; }
@@ -133,7 +151,7 @@ protected:
 	TRefCountPtr<const FReplicationStateDescriptor> BlobDescriptor;
 
 	/** The state buffer that holds the data described by the descriptor. */
-	TUniquePtr<uint8> QuantizedBlobState;
+	FQuantizedBlobState QuantizedBlobState;
 
 private:
 	mutable std::atomic<int32> RefCount;
@@ -178,6 +196,32 @@ protected:
 	FNetObjectReference TargetObjectReference;
 };
 
+inline FNetBlob::FQuantizedBlobState::FQuantizedBlobState(FNetBlob::FQuantizedBlobState&& Other)
+{
+	StateBuffer = Other.StateBuffer;
+	Other.StateBuffer = nullptr;
+}
+
+inline FNetBlob::FQuantizedBlobState::~FQuantizedBlobState()
+{
+	if (StateBuffer)
+	{
+		GMalloc->Free(StateBuffer);
+	}
+}
+
+inline FNetBlob::FQuantizedBlobState& FNetBlob::FQuantizedBlobState::operator=(FNetBlob::FQuantizedBlobState&& Other)
+{
+	if (StateBuffer)
+	{
+		GMalloc->Free(StateBuffer);
+	}
+
+	StateBuffer = Other.StateBuffer;
+	Other.StateBuffer = nullptr;
+	return *this;
+}
+
 inline void FNetObjectAttachment::SetNetObjectReference(const FNetObjectReference& InQueueOwnerReference, const FNetObjectReference& InTargetObjectReference)
 {
 	NetObjectReference = InQueueOwnerReference;
