Chaos Cloth: Add Self Collision Thickness Map

Thickness is baked into the Spatial hierarchy when using a map. Should be very little overhead when not using a map.

Tested with new and old style solver (with and without offsets).

Also includes: fix to an API change to GetPointNormals. Fix to a crash in EdgeIntersections detection when you find more than the preallocated number of intersections.

diff --git a/Engine/Plugins/Experimental/ChaosCloth/Source/Chaos/Private/ChaosCloth/ChaosClothConstraints.cpp b/Engine/Plugins/Experimental/ChaosCloth/Source/Chaos/Private/ChaosCloth/ChaosClothConstraints.cpp
--- a/Engine/Plugins/Experimental/ChaosCloth/Source/Chaos/Private/ChaosCloth/ChaosClothConstraints.cpp
+++ b/Engine/Plugins/Experimental/ChaosCloth/Source/Chaos/Private/ChaosCloth/ChaosClothConstraints.cpp
@@ -1693,17 +1693,17 @@ void FClothConstraints::CreateForceBasedRules(const TSharedPtr<Softs::FMultiResC
 		RuleCreator.AddPerIterationCollisionPBDConstraintRule_Apply(CollisionConstraint.Get());
 		RuleCreator.AddUpdateLinearSystemCollisionsRule_UpdateLinearSystem(CollisionConstraint.Get());
 	}
+
 	if (SelfCollisionInit && SelfCollisionConstraints)
 	{
 		RuleCreator.AddPostInitialGuessParallelInitRule(
-			[this](const Softs::FSolverParticlesRange& Particles, const Softs::FSolverReal Dt, const Softs::ESolverMode SolverMode)
+		[this](const Softs::FSolverParticlesRange& Particles, const Softs::FSolverReal Dt, const Softs::ESolverMode SolverMode)
 		{
 			if (bSkipSelfCollisionInit)
 			{
 				return;
 			}
-			// Thickness * 2 to account for collision radius for both particles
-			SelfCollisionInit->Init(Particles, SelfCollisionConstraints->GetThickness() * (Softs::FSolverReal)2.f);
+			SelfCollisionInit->Init(Particles, SelfCollisionConstraints->GetThicknessWeighted());
 			SelfCollisionConstraints->Init(Particles, Dt, SelfCollisionInit->GetCollidableSubMesh(), SelfCollisionInit->GetDynamicSpatialHash(), SelfCollisionInit->GetKinematicColliderSpatialHash(), SelfCollisionInit->GetVertexGIAColors(), SelfCollisionInit->GetTriangleGIAColors());
 		});
 
@@ -1711,6 +1711,7 @@ void FClothConstraints::CreateForceBasedRules(const TSharedPtr<Softs::FMultiResC
 		RuleCreator.AddPerIterationPBDConstraintRule_Apply<bPostCollisions>(SelfCollisionConstraints.Get());
 		RuleCreator.AddUpdateLinearSystemRule_UpdateLinearSystem(SelfCollisionConstraints.Get());
 	}
+
 	if (SelfCollisionSphereConstraints)
 	{
 		RuleCreator.AddPostInitialGuessParallelInitRule(
@@ -2109,7 +2110,7 @@ void FClothConstraints::CreatePBDRules()
 					return;
 				}
 				// Thickness * 2 to account for collision radius for both particles
-				SelfCollisionInit->Init(Particles, SelfCollisionConstraints->GetThickness() * (Softs::FSolverReal)2.f);
+				SelfCollisionInit->Init(Particles, SelfCollisionConstraints->GetThicknessWeighted());
 				SelfCollisionConstraints->Init(Particles, Dt, SelfCollisionInit->GetCollidableSubMesh(), SelfCollisionInit->GetDynamicSpatialHash(), SelfCollisionInit->GetKinematicColliderSpatialHash(), SelfCollisionInit->GetVertexGIAColors(), SelfCollisionInit->GetTriangleGIAColors());
 			};
 
diff --git a/Engine/Plugins/Experimental/ChaosCloth/Source/Chaos/Private/ChaosCloth/ChaosClothVisualization.cpp b/Engine/Plugins/Experimental/ChaosCloth/Source/Chaos/Private/ChaosCloth/ChaosClothVisualization.cpp
--- a/Engine/Plugins/Experimental/ChaosCloth/Source/Chaos/Private/ChaosCloth/ChaosClothVisualization.cpp
+++ b/Engine/Plugins/Experimental/ChaosCloth/Source/Chaos/Private/ChaosCloth/ChaosClothVisualization.cpp
@@ -555,21 +555,46 @@ FLinearColor PseudoRandomColor(int32 NumColorRotations)
 #endif
 	}
 
-	static void DrawTaperedCylinder(FPrimitiveDrawInterface* PDI, const ::Chaos::FTaperedCylinder& TaperedCylinder, const FQuat& Rotation, const FVector& Position, const FLinearColor& Color)
+#if WITH_EDITOR
+	static void AppendTaperedCylinderTriangles(FDynamicMeshBuilder& MeshBuilder, const FVector3f& Position1, const FVector3f& Position2, const FRealSingle Radius1, const FRealSingle Radius2, const int32 NumSides, const FLinearColor& Color)
+	{
+		const FQuat4f Q = (Position2 - Position1).ToOrientationQuat();
+		const FVector3f I = Q.GetRightVector();
+		const FVector3f J = Q.GetUpVector();
+		const FVector3f K = Q.GetForwardVector();
+
+		const FRealSingle AngleDelta = (FRealSingle)2. * (FRealSingle)PI / NumSides;
+		int32 LastVertex1 = MeshBuilder.AddVertex(FDynamicMeshVertex(Position1 + I * Radius1, -K, I, FVector2f(0.f, 0.f), Color.ToFColor(true)));
+		int32 LastVertex2 = MeshBuilder.AddVertex(FDynamicMeshVertex(Position2 + I * Radius2, -K, I, FVector2f(1.f, 0.f), Color.ToFColor(true))); 
+		for (int32 SideIndex = 1; SideIndex <= NumSides; ++SideIndex)
+		{
+			const FRealSingle Angle = AngleDelta * FRealSingle(SideIndex);
+			const FVector3f ArcPos = I * FMath::Cos(Angle) + J * FMath::Sin(Angle);
+			
+			const FVector3f Pos1 = Position1 + ArcPos * Radius1;
+			const FVector3f Pos2 = Position2 + ArcPos * Radius2;
+			const FVector3f Normal = (Pos1 - Position1).GetSafeNormal();
+
+			const int32 Vertex1 = MeshBuilder.AddVertex(FDynamicMeshVertex(Pos1, -K, Normal, FVector2f(0.f, 0.f), Color.ToFColor(true)));
+			const int32 Vertex2 = MeshBuilder.AddVertex(FDynamicMeshVertex(Pos2, -K, Normal, FVector2f(1.f, 0.f), Color.ToFColor(true)));
+			MeshBuilder.AddTriangle(LastVertex1, LastVertex2, Vertex1);
+			MeshBuilder.AddTriangle(LastVertex2, Vertex2, Vertex1);
+
+			LastVertex1 = Vertex1;
+			LastVertex2 = Vertex2;
+		}
+	}
+#endif
+
+	static void DrawTaperedCylinder(FPrimitiveDrawInterface* PDI, const FVector& Position1, const FVector& Position2, const FReal Radius1, const FReal Radius2, const int32 NumSides, const FLinearColor& Color)
 	{
-		const FReal HalfHeight = TaperedCylinder.GetHeight() * 0.5f;
-		const FReal Radius1 = TaperedCylinder.GetRadius1();
-		const FReal Radius2 = TaperedCylinder.GetRadius2();
-		const FVector Position1 = Position + Rotation.RotateVector(TaperedCylinder.GetX1());
-		const FVector Position2 = Position + Rotation.RotateVector(TaperedCylinder.GetX2());
 		const FQuat Q = (Position2 - Position1).ToOrientationQuat();
 		const FVector I = Q.GetRightVector();
 		const FVector J = Q.GetUpVector();
 
-		static const int32 NumSides = 12;
-		static const FReal	AngleDelta = (FReal)2. * (FReal)PI / NumSides;
-		FVector	LastVertex1 = Position1 + I * Radius1;
-		FVector	LastVertex2 = Position2 + I * Radius2;
+		const FReal	AngleDelta = (FReal)2. * (FReal)PI / NumSides;
+		FVector LastVertex1 = Position1 + I * Radius1;
+		FVector LastVertex2 = Position2 + I * Radius2;
 
 		for (int32 SideIndex = 1; SideIndex <= NumSides; ++SideIndex)
 		{
@@ -587,6 +612,15 @@ FLinearColor PseudoRandomColor(int32 NumColorRotations)
 		}
 	}
 
+	static void DrawTaperedCylinder(FPrimitiveDrawInterface* PDI, const ::Chaos::FTaperedCylinder& TaperedCylinder, const FQuat& Rotation, const FVector& Position, const FLinearColor& Color)
+	{
+		const FReal Radius1 = TaperedCylinder.GetRadius1();
+		const FReal Radius2 = TaperedCylinder.GetRadius2();
+		const FVector Position1 = Position + Rotation.RotateVector(TaperedCylinder.GetX1());
+		const FVector Position2 = Position + Rotation.RotateVector(TaperedCylinder.GetX2());
+		DrawTaperedCylinder(PDI, Position1, Position2, Radius1, Radius2, 12, Color);
+	}
+
 	static void DrawConvex(FPrimitiveDrawInterface* PDI, const ::Chaos::FConvex& Convex, const FQuat& Rotation, const FVector& Position, const FLinearColor& Color)
 	{
 		const TArray<FConvex::FPlaneType>& Planes = Convex.GetFaces();
@@ -2258,12 +2292,11 @@ FLinearColor PseudoRandomColor(int32 NumColorRotations)
 				const int32 Offset = Solver->IsForceBasedSolver() ? 0 : ParticleRangeId;
 				const TArray<TVec4<int32>>& Constraints = SelfCollisionConstraints->GetConstraints();
 				const TArray<Softs::FSolverVec3>& Barys = SelfCollisionConstraints->GetBarys();
-				const FReal Thickness = (FReal)SelfCollisionConstraints->GetThickness();
-				const FReal Height = Thickness + Thickness;
 				const TArray<bool>& FlipNormals = SelfCollisionConstraints->GetFlipNormals();
 
 				for (int32 Index = 0; Index < Constraints.Num(); ++Index)
 				{
+					const FReal Height = (FReal)SelfCollisionConstraints->GetConstraintThickness(Index);
 					const TVec4<int32>& Constraint = Constraints[Index];
 					const FVec3 Bary(Barys[Index]);
 
@@ -2498,6 +2531,122 @@ FLinearColor PseudoRandomColor(int32 NumColorRotations)
 		}
 	}
 
+	void FClothVisualization::DrawSelfCollisionThickness(FPrimitiveDrawInterface* PDI) const
+	{
+		if (!Solver)
+		{
+			return;
+		}
+
+		const FVec3& LocalSpaceLocation = Solver->GetLocalSpaceLocation();
+
+#if WITH_EDITOR
+		if (ClothMaterialColor)
+		{
+			FDynamicMeshBuilder MeshBuilder(PDI->View->GetFeatureLevel());
+
+			for (const FClothingSimulationCloth* const Cloth : Solver->GetCloths())
+			{
+				const int32 ParticleRangeId = Cloth->GetParticleRangeId(Solver);
+				if (ParticleRangeId == INDEX_NONE)
+				{
+					continue;
+				}
+
+				const FClothConstraints& ClothConstraints = Solver->GetClothConstraints(ParticleRangeId);
+
+				if (const Softs::FPBDCollisionSpringConstraints* const SelfCollisionConstraints = ClothConstraints.GetSelfCollisionConstraints().Get())
+				{
+					const int32 Offset = Solver->IsForceBasedSolver() ? 0 : ParticleRangeId;
+					const TConstArrayView<Softs::FSolverVec3> Positions = Cloth->GetParticlePositions(Solver);
+					const TConstArrayView<Softs::FSolverReal> InvMasses = Cloth->GetParticleInvMasses(Solver);
+					const TConstArrayView<int32>& WeightMap = Cloth->GetFaceIntMapByProperty(Solver, TEXT("SelfCollisionLayers"));
+
+					const TConstArrayView<TVec2<int32>> Edges = Cloth->GetTriangleMesh(Solver).GetSegmentMesh().GetElements();
+					const TArray<TVec2<int32>>& EdgeToFaces = Cloth->GetTriangleMesh(Solver).GetEdgeToFaces();
+					for(int32 EdgeIndex = 0; EdgeIndex < Edges.Num(); ++EdgeIndex)
+					{
+						const TVec2<int32>& Edge = Edges[EdgeIndex];
+					
+						const bool bIsKinematic0 = (InvMasses[Edge[0] - Offset] == (Softs::FSolverReal)0.);
+						const bool bIsKinematic1 = (InvMasses[Edge[1] - Offset] == (Softs::FSolverReal)0.);
+						if (bIsKinematic0 && bIsKinematic1)
+						{
+							continue;
+						}
+
+						const FVector3f Position1(Positions[Edge[0] - Offset]); // TODO: Triangle Mesh shouldn't really be solver dependent (ie not use an offset)
+						const FVector3f Position2(Positions[Edge[1] - Offset]);
+
+						const FRealSingle Radius1 = (FRealSingle)SelfCollisionConstraints->GetParticleThickness(Edge[0]);
+						const FRealSingle Radius2 = (FRealSingle)SelfCollisionConstraints->GetParticleThickness(Edge[1]);
+						const int32 Face1Layer = WeightMap.IsValidIndex(EdgeToFaces[EdgeIndex][0]) ? WeightMap[EdgeToFaces[EdgeIndex][0]] : INDEX_NONE;
+						const int32 Face2Layer = WeightMap.IsValidIndex(EdgeToFaces[EdgeIndex][1]) ? WeightMap[EdgeToFaces[EdgeIndex][1]] : INDEX_NONE;
+						const FLinearColor Color = Chaos::Private::PseudoRandomColor((Face1Layer == Face2Layer || Face2Layer == INDEX_NONE) ? Face1Layer : INDEX_NONE);
+						
+						AppendTaperedCylinderTriangles(MeshBuilder, Position1, Position2, Radius1, Radius2, 6, Color);
+					}
+				}
+			}
+			FMatrix LocalSimSpaceToWorld(FMatrix::Identity);
+			LocalSimSpaceToWorld.SetOrigin(Solver->GetLocalSpaceLocation());
+			MeshBuilder.Draw(PDI, LocalSimSpaceToWorld, ClothMaterialColor->GetRenderProxy(), SDPG_World, false, false);
+		}
+		else
+#endif
+		{
+			for (const FClothingSimulationCloth* const Cloth : Solver->GetCloths())
+			{
+				const int32 ParticleRangeId = Cloth->GetParticleRangeId(Solver);
+				if (ParticleRangeId == INDEX_NONE)
+				{
+					continue;
+				}
+
+				const FClothConstraints& ClothConstraints = Solver->GetClothConstraints(ParticleRangeId);
+
+				if (const Softs::FPBDCollisionSpringConstraints* const SelfCollisionConstraints = ClothConstraints.GetSelfCollisionConstraints().Get())
+				{
+					const int32 Offset = Solver->IsForceBasedSolver() ? 0 : ParticleRangeId;
+					const TConstArrayView<Softs::FSolverVec3> Positions = Cloth->GetParticlePositions(Solver);
+					const TConstArrayView<Softs::FSolverReal> InvMasses = Cloth->GetParticleInvMasses(Solver);
+
+					const TConstArrayView<TVec2<int32>> Edges = Cloth->GetTriangleMesh(Solver).GetSegmentMesh().GetElements();
+					for (const TVec2<int32>& Edge : Edges)
+					{
+						const bool bIsKinematic0 = (InvMasses[Edge[0] - Offset] == (Softs::FSolverReal)0.);
+						const bool bIsKinematic1 = (InvMasses[Edge[1] - Offset] == (Softs::FSolverReal)0.);
+						if (bIsKinematic0 && bIsKinematic1)
+						{
+							continue;
+						}
+
+						const FVector Position1(Positions[Edge[0] - Offset]); // TODO: Triangle Mesh shouldn't really be solver dependent (ie not use an offset)
+						const FVector Position2(Positions[Edge[1] - Offset]);
+
+						const FReal Radius1 = (FReal)SelfCollisionConstraints->GetParticleThickness(Edge[0]);
+						const FReal Radius2 = (FReal)SelfCollisionConstraints->GetParticleThickness(Edge[1]);
+						DrawTaperedCylinder(PDI, Position1 + LocalSpaceLocation, Position2 + LocalSpaceLocation, Radius1, Radius2, 6, FLinearColor::Gray);
+					}
+
+					for (int32 VertexIndex = 0; VertexIndex < Positions.Num(); ++VertexIndex)
+					{
+						const bool bIsKinematic0 = (InvMasses[VertexIndex] == (Softs::FSolverReal)0.);
+						if (bIsKinematic0)
+						{
+							continue;
+						}
+
+						const FVector Position1(Positions[VertexIndex]);
+						const FReal Radius1 = (FReal)SelfCollisionConstraints->GetParticleThickness(VertexIndex + Offset);
+						const FTransform Transform(Position1 + LocalSpaceLocation);
+						DrawWireSphere(PDI, Transform, FLinearColor::Gray, Radius1, 6, SDPG_World, 0.0f, 0.001f, false);
+					}
+				}
+			}
+		}
+	}
+
 	void FClothVisualization::DrawKinematicColliderWired(FPrimitiveDrawInterface* PDI) const
 	{
 		if (!Solver)
diff --git a/Engine/Plugins/Experimental/ChaosCloth/Source/Chaos/Public/ChaosCloth/ChaosClothVisualization.h b/Engine/Plugins/Experimental/ChaosCloth/Source/Chaos/Public/ChaosCloth/ChaosClothVisualization.h
--- a/Engine/Plugins/Experimental/ChaosCloth/Source/Chaos/Public/ChaosCloth/ChaosClothVisualization.h
+++ b/Engine/Plugins/Experimental/ChaosCloth/Source/Chaos/Public/ChaosCloth/ChaosClothVisualization.h
@@ -49,6 +49,7 @@ namespace Chaos
 		CHAOSCLOTH_API void DrawLocalSpace(FPrimitiveDrawInterface* PDI = nullptr) const;
 		CHAOSCLOTH_API void DrawSelfCollision(FPrimitiveDrawInterface* PDI = nullptr) const;
 		CHAOSCLOTH_API void DrawSelfIntersection(FPrimitiveDrawInterface* PDI = nullptr) const;
+		CHAOSCLOTH_API void DrawSelfCollisionThickness(FPrimitiveDrawInterface* PDI = nullptr) const;
 		CHAOSCLOTH_API void DrawKinematicColliderWired(FPrimitiveDrawInterface* PDI = nullptr) const;
 		CHAOSCLOTH_API void DrawBounds(FPrimitiveDrawInterface* PDI = nullptr) const;
 		CHAOSCLOTH_API void DrawGravity(FPrimitiveDrawInterface* PDI = nullptr) const;
@@ -76,7 +77,8 @@ namespace Chaos
 		void DrawLocalSpace(FPrimitiveDrawInterface* /*PDI*/ = nullptr) const {}
 		void DrawSelfCollision(FPrimitiveDrawInterface* /*PDI*/ = nullptr) const {}
 		void DrawSelfIntersection(FPrimitiveDrawInterface* /*PDI*/ = nullptr) const {}
-		CHAOSCLOTH_API void DrawKinematicColliderWired(FPrimitiveDrawInterface* /*PDI*/ = nullptr) const {}
+		void DrawSelfCollisionThickness(FPrimitiveDrawInterface* /*PDI*/ = nullptr) const {}
+		void DrawKinematicColliderWired(FPrimitiveDrawInterface* /*PDI*/ = nullptr) const {}
 		void DrawBounds(FPrimitiveDrawInterface* /*PDI*/ = nullptr) const {}
 		void DrawGravity(FPrimitiveDrawInterface* /*PDI*/ = nullptr) const {}
 		void DrawFictitiousAngularForces(FPrimitiveDrawInterface* /*PDI*/ = nullptr) const {}
diff --git a/Engine/Plugins/Experimental/ChaosClothAssetEditor/Source/ChaosClothAssetDataflowNodes/Private/ChaosClothAsset/SimulationSelfCollisionConfigNode.cpp b/Engine/Plugins/Experimental/ChaosClothAssetEditor/Source/ChaosClothAssetDataflowNodes/Private/ChaosClothAsset/SimulationSelfCollisionConfigNode.cpp
--- a/Engine/Plugins/Experimental/ChaosClothAssetEditor/Source/ChaosClothAssetDataflowNodes/Private/ChaosClothAsset/SimulationSelfCollisionConfigNode.cpp
+++ b/Engine/Plugins/Experimental/ChaosClothAssetEditor/Source/ChaosClothAssetDataflowNodes/Private/ChaosClothAsset/SimulationSelfCollisionConfigNode.cpp
@@ -14,12 +14,13 @@ FChaosClothAssetSimulationSelfCollisionConfigNode::FChaosClothAssetSimulationSel
 	RegisterInputConnection(&SelfCollisionDisabledFaces.StringValue, GET_MEMBER_NAME_CHECKED(FChaosClothAssetConnectableIStringValue, StringValue));
 	RegisterInputConnection(&SelfCollisionEnabledKinematicFaces.StringValue, GET_MEMBER_NAME_CHECKED(FChaosClothAssetConnectableIStringValue, StringValue));
 	RegisterInputConnection(&SelfCollisionKinematicColliderFrictionWeighted.WeightMap);
+	RegisterInputConnection(&SelfCollisionThicknessWeighted.WeightMap);
 }
 
 void FChaosClothAssetSimulationSelfCollisionConfigNode::AddProperties(FPropertyHelper& PropertyHelper) const
 {
 	PropertyHelper.SetPropertyBool(FName("UseSelfCollisions"), true);
-	PropertyHelper.SetProperty(this, &SelfCollisionThickness);
+	PropertyHelper.SetPropertyWeighted(TEXT("SelfCollisionThickness"), SelfCollisionThicknessWeighted);
 	PropertyHelper.SetProperty(this, &SelfCollisionStiffness);
 	PropertyHelper.SetProperty(this, &SelfCollisionFriction);
 	PropertyHelper.SetProperty(this, &SelfCollisionDisableNeighborDistance, {}, ECollectionPropertyFlags::None); // Non animatable
@@ -51,6 +52,11 @@ void FChaosClothAssetSimulationSelfCollisionConfigNode::Serialize(FArchive& Ar)
 			SelfCollisionKinematicColliderFrictionWeighted.Low = SelfCollisionKinematicColliderFrictionWeighted.High = SelfCollisionKinematicColliderFriction_DEPRECATED;
 			SelfCollisionKinematicColliderFriction_DEPRECATED = FrictionDeprecatedValue;
 		}
+		if (SelfCollisionThickness_DEPRECATED != SelfCollisionThicknessDeprecatedValue)
+		{
+			SelfCollisionThicknessWeighted.Low = SelfCollisionThicknessWeighted.High = SelfCollisionThickness_DEPRECATED;
+			SelfCollisionThickness_DEPRECATED = SelfCollisionThicknessDeprecatedValue;
+		}
 #endif
 	}
 }
\ No newline at end of file
diff --git a/Engine/Plugins/Experimental/ChaosClothAssetEditor/Source/ChaosClothAssetDataflowNodes/Public/ChaosClothAsset/SimulationSelfCollisionConfigNode.h b/Engine/Plugins/Experimental/ChaosClothAssetEditor/Source/ChaosClothAssetDataflowNodes/Public/ChaosClothAsset/SimulationSelfCollisionConfigNode.h
--- a/Engine/Plugins/Experimental/ChaosClothAssetEditor/Source/ChaosClothAssetDataflowNodes/Public/ChaosClothAsset/SimulationSelfCollisionConfigNode.h
+++ b/Engine/Plugins/Experimental/ChaosClothAssetEditor/Source/ChaosClothAssetDataflowNodes/Public/ChaosClothAsset/SimulationSelfCollisionConfigNode.h
@@ -16,7 +16,7 @@ struct FChaosClothAssetSimulationSelfCollisionConfigNode : public FChaosClothAss
 public:
 	/** The self collision offset per side. Total thickness of cloth is 2x this value. */
 	UPROPERTY(EditAnywhere, Category = "Self-Collision Properties", meta = (UIMin = "0", UIMax = "100", ClampMin = "0", ClampMax = "1000", EditCondition = "bUseSelfCollisions"))
-	float SelfCollisionThickness = 0.5f;
+	FChaosClothAssetWeightedValue SelfCollisionThicknessWeighted = {true, 0.5f, 0.5f, TEXT("SelfCollisionThickness")};
 
 	/** The stiffness of the springs used to control self collision (PBD Solver). */
 	UPROPERTY(EditAnywhere, Category = "Self-Collision Properties", meta = (UIMin = "0", UIMax = "1", ClampMin = "0", ClampMax = "1", EditCondition = "!bSelfCollideAgainstKinematicCollidersOnly"))
@@ -98,6 +98,10 @@ private:
 	static constexpr float FrictionDeprecatedValue = -1.f;
 	UPROPERTY()
 	float SelfCollisionKinematicColliderFriction_DEPRECATED = FrictionDeprecatedValue;
+
+	static constexpr float SelfCollisionThicknessDeprecatedValue = -1.f;
+	UPROPERTY()
+	float SelfCollisionThickness_DEPRECATED = SelfCollisionThicknessDeprecatedValue;
 #endif
 	virtual void AddProperties(FPropertyHelper& PropertyHelper) const override;
 };
diff --git a/Engine/Plugins/Experimental/ChaosClothAssetEditor/Source/ChaosClothAssetEditor/Private/ChaosClothAsset/ClothEditorSimulationVisualization.cpp b/Engine/Plugins/Experimental/ChaosClothAssetEditor/Source/ChaosClothAssetEditor/Private/ChaosClothAsset/ClothEditorSimulationVisualization.cpp
--- a/Engine/Plugins/Experimental/ChaosClothAssetEditor/Source/ChaosClothAssetEditor/Private/ChaosClothAsset/ClothEditorSimulationVisualization.cpp
+++ b/Engine/Plugins/Experimental/ChaosClothAssetEditor/Source/ChaosClothAssetEditor/Private/ChaosClothAsset/ClothEditorSimulationVisualization.cpp
@@ -170,6 +170,7 @@ const FVisualizationOption FVisualizationOption::OptionData[] =
 	FVisualizationOption(&::Chaos::FClothVisualization::DrawSelfCollision, TEXT("p.ChaosClothAssetEditor.DebugDrawSelfCollision"), LOCTEXT("ChaosVisName_SelfCollision"       , "Self Collision"), LOCTEXT("ChaosVisName_SelfCollision_Tooltip", "Draws the self collision thickness/debugging information")),
 	FVisualizationOption(&::Chaos::FClothVisualization::DrawSelfIntersection, TEXT("p.ChaosClothAssetEditor.DebugDrawSelfIntersection"), LOCTEXT("ChaosVisName_SelfIntersection", "Self Intersection"), LOCTEXT("ChaosVisName_SelfIntersection_Tooltip", "Draws the self intersection contour/region information")),
 	FVisualizationOption(&::Chaos::FClothVisualization::DrawSelfCollisionLayers, TEXT("p.ChaosClothAssetEditor.DebugDrawSelfCollisionLayers"), LOCTEXT("ChaosVisName_SelfCollisionLayers", "Self Collision Layers"), LOCTEXT("ChaosVisName_SelfCollisionLayers_Tooltip", "Draws the self collision layers"), /*bDisablesSimulation =*/false, /*bHidesClothSections=*/true),
+	FVisualizationOption(&::Chaos::FClothVisualization::DrawSelfCollisionThickness, TEXT("p.ChaosClothAssetEditor.DebugDrawSelfCollisionThickness"), LOCTEXT("ChaosVisName_SelfCollisionThickness", "Self Collision Thickness"), LOCTEXT("ChaosVisName_SelfCollisionThickness_Tooltip", "Draws the self collision Thickness")),
 	FVisualizationOption(&::Chaos::FClothVisualization::DrawKinematicColliderShaded, TEXT("p.ChaosClothAssetEditor.DebugDrawKinematicColliderShaded"), LOCTEXT("ChaosVisName_DrawKinematicColliderShaded", "Draw Kinematic Colliders (Shaded)"), LOCTEXT("ChaosVisName_DrawKinematicColliderShaded_Tooltip", "Draw kinematic cloth colliders with flat shading.")),
 	FVisualizationOption(&::Chaos::FClothVisualization::DrawKinematicColliderWired, TEXT("p.ChaosClothAssetEditor.DebugDrawKinematicColliderWired"), LOCTEXT("ChaosVisName_DrawKinematicColliderWired", "Draw Kinematic Colliders (Wireframe)"), LOCTEXT("ChaosVisName_DrawKinematicColliderWired_Tooltip", "Draw kinematic cloth colliders in wireframe.")),
 	FVisualizationOption(&::Chaos::FClothVisualization::DrawMultiResConstraint, TEXT("p.ChaosClothAssetEditor.DebugDrawMultiResConstraint"), LOCTEXT("ChaosVisName_DrawMultiResConstraint", "Draw Multi Res Constraint"), LOCTEXT("ChaosVisName_DrawMultiResConstraint_Tooltip", "Draw multi res constraint coarse mesh and targets.")),
diff --git a/Engine/Source/Runtime/Experimental/Chaos/Private/Chaos/PBDCollisionSpringConstraints.ispc b/Engine/Source/Runtime/Experimental/Chaos/Private/Chaos/PBDCollisionSpringConstraints.ispc
--- a/Engine/Source/Runtime/Experimental/Chaos/Private/Chaos/PBDCollisionSpringConstraints.ispc
+++ b/Engine/Source/Runtime/Experimental/Chaos/Private/Chaos/PBDCollisionSpringConstraints.ispc
@@ -270,7 +270,7 @@ export void ApplyKinematicCollisionSpringConstraints(
 	}
 }
 
-export void ApplyKinematicCollisionSpringConstraintsWithMaps(
+export void ApplyKinematicCollisionSpringConstraintsWithKinematicFrictionMap(
 			uniform FVector4f PandInvM[],
 			const uniform FVector3f ParticlesX[],
 			const uniform int32 KinematicCollidingParticles[],
@@ -282,6 +282,7 @@ export void ApplyKinematicCollisionSpringConstraintsWithMaps(
 			const uniform float KinematicColliderStiffness, 
 			const uniform FVector2f& KinematicColliderFrictionCoefficientOffsetRange, 
 			const uniform float KinematicColliderFrictionCoefficientMapValues[],
+			const uniform int32 ParticleOffset,
 			const uniform int32 NumConstraints)
 {
 	foreach(Index = 0...NumConstraints)
@@ -296,7 +297,7 @@ export void ApplyKinematicCollisionSpringConstraintsWithMaps(
 		const varying FIntVector Elements = VectorLoad(&KinematicColliderElements[extract(Index,0)]);
 		
 		#pragma ignore warning(perf)
-		const varying float FrictionMapValue = KinematicColliderFrictionCoefficientMapValues[Index1];
+		const varying float FrictionMapValue = KinematicColliderFrictionCoefficientMapValues[Index1 - ParticleOffset];
 		const varying float FrictionCoefficient = KinematicColliderFrictionCoefficientOffsetRange.V[0] + FrictionMapValue * KinematicColliderFrictionCoefficientOffsetRange.V[1];
 
 		for(uniform int32 EIndex = 0; EIndex < 3; ++EIndex)
@@ -369,4 +370,120 @@ export void ApplyKinematicCollisionSpringConstraintsWithMaps(
 
 		VectorScatter(&PandInvM[Index1], SetVector4(P1, IM));
 	}
-}
\ No newline at end of file
+}
+
+export void ApplyKinematicCollisionSpringConstraintsWithMaps(
+			uniform FVector4f PandInvM[],
+			const uniform FVector3f ParticlesX[],
+			const uniform int32 KinematicCollidingParticles[],
+			const uniform FIntVector KinematicColliderElements[],
+			const uniform FIntVector TriangleMeshElements[],
+			const uniform FVector2f& ParticleThicknessOffsetRange,
+			const uniform float ParticleThicknessMapValues[],
+			const uniform float KinematicColliderThickness,
+			const uniform float KinematicColliderFalloffMultiplier,
+			const uniform float KinematicColliderMaxDepthMultiplier,
+			const uniform float KinematicColliderStiffness, 
+			const uniform bool bKinematicColliderFrictionHasMap,
+			const uniform FVector2f& KinematicColliderFrictionCoefficientOffsetRange, 
+			const uniform float KinematicColliderFrictionCoefficientMapValues[],
+			const uniform int32 ParticleOffset,
+			const uniform int32 NumConstraints)
+{
+	foreach(Index = 0...NumConstraints)
+	{
+		const int32 Index1 = KinematicCollidingParticles[Index];
+
+		const FVector4f PandInvM1 = VectorGather(&PandInvM[Index1]);
+		varying FVector3f P1;
+		varying float IM;
+		UnzipPandInvM(PandInvM1, P1, IM);
+
+		const varying FIntVector Elements = VectorLoad(&KinematicColliderElements[extract(Index,0)]);
+		
+		#pragma ignore warning(perf)
+		const varying float ThicknessMapValue = ParticleThicknessMapValues[Index1 - ParticleOffset];
+		const varying float ParticleThickness = ParticleThicknessOffsetRange.V[0] + ThicknessMapValue * ParticleThicknessOffsetRange.V[1];
+		const float Height = ParticleThickness + KinematicColliderThickness;
+		const float OneOverTangentialFalloffDist = 1.f / max(KinematicColliderFalloffMultiplier * Height, FLOAT_KINDA_SMALL_NUMBER);
+		const float MaxDepth = -Height * KinematicColliderMaxDepthMultiplier;
+
+		varying float FrictionCoefficient = KinematicColliderFrictionCoefficientOffsetRange.V[0];
+		if(bKinematicColliderFrictionHasMap)
+		{
+			#pragma ignore warning(perf)
+			const varying float FrictionMapValue = KinematicColliderFrictionCoefficientMapValues[Index1 - ParticleOffset];
+			FrictionCoefficient += FrictionMapValue * KinematicColliderFrictionCoefficientOffsetRange.V[1];
+		}
+
+		for(uniform int32 EIndex = 0; EIndex < 3; ++EIndex)
+		{
+			const varying int32 ElementIndex = Elements.V[EIndex];
+			if(ElementIndex == -1)
+			{
+				break;
+			}
+
+			const FIntVector Element = VectorGather(&TriangleMeshElements[ElementIndex]);
+
+			const varying int32 Index2 = Element.V[0];
+			const varying int32 Index3 = Element.V[1];
+			const varying int32 Index4 = Element.V[2];
+		
+			const varying FVector4f PandInvM2 = VectorGather(&PandInvM[Index2]);
+			const varying FVector4f PandInvM3 = VectorGather(&PandInvM[Index3]);
+			const varying FVector4f PandInvM4 = VectorGather(&PandInvM[Index4]);
+
+			varying FVector3f P2, P3, P4;
+			varying float IM2, IM3, IM4;
+		
+			UnzipPandInvM(PandInvM2, P2, IM2);
+			UnzipPandInvM(PandInvM3, P3, IM3);
+			UnzipPandInvM(PandInvM4, P4, IM4);
+
+			const varying FVector3f Normal = GetTriangleNormalFlipped(P2, P3, P4);
+
+			varying FVector3f Bary;
+			const varying FVector3f P = FindClosestPointAndBaryOnTriangle(P2, P3, P4, P1, Bary);
+			const varying FVector3f Difference = P1 - P;
+			const varying float NormalDifference = VectorDot(Difference, Normal);
+			if(NormalDifference >= Height || NormalDifference < MaxDepth)
+			{
+				continue;
+			}
+
+			const varying float TangentialDifference = VectorSize(Difference - NormalDifference * Normal);
+			const uniform float TangentialFalloffMultiplier = 1.5f;
+			const varying float TangentialFalloff = 1.f - TangentialDifference * OneOverTangentialFalloffDist;
+			if(TangentialFalloff <= 0.f)
+			{
+				continue;
+			}
+
+			const float NormalDelta = Height - NormalDifference;
+			const FVector3f RepulsionDelta = KinematicColliderStiffness * TangentialFalloff * NormalDelta *  Normal;
+
+			P1 = P1 + RepulsionDelta;
+
+
+			if (FrictionCoefficient > 0)
+			{
+				const FVector3f X1 = VectorGather(&ParticlesX[Index1]);
+				const FVector3f X2 = VectorGather(&ParticlesX[Index2]);
+				const FVector3f X3 = VectorGather(&ParticlesX[Index3]);
+				const FVector3f X4 = VectorGather(&ParticlesX[Index4]);
+				const FVector3f X = Bary.V[0] * X2 + Bary.V[1] * X3 + Bary.V[2] * X4;
+				const FVector3f RelativeDisplacement = (P1 - X1) - (P - X);
+				const FVector3f RelativeDisplacementTangent = RelativeDisplacement - VectorDot(RelativeDisplacement,Normal) * Normal;
+				const float RelativeDisplacementTangentLength = VectorSize(RelativeDisplacementTangent);
+				const float PositionCorrection = min(NormalDelta * FrictionCoefficient, RelativeDisplacementTangentLength);
+				const float CorrectionRatio = RelativeDisplacementTangentLength < FLOAT_SMALL_NUMBER ? 0.f : PositionCorrection / RelativeDisplacementTangentLength;
+				const FVector3f FrictionDelta = -CorrectionRatio * RelativeDisplacementTangent;
+
+				P1 = P1 + FrictionDelta;
+			}
+		}
+
+		VectorScatter(&PandInvM[Index1], SetVector4(P1, IM));
+	}
+}
diff --git a/Engine/Source/Runtime/Experimental/Chaos/Private/Chaos/PBDCollisionSpringConstraintsBase.cpp b/Engine/Source/Runtime/Experimental/Chaos/Private/Chaos/PBDCollisionSpringConstraintsBase.cpp
--- a/Engine/Source/Runtime/Experimental/Chaos/Private/Chaos/PBDCollisionSpringConstraintsBase.cpp
+++ b/Engine/Source/Runtime/Experimental/Chaos/Private/Chaos/PBDCollisionSpringConstraintsBase.cpp
@@ -65,9 +65,10 @@ FPBDCollisionSpringConstraintsBase::FPBDCollisionSpringConstraintsBase(
 	const FTriangleMesh& InTriangleMesh,
 	const TArray<FSolverVec3>* InReferencePositions,
 	TSet<TVec2<int32>>&& InDisabledCollisionElements,
+	const TConstArrayView<FRealSingle>& InThicknessMultipliers,
 	const TConstArrayView<FRealSingle>& InKinematicColliderFrictionMultipliers,
 	const TConstArrayView<int32>& InSelfCollisionLayers,
-	const FSolverReal InThickness,
+	const FSolverVec2 InThickness,
 	const FSolverReal InStiffness,
 	const FSolverReal InFrictionCoefficient,
 	const bool bInOnlyCollideKinematics,
@@ -75,7 +76,7 @@ FPBDCollisionSpringConstraintsBase::FPBDCollisionSpringConstraintsBase(
 	const FSolverReal InKinematicColliderStiffness,
 	const FSolverVec2 InKinematicColliderFrictionCoefficient,
 	const FSolverReal InProximityStiffness)
-	: Thickness(InThickness)
+	: ThicknessWeighted(InThickness, InThicknessMultipliers, InNumParticles)
 	, Stiffness(InStiffness)
 	, FrictionCoefficient(InFrictionCoefficient)
 	, bOnlyCollideKinematics(bInOnlyCollideKinematics)
@@ -90,6 +91,9 @@ FPBDCollisionSpringConstraintsBase::FPBDCollisionSpringConstraintsBase(
 	, NumParticles(InNumParticles)
 	, bGlobalIntersectionAnalysis(false)
 {
+	PRAGMA_DISABLE_DEPRECATION_WARNINGS
+	Thickness = GetMaxThickness();
+	PRAGMA_ENABLE_DEPRECATION_WARNINGS
 	UpdateCollisionLayers(InSelfCollisionLayers);
 }
 
@@ -202,57 +206,78 @@ void FPBDCollisionSpringConstraintsBase::Init(const SolverParticlesOrRange& Part
 				const bool bVertexHasCollisionLayers = VertexCollisionLayers.IsValidIndex(i) && VertexCollisionLayers[i][0] != INDEX_NONE;
 				check(!bVertexHasCollisionLayers || VertexCollisionLayers[i][0] <= VertexCollisionLayers[i][1]);
 
+				const FSolverReal ParticleThickness = ThicknessWeighted.GetValue(i);
+
 				// Dynamic collisions
 				TArray< TTriangleCollisionPoint<FSolverReal> > DynamicResult;
-				if (!bOnlyCollideKinematics && CollidableSubMesh.GetDynamicSubMesh().PointProximityQuery(DynamicSpatial, static_cast<const TArrayView<const FSolverVec3>&>(Particles.XArray()), Index, Particles.X(Index), Thickness * ExtraThicknessMult, Thickness * ExtraThicknessMult,
-					[this, bVertexHasCollisionLayers, &Particles, &CollidableSubMesh, &VertexGIAColors, &TriangleGIAColors](const int32 PointIndex, const int32 SubMeshTriangleIndex)->bool
-					{
-						const TVector<int32, 3>& Elem = CollidableSubMesh.GetDynamicSubMesh().GetElements()[SubMeshTriangleIndex];
-						const int32 FullMeshTriangleIndex = CollidableSubMesh.GetFullMeshElementIndexFromDynamicElement(SubMeshTriangleIndex);
 
-						bool bUseCollisionLayerOverride = false;
-						if (bVertexHasCollisionLayers && FaceCollisionLayers[FullMeshTriangleIndex] != INDEX_NONE)
+				auto BroadphaseTest = [this, bVertexHasCollisionLayers, &Particles, &CollidableSubMesh, &VertexGIAColors, &TriangleGIAColors](const int32 PointIndex, const int32 SubMeshTriangleIndex)->bool
+				{
+					const TVector<int32, 3>& Elem = CollidableSubMesh.GetDynamicSubMesh().GetElements()[SubMeshTriangleIndex];
+					const int32 FullMeshTriangleIndex = CollidableSubMesh.GetFullMeshElementIndexFromDynamicElement(SubMeshTriangleIndex);
+
+					bool bUseCollisionLayerOverride = false;
+					if (bVertexHasCollisionLayers && FaceCollisionLayers[FullMeshTriangleIndex] != INDEX_NONE)
+					{
+						if (FaceCollisionLayers[FullMeshTriangleIndex] < VertexCollisionLayers[PointIndex - Offset][0] || FaceCollisionLayers[FullMeshTriangleIndex] > VertexCollisionLayers[PointIndex - Offset][1])
 						{
-							if (FaceCollisionLayers[FullMeshTriangleIndex] < VertexCollisionLayers[PointIndex - Offset][0] || FaceCollisionLayers[FullMeshTriangleIndex] > VertexCollisionLayers[PointIndex - Offset][1])
-							{
-								bUseCollisionLayerOverride = true;
-							}
+							bUseCollisionLayerOverride = true;
 						}
-						if (!bUseCollisionLayerOverride && bGlobalIntersectionAnalysis)
+					}
+					if (!bUseCollisionLayerOverride && bGlobalIntersectionAnalysis)
+					{
+						const bool bIsAnyBoundary = VertexGIAColors[PointIndex].IsBoundary()
+							|| VertexGIAColors[Elem[0]].IsBoundary()
+							|| VertexGIAColors[Elem[1]].IsBoundary()
+							|| VertexGIAColors[Elem[2]].IsBoundary();
+						if (bIsAnyBoundary)
 						{
-							const bool bIsAnyBoundary = VertexGIAColors[PointIndex].IsBoundary()
-								|| VertexGIAColors[Elem[0]].IsBoundary() 
-								|| VertexGIAColors[Elem[1]].IsBoundary() 
-								|| VertexGIAColors[Elem[2]].IsBoundary();
-							if (bIsAnyBoundary)
-							{
-								return false;
-							}
-
-							const bool bAreBothLoop = VertexGIAColors[PointIndex].IsLoop() &&
-								(VertexGIAColors[Elem[0]].IsLoop() 
-									|| VertexGIAColors[Elem[1]].IsLoop() 
-									|| VertexGIAColors[Elem[2]].IsLoop() 
-									|| TriangleGIAColors[SubMeshTriangleIndex].IsLoop());
-
-							if (bAreBothLoop)
-							{
-								return false;
-							}
+							return false;
 						}
 
-						if (DisabledCollisionElements.Contains({ PointIndex, Elem[0] }) ||
-							DisabledCollisionElements.Contains({ PointIndex, Elem[1] }) ||
-							DisabledCollisionElements.Contains({ PointIndex, Elem[2] }))
+						const bool bAreBothLoop = VertexGIAColors[PointIndex].IsLoop() &&
+							(VertexGIAColors[Elem[0]].IsLoop()
+								|| VertexGIAColors[Elem[1]].IsLoop()
+								|| VertexGIAColors[Elem[2]].IsLoop()
+								|| TriangleGIAColors[SubMeshTriangleIndex].IsLoop());
+
+						if (bAreBothLoop)
 						{
 							return false;
 						}
+					}
 
-						return true;
-					},
-					DynamicResult))
+					if (DisabledCollisionElements.Contains({ PointIndex, Elem[0] }) ||
+						DisabledCollisionElements.Contains({ PointIndex, Elem[1] }) ||
+						DisabledCollisionElements.Contains({ PointIndex, Elem[2] }))
+					{
+						return false;
+					}
+
+					return true;
+				};
+
+				if (!bOnlyCollideKinematics)
 				{
+					if (ThicknessWeighted.HasWeightMap())
+					{
+						if constexpr (std::is_same<SpatialAccelerator, THierarchicalSpatialHash<int32, FSolverReal>>::value)
+						{
+							CollidableSubMesh.GetDynamicSubMesh().PointProximityQuery(DynamicSpatial, static_cast<const TArrayView<const FSolverVec3>&>(Particles.XArray()), Index, Particles.X(Index), ParticleThickness * ExtraThicknessMult, ThicknessWeighted, ExtraThicknessMult, Offset, BroadphaseTest, DynamicResult);
+						}
+						else
+						{
+							check(false);
+						}
+					}
+					else
+					{
+						CollidableSubMesh.GetDynamicSubMesh().PointProximityQuery(DynamicSpatial, static_cast<const TArrayView<const FSolverVec3>&>(Particles.XArray()), Index, Particles.X(Index), ParticleThickness * ExtraThicknessMult, (FSolverReal)ThicknessWeighted * ExtraThicknessMult, BroadphaseTest, DynamicResult);
+					}
+				}
 
+				if ( DynamicResult.Num() )
+				{
 					if (DynamicResult.Num() > MaxConnectionsPerPoint)
 					{
 						// TODO: once we have a PartialSort, use that instead here.
@@ -275,7 +300,10 @@ void FPBDCollisionSpringConstraintsBase::Init(const SolverParticlesOrRange& Part
 							const FSolverVec3& RefP1 = ReferencePositionsView[Elem[1]];
 							const FSolverVec3& RefP2 = ReferencePositionsView[Elem[2]];
 							const FSolverVec3 RefDiff = RefP - CollisionPoint.Bary[1] * RefP0 - CollisionPoint.Bary[2] * RefP1 - CollisionPoint.Bary[3] * RefP2;
-							if (RefDiff.SizeSquared() < FMath::Square(Thickness))
+							const FSolverReal TriangleThickness = ThicknessWeighted.HasWeightMap() ?
+								CollisionPoint.Bary[1] * ThicknessWeighted.GetValue(Elem[0] - Offset) + CollisionPoint.Bary[2] * ThicknessWeighted.GetValue(Elem[1] - Offset) + CollisionPoint.Bary[3] * ThicknessWeighted.GetValue(Elem[2] - Offset) : (FSolverReal)ThicknessWeighted;
+
+							if (RefDiff.SizeSquared() < FMath::Square(ParticleThickness + TriangleThickness))
 							{
 								continue;
 							}
@@ -350,7 +378,7 @@ void FPBDCollisionSpringConstraintsBase::Init(const SolverParticlesOrRange& Part
 				if(CollidableSubMesh.GetKinematicColliderSubMesh().GetNumElements() > 0)
 				{
 					TArray< TTriangleCollisionPoint<FSolverReal> > KinematicResult;
-					if (CollidableSubMesh.GetKinematicColliderSubMesh().PointProximityQuery(KinematicColliderSpatial, static_cast<const TArrayView<const FSolverVec3>&>(Particles.XArray()), Index, Particles.X(Index), Thickness * ExtraThicknessMult, KinematicColliderThickness * ExtraThicknessMult, [](const int32 PointIndex, const int32 SubMeshTriangleIndex)->bool { return true; }, KinematicResult))
+					if (CollidableSubMesh.GetKinematicColliderSubMesh().PointProximityQuery(KinematicColliderSpatial, static_cast<const TArrayView<const FSolverVec3>&>(Particles.XArray()), Index, Particles.X(Index), ParticleThickness * ExtraThicknessMult, KinematicColliderThickness * ExtraThicknessMult, [](const int32 PointIndex, const int32 SubMeshTriangleIndex)->bool { return true; }, KinematicResult))
 					{
 						if (KinematicResult.Num() > MaxKinematicConnectionsPerPoint)
 						{
@@ -470,11 +498,10 @@ FSolverVec3 FPBDCollisionSpringConstraintsBase::GetDelta(const SolverParticlesOr
 		return FSolverVec3(0);
 	}
 
-	const FSolverReal ConstraintStiffness = GetConstraintStiffness(ConstraintIndex);
 	const FSolverReal ConstraintFriction = GetConstraintFrictionCoefficient(ConstraintIndex);
 
 	const FSolverReal NormalDelta = Height - NormalDifference;
-	const FSolverVec3 RepulsionDelta = ConstraintStiffness * NormalDelta * Normal / CombinedMass;
+	const FSolverVec3 RepulsionDelta = Stiffness * NormalDelta * Normal / CombinedMass;
 
 	if (ConstraintFriction > 0)
 	{
@@ -570,21 +597,90 @@ void FPBDCollisionSpringConstraintsBase::UpdateLinearSystem(const FSolverParticl
 template<typename SolverParticlesOrRange>
 void FPBDCollisionSpringConstraintsBase::Apply(SolverParticlesOrRange& InParticles, const FSolverReal Dt) const
 {
+	ApplyDynamicConstraints(InParticles, Dt);
+	ApplyKinematicConstraints(InParticles, Dt);
+}
+template void CHAOS_API FPBDCollisionSpringConstraintsBase::Apply(FSolverParticles& Particles, const FSolverReal Dt) const;
+template void CHAOS_API FPBDCollisionSpringConstraintsBase::Apply(FSolverParticlesRange& Particles, const FSolverReal Dt) const;
+
+template<typename SolverParticlesOrRange>
+void FPBDCollisionSpringConstraintsBase::ApplyDynamicConstraints(SolverParticlesOrRange& Particles, const FSolverReal Dt) const
+{
+	TRACE_CPUPROFILER_EVENT_SCOPE(ChaosPBDCollisionSpring_ApplyDynamic);
 	PRAGMA_DISABLE_DEPRECATION_WARNINGS
-	// Dynamic collisions
+
+	check(Constraints.Num() == Barys.Num());
+	check(Constraints.Num() == FlipNormal.Num());
+
+	const TVec4<int32>* const ConstraintsData = Constraints.GetData();
+	const FSolverVec3* const BarysData = Barys.GetData();
+	const bool* const FlipNormalData = FlipNormal.GetData();
+
+	FPAndInvM* const PAndInvM = Particles.GetPAndInvM().GetData();
+	const FSolverVec3* const X = Particles.XArray().GetData();
+
+	for (int32 ConstraintIndex = 0; ConstraintIndex < Constraints.Num(); ++ConstraintIndex)
 	{
-		TRACE_CPUPROFILER_EVENT_SCOPE(ChaosPBDCollisionSpring_ApplyDynamic);
-		for (int32 ConstraintIndex = 0; ConstraintIndex < Constraints.Num(); ++ConstraintIndex)
+		const TVec4<int32>& Constraint = ConstraintsData[ConstraintIndex];
+		const int32 Index1 = Constraint[0];
+		const int32 Index2 = Constraint[1];
+		const int32 Index3 = Constraint[2];
+		const int32 Index4 = Constraint[3];
+
+		const FSolverReal TrianglePointInvM =
+			PAndInvM[Index2].InvM * BarysData[ConstraintIndex][0] +
+			PAndInvM[Index3].InvM * BarysData[ConstraintIndex][1] +
+			PAndInvM[Index4].InvM * BarysData[ConstraintIndex][2];
+
+		const FSolverReal CombinedMass = PAndInvM[Index1].InvM + TrianglePointInvM;
+		if (CombinedMass <= (FSolverReal)1e-7)
+		{
+			continue;
+		}
+
+		FSolverVec3& P1 = PAndInvM[Index1].P;
+		FSolverVec3& P2 = PAndInvM[Index2].P;
+		FSolverVec3& P3 = PAndInvM[Index3].P;
+		FSolverVec3& P4 = PAndInvM[Index4].P;
+
+		const FSolverReal Height = GetConstraintThickness(ConstraintIndex);
+		const TTriangle<FSolverReal> Triangle(P2, P3, P4);
+		const FSolverVec3 Normal = FlipNormalData[ConstraintIndex] ? -Triangle.GetNormal() : Triangle.GetNormal();
+		const FSolverVec3 P = BarysData[ConstraintIndex][0] * P2 + BarysData[ConstraintIndex][1] * P3 + BarysData[ConstraintIndex][2] * P4;
+		const FSolverVec3 Difference = P1 - P;
+		const FSolverReal NormalDifference = Difference.Dot(Normal);
+
+		// Normal repulsion with friction
+		if (NormalDifference > Height)
 		{
-			Apply(InParticles, Dt, ConstraintIndex);
+			continue;
 		}
+
+		const FSolverReal ConstraintFriction = GetConstraintFrictionCoefficient(ConstraintIndex);
+
+		const FSolverReal NormalDelta = Height - NormalDifference;
+		const FSolverVec3 RepulsionDelta = Stiffness * NormalDelta * Normal / CombinedMass; 
+		FSolverVec3 FrictionDelta((FSolverReal)0.);
+		if (ConstraintFriction > 0)
+		{
+			const FSolverVec3& X1 = X[Index1];
+			const FSolverVec3 XP = BarysData[ConstraintIndex][0] * X[Index2] + BarysData[ConstraintIndex][1] * X[Index3] + BarysData[ConstraintIndex][2] * X[Index4];
+			const FSolverVec3 RelativeDisplacement = (P1 - X1) - (P - XP) + (PAndInvM[Index1].InvM - TrianglePointInvM) * RepulsionDelta;
+			const FSolverVec3 RelativeDisplacementTangent = RelativeDisplacement - RelativeDisplacement.Dot(Normal) * Normal;
+			const FSolverReal RelativeDisplacementTangentLength = RelativeDisplacementTangent.Length();
+			const FSolverReal PositionCorrection = FMath::Min(NormalDelta * ConstraintFriction, RelativeDisplacementTangentLength);
+			const FSolverReal CorrectionRatio = RelativeDisplacementTangentLength < UE_SMALL_NUMBER ? 0.f : PositionCorrection / RelativeDisplacementTangentLength;
+			FrictionDelta = -CorrectionRatio * RelativeDisplacementTangent / CombinedMass;
+		}
+		const FSolverVec3 Delta = RepulsionDelta + FrictionDelta;
+
+		P1 += PAndInvM[Index1].InvM * Delta;
+		P2 -= PAndInvM[Index2].InvM * BarysData[ConstraintIndex][0] * Delta;
+		P3 -= PAndInvM[Index3].InvM * BarysData[ConstraintIndex][1] * Delta;
+		P4 -= PAndInvM[Index4].InvM * BarysData[ConstraintIndex][2] * Delta;
 	}
 	PRAGMA_ENABLE_DEPRECATION_WARNINGS
-
-	ApplyKinematicConstraints(InParticles, Dt);
 }
-template void CHAOS_API FPBDCollisionSpringConstraintsBase::Apply(FSolverParticles& Particles, const FSolverReal Dt) const;
-template void CHAOS_API FPBDCollisionSpringConstraintsBase::Apply(FSolverParticlesRange& Particles, const FSolverReal Dt) const;
 
 template<typename SolverParticlesOrRange>
 void FPBDCollisionSpringConstraintsBase::ApplyKinematicConstraints(SolverParticlesOrRange& Particles, const FSolverReal Dt) const
@@ -595,14 +691,12 @@ void FPBDCollisionSpringConstraintsBase::ApplyKinematicConstraints(SolverParticl
 	}
 	TRACE_CPUPROFILER_EVENT_SCOPE(ChaosPBDCollisionSpring_ApplyKinematicConstraints);
 
-	const FSolverReal Height = Thickness + KinematicColliderThickness;
-	const FSolverReal OneOverTangentialFalloffDist = (FSolverReal)1.f / FMath::Max(KinematicColliderFalloffMultiplier * Height, UE_KINDA_SMALL_NUMBER);
-	const FSolverReal MaxDepth = -Height * KinematicColliderMaxDepthMultiplier;
+
 #if INTEL_ISPC
 	static_assert(sizeof(ispc::FIntVector) == sizeof(TVector<int32, MaxKinematicConnectionsPerPoint>), "sizeof(ispc::FIntVector) != sizeof(TVector<int32, MaxKinematicConnectionsPerPoint>)");
 	if (bRealTypeCompatibleWithISPC && bChaos_CollisionSpring_ISPC_Enabled)
 	{
-		if (KinematicColliderFrictionCoefficient.HasWeightMap())
+		if (ThicknessWeighted.HasWeightMap())
 		{
 			ispc::ApplyKinematicCollisionSpringConstraintsWithMaps(
 				(ispc::FVector4f*)Particles.GetPAndInvM().GetData(),
@@ -610,31 +704,59 @@ void FPBDCollisionSpringConstraintsBase::ApplyKinematicConstraints(SolverParticl
 				KinematicCollidingParticles.GetData(),
 				(const ispc::FIntVector*)KinematicColliderElements.GetData(),
 				(const ispc::FIntVector*)TriangleMesh.GetElements().GetData(),
-				Height,
-				OneOverTangentialFalloffDist,
-				MaxDepth,
+				reinterpret_cast<const ispc::FVector2f&>(ThicknessWeighted.GetOffsetRange()),
+				ThicknessWeighted.GetMapValues().GetData(),
+				KinematicColliderThickness,
+				KinematicColliderFalloffMultiplier,
+				KinematicColliderMaxDepthMultiplier,
 				KinematicColliderStiffness,
+				KinematicColliderFrictionCoefficient.HasWeightMap(),
 				reinterpret_cast<const ispc::FVector2f&>(KinematicColliderFrictionCoefficient.GetOffsetRange()),
 				KinematicColliderFrictionCoefficient.GetMapValues().GetData(),
+				Offset,
 				KinematicCollidingParticles.Num()
 			);
-
 		}
 		else
 		{
-			ispc::ApplyKinematicCollisionSpringConstraints(
-				(ispc::FVector4f*)Particles.GetPAndInvM().GetData(),
-				(const ispc::FVector3f*)Particles.XArray().GetData(),
-				KinematicCollidingParticles.GetData(),
-				(const ispc::FIntVector*)KinematicColliderElements.GetData(),
-				(const ispc::FIntVector*)TriangleMesh.GetElements().GetData(),
-				Height,
-				OneOverTangentialFalloffDist,
-				MaxDepth,
-				KinematicColliderStiffness,
-				(FSolverReal)KinematicColliderFrictionCoefficient,
-				KinematicCollidingParticles.Num()
-			);
+			const FSolverReal Height = (FSolverReal)ThicknessWeighted + KinematicColliderThickness;
+			const FSolverReal OneOverTangentialFalloffDist = (FSolverReal)1.f / FMath::Max(KinematicColliderFalloffMultiplier * Height, UE_KINDA_SMALL_NUMBER);
+			const FSolverReal MaxDepth = -Height * KinematicColliderMaxDepthMultiplier;
+			if (KinematicColliderFrictionCoefficient.HasWeightMap())
+			{
+				ispc::ApplyKinematicCollisionSpringConstraintsWithKinematicFrictionMap(
+					(ispc::FVector4f*)Particles.GetPAndInvM().GetData(),
+					(const ispc::FVector3f*)Particles.XArray().GetData(),
+					KinematicCollidingParticles.GetData(),
+					(const ispc::FIntVector*)KinematicColliderElements.GetData(),
+					(const ispc::FIntVector*)TriangleMesh.GetElements().GetData(),
+					Height,
+					OneOverTangentialFalloffDist,
+					MaxDepth,
+					KinematicColliderStiffness,
+					reinterpret_cast<const ispc::FVector2f&>(KinematicColliderFrictionCoefficient.GetOffsetRange()),
+					KinematicColliderFrictionCoefficient.GetMapValues().GetData(),
+					Offset,
+					KinematicCollidingParticles.Num()
+				);
+
+			}
+			else
+			{
+				ispc::ApplyKinematicCollisionSpringConstraints(
+					(ispc::FVector4f*)Particles.GetPAndInvM().GetData(),
+					(const ispc::FVector3f*)Particles.XArray().GetData(),
+					KinematicCollidingParticles.GetData(),
+					(const ispc::FIntVector*)KinematicColliderElements.GetData(),
+					(const ispc::FIntVector*)TriangleMesh.GetElements().GetData(),
+					Height,
+					OneOverTangentialFalloffDist,
+					MaxDepth,
+					KinematicColliderStiffness,
+					(FSolverReal)KinematicColliderFrictionCoefficient,
+					KinematicCollidingParticles.Num()
+				);
+			}
 		}
 	}
 	else
@@ -643,6 +765,12 @@ void FPBDCollisionSpringConstraintsBase::ApplyKinematicConstraints(SolverParticl
 		for (int32 Index = 0; Index < KinematicCollidingParticles.Num(); ++Index)
 		{
 			const int32 Index1 = KinematicCollidingParticles[Index];
+			const FSolverReal ParticleThickness = ThicknessWeighted.GetValue(Index1 - Offset);
+
+			const FSolverReal Height = ParticleThickness + KinematicColliderThickness;
+			const FSolverReal OneOverTangentialFalloffDist = (FSolverReal)1.f / FMath::Max(KinematicColliderFalloffMultiplier * Height, UE_KINDA_SMALL_NUMBER);
+			const FSolverReal MaxDepth = -Height * KinematicColliderMaxDepthMultiplier;
+
 			for (int32 EIndex = 0; EIndex < MaxKinematicConnectionsPerPoint; ++EIndex)
 			{
 				const int32 ElemIndex = KinematicColliderElements[Index][EIndex];
@@ -685,7 +813,7 @@ void FPBDCollisionSpringConstraintsBase::ApplyKinematicConstraints(SolverParticl
 
 				P1 += RepulsionDelta;
 
-				const FSolverReal KinematicFrictionCoefficient = KinematicColliderFrictionCoefficient.GetValue(Index1);
+				const FSolverReal KinematicFrictionCoefficient = KinematicColliderFrictionCoefficient.GetValue(Index1 - Offset);
 				if (KinematicFrictionCoefficient > 0)
 				{
 					const FSolverVec3& X1 = Particles.X(Index1);
diff --git a/Engine/Source/Runtime/Experimental/Chaos/Private/Chaos/PBDTriangleMeshCollisions.cpp b/Engine/Source/Runtime/Experimental/Chaos/Private/Chaos/PBDTriangleMeshCollisions.cpp
--- a/Engine/Source/Runtime/Experimental/Chaos/Private/Chaos/PBDTriangleMeshCollisions.cpp
+++ b/Engine/Source/Runtime/Experimental/Chaos/Private/Chaos/PBDTriangleMeshCollisions.cpp
@@ -103,7 +103,7 @@ static void FindEdgeFaceIntersections(const FTriangleMesh& TriangleMesh, const S
 	);
 
 	// Set Intersections Num to actual number found.
-	const int32 IntersectionNum = IntersectionIndex.load();
+	const int32 IntersectionNum = FMath::Min(IntersectionIndex.load(), PreallocatedIntersectionsNum);
 	Intersections.SetNum(IntersectionNum, EAllowShrinking::No);
 
 	// Append any ExtraIntersections
@@ -1378,7 +1378,7 @@ void FPBDTriangleMeshCollisions::FTriangleSubMesh::InitAllDynamic()
 }
 
 template<typename SolverParticlesOrRange>
-void FPBDTriangleMeshCollisions::Init(const SolverParticlesOrRange& Particles, const FSolverReal MinProximityQueryRadius)
+void FPBDTriangleMeshCollisions::Init(const SolverParticlesOrRange& Particles, const FPBDFlatWeightMap& ThicknessMap)
 {
 	const bool bDoSelfIntersections = bGlobalIntersectionAnalysis || bContourMinimization;
 	if (bCollidableSubMeshDirty)
@@ -1396,8 +1396,19 @@ void FPBDTriangleMeshCollisions::Init(const SolverParticlesOrRange& Particles, c
 
 	{
 		TRACE_CPUPROFILER_EVENT_SCOPE(ChaosFPBDTriangleMeshCollisions_BuildSpatialHash);
+
 		constexpr FSolverReal RadiusToLodSizeMultiplier = 2.; // Radius to Diameter
-		DynamicSubMesh.BuildSpatialHash(static_cast<const TArrayView<const FSolverVec3>&>(Particles.XArray()), DynamicSubMeshSpatialHash, RadiusToLodSizeMultiplier * MinProximityQueryRadius);
+		FSolverReal MinProximityQueryRadius;
+		if (!ThicknessMap.HasWeightMap())
+		{	
+			MinProximityQueryRadius = ((FSolverReal)ThicknessMap) * 2.f; // When thickness is constant, we don't put any thickness in the spatial hash and instead query with particles with 2*thickness radius.
+			DynamicSubMesh.BuildSpatialHash(static_cast<const TArrayView<const FSolverVec3>&>(Particles.XArray()), DynamicSubMeshSpatialHash, RadiusToLodSizeMultiplier * MinProximityQueryRadius);
+		}
+		else
+		{
+			MinProximityQueryRadius = FMath::Max(ThicknessMap.GetLow(), ThicknessMap.GetHigh());
+			DynamicSubMesh.BuildSpatialHash(static_cast<const TArrayView<const FSolverVec3>&>(Particles.XArray()), DynamicSubMeshSpatialHash, ThicknessMap, Offset, RadiusToLodSizeMultiplier * MinProximityQueryRadius);
+		}
 		if (CollidableSubMesh.GetKinematicColliderSubMesh().GetNumElements() > 0)
 		{
 			CollidableSubMesh.GetKinematicColliderSubMesh().BuildSpatialHash(static_cast<const TArrayView<const FSolverVec3>&>(Particles.XArray()), KinematicSubMeshSpatialHash, RadiusToLodSizeMultiplier * MinProximityQueryRadius);
@@ -1453,8 +1464,8 @@ void FPBDTriangleMeshCollisions::Init(const SolverParticlesOrRange& Particles, c
 		}
 	}
 }
-template CHAOS_API void FPBDTriangleMeshCollisions::Init(const FSolverParticles& Particles, const FSolverReal MinProximityQueryRadius);
-template CHAOS_API void FPBDTriangleMeshCollisions::Init(const FSolverParticlesRange& Particles, const FSolverReal MinProximityQueryRadius);
+template CHAOS_API void FPBDTriangleMeshCollisions::Init(const FSolverParticles& Particles, const FPBDFlatWeightMap& ThicknessMap);
+template CHAOS_API void FPBDTriangleMeshCollisions::Init(const FSolverParticlesRange& Particles, const FPBDFlatWeightMap& ThicknessMap);
 
 template<typename SolverParticlesOrRange>
 void FPBDTriangleMeshCollisions::PostStepInit(const SolverParticlesOrRange& Particles)
diff --git a/Engine/Source/Runtime/Experimental/Chaos/Private/Chaos/TriangleMesh.cpp b/Engine/Source/Runtime/Experimental/Chaos/Private/Chaos/TriangleMesh.cpp
--- a/Engine/Source/Runtime/Experimental/Chaos/Private/Chaos/TriangleMesh.cpp
+++ b/Engine/Source/Runtime/Experimental/Chaos/Private/Chaos/TriangleMesh.cpp
@@ -13,6 +13,7 @@
 #include "Templates/Sorting.h"
 #include "Templates/TypeHash.h"
 #include "Chaos/HierarchicalSpatialHash.h"
+#include "Chaos/PBDFlatWeightMap.h"
 
 #include <algorithm>
 #include <iostream>
@@ -97,6 +98,53 @@ struct TTriangleMeshBvData
 	}
 };
 
+struct FTriangleMeshBvDataWithThicknessEntry
+{
+	const FTriangleMesh* TmData;
+	const TConstArrayView<Softs::FSolverVec3>& Points;
+	const Softs::FPBDFlatWeightMap& PointThicknesses;
+	int32 ThicknessMapIndexOffset;
+	int32 Index;
+
+	bool HasBoundingBox() const
+	{
+		return true;
+	}
+
+	typename FTriangleMesh::TSpatialHashType<Softs::FSolverReal>::FVectorAABB VectorAABB() const
+	{
+		const TVec3<int32>& Tri = TmData->GetElements()[Index];
+		const Softs::FSolverVec3& A = Points[Tri[0]];
+		const Softs::FSolverVec3& B = Points[Tri[1]];
+		const Softs::FSolverVec3& C = Points[Tri[2]];
+		typename FTriangleMesh::TSpatialHashType<Softs::FSolverReal>::FVectorAABB AABB(A, PointThicknesses.GetValue(Tri[0] - ThicknessMapIndexOffset));
+		AABB.GrowToInclude(B, PointThicknesses.GetValue(Tri[1] - ThicknessMapIndexOffset));
+		AABB.GrowToInclude(C, PointThicknesses.GetValue(Tri[2] - ThicknessMapIndexOffset));
+		return AABB;
+	}
+
+	template<typename TPayloadType>
+	int32 GetPayload(int32 Idx) const { return Idx; }
+};
+
+struct FTriangleMeshBvDataWithThickness
+{
+	const FTriangleMesh* TmData;
+	const TConstArrayView<Softs::FSolverVec3>& Points;
+	const Softs::FPBDFlatWeightMap& PointThicknesses;
+	const int32 ThicknessMapIndexOffset;
+
+	FTriangleMeshBvDataWithThicknessEntry operator[](const int32 ParticleIndex) const
+	{
+		return FTriangleMeshBvDataWithThicknessEntry{ TmData, Points, PointThicknesses, ThicknessMapIndexOffset, ParticleIndex };
+	}
+
+	int32 Num() const
+	{
+		return TmData->GetNumElements();
+	}
+};
+
 FTriangleMesh::FTriangleMesh()
     : MStartIdx(0)
     , MNumIndices(0)
@@ -474,6 +522,12 @@ CHAOS_API void FTriangleMesh::GetPointNormals(TArrayView<TVec3<FRealSingle>> Poi
 	}
 }
 
+template<>
+CHAOS_API void FTriangleMesh::GetPointNormals<FRealSingle>(TArrayView<TVec3<FRealSingle>> PointNormals, const TConstArrayView<TVec3<FRealSingle>>& FaceNormals, const bool bUseGlobalArray) const
+{
+	return GetPointNormals(PointNormals, FaceNormals, bUseGlobalArray);
+}
+
 template<class T>
 void AddTrianglesToHull(const TConstArrayView<FVec3>& Points, const int32 I0, const int32 I1, const int32 I2, const TPlane<T, 3>& SplitPlane, const TArray<int32>& InIndices, TArray<TVec3<int32>>& OutIndices)
 {
@@ -1712,6 +1766,12 @@ void FTriangleMesh::BuildSpatialHash(const TConstArrayView<TVec3<T>>& Points, TS
 template void FTriangleMesh::BuildSpatialHash<FRealSingle>(const TConstArrayView<TVec3<FRealSingle>>& Points, TSpatialHashType<FRealSingle>& SpatialHash, const FRealSingle MinSpatialLodSize) const;
 template void FTriangleMesh::BuildSpatialHash<FRealDouble>(const TConstArrayView<TVec3<FRealDouble>>& Points, TSpatialHashType<FRealDouble>& SpatialHash, const FRealDouble MinSpatialLodSize) const;
 
+void FTriangleMesh::BuildSpatialHash(const TConstArrayView<Softs::FSolverVec3>& Points, TSpatialHashType<Softs::FSolverReal>& SpatialHash, const Softs::FPBDFlatWeightMap& PointThicknesses, int32 ThicknessMapIndexOffset, const Softs::FSolverReal MinSpatialLodSize) const
+{
+	const FTriangleMeshBvDataWithThickness BvData({ this, Points, PointThicknesses, ThicknessMapIndexOffset });
+	SpatialHash.Initialize(BvData, MinSpatialLodSize);
+}
+
 template<typename T>
 bool FTriangleMesh::PointProximityQuery(const TSpatialHashType<T>& SpatialHash, const TConstArrayView<TVec3<T>>& Points, const int32 PointIndex, const TVec3<T>& PointPosition, const T PointThickness, const T ThisThickness,
 	TFunctionRef<bool(const int32 PointIndex, const int32 TriangleIndex)> BroadphaseTest, TArray<TTriangleCollisionPoint<T>>& Result) const
@@ -1763,6 +1823,60 @@ bool FTriangleMesh::PointProximityQuery(const TSpatialHashType<T>& SpatialHash,
 template bool FTriangleMesh::PointProximityQuery<FRealSingle>(const TSpatialHashType<FRealSingle>& SpatialHash, const TConstArrayView<TVector<FRealSingle, 3>>& Points, const int32 PointIndex, const TVector<FRealSingle, 3>& PointPosition, const FRealSingle PointThickness, const FRealSingle ThisThickness, TFunctionRef<bool(const int32 PointIndex, const int32 TriangleIndex)> BroadphaseTest, TArray<TTriangleCollisionPoint<FRealSingle>>& Result) const;
 template bool FTriangleMesh::PointProximityQuery<FRealDouble>(const TSpatialHashType<FRealDouble>& SpatialHash, const TConstArrayView<TVector<FRealDouble, 3>>& Points, const int32 PointIndex, const TVector<FRealDouble, 3>& PointPosition, const FRealDouble PointThickness, const FRealDouble ThisThickness, TFunctionRef<bool(const int32 PointIndex, const int32 TriangleIndex)> BroadphaseTest, TArray<TTriangleCollisionPoint<FRealDouble>>& Result) const;
 
+bool FTriangleMesh::PointProximityQuery(const TSpatialHashType<Softs::FSolverReal>& SpatialHash, const TConstArrayView<Softs::FSolverVec3>& Points, const int32 PointIndex, const Softs::FSolverVec3& PointPosition, const Softs::FSolverReal PointThickness, const Softs::FPBDFlatWeightMap& ThisThicknesses,
+	const Softs::FSolverReal ThisThicknessExtraMultiplier, int32 ThicknessMapIndexOffset, TFunctionRef<bool(const int32 PointIndex, const int32 TriangleIndex)> BroadphaseTest, TArray<TTriangleCollisionPoint<Softs::FSolverReal>>& Result) const
+{
+	typename TSpatialHashType<FRealSingle>::FVectorAABB QueryBounds(PointPosition);
+	QueryBounds.Thicken(PointThickness);
+
+	const TArray<int32> PotentialIntersections = SpatialHash.FindAllIntersections(QueryBounds,
+		[PointIndex, &BroadphaseTest](int32 Payload)
+	{
+		return BroadphaseTest(PointIndex, Payload);
+	});
+
+	Result.Reset(PotentialIntersections.Num());
+
+	for (int32 TriIdx : PotentialIntersections)
+	{
+
+		const Softs::FSolverVec3& A = Points[MElements[TriIdx][0]];
+		const Softs::FSolverVec3& B = Points[MElements[TriIdx][1]];
+		const Softs::FSolverVec3& C = Points[MElements[TriIdx][2]];
+		Softs::FSolverVec3 Bary;
+		const Softs::FSolverVec3 ClosestPoint = FindClosestPointAndBaryOnTriangle(A, B, C, PointPosition, Bary);
+
+		const Softs::FSolverReal TriangleThickness = (ThisThicknesses.GetValue(MElements[TriIdx][0] - ThicknessMapIndexOffset) * Bary[0] +
+			ThisThicknesses.GetValue(MElements[TriIdx][1] - ThicknessMapIndexOffset) * Bary[1] +
+			ThisThicknesses.GetValue(MElements[TriIdx][2] - ThicknessMapIndexOffset) * Bary[2]) * ThisThicknessExtraMultiplier;
+
+		const Softs::FSolverReal TotalThickness = PointThickness + TriangleThickness;
+		const Softs::FSolverReal TotalThicknessSq = TotalThickness * TotalThickness;
+
+		const Softs::FSolverReal DistSq = (PointPosition - ClosestPoint).SizeSquared();
+		if (DistSq > TotalThicknessSq)
+		{
+			// Failed narrow test.
+			continue;
+		}
+
+		Softs::FSolverVec3 Normal = Softs::FSolverVec3::CrossProduct(B - A, C - A).GetSafeNormal();
+		Normal = (Softs::FSolverVec3::DotProduct(Normal, PointPosition - A) > 0) ? Normal : -Normal;
+
+		TTriangleCollisionPoint<Softs::FSolverReal> CollisionPoint;
+		CollisionPoint.ContactType = TTriangleCollisionPoint<Softs::FSolverReal>::EContactType::PointFace;
+		CollisionPoint.Indices[0] = PointIndex;
+		CollisionPoint.Indices[1] = TriIdx;
+		CollisionPoint.Bary = TVec4<Softs::FSolverReal>((Softs::FSolverReal)1., Bary.X, Bary.Y, Bary.Z);
+		CollisionPoint.Location = ClosestPoint;
+		CollisionPoint.Normal = Normal;
+		CollisionPoint.Phi = FMath::Sqrt(DistSq);
+		Result.Add(CollisionPoint);
+	}
+	return Result.Num() > 0;
+
+}
+
 template<typename T>
 bool FTriangleMesh::EdgeIntersectionQuery(const TSpatialHashType<T>& SpatialHash, const TConstArrayView<TVec3<T>>& Points, const int32 EdgeIndex, const TVec3<T>& EdgePosition1, const TVec3<T>& EdgePosition2,
 	TFunctionRef<bool(const int32 EdgeIndex, const int32 TriangleIndex)> BroadphaseTest, TArray<TTriangleCollisionPoint<T>>& Result) const
diff --git a/Engine/Source/Runtime/Experimental/Chaos/Public/Chaos/HierarchicalSpatialHash.h b/Engine/Source/Runtime/Experimental/Chaos/Public/Chaos/HierarchicalSpatialHash.h
--- a/Engine/Source/Runtime/Experimental/Chaos/Public/Chaos/HierarchicalSpatialHash.h
+++ b/Engine/Source/Runtime/Experimental/Chaos/Public/Chaos/HierarchicalSpatialHash.h
@@ -314,6 +314,13 @@ public:
 			, Max(Min)
 		{}
 
+		FVectorAABB(const TVec3<T>& Point, const T Radius)
+			: Min(MakeVectorRegister(Point[0], Point[1], Point[2], (T)0.))
+			, Max(Min)
+		{
+			Thicken(Radius);
+		}
+
 		void GrowToInclude(const TVec3<T>& Point)
 		{
 			const TVectorRegisterType<T> Vector = MakeVectorRegister(Point[0], Point[1], Point[2], (T)0);
@@ -321,6 +328,14 @@ public:
 			Max = VectorMax(Max, Vector);
 		}
 
+		void GrowToInclude(const TVec3<T>& Point, const T Radius)
+		{
+			const TVectorRegisterType<T> Vector = MakeVectorRegister(Point[0], Point[1], Point[2], (T)0);
+			const TVectorRegisterType<T> ThicknessVec = MakeVectorRegister(Radius, Radius, Radius, (T)0);
+			Min = VectorMin(Min, VectorSubtract(Vector, ThicknessVec));
+			Max = VectorMax(Max, VectorAdd(Vector, ThicknessVec));
+		}
+
 		void Thicken(const T Thickness)
 		{
 			const TVectorRegisterType<T> ThicknessVec = MakeVectorRegister(Thickness, Thickness, Thickness, (T)0);
diff --git a/Engine/Source/Runtime/Experimental/Chaos/Public/Chaos/PBDCollisionSpringConstraints.h b/Engine/Source/Runtime/Experimental/Chaos/Public/Chaos/PBDCollisionSpringConstraints.h
--- a/Engine/Source/Runtime/Experimental/Chaos/Public/Chaos/PBDCollisionSpringConstraints.h
+++ b/Engine/Source/Runtime/Experimental/Chaos/Public/Chaos/PBDCollisionSpringConstraints.h
@@ -36,8 +36,9 @@ public:
 			InRestPositions,
 			GenerateDisabledCollisionElements(InTriangleMesh, PropertyCollection),
 			WeightMaps.FindRef(GetSelfCollisionKinematicColliderFrictionString(PropertyCollection, SelfCollisionKinematicColliderFrictionName.ToString())),
+			WeightMaps.FindRef(GetSelfCollisionThicknessString(PropertyCollection, SelfCollisionThicknessName.ToString())),
 			FaceIntMaps.FindRef(GetSelfCollisionLayersString(PropertyCollection, SelfCollisionLayersName.ToString())),
-			(FSolverReal)FMath::Max(GetSelfCollisionThickness(PropertyCollection, Base::BackCompatThickness), 0.f),
+			FSolverVec2::Max(FSolverVec2(GetWeightedFloatSelfCollisionThickness(PropertyCollection, Base::BackCompatThickness)), FSolverVec2(0.f)),
 			(FSolverReal)FMath::Clamp(GetSelfCollisionStiffness(PropertyCollection, Base::BackCompatStiffness), 0.f, 1.f),
 			FMath::Clamp((FSolverReal)GetSelfCollisionFriction(PropertyCollection, Base::BackCompatFrictionCoefficient), MinFrictionCoefficient, MaxFrictionCoefficient),
 			GetSelfCollideAgainstKinematicCollidersOnly(PropertyCollection, false),
@@ -91,8 +92,9 @@ public:
 			InRestPositions,
 			MoveTemp(InDisabledCollisionElements),
 			TConstArrayView<FRealSingle>(),
+			TConstArrayView<FRealSingle>(),
 			TConstArrayView<int32>(),
-			InThickness,
+			FSolverVec2(InThickness),
 			InStiffness,
 			InFrictionCoefficient)
 		, SelfCollisionThicknessIndex(ForceInit)
@@ -116,7 +118,19 @@ public:
 	{
 		if (IsSelfCollisionThicknessMutable(PropertyCollection))
 		{
-			Thickness = (FSolverReal)FMath::Max(GetSelfCollisionThickness(PropertyCollection), 0.f);
+			const FSolverVec2 WeightedValue = FSolverVec2::Max(FSolverVec2(GetWeightedFloatSelfCollisionThickness(PropertyCollection)), FSolverVec2(0.f));
+			if (IsSelfCollisionThicknessStringDirty(PropertyCollection))
+			{
+				const FString& WeightMapName = GetSelfCollisionThicknessString(PropertyCollection);
+				ThicknessWeighted = FPBDFlatWeightMap(WeightedValue, WeightMaps.FindRef(WeightMapName), GetNumParticles());
+			}
+			else
+			{
+				ThicknessWeighted.SetWeightedValue(WeightedValue);
+			}
+PRAGMA_DISABLE_DEPRECATION_WARNINGS
+			Thickness = GetMaxThickness();
+PRAGMA_ENABLE_DEPRECATION_WARNINGS
 		}
 		if (IsSelfCollisionStiffnessMutable(PropertyCollection))
 		{
@@ -172,7 +186,7 @@ public:
 	}
 
 private:
-	using Base::Thickness;
+	using Base::ThicknessWeighted;
 	using Base::Stiffness;
 	using Base::FrictionCoefficient;
 	using Base::ProximityStiffness;
diff --git a/Engine/Source/Runtime/Experimental/Chaos/Public/Chaos/PBDCollisionSpringConstraintsBase.h b/Engine/Source/Runtime/Experimental/Chaos/Public/Chaos/PBDCollisionSpringConstraintsBase.h
--- a/Engine/Source/Runtime/Experimental/Chaos/Public/Chaos/PBDCollisionSpringConstraintsBase.h
+++ b/Engine/Source/Runtime/Experimental/Chaos/Public/Chaos/PBDCollisionSpringConstraintsBase.h
@@ -36,9 +36,10 @@ public:
 		const FTriangleMesh& InTriangleMesh,
 		const TArray<FSolverVec3>* InReferencePositions,
 		TSet<TVec2<int32>>&& InDisabledCollisionElements,
+		const TConstArrayView<FRealSingle>& InThicknessMultipliers,
 		const TConstArrayView<FRealSingle>& InKinematicColliderFrictionMultipliers,
 		const TConstArrayView<int32>& InSelfCollisionLayers,
-		const FSolverReal InThickness = BackCompatThickness,
+		const FSolverVec2 InThickness = FSolverVec2(BackCompatThickness),
 		const FSolverReal InStiffness = BackCompatStiffness,
 		const FSolverReal InFrictionCoefficient = BackCompatFrictionCoefficient,
 		const bool bInOnlyCollideKinematics = false,
@@ -58,7 +59,7 @@ public:
 		const FSolverReal InStiffness = BackCompatStiffness,
 		const FSolverReal InFrictionCoefficient = BackCompatFrictionCoefficient)
 		: FPBDCollisionSpringConstraintsBase(InOffset, InNumParticles, InTriangleMesh, InReferencePositions, MoveTemp(InDisabledCollisionElements),
-			TConstArrayView<FRealSingle>(), TConstArrayView<int32>(), InThickness, InStiffness, InFrictionCoefficient)
+			TConstArrayView<FRealSingle>(), TConstArrayView<FRealSingle>(), TConstArrayView<int32>(), FSolverVec2(InThickness), InStiffness, InFrictionCoefficient)
 	{}
 
 	PRAGMA_DISABLE_DEPRECATION_WARNINGS
@@ -89,7 +90,13 @@ public:
 		PRAGMA_ENABLE_DEPRECATION_WARNINGS
 	}
 
-	FSolverReal GetThickness() const { return Thickness; }
+	UE_DEPRECATED(5.4, "Thickness is now weighted.")
+	FSolverReal GetThickness() const { return (FSolverReal)ThicknessWeighted; }
+
+	FSolverReal GetMaxThickness() const { return FMath::Max(ThicknessWeighted.GetLow(), ThicknessWeighted.GetHigh()); }
+	FSolverReal GetParticleThickness(int32 ParticleIndex) const { return ThicknessWeighted.GetValue(ParticleIndex - Offset); }
+	const FPBDFlatWeightMap& GetThicknessWeighted() const { return ThicknessWeighted; }
+
 	bool GetGlobalIntersectionAnalysis() const { return bGlobalIntersectionAnalysis; }
 	const TArray<bool>& GetFlipNormals() const 
 	{
@@ -102,7 +109,12 @@ public:
 	const TArray<TMap<int32, FSolverReal>>& GetKinematicColliderTimers() const { return KinematicColliderTimers; }
 	const FTriangleMesh& GetTriangleMesh() const { return TriangleMesh; }
 
-	void SetThickness(FSolverReal InThickness) { Thickness = FMath::Max(InThickness, (FSolverReal)0.);  }
+	UE_DEPRECATED(5.4, "Thickness is now weighted.")
+	void SetThickness(FSolverReal InThickness) 
+	{ 
+		SetThicknessWeighted(FSolverVec2(InThickness));
+	}
+	void SetThicknessWeighted(const FSolverVec2 InThickness) { ThicknessWeighted.SetWeightedValue(FSolverVec2::Max(InThickness, FSolverVec2(0.f))); }
 	void SetFrictionCoefficient(FSolverReal InFrictionCoefficient) { FrictionCoefficient = InFrictionCoefficient; }
 
 	template<typename SolverParticlesOrRange>
@@ -152,24 +164,20 @@ public:
 
 	FSolverReal GetConstraintThickness(const int32 ConstraintIndex) const
 	{
-		switch (ConstraintTypes[ConstraintIndex])
-		{
-		default:
-		case EConstraintType::Default:
-		case EConstraintType::GIAFlipped:
-			return 2.f * Thickness;
-		}		
-	}
-
-	FSolverReal GetConstraintStiffness(const int32 ConstraintIndex) const
-	{
-		switch (ConstraintTypes[ConstraintIndex])
+		if (!ThicknessWeighted.HasWeightMap())
 		{
-		default:
-		case EConstraintType::Default:
-		case EConstraintType::GIAFlipped:
-			return Stiffness;
+			return 2.f * (FSolverReal)ThicknessWeighted;
 		}
+		PRAGMA_DISABLE_DEPRECATION_WARNINGS
+		const TVec4<int32>& Constraint = Constraints[ConstraintIndex];
+		const FSolverVec3& Bary = Barys[ConstraintIndex];
+		PRAGMA_ENABLE_DEPRECATION_WARNINGS
+		const int32 Index1 = Constraint[0] - Offset;
+		const int32 Index2 = Constraint[1] - Offset;
+		const int32 Index3 = Constraint[2] - Offset;
+		const int32 Index4 = Constraint[3] - Offset;
+
+		return ThicknessWeighted.GetValue(Index1) + Bary[0] * ThicknessWeighted.GetValue(Index2) + Bary[1] * ThicknessWeighted.GetValue(Index3) + Bary[2] * ThicknessWeighted.GetValue(Index4);
 	}
 
 	FSolverReal GetConstraintFrictionCoefficient(const int32 ConstraintIndex) const
@@ -183,9 +191,9 @@ public:
 			return (FSolverReal)0.f;
 		}
 	}
-protected:
 
-	FSolverReal Thickness;
+protected:
+	FPBDFlatWeightMap ThicknessWeighted;
 	FSolverReal Stiffness; // (0-1 compliance for PBD)
 	FSolverReal FrictionCoefficient; 
 	bool bOnlyCollideKinematics;
@@ -194,6 +202,9 @@ protected:
 	FPBDFlatWeightMap KinematicColliderFrictionCoefficient;
 	FSolverReal ProximityStiffness; // (actual spring stiffness for force-based solver)
 
+	UE_DEPRECATED(5.4, "Use ThicknessWeighted instead")
+	FSolverReal Thickness;
+
 	UE_DEPRECATED(5.4, "Constraints will be made private")
 	TArray<TVec4<int32>> Constraints;
 	UE_DEPRECATED(5.4, "Barys will be made private")
@@ -206,10 +217,13 @@ protected:
 	int32 GetNumParticles() const { return NumParticles; }
 
 private:
+	template<typename SolverParticlesOrRange>
+	void ApplyDynamicConstraints(SolverParticlesOrRange& InParticles, const FSolverReal Dt) const;
 
 	template<typename SolverParticlesOrRange>
 	void ApplyKinematicConstraints(SolverParticlesOrRange& InParticles, const FSolverReal Dt) const;
 
+
 	const FTriangleMesh& TriangleMesh;
 	const TArray<FSolverVec3>* ReferencePositions;
 	const TSet<TVec2<int32>> DisabledCollisionElements;  // TODO: Make this a bitarray
diff --git a/Engine/Source/Runtime/Experimental/Chaos/Public/Chaos/PBDTriangleMeshCollisions.h b/Engine/Source/Runtime/Experimental/Chaos/Public/Chaos/PBDTriangleMeshCollisions.h
--- a/Engine/Source/Runtime/Experimental/Chaos/Public/Chaos/PBDTriangleMeshCollisions.h
+++ b/Engine/Source/Runtime/Experimental/Chaos/Public/Chaos/PBDTriangleMeshCollisions.h
@@ -2,6 +2,7 @@
 #pragma once
 
 #include "Chaos/Core.h"
+#include "Chaos/PBDFlatWeightMap.h"
 #include "Chaos/PBDSoftsEvolutionFwd.h"
 #include "Chaos/TriangleMesh.h"
 #include "Chaos/HierarchicalSpatialHash.h"
@@ -357,7 +358,13 @@ public:
 	}
 
 	template<typename SolverParticlesOrRange>
-	CHAOS_API void Init(const SolverParticlesOrRange& Particles, const FSolverReal MinProximityQueryRadius = (FSolverReal)0.);
+	void Init(const SolverParticlesOrRange& Particles, const FSolverReal MinProximityQueryRadius = (FSolverReal)0.)
+	{
+		Init(Particles, FPBDFlatWeightMap(FSolverVec2(MinProximityQueryRadius * (FSolverReal).5f)));
+	}
+
+	template<typename SolverParticlesOrRange>
+	CHAOS_API void Init(const SolverParticlesOrRange& Particles, const FPBDFlatWeightMap& ThicknessMap);
 
 	template<typename SolverParticlesOrRange>
 	CHAOS_API void PostStepInit(const SolverParticlesOrRange& Particles);
diff --git a/Engine/Source/Runtime/Experimental/Chaos/Public/Chaos/TriangleMesh.h b/Engine/Source/Runtime/Experimental/Chaos/Public/Chaos/TriangleMesh.h
--- a/Engine/Source/Runtime/Experimental/Chaos/Public/Chaos/TriangleMesh.h
+++ b/Engine/Source/Runtime/Experimental/Chaos/Public/Chaos/TriangleMesh.h
@@ -15,6 +15,11 @@ namespace Chaos
 
 	template<typename T> struct TTriangleCollisionPoint;
 
+	namespace Softs
+	{
+	class FPBDFlatWeightMap;
+	}
+
 	class FTriangleMesh
 	{
 	public:
@@ -164,7 +169,7 @@ namespace Chaos
 		 * @param bUseGlobalArray When true, fill the array from the StartIdx to StartIdx + NumIndices - 1 positions, otherwise fill the array from the 0 to NumIndices - 1 positions.
 		 */
 		template <typename T>
-		void GetPointNormals(TArrayView<TVec3<T>> PointNormals, const TConstArrayView<TVec3<T>>& FaceNormals, const bool bUseGlobalArray) const;
+		CHAOS_API void GetPointNormals(TArrayView<TVec3<T>> PointNormals, const TConstArrayView<TVec3<T>>& FaceNormals, const bool bUseGlobalArray) const;
 
 		CHAOS_API void GetPointNormals(TArrayView<TVec3<FRealSingle>> PointNormals, const TConstArrayView<TVec3<FRealSingle>>& FaceNormals, const bool bUseGlobalArray) const;
 
@@ -311,10 +316,13 @@ namespace Chaos
 		// Hierarchy will only go down to lods as small as MinSpatialLodSize.
 		template<typename T>
 		void BuildSpatialHash(const TConstArrayView<TVec3<T>>& Points, TSpatialHashType<T>& SpatialHash, const T MinSpatialLodSize = (T)0.) const;
+		void BuildSpatialHash(const TConstArrayView<TVec3<FRealSingle>>& Points, TSpatialHashType<FRealSingle>& SpatialHash, const Softs::FPBDFlatWeightMap& PointThicknesses, int32 ThicknessMapIndexOffset, const FRealSingle MinSpatialLodSize = 0.f) const;
 
 		template<typename T>
 		bool PointProximityQuery(const TSpatialHashType<T>& SpatialHash, const TConstArrayView<TVec3<T>>& Points, const int32 PointIndex, const TVec3<T>& PointPosition, const T PointThickness, const T ThisThickness,
 			TFunctionRef<bool(const int32 PointIndex, const int32 TriangleIndex)> BroadphaseTest, TArray<TTriangleCollisionPoint<T>>& Result) const;
+		bool PointProximityQuery(const TSpatialHashType<FRealSingle>& SpatialHash, const TConstArrayView<TVec3<FRealSingle>>& Points, const int32 PointIndex, const TVec3<FRealSingle>& PointPosition, const FRealSingle PointThickness, const Softs::FPBDFlatWeightMap& ThisThicknesses,
+			const FRealSingle ThisThicknessExtraMultiplier, int32 ThicknessMapIndexOffset, TFunctionRef<bool(const int32 PointIndex, const int32 TriangleIndex)> BroadphaseTest, TArray<TTriangleCollisionPoint<FRealSingle>>& Result) const;
 
 		template<typename T>
 		bool EdgeIntersectionQuery(const TSpatialHashType<T>& SpatialHash, const TConstArrayView<TVec3<T>>& Points, const int32 EdgeIndex, const TVec3<T>& EdgePosition1, const TVec3<T>& EdgePosition2,
