Make FTransformPair a chaos type instead of being defined independently in multiple places

diff --git a/Engine/Source/Runtime/Experimental/Chaos/Public/Chaos/PBDJointConstraintData.h b/Engine/Source/Runtime/Experimental/Chaos/Public/Chaos/PBDJointConstraintData.h
--- a/Engine/Source/Runtime/Experimental/Chaos/Public/Chaos/PBDJointConstraintData.h
+++ b/Engine/Source/Runtime/Experimental/Chaos/Public/Chaos/PBDJointConstraintData.h
@@ -31,7 +31,6 @@ namespace Chaos
 
 
 	using FJointConstraintDirtyFlags = TDirtyFlags<EJointConstraintFlags>;
-	typedef TVector<FTransform, 2> FTransformPair;
 
 	class FJointConstraint : public FConstraintBase
 	{
diff --git a/Engine/Source/Runtime/Experimental/Chaos/Public/Chaos/PBDJointConstraintTypes.h b/Engine/Source/Runtime/Experimental/Chaos/Public/Chaos/PBDJointConstraintTypes.h
--- a/Engine/Source/Runtime/Experimental/Chaos/Public/Chaos/PBDJointConstraintTypes.h
+++ b/Engine/Source/Runtime/Experimental/Chaos/Public/Chaos/PBDJointConstraintTypes.h
@@ -113,8 +113,6 @@ namespace Chaos
 	class FPBDJointSettings
 	{
 	public:
-		using FTransformPair = TVector<FRigidTransform3, 2>;
-
 		CHAOS_API FPBDJointSettings();
 
 		// Ensure that settings are consistent and within valid ranges. Should be called
diff --git a/Engine/Source/Runtime/Experimental/Chaos/Public/Chaos/PBDJointConstraints.h b/Engine/Source/Runtime/Experimental/Chaos/Public/Chaos/PBDJointConstraints.h
--- a/Engine/Source/Runtime/Experimental/Chaos/Public/Chaos/PBDJointConstraints.h
+++ b/Engine/Source/Runtime/Experimental/Chaos/Public/Chaos/PBDJointConstraints.h
@@ -153,7 +153,6 @@ namespace Chaos
 
 		using FConstraintContainerHandle = FPBDJointConstraintHandle;
 		using FConstraintHandleAllocator = TConstraintHandleAllocator<FPBDJointConstraints>;
-		using FTransformPair = TVector<FRigidTransform3, 2>;
 		using FHandles = TArray<FConstraintContainerHandle*>;
 
 		CHAOS_API FPBDJointConstraints();
diff --git a/Engine/Source/Runtime/Experimental/ChaosCore/Public/Chaos/Core.h b/Engine/Source/Runtime/Experimental/ChaosCore/Public/Chaos/Core.h
--- a/Engine/Source/Runtime/Experimental/ChaosCore/Public/Chaos/Core.h
+++ b/Engine/Source/Runtime/Experimental/ChaosCore/Public/Chaos/Core.h
@@ -27,6 +27,7 @@ namespace Chaos
 	using FVec3f = TVector<FRealSingle, 3>;
 	using FRotation3f = TRotation<FRealSingle, 3>;
 	using FRigidTransform3f = TRigidTransform<FRealSingle, 3>;
+	using FTransformPair = TVector<FRigidTransform3, 2>;
 
 	using FAABB3f = TAABB<FRealSingle, 3>;
 
