Dynamically grow selected internal buffers when the pre-allocated number of replicated objects is reached.

diff --git a/Engine/Source/Runtime/Engine/Classes/Engine/NetDriver.h b/Engine/Source/Runtime/Engine/Classes/Engine/NetDriver.h
--- a/Engine/Source/Runtime/Engine/Classes/Engine/NetDriver.h
+++ b/Engine/Source/Runtime/Engine/Classes/Engine/NetDriver.h
@@ -634,6 +634,18 @@ struct FNetDriverReplicationSystemConfig
 	/** Override the max object count when running as a server. If 0 use the default system value. */
 	UPROPERTY()
 	uint32 MaxReplicatedObjectServerCount = 0;
+
+	/** Override the number of pre-allocated objects when running as a client. */
+	UPROPERTY()
+	uint32 PreAllocatedReplicatedObjectClientCount = 0;
+
+	/** Override the number of pre-allocated objects when running as a server. */
+	UPROPERTY()
+	uint32 PreAllocatedReplicatedObjectServerCount = 0;
+
+	/** Override the number of pre-allocated objects in FReplicationWriter on the client. */
+	UPROPERTY()
+	uint32 MaxReplicatedWriterObjectClientCount = 0;
 	
 	/** Override the max compressed object count. If 0 use the default system value. */
 	UPROPERTY()
diff --git a/Engine/Source/Runtime/Engine/Private/NetDriver.cpp b/Engine/Source/Runtime/Engine/Private/NetDriver.cpp
--- a/Engine/Source/Runtime/Engine/Private/NetDriver.cpp
+++ b/Engine/Source/Runtime/Engine/Private/NetDriver.cpp
@@ -257,6 +257,13 @@ namespace UE::Net::Private
 			{
 				OutParams.MaxReplicatedObjectCount = ReplicationSystemConfig.MaxReplicatedObjectServerCount;
 			}
+
+			if (ReplicationSystemConfig.PreAllocatedReplicatedObjectServerCount >= 0)
+			{
+				OutParams.PreAllocatedReplicatedObjectCount = ReplicationSystemConfig.PreAllocatedReplicatedObjectServerCount;
+			}
+
+			OutParams.MaxReplicatedWriterObjectCount = OutParams.MaxReplicatedObjectCount;
 		}
 		else
 		{
@@ -264,6 +271,16 @@ namespace UE::Net::Private
 			{
 				OutParams.MaxReplicatedObjectCount = ReplicationSystemConfig.MaxReplicatedObjectClientCount;
 			}
+
+			if (ReplicationSystemConfig.PreAllocatedReplicatedObjectClientCount >= 0)
+			{
+				OutParams.PreAllocatedReplicatedObjectCount = ReplicationSystemConfig.PreAllocatedReplicatedObjectClientCount;
+			}
+
+			if (ReplicationSystemConfig.MaxReplicatedWriterObjectClientCount != 0)
+			{
+				OutParams.MaxReplicatedWriterObjectCount = ReplicationSystemConfig.MaxReplicatedWriterObjectClientCount;
+			}
 		}
 
 		if (ReplicationSystemConfig.MaxDeltaCompressedObjectCount != 0)
diff --git a/Engine/Source/Runtime/Experimental/Iris/Core/Private/Iris/ReplicationSystem/Filtering/NetObjectFilter.cpp b/Engine/Source/Runtime/Experimental/Iris/Core/Private/Iris/ReplicationSystem/Filtering/NetObjectFilter.cpp
--- a/Engine/Source/Runtime/Experimental/Iris/Core/Private/Iris/ReplicationSystem/Filtering/NetObjectFilter.cpp
+++ b/Engine/Source/Runtime/Experimental/Iris/Core/Private/Iris/ReplicationSystem/Filtering/NetObjectFilter.cpp
@@ -6,7 +6,6 @@
 FNetObjectFilteringParams::FNetObjectFilteringParams(const UE::Net::FNetBitArrayView InFilteredObjects)
 : FilteredObjects(InFilteredObjects)
 , FilteringInfos(nullptr)
-, StateBuffers(nullptr)
 , ConnectionId(0)
 {
 }
diff --git a/Engine/Source/Runtime/Experimental/Iris/Core/Private/Iris/ReplicationSystem/Filtering/ReplicationFiltering.cpp b/Engine/Source/Runtime/Experimental/Iris/Core/Private/Iris/ReplicationSystem/Filtering/ReplicationFiltering.cpp
--- a/Engine/Source/Runtime/Experimental/Iris/Core/Private/Iris/ReplicationSystem/Filtering/ReplicationFiltering.cpp
+++ b/Engine/Source/Runtime/Experimental/Iris/Core/Private/Iris/ReplicationSystem/Filtering/ReplicationFiltering.cpp
@@ -457,7 +457,7 @@ bool FReplicationFiltering::SetFilter(FInternalNetRefIndex ObjectIndex, FNetObje
 	}
 
 	const FNetRefHandleManager::FReplicatedObjectData& ObjectData = NetRefHandleManager->GetReplicatedObjectDataNoCheck(ObjectIndex);
-	const TArray<uint8*>& ReplicatedObjectsStateBuffers = NetRefHandleManager->GetReplicatedObjectStateBuffers();
+	const TNetChunkedArray<uint8*>& ReplicatedObjectsStateBuffers = NetRefHandleManager->GetReplicatedObjectStateBuffers();
 	// Let subobjects be filtered like their owners.
 	if (bWantsToUseDynamicFilter && (ObjectData.SubObjectRootIndex != FNetRefHandleManager::InvalidInternalIndex))
 	{
@@ -478,7 +478,7 @@ bool FReplicationFiltering::SetFilter(FInternalNetRefIndex ObjectIndex, FNetObje
 	{
 		FNetObjectFilteringInfo& NetObjectFilteringInfo = this->NetObjectFilteringInfos[ObjIndex];
 		NetObjectFilteringInfo = {};
-		FNetObjectFilterAddObjectParams AddParams = { NetObjectFilteringInfo, ObjectData.InstanceProtocol, ObjectData.Protocol, ReplicatedObjectsStateBuffers.GetData()[ObjIndex] };
+		FNetObjectFilterAddObjectParams AddParams = { NetObjectFilteringInfo, ObjectData.InstanceProtocol, ObjectData.Protocol, ReplicatedObjectsStateBuffers[ObjIndex] };
 		FFilterInfo& FilterInfo = this->DynamicFilterInfos[FilterIndex];
 		if (FilterInfo.Filter->AddObject(ObjIndex, AddParams))
 		{
@@ -1308,7 +1308,6 @@ void FReplicationFiltering::UpdateDynamicFiltering(ENetFilterType FilterType)
 	const uint32* SubObjectsData = NetRefHandleManager->GetSubObjectInternalIndices().GetData();
 	const uint32* DependentObjectsData = NetRefHandleManager->GetDependentObjectInternalIndices().GetData();
 	const uint32* ObjectsRequiringDynamicFilterUpdateData = ObjectsRequiringDynamicFilterUpdate.GetData();
-	uint8* const* ObjectsStateBuffers = FilterType == ENetFilterType::PostPoll_FragmentBased ? NetRefHandleManager->GetReplicatedObjectStateBuffers().GetData() : nullptr;
 
 	uint32* ConnectionIds = static_cast<uint32*>(FMemory_Alloca(ValidConnections.GetNumBits() * sizeof(uint32)));
 	uint32 ConnectionCount = 0;
@@ -1347,7 +1346,6 @@ void FReplicationFiltering::UpdateDynamicFiltering(ENetFilterType FilterType)
 				FNetObjectFilteringParams FilteringParams(MakeNetBitArrayView(Info.FilteredObjects));
 				FilteringParams.OutAllowedObjects = AllowedObjects;
 				FilteringParams.FilteringInfos = NetObjectFilteringInfos.GetData();
-				FilteringParams.StateBuffers = ObjectsStateBuffers;
 				FilteringParams.ConnectionId = ConnId;
 				FilteringParams.View = Connections->GetReplicationView(ConnId);
 
@@ -1526,11 +1524,6 @@ void FReplicationFiltering::BatchNotifyFiltersOfDirtyObjects(FUpdateDirtyObjects
 	FNetObjectFilterUpdateParams UpdateParameters;
 	UpdateParameters.FilteringInfos = NetObjectFilteringInfos.GetData();
 
-	if (BatchHelper.HasProtocolBuffers())
-	{
-		UpdateParameters.StateBuffers = NetRefHandleManager->GetReplicatedObjectStateBuffers().GetData();
-	}
-
 	for (const FUpdateDirtyObjectsBatchHelper::FPerFilterInfo& PerFilterInfo : BatchHelper.PerFilterInfos)
 	{
 		if (PerFilterInfo.ObjectCount == 0)
diff --git a/Engine/Source/Runtime/Experimental/Iris/Core/Private/Iris/ReplicationSystem/NetRefHandleManager.cpp b/Engine/Source/Runtime/Experimental/Iris/Core/Private/Iris/ReplicationSystem/NetRefHandleManager.cpp
--- a/Engine/Source/Runtime/Experimental/Iris/Core/Private/Iris/ReplicationSystem/NetRefHandleManager.cpp
+++ b/Engine/Source/Runtime/Experimental/Iris/Core/Private/Iris/ReplicationSystem/NetRefHandleManager.cpp
@@ -21,9 +21,11 @@ CSV_DEFINE_CATEGORY(IrisCommon, true);
 namespace UE::Net::Private
 {
 
-FNetRefHandleManager::FNetRefHandleManager(FReplicationProtocolManager& InReplicationProtocolManager, uint32 InReplicationSystemId, uint32 InMaxActiveObjectCount)
+FNetRefHandleManager::FNetRefHandleManager(FReplicationProtocolManager& InReplicationProtocolManager, uint32 InReplicationSystemId, uint32 InMaxActiveObjectCount, uint32 InPreAllocatedObjectCount)
 : ActiveObjectCount(0)
 , MaxActiveObjectCount(InMaxActiveObjectCount)
+, PreAllocatedObjectCount(InPreAllocatedObjectCount > 0 ? InPreAllocatedObjectCount : 1)	// PreAllocatedObjectCount must be a minimum of 1 to account for InvalidInternalIndex.
+, LargestInternalIndex(InPreAllocatedObjectCount > 0 ? InPreAllocatedObjectCount - 1 : 0)	// LargestInternalIndex must be a minimum of 0 to account for InvalidInternalIndex.
 , ReplicationSystemId(InReplicationSystemId)
 , GlobalScopableInternalIndices(MaxActiveObjectCount)
 , ScopeFrameData(MaxActiveObjectCount)
@@ -40,20 +42,57 @@ FNetRefHandleManager::FNetRefHandleManager(FReplicationProtocolManager& InReplic
 , NextDynamicHandleIndex(1)
 , ReplicationProtocolManager(InReplicationProtocolManager)
 {
+	// Ensure that the pre-allocated object count is not greater than the maximum. If it is, just set it to the maximum.
+	if (PreAllocatedObjectCount > MaxActiveObjectCount)
+	{
+		UE_LOG(LogIris, Error, TEXT("PreAllocatedObjectCount (%d) must be smaller or equal to MaxActiveObjectCount (%d)!"), PreAllocatedObjectCount, MaxActiveObjectCount);
+		PreAllocatedObjectCount = MaxActiveObjectCount;
+	}
+
+	UE_LOG(LogIris, Log, TEXT("NetRefHandleManager: Configured with MaxActiveObjectCount=%d, PreAllocatedObjectCount=%d and LargestInternalIndex=%d."), MaxActiveObjectCount, PreAllocatedObjectCount, LargestInternalIndex);
+
 	static_assert(InvalidInternalIndex == 0, "FNetRefHandleManager::InvalidInternalIndex has an unexpected value");
 	// Mark the invalid index as used
 	AssignedInternalIndices.SetBit(0);
 
-	ReplicatedObjectData.SetNumUninitialized(MaxActiveObjectCount);
-	ReplicatedObjectRefCount.SetNumZeroed(MaxActiveObjectCount);
-	ReplicatedObjectStateBuffers.SetNumZeroed(MaxActiveObjectCount);
-	
-	ReplicatedInstances.SetNumZeroed(MaxActiveObjectCount);
+	ReplicatedObjectData = TNetChunkedArray<FReplicatedObjectData>(PreAllocatedObjectCount, EInitMemory::Constructor);
+	ReplicatedObjectRefCount = TNetChunkedArray<uint16>(PreAllocatedObjectCount, EInitMemory::Zero);
+	ReplicatedObjectStateBuffers = TNetChunkedArray<uint8*>(PreAllocatedObjectCount, EInitMemory::Zero);
+
+	ReplicatedInstances = TNetChunkedArray<TObjectPtr<UObject>>(PreAllocatedObjectCount, EInitMemory::Zero);
 
 	// For convenience we initialize ReplicatedObjectData for InvalidInternalIndex so that GetReplicatedObjectDataNoCheck returns something useful.
 	ReplicatedObjectData[InvalidInternalIndex] = FReplicatedObjectData();
 }
 
+void FNetRefHandleManager::GrowBuffersToLargestIndex(uint32 InternalIndex)
+{
+	if (InternalIndex > LargestInternalIndex)
+	{
+		QUICK_SCOPE_CYCLE_COUNTER(STAT_FNetRefHandleManager_GrowBuffersToLargestIndex);
+		CSV_CUSTOM_STAT(IrisCommon, DynamicBufferGrowCount, 1, ECsvCustomStatOp::Accumulate);
+
+		// This call will add the neccessary number of elements and chunks to ReplicatedObjectRefCount
+		// to accomidate InternalIndex. Once this is done, we determine how many more elements could be
+		// added to the array without adding a new chunk and ensure all of the other buffers have this
+		// many elements. 
+		// 
+		// This optimization will reduce the number of calls to AddToIndexUninitialized() and AddToIndexZeroed()
+		// but assumes each buffer is going to have the same number of elements all the time.
+		ReplicatedObjectRefCount.AddToIndexUninitialized(InternalIndex);
+		const int32 LargestIndexInCurrentChunk = ReplicatedObjectRefCount.Capacity() - 1;
+
+		ReplicatedObjectRefCount.AddToIndexUninitialized(LargestIndexInCurrentChunk);
+		ReplicatedObjectStateBuffers.AddToIndexZeroed(LargestIndexInCurrentChunk);
+		ReplicatedInstances.AddToIndexZeroed(LargestIndexInCurrentChunk);
+		ReplicatedObjectData.AddToIndexZeroed(LargestIndexInCurrentChunk);
+		
+		LargestInternalIndex = static_cast<uint32>(LargestIndexInCurrentChunk);
+
+		OnLargestIndexIncreaseDelegate.Broadcast(LargestInternalIndex);
+	}
+}
+
 uint64 FNetRefHandleManager::GetNextNetRefHandleId(uint64 HandleId) const
 {
 	// Since we use the lowest bit in the index to indicate if the handle is static or dynamic we cannot use all bits as the index
@@ -94,6 +133,9 @@ FInternalNetRefIndex FNetRefHandleManager::InternalCreateNetObject(const FNetRef
 	{
 		UE_LOG(LogIris, Verbose, TEXT("FNetRefHandleManager::InternalCreateNetObject ( InternalIndex: %u ) %s"), InternalIndex, *NetRefHandle.ToString());
 
+		// Track the largest internal index and grow internal buffers if neccessary.
+		GrowBuffersToLargestIndex(InternalIndex);
+
 		// Store data;
 		FReplicatedObjectData& Data = ReplicatedObjectData[InternalIndex];
 
@@ -103,7 +145,7 @@ FInternalNetRefIndex FNetRefHandleManager::InternalCreateNetObject(const FNetRef
 		Data.NetHandle = GlobalHandle;
 		Data.Protocol = ReplicationProtocol;
 		Data.InstanceProtocol = nullptr;
-		ReplicatedObjectStateBuffers.GetData()[InternalIndex] = nullptr;
+		ReplicatedObjectStateBuffers[InternalIndex] = nullptr;
 		Data.ReceiveStateBuffer = nullptr;
 		Data.bShouldPropagateChangedStates = 1U;
 
@@ -219,7 +261,7 @@ FNetRefHandle FNetRefHandleManager::CreateNetObject(FNetRefHandle WantedHandle,
 			ConditionalChangeMask.SetAllBits();
 		}
 
-		ReplicatedObjectStateBuffers.GetData()[InternalIndex] = StateBuffer;
+		ReplicatedObjectStateBuffers[InternalIndex] = StateBuffer;
 
 		// Bump protocol refcount
 		ReplicationProtocol->AddRef();
@@ -275,7 +317,7 @@ void FNetRefHandleManager::InternalDestroyNetObject(FInternalNetRefIndex Interna
 
 	UE_LOG(LogIris, Verbose, TEXT("FNetRefHandleManager::InternalDestroyNetObject ( InternalIndex: %u ) %s"), InternalIndex, *Data.RefHandle.ToString());
 
-	uint8* StateBuffer = ReplicatedObjectStateBuffers.GetData()[InternalIndex];
+	uint8* StateBuffer = ReplicatedObjectStateBuffers[InternalIndex];
 	// Free any allocated resources
 	if (EnumHasAnyFlags(Data.Protocol->ProtocolTraits, EReplicationProtocolTraits::HasDynamicState))
 	{
@@ -330,7 +372,7 @@ void FNetRefHandleManager::InternalDestroyNetObject(FInternalNetRefIndex Interna
 	FMemory::Free(Data.ReceiveStateBuffer);
 
 	// Clear pointer to state buffer
-	ReplicatedObjectStateBuffers.GetData()[InternalIndex] = nullptr;
+	ReplicatedObjectStateBuffers[InternalIndex] = nullptr;
 
 	UE_NET_TRACE_NETHANDLE_DESTROYED(Data.RefHandle);
 
@@ -783,7 +825,10 @@ void FNetRefHandleManager::SetShouldPropagateChangedStates(FNetRefHandle Handle,
 
 void FNetRefHandleManager::AddReferencedObjects(FReferenceCollector& Collector)
 {
-	Collector.AddReferencedObjects(ReplicatedInstances);
+	for (TObjectPtr<UObject>& Object : ReplicatedInstances)
+	{
+		Collector.AddReferencedObject(Object);
+	}
 }
 
 uint64 FNetRefHandleManager::MakeNetRefHandleId(uint64 Id, bool bIsStatic)
diff --git a/Engine/Source/Runtime/Experimental/Iris/Core/Private/Iris/ReplicationSystem/NetRefHandleManager.h b/Engine/Source/Runtime/Experimental/Iris/Core/Private/Iris/ReplicationSystem/NetRefHandleManager.h
--- a/Engine/Source/Runtime/Experimental/Iris/Core/Private/Iris/ReplicationSystem/NetRefHandleManager.h
+++ b/Engine/Source/Runtime/Experimental/Iris/Core/Private/Iris/ReplicationSystem/NetRefHandleManager.h
@@ -8,10 +8,12 @@
 #include "Containers/ArrayView.h"
 #include "Net/Core/NetBitArray.h"
 #include "Net/Core/NetHandle/NetHandle.h"
+#include "Iris/Core/NetChunkedArray.h"
 #include "Iris/ReplicationState/ReplicationStateDescriptor.h"
 #include "Iris/ReplicationSystem/NetRefHandle.h"
 #include "Iris/ReplicationSystem/NetDependencyData.h"
 #include "UObject/ObjectPtr.h"
+#include "Delegates/Delegate.h"
 
 class FReferenceCollector;
 namespace UE::Net
@@ -119,7 +121,7 @@ public:
 	typedef TMap<FNetHandle, FInternalNetRefIndex> FNetHandleMap;
 
 public:
-	FNetRefHandleManager(FReplicationProtocolManager& InReplicationProtocolManager, uint32 InReplicationSystemId, uint32 MaxActiveObjects);
+	FNetRefHandleManager(FReplicationProtocolManager& InReplicationProtocolManager, uint32 InReplicationSystemId, uint32 InMaxActiveObjectCount, uint32 InPreAllocatedObjectCount);
 
 	/** Callback triggered at the beginning of PreSendUpdate. Used to sync current frame data. */
 	void OnPreSendUpdate();
@@ -167,9 +169,9 @@ public:
 
 	inline const FReplicatedObjectData& GetReplicatedObjectData(FInternalNetRefIndex InternalIndex) const;
 
-	inline const uint8* GetReplicatedObjectStateBufferNoCheck(FInternalNetRefIndex InternalObjectIndex) const { return ReplicatedObjectStateBuffers.GetData()[InternalObjectIndex]; }
-	inline uint8* GetReplicatedObjectStateBufferNoCheck(FInternalNetRefIndex InternalObjectIndex) { return ReplicatedObjectStateBuffers.GetData()[InternalObjectIndex]; }
-	inline const TArray<uint8*>& GetReplicatedObjectStateBuffers() const { return ReplicatedObjectStateBuffers; }
+	inline const uint8* GetReplicatedObjectStateBufferNoCheck(FInternalNetRefIndex InternalObjectIndex) const { return ReplicatedObjectStateBuffers[InternalObjectIndex]; }
+	inline uint8* GetReplicatedObjectStateBufferNoCheck(FInternalNetRefIndex InternalObjectIndex) { return ReplicatedObjectStateBuffers[InternalObjectIndex]; }
+	inline const TNetChunkedArray<uint8*>& GetReplicatedObjectStateBuffers() const { return ReplicatedObjectStateBuffers; }
 
 	/** Verify a handle against internal handle. A handle is valid if it matches internal storage. */
 	inline bool IsValidNetRefHandle(FNetRefHandle Handle) const;
@@ -244,6 +246,7 @@ public:
 
 	uint32 GetMaxActiveObjectCount() const { return MaxActiveObjectCount; }
 	uint32 GetActiveObjectCount() const { return ActiveObjectCount; }
+	uint32 GetPreAllocatedObjectCount() const { return PreAllocatedObjectCount; }
 
 	// We do refcount objects tracked by each connection in order to know when it is safe to reuse an InternalIndex
 	void AddNetObjectRef(FInternalNetRefIndex InternalIndex) { ++ReplicatedObjectRefCount[InternalIndex]; }
@@ -272,7 +275,7 @@ public:
 	UObject* GetReplicatedObjectInstance(FInternalNetRefIndex ObjectIndex) const { return ReplicatedInstances[ObjectIndex]; }
 
 	/** Get the array of all held object pointers */
-	const TArray<UObject*>& GetReplicatedInstances() const { return ReplicatedInstances; }
+	const TNetChunkedArray<TObjectPtr<UObject>>& GetReplicatedInstances() const { return ReplicatedInstances; }
 
 	void AddReferencedObjects(FReferenceCollector& Collector);
 
@@ -287,6 +290,14 @@ public:
 
 	/** Return a string to identify the object linked to an index in logs */
 	FString PrintObjectFromIndex(FInternalNetRefIndex ObjectIndex) const;
+	
+	DECLARE_MULTICAST_DELEGATE_OneParam(FOnLargestIndexIncrease, uint32 LargestIndex);
+
+	/** Return a delegate that will notify when the largest internal index has increased. */
+	FOnLargestIndexIncrease& GetLargestIndexIncreaseDelegate() { return OnLargestIndexIncreaseDelegate; };
+
+	/** Return the largest internal index that has been used. */
+	uint32 GetLargestInternalIndex() const { return LargestInternalIndex; };
 
 public:
 
@@ -323,12 +334,24 @@ private:
 	void InternalRemoveDependentObject(FInternalNetRefIndex ParentInternalIndex, FInternalNetRefIndex DependentInternalIndex, ERemoveDependentObjectFlags Flags = ERemoveDependentObjectFlags::All);
 	void InternalRemoveDependentObject(FInternalNetRefIndex DependentInternalIndex);
 
+	// Ensure that all buffers that depend on the largest internal index are sized correctly.
+	void GrowBuffersToLargestIndex(uint32 InternalIndex);
+
 	// The current replicated object count
 	uint32 ActiveObjectCount;
 
 	// Max allowed replicated object count
 	uint32 MaxActiveObjectCount;
 
+	// The number of pre-allocated objects used by internal buffers.
+	uint32 PreAllocatedObjectCount;
+
+	// The largest internal index value that has been used.
+	uint32 LargestInternalIndex;
+
+	// A delegate that is triggered when the largest encountered internal index increases.
+	FOnLargestIndexIncrease OnLargestIndexIncreaseDelegate;
+
 	uint32 ReplicationSystemId;
 
 	FRefHandleMap RefHandleToInternalIndex;
@@ -399,16 +422,16 @@ private:
 	TArray<FInternalNetRefIndex> PendingDestroyInternalIndices;
 
 	// Just an array containing data about our replicated objects
-	TArray<FReplicatedObjectData> ReplicatedObjectData;
+	TNetChunkedArray<FReplicatedObjectData> ReplicatedObjectData;
 
 	// Pointers to state buffers for all replicated objects
-	TArray<uint8*> ReplicatedObjectStateBuffers;
+	TNetChunkedArray<uint8*> ReplicatedObjectStateBuffers;
 
 	// Refcounts for all tracked objects
-	TArray<uint16> ReplicatedObjectRefCount;
+	TNetChunkedArray<uint16> ReplicatedObjectRefCount;
 
 	// Raw pointers to all bound instances
-	TArray<TObjectPtr<UObject>> ReplicatedInstances;
+	TNetChunkedArray<TObjectPtr<UObject>> ReplicatedInstances;
 
 	// Assign handles
 	uint64 NextStaticHandleIndex;
diff --git a/Engine/Source/Runtime/Experimental/Iris/Core/Private/Iris/ReplicationSystem/ObjectReplicationBridge.cpp b/Engine/Source/Runtime/Experimental/Iris/Core/Private/Iris/ReplicationSystem/ObjectReplicationBridge.cpp
--- a/Engine/Source/Runtime/Experimental/Iris/Core/Private/Iris/ReplicationSystem/ObjectReplicationBridge.cpp
+++ b/Engine/Source/Runtime/Experimental/Iris/Core/Private/Iris/ReplicationSystem/ObjectReplicationBridge.cpp
@@ -817,7 +817,7 @@ void UObjectReplicationBridge::PruneStaleObjects()
 	GarbageCollectionAffectedObjects = ObjectsWithObjectReferences;
 
 	FNetRefHandleManager& LocalNetRefHandleManager = GetReplicationSystem()->GetReplicationSystemInternal()->GetNetRefHandleManager();
-	const TArray<UObject*>& ReplicatedInstances = LocalNetRefHandleManager.GetReplicatedInstances();
+	const TNetChunkedArray<TObjectPtr<UObject>>& ReplicatedInstances = LocalNetRefHandleManager.GetReplicatedInstances();
 
 	TArray<FNetRefHandle> StaleObjects;
 
@@ -1138,7 +1138,7 @@ void UObjectReplicationBridge::UpdateInstancesWorldLocation()
 	FReplicationSystemInternal* ReplicationSystemInternal = GetReplicationSystem()->GetReplicationSystemInternal();
 	const FNetRefHandleManager& LocalNetRefHandleManager = ReplicationSystemInternal->GetNetRefHandleManager();
 	FWorldLocations& WorldLocations = ReplicationSystemInternal->GetWorldLocations();
-	const TArray<UObject*>& ReplicatedInstances = LocalNetRefHandleManager.GetReplicatedInstances();
+	const TNetChunkedArray<TObjectPtr<UObject>>& ReplicatedInstances = LocalNetRefHandleManager.GetReplicatedInstances();
 
 	// Retrieve the world location for instances that supports it. Only dirty objects are considered.
 	auto UpdateInstanceWorldLocation = [this, &ReplicatedInstances, &LocalNetRefHandleManager, &WorldLocations](uint32 InternalObjectIndex)
diff --git a/Engine/Source/Runtime/Experimental/Iris/Core/Private/Iris/ReplicationSystem/Polling/ObjectPoller.h b/Engine/Source/Runtime/Experimental/Iris/Core/Private/Iris/ReplicationSystem/Polling/ObjectPoller.h
--- a/Engine/Source/Runtime/Experimental/Iris/Core/Private/Iris/ReplicationSystem/Polling/ObjectPoller.h
+++ b/Engine/Source/Runtime/Experimental/Iris/Core/Private/Iris/ReplicationSystem/Polling/ObjectPoller.h
@@ -3,6 +3,7 @@
 #pragma once
 
 #include "Net/Core/NetBitArray.h"
+#include "Iris/Core/NetChunkedArray.h"
 
 // Forward declarations
 class UObjectReplicationBridge;
@@ -79,7 +80,7 @@ private:
 
 	FNetRefHandleManager& LocalNetRefHandleManager;
 	FNetStatsContext* NetStatsContext = nullptr;
-	const TArray<UObject*>& ReplicatedInstances;
+	const TNetChunkedArray<TObjectPtr<UObject>>& ReplicatedInstances;
 
 	const FNetBitArrayView AccumulatedDirtyObjects;
 
diff --git a/Engine/Source/Runtime/Experimental/Iris/Core/Private/Iris/ReplicationSystem/Prioritization/ReplicationPrioritization.cpp b/Engine/Source/Runtime/Experimental/Iris/Core/Private/Iris/ReplicationSystem/Prioritization/ReplicationPrioritization.cpp
--- a/Engine/Source/Runtime/Experimental/Iris/Core/Private/Iris/ReplicationSystem/Prioritization/ReplicationPrioritization.cpp
+++ b/Engine/Source/Runtime/Experimental/Iris/Core/Private/Iris/ReplicationSystem/Prioritization/ReplicationPrioritization.cpp
@@ -679,7 +679,6 @@ void FReplicationPrioritization::PrioritizeForConnection(uint32 ConnId, FPriorit
 	{
 		PrioParameters.Priorities = ConnInfo.Priorities.GetData();
 		PrioParameters.PrioritizationInfos = NetObjectPrioritizationInfos.GetData();
-		PrioParameters.StateBuffers = NetRefHandleManager->GetReplicatedObjectStateBuffers().GetData();
 		PrioParameters.ConnectionId = ConnId;
 		PrioParameters.View = Connections->GetReplicationView(ConnId);
 	}
@@ -802,7 +801,6 @@ void FReplicationPrioritization::BatchNotifyPrioritizersOfDirtyObjects(FUpdateDi
 	BatchHelper.PrepareBatch(ObjectIndices, ObjectCount, ObjectIndexToPrioritizer.GetData());
 
 	FNetObjectPrioritizerUpdateParams UpdateParameters;
-	UpdateParameters.StateBuffers = NetRefHandleManager->GetReplicatedObjectStateBuffers().GetData();
 	UpdateParameters.PrioritizationInfos = NetObjectPrioritizationInfos.GetData();
 
 	for (const FUpdateDirtyObjectsBatchHelper::FPerPrioritizerInfo& PerPrioritizerInfo : BatchHelper.PerPrioritizerInfos)
diff --git a/Engine/Source/Runtime/Experimental/Iris/Core/Private/Iris/ReplicationSystem/ReplicationSystem.cpp b/Engine/Source/Runtime/Experimental/Iris/Core/Private/Iris/ReplicationSystem/ReplicationSystem.cpp
--- a/Engine/Source/Runtime/Experimental/Iris/Core/Private/Iris/ReplicationSystem/ReplicationSystem.cpp
+++ b/Engine/Source/Runtime/Experimental/Iris/Core/Private/Iris/ReplicationSystem/ReplicationSystem.cpp
@@ -57,7 +57,7 @@ public:
 
 	explicit FReplicationSystemImpl(UReplicationSystem* InReplicationSystem, const UReplicationSystem::FReplicationSystemParams& Params)
 	: ReplicationSystem(InReplicationSystem)
-	, ReplicationSystemInternal(FReplicationSystemInternalInitParams({ InReplicationSystem->GetId(), Params.MaxReplicatedObjectCount }))
+	, ReplicationSystemInternal(FReplicationSystemInternalInitParams({ InReplicationSystem->GetId(), Params.MaxReplicatedObjectCount, Params.PreAllocatedReplicatedObjectCount, Params.MaxReplicatedWriterObjectCount }))
 	{
 	}
 
@@ -497,6 +497,9 @@ public:
 			Params.PacketSendWindowSize = 256;
 			Params.ConnectionId = ConnectionId;
 			Params.MaxActiveReplicatedObjectCount = ReplicationSystemInternal.GetNetRefHandleManager().GetMaxActiveObjectCount();
+			Params.PreAllocatedReplicatedObjectCount = ReplicationSystemInternal.GetNetRefHandleManager().GetPreAllocatedObjectCount();
+			Params.MaxReplicatedWriterObjectCount = ReplicationSystemInternal.GetInitParams().MaxReplicatedWriterObjectCount;
+
 			/** 
 			  * Currently we expect all objects to be replicated from server to client.
 			  * That means we will have to support sending attachments such as RPCs from
diff --git a/Engine/Source/Runtime/Experimental/Iris/Core/Private/Iris/ReplicationSystem/ReplicationSystemInternal.h b/Engine/Source/Runtime/Experimental/Iris/Core/Private/Iris/ReplicationSystem/ReplicationSystemInternal.h
--- a/Engine/Source/Runtime/Experimental/Iris/Core/Private/Iris/ReplicationSystem/ReplicationSystemInternal.h
+++ b/Engine/Source/Runtime/Experimental/Iris/Core/Private/Iris/ReplicationSystem/ReplicationSystemInternal.h
@@ -35,13 +35,16 @@ struct FReplicationSystemInternalInitParams
 {
 	uint32 ReplicationSystemId;
 	uint32 MaxReplicatedObjectCount;
+	uint32 PreAllocatedReplicatedObjectCount;
+	uint32 MaxReplicatedWriterObjectCount;
 };
 
 class FReplicationSystemInternal
 {
 public:
 	explicit FReplicationSystemInternal(const FReplicationSystemInternalInitParams& Params)
-	: NetRefHandleManager(ReplicationProtocolManager, Params.ReplicationSystemId, Params.MaxReplicatedObjectCount)
+	: NetRefHandleManager(ReplicationProtocolManager, Params.ReplicationSystemId, Params.MaxReplicatedObjectCount, Params.PreAllocatedReplicatedObjectCount)
+	, InternalInitParams(Params)
 	, DirtyNetObjectTracker()
 	, ReplicationBridge(nullptr)
 	, IrisObjectReferencePackageMap(nullptr)
@@ -101,6 +104,8 @@ public:
 
 	FNetTypeStats& GetNetTypeStats() { return TypeStats; }
 
+	FReplicationSystemInternalInitParams& GetInitParams() { return InternalInitParams; }
+
 	FNetSendStats& GetSendStats()
 	{ 
 		return SendStats;
@@ -118,6 +123,7 @@ public:
 private:
 	FReplicationProtocolManager ReplicationProtocolManager;
 	FNetRefHandleManager NetRefHandleManager;
+	FReplicationSystemInternalInitParams InternalInitParams;
 	FDirtyNetObjectTracker DirtyNetObjectTracker;
 	FReplicationStateStorage ReplicationStateStorage;
 	FReplicationStateDescriptorRegistry ReplicationStateDescriptorRegistry;
diff --git a/Engine/Source/Runtime/Experimental/Iris/Core/Private/Iris/ReplicationSystem/ReplicationTypes.h b/Engine/Source/Runtime/Experimental/Iris/Core/Private/Iris/ReplicationSystem/ReplicationTypes.h
--- a/Engine/Source/Runtime/Experimental/Iris/Core/Private/Iris/ReplicationSystem/ReplicationTypes.h
+++ b/Engine/Source/Runtime/Experimental/Iris/Core/Private/Iris/ReplicationSystem/ReplicationTypes.h
@@ -11,6 +11,8 @@ namespace UE::Net::Private
 struct FReplicationParameters
 {
 	uint32 MaxActiveReplicatedObjectCount;	
+	uint32 PreAllocatedReplicatedObjectCount;
+	uint32 MaxReplicatedWriterObjectCount;
 	uint32 PacketSendWindowSize;
 	uint32 ConnectionId;
 	UReplicationSystem* ReplicationSystem = nullptr;
diff --git a/Engine/Source/Runtime/Experimental/Iris/Core/Private/Iris/ReplicationSystem/ReplicationWriter.cpp b/Engine/Source/Runtime/Experimental/Iris/Core/Private/Iris/ReplicationSystem/ReplicationWriter.cpp
--- a/Engine/Source/Runtime/Experimental/Iris/Core/Private/Iris/ReplicationSystem/ReplicationWriter.cpp
+++ b/Engine/Source/Runtime/Experimental/Iris/Core/Private/Iris/ReplicationSystem/ReplicationWriter.cpp
@@ -178,11 +178,11 @@ struct TReplicationRecordHelper
 	typedef FReplicationRecord::FRecordInfoList FRecordInfoList;
 	typedef FReplicationWriter::EReplicatedObjectState EReplicatedObjectState;
 
-	FReplicationInfo* ReplicationInfos;
-	FRecordInfoList* ReplicationInfosRecordInfoLists;
+	TArray<FReplicationInfo>& ReplicationInfos;
+	TArray<FReplicationRecord::FRecordInfoList>& ReplicationInfosRecordInfoLists;
 	FReplicationRecord* ReplicationRecord;
 
-	TReplicationRecordHelper(FReplicationInfo* InReplicationInfos, FRecordInfoList* InReplicationInfosRecordInfoLists, FReplicationRecord* InReplicationRecordRecord)
+	TReplicationRecordHelper(TArray<FReplicationInfo>& InReplicationInfos, TArray<FReplicationRecord::FRecordInfoList>& InReplicationInfosRecordInfoLists, FReplicationRecord* InReplicationRecordRecord)
 	: ReplicationInfos(InReplicationInfos)
 	, ReplicationInfosRecordInfoLists(InReplicationInfosRecordInfoLists)
 	, ReplicationRecord(InReplicationRecordRecord)
@@ -271,6 +271,9 @@ static bool s_ValidateReplicationRecord(const FReplicationRecord* ReplicationRec
 
 FReplicationWriter::~FReplicationWriter()
 {
+	// NOTE: Currently disabled because FReplicationWriter until the performance impact of TNetChunkedArray can be measured on the server.
+	//NetRefHandleManager->GetLargestIndexIncreaseDelegate().Remove(OnLargestIndexIncreaseHandle);
+
 	DiscardAllRecords();
 
 	// Freeing the huge object queue needs to be done before calling StopAllReplication() in order to be able to free any changemask allocations.
@@ -358,6 +361,8 @@ void FReplicationWriter::Init(const FReplicationParameters& InParameters)
 	// Store copy of parameters
 	Parameters = InParameters;
 
+	UE_LOG(LogIris, Log, TEXT("ReplicationWriter: Configured with MaxActiveReplicatedObjectCount=%d, PreallocatedObjectCount=%d and MaxReplicatedWriterObjectCount=%d."), Parameters.MaxActiveReplicatedObjectCount, Parameters.PreAllocatedReplicatedObjectCount, Parameters.MaxReplicatedWriterObjectCount);
+
 	// Cache internal systems
 	ReplicationSystemInternal = Parameters.ReplicationSystem->GetReplicationSystemInternal();
 	NetRefHandleManager = &ReplicationSystemInternal->GetNetRefHandleManager();
@@ -372,9 +377,13 @@ void FReplicationWriter::Init(const FReplicationParameters& InParameters)
 	NetTypeStats = &ReplicationSystemInternal->GetNetTypeStats();
 
 	// Init book keeping
-	ReplicatedObjects.SetNumZeroed(Parameters.MaxActiveReplicatedObjectCount);
-	ReplicatedObjectsRecordInfoLists.SetNumZeroed(Parameters.MaxActiveReplicatedObjectCount);
-	SchedulingPriorities.SetNumZeroed(Parameters.MaxActiveReplicatedObjectCount);
+	const int32 PreAllocatedBufferSize = Parameters.MaxReplicatedWriterObjectCount;
+	ReplicatedObjects.SetNumUninitialized(PreAllocatedBufferSize);
+	ReplicatedObjectsRecordInfoLists.SetNumUninitialized(PreAllocatedBufferSize);
+	SchedulingPriorities.SetNumUninitialized(PreAllocatedBufferSize);
+
+	// NOTE: Currently disabled because FReplicationWriter until the performance impact of TNetChunkedArray can be measured on the server.
+	//OnLargestIndexIncreaseHandle = NetRefHandleManager->GetLargestIndexIncreaseDelegate().AddRaw(this, &FReplicationWriter::OnLargestIndexIncrease);
 
 	ObjectsPendingDestroy.Init(Parameters.MaxActiveReplicatedObjectCount);
 	ObjectsWithDirtyChanges.Init(Parameters.MaxActiveReplicatedObjectCount);
@@ -866,7 +875,8 @@ const FNetBitArray& FReplicationWriter::GetObjectsRequiringPriorityUpdate() cons
 void FReplicationWriter::UpdatePriorities(const float* UpdatedPriorities)
 {
 	IRIS_PROFILER_SCOPE(FReplicationWriter_UpdatePriorities);
-	auto UpdatePriority = [LocalPriorities = SchedulingPriorities.GetData(), UpdatedPriorities](uint32 Index)
+
+	auto UpdatePriority = [&LocalPriorities = SchedulingPriorities, UpdatedPriorities](uint32 Index)
 	{
 		LocalPriorities[Index] += UpdatedPriorities[Index];
 	};
@@ -874,7 +884,7 @@ void FReplicationWriter::UpdatePriorities(const float* UpdatedPriorities)
 	ObjectsWithDirtyChanges.ForAllSetBits(UpdatePriority);
 }
 
-void FReplicationWriter::ScheduleDependentObjects(uint32 Index, float ParentPriority, float* LocalPriorities, FScheduleObjectInfo* ScheduledObjectIndices, uint32& OutScheduledObjectCount)
+void FReplicationWriter::ScheduleDependentObjects(uint32 Index, float ParentPriority, TArray<float>& LocalPriorities, FScheduleObjectInfo* ScheduledObjectIndices, uint32& OutScheduledObjectCount)
 {
 	const float DependentObjectPriorityBump = UE_KINDA_SMALL_NUMBER;
 
@@ -917,7 +927,6 @@ uint32 FReplicationWriter::ScheduleObjects(FScheduleObjectInfo* OutScheduledObje
 	IRIS_PROFILER_SCOPE(FReplicationWriter_ScheduleObjects);
 
 	uint32 ScheduledObjectCount = 0;
-	float* LocalPriorities = SchedulingPriorities.GetData();
 
 	FScheduleObjectInfo* ScheduledObjectIndices = OutScheduledObjectIndices;
 
@@ -927,9 +936,9 @@ uint32 FReplicationWriter::ScheduleObjects(FScheduleObjectInfo* OutScheduledObje
 	const FNetBitArray& UpdatedObjects = ObjectsWithDirtyChanges;
 	const FNetBitArray& SubObjects = NetRefHandleManager->GetSubObjectInternalIndices();
 
-	auto FillIndexListFunc = [&LocalPriorities, &ScheduledObjectIndices, &ScheduledObjectCount, this](uint32 Index)
+	auto FillIndexListFunc = [&ScheduledObjectIndices, &ScheduledObjectCount, this](uint32 Index)
 	{
-		const float UpdatedPriority = LocalPriorities[Index];
+		const float UpdatedPriority = SchedulingPriorities[Index];
 
 		FScheduleObjectInfo& ScheduledObjectInfo = ScheduledObjectIndices[ScheduledObjectCount];
 		ScheduledObjectInfo.Index = Index;
@@ -942,7 +951,7 @@ uint32 FReplicationWriter::ScheduleObjects(FScheduleObjectInfo* OutScheduledObje
 			// If we have dependent objects that needs to replicate before parent we need to schedule them as well.
 			if (NetRefHandleManager->GetObjectsWithDependentObjectsInternalIndices().GetBit(Index))
 			{
-				ScheduleDependentObjects(Index, UpdatedPriority, LocalPriorities, ScheduledObjectIndices, ScheduledObjectCount);
+				ScheduleDependentObjects(Index, UpdatedPriority, SchedulingPriorities, ScheduledObjectIndices, ScheduledObjectCount);
 			}
 		}
 	};
@@ -1536,7 +1545,7 @@ void FReplicationWriter::ProcessDeliveryNotification(EPacketDeliveryStatus Packe
 
 	if (RecordCount > 0)
 	{
-		TReplicationRecordHelper Helper(ReplicatedObjects.GetData(), ReplicatedObjectsRecordInfoLists.GetData(), &ReplicationRecord);
+		TReplicationRecordHelper Helper(ReplicatedObjects, ReplicatedObjectsRecordInfoLists, &ReplicationRecord);
 
 		if (PacketDeliveryStatus == EPacketDeliveryStatus::Delivered)
 		{
@@ -3494,7 +3503,7 @@ bool FReplicationWriter::IsWriteObjectSuccess(EWriteObjectStatus Status) const
 
 void FReplicationWriter::DiscardAllRecords()
 {
-	TReplicationRecordHelper Helper(ReplicatedObjects.GetData(), ReplicatedObjectsRecordInfoLists.GetData(), &ReplicationRecord);
+	TReplicationRecordHelper Helper(ReplicatedObjects, ReplicatedObjectsRecordInfoLists, &ReplicationRecord);
 
 	const uint32 RecordCount = ReplicationRecord.GetRecordCount();
 	for (uint32 RecordIt = 0, RecordEndIt = RecordCount; RecordIt != RecordEndIt; ++RecordIt)
@@ -3524,9 +3533,11 @@ void FReplicationWriter::StopAllReplication()
 	ReplicatedObjects[ObjectIndexForOOBAttachment].State = (uint8)EReplicatedObjectState::Invalid;
 
 	// We cannot tell for sure which objects need processing so we check them all.
-	const FReplicationInfo* FirstInfo = ReplicatedObjects.GetData();
-	for (const FReplicationInfo& Info : ReplicatedObjects)
+	int32 ReplicatedObjectsCount = ReplicatedObjects.Num();
+	for (int32 InternalIndex = 0; InternalIndex < ReplicatedObjectsCount; InternalIndex++)
 	{
+		const FReplicationInfo& Info = GetReplicationInfo(InternalIndex);
+
 		if (Info.GetState() == EReplicatedObjectState::Invalid)
 		{
 			continue;
@@ -3536,7 +3547,6 @@ void FReplicationWriter::StopAllReplication()
 		FChangeMaskStorageOrPointer::Free(Info.ChangeMaskOrPtr, Info.ChangeMaskBitCount, s_DefaultChangeMaskAllocator);
 
 		// Release object reference
-		const FInternalNetRefIndex InternalIndex = static_cast<FInternalNetRefIndex>(&Info - FirstInfo);
 		NetRefHandleManager->ReleaseNetObjectRef(InternalIndex);
 	}
 }
@@ -3555,6 +3565,13 @@ void FReplicationWriter::MarkObjectDirty(FInternalNetRefIndex InternalIndex, con
 	ObjectsWithDirtyChanges.SetBit(InternalIndex);
 }
 
+void FReplicationWriter::OnLargestIndexIncrease(uint32 InternalIndex)
+{
+	ReplicatedObjects.SetNumZeroed(InternalIndex);
+	ReplicatedObjectsRecordInfoLists.SetNumUninitialized(InternalIndex);
+	SchedulingPriorities.SetNumUninitialized(InternalIndex);
+}
+
 
 FReplicationWriter::FHugeObjectContext::FHugeObjectContext() = default;
 
diff --git a/Engine/Source/Runtime/Experimental/Iris/Core/Private/Iris/ReplicationSystem/ReplicationWriter.h b/Engine/Source/Runtime/Experimental/Iris/Core/Private/Iris/ReplicationSystem/ReplicationWriter.h
--- a/Engine/Source/Runtime/Experimental/Iris/Core/Private/Iris/ReplicationSystem/ReplicationWriter.h
+++ b/Engine/Source/Runtime/Experimental/Iris/Core/Private/Iris/ReplicationSystem/ReplicationWriter.h
@@ -435,7 +435,7 @@ private:
 
 	void CommitBatchRecord(const FBatchRecord& BatchRecord);
 
-	void ScheduleDependentObjects(uint32 Index, float ParentPriority, float* LocalPriorities, FScheduleObjectInfo* ScheduledObjectIndices, uint32& OutScheduledObjectCount);
+	void ScheduleDependentObjects(uint32 Index, float ParentPriority, TArray<float>& LocalPriorities, FScheduleObjectInfo* ScheduledObjectIndices, uint32& OutScheduledObjectCount);
 
 	uint32 ScheduleObjects(FScheduleObjectInfo* ScheduledObjectIndices);
 	
@@ -527,6 +527,9 @@ private:
 
 	void MarkObjectDirty(FInternalNetRefIndex InternalIndex, const char* Caller);
 
+	/* Grow internal buffers to support a given index. */
+	void OnLargestIndexIncrease(uint32 InternalIndex);
+
 private:
 	// Replication parameters
 	FReplicationParameters Parameters;
@@ -555,6 +558,8 @@ private:
 	// Handles logic for all attachments to objects.
 	FNetObjectAttachmentsWriter Attachments;
 
+	FDelegateHandle OnLargestIndexIncreaseHandle;
+
 	// Cached internal systems
 	FReplicationSystemInternal* ReplicationSystemInternal = nullptr;
 	FNetRefHandleManager* NetRefHandleManager = nullptr;
diff --git a/Engine/Source/Runtime/Experimental/Iris/Core/Public/Iris/Core/NetChunkedArray.h b/Engine/Source/Runtime/Experimental/Iris/Core/Public/Iris/Core/NetChunkedArray.h
new file mode 100644
--- /dev/null
+++ b/Engine/Source/Runtime/Experimental/Iris/Core/Public/Iris/Core/NetChunkedArray.h
@@ -0,0 +1,140 @@
+// Copyright Epic Games, Inc. All Rights Reserved.
+
+#pragma once
+
+#include "Containers/ChunkedArray.h"
+
+namespace UE::Net
+{
+
+// Configure what happens when new chunks are allocated.
+enum class EInitMemory : uint8
+{
+	Zero,		// Memory will be filled with zeros.
+	Constructor	// Memory will be initialized by the elements constructor.
+};
+
+/**
+ * A variation of TChunkArray that is optimized for use in the Iris networking system. The initial chunks
+ * created during construction are placed in a contiguous block of memory to promote locality of reference.
+ */
+template<typename InElementType, uint32 ElementsPerChunk = 100, typename AllocatorType = FDefaultAllocator>
+class TNetChunkedArray : public TChunkedArray<InElementType, sizeof(InElementType)* ElementsPerChunk, AllocatorType>
+{
+public:
+	
+	using Super = TChunkedArray<InElementType, sizeof(InElementType) * ElementsPerChunk, AllocatorType>;
+
+	TNetChunkedArray(int32 InNumElements = 0, EInitMemory InitMemory = EInitMemory::Constructor)
+	{
+		this->NumElements = InNumElements;
+
+		// Compute the number of chunks needed.
+		const int32 NumChunks = (this->NumElements + this->NumElementsPerChunk - 1) / this->NumElementsPerChunk;
+
+		// Initial blocks should come from a single block of memory.
+		this->Chunks.Empty(NumChunks);
+		typename Super::FChunk* StartChunks = new typename Super::FChunk[NumChunks];
+		for (int32 ChunkIndex = 0; ChunkIndex < NumChunks; ChunkIndex++)
+		{
+			typename Super::FChunk* CurrentChunk = (StartChunks + ChunkIndex);
+			if (InitMemory == EInitMemory::Zero)
+			{
+				FMemory::Memset<typename Super::FChunk>(*CurrentChunk, 0);
+			}
+
+			this->Chunks.Add(CurrentChunk);
+		}
+
+		NumPreAllocatedChunks = NumChunks;
+	}
+
+	~TNetChunkedArray()
+	{
+		InvalidatePreAllocatedChunks();
+	}
+
+	int32 NumChunks() const
+	{
+		return this->Chunks.Num();
+	}
+
+	/** 
+	 * Return the maximum number of elements the array can hold before having to add another chunk.
+	 */
+	int32 Capacity() const
+	{
+		return this->Chunks.Num() * this->NumElementsPerChunk;
+	}
+
+	void Empty(int32 Slack = 0)
+	{
+		checkf(false, TEXT("This function is not supported"));
+	}
+
+	void Reset(int32 NewSize = 0)
+	{
+		checkf(false, TEXT("This function is not supported"));
+	}
+
+	/** 
+	 * Add elements to the array so that an index can be successfully addressed, leaving 
+	 * any new element's memory unitialized or initialized by the element's constructor.
+	 * 
+	 * @param Index The index that must be addressable.
+	 */
+	void AddToIndexUninitialized(int32 Index)
+	{
+		if (Index >= this->NumElements)
+		{
+			int32 NewElementCount = (Index - this->NumElements) + 1;
+			this->Add(NewElementCount);
+		}
+	}
+
+	/**
+	 * Add elements to the array so that an index can be successfully addressed, zeroing
+	 * out the memory for each new element.
+	 *
+	 * @param Index The index that must be addressable.
+	 */
+	void AddToIndexZeroed(int32 Index)
+	{
+		const int32 OldChunkCount = this->Chunks.Num();
+		AddToIndexUninitialized(Index);
+		for (int32 ChunkIndex = OldChunkCount; ChunkIndex < this->Chunks.Num(); ChunkIndex++)
+		{
+			FMemory::Memset<typename Super::FChunk>(*this->Chunks.GetData()[ChunkIndex], 0);
+		}
+	}
+
+protected:
+	/* The number of preallocated chunks. */
+	int32 NumPreAllocatedChunks = 0;
+
+private:
+	/**
+	 * Invalidate the pre-allocated chunk pointers and free the associated block of memory.
+	 *
+	 * This function must be called before TChunkedArray frees memory (e.g. the constructor) because
+	 * it assumes that each chunk is a seperate allocation and is unaware of the single block
+	 * of memory used by the pre-allocated chunks.
+	 */
+	void InvalidatePreAllocatedChunks()
+	{
+		// Invalid pre-allocated chunks from the chunks array.
+		typename Super::FChunk* FirstChunk = nullptr;
+		for (int32 ChunkIndex = 0; ChunkIndex < NumPreAllocatedChunks; ChunkIndex++)
+		{
+			FirstChunk = (FirstChunk == nullptr) ? this->Chunks.GetData()[ChunkIndex] : FirstChunk;
+			this->Chunks.GetData()[ChunkIndex] = nullptr;
+		}
+
+		// The first chunk points to the beginning of the block of memory used by the pre-allocated chunks.
+		delete[] FirstChunk;
+
+		NumPreAllocatedChunks = 0;
+	}
+};
+
+}
diff --git a/Engine/Source/Runtime/Experimental/Iris/Core/Public/Iris/ReplicationSystem/Filtering/NetObjectFilter.h b/Engine/Source/Runtime/Experimental/Iris/Core/Public/Iris/ReplicationSystem/Filtering/NetObjectFilter.h
--- a/Engine/Source/Runtime/Experimental/Iris/Core/Public/Iris/ReplicationSystem/Filtering/NetObjectFilter.h
+++ b/Engine/Source/Runtime/Experimental/Iris/Core/Public/Iris/ReplicationSystem/Filtering/NetObjectFilter.h
@@ -56,9 +56,6 @@ struct FNetObjectFilteringParams
 	/** FilteringInfos for all objects. Index using the set bit indices in FilteredObjects. */
 	const FNetObjectFilteringInfo* FilteringInfos = nullptr;
 
-	/** State buffers for all objects. Index using the set bit indices in FilteredObjects. */
-	uint8 *const* StateBuffers = nullptr;
-
 	/** ID of the connection that the filtering applies to. */
 	uint32 ConnectionId;
 
@@ -196,12 +193,6 @@ struct FNetObjectFilterUpdateParams
 	* NOTE: Only for filters of type FragmentBased; null for Raw types
 	*/
 	UE::Net::FReplicationInstanceProtocol const* const* InstanceProtocols = nullptr;
-
-	/** 
-	* State buffers for all objects. Index using ObjectIndices[0..ObjectCount-1]. 
-	* NOTE: Only for filters of type FragmentBased; null for Raw types
-	*/
-	uint8* const* StateBuffers = nullptr;
 };
 
 UCLASS(Abstract)
diff --git a/Engine/Source/Runtime/Experimental/Iris/Core/Public/Iris/ReplicationSystem/Prioritization/NetObjectPrioritizer.h b/Engine/Source/Runtime/Experimental/Iris/Core/Public/Iris/ReplicationSystem/Prioritization/NetObjectPrioritizer.h
--- a/Engine/Source/Runtime/Experimental/Iris/Core/Public/Iris/ReplicationSystem/Prioritization/NetObjectPrioritizer.h
+++ b/Engine/Source/Runtime/Experimental/Iris/Core/Public/Iris/ReplicationSystem/Prioritization/NetObjectPrioritizer.h
@@ -41,9 +41,6 @@ struct FNetObjectPrioritizationParams
 	/** PrioritizationInfos for all objects. Index using ObjectIndices[0..ObjectCount-1]. */
 	const FNetObjectPrioritizationInfo* PrioritizationInfos;
 
-	/** State buffers for all objects. Index using ObjectIndices[0..ObjectCount-1]. */
-	uint8 *const* StateBuffers;
-
 	/** ID of the connection that objects are prioritized for. */
 	uint32 ConnectionId;
 
@@ -124,9 +121,6 @@ struct FNetObjectPrioritizerUpdateParams
 	/** InstanceProtocols for updated objects. Index using 0..ObjectCount-1. */
 	UE::Net::FReplicationInstanceProtocol const*const* InstanceProtocols;
 
-	/** State buffers for all objects. Index using ObjectIndices[0..ObjectCount-1]. */
-	uint8 *const* StateBuffers;
-
 	/** Infos for all objects. Index using ObjectIndices[0..ObjectCount-1]. */
 	FNetObjectPrioritizationInfo* PrioritizationInfos;
 };
diff --git a/Engine/Source/Runtime/Experimental/Iris/Core/Public/Iris/ReplicationSystem/ReplicationSystem.h b/Engine/Source/Runtime/Experimental/Iris/Core/Public/Iris/ReplicationSystem/ReplicationSystem.h
--- a/Engine/Source/Runtime/Experimental/Iris/Core/Public/Iris/ReplicationSystem/ReplicationSystem.h
+++ b/Engine/Source/Runtime/Experimental/Iris/Core/Public/Iris/ReplicationSystem/ReplicationSystem.h
@@ -63,6 +63,8 @@ public:
 		//$IRIS TODO: These need documentation
 		UReplicationBridge* ReplicationBridge = nullptr;
 		uint32 MaxReplicatedObjectCount = 65535U;
+		uint32 PreAllocatedReplicatedObjectCount = 65535U;
+		uint32 MaxReplicatedWriterObjectCount = 65535U;
 		uint32 MaxDeltaCompressedObjectCount = 2048U;
 		uint32 MaxNetObjectGroupCount = 2048U;
 		bool bIsServer = false;
