Add a FilterConfig option to force apply a dynamic filter (ignoring bOnlyRelevantToOwner).

diff --git a/Engine/Source/Runtime/Engine/Private/Net/Iris/ReplicationSystem/ActorReplicationBridge.cpp b/Engine/Source/Runtime/Engine/Private/Net/Iris/ReplicationSystem/ActorReplicationBridge.cpp
--- a/Engine/Source/Runtime/Engine/Private/Net/Iris/ReplicationSystem/ActorReplicationBridge.cpp
+++ b/Engine/Source/Runtime/Engine/Private/Net/Iris/ReplicationSystem/ActorReplicationBridge.cpp
@@ -255,7 +255,13 @@ UE::Net::FNetRefHandle UActorReplicationBridge::BeginReplication(AActor* Actor,
 	{
 		if (Actor->bOnlyRelevantToOwner && !Actor->bAlwaysRelevant)
 		{
-			GetReplicationSystem()->SetFilter(ActorRefHandle, ToOwnerFilterHandle);
+			// Only apply owner filter if we haven't force enabled a dynamic filter.
+			constexpr bool bRequireForceEnabled = true;
+			FNetObjectFilterHandle FilterHandle = GetDynamicFilter(Actor->GetClass(), bRequireForceEnabled);
+			if (FilterHandle == InvalidNetObjectFilterHandle)
+			{
+				GetReplicationSystem()->SetFilter(ActorRefHandle, ToOwnerFilterHandle);
+			}
 		}
 	}
 
diff --git a/Engine/Source/Runtime/Experimental/Iris/Core/Private/Iris/ReplicationSystem/ObjectReplicationBridge.cpp b/Engine/Source/Runtime/Experimental/Iris/Core/Private/Iris/ReplicationSystem/ObjectReplicationBridge.cpp
--- a/Engine/Source/Runtime/Experimental/Iris/Core/Private/Iris/ReplicationSystem/ObjectReplicationBridge.cpp
+++ b/Engine/Source/Runtime/Experimental/Iris/Core/Private/Iris/ReplicationSystem/ObjectReplicationBridge.cpp
@@ -424,7 +424,8 @@ void UObjectReplicationBridge::AssignDynamicFilter(UObject* Instance, const FCre
 	}
 	else if (Params.bUseClassConfigDynamicFilter)
 	{
-		FilterHandle = GetDynamicFilter(Instance->GetClass());
+		constexpr bool bRequireForceEnabled = false;
+		FilterHandle = GetDynamicFilter(Instance->GetClass(), bRequireForceEnabled);
 	}
 
 	if (FilterHandle != InvalidNetObjectFilterHandle)
@@ -1525,12 +1526,12 @@ void UObjectReplicationBridge::SetClassDynamicFilterConfig(FName ClassPathName,
 
 	if (UE_LOG_ACTIVE(LogIrisFilterConfig, Log))
 	{
-		if (const UE::Net::FNetObjectFilterHandle* OldFilter = ClassesWithDynamicFilter.Find(ClassPathName))
+		if (const FClassFilterInfo* OldFilterInfo = ClassesWithDynamicFilter.Find(ClassPathName))
 		{
-			if (*OldFilter != FilterHandle)
+			if (OldFilterInfo->FilterHandle != FilterHandle)
 			{
 				UE_LOG(LogIrisFilterConfig, Log, TEXT("SetClassDynamicFilterConfig assigned %s to use filter %s. Previously using filter %s."),
-					*ClassPathName.ToString(), *(ReplicationSystem->GetFilterName(FilterHandle).ToString()), *(ReplicationSystem->GetFilterName(*OldFilter).ToString()));
+					*ClassPathName.ToString(), *(ReplicationSystem->GetFilterName(FilterHandle).ToString()), *(ReplicationSystem->GetFilterName(OldFilterInfo->FilterHandle).ToString()));
 			}
 			else
 			{
@@ -1543,7 +1544,10 @@ void UObjectReplicationBridge::SetClassDynamicFilterConfig(FName ClassPathName,
 		}
 	}
 
-	ClassesWithDynamicFilter.Add(ClassPathName, FilterHandle);
+	FClassFilterInfo FilterInfo;
+	FilterInfo.FilterHandle = FilterHandle;
+	FilterInfo.bForceEnable = false;
+	ClassesWithDynamicFilter.Add(ClassPathName, FilterInfo);
 }
 
 void UObjectReplicationBridge::SetClassDynamicFilterConfig(FName ClassPathName, FName FilterName)
@@ -1571,7 +1575,7 @@ void UObjectReplicationBridge::SetClassDynamicFilterConfig(FName ClassPathName,
 	}
 }
 
-UE::Net::FNetObjectFilterHandle UObjectReplicationBridge::GetDynamicFilter(const UClass* Class)
+UE::Net::FNetObjectFilterHandle UObjectReplicationBridge::GetDynamicFilter(const UClass* Class, bool bRequireForceEnabled)
 {
 	using namespace UE::Net;
 
@@ -1580,9 +1584,10 @@ UE::Net::FNetObjectFilterHandle UObjectReplicationBridge::GetDynamicFilter(const
 		const FName ClassName = GetConfigClassPathName(Class);
 
 		// Try exact match first.
-		if (FNetObjectFilterHandle* FilterHandlePtr = ClassesWithDynamicFilter.Find(ClassName))
+		if (FClassFilterInfo* FilterInfoPtr = ClassesWithDynamicFilter.Find(ClassName))
 		{
-			return *FilterHandlePtr;
+			const bool bUseFilter = !bRequireForceEnabled || FilterInfoPtr->bForceEnable;
+			return bUseFilter ? FilterInfoPtr->FilterHandle : InvalidNetObjectFilterHandle;
 		}
 
 		/**
@@ -1595,13 +1600,15 @@ UE::Net::FNetObjectFilterHandle UObjectReplicationBridge::GetDynamicFilter(const
 			const FName SuperClassName = GetConfigClassPathName(SuperClass);
 
 			// Try to get exact match first.
-			if (FNetObjectFilterHandle* FilterHandlePtr = ClassesWithDynamicFilter.Find(SuperClassName))
+			if (FClassFilterInfo* FilterInfoPtr = ClassesWithDynamicFilter.Find(SuperClassName))
 			{
 				if (ShouldSubclassUseSameFilterFunction(SuperClass, Class))
 				{
-					const FNetObjectFilterHandle FilterHandle = *FilterHandlePtr;
-					ClassesWithDynamicFilter.Add(ClassName, FilterHandle);
-					return FilterHandle;
+					FClassFilterInfo FilterInfo = *FilterInfoPtr;
+					ClassesWithDynamicFilter.Add(ClassName, FilterInfo);
+
+					const bool bUseFilter = !bRequireForceEnabled || FilterInfo.bForceEnable;
+					return bUseFilter ? FilterInfo.FilterHandle : InvalidNetObjectFilterHandle;
 				}
 
 				// Here's a good place to put a line of code and set a breakpoint to debug inheritance issues.
@@ -1611,9 +1618,11 @@ UE::Net::FNetObjectFilterHandle UObjectReplicationBridge::GetDynamicFilter(const
 		}
 
 		// Either super class wasn't found or it wasn't considered equal. Let's add a new filter mapping.
-		const FNetObjectFilterHandle FilterHandle = ShouldUseDefaultSpatialFilterFunction(Class) ? DefaultSpatialFilterHandle : InvalidNetObjectFilterHandle;
-		ClassesWithDynamicFilter.Add(ClassName, FilterHandle);
-		return FilterHandle;
+		FClassFilterInfo FilterInfo;
+		FilterInfo.FilterHandle = ShouldUseDefaultSpatialFilterFunction(Class) ? DefaultSpatialFilterHandle : InvalidNetObjectFilterHandle;
+		FilterInfo.bForceEnable = false;
+		ClassesWithDynamicFilter.Add(ClassName, FilterInfo);
+		return FilterInfo.FilterHandle;
 	}
 
 	/**
@@ -1718,8 +1727,10 @@ void UObjectReplicationBridge::LoadConfig()
 	{
 		for (const FObjectReplicationBridgeFilterConfig& FilterConfig : BridgeConfig->GetFilterConfigs())
 		{
-			const UE::Net::FNetObjectFilterHandle FilterHandle = ReplicationSystem->GetFilterHandle(FilterConfig.DynamicFilterName);
-			ClassesWithDynamicFilter.Add(FilterConfig.ClassName, FilterHandle);
+			FClassFilterInfo FilterInfo;
+			FilterInfo.FilterHandle = ReplicationSystem->GetFilterHandle(FilterConfig.DynamicFilterName);
+			FilterInfo.bForceEnable = FilterConfig.bForceEnableOnAllInstances;
+			ClassesWithDynamicFilter.Add(FilterConfig.ClassName, FilterInfo);
 		}
 	}
 
diff --git a/Engine/Source/Runtime/Experimental/Iris/Core/Private/Iris/ReplicationSystem/ObjectReplicationBridgeDebugging.cpp b/Engine/Source/Runtime/Experimental/Iris/Core/Private/Iris/ReplicationSystem/ObjectReplicationBridgeDebugging.cpp
--- a/Engine/Source/Runtime/Experimental/Iris/Core/Private/Iris/ReplicationSystem/ObjectReplicationBridgeDebugging.cpp
+++ b/Engine/Source/Runtime/Experimental/Iris/Core/Private/Iris/ReplicationSystem/ObjectReplicationBridgeDebugging.cpp
@@ -155,15 +155,15 @@ void UObjectReplicationBridge::PrintDynamicFilterClassConfig() const
 	UE_LOG(LogIrisFilterConfig, Display, TEXT(""));
 	UE_LOG(LogIrisFilterConfig, Display, TEXT("Default Dynamic Filter Class Config:"));
 	{
-		TMap<FName, UE::Net::FNetObjectFilterHandle> SortedClassConfig = ClassesWithDynamicFilter;
+		TMap<FName, FClassFilterInfo> SortedClassConfig = ClassesWithDynamicFilter;
 
 		SortedClassConfig.KeyStableSort([](FName lhs, FName rhs){return lhs.Compare(rhs) < 0;});
 		for (auto MapIt = SortedClassConfig.CreateConstIterator(); MapIt; ++MapIt)
 		{
 			const FName ClassName = MapIt.Key();
-			const UE::Net::FNetObjectFilterHandle FilterHandle = MapIt.Value();
+			const FClassFilterInfo FilterInfo = MapIt.Value();
 
-			UE_LOG(LogIrisFilterConfig, Display, TEXT("\t%s -> %s"), *ClassName.ToString(), *RepSystem->GetFilterName(FilterHandle).ToString());
+			UE_LOG(LogIrisFilterConfig, Display, TEXT("\t%s -> %s"), *ClassName.ToString(), *RepSystem->GetFilterName(FilterInfo.FilterHandle).ToString());
 		}
 	}
 }
diff --git a/Engine/Source/Runtime/Experimental/Iris/Core/Public/Iris/ReplicationSystem/ObjectReplicationBridge.h b/Engine/Source/Runtime/Experimental/Iris/Core/Public/Iris/ReplicationSystem/ObjectReplicationBridge.h
--- a/Engine/Source/Runtime/Experimental/Iris/Core/Public/Iris/ReplicationSystem/ObjectReplicationBridge.h
+++ b/Engine/Source/Runtime/Experimental/Iris/Core/Public/Iris/ReplicationSystem/ObjectReplicationBridge.h
@@ -350,9 +350,11 @@ private:
 	/** Tries to load the classes used in poll period overrides. */
 	void FindClassesInPollPeriodOverrides();
 
+protected:
 	/** Retrieves the dynamic filter to set for the given class. Will return an invalid handle if no dynamic filter should be set. */
-	UE::Net::FNetObjectFilterHandle GetDynamicFilter(const UClass* Class);
+	UE::Net::FNetObjectFilterHandle GetDynamicFilter(const UClass* Class, bool bRequireForceEnabled);
 
+private:
 	/** Retrieves the prioritizer to set for the given class. If bRequireForceEnabled the config needs to have bForceEnableOnAllInstances set in order for this method to return the configured prioritizer. Returns an invalid handle if no prioritizer should be set. */
 	UE::Net::FNetObjectPrioritizerHandle GetPrioritizer(const UClass* Class, bool bRequireForceEnabled);
 
@@ -385,7 +387,14 @@ private:
 	struct FClassPrioritizerInfo
 	{
 		UE::Net::FNetObjectPrioritizerHandle PrioritizerHandle;
-		uint32 bForceEnable : 1;
+		bool bForceEnable = false;
+	};
+	
+	// Filtering
+	struct FClassFilterInfo
+	{
+		UE::Net::FNetObjectFilterHandle FilterHandle;
+		bool bForceEnable = false;
 	};
 
 	// Polling
@@ -411,7 +420,7 @@ private:
 	// Filter mapping
 	//$IRIS TODO: Look into improving this class map by balancing runtime speed (implicit addition of new classes) with runtime modifications of base classes.
     //		      Right-now any changes to say APawn's filter will not be read if a APlayerPawn entry was created based on the APawn entry.
-	TMap<FName, UE::Net::FNetObjectFilterHandle> ClassesWithDynamicFilter;
+	TMap<FName, FClassFilterInfo> ClassesWithDynamicFilter;
 	TFunction<bool(const UClass*)> ShouldUseDefaultSpatialFilterFunction;
 	TFunction<bool(const UClass*,const UClass*)> ShouldSubclassUseSameFilterFunction;
 
diff --git a/Engine/Source/Runtime/Experimental/Iris/Core/Public/Iris/ReplicationSystem/ObjectReplicationBridgeConfig.h b/Engine/Source/Runtime/Experimental/Iris/Core/Public/Iris/ReplicationSystem/ObjectReplicationBridgeConfig.h
--- a/Engine/Source/Runtime/Experimental/Iris/Core/Public/Iris/ReplicationSystem/ObjectReplicationBridgeConfig.h
+++ b/Engine/Source/Runtime/Experimental/Iris/Core/Public/Iris/ReplicationSystem/ObjectReplicationBridgeConfig.h
@@ -46,6 +46,10 @@ struct FObjectReplicationBridgeFilterConfig
 	/** The name of the filter to set on the class instances. */
 	UPROPERTY()
 	FName DynamicFilterName;
+
+	/** Whether this filter should be used for all instances of this class and subclasses, regardless of bAlwaysRelevant and bOnlyRelevantToOwner settings on instance. */
+	UPROPERTY()
+	bool bForceEnableOnAllInstances = false;
 };
 
 USTRUCT()
