Add ClusterUnion SetMassOverrideInKg support.

Caveat: Cluster Unions recalculate their mass following destruction events, so to maintain a custom mass you also need to hook into the cluster's add/remove component events.

diff --git a/Engine/Source/Runtime/Engine/Classes/PhysicsEngine/ClusterUnionComponent.h b/Engine/Source/Runtime/Engine/Classes/PhysicsEngine/ClusterUnionComponent.h
--- a/Engine/Source/Runtime/Engine/Classes/PhysicsEngine/ClusterUnionComponent.h
+++ b/Engine/Source/Runtime/Engine/Classes/PhysicsEngine/ClusterUnionComponent.h
@@ -313,6 +313,10 @@ public:
 	ENGINE_API virtual void WakeAllRigidBodies() override;
 	ENGINE_API virtual bool IsAnyRigidBodyAwake() override;
 
+	// Set the cluster union total mass. NOTE: if the cluster breaks the mass will be recalculated from the
+	// remaining components. Use OnComponentAddedEvent and OnComponentRemovedEvent to trap this and re-apply custom masses.
+	ENGINE_API virtual void SetMassOverrideInKg(FName BoneName, float MassInKg, bool bOverrideMass) override;
+
 	// Multi-trace/sweep functions that only make sense in the context of a cluster union.
 	ENGINE_API bool LineTraceComponent(TArray<FHitResult>& OutHit, const FVector Start, const FVector End, ECollisionChannel TraceChannel, const struct FCollisionQueryParams& Params, const struct FCollisionResponseParams& ResponseParams, const struct FCollisionObjectQueryParams& ObjectParams);
 	ENGINE_API bool SweepComponent(TArray<FHitResult>& OutHit, const FVector Start, const FVector End, const FQuat& ShapeWorldRotation, const FPhysicsGeometry& Geometry, ECollisionChannel TraceChannel, const struct FCollisionQueryParams& Params, const struct FCollisionResponseParams& ResponseParams, const struct FCollisionObjectQueryParams& ObjectParams);
diff --git a/Engine/Source/Runtime/Engine/Private/PhysicsEngine/ClusterUnionComponent.cpp b/Engine/Source/Runtime/Engine/Private/PhysicsEngine/ClusterUnionComponent.cpp
--- a/Engine/Source/Runtime/Engine/Private/PhysicsEngine/ClusterUnionComponent.cpp
+++ b/Engine/Source/Runtime/Engine/Private/PhysicsEngine/ClusterUnionComponent.cpp
@@ -1401,6 +1401,14 @@ bool UClusterUnionComponent::IsAnyRigidBodyAwake()
 	return !Interface->AreAllSleeping(PhysicsObjects);
 }
 
+void UClusterUnionComponent::SetMassOverrideInKg(FName BoneName, float MassInKg, bool bOverrideMass)
+{
+	if (PhysicsProxy)
+	{
+		PhysicsProxy->SetMass_External(MassInKg);
+	}
+}
+
 DECLARE_CYCLE_STAT(TEXT("UClusterUnionComponent::LineTraceComponentMulti"), STAT_ClusterUnionComponentLineTraceComponentMulti, STATGROUP_Chaos);
 bool UClusterUnionComponent::LineTraceComponent(TArray<FHitResult>& OutHit, const FVector Start, const FVector End, ECollisionChannel TraceChannel, const struct FCollisionQueryParams& Params, const struct FCollisionResponseParams& ResponseParams, const struct FCollisionObjectQueryParams& ObjectParams)
 {
diff --git a/Engine/Source/Runtime/Experimental/Chaos/Private/PhysicsProxy/ClusterUnionPhysicsProxy.cpp b/Engine/Source/Runtime/Experimental/Chaos/Private/PhysicsProxy/ClusterUnionPhysicsProxy.cpp
--- a/Engine/Source/Runtime/Experimental/Chaos/Private/PhysicsProxy/ClusterUnionPhysicsProxy.cpp
+++ b/Engine/Source/Runtime/Experimental/Chaos/Private/PhysicsProxy/ClusterUnionPhysicsProxy.cpp
@@ -297,6 +297,52 @@ namespace Chaos
 		);
 	}
 
+	void FClusterUnionPhysicsProxy::SetMass_External(Chaos::FReal InMass)
+	{
+		if (!Solver || !ensure(Particle_External))
+		{
+			return;
+		}
+
+		Chaos::FRealSingle Mass = Chaos::FRealSingle(InMass);
+		Chaos::FRealSingle InvMass = 0;
+		Chaos::FVec3f Inertia = FVec3f(0);
+		Chaos::FVec3f InvInertia = FVec3f(0);
+		if (Mass > 0)
+		{
+			InvMass = 1.0f / Mass;
+
+			Chaos::FRealSingle OldMass = Chaos::FRealSingle(Particle_External->M());
+			if (OldMass > 0)
+			{
+				const Chaos::FRealSingle MassScale = Mass / OldMass;
+				Inertia = Particle_External->I() * MassScale;
+				InvInertia = Particle_External->InvI() / MassScale;
+			}
+		}
+
+		if (Particle_External)
+		{
+			Particle_External->SetM(Mass);
+			Particle_External->SetInvM(InvMass);
+			Particle_External->SetI(Inertia);
+			Particle_External->SetInvI(InvInertia);
+		}
+
+		Solver->EnqueueCommandImmediate(
+			[this, Mass, InvMass, Inertia, InvInertia]()
+			{
+				if (Particle_Internal)
+				{
+					Particle_Internal->SetM(Mass);
+					Particle_Internal->SetInvM(InvMass);
+					Particle_Internal->SetI(Inertia);
+					Particle_Internal->SetInvI(InvInertia);
+				}
+			}
+		);
+	}
+
 	void FClusterUnionPhysicsProxy::RemoveShapes_External(const TArray<FPBDRigidParticle*>& ShapeParticles)
     {
 		RemoveParticlesFromClusterUnionGeometry(Particle_External.Get(), ShapeParticles, GeometryChildParticles_External);
diff --git a/Engine/Source/Runtime/Experimental/Chaos/Public/PhysicsProxy/ClusterUnionPhysicsProxy.h b/Engine/Source/Runtime/Experimental/Chaos/Public/PhysicsProxy/ClusterUnionPhysicsProxy.h
--- a/Engine/Source/Runtime/Experimental/Chaos/Public/PhysicsProxy/ClusterUnionPhysicsProxy.h
+++ b/Engine/Source/Runtime/Experimental/Chaos/Public/PhysicsProxy/ClusterUnionPhysicsProxy.h
@@ -85,6 +85,10 @@ namespace Chaos
 		// Explicitly wake the physics object. The wakes the physics object if it is sleeping and resets any sleep state (whether awake or sleeping)
 		CHAOS_API void Wake_External();
 
+		// Set the cluster mass.
+		// NOTE: When a cluster breaks its mass will be recalculated from the remaining children on the PT, effectively undoing this work.
+		CHAOS_API void SetMass_External(Chaos::FReal Mass);
+
 		// Set GT geometry - this is only for smoothing over any changes until the PT syncs back to the GT.
 		CHAOS_API void SetGeometry_External(const Chaos::FImplicitObjectPtr& Geometry, const TArray<FPBDRigidParticle*>& ShapeParticles);
 
