Expose fully decayed flag on geometry collection components

Gameplay systems may want to bind to the decayed event on a geometry collection, but if it's already fully decayed the event will never fire. In order to handle this we can expose the flag so gameplay can check first before binding to the event

diff --git a/Engine/Source/Runtime/Experimental/GeometryCollectionEngine/Private/GeometryCollection/GeometryCollectionComponent.cpp b/Engine/Source/Runtime/Experimental/GeometryCollectionEngine/Private/GeometryCollection/GeometryCollectionComponent.cpp
--- a/Engine/Source/Runtime/Experimental/GeometryCollectionEngine/Private/GeometryCollection/GeometryCollectionComponent.cpp
+++ b/Engine/Source/Runtime/Experimental/GeometryCollectionEngine/Private/GeometryCollection/GeometryCollectionComponent.cpp
@@ -1367,6 +1367,11 @@ void UGeometryCollectionComponent::AddTorqueInRadians(FVector Torque, FName Bone
 	DispatchFieldCommand(Command);
 }
 
+bool UGeometryCollectionComponent::IsFullyDecayed() const
+{
+	return bAlreadyFullyDecayed;
+}
+
 void UGeometryCollectionComponent::DispatchBreakEvent(const FChaosBreakEvent& Event)
 {
 	// native
diff --git a/Engine/Source/Runtime/Experimental/GeometryCollectionEngine/Public/GeometryCollection/GeometryCollectionComponent.h b/Engine/Source/Runtime/Experimental/GeometryCollectionEngine/Public/GeometryCollection/GeometryCollectionComponent.h
--- a/Engine/Source/Runtime/Experimental/GeometryCollectionEngine/Public/GeometryCollection/GeometryCollectionComponent.h
+++ b/Engine/Source/Runtime/Experimental/GeometryCollectionEngine/Public/GeometryCollection/GeometryCollectionComponent.h
@@ -1249,6 +1249,8 @@ public:
 	FOnGeometryCollectionRootMovedEvent OnRootMovedEvent;
 	FOnGeometryCollectionRootMovedNativeEvent OnRootMovedNativeEvent;
 
+	GEOMETRYCOLLECTIONENGINE_API bool IsFullyDecayed() const;
+
 	GEOMETRYCOLLECTIONENGINE_API void DispatchBreakEvent(const FChaosBreakEvent& Event);
 
 	GEOMETRYCOLLECTIONENGINE_API void DispatchRemovalEvent(const FChaosRemovalEvent& Event);
