Chaos - fix mid/narrow-phase costs for shapes that are currently QueryOnly but were previously QueryAndPhysics.

Calling SetCollisionEnabled(QueryOnly) on a body clears the bSimEnabled flag on the collision settings, but this was not being factored into the broadphase pair rejection logic, so we would end up generating unnecessary midphase objects. This is especially problems for particles with many shapes because we don't pre-filter the set of potentially colliding shape pairs.

diff --git a/Engine/Source/Runtime/Experimental/Chaos/Public/Chaos/ParticleHandle.h b/Engine/Source/Runtime/Experimental/Chaos/Public/Chaos/ParticleHandle.h
--- a/Engine/Source/Runtime/Experimental/Chaos/Public/Chaos/ParticleHandle.h
+++ b/Engine/Source/Runtime/Experimental/Chaos/Public/Chaos/ParticleHandle.h
@@ -283,13 +283,15 @@ public:
 
 public:
 	/**
-	* compute the aggregated query collision filter from all associated shapes
+	* compute the aggregated query collision filter from all associated shapes.
+	* NOTE: may not initialize OutQueryFilterData, so it must be cleared before calling if necessary (e.g., when changing overlapping channels)
 	**/
 	template <typename TParticle>
 	static void ComputeParticleQueryFilterDataFromShapes(const TParticle& Particle, FCollisionFilterData& OutQueryFilterData);
 
 	/**
 	* compute the aggregated sim collision filter from all associated shapes
+	* NOTE: may not initialize OutQueryFilterData, so it must be cleared before calling if necessary (e.g., when changing blocking channels)
 	**/
 	template <typename TParticle>
 	static void ComputeParticleSimFilterDataFromShapes(const TParticle& Particle, FCollisionFilterData& OutSimFilterData);
@@ -3811,11 +3813,14 @@ template <typename TParticle>
 	const auto& Shapes = Particle.ShapesArray();
 	for (const auto& Shape : Shapes)
 	{
-		const FCollisionFilterData& ShapeQueryData = Shape->GetQueryData();
-		OutQueryFilterData.Word0 |= ShapeQueryData.Word0;
-		OutQueryFilterData.Word1 |= ShapeQueryData.Word1;
-		OutQueryFilterData.Word2 |= ShapeQueryData.Word2;
-		OutQueryFilterData.Word3 |= ShapeQueryData.Word3;
+		if (Shape->GetQueryEnabled())
+		{
+			const FCollisionFilterData& ShapeQueryData = Shape->GetQueryData();
+			OutQueryFilterData.Word0 |= ShapeQueryData.Word0;
+			OutQueryFilterData.Word1 |= ShapeQueryData.Word1;
+			OutQueryFilterData.Word2 |= ShapeQueryData.Word2;
+			OutQueryFilterData.Word3 |= ShapeQueryData.Word3;
+		}
 	}
 }
 
@@ -3825,11 +3830,14 @@ template <typename TParticle>
 	const auto& Shapes = Particle.ShapesArray();
 	for (const auto& Shape : Shapes)
 	{
-		const FCollisionFilterData& ShapeSimData = Shape->GetSimData();
-		OutSimFilterData.Word0 |= ShapeSimData.Word0;
-		OutSimFilterData.Word1 |= ShapeSimData.Word1;
-		OutSimFilterData.Word2 |= ShapeSimData.Word2;
-		OutSimFilterData.Word3 |= ShapeSimData.Word3;
+		if (Shape->GetSimEnabled())
+		{
+			const FCollisionFilterData& ShapeSimData = Shape->GetSimData();
+			OutSimFilterData.Word0 |= ShapeSimData.Word0;
+			OutSimFilterData.Word1 |= ShapeSimData.Word1;
+			OutSimFilterData.Word2 |= ShapeSimData.Word2;
+			OutSimFilterData.Word3 |= ShapeSimData.Word3;
+		}
 	}
 }
 
