Chaos - fix MACD flag not propagating to child particles on cluster/GC break

diff --git a/Engine/Source/Runtime/Experimental/Chaos/Private/Chaos/ClusterUnionManager.cpp b/Engine/Source/Runtime/Experimental/Chaos/Private/Chaos/ClusterUnionManager.cpp
--- a/Engine/Source/Runtime/Experimental/Chaos/Private/Chaos/ClusterUnionManager.cpp
+++ b/Engine/Source/Runtime/Experimental/Chaos/Private/Chaos/ClusterUnionManager.cpp
@@ -655,8 +655,8 @@ namespace Chaos
 		// Cluster has one-way interaction only if all children are also one-way
 		bool bIsOneWayInteraction = (bIsNewCluster) || Cluster->InternalCluster->OneWayInteraction();
 
-		// Cluster uses MACD is any children require MACD
-		bool bIsMACD = false;
+		// Cluster uses MACD is any children require MACD (or it is already enabled on this cluster)
+		bool bIsMACD = Cluster->InternalCluster->MACDEnabled();
 
 		// Use the minimum sleep multiplier of all member particles
 		FRealSingle MinSleepThresholdMultiplier = TNumericLimits<FRealSingle>::Max();
diff --git a/Engine/Source/Runtime/Experimental/Chaos/Private/Chaos/PBDRigidClustering.cpp b/Engine/Source/Runtime/Experimental/Chaos/Private/Chaos/PBDRigidClustering.cpp
--- a/Engine/Source/Runtime/Experimental/Chaos/Private/Chaos/PBDRigidClustering.cpp
+++ b/Engine/Source/Runtime/Experimental/Chaos/Private/Chaos/PBDRigidClustering.cpp
@@ -254,10 +254,12 @@ namespace Chaos
 
 		bool bClusterIsAsleep = true;
 		bool bClusterIsOneWayInteraction = true;
+		bool bClusterIsMACD = false;
 		for (FPBDRigidParticleHandle* Child : ChildrenSet)
 		{
 			bClusterIsAsleep &= Child->Sleeping();
 			bClusterIsOneWayInteraction &= Child->OneWayInteraction();
+			bClusterIsMACD |= Child->MACDEnabled();
 
 			if (FPBDRigidClusteredParticleHandle* ClusteredChild = Child->CastToClustered())
 			{
@@ -287,6 +289,7 @@ namespace Chaos
 			// children are added to an empty cluster, but we also shouldn't set non-default values before then
 			NewParticle->SetSleeping(bClusterIsAsleep);
 			NewParticle->SetOneWayInteraction(bClusterIsOneWayInteraction);
+			NewParticle->SetMACDEnabled(bClusterIsMACD);
 		}
 
 		if (ForceMassOrientation)
@@ -569,6 +572,7 @@ namespace Chaos
 			MEvolution.DisableParticles(ChildrenHandles);
 		}
 		bool bClusterIsOneWayInteraction = true;
+		bool bClusterIsMACD = false;
 		for (FPBDRigidParticleHandle* Child : ChildrenArray)
 		{
 			if (FPBDRigidClusteredParticleHandle* ClusteredChild = Child->CastToClustered())
@@ -588,6 +592,7 @@ namespace Chaos
 				Child->SetCollisionGroup(FMath::Min(NewParticle->CollisionGroup(), Child->CollisionGroup()));
 
 				bClusterIsOneWayInteraction &= Child->OneWayInteraction();
+				bClusterIsMACD |= Child->MACDEnabled();
 			}
 		}
 
@@ -605,6 +610,7 @@ namespace Chaos
 		UpdateGeometry(NewParticle, ChildrenSet, MChildren, FImplicitObjectPtr(nullptr), NoCleanParams);
 
 		NewParticle->SetOneWayInteraction(bClusterIsOneWayInteraction);
+		NewParticle->SetMACDEnabled(bClusterIsMACD);
 
 		// Build the convex optimizer if required
 		FRigidClustering::BuildConvexOptimizer(NewParticle);
