add ability to rebase the geometry collection component space transforms off the new world transform and keep the particles' old world space transforms

diff --git a/Engine/Source/Runtime/Experimental/Chaos/Private/PhysicsProxy/GeometryCollectionPhysicsProxy.cpp b/Engine/Source/Runtime/Experimental/Chaos/Private/PhysicsProxy/GeometryCollectionPhysicsProxy.cpp
--- a/Engine/Source/Runtime/Experimental/Chaos/Private/PhysicsProxy/GeometryCollectionPhysicsProxy.cpp
+++ b/Engine/Source/Runtime/Experimental/Chaos/Private/PhysicsProxy/GeometryCollectionPhysicsProxy.cpp
@@ -4416,32 +4416,9 @@ bool FGeometryCollectionPhysicsProxy::PullFromPhysicsState(const Chaos::FDirtyGe
 						}
 						if (bWasModified)
 						{
-							const FParticle& GTParticle = *GTParticles[ParticleIndex];
-							const FTransform& ParticleMassToLocal = MassToLocal[TransformIndex];
-							const int32 ParentTransformIndex = GameThreadCollection.GetParent(TransformIndex);
-
-							const FTransform& WorldTransform = ParticleMassToLocal.Inverse() * FTransform { GTParticle.R(), GTParticle.X() };
-
-							// by default parent is the component's world transform 
-							FTransform ParentWorldTransform = WorldTransform_External;
-							if (ParentTransformIndex != INDEX_NONE && FromTransformToParticleIndex[ParentTransformIndex] != INDEX_NONE)
-							{
-								if (const FParticle* GTParentParticle = GTParticles[FromTransformToParticleIndex[ParentTransformIndex]].Get())
-								{
-									const FTransform& ParentMassToLocal = MassToLocal[TransformIndex];
-									ParentWorldTransform = ParentMassToLocal.Inverse() * FTransform { GTParentParticle->R(), GTParentParticle->X() };
-								}
-							}
-
-							FTransform NewTransform = WorldTransform.GetRelativeTransform(ParentWorldTransform);
-							if (ParentTransformIndex == INDEX_NONE && bIsComponentTransformScaled)
-							{
-								NewTransform = ParticleMassToLocal.Inverse() * ComponentScaleTransform * ParticleMassToLocal * NewTransform;
-							}
-							if (!NewTransform.Equals(FTransform(GameThreadCollection.GetTransform(TransformIndex)), GeometryCollectionTransformTolerance))
+							if (RebaseParticleGameThreadCollectionTransformOnNewWorldTransform_External(ParticleIndex, MassToLocal, bIsComponentTransformScaled, ComponentScaleTransform))
 							{
 								bHasDifferentTransforms = true;
-								GameThreadCollection.SetTransform(TransformIndex, FTransform3f(NewTransform));
 							}
 						}
 					}
@@ -5918,6 +5895,65 @@ const FGeometryCollectionPhysicsProxy::FParticleHandle* FGeometryCollectionPhysi
 	return nullptr;
 }
 
+
+void FGeometryCollectionPhysicsProxy::RebaseAllGameThreadCollectionTransformsOnNewWorldTransform_External()
+{
+	check(IsInGameThread());
+	const TManagedArray<bool>* AnimationsActive = GameThreadCollection.GetAnimateTransformAttribute();
+	const TManagedArray<FTransform>& MassToLocal = Parameters.RestCollection->GetAttribute<FTransform>(MassToLocalAttributeName, FTransformCollection::TransformGroup);
+	const bool bIsComponentTransformScaled = !WorldTransform_External.GetScale3D().Equals(FVector::OneVector);
+	const FTransform ComponentScaleTransform(FQuat::Identity, FVector::ZeroVector, WorldTransform_External.GetScale3D());
+
+	for (int32 ParticleIndex = 0; ParticleIndex < NumEffectiveParticles; ++ParticleIndex)
+	{
+		const int32 TransformIndex = FromParticleToTransformIndex[ParticleIndex];
+		const bool bAnimatingWhileDisabled = AnimationsActive ? (*AnimationsActive)[TransformIndex] : false;
+		const bool bActive = GameThreadCollection.Active[TransformIndex];
+		if (bActive || bAnimatingWhileDisabled)
+		{
+			if (GTParticles[ParticleIndex] != nullptr)
+			{
+				RebaseParticleGameThreadCollectionTransformOnNewWorldTransform_External(ParticleIndex, MassToLocal, bIsComponentTransformScaled, ComponentScaleTransform);
+			}
+		}
+	}
+}
+
+bool FGeometryCollectionPhysicsProxy::RebaseParticleGameThreadCollectionTransformOnNewWorldTransform_External(int32 ParticleIndex, const TManagedArray<FTransform>& MassToLocal, bool bIsComponentTransformScaled, const FTransform& ComponentScaleTransform)
+{
+	check(IsInGameThread());
+	bool bHasDifferentTransforms = false;
+	const int32 TransformIndex = FromParticleToTransformIndex[ParticleIndex];
+	const FParticle& GTParticle = *GTParticles[ParticleIndex];
+	const FTransform& ParticleMassToLocal = MassToLocal[TransformIndex];
+	const int32 ParentTransformIndex = GameThreadCollection.GetParent(TransformIndex);
+
+	const FTransform& WorldTransform = ParticleMassToLocal.Inverse() * FTransform { GTParticle.R(), GTParticle.X() };
+
+	// by default parent is the component's world transform 
+	FTransform ParentWorldTransform = WorldTransform_External;
+	if (ParentTransformIndex != INDEX_NONE && FromTransformToParticleIndex[ParentTransformIndex] != INDEX_NONE)
+	{
+		if (const FParticle* GTParentParticle = GTParticles[FromTransformToParticleIndex[ParentTransformIndex]].Get())
+		{
+			const FTransform& ParentMassToLocal = MassToLocal[TransformIndex];
+			ParentWorldTransform = ParentMassToLocal.Inverse() * FTransform { GTParentParticle->R(), GTParentParticle->X() };
+		}
+	}
+
+	FTransform NewTransform = WorldTransform.GetRelativeTransform(ParentWorldTransform);
+	if (ParentTransformIndex == INDEX_NONE && bIsComponentTransformScaled)
+	{
+		NewTransform = ParticleMassToLocal.Inverse() * ComponentScaleTransform * ParticleMassToLocal * NewTransform;
+	}
+	if (!NewTransform.Equals(FTransform(GameThreadCollection.GetTransform(TransformIndex)), GeometryCollectionTransformTolerance))
+	{
+		bHasDifferentTransforms = true;
+		GameThreadCollection.SetTransform(TransformIndex, FTransform3f(NewTransform));
+	}
+	return bHasDifferentTransforms;
+}
+
 void FDamageCollector::Reset(int32 NumTransforms)
 {
 	DamageData.Reset();
diff --git a/Engine/Source/Runtime/Experimental/Chaos/Public/PhysicsProxy/GeometryCollectionPhysicsProxy.h b/Engine/Source/Runtime/Experimental/Chaos/Public/PhysicsProxy/GeometryCollectionPhysicsProxy.h
--- a/Engine/Source/Runtime/Experimental/Chaos/Public/PhysicsProxy/GeometryCollectionPhysicsProxy.h
+++ b/Engine/Source/Runtime/Experimental/Chaos/Public/PhysicsProxy/GeometryCollectionPhysicsProxy.h
@@ -536,6 +536,7 @@ public:
 	CHAOS_API TArray<Chaos::FPhysicsObjectHandle> GetAllPhysicsObjects() const ;
 	CHAOS_API TArray<Chaos::FPhysicsObjectHandle> GetAllPhysicsObjectIncludingNulls() const;
 	CHAOS_API Chaos::FPhysicsObjectHandle GetPhysicsObjectByIndex(int32 Index) const;
+	CHAOS_API void RebaseAllGameThreadCollectionTransformsOnNewWorldTransform_External();
 
 	UE_DEPRECATED(5.4, "Use GetNumTransforms instead")
 	int32 GetNumParticles() const { return NumTransforms; }
@@ -551,6 +552,8 @@ public:
 
 protected:
 
+	bool RebaseParticleGameThreadCollectionTransformOnNewWorldTransform_External(int32 ParticleIndex, const TManagedArray<FTransform>& MassToLocal, bool bIsComponentTransformScaled, const FTransform& ComponentScaleTransform);
+
 	CHAOS_API float ComputeMaterialBasedDamageThreshold_Internal(int32 TransformIndex) const;
 
 	/**
diff --git a/Engine/Source/Runtime/Experimental/GeometryCollectionEngine/Private/GeometryCollection/GeometryCollectionComponent.cpp b/Engine/Source/Runtime/Experimental/GeometryCollectionEngine/Private/GeometryCollection/GeometryCollectionComponent.cpp
--- a/Engine/Source/Runtime/Experimental/GeometryCollectionEngine/Private/GeometryCollection/GeometryCollectionComponent.cpp
+++ b/Engine/Source/Runtime/Experimental/GeometryCollectionEngine/Private/GeometryCollection/GeometryCollectionComponent.cpp
@@ -7383,3 +7383,14 @@ bool UGeometryCollectionComponent::IsHLODRelevant() const
 }
 
 #endif // #if WITH_EDITOR
+
+void UGeometryCollectionComponent::RebaseDynamicCollectionTransformsOnNewWorldTransform()
+{
+	if (!PhysicsProxy)
+	{
+		return;
+	}
+
+	PhysicsProxy->RebaseAllGameThreadCollectionTransformsOnNewWorldTransform_External();
+	ComponentSpaceTransforms.MarkDirty();
+}
\ No newline at end of file
diff --git a/Engine/Source/Runtime/Experimental/GeometryCollectionEngine/Public/GeometryCollection/GeometryCollectionComponent.h b/Engine/Source/Runtime/Experimental/GeometryCollectionEngine/Public/GeometryCollection/GeometryCollectionComponent.h
--- a/Engine/Source/Runtime/Experimental/GeometryCollectionEngine/Public/GeometryCollection/GeometryCollectionComponent.h
+++ b/Engine/Source/Runtime/Experimental/GeometryCollectionEngine/Public/GeometryCollection/GeometryCollectionComponent.h
@@ -1878,6 +1878,14 @@ public:
 	//~ End UActorComponent interface.
 #endif
 
+	/**
+	 * Currently, component space transforms for every particle is compute relative to the component transform.
+	 * By default, this means that when the component transform changes, every particle is shifted along with the component.
+	 * However, there are times when you may change the component transform but instead want every particle to stay in the same world position.
+	 * In those cases, RebaseDynamicCollectionTransformsOnNewWorldTransform can be called.
+	 */
+	GEOMETRYCOLLECTIONENGINE_API void RebaseDynamicCollectionTransformsOnNewWorldTransform();
+
 	friend struct FGCCollisionProfileScopedTransaction;
 };
 
