Iris
* When NetGridFilterConfig MaxDistance is set to 0 accept all cull distance values
* Changed default config value to be zero
* Added warning when detecting objects with culldistances outside the max tolarated

#rnx

diff --git a/Engine/Source/Runtime/Experimental/Iris/Core/Private/Iris/ReplicationSystem/Filtering/NetObjectGridFilter.cpp b/Engine/Source/Runtime/Experimental/Iris/Core/Private/Iris/ReplicationSystem/Filtering/NetObjectGridFilter.cpp
--- a/Engine/Source/Runtime/Experimental/Iris/Core/Private/Iris/ReplicationSystem/Filtering/NetObjectGridFilter.cpp
+++ b/Engine/Source/Runtime/Experimental/Iris/Core/Private/Iris/ReplicationSystem/Filtering/NetObjectGridFilter.cpp
@@ -54,11 +54,11 @@ bool UNetObjectGridFilter::AddObject(uint32 ObjectIndex, FNetObjectFilterAddObje
 
 	AddCellInfoForObject(ObjectLocationInfo, Params.InstanceProtocol);
 	
-	if (PerObjectInfo.GetCullDistance() > Config->MaxCullDistance)
+	if (Config->MaxCullDistance > 0.0f && PerObjectInfo.GetCullDistance() > Config->MaxCullDistance)
 	{
 		// Too big an object. We expect it to be costly to move it across cells.
+		UE_LOG(LogIris, Warning, TEXT("ReplicatedObject %u cull distance %f is above the max %f. Object will become always relevant instead"), ObjectIndex, PerObjectInfo.GetCullDistance(), Config->MaxCullDistance);
 		RemoveObject(ObjectIndex, ObjectLocationInfo);
-		//TODO: Should we log these objects so they can get flagged and properly set always relevant ?
 		return false;
 	}
 
diff --git a/Engine/Source/Runtime/Experimental/Iris/Core/Public/Iris/ReplicationSystem/Filtering/NetObjectGridFilter.h b/Engine/Source/Runtime/Experimental/Iris/Core/Public/Iris/ReplicationSystem/Filtering/NetObjectGridFilter.h
--- a/Engine/Source/Runtime/Experimental/Iris/Core/Public/Iris/ReplicationSystem/Filtering/NetObjectGridFilter.h
+++ b/Engine/Source/Runtime/Experimental/Iris/Core/Public/Iris/ReplicationSystem/Filtering/NetObjectGridFilter.h
@@ -33,9 +33,9 @@ public:
 	UPROPERTY(Config)
 	float CellSizeY = 20000.0f;
 
-	/** Objects with larger sqrt(NetCullDistanceSqr) will be rejected. */
+	/** Objects with larger sqrt(NetCullDistanceSqr) will be rejected. Disabled when value is zero. */
 	UPROPERTY(Config)
-	float MaxCullDistance = 20000.0f;
+	float MaxCullDistance = 0;
 
 	/** Objects without a NetCullDistanceSquared property will assume to have this value but squared unless there's a cull distance override. */
 	UPROPERTY(Config)
