Added more cpu scopes to the different operations in the Iris filter pass

#rnx

diff --git a/Engine/Source/Runtime/Experimental/Iris/Core/Private/Iris/ReplicationSystem/Filtering/ReplicationFiltering.cpp b/Engine/Source/Runtime/Experimental/Iris/Core/Private/Iris/ReplicationSystem/Filtering/ReplicationFiltering.cpp
--- a/Engine/Source/Runtime/Experimental/Iris/Core/Private/Iris/ReplicationSystem/Filtering/ReplicationFiltering.cpp
+++ b/Engine/Source/Runtime/Experimental/Iris/Core/Private/Iris/ReplicationSystem/Filtering/ReplicationFiltering.cpp
@@ -833,6 +833,8 @@ void FReplicationFiltering::ResetRemovedConnections()
 		return;
 	}
 
+	IRIS_PROFILER_SCOPE(FReplicationFiltering_ResetRemovedConnections);
+
 	bHasRemovedConnection = 0;
 
 	// Reset group filter status
@@ -858,6 +860,8 @@ void FReplicationFiltering::ResetRemovedConnections()
 
 void FReplicationFiltering::UpdateObjectsInScope()
 {
+	IRIS_PROFILER_SCOPE(FReplicationFiltering_UpdateObjectsInScope);
+
 	const FNetBitArrayView ObjectsInScope = NetRefHandleManager->GetCurrentFrameScopableInternalIndices();
 	const FNetBitArrayView PrevObjectsInScope = NetRefHandleManager->GetPrevFrameScopableInternalIndices();
 
@@ -1073,6 +1077,8 @@ void FReplicationFiltering::UpdateOwnerAndConnectionFiltering()
 	// Update owners
 	if (bHasDirtyOwner)
 	{
+		IRIS_PROFILER_SCOPE(FReplicationFiltering_UpdateDirtyOwnerValue);
+
 		uint16* const ObjectIndexToOwningConnectionStorage = ObjectIndexToOwningConnection.GetData();
 		auto UpdateOwners = [this, ObjectIndexToOwningConnectionStorage](uint32 ObjectIndex)
 		{
@@ -1096,6 +1102,8 @@ void FReplicationFiltering::UpdateOwnerAndConnectionFiltering()
 	// Update filtering
 	if (bHasDirtyConnectionFilter)
 	{
+		IRIS_PROFILER_SCOPE(FReplicationFiltering_UpdateDirtyConnectionFilter);
+
 		ObjectsWithDirtyConnectionFilter.ForAllSetBits([this](uint32 DirtyObjectIndex)
 		{
 			AllConnectionFilteredObjects.SetBitValue(DirtyObjectIndex, HasConnectionFilter(DirtyObjectIndex));
@@ -1369,6 +1377,7 @@ void FReplicationFiltering::UpdateDynamicFiltering(ENetFilterType FilterType)
 		 * for objects that have changed filter status since the previous frame.
 		 */
 		{
+			IRIS_PROFILER_SCOPE(FReplicationFiltering_OnFilterStatusChanged);
 			DisabledDependentObjects.Reset();
 
 			const uint32* DynamicFilterEnabledObjectsData = DynamicFilterEnabledObjects.GetData();
@@ -1435,6 +1444,8 @@ void FReplicationFiltering::UpdateDynamicFiltering(ENetFilterType FilterType)
 
 		// Update the entire scope for the connection
 		{
+			IRIS_PROFILER_SCOPE(FReplicationFiltering_UpdateConnectionScope);
+
 			uint32* ObjectsInScopeData = ConnectionInfo.ObjectsInScope.GetData();
 			const uint32* ObjectsInScopeBeforeDynamicFilteringData = ConnectionInfo.ObjectsInScopeBeforeDynamicFiltering.GetData();
 			const uint32* DynamicFilteredOutObjectsData = ConnectionInfo.DynamicFilteredOutObjects.GetData();
@@ -1452,6 +1463,7 @@ void FReplicationFiltering::UpdateDynamicFiltering(ENetFilterType FilterType)
 		// The scope for the connection is now fully updated, apart from disabled dependent objects.
 		// If any object that has a dependency that isn't filtered out we must re-enable the dependent object.
 		{
+			IRIS_PROFILER_SCOPE(FReplicationFiltering_UpdateDependentObjects);
 			for (const uint32 DependentObjectIndex : DisabledDependentObjects)
 			{
 				if (GetDependentObjectFilterStatus(NetRefHandleManager, ConnectionInfo.ObjectsInScope, DependentObjectIndex) == ENetFilterStatus::Allow)
@@ -1496,6 +1508,8 @@ void FReplicationFiltering::PostUpdateDynamicFiltering(ENetFilterType FilterType
 
 void FReplicationFiltering::NotifyFiltersOfDirtyObjects(ENetFilterType FilterType)
 {
+	IRIS_PROFILER_SCOPE(FReplicationFiltering_UpdateFilterWithDirtyObjects);
+
 	FDirtyObjectsAccessor DirtyObjectsAccessor(ReplicationSystem->GetReplicationSystemInternal()->GetDirtyNetObjectTracker());
 	const FNetBitArrayView DirtyObjectsThisFrame = DirtyObjectsAccessor.GetDirtyNetObjects();
 
@@ -1534,7 +1548,7 @@ void FReplicationFiltering::BatchNotifyFiltersOfDirtyObjects(FUpdateDirtyObjects
 		UpdateParameters.ObjectIndices = PerFilterInfo.ObjectIndices;
 		UpdateParameters.ObjectCount = PerFilterInfo.ObjectCount;
 		UpdateParameters.InstanceProtocols = PerFilterInfo.InstanceProtocols;
-
+		// $IRIS TODO: Add a trait asking if the Filter needs to receive UpdateObjects. Most do nothing in UpdateObjects()
 		const int32 FilterIndex = static_cast<int32>(&PerFilterInfo - BatchHelper.PerFilterInfos.GetData());
 		UNetObjectFilter* Filter = DynamicFilterInfos[FilterIndex].Filter.Get();
 		Filter->UpdateObjects(UpdateParameters);
@@ -1680,6 +1694,8 @@ void FReplicationFiltering::RemoveSubObjectFilter(FNetObjectGroupHandle GroupHan
 
 void FReplicationFiltering::UpdateSubObjectFilters()
 {
+	IRIS_PROFILER_SCOPE(FReplicationFiltering_UpdateSubObjectFilters);
+
 	// We want to remove all groups that have no members and no enabled connections
 	auto UpdateSubObjectFilterGroup = [this](uint32 GroupIndex)
 	{
