Chaos - fix broadphase ignoring probe-only shapes

This was a knock-on from a prior change that was supposed to ignore query-only shapes in the broadphase but also ignored probe-only shapes.

diff --git a/Engine/Source/Runtime/Experimental/Chaos/Public/Chaos/ParticleHandle.h b/Engine/Source/Runtime/Experimental/Chaos/Public/Chaos/ParticleHandle.h
--- a/Engine/Source/Runtime/Experimental/Chaos/Public/Chaos/ParticleHandle.h
+++ b/Engine/Source/Runtime/Experimental/Chaos/Public/Chaos/ParticleHandle.h
@@ -3813,6 +3813,7 @@ template <typename TParticle>
 	const auto& Shapes = Particle.ShapesArray();
 	for (const auto& Shape : Shapes)
 	{
+		// For queries we can ignore any shapes that are not query-enabled (e.g., sim- and probe-only)
 		if (Shape->GetQueryEnabled())
 		{
 			const FCollisionFilterData& ShapeQueryData = Shape->GetQueryData();
@@ -3830,7 +3831,8 @@ template <typename TParticle>
 	const auto& Shapes = Particle.ShapesArray();
 	for (const auto& Shape : Shapes)
 	{
-		if (Shape->GetSimEnabled())
+		// In the simulation we can ignore query-only shapes, but need to know about sim- and probe-enabled shapes
+		if (Shape->GetSimEnabled() || Shape->GetIsProbe())
 		{
 			const FCollisionFilterData& ShapeSimData = Shape->GetSimData();
 			OutSimFilterData.Word0 |= ShapeSimData.Word0;
