Iris
* NetGridFilter::RemoveObject will also remove the object from a connection's recently relevant list.

#rnx

diff --git a/Engine/Source/Runtime/Experimental/Iris/Core/Private/Iris/ReplicationSystem/Filtering/NetObjectGridFilter.cpp b/Engine/Source/Runtime/Experimental/Iris/Core/Private/Iris/ReplicationSystem/Filtering/NetObjectGridFilter.cpp
--- a/Engine/Source/Runtime/Experimental/Iris/Core/Private/Iris/ReplicationSystem/Filtering/NetObjectGridFilter.cpp
+++ b/Engine/Source/Runtime/Experimental/Iris/Core/Private/Iris/ReplicationSystem/Filtering/NetObjectGridFilter.cpp
@@ -67,6 +67,12 @@ bool UNetObjectGridFilter::AddObject(uint32 ObjectIndex, FNetObjectFilterAddObje
 
 void UNetObjectGridFilter::RemoveObject(uint32 ObjectIndex, const FNetObjectFilteringInfo& Info)
 {
+	// Remove the object from connection lists
+	for (FPerConnectionInfo& ConnectionInfo : PerConnectionInfos)
+	{
+		ConnectionInfo.RecentObjectFrameCount.Remove(ObjectIndex);
+	}
+
 	const FObjectLocationInfo& ObjectLocationInfo = static_cast<const FObjectLocationInfo&>(Info);
 	RemoveCellInfoForObject(ObjectLocationInfo);
 
