Chaos VD
- rename and move CVD data collection steps to be more useful (PostIntegrate, PreSolve, PostSolve)
- fix position saved in each step (should be P/Q for dynamics, not X/R)

Chaos
- moved ApplyKinematicTargets into cpp file to reduce recompilation and keep the code near its partner Integrate()

diff --git a/Engine/Plugins/Experimental/ChaosVD/Source/ChaosVD/Private/Widgets/SChaosVDMainTab.cpp b/Engine/Plugins/Experimental/ChaosVD/Source/ChaosVD/Private/Widgets/SChaosVDMainTab.cpp
--- a/Engine/Plugins/Experimental/ChaosVD/Source/ChaosVD/Private/Widgets/SChaosVDMainTab.cpp
+++ b/Engine/Plugins/Experimental/ChaosVD/Source/ChaosVD/Private/Widgets/SChaosVDMainTab.cpp
@@ -593,14 +593,24 @@ FReply SChaosVDMainTab::HandleSessionConnectionClicked()
 
 FText SChaosVDMainTab::GetConnectButtonText() const
 {
-	const bool bIsAlreadyInLiveSession = GetChaosVDEngineInstance()->GetPlaybackController()->IsPlayingLiveSession();
-	return bIsAlreadyInLiveSession ? LOCTEXT("DisconnectFromSession", "Disconnect from Session") : LOCTEXT("ConnectToSession", "Connect to Session");
+	TSharedPtr<FChaosVDPlaybackController> PlaybackControllerPtr = GetChaosVDEngineInstance()->GetPlaybackController();
+	if (PlaybackControllerPtr.IsValid())
+	{
+		const bool bIsAlreadyInLiveSession = PlaybackControllerPtr->IsPlayingLiveSession();
+		return bIsAlreadyInLiveSession ? LOCTEXT("DisconnectFromSession", "Disconnect from Session") : LOCTEXT("ConnectToSession", "Connect to Session");
+	}
+	return FText();
 }
 
 FText SChaosVDMainTab::GetConnectButtonTooltipText() const
 {
-	const bool bIsAlreadyInLiveSession = GetChaosVDEngineInstance()->GetPlaybackController()->IsPlayingLiveSession();
-	return bIsAlreadyInLiveSession ? LOCTEXT("DisconnectFromSessionTooltip", "Disconnects from the current live session but it does not stop it") : LOCTEXT("ConnectToSessionTooltip", "Opens a panel where you can browse active live sessions and connect to one");
+	TSharedPtr<FChaosVDPlaybackController> PlaybackControllerPtr = GetChaosVDEngineInstance()->GetPlaybackController();
+	if (PlaybackControllerPtr.IsValid())
+	{
+		const bool bIsAlreadyInLiveSession = PlaybackControllerPtr->IsPlayingLiveSession();
+		return bIsAlreadyInLiveSession ? LOCTEXT("DisconnectFromSessionTooltip", "Disconnects from the current live session but it does not stop it") : LOCTEXT("ConnectToSessionTooltip", "Opens a panel where you can browse active live sessions and connect to one");
+	}
+	return FText();
 }
 
 #undef LOCTEXT_NAMESPACE
diff --git a/Engine/Source/Runtime/Experimental/Chaos/Private/Chaos/PBDRigidsEvolutionGBF.cpp b/Engine/Source/Runtime/Experimental/Chaos/Private/Chaos/PBDRigidsEvolutionGBF.cpp
--- a/Engine/Source/Runtime/Experimental/Chaos/Private/Chaos/PBDRigidsEvolutionGBF.cpp
+++ b/Engine/Source/Runtime/Experimental/Chaos/Private/Chaos/PBDRigidsEvolutionGBF.cpp
@@ -445,7 +445,7 @@ void FPBDRigidsEvolutionGBF::AdvanceOneTimeStepImpl(const FReal Dt, const FSubSt
 #endif
 
 	{
-		CVD_SCOPE_TRACE_SOLVER_STEP(CVDDC_EvolutionStart, TEXT("Evolution Start"))
+		CVD_SCOPE_TRACE_SOLVER_STEP(CVDDC_EvolutionStart, TEXT("Evolution Start"));
 		CVD_TRACE_PARTICLES_SOA(Particles);
 	}
 
@@ -468,19 +468,22 @@ void FPBDRigidsEvolutionGBF::AdvanceOneTimeStepImpl(const FReal Dt, const FSubSt
 	}
 
 	{
-		CVD_SCOPE_TRACE_SOLVER_STEP(CVDDC_Integrate, TEXT("Integrate"))
 		SCOPE_CYCLE_COUNTER(STAT_Evolution_Integrate);
 		CSV_SCOPED_TIMING_STAT(PhysicsVerbose, StepSolver_Integrate);
 		Integrate(Dt);
 	}
 
 	{
-		CVD_SCOPE_TRACE_SOLVER_STEP(CVDDC_ApplyKinematicTargets, TEXT("ApplyKinematicTargets"))
 		SCOPE_CYCLE_COUNTER(STAT_Evolution_KinematicTargets);
 		CSV_SCOPED_TIMING_STAT(PhysicsVerbose, StepSolver_KinematicTargets);
 		ApplyKinematicTargets(Dt, SubStepInfo.PseudoFraction);
 	}
 
+	{
+		CVD_SCOPE_TRACE_SOLVER_STEP(CVDDC_PostIntegrate, TEXT("Post Integrate"));
+		CVD_TRACE_PARTICLES_SOA(Particles);
+	}
+
 	if (PostIntegrateCallback != nullptr)
 	{
 		SCOPE_CYCLE_COUNTER(STAT_Evolution_PostIntegrateCallback);
@@ -508,7 +511,7 @@ void FPBDRigidsEvolutionGBF::AdvanceOneTimeStepImpl(const FReal Dt, const FSubSt
 		CollisionDetector.GetBroadPhase().SetSpatialAcceleration(InternalAcceleration);
 
 		{
-			CVD_SCOPE_TRACE_SOLVER_STEP(CVDDC_CollisionDetectionBroadPhase, TEXT("Collision Detection Broad Phase"))
+			CVD_SCOPE_TRACE_SOLVER_STEP(CVDDC_CollisionDetectionBroadPhase, TEXT("Collision Detection Broad Phase"));
 			CollisionDetector.RunBroadPhase(Dt, GetCurrentStepResimCache());
 		}
 
@@ -519,7 +522,7 @@ void FPBDRigidsEvolutionGBF::AdvanceOneTimeStepImpl(const FReal Dt, const FSubSt
 		}
 
 		{
-			CVD_SCOPE_TRACE_SOLVER_STEP(CVDDC_CollisionDetectionNarrowPhase, TEXT("Collision Detection Narrow Phase"))
+			CVD_SCOPE_TRACE_SOLVER_STEP(CVDDC_CollisionDetectionNarrowPhase, TEXT("Collision Detection Narrow Phase"));
 			CollisionDetector.RunNarrowPhase(Dt, GetCurrentStepResimCache());
 		}
 	}
@@ -596,6 +599,11 @@ void FPBDRigidsEvolutionGBF::AdvanceOneTimeStepImpl(const FReal Dt, const FSubSt
 		ReloadParticlesCache();
 	}
 
+	{
+		CVD_SCOPE_TRACE_SOLVER_STEP(CVDDC_PreConstraintSolve, TEXT("Pre Solve"));
+		CVD_TRACE_PARTICLES_SOA(Particles);
+	}
+
 	// Assign all islands to a set of groups. Each group is solved in parallel with the others.
 	int32 NumGroups = 0;
 	{
@@ -633,6 +641,11 @@ void FPBDRigidsEvolutionGBF::AdvanceOneTimeStepImpl(const FReal Dt, const FSubSt
 		}
 	}
 
+	{
+		CVD_SCOPE_TRACE_SOLVER_STEP(CVDDC_PostConstraintSolve, TEXT("Post Solve"));
+		CVD_TRACE_PARTICLES_SOA(Particles);
+	}
+
 	if (PostSolveCallback != nullptr)
 	{
 		SCOPE_CYCLE_COUNTER(STAT_Evolution_PostSolveCallback);
@@ -742,13 +755,8 @@ void FPBDRigidsEvolutionGBF::Integrate(FReal Dt)
 	const bool bAllowMACD = CVars::bChaosUseMACD;
 	const bool bForceMACD = CVars::bChaosForceMACD;
 
-	FChaosVDContextWrapper CVDContext;
-	CVD_GET_WRAPPED_CURRENT_CONTEXT(CVDContext);
-
 	ParticlesView.ParallelFor([&](auto& GeomParticle, int32 Index)
 		{
-			CVD_SCOPE_CONTEXT(CVDContext.Context);
-
 			//question: can we enforce this at the API layer? Right now islands contain non dynamic which makes this hard
 			auto PBDParticle = GeomParticle.CastToRigidParticle();
 			if (PBDParticle && PBDParticle->ObjectState() == EObjectStateType::Dynamic)
@@ -875,8 +883,6 @@ void FPBDRigidsEvolutionGBF::Integrate(FReal Dt)
 						Particle.UpdateWorldSpaceStateSwept(FRigidTransform3(Particle.GetP(), Particle.GetQ()), FVec3(BoundsThickness), VelocityBoundsDelta);
 					}
 				}
-
-				CVD_TRACE_PARTICLE(PBDParticle->Handle())
 			}
 		});
 
@@ -886,6 +892,129 @@ void FPBDRigidsEvolutionGBF::Integrate(FReal Dt)
 	}
 }
 
+void FPBDRigidsEvolutionGBF::ApplyKinematicTargets(const FReal Dt, const FReal StepFraction)
+{
+	check(StepFraction > (FReal)0);
+	check(StepFraction <= (FReal)1);
+
+	const bool IsLastStep = (FMath::IsNearlyEqual(StepFraction, (FReal)1, (FReal)UE_KINDA_SMALL_NUMBER));
+
+	// NOTE: ApplyKinematicTargetForParticle is run in a parallel-for. We only write to particle state
+	const auto& ApplyParticleKinematicTarget =
+	[Dt, StepFraction, IsLastStep](FTransientPBDRigidParticleHandle& Particle, const int32 ParticleIndex) -> void
+	{
+		TKinematicTarget<FReal, 3>& KinematicTarget = Particle.KinematicTarget();
+		const FVec3 CurrentX = Particle.GetX();
+		const FRotation3 CurrentR = Particle.GetR();
+		constexpr FReal MinDt = 1e-6f;
+
+		bool bMoved = false;
+		switch (KinematicTarget.GetMode())
+		{
+		case EKinematicTargetMode::None:
+			// Nothing to do
+			break;
+
+		case EKinematicTargetMode::Reset:
+		{
+			// Reset velocity and then switch to do-nothing mode
+			Particle.SetVf(FVec3f(0.0f, 0.0f, 0.0f));
+			Particle.SetWf(FVec3f(0.0f, 0.0f, 0.0f));
+			Particle.ClearIsMovingKinematic();
+			KinematicTarget.SetMode(EKinematicTargetMode::None);
+			break;
+		}
+
+		case EKinematicTargetMode::Position:
+		{
+			// Move to kinematic target and update velocities to match
+			// Target positions only need to be processed once, and we reset the velocity next frame (if no new target is set)
+			FVec3 NewX;
+			FRotation3 NewR;
+			if (IsLastStep)
+			{
+				NewX = KinematicTarget.GetTarget().GetLocation();
+				NewR = KinematicTarget.GetTarget().GetRotation();
+				KinematicTarget.SetMode(EKinematicTargetMode::Reset);
+			}
+			else
+			{
+				// as a reminder, stepfraction is the remaing fraction of the step from the remaining steps
+				// for total of 4 steps and current step of 2, this will be 1/3 ( 1 step passed, 3 steps remains )
+				NewX = FVec3::Lerp(CurrentX, KinematicTarget.GetTarget().GetLocation(), StepFraction);
+				NewR = FRotation3::Slerp(CurrentR, KinematicTarget.GetTarget().GetRotation(), decltype(FQuat::X)(StepFraction));
+			}
+
+			const bool bPositionChanged = !FVec3::IsNearlyEqual(NewX, CurrentX, UE_SMALL_NUMBER);
+			const bool bRotationChanged = !FRotation3::IsNearlyEqual(NewR, CurrentR, UE_SMALL_NUMBER);
+			bMoved = bPositionChanged || bRotationChanged;
+			FVec3 NewV = FVec3(0);
+			FVec3 NewW = FVec3(0);
+			if (Dt > MinDt)
+			{
+				if (bPositionChanged)
+				{
+					NewV = FVec3::CalculateVelocity(CurrentX, NewX, Dt);
+				}
+				if (bRotationChanged)
+				{
+					NewW = FRotation3::CalculateAngularVelocity(CurrentR, NewR, Dt);
+				}
+			}
+			Particle.SetX(NewX);
+			Particle.SetR(NewR);
+			Particle.SetV(NewV);
+			Particle.SetW(NewW);
+			Particle.SetIsMovingKinematic();
+
+			break;
+		}
+
+		case EKinematicTargetMode::Velocity:
+		{
+			// Move based on velocity
+			bMoved = true;
+			Particle.SetX(Particle.GetX() + Particle.GetV() * Dt);
+			Particle.SetRf(FRotation3f::IntegrateRotationWithAngularVelocity(Particle.GetRf(), Particle.GetWf(), FRealSingle(Dt)));
+			Particle.SetIsMovingKinematic();
+
+			break;
+		}
+		}
+
+		// Set positions and previous velocities if we can
+		// Note: At present kinematics are in fact rigid bodies
+		Particle.SetP(Particle.GetX());
+		Particle.SetQf(Particle.GetRf());
+		Particle.SetPreVf(Particle.GetVf());
+		Particle.SetPreWf(Particle.GetWf());
+
+		if (bMoved)
+		{
+			if (!Particle.CCDEnabled())
+			{
+				Particle.UpdateWorldSpaceState(FRigidTransform3(Particle.GetP(), Particle.GetQ()), FVec3(0));
+			}
+			else
+			{
+				Particle.UpdateWorldSpaceStateSwept(FRigidTransform3(Particle.GetP(), Particle.GetQ()), FVec3(0), -Particle.GetV() * Dt);
+			}
+		}
+	};
+
+	// Apply kinematic targets in parallel
+	Particles.GetActiveMovingKinematicParticlesView().ParallelFor(ApplyParticleKinematicTarget);
+
+	// done with update, let's clear the tracking structures
+	if (IsLastStep)
+	{
+		Particles.UpdateAllMovingKinematic();
+	}
+
+	// If we changed any particle state, the views need to be refreshed
+	Particles.UpdateDirtyViews();
+}
+
 void FPBDRigidsEvolutionGBF::SetIsDeterministic(const bool bInIsDeterministic)
 {
 	// We detect collisions in parallel, so order is non-deterministic without additional processing
diff --git a/Engine/Source/Runtime/Experimental/Chaos/Private/ChaosVisualDebugger/ChaosVDDataWrapperUtils.cpp b/Engine/Source/Runtime/Experimental/Chaos/Private/ChaosVisualDebugger/ChaosVDDataWrapperUtils.cpp
--- a/Engine/Source/Runtime/Experimental/Chaos/Private/ChaosVisualDebugger/ChaosVDDataWrapperUtils.cpp
+++ b/Engine/Source/Runtime/Experimental/Chaos/Private/ChaosVisualDebugger/ChaosVDDataWrapperUtils.cpp
@@ -87,7 +87,7 @@ FChaosVDParticleDataWrapper FChaosVDDataWrapperUtils::BuildParticleDataWrapperFr
 	WrappedParticleData.DebugNamePtr = ParticleHandlePtr->DebugName().IsValid() ? ParticleHandlePtr->DebugName().Get() : nullptr;
 #endif
 
-	WrappedParticleData.ParticlePositionRotation.CopyFrom(*ParticleHandlePtr);
+	WrappedParticleData.ParticlePositionRotation.CopyFrom(Chaos::FConstGenericParticleHandle(ParticleHandlePtr));
 
 	if (const Chaos::TKinematicGeometryParticleHandleImp<Chaos::FReal, 3, true>* KinematicParticle = ParticleHandlePtr->CastToKinematicParticle())
 	{
diff --git a/Engine/Source/Runtime/Experimental/Chaos/Private/ChaosVisualDebugger/ChaosVDOptionalDataChannel.cpp b/Engine/Source/Runtime/Experimental/Chaos/Private/ChaosVisualDebugger/ChaosVDOptionalDataChannel.cpp
--- a/Engine/Source/Runtime/Experimental/Chaos/Private/ChaosVisualDebugger/ChaosVDOptionalDataChannel.cpp
+++ b/Engine/Source/Runtime/Experimental/Chaos/Private/ChaosVisualDebugger/ChaosVDOptionalDataChannel.cpp
@@ -75,11 +75,12 @@ namespace Chaos::VisualDebugger
 
 CVD_DEFINE_OPTIONAL_DATA_CHANNEL(Default, EChaosVDDataChannelInitializationFlags::StartEnabled);
 CVD_DEFINE_OPTIONAL_DATA_CHANNEL(EvolutionStart, EChaosVDDataChannelInitializationFlags::StartEnabled | EChaosVDDataChannelInitializationFlags::CanChangeEnabledState);
-CVD_DEFINE_OPTIONAL_DATA_CHANNEL(Integrate, EChaosVDDataChannelInitializationFlags::CanChangeEnabledState);
-CVD_DEFINE_OPTIONAL_DATA_CHANNEL(ApplyKinematicTargets, EChaosVDDataChannelInitializationFlags::StartEnabled | EChaosVDDataChannelInitializationFlags::CanChangeEnabledState);
+CVD_DEFINE_OPTIONAL_DATA_CHANNEL(PostIntegrate, EChaosVDDataChannelInitializationFlags::CanChangeEnabledState);
+CVD_DEFINE_OPTIONAL_DATA_CHANNEL(PreConstraintSolve, EChaosVDDataChannelInitializationFlags::CanChangeEnabledState);
 CVD_DEFINE_OPTIONAL_DATA_CHANNEL(CollisionDetectionBroadPhase, EChaosVDDataChannelInitializationFlags::CanChangeEnabledState);
 CVD_DEFINE_OPTIONAL_DATA_CHANNEL(CollisionDetectionNarrowPhase, EChaosVDDataChannelInitializationFlags::CanChangeEnabledState);
-CVD_DEFINE_OPTIONAL_DATA_CHANNEL(EndOfEvolutionCollisionConstraints, EChaosVDDataChannelInitializationFlags::StartEnabled| EChaosVDDataChannelInitializationFlags::CanChangeEnabledState);
+CVD_DEFINE_OPTIONAL_DATA_CHANNEL(PostConstraintSolve, EChaosVDDataChannelInitializationFlags::CanChangeEnabledState);
+CVD_DEFINE_OPTIONAL_DATA_CHANNEL(EndOfEvolutionCollisionConstraints, EChaosVDDataChannelInitializationFlags::StartEnabled | EChaosVDDataChannelInitializationFlags::CanChangeEnabledState);
 CVD_DEFINE_OPTIONAL_DATA_CHANNEL(EvolutionEnd, EChaosVDDataChannelInitializationFlags::StartEnabled); // Intentionally not allowing changing the enabled state because otherwise there might be nothing to visualize
 CVD_DEFINE_OPTIONAL_DATA_CHANNEL(SceneQueries, EChaosVDDataChannelInitializationFlags::StartEnabled | EChaosVDDataChannelInitializationFlags::CanChangeEnabledState);
 CVD_DEFINE_OPTIONAL_DATA_CHANNEL(JointConstraints, EChaosVDDataChannelInitializationFlags::CanChangeEnabledState);
diff --git a/Engine/Source/Runtime/Experimental/Chaos/Public/Chaos/PBDRigidsEvolution.h b/Engine/Source/Runtime/Experimental/Chaos/Public/Chaos/PBDRigidsEvolution.h
--- a/Engine/Source/Runtime/Experimental/Chaos/Public/Chaos/PBDRigidsEvolution.h
+++ b/Engine/Source/Runtime/Experimental/Chaos/Public/Chaos/PBDRigidsEvolution.h
@@ -733,137 +733,7 @@ public:
 		}
 	}
 
-	void ApplyKinematicTargets(const FReal Dt, const FReal StepFraction)
-	{
-		check(StepFraction > (FReal)0);
-		check(StepFraction <= (FReal)1);
-
-		const bool IsLastStep = (FMath::IsNearlyEqual(StepFraction, (FReal)1, (FReal)UE_KINDA_SMALL_NUMBER));
-
-		FChaosVDContextWrapper CVDContext;
-		CVD_GET_WRAPPED_CURRENT_CONTEXT(CVDContext);
-
-		// NOTE: ApplyKinematicTargetForParticle is run in a parallel-for. We only write to particle state
-		const auto& ApplyParticleKinematicTarget = 
-		[Dt, StepFraction, IsLastStep, CVDContext]
-		(FTransientPBDRigidParticleHandle& Particle, const int32 ParticleIndex)
-		-> void
-		{
-			CVD_SCOPE_CONTEXT(CVDContext.Context);
-
-			TKinematicTarget<FReal, 3>& KinematicTarget = Particle.KinematicTarget();
-			const FVec3 CurrentX = Particle.GetX();
-			const FRotation3 CurrentR = Particle.GetR();
-			constexpr FReal MinDt = 1e-6f;
-
-			bool bMoved = false;
-			switch (KinematicTarget.GetMode())
-			{
-			case EKinematicTargetMode::None:
-				// Nothing to do
-				break;
-
-			case EKinematicTargetMode::Reset:
-			{
-				// Reset velocity and then switch to do-nothing mode
-				Particle.SetVf(FVec3f(0.0f, 0.0f, 0.0f));
-				Particle.SetWf(FVec3f(0.0f, 0.0f, 0.0f));
-				Particle.ClearIsMovingKinematic();
-				KinematicTarget.SetMode(EKinematicTargetMode::None);
-				break;
-			}
-
-			case EKinematicTargetMode::Position:
-			{
-				// Move to kinematic target and update velocities to match
-				// Target positions only need to be processed once, and we reset the velocity next frame (if no new target is set)
-				FVec3 NewX;
-				FRotation3 NewR;
-				if (IsLastStep)
-				{
-					NewX = KinematicTarget.GetTarget().GetLocation();
-					NewR = KinematicTarget.GetTarget().GetRotation();
-					KinematicTarget.SetMode(EKinematicTargetMode::Reset);
-				}
-				else
-				{
-					// as a reminder, stepfraction is the remaing fraction of the step from the remaining steps
-					// for total of 4 steps and current step of 2, this will be 1/3 ( 1 step passed, 3 steps remains )
-					NewX = FVec3::Lerp(CurrentX, KinematicTarget.GetTarget().GetLocation(), StepFraction);
-					NewR = FRotation3::Slerp(CurrentR, KinematicTarget.GetTarget().GetRotation(), decltype(FQuat::X)(StepFraction));
-				}
-
-				const bool bPositionChanged = !FVec3::IsNearlyEqual(NewX, CurrentX, UE_SMALL_NUMBER);
-				const bool bRotationChanged = !FRotation3::IsNearlyEqual(NewR, CurrentR, UE_SMALL_NUMBER);
-				bMoved = bPositionChanged || bRotationChanged;
-				FVec3 NewV = FVec3(0);
-				FVec3 NewW = FVec3(0);
-				if (Dt > MinDt)
-				{
-					if (bPositionChanged)
-					{
-						NewV = FVec3::CalculateVelocity(CurrentX, NewX, Dt);
-					}
-					if (bRotationChanged)
-					{
-						NewW = FRotation3::CalculateAngularVelocity(CurrentR, NewR, Dt);
-					}
-				}
-				Particle.SetX(NewX);
-				Particle.SetR(NewR);
-				Particle.SetV(NewV);
-				Particle.SetW(NewW);
-				Particle.SetIsMovingKinematic();
-
-				break;
-			}
-
-			case EKinematicTargetMode::Velocity:
-			{
-				// Move based on velocity
-				bMoved = true;
-				Particle.SetX(Particle.GetX() + Particle.GetV() * Dt);
-				Particle.SetRf(FRotation3f::IntegrateRotationWithAngularVelocity(Particle.GetRf(), Particle.GetWf(), FRealSingle(Dt)));
-				Particle.SetIsMovingKinematic();
-
-				break;
-			}
-			}
-			
-			// Set positions and previous velocities if we can
-			// Note: At present kinematics are in fact rigid bodies
-			Particle.SetP(Particle.GetX());
-			Particle.SetQf(Particle.GetRf());
-			Particle.SetPreVf(Particle.GetVf());
-			Particle.SetPreWf(Particle.GetWf());
-
-			if (bMoved)
-			{
-				if (!Particle.CCDEnabled())
-				{
-					Particle.UpdateWorldSpaceState(FRigidTransform3(Particle.GetP(), Particle.GetQ()), FVec3(0));
-				}
-				else
-				{
-					Particle.UpdateWorldSpaceStateSwept(FRigidTransform3(Particle.GetP(), Particle.GetQ()), FVec3(0), -Particle.GetV() * Dt);
-				}
-			}
-
-			CVD_TRACE_PARTICLE(Particle.Handle())
-		};
-
-		// Apply kinematic targets in parallel
-		Particles.GetActiveMovingKinematicParticlesView().ParallelFor(ApplyParticleKinematicTarget);
-
-		// done with update, let's clear the tracking structures
-		if (IsLastStep)
-		{
-			Particles.UpdateAllMovingKinematic();
-		}
-
-		// If we changed any particle state, the views need to be refreshed
-		Particles.UpdateDirtyViews();
-	}
+	CHAOS_API virtual void ApplyKinematicTargets(const FReal Dt, const FReal StepFraction) {}
 
 	/** Make a copy of the acceleration structure to allow for external modification.
 	    This is needed for supporting sync operations on SQ structure from game thread. You probably want to go through solver which maintains PendingExternal */
diff --git a/Engine/Source/Runtime/Experimental/Chaos/Public/Chaos/PBDRigidsEvolutionGBF.h b/Engine/Source/Runtime/Experimental/Chaos/Public/Chaos/PBDRigidsEvolutionGBF.h
--- a/Engine/Source/Runtime/Experimental/Chaos/Public/Chaos/PBDRigidsEvolutionGBF.h
+++ b/Engine/Source/Runtime/Experimental/Chaos/Public/Chaos/PBDRigidsEvolutionGBF.h
@@ -211,6 +211,8 @@ namespace Chaos
 
 		CHAOS_API void Integrate(FReal Dt);
 
+		CHAOS_API virtual void ApplyKinematicTargets(const FReal Dt, const FReal StepFraction) override final;
+
 		CHAOS_API void Serialize(FChaosArchive& Ar);
 
 		CHAOS_API TUniquePtr<IResimCacheBase> CreateExternalResimCache() const;
diff --git a/Engine/Source/Runtime/Experimental/Chaos/Public/ChaosVisualDebugger/ChaosVDOptionalDataChannel.h b/Engine/Source/Runtime/Experimental/Chaos/Public/ChaosVisualDebugger/ChaosVDOptionalDataChannel.h
--- a/Engine/Source/Runtime/Experimental/Chaos/Public/ChaosVisualDebugger/ChaosVDOptionalDataChannel.h
+++ b/Engine/Source/Runtime/Experimental/Chaos/Public/ChaosVisualDebugger/ChaosVDOptionalDataChannel.h
@@ -142,13 +142,14 @@ namespace Chaos::VisualDebugger
 
 // Declare CVD's Default set of channels
 CVD_DECLARE_OPTIONAL_DATA_CHANNEL(Default);
-CVD_DECLARE_OPTIONAL_DATA_CHANNEL(EvolutionStart);
-CVD_DECLARE_OPTIONAL_DATA_CHANNEL(Integrate);
-CVD_DECLARE_OPTIONAL_DATA_CHANNEL(ApplyKinematicTargets);
+CVD_DECLARE_OPTIONAL_DATA_CHANNEL(EvolutionStart);							// Initial particle positions for this tick. Same as end of previous EvolutionEnd but with any external changes applied (set position, velocity, etc).
+CVD_DECLARE_OPTIONAL_DATA_CHANNEL(PostIntegrate);							// Predicted particle positions before user callbacks and constraint solver. Includes this tick's movement from external forces, velocity, kinematic targets, etc.
+CVD_DECLARE_OPTIONAL_DATA_CHANNEL(PreConstraintSolve);						// Predicted particle positions as used by the constraint solver. After user callbacks, collision modifiers, etc., but before collisions and joints are resolved.
 CVD_DECLARE_OPTIONAL_DATA_CHANNEL(CollisionDetectionBroadPhase);
 CVD_DECLARE_OPTIONAL_DATA_CHANNEL(CollisionDetectionNarrowPhase);
 CVD_DECLARE_OPTIONAL_DATA_CHANNEL(EndOfEvolutionCollisionConstraints);
-CVD_DECLARE_OPTIONAL_DATA_CHANNEL(EvolutionEnd);
+CVD_DECLARE_OPTIONAL_DATA_CHANNEL(PostConstraintSolve);						// Particle positions corrected by the constraint solver (collisions, joints, etc.), but before user callbacks or destruction handling.
+CVD_DECLARE_OPTIONAL_DATA_CHANNEL(EvolutionEnd);							// Final particle positions, including destruction and user callbacks.
 CVD_DECLARE_OPTIONAL_DATA_CHANNEL(SceneQueries);
 CVD_DECLARE_OPTIONAL_DATA_CHANNEL(JointConstraints);
 
diff --git a/Engine/Source/Runtime/Experimental/Chaos/Public/ChaosVisualDebugger/ChaosVDTraceMacros.h b/Engine/Source/Runtime/Experimental/Chaos/Public/ChaosVisualDebugger/ChaosVDTraceMacros.h
--- a/Engine/Source/Runtime/Experimental/Chaos/Public/ChaosVisualDebugger/ChaosVDTraceMacros.h
+++ b/Engine/Source/Runtime/Experimental/Chaos/Public/ChaosVisualDebugger/ChaosVDTraceMacros.h
@@ -86,7 +86,7 @@
 	#ifndef CVD_SCOPE_TRACE_SOLVER_STEP
 		#define CVD_SCOPE_TRACE_SOLVER_STEP(DataChannel, StepName) \
 			CVD_SCOPED_DATA_CHANNEL_OVERRIDE(DataChannel) \
-			FChaosVDScopeSolverStep ScopeSolverStep(StepName);
+			FChaosVDScopeSolverStep ScopeSolverStep(StepName)
 	#endif
 
 	#ifndef CVD_TRACE_BINARY_DATA
diff --git a/Engine/Source/Runtime/Experimental/ChaosVisualDebugger/Public/DataWrappers/ChaosVDParticleDataWrapper.h b/Engine/Source/Runtime/Experimental/ChaosVisualDebugger/Public/DataWrappers/ChaosVDParticleDataWrapper.h
--- a/Engine/Source/Runtime/Experimental/ChaosVisualDebugger/Public/DataWrappers/ChaosVDParticleDataWrapper.h
+++ b/Engine/Source/Runtime/Experimental/ChaosVisualDebugger/Public/DataWrappers/ChaosVDParticleDataWrapper.h
@@ -130,8 +130,8 @@ struct FChaosVDParticlePositionRotation : public FChaosVDWrapperDataBase
 	template <typename TOther>
 	void CopyFrom(const TOther& Other)
 	{
-		MX = Other.GetX();
-		MR = Other.GetR();
+		MX = Other->GetP();
+		MR = Other->GetQ();
 		
 		bHasValidData = true;
 	}
