Prevent duplicate unresolved handles from being resolved.

diff --git a/Engine/Source/Runtime/Experimental/Iris/Core/Private/Iris/ReplicationSystem/ReplicationReader.cpp b/Engine/Source/Runtime/Experimental/Iris/Core/Private/Iris/ReplicationSystem/ReplicationReader.cpp
--- a/Engine/Source/Runtime/Experimental/Iris/Core/Private/Iris/ReplicationSystem/ReplicationReader.cpp
+++ b/Engine/Source/Runtime/Experimental/Iris/Core/Private/Iris/ReplicationSystem/ReplicationReader.cpp
@@ -1690,13 +1690,17 @@ void FReplicationReader::ResolveAndDispatchUnresolvedReferences()
 	// Currently we brute force this by iterating over all handles pending resolve and update all objects pending resolve
 	TArray<FNetRefHandle> UpdatedHandles;
 	UpdatedHandles.Reserve(128);
-	UnresolvedHandleToDependents.GenerateKeyArray(UpdatedHandles);
+	UnresolvedHandleToDependents.GetKeys(UpdatedHandles);
 	
 	TSet<uint32> InternalObjectsToResolve;
 	InternalObjectsToResolve.Reserve(UnresolvedHandleToDependents.Num());
 
+	UE_LOG(LogIris, Log, TEXT("*** BEGIN RESOLVING ***"));
+
 	for (FNetRefHandle Handle : UpdatedHandles)
 	{
+		UE_LOG(LogIris, Log, TEXT("Resolving Id=%d"), Handle.GetId());
+
 		// Only make sense to update dependant objects if handle is resolvable
 		if (ObjectReferenceCache->ResolveObjectReferenceHandle(Handle, ResolveContext) != nullptr)
 		{
@@ -1716,6 +1720,9 @@ void FReplicationReader::ResolveAndDispatchUnresolvedReferences()
 		ResolveAndDispatchUnresolvedReferencesForObject(Context, InternalIndex);
 	}
 
+	CSV_CUSTOM_STAT(IrisClient, UnresolvedHandlesToResolve, UpdatedHandles.Num(), ECsvCustomStatOp::Accumulate);
+	CSV_CUSTOM_STAT(IrisClient, UnresolvedObjectsToResolve, InternalObjectsToResolve.Num(), ECsvCustomStatOp::Accumulate);
+
 	if (NumHandlesPendingResolveLastUpdate != UpdatedHandles.Num() || ObjectsWithAttachmentPendingResolve.Num() > 0)
 	{
 		UE_LOG_REPLICATIONREADER(TEXT("FReplicationReader::ResolveAndDispatchUnresolvedReferences NetHandles pending: %u Attachments pending: %u)"), UpdatedHandles.Num(), ObjectsWithAttachmentPendingResolve.Num());
