Added a helper debug method that pairs w/ GetMaxBreakStrainRatio, so you can get the max applied strain that caused the max break ratio

diff --git a/Engine/Source/Runtime/Experimental/Chaos/Private/Chaos/StrainModification.cpp b/Engine/Source/Runtime/Experimental/Chaos/Private/Chaos/StrainModification.cpp
--- a/Engine/Source/Runtime/Experimental/Chaos/Private/Chaos/StrainModification.cpp
+++ b/Engine/Source/Runtime/Experimental/Chaos/Private/Chaos/StrainModification.cpp
@@ -135,6 +135,33 @@ float Chaos::FStrainedProxyModifier::GetMaxBreakStrainRatio(const float FatigueT
 	return MaxBreakStrainRatio;
 }
 
+float Chaos::FStrainedProxyModifier::GetStrainUsedForBreakRatio(const float FatigueThresholdPercent, const float FatigueThresholdMinimum, bool bRelative, const uint8 StrainTypes)
+{
+	float MaxBreakStrainRatio = 0.f;
+	float StrainForBreakRatio = 0.f;
+	ForEachRootChildParticle(ProxyAndRoot, RestChildren,
+		[this, &MaxBreakStrainRatio, &StrainForBreakRatio, FatigueThresholdPercent, FatigueThresholdMinimum, bRelative, StrainTypes](Chaos::FPBDRigidClusteredParticleHandle* ChildHandle)
+		{
+			// compute the strain ratio
+			const float InternalStrain = ChildHandle->GetInternalStrains();
+			const float MaxAppliedStrain = GetMaxAppliedStrain(ChildHandle, StrainTypes);
+			const float FatigueThreshold = FMath::Max(FatigueThresholdMinimum, (FatigueThresholdPercent * InternalStrain));
+			if (MaxAppliedStrain >= FatigueThreshold)
+			{
+				const float StrainRange = bRelative ? FMath::Max(0.f, InternalStrain - FatigueThreshold) : InternalStrain;
+				const float Strain = bRelative ? FMath::Max(0.f, MaxAppliedStrain - FatigueThreshold) : MaxAppliedStrain;
+				const float StrainRatio = (StrainRange > SMALL_NUMBER) ? (Strain / StrainRange) : 1.0f;
+				if (StrainRatio > MaxBreakStrainRatio)
+				{
+					MaxBreakStrainRatio = StrainRatio;
+					StrainForBreakRatio = MaxAppliedStrain;
+				}
+			}
+		});
+
+	return StrainForBreakRatio;
+}
+
 void Chaos::FStrainedProxyModifier::AdjustStrainForBreak(const float FatigueThresholdPercent, const float FatigueThresholdMinimum, const uint8 StrainTypes)
 {
 	ForEachRootChildParticle(ProxyAndRoot, RestChildren,
diff --git a/Engine/Source/Runtime/Experimental/Chaos/Public/Chaos/StrainModification.h b/Engine/Source/Runtime/Experimental/Chaos/Public/Chaos/StrainModification.h
--- a/Engine/Source/Runtime/Experimental/Chaos/Public/Chaos/StrainModification.h
+++ b/Engine/Source/Runtime/Experimental/Chaos/Public/Chaos/StrainModification.h
@@ -88,6 +88,9 @@ namespace Chaos
 		// @param StrainTypes type strain to account for ( collision and/or external )
 		CHAOS_API float GetMaxBreakStrainRatio(const float FatigueThresholdPercent, const float FatigueThresholdMinimum, bool bRelative, const uint8 StrainTypes) const;
 
+		// A helper debug method to pair w/ GetMaxBreakStrainRatio, where it will return the Max Applied Strain used for the Max Break Strain. 
+		CHAOS_API float GetStrainUsedForBreakRatio(const float FatigueThresholdPercent, const float FatigueThresholdMinimum, bool bRelative, const uint8 StrainTypes);
+
 		// Clear strains for all strained cluster children
 		CHAOS_API void ClearStrains();
 
