Chaos - fix collision between BVH-wrapped shapes and non-scaled tri mesh.

The inner collision detection call does not handle the TInstancedImplicit wrapper that it assumes was stripped but the BVH midphase path did not strip the wrapper.

diff --git a/Engine/Source/Runtime/Experimental/Chaos/Private/Chaos/Collision/ParticlePairMidPhase.cpp b/Engine/Source/Runtime/Experimental/Chaos/Private/Chaos/Collision/ParticlePairMidPhase.cpp
--- a/Engine/Source/Runtime/Experimental/Chaos/Private/Chaos/Collision/ParticlePairMidPhase.cpp
+++ b/Engine/Source/Runtime/Experimental/Chaos/Private/Chaos/Collision/ParticlePairMidPhase.cpp
@@ -1584,6 +1584,16 @@ namespace Chaos
 		bool bShouldSwapParticles = false;
 		const EContactShapesType ShapePairType = Collisions::CalculateShapePairType(ParticleA, ImplicitA, ParticleB, ImplicitB, bShouldSwapParticles);
 
+		// Strip the Instanced wrapper if there is one
+		if (const FImplicitObjectInstanced* InstancedA = ImplicitA->AsA<FImplicitObjectInstanced>())
+		{
+			ImplicitA = InstancedA->GetInnerObject().Get();
+		}
+		if (const FImplicitObjectInstanced* InstancedB = ImplicitB->AsA<FImplicitObjectInstanced>())
+		{
+			ImplicitB = InstancedB->GetInnerObject().Get();
+		}
+
 		// Create the constraint for this shape pair
 		// NOTE: this just creates the object. The collision detection is done in ProcessNewConstraints
 		if (!bShouldSwapParticles)
