Remove redundant Iris logging.

diff --git a/Engine/Source/Runtime/Experimental/Iris/Core/Private/Iris/ReplicationSystem/ReplicationReader.cpp b/Engine/Source/Runtime/Experimental/Iris/Core/Private/Iris/ReplicationSystem/ReplicationReader.cpp
--- a/Engine/Source/Runtime/Experimental/Iris/Core/Private/Iris/ReplicationSystem/ReplicationReader.cpp
+++ b/Engine/Source/Runtime/Experimental/Iris/Core/Private/Iris/ReplicationSystem/ReplicationReader.cpp
@@ -1695,12 +1695,8 @@ void FReplicationReader::ResolveAndDispatchUnresolvedReferences()
 	TSet<uint32> InternalObjectsToResolve;
 	InternalObjectsToResolve.Reserve(UnresolvedHandleToDependents.Num());
 
-	UE_LOG(LogIris, Log, TEXT("*** BEGIN RESOLVING ***"));
-
 	for (FNetRefHandle Handle : UpdatedHandles)
 	{
-		UE_LOG(LogIris, Log, TEXT("Resolving Id=%d"), Handle.GetId());
-
 		// Only make sense to update dependant objects if handle is resolvable
 		if (ObjectReferenceCache->ResolveObjectReferenceHandle(Handle, ResolveContext) != nullptr)
 		{
