Reduce severity of nav export warning in geoemtry collection. May be useful to enable when disagnosing nav issues but in editor modes the log messages are overly verbose and slow down movement of objects.

diff --git a/Engine/Source/Runtime/Experimental/GeometryCollectionEngine/Private/GeometryCollection/GeometryCollectionComponent.cpp b/Engine/Source/Runtime/Experimental/GeometryCollectionEngine/Private/GeometryCollection/GeometryCollectionComponent.cpp
--- a/Engine/Source/Runtime/Experimental/GeometryCollectionEngine/Private/GeometryCollection/GeometryCollectionComponent.cpp
+++ b/Engine/Source/Runtime/Experimental/GeometryCollectionEngine/Private/GeometryCollection/GeometryCollectionComponent.cpp
@@ -1452,7 +1452,7 @@ bool UGeometryCollectionComponent::DoCustomNavigableGeometryExport(FNavigableGeo
 
 	if (RestCollection->bStripOnCook)
 	{
-		UE_LOG(LogGeometryCollectionComponent, Warning,
+		UE_LOG(LogGeometryCollectionComponent, Log,
 			TEXT("Collection data is set to be stripped on cook so there is nothing left to export. Make sure to use bUseRootProxyForNavigation with bStripOnCook for %s."), *GetFullNameSafe(GetOwner()));
 	}
 	
