Chaos - api to prevent friction undoing server corrections

This is especially problemtic with high friction and could undo all movement or cause jitter.

diff --git a/Engine/Source/Runtime/Experimental/Chaos/Private/Chaos/PBDRigidsEvolutionGBF.cpp b/Engine/Source/Runtime/Experimental/Chaos/Private/Chaos/PBDRigidsEvolutionGBF.cpp
--- a/Engine/Source/Runtime/Experimental/Chaos/Private/Chaos/PBDRigidsEvolutionGBF.cpp
+++ b/Engine/Source/Runtime/Experimental/Chaos/Private/Chaos/PBDRigidsEvolutionGBF.cpp
@@ -1338,7 +1338,7 @@ void FPBDRigidsEvolutionGBF::DestroyTransientConstraints()
 	}
 }
 
-CHAOS_API void FPBDRigidsEvolutionGBF::SetParticleTransform(FGeometryParticleHandle* InParticle, const FVec3& InPos, const FRotation3& InRot, const bool bIsTeleport)
+void FPBDRigidsEvolutionGBF::SetParticleTransform(FGeometryParticleHandle* InParticle, const FVec3& InPos, const FRotation3& InRot, const bool bIsTeleport)
 {
 	const FVec3 PrevX = InParticle->GetX();
 	const FRotation3 PrevR = InParticle->GetR();
@@ -1355,7 +1355,6 @@ CHAOS_API void FPBDRigidsEvolutionGBF::SetParticleTransform(FGeometryParticleHan
 #endif
 }
 
-
 void FPBDRigidsEvolutionGBF::SetParticleTransformSwept(FGeometryParticleHandle* InParticle, const FVec3& InPos, const FRotation3& InRot, const bool bIsTeleport)
 {
 	// If the rotation has changed, we need to adjust the initial sweep position to account for the fact that the
@@ -1412,6 +1411,41 @@ void FPBDRigidsEvolutionGBF::SetParticleKinematicTarget(FGeometryParticleHandle*
 	}
 }
 
+void FPBDRigidsEvolutionGBF::SetParticleTransformCorrection(FGeometryParticleHandle* InParticle, const FVec3& InPosDelta, const FVec3& InRotDelta)
+{
+	FGenericParticleHandle Particle = InParticle;
+	const bool bHavePositionDelta = !InPosDelta.IsZero();
+	const bool bHaveRotationDelta = !InRotDelta.IsZero();
+
+	if (bHavePositionDelta || bHaveRotationDelta)
+	{
+		const FVec3 PrevX = Particle->GetX();
+		const FRotation3 PrevR = Particle->GetR();
+
+		// Apply the position correction
+		if (bHavePositionDelta)
+		{
+			const FVec3 NewX = Particle->GetX() + InPosDelta;
+			Particle->SetX(NewX);
+		}
+
+		// Apply the rotation correction
+		if (bHaveRotationDelta)
+		{
+			const FRotation3 NewR = FRotation3::IntegrateRotationWithAngularVelocity(Particle->GetR(), InRotDelta, 1.0f);
+			Particle->SetR(NewR);
+		}
+
+		// If we applied a correction we must fix the collision anchors so that friction doesn't undo prevent our move or rotation
+		InParticle->ParticleCollisions().VisitCollisions(
+			[this, InParticle](FPBDCollisionConstraint& Collision)
+			{
+				Collision.UpdateParticleTransform(InParticle);
+				return ECollisionVisitorResult::Continue;
+			});
+	}
+}
+
 void FPBDRigidsEvolutionGBF::OnParticleMoved(FGeometryParticleHandle* InParticle, const FVec3& PrevX, const FRotation3& PrevR, const bool bIsTeleport)
 {
 	// When a particle is moved, we need to 
diff --git a/Engine/Source/Runtime/Experimental/Chaos/Public/Chaos/PBDRigidsEvolutionGBF.h b/Engine/Source/Runtime/Experimental/Chaos/Public/Chaos/PBDRigidsEvolutionGBF.h
--- a/Engine/Source/Runtime/Experimental/Chaos/Public/Chaos/PBDRigidsEvolutionGBF.h
+++ b/Engine/Source/Runtime/Experimental/Chaos/Public/Chaos/PBDRigidsEvolutionGBF.h
@@ -171,6 +171,13 @@ namespace Chaos
 		*/
 		CHAOS_API void SetParticleKinematicTarget(FGeometryParticleHandle* ParticleHandle, const FKinematicTarget& NewKinematicTarget);
 
+		/*
+		 * Apply a momentumless correction to the particle transform, usually as a result of a server correction.
+		 * This will shift the particle by the supplied delta and handle updating of friction anchors or anything else that might prevent or undo the shift.
+		 * NOTE: must be called prior to Integrate() to be effective.
+		 */
+		CHAOS_API void SetParticleTransformCorrection(FGeometryParticleHandle* InParticle, const FVec3& InPosDelta, const FVec3& InRotDelta);
+
 		/**
 		 * Called when a particle is moved. We need to reset some friction properties, sleeping properties, etc
 		 */
