Iris - Fix issue causing lost state during an object being marked for destroy to never be resent. If the object ends up cancelling pending destroy this could cause state to be lost forever and also bitstream errors on objects with delta compression enabled.

#rnx

diff --git a/Engine/Plugins/Runtime/ReplicationSystemTestPlugin/Source/Private/Tests/ReplicationSystem/TestObjectDeltaSerialization.cpp b/Engine/Plugins/Runtime/ReplicationSystemTestPlugin/Source/Private/Tests/ReplicationSystem/TestObjectDeltaSerialization.cpp
--- a/Engine/Plugins/Runtime/ReplicationSystemTestPlugin/Source/Private/Tests/ReplicationSystem/TestObjectDeltaSerialization.cpp
+++ b/Engine/Plugins/Runtime/ReplicationSystemTestPlugin/Source/Private/Tests/ReplicationSystem/TestObjectDeltaSerialization.cpp
@@ -433,6 +433,77 @@ UE_NET_TEST_FIXTURE(FTestObjectDeltaSerialization, InFlightChangesForDisabledCon
 	UE_NET_ASSERT_EQ(ClientObject->IntA, ServerObject->IntA);
 }
 
+UE_NET_TEST_FIXTURE(FTestObjectDeltaSerialization, LostChangesDuringPendingDestroyArePartOfBaselineAfterCancellingPendingDestroy)
+{
+	FReplicationSystemTestClient* Client = CreateClient();
+
+	// Spawn object on server
+	UTestReplicatedIrisObject* ServerObject = Server->CreateObject({ .DynamicStateComponentCount = 1 });
+	Server->ReplicationSystem->SetDeltaCompressionStatus(ServerObject->NetRefHandle, ENetObjectDeltaCompressionStatus::Allow);
+
+	// Send and deliver packet
+	Server->UpdateAndSend({ Client });
+
+	const UTestReplicatedIrisObject* ClientObject = Cast<UTestReplicatedIrisObject>(Client->GetReplicationBridge()->GetReplicatedObject(ServerObject->NetRefHandle));
+	UE_NET_ASSERT_NE(ClientObject, nullptr);
+
+	AllowNewBaselineCreation();
+
+	// Modify array
+	ServerObject->DynamicStateComponents[0]->IntArray.Add(23456789);
+
+	// Send and deliver packet
+	Server->UpdateAndSend({ Client });
+
+	// Clear array
+	DisallowNewBaselineCreation();
+	ServerObject->DynamicStateComponents[0]->IntArray.Empty();
+
+	// Fake latency by acking/naking later
+	Server->PreSendUpdate();
+	Server->SendTo(Client);
+	Server->PostSendUpdate();
+
+	// Filter out object to cause a PendingDestroy
+	Server->GetReplicationSystem()->AddToGroup(Server->GetReplicationSystem()->GetNotReplicatedNetObjectGroup(), ServerObject->NetRefHandle);
+	Server->PreSendUpdate();
+	Server->SendTo(Client);
+	Server->PostSendUpdate();
+
+	// Drop all packets
+	{
+		const auto& ConnectionInfo = Server->GetConnectionInfo(Client->ConnectionIdOnServer);
+		SIZE_T PacketCount = ConnectionInfo.WrittenPackets.Count();
+		for (SIZE_T PacketIt = 0; PacketIt != PacketCount; ++PacketIt)
+		{
+			Server->DeliverTo(Client, DoNotDeliverPacket);
+		}
+	}
+
+	// Remove object from filter to cause object to end up in CancelPendingDestroy
+	DisallowNewBaselineCreation();
+	Server->GetReplicationSystem()->RemoveFromGroup(Server->GetReplicationSystem()->GetNotReplicatedNetObjectGroup(), ServerObject->NetRefHandle);
+	Server->UpdateAndSend({ Client });
+
+	// Change other properties to force replication
+	ServerObject->IntA += 1;
+	ServerObject->DynamicStateComponents[0]->IntStaticArray[0] += 1;
+
+	AllowNewBaselineCreation();
+	Server->UpdateAndSend({ Client });
+
+	// Finally modify the IntArray.
+	ServerObject->IntA += 1;
+	ServerObject->DynamicStateComponents[0]->IntArray.Add(-12345);
+
+	DisallowNewBaselineCreation();
+	Server->UpdateAndSend({ Client });
+
+	// Verify IntArray contents
+	UE_NET_ASSERT_EQ(ClientObject->DynamicStateComponents[0]->IntArray.Num(), ServerObject->DynamicStateComponents[0]->IntArray.Num());
+	UE_NET_ASSERT_EQ(ClientObject->DynamicStateComponents[0]->IntArray[0], ServerObject->DynamicStateComponents[0]->IntArray[0]);
+}
+
 // FTestObjectDeltaSerialization implementation
 void FTestObjectDeltaSerialization::SetUp()
 {
diff --git a/Engine/Source/Runtime/Experimental/Iris/Core/Private/Iris/ReplicationSystem/ReplicationWriter.cpp b/Engine/Source/Runtime/Experimental/Iris/Core/Private/Iris/ReplicationSystem/ReplicationWriter.cpp
--- a/Engine/Source/Runtime/Experimental/Iris/Core/Private/Iris/ReplicationSystem/ReplicationWriter.cpp
+++ b/Engine/Source/Runtime/Experimental/Iris/Core/Private/Iris/ReplicationSystem/ReplicationWriter.cpp
@@ -1262,8 +1262,8 @@ void FReplicationWriter::HandleDroppedRecord<FReplicationWriter::EReplicatedObje
 {
 	const uint32 InternalIndex = RecordInfo.Index;
 
-	// For now we do not support flush (https://jira.it.epicgames.com/browse/UENET-1079), so if we drop data in flight while we are pending destroy/tear-off we will ignore the lost data.
-	if (CurrentState < EReplicatedObjectState::PendingDestroy)
+	// An object in PendingDestroy/WaitOnDestroyConfirmation can end up being replicated again via CancelPendingDestroy.
+	if (CurrentState < EReplicatedObjectState::Destroyed)
 	{
 		// Mask in any lost changes
 		bool bNeedToResendAttachments = RecordInfo.HasAttachments;
@@ -1298,36 +1298,39 @@ void FReplicationWriter::HandleDroppedRecord<FReplicationWriter::EReplicatedObje
 				ChangeMask.Combine(LostChangeMask, FNetBitArrayView::OrOp);
 			}
 
-			// Mark object as having dirty changes
-			MarkObjectDirty(InternalIndex, "DroppedCreated");
+			if (CurrentState < EReplicatedObjectState::PendingDestroy)
+			{
+				// Mark object as having dirty changes
+				MarkObjectDirty(InternalIndex, "DroppedCreated");
 
-			// Mark changemask as dirty
-			Info.HasDirtyChangeMask |= bNeedToResendState;
+				// Mark changemask as dirty
+				Info.HasDirtyChangeMask |= bNeedToResendState;
 
-			// Mark attachments as dirty
-			Info.HasAttachments |= bNeedToResendAttachments;
+				// Mark attachments as dirty
+				Info.HasAttachments |= bNeedToResendAttachments;
 
-			// Give slight priority bump
-			SchedulingPriorities[InternalIndex] += FReplicationWriter::LostStatePriorityBump;
+				// Give slight priority bump
+				SchedulingPriorities[InternalIndex] += FReplicationWriter::LostStatePriorityBump;
 
-			if (Info.IsSubObject)
-			{
-				// Mark owner dirty as well as subobjects only are scheduled together with owner
-				const FNetRefHandleManager::FReplicatedObjectData& ObjectData = NetRefHandleManager->GetReplicatedObjectData(InternalIndex);
-				uint32 SubObjectOwnerInternalIndex = ObjectData.SubObjectRootIndex;
+				if (Info.IsSubObject)
+				{
+					// Mark owner dirty as well as subobjects only are scheduled together with owner
+					const FNetRefHandleManager::FReplicatedObjectData& ObjectData = NetRefHandleManager->GetReplicatedObjectData(InternalIndex);
+					uint32 SubObjectOwnerInternalIndex = ObjectData.SubObjectRootIndex;
 
-				FReplicationInfo& SubObjectOwnerReplicationInfo = GetReplicationInfo(SubObjectOwnerInternalIndex);
+					FReplicationInfo& SubObjectOwnerReplicationInfo = GetReplicationInfo(SubObjectOwnerInternalIndex);
 
-				if (ensure(SubObjectOwnerReplicationInfo.GetState() < EReplicatedObjectState::PendingDestroy))
-				{
-					// Mark owner as dirty
-					MarkObjectDirty(SubObjectOwnerInternalIndex, "DroppedCreated2");
+					if (ensure(SubObjectOwnerReplicationInfo.GetState() < EReplicatedObjectState::PendingDestroy))
+					{
+						// Mark owner as dirty
+						MarkObjectDirty(SubObjectOwnerInternalIndex, "DroppedCreated2");
 
-					// Indicate that we have dirty subobjects
-					SubObjectOwnerReplicationInfo.HasDirtySubObjects = 1U;
+						// Indicate that we have dirty subobjects
+						SubObjectOwnerReplicationInfo.HasDirtySubObjects = 1U;
 
-					// Give slight priority bump to owner
-					SchedulingPriorities[SubObjectOwnerInternalIndex] += FReplicationWriter::LostStatePriorityBump;
+						// Give slight priority bump to owner
+						SchedulingPriorities[SubObjectOwnerInternalIndex] += FReplicationWriter::LostStatePriorityBump;
+					}
 				}
 			}
 		}
