Gracefully handle the case where the pre-allocated handle configuration is greater than the maximum in NetRefHandleManager.

diff --git a/Engine/Source/Runtime/Experimental/Iris/Core/Private/Iris/ReplicationSystem/NetRefHandleManager.cpp b/Engine/Source/Runtime/Experimental/Iris/Core/Private/Iris/ReplicationSystem/NetRefHandleManager.cpp
--- a/Engine/Source/Runtime/Experimental/Iris/Core/Private/Iris/ReplicationSystem/NetRefHandleManager.cpp
+++ b/Engine/Source/Runtime/Experimental/Iris/Core/Private/Iris/ReplicationSystem/NetRefHandleManager.cpp
@@ -25,7 +25,6 @@ FNetRefHandleManager::FNetRefHandleManager(FReplicationProtocolManager& InReplic
 : ActiveObjectCount(0)
 , MaxActiveObjectCount(InMaxActiveObjectCount)
 , PreAllocatedObjectCount(InPreAllocatedObjectCount > 0 ? InPreAllocatedObjectCount : 1)	// PreAllocatedObjectCount must be a minimum of 1 to account for InvalidInternalIndex.
-, LargestInternalIndex(InPreAllocatedObjectCount > 0 ? InPreAllocatedObjectCount - 1 : 0)	// LargestInternalIndex must be a minimum of 0 to account for InvalidInternalIndex.
 , ReplicationSystemId(InReplicationSystemId)
 , GlobalScopableInternalIndices(MaxActiveObjectCount)
 , ScopeFrameData(MaxActiveObjectCount)
@@ -45,10 +44,13 @@ FNetRefHandleManager::FNetRefHandleManager(FReplicationProtocolManager& InReplic
 	// Ensure that the pre-allocated object count is not greater than the maximum. If it is, just set it to the maximum.
 	if (PreAllocatedObjectCount > MaxActiveObjectCount)
 	{
-		UE_LOG(LogIris, Error, TEXT("PreAllocatedObjectCount (%d) must be smaller or equal to MaxActiveObjectCount (%d)!"), PreAllocatedObjectCount, MaxActiveObjectCount);
+		UE_LOG(LogIris, Error, TEXT("PreAllocatedObjectCount (%d) must be smaller or equal to MaxActiveObjectCount (%d). Setting PreAllocatedObjectCount to MaxActiveObjectCount."), PreAllocatedObjectCount, MaxActiveObjectCount);
 		PreAllocatedObjectCount = MaxActiveObjectCount;
 	}
 
+	// Calculate the largest internal index and must be a minimum of 0 to support InvalidInternalIndex.
+	LargestInternalIndex = PreAllocatedObjectCount > 0 ? PreAllocatedObjectCount - 1 : 0;
+
 	UE_LOG(LogIris, Log, TEXT("NetRefHandleManager: Configured with MaxActiveObjectCount=%d, PreAllocatedObjectCount=%d and LargestInternalIndex=%d."), MaxActiveObjectCount, PreAllocatedObjectCount, LargestInternalIndex);
 
 	static_assert(InvalidInternalIndex == 0, "FNetRefHandleManager::InvalidInternalIndex has an unexpected value");
