Fix world locations not being reset when destroying an object

#rnx

diff --git a/Engine/Source/Runtime/Experimental/Iris/Core/Private/Iris/ReplicationSystem/NetRefHandleManager.cpp b/Engine/Source/Runtime/Experimental/Iris/Core/Private/Iris/ReplicationSystem/NetRefHandleManager.cpp
--- a/Engine/Source/Runtime/Experimental/Iris/Core/Private/Iris/ReplicationSystem/NetRefHandleManager.cpp
+++ b/Engine/Source/Runtime/Experimental/Iris/Core/Private/Iris/ReplicationSystem/NetRefHandleManager.cpp
@@ -701,7 +701,7 @@ bool FNetRefHandleManager::AddDependentObject(FNetRefHandle ParentRefHandle, FNe
 	{
 		// If this gets out of sync something is messed up
 		check(bDependentHadParentAlready == bDependentIsAlreadyDependant);
-		ensureAlwaysMsgf(false, TEXT("FNetRefHandleManager::AddDependentObject %s already is a child of %s parent %s"), *DependentObjectRefHandle.ToString(), *GetReplicatedObjectDataNoCheck(DependentObjectData.SubObjectRootIndex).RefHandle.ToString(), *ParentRefHandle.ToString());
+		ensureMsgf(false, TEXT("FNetRefHandleManager::AddDependentObject %s already is a child of %s parent %s"), *PrintObjectFromNetRefHandle(DependentObjectRefHandle), *PrintObjectFromIndex(DependentObjectData.SubObjectRootIndex), *PrintObjectFromNetRefHandle(ParentRefHandle));
 	}
 
 	return true;
diff --git a/Engine/Source/Runtime/Experimental/Iris/Core/Private/Iris/ReplicationSystem/NetRefHandleManager.h b/Engine/Source/Runtime/Experimental/Iris/Core/Private/Iris/ReplicationSystem/NetRefHandleManager.h
--- a/Engine/Source/Runtime/Experimental/Iris/Core/Private/Iris/ReplicationSystem/NetRefHandleManager.h
+++ b/Engine/Source/Runtime/Experimental/Iris/Core/Private/Iris/ReplicationSystem/NetRefHandleManager.h
@@ -290,7 +290,8 @@ public:
 
 	/** Return a string to identify the object linked to an index in logs */
 	FString PrintObjectFromIndex(FInternalNetRefIndex ObjectIndex) const;
-	
+	FString PrintObjectFromNetRefHandle(FNetRefHandle ObjectHandle) const { return PrintObjectFromIndex(GetInternalIndex(ObjectHandle)); }
+
 	DECLARE_MULTICAST_DELEGATE_OneParam(FOnLargestIndexIncrease, uint32 LargestIndex);
 
 	/** Return a delegate that will notify when the largest internal index has increased. */
diff --git a/Engine/Source/Runtime/Experimental/Iris/Core/Private/Iris/ReplicationSystem/ReplicationBridge.cpp b/Engine/Source/Runtime/Experimental/Iris/Core/Private/Iris/ReplicationSystem/ReplicationBridge.cpp
--- a/Engine/Source/Runtime/Experimental/Iris/Core/Private/Iris/ReplicationSystem/ReplicationBridge.cpp
+++ b/Engine/Source/Runtime/Experimental/Iris/Core/Private/Iris/ReplicationSystem/ReplicationBridge.cpp
@@ -148,6 +148,8 @@ void UReplicationBridge::DestroyNetObjectFromRemote(FNetRefHandle Handle, ERepli
 
 	if (Handle.IsValid())
 	{
+		UE_LOG_REPLICATIONBRIDGE(Verbose, TEXT("DestroyNetObjectFromRemote for %s | DestroyReason: %s | DestroyFlags: %s "), *NetRefHandleManager->PrintObjectFromNetRefHandle(Handle), LexToString(DestroyReason), LexToString(DestroyFlags) );
+
 		FInternalNetRefIndex OwnerInternalIndex = NetRefHandleManager->GetInternalIndex(Handle);
 		FNetRefHandleManager::FReplicatedObjectData& ObjectData = NetRefHandleManager->GetReplicatedObjectDataNoCheck(OwnerInternalIndex);
 		ObjectData.bTearOff = (DestroyReason == EReplicationBridgeDestroyInstanceReason::TearOff);
@@ -165,8 +167,6 @@ void UReplicationBridge::DestroyNetObjectFromRemote(FNetRefHandle Handle, ERepli
 	
 		// Destroy the NetRefHandle
 		InternalDestroyNetObject(Handle);
-
-		UE_LOG_REPLICATIONBRIDGE(Verbose, TEXT("DestroyNetObjectFromRemote Remote %s"), *Handle.ToString());
 	}
 }
 
@@ -407,9 +407,14 @@ void UReplicationBridge::InternalDestroyNetObject(FNetRefHandle Handle)
 
 	if (const FInternalNetRefIndex ObjectInternalIndex = NetRefHandleManager->GetInternalIndex(Handle))
 	{
-		FNetCullDistanceOverrides& NetCullDistanceOverrides = ReplicationSystem->GetReplicationSystemInternal()->GetNetCullDistanceOverrides();
+		FReplicationSystemInternal* ReplicationSystemInternal = ReplicationSystem->GetReplicationSystemInternal();
+
+		FNetCullDistanceOverrides& NetCullDistanceOverrides = ReplicationSystemInternal->GetNetCullDistanceOverrides();
 		NetCullDistanceOverrides.ClearCullDistanceSqr(ObjectInternalIndex);
 
+		FWorldLocations& WorldLocations = ReplicationSystemInternal->GetWorldLocations();
+		WorldLocations.RemoveObjectInfoCache(ObjectInternalIndex);
+
 		const FNetRefHandleManager::FReplicatedObjectData& ObjectData = NetRefHandleManager->GetReplicatedObjectDataNoCheck(ObjectInternalIndex);
 
 		if (ObjectData.bTearOff && NetRefHandleManager->GetNetObjectRefCount(ObjectInternalIndex) > 0U)
@@ -417,7 +422,7 @@ void UReplicationBridge::InternalDestroyNetObject(FNetRefHandle Handle)
 			// We need to explicitly notify all ReplicationWriters that we are destroying objects pending tearoff
 			// The handle will automatically be removed from HandlesPendingTearOff after the next update
 
-			FReplicationConnections& Connections = ReplicationSystem->GetReplicationSystemInternal()->GetConnections();
+			FReplicationConnections& Connections = ReplicationSystemInternal->GetConnections();
 
 			auto NotifyDestroyedObjectPendingTearOff = [&Connections, &ObjectInternalIndex](uint32 ConnectionId)
 			{
@@ -435,6 +440,8 @@ void UReplicationBridge::InternalDestroyNetObject(FNetRefHandle Handle)
 
 void UReplicationBridge::DestroyLocalNetHandle(FNetRefHandle Handle, EEndReplicationFlags EndReplicationFlags)
 {
+	UE_LOG_REPLICATIONBRIDGE(Verbose, TEXT("DestroyLocalNetHandle for %s | EndReplicationFlags: %s"), *NetRefHandleManager->PrintObjectFromNetRefHandle(Handle), *LexToString(EndReplicationFlags));
+
 	if (EnumHasAnyFlags(EndReplicationFlags, EEndReplicationFlags::DestroyNetHandle | EEndReplicationFlags::ClearNetPushId))
 	{
 		const UE::Net::Private::FInternalNetRefIndex InternalReplicationIndex = NetRefHandleManager->GetInternalIndex(Handle);
@@ -464,8 +471,6 @@ void UReplicationBridge::DestroyLocalNetHandle(FNetRefHandle Handle, EEndReplica
 
 	// Tell ReplicationSystem to destroy the handle
 	InternalDestroyNetObject(Handle);
-
-	UE_LOG_REPLICATIONBRIDGE(Verbose, TEXT("DestroyLocalNetHandle Local %s"), *Handle.ToString());
 }
 
 void UReplicationBridge::InternalAddSubObject(FNetRefHandle OwnerHandle, FNetRefHandle SubObjectHandle, FNetRefHandle InsertRelativeToSubObjectHandle, ESubObjectInsertionOrder InsertionOrder)
@@ -504,7 +509,7 @@ void UReplicationBridge::InternalDestroySubObjects(FNetRefHandle OwnerHandle, EE
 			if (bDestroySubObjectWithOwner && NetRefHandleManager->IsScopableIndex(SubObjectInternalIndex))
 			{
 				SubObjectData.bPendingEndReplication = 1U;
-				UE_LOG_REPLICATIONBRIDGE(Verbose, TEXT("InternalDestroySubObjects %s - SubObject %s"), *OwnerHandle.ToString(), *SubObjectHandle.ToString());
+				UE_LOG_REPLICATIONBRIDGE(Verbose, TEXT("InternalDestroySubObjects %s"), *NetRefHandleManager->PrintObjectFromNetRefHandle(SubObjectHandle));
 				DestroyLocalNetHandle(SubObjectHandle, Flags);
 			}
 		}
@@ -875,7 +880,7 @@ UE::Net::FNetRefHandle UReplicationBridge::InternalAddDestructionInfo(FNetRefHan
 		GetReplicationSystem()->SetPrioritizer(DestructionInfoHandle, DefaultSpatialNetObjectPrioritizerHandle);
 	}
 
-	UE_LOG_REPLICATIONBRIDGE(Verbose, TEXT("UReplicationBridge::InternalAddDestructionInfo %s ( InternalIndex: %u ) for %s GroupIndex: %u"), *DestructionInfoHandle.ToString(), InternalReplicationIndex,  *Handle.ToString(), LevelGroupHandle.GetGroupIndex());
+	UE_LOG_REPLICATIONBRIDGE(Verbose, TEXT("UReplicationBridge::InternalAddDestructionInfo %s (InternalIndex: %u ) for %s GroupIndex: %u"), *DestructionInfoHandle.ToString(), InternalReplicationIndex,  *PrintObjectFromNetRefHandle(Handle), LevelGroupHandle.GetGroupIndex());
 	
 	return DestructionInfoHandle;
 }
@@ -1005,6 +1010,11 @@ void UReplicationBridge::ClearNetPushIds(UE::Net::Private::FInternalNetRefIndex
 #endif
 }
 
+FString UReplicationBridge::PrintObjectFromNetRefHandle(FNetRefHandle RefHandle) const
+{
+	return NetRefHandleManager->PrintObjectFromNetRefHandle(RefHandle);
+}
+
 const TCHAR* LexToString(EReplicationBridgeDestroyInstanceReason Reason)
 {
 	switch (Reason)
@@ -1025,5 +1035,67 @@ const TCHAR* LexToString(EReplicationBridgeDestroyInstanceReason Reason)
 	{
 		return TEXT("[Invalid]");
 	}
-	};
+	}
 }
+
+const TCHAR* LexToString(EReplicationBridgeDestroyInstanceFlags DestroyFlags)
+{
+	switch (DestroyFlags)
+	{
+	case EReplicationBridgeDestroyInstanceFlags::None:
+	{
+		return TEXT("None");
+	}
+	case EReplicationBridgeDestroyInstanceFlags::AllowDestroyInstanceFromRemote:
+	{
+		return TEXT("AllowDestroyInstanceFromRemote");
+	}
+	default:
+	{
+		return TEXT("[Invalid]");
+	}
+	}
+}
+
+FString LexToString(EEndReplicationFlags EndReplicationFlags)
+{
+	
+	if (EndReplicationFlags == EEndReplicationFlags::None)
+	{
+		return TEXT("None");
+	}
+
+	FString Flags;
+
+	if (EnumHasAnyFlags(EndReplicationFlags, EEndReplicationFlags::Destroy))
+	{
+		Flags += TEXT("Destroy");
+	}
+	if (EnumHasAnyFlags(EndReplicationFlags, EEndReplicationFlags::TearOff))
+	{
+		if (!Flags.IsEmpty()) { Flags += TEXT(','); }
+		Flags += TEXT("TearOff");
+	}
+	if (EnumHasAnyFlags(EndReplicationFlags, EEndReplicationFlags::Flush))
+	{
+		if (!Flags.IsEmpty()) { Flags += TEXT(','); }
+		Flags += TEXT("Flush");
+	}
+	if (EnumHasAnyFlags(EndReplicationFlags, EEndReplicationFlags::DestroyNetHandle))
+	{
+		if (!Flags.IsEmpty()) { Flags += TEXT(','); };
+		Flags += TEXT("DestroyNetHandle");
+	}
+	if (EnumHasAnyFlags(EndReplicationFlags, EEndReplicationFlags::ClearNetPushId))
+	{
+		if (!Flags.IsEmpty()) { Flags += TEXT(','); };
+		Flags += TEXT("ClearNetPushId");
+	}
+	if (EnumHasAnyFlags(EndReplicationFlags, EEndReplicationFlags::SkipPendingEndReplicationValidation))
+	{
+		if (!Flags.IsEmpty()) { Flags += TEXT(','); };
+		Flags += TEXT("SkipPendingEndReplicationValidation");
+	}
+	
+	return Flags;
+}
\ No newline at end of file
diff --git a/Engine/Source/Runtime/Experimental/Iris/Core/Public/Iris/ReplicationSystem/ReplicationBridge.h b/Engine/Source/Runtime/Experimental/Iris/Core/Public/Iris/ReplicationSystem/ReplicationBridge.h
--- a/Engine/Source/Runtime/Experimental/Iris/Core/Public/Iris/ReplicationSystem/ReplicationBridge.h
+++ b/Engine/Source/Runtime/Experimental/Iris/Core/Public/Iris/ReplicationSystem/ReplicationBridge.h
@@ -79,6 +79,8 @@ enum class EEndReplicationFlags : uint32
 };
 ENUM_CLASS_FLAGS(EEndReplicationFlags);
 
+IRISCORE_API FString LexToString(EEndReplicationFlags EndReplicationFlags);
+
 enum class EReplicationBridgeCreateNetRefHandleResultFlags : unsigned
 {
 	None = 0U,
@@ -110,6 +112,8 @@ enum class EReplicationBridgeDestroyInstanceFlags : unsigned
 };
 ENUM_CLASS_FLAGS(EReplicationBridgeDestroyInstanceFlags);
 
+IRISCORE_API const TCHAR* LexToString(EReplicationBridgeDestroyInstanceFlags DestroyFlags);
+
 UCLASS(Transient, MinimalAPI)
 class UReplicationBridge : public UObject
 {
@@ -160,6 +164,9 @@ public:
 	/** Get the group associated with the level in order to control connection filtering for it. */
 	IRISCORE_API UE::Net::FNetObjectGroupHandle GetLevelGroup(const UObject* Level) const;
 
+	/** Print common information about this handle and the object it is mapped to */
+	[[nodiscard]] IRISCORE_API FString PrintObjectFromNetRefHandle(FNetRefHandle RefHandle) const;
+
 protected:
 	/** Initializes the bridge. Is called during ReplicationSystem initialization. */
 	IRISCORE_API virtual void Initialize(UReplicationSystem* InReplicationSystem);
