Iris
* Added NetObjectFilters::PrintDebugInfoForObject virtual so filters can output loggable info on objects they filter
* Implemented NetObjectGridFilter::PrintDebugInfoForObject to log distance and relevancy list of objects for a given connection
* Use "net.Iris.PrintRelevantObjectsToConnection WithFilter"  to log filter information on relevant RootObjects
* Added optional params to many net.Iris.PrintXXX commands:
    * WithSubObjects -> also log the subobjects attached to a root object
    * SortByClass -> sort the list by class name
    * SortByNetRefHandle -> sort the list by Id with static objects in one group followed by dynamic objects
* Added "Net.Iris.PrintNetCullDistances" that prints the culldistances values of actors that can be culled out.  Grouped by class and sorted by CullDistance values of the class CDO.
* Unified alot of log lines so they use PrintObjectFromIndex utility functions

#rnx

diff --git a/Engine/Source/Runtime/Engine/Private/Net/Iris/ReplicationSystem/ActorReplicationBridge.cpp b/Engine/Source/Runtime/Engine/Private/Net/Iris/ReplicationSystem/ActorReplicationBridge.cpp
--- a/Engine/Source/Runtime/Engine/Private/Net/Iris/ReplicationSystem/ActorReplicationBridge.cpp
+++ b/Engine/Source/Runtime/Engine/Private/Net/Iris/ReplicationSystem/ActorReplicationBridge.cpp
@@ -68,7 +68,7 @@ void ActorReplicationBridgePreUpdateFunction(FNetRefHandle Handle, UObject* Inst
 	}
 }
 
-void ActorReplicationBridgeGetActorWorldObjectInfo(FNetRefHandle Handle, const UObject* Instance, FVector& OutWorldLocation, float& OutCullDistance)
+void ActorReplicationBridgeGetActorWorldObjectInfo(FNetRefHandle /*Handle*/, const UObject* Instance, FVector& OutWorldLocation, float& OutCullDistance)
 {
 	if (const AActor* Actor = Cast<AActor>(Instance))
 	{
@@ -251,6 +251,8 @@ UE::Net::FNetRefHandle UActorReplicationBridge::BeginReplication(AActor* Actor,
 		return FNetRefHandle::GetInvalid();
 	}
 
+    UE_CLOG(Actor->bAlwaysRelevant, LogIrisBridge, Verbose, TEXT("BeginReplication of AlwaysRelevant actor %s"), *PrintObjectFromNetRefHandle(ActorRefHandle));
+
 	// Set owning connection filtering if actor is only relevant to owner
 	{
 		if (Actor->bOnlyRelevantToOwner && !Actor->bAlwaysRelevant)
@@ -1221,13 +1223,13 @@ void UActorReplicationBridge::ReportErrorWithNetRefHandle(uint32 ErrorType, FNet
 	}
 }
 
-FString UActorReplicationBridge::PrintConnectionInfo(uint32 ConnectionId)
+FString UActorReplicationBridge::PrintConnectionInfo(uint32 ConnectionId) const
 {
 	if (NetDriver)
 	{
 		if (UNetConnection* ClientConnection = NetDriver->GetConnectionById(ConnectionId))
 		{
-			return FString::Printf(TEXT("ConnectionId:%u %s"), ConnectionId, *ClientConnection->Describe());
+			return FString::Printf(TEXT("ConnectionId:%u ViewTarget: %s Named: %s"), ConnectionId, *GetNameSafe(ClientConnection->ViewTarget), *ClientConnection->Describe());
 		}
 		else
 		{
diff --git a/Engine/Source/Runtime/Engine/Public/Net/Iris/ReplicationSystem/ActorReplicationBridge.h b/Engine/Source/Runtime/Engine/Public/Net/Iris/ReplicationSystem/ActorReplicationBridge.h
--- a/Engine/Source/Runtime/Engine/Public/Net/Iris/ReplicationSystem/ActorReplicationBridge.h
+++ b/Engine/Source/Runtime/Engine/Public/Net/Iris/ReplicationSystem/ActorReplicationBridge.h
@@ -106,7 +106,7 @@ protected:
 
 	virtual float GetPollFrequencyOfRootObject(const UObject* ReplicatedObject) const override;
 
-	[[nodiscard]] virtual FString PrintConnectionInfo(uint32 ConnectionId) override;
+	[[nodiscard]] virtual FString PrintConnectionInfo(uint32 ConnectionId) const override;
 
 private:
 	void GetActorCreationHeader(const AActor* Actor, UE::Net::Private::FActorCreationHeader& Header) const;
diff --git a/Engine/Source/Runtime/Experimental/Iris/Core/Private/Iris/ReplicationSystem/Filtering/NetObjectGridFilter.cpp b/Engine/Source/Runtime/Experimental/Iris/Core/Private/Iris/ReplicationSystem/Filtering/NetObjectGridFilter.cpp
--- a/Engine/Source/Runtime/Experimental/Iris/Core/Private/Iris/ReplicationSystem/Filtering/NetObjectGridFilter.cpp
+++ b/Engine/Source/Runtime/Experimental/Iris/Core/Private/Iris/ReplicationSystem/Filtering/NetObjectGridFilter.cpp
@@ -5,6 +5,8 @@
 #include "Iris/ReplicationSystem/RepTag.h"
 #include "Iris/ReplicationSystem/ReplicationProtocol.h"
 #include "Iris/ReplicationSystem/ReplicationSystem.h"
+#include "Iris/ReplicationSystem/ReplicationSystemInternal.h"
+#include "Iris/ReplicationSystem/NetRefHandleManager.h"
 #include "Iris/ReplicationSystem/WorldLocations.h"
 #include "Iris/Core/IrisCsv.h"
 #include "Iris/Core/IrisLog.h"
@@ -21,6 +23,7 @@ void UNetObjectGridFilter::OnInit(FNetObjectFilterInitParams& Params)
 
 	PerConnectionInfos.SetNum(Params.MaxConnectionCount + 1);
 
+	NetRefHandleManager = &Params.ReplicationSystem->GetReplicationSystemInternal()->GetNetRefHandleManager();
 	NetCullDistanceOverrides = &Params.ReplicationSystem->GetNetCullDistanceOverrides();
 }
 
@@ -57,7 +60,7 @@ bool UNetObjectGridFilter::AddObject(uint32 ObjectIndex, FNetObjectFilterAddObje
 	if (Config->MaxCullDistance > 0.0f && PerObjectInfo.GetCullDistance() > Config->MaxCullDistance)
 	{
 		// Too big an object. We expect it to be costly to move it across cells.
-		UE_LOG(LogIris, Warning, TEXT("ReplicatedObject %u cull distance %f is above the max %f. Object will become always relevant instead"), ObjectIndex, PerObjectInfo.GetCullDistance(), Config->MaxCullDistance);
+		UE_LOG(LogIris, Warning, TEXT("ReplicatedObject %s cull distance %f is above the max %f. Object will become always relevant instead"), *NetRefHandleManager->PrintObjectFromIndex(ObjectIndex), PerObjectInfo.GetCullDistance(), Config->MaxCullDistance);
 		RemoveObject(ObjectIndex, ObjectLocationInfo);
 		return false;
 	}
@@ -468,6 +471,39 @@ bool UNetObjectGridFilter::DoesCellContainCoord(const UNetObjectGridFilter::FCel
 	return (Coord.X >= Cell.MinX) & (Coord.X <= Cell.MaxX) & (Coord.Y >= Cell.MinY) & (Coord.Y <= Cell.MaxY);
 }
 
+FString UNetObjectGridFilter::PrintDebugInfoForObject(const FDebugInfoParams& Params, uint32 ObjectIndex) const
+{
+	using namespace UE::Net;
+
+	const FObjectLocationInfo& ObjectLocationInfo = static_cast<const FObjectLocationInfo&>(Params.FilteringInfos[ObjectIndex]);
+
+	if (ObjectLocationInfo.GetInfoIndex() == 0)
+	{
+		return TEXT("NoGridFilter Info");
+	}
+
+	const FPerObjectInfo& PerObjectInfo = ObjectInfos[ObjectLocationInfo.GetInfoIndex()];
+
+	const FPerConnectionInfo* const PerConnectionInfo = PerConnectionInfos.IsValidIndex(Params.ConnectionId) ? &PerConnectionInfos[Params.ConnectionId] : nullptr;
+	const uint32* CulledFrameCountPtr = PerConnectionInfo ? PerConnectionInfo->RecentObjectFrameCount.Find(ObjectIndex) : nullptr;
+	const uint32 CulledFrameCount = CulledFrameCountPtr ? *CulledFrameCountPtr : 0;
+
+	double Dist2d = DOUBLE_BIG_NUMBER;
+	double DistZ = DOUBLE_BIG_NUMBER;
+	
+	for (const FReplicationView::FView& View : Params.View.Views)
+	{
+		Dist2d = FMath::Min(FVector::Dist2D(PerObjectInfo.Position, View.Pos), Dist2d);
+		DistZ = FMath::Min(FMath::Abs(PerObjectInfo.Position.Z - View.Pos.Z), DistZ);
+	}
+	
+
+	return FString::Printf(TEXT("[GridFilter] Dist2D: %.2f, CullDistance: %.2f, CulledFrameCount: %u, DistZ: %.2f, Pos: %s"),
+		Dist2d, PerObjectInfo.GetCullDistance(), CulledFrameCount, DistZ, *PerObjectInfo.Position.ToCompactString()
+	);
+}
+
+
 //*************************************************************************************************
 // UNetObjectGridWorldLocFilter
 //*************************************************************************************************
diff --git a/Engine/Source/Runtime/Experimental/Iris/Core/Private/Iris/ReplicationSystem/Filtering/ReplicationFiltering.cpp b/Engine/Source/Runtime/Experimental/Iris/Core/Private/Iris/ReplicationSystem/Filtering/ReplicationFiltering.cpp
--- a/Engine/Source/Runtime/Experimental/Iris/Core/Private/Iris/ReplicationSystem/Filtering/ReplicationFiltering.cpp
+++ b/Engine/Source/Runtime/Experimental/Iris/Core/Private/Iris/ReplicationSystem/Filtering/ReplicationFiltering.cpp
@@ -35,7 +35,8 @@ FName GetStaticFilterName(FNetObjectFilterHandle Filter)
 	{
 		case InvalidNetObjectFilterHandle:
 		{
-			return NAME_None;
+			static const FName NoFilterName = TEXT("NoFilter");
+			return NoFilterName;
 		} break;
 
 		case ToOwnerFilterHandle:
@@ -439,6 +440,8 @@ bool FReplicationFiltering::SetFilter(FInternalNetRefIndex ObjectIndex, FNetObje
 		return false;
 	}
 
+	UE_LOG(LogIrisFiltering, Verbose, TEXT("Setting filter %s to %s"), *GetFilterName(Filter).ToString(), *NetRefHandleManager->PrintObjectFromIndex(ObjectIndex));
+
 	const bool bWantsToUseDynamicFilter = FNetObjectFilterHandleUtil::IsDynamicFilter(Filter);
 	const uint8 OldDynamicFilterIndex = ObjectIndexToDynamicFilterIndex[ObjectIndex];
 	const uint32 NewDynamicFilterIndex = bWantsToUseDynamicFilter ? FNetObjectFilterHandleUtil::GetDynamicFilterIndex(Filter) : InvalidDynamicFilterIndex;
@@ -461,7 +464,7 @@ bool FReplicationFiltering::SetFilter(FInternalNetRefIndex ObjectIndex, FNetObje
 	// Let subobjects be filtered like their owners.
 	if (bWantsToUseDynamicFilter && (ObjectData.SubObjectRootIndex != FNetRefHandleManager::InvalidInternalIndex))
 	{
-		UE_LOG(LogIrisFiltering, Warning, TEXT("Ignoring request to use dynamic filter on object %s due to it being a subobject. Any filtering already successfully applied remains as is."), ToCStr(ObjectData.RefHandle.ToString()));
+		UE_LOG(LogIrisFiltering, Warning, TEXT("Cannot set dynamic filters on subobjects. Filter change for %s is ignored"), *NetRefHandleManager->PrintObjectFromIndex(ObjectIndex));
 		return false;
 	}
 
@@ -513,7 +516,7 @@ bool FReplicationFiltering::SetFilter(FInternalNetRefIndex ObjectIndex, FNetObje
 		}
 		else
 		{
-			UE_LOG(LogIrisFiltering, Verbose, TEXT("Filter '%s' does not support object %u."), ToCStr(DynamicFilterInfos[NewDynamicFilterIndex].Filter->GetFName().GetPlainNameString()), ObjectIndex);
+			UE_LOG(LogIrisFiltering, Verbose, TEXT("Filter '%s' does not support object %s."), ToCStr(DynamicFilterInfos[NewDynamicFilterIndex].Filter->GetFName().GetPlainNameString()), *NetRefHandleManager->PrintObjectFromIndex(ObjectIndex));
 			return false;
 		}
 	}
@@ -1538,6 +1541,7 @@ void FReplicationFiltering::BatchNotifyFiltersOfDirtyObjects(FUpdateDirtyObjects
 	FNetObjectFilterUpdateParams UpdateParameters;
 	UpdateParameters.FilteringInfos = NetObjectFilteringInfos.GetData();
 
+	// $IRIS TODO: We should probably have a trait asking if the Filter needs to receive UpdateObjects.
 	for (const FUpdateDirtyObjectsBatchHelper::FPerFilterInfo& PerFilterInfo : BatchHelper.PerFilterInfos)
 	{
 		if (PerFilterInfo.ObjectCount == 0)
@@ -2307,7 +2311,7 @@ void FReplicationFiltering::NotifyObjectAddedToGroup(FNetObjectGroupHandle Group
 	const FNetObjectGroupHandle::FGroupIndexType GroupIndex = GroupHandle.GetGroupIndex();
 	if (SubObjectFilterGroups.GetBit(GroupIndex))
 	{
-		UE_LOG(LogIrisFiltering, Verbose, TEXT("ReplicationFiltering::NotifyObjectAddedToGroup Added %s to SubObjectFilter with GroupIndex: %u"), *(NetRefHandleManager->GetReplicatedObjectDataNoCheck(ObjectIndex).RefHandle.ToString()), GroupIndex);
+		UE_LOG(LogIrisFiltering, Verbose, TEXT("ReplicationFiltering::NotifyObjectAddedToGroup Added %s to SubObjectFilter with GroupIndex: %u"), *(NetRefHandleManager->PrintObjectFromIndex(ObjectIndex)), GroupIndex);
 	}
 	else if (ExclusionFilterGroups.GetBit(GroupIndex))
 	{
@@ -2320,7 +2324,7 @@ void FReplicationFiltering::NotifyObjectAddedToGroup(FNetObjectGroupHandle Group
 			bHasDirtyExclusionFilterGroup = 1;
 		}
 
-		UE_LOG(LogIrisFiltering, Verbose, TEXT("ReplicationFiltering::NotifyObjectAddedToGroup Added %s to ExclusionGroup filter with GroupIndex: %u"), *(NetRefHandleManager->GetReplicatedObjectDataNoCheck(ObjectIndex).RefHandle.ToString()), GroupIndex);
+		UE_LOG(LogIrisFiltering, Verbose, TEXT("ReplicationFiltering::NotifyObjectAddedToGroup Added %s to ExclusionGroup filter with GroupIndex: %u"), *(NetRefHandleManager->PrintObjectFromIndex(ObjectIndex)), GroupIndex);
 	}
 	else if (InclusionFilterGroups.GetBit(GroupIndex))
 	{
@@ -2333,7 +2337,7 @@ void FReplicationFiltering::NotifyObjectAddedToGroup(FNetObjectGroupHandle Group
 			bHasDirtyInclusionFilterGroup = 1;
 		}
 
-		UE_LOG(LogIrisFiltering, Verbose, TEXT("ReplicationFiltering::NotifyObjectAddedToGroup Added %s to InclusionGroup filter with GroupIndex: %u"), *(NetRefHandleManager->GetReplicatedObjectDataNoCheck(ObjectIndex).RefHandle.ToString()), GroupIndex);
+		UE_LOG(LogIrisFiltering, Verbose, TEXT("ReplicationFiltering::NotifyObjectAddedToGroup Added %s to InclusionGroup filter with GroupIndex: %u"), *(NetRefHandleManager->PrintObjectFromIndex(ObjectIndex)), GroupIndex);
 	}
 }
 
@@ -2341,7 +2345,7 @@ void FReplicationFiltering::NotifyObjectRemovedFromGroup(FNetObjectGroupHandle G
 {
 	const FNetObjectGroupHandle::FGroupIndexType GroupIndex = GroupHandle.GetGroupIndex();
 
-	UE_LOG(LogIrisFiltering, Verbose, TEXT("ReplicationFiltering::NotifyObjectRemovedFromGroup Removing %s from GroupIndex: %u"), ToCStr(NetRefHandleManager->GetReplicatedObjectDataNoCheck(ObjectIndex).RefHandle.ToString()), GroupIndex);
+	UE_LOG(LogIrisFiltering, Verbose, TEXT("ReplicationFiltering::NotifyObjectRemovedFromGroup Removing %s from GroupIndex: %u"), ToCStr(NetRefHandleManager->PrintObjectFromIndex(ObjectIndex)), GroupIndex);
 
 	if (SubObjectFilterGroups.GetBit(GroupIndex))
 	{
@@ -2482,6 +2486,9 @@ void FReplicationFiltering::InitFilters()
 
 void FReplicationFiltering::RemoveFromDynamicFilter(uint32 ObjectIndex, uint32 FilterIndex)
 {
+	UE_LOG(LogIrisFiltering, Verbose, TEXT("RemoveFromDynamicFilter removing %s from Dynamic Filter %s"), 
+		*NetRefHandleManager->PrintObjectFromIndex(ObjectIndex), *GetFilterName(FNetObjectFilterHandleUtil::MakeDynamicFilterHandle(ObjectIndexToDynamicFilterIndex[ObjectIndex])).ToString());
+
 	ObjectIndexToDynamicFilterIndex[ObjectIndex] = InvalidDynamicFilterIndex;
 	FNetObjectFilteringInfo& NetObjectFilteringInfo = NetObjectFilteringInfos[ObjectIndex];
 	FFilterInfo& FilterInfo = DynamicFilterInfos[FilterIndex];
@@ -2535,6 +2542,45 @@ TArrayView<FNetObjectFilteringInfo> FReplicationFiltering::GetNetObjectFiltering
 	return MakeArrayView(NetObjectFilteringInfos);
 }
 
+FString FReplicationFiltering::PrintFilterObjectInfo(FInternalNetRefIndex ObjectIndex, uint32 ConnectionId) const
+{
+	// $IRIS TODO: Print info if the object is part of a non-dynamic filter.
+
+	const uint8 DynamicFilterIndex = ObjectIndexToDynamicFilterIndex[ObjectIndex];
+	if (DynamicFilterIndex == InvalidDynamicFilterIndex)
+	{
+		return TEXT("[NoDynamicFilter]");
+	}
+
+	const FFilterInfo& FilterInfo = DynamicFilterInfos[DynamicFilterIndex];
+
+	if (!FilterInfo.FilteredObjects.IsBitSet(ObjectIndex))
+	{
+		ensureMsgf(false, TEXT("Problem with Filter configs for %s.  DynamicIndex %u Filter %s but not in FilteredObjects list"), *NetRefHandleManager->PrintObjectFromIndex(ObjectIndex), DynamicFilterIndex, *FilterInfo.Name.ToString());
+		return TEXT("[WrongSettings]");
+	}
+
+	UNetObjectFilter::FDebugInfoParams DebugParams;
+	DebugParams.FilterName = FilterInfo.Name;
+	DebugParams.FilteringInfos = NetObjectFilteringInfos.GetData();
+	DebugParams.ConnectionId = ConnectionId;
+	DebugParams.View = Connections->GetReplicationView(ConnectionId);
+
+	return FilterInfo.Filter->PrintDebugInfoForObject(DebugParams, ObjectIndex);
+}
+
+void FReplicationFiltering::BuildObjectsInFilterList(FNetBitArrayView OutObjectsInFilter, FName FilterName) const
+{
+	for (const FFilterInfo& FilterInfo : DynamicFilterInfos)
+	{
+		if (FilterInfo.Name == FilterName)
+		{
+			OutObjectsInFilter.Copy(FilterInfo.FilteredObjects);
+			return;
+		}
+	}
+}
+
 //*************************************************************************************************
 // FNetObjectFilteringInfoAccessor
 //*************************************************************************************************
diff --git a/Engine/Source/Runtime/Experimental/Iris/Core/Private/Iris/ReplicationSystem/Filtering/ReplicationFiltering.h b/Engine/Source/Runtime/Experimental/Iris/Core/Private/Iris/ReplicationSystem/Filtering/ReplicationFiltering.h
--- a/Engine/Source/Runtime/Experimental/Iris/Core/Private/Iris/ReplicationSystem/Filtering/ReplicationFiltering.h
+++ b/Engine/Source/Runtime/Experimental/Iris/Core/Private/Iris/ReplicationSystem/Filtering/ReplicationFiltering.h
@@ -100,6 +100,9 @@ public:
 	/** Fill the passed in list with root objects considered always relevant */
 	void BuildAlwaysRelevantList(FNetBitArrayView OutAlwaysRelevantList, const FNetBitArrayView ScopeList) const;
 
+	/** Fill the passed in list with objects affected by a given filter*/
+	void BuildObjectsInFilterList(FNetBitArrayView OutObjectsInFilter, FName FilterName) const;
+
 	// Connection handling
 	void AddConnection(uint32 ConnectionId);
 	void RemoveConnection(uint32 ConnectionId);
@@ -133,6 +136,9 @@ public:
 	void SetSubObjectFilterStatus(FNetObjectGroupHandle GroupHandle, uint32 ConnectionId, ENetFilterStatus ReplicationStatus);
 	bool GetSubObjectFilterStatus(FNetObjectGroupHandle GroupHandle, uint32 ConnectionId, ENetFilterStatus& OutReplicationStatus) const;
 
+	/** Print the filter information we have regarding the passed ObjectIndex and his relation to the passed Connection */
+	FString PrintFilterObjectInfo(FInternalNetRefIndex ObjectIndex, uint32 ConnectionId) const;
+
 private:
 	struct FPerConnectionInfo
 	{
diff --git a/Engine/Source/Runtime/Experimental/Iris/Core/Private/Iris/ReplicationSystem/ObjectReplicationBridge.cpp b/Engine/Source/Runtime/Experimental/Iris/Core/Private/Iris/ReplicationSystem/ObjectReplicationBridge.cpp
--- a/Engine/Source/Runtime/Experimental/Iris/Core/Private/Iris/ReplicationSystem/ObjectReplicationBridge.cpp
+++ b/Engine/Source/Runtime/Experimental/Iris/Core/Private/Iris/ReplicationSystem/ObjectReplicationBridge.cpp
@@ -430,7 +430,7 @@ void UObjectReplicationBridge::AssignDynamicFilter(UObject* Instance, const FCre
 
 	if (FilterHandle != InvalidNetObjectFilterHandle)
 	{
-		UE_LOG_OBJECTREPLICATIONBRIDGE(Verbose, TEXT("BeginReplication Filter: %s will be used for Object: %s "), *(ReplicationSystem->GetFilterName(FilterHandle).ToString()), *Instance->GetName());
+		UE_LOG_OBJECTREPLICATIONBRIDGE(Verbose, TEXT("BeginReplication Filter: %s will be used for %s "), *(ReplicationSystem->GetFilterName(FilterHandle).ToString()), *NetRefHandleManager->PrintObjectFromNetRefHandle(RefHandle));
 		ReplicationSystem->SetFilter(RefHandle, FilterHandle);
 	}
 	
@@ -473,10 +473,10 @@ UE::Net::FNetRefHandle UObjectReplicationBridge::BeginReplication(FNetRefHandle
 	if (SubObjectRefHandle.IsValid())
 	{
 		// Add subobject
-		UE_LOG_OBJECTREPLICATIONBRIDGE(Verbose, TEXT("BeginReplication Added SubObject %s to Owner %s RelativeToSubObjectHandle %s"), *SubObjectRefHandle.ToString(), *OwnerRefHandle.ToString(), *InsertRelativeToSubObjectRefHandle.ToString());
-
 		InternalAddSubObject(OwnerRefHandle, SubObjectRefHandle, InsertRelativeToSubObjectRefHandle, InsertionOrder);
 
+		UE_LOG_OBJECTREPLICATIONBRIDGE(Verbose, TEXT("BeginReplication Added %s RelativeToSubObjectHandle %s"), *PrintObjectFromNetRefHandle(SubObjectRefHandle), *PrintObjectFromNetRefHandle(InsertRelativeToSubObjectRefHandle));
+
 		// SubObjects should always poll with owner
 		SetPollWithObject(OwnerRefHandle, SubObjectRefHandle);
 
@@ -502,12 +502,12 @@ void UObjectReplicationBridge::SetSubObjectNetCondition(FNetRefHandle SubObjectR
 	const FInternalNetRefIndex SubObjectInternalIndex = LocalNetRefHandleManager.GetInternalIndex(SubObjectRefHandle);
 	if (LocalNetRefHandleManager.SetSubObjectNetCondition(SubObjectInternalIndex, (int8)Condition))
 	{
-		UE_LOG_OBJECTREPLICATIONBRIDGE(Verbose, TEXT("SetSubObjectNetCondition for SubObject %s Condition %s"), *SubObjectRefHandle.ToString(), *UEnum::GetValueAsString<ELifetimeCondition>(Condition));
+		UE_LOG_OBJECTREPLICATIONBRIDGE(Verbose, TEXT("SetSubObjectNetCondition for SubObject %s Condition %s"), *PrintObjectFromNetRefHandle(SubObjectRefHandle), *UEnum::GetValueAsString<ELifetimeCondition>(Condition));
 		MarkNetObjectStateDirty(ReplicationSystem->GetId(), SubObjectInternalIndex);
 	}
 	else
 	{
-		UE_LOG_OBJECTREPLICATIONBRIDGE(Verbose, TEXT("Failed to SetSubObjectNetCondition for SubObject %s Condition %s"), *SubObjectRefHandle.ToString(), *UEnum::GetValueAsString<ELifetimeCondition>(Condition));
+		UE_LOG_OBJECTREPLICATIONBRIDGE(Warning, TEXT("Failed to Set SubObjectNetCondition for SubObject %s Condition %s"), *PrintObjectFromNetRefHandle(SubObjectRefHandle), *UEnum::GetValueAsString<ELifetimeCondition>(Condition));
 	}
 }
 
@@ -535,11 +535,11 @@ void UObjectReplicationBridge::AddDependentObject(FNetRefHandle ParentHandle, FN
 		const FInternalNetRefIndex DependentInternalIndex = LocalNetRefHandleManager.GetInternalIndex(DependentHandle);
 		Filtering.NotifyAddedDependentObject(DependentInternalIndex);
 
-		UE_LOG_OBJECTREPLICATIONBRIDGE(Verbose, TEXT("AddDependentObject Added dependent object %s to parent %s"), *DependentHandle.ToString(), *ParentHandle.ToString());
+		UE_LOG_OBJECTREPLICATIONBRIDGE(Verbose, TEXT("AddDependentObject Added dependent object %s to parent %s"), *PrintObjectFromNetRefHandle(DependentHandle), *PrintObjectFromNetRefHandle(ParentHandle));
 	}
 	else
 	{
-		UE_LOG_OBJECTREPLICATIONBRIDGE(Verbose, TEXT("AddDependentObject Failed to add dependent object %s to parent %s"), *DependentHandle.ToString(), *ParentHandle.ToString());
+		UE_LOG_OBJECTREPLICATIONBRIDGE(Warning , TEXT("AddDependentObject Failed to add dependent object %s to parent %s"), *PrintObjectFromNetRefHandle(DependentHandle), *PrintObjectFromNetRefHandle(ParentHandle));
 	}
 }
 
@@ -581,7 +581,7 @@ void UObjectReplicationBridge::EndReplication(UObject* Instance, EEndReplication
 
 void UObjectReplicationBridge::DetachInstanceFromRemote(FNetRefHandle Handle, EReplicationBridgeDestroyInstanceReason DestroyReason, EReplicationBridgeDestroyInstanceFlags DestroyFlags)
 {
-	UE_LOG_OBJECTREPLICATIONBRIDGE(Verbose, TEXT("DetachInstanceFromRemote %s DestroyReason: %s DestroyFlags: %u"), *Handle.ToString(), ToCStr(LexToString(DestroyReason)), unsigned(DestroyFlags));
+	UE_LOG_OBJECTREPLICATIONBRIDGE(Verbose, TEXT("DetachInstanceFromRemote %s DestroyReason: %s DestroyFlags: %u"), *PrintObjectFromNetRefHandle(Handle), ToCStr(LexToString(DestroyReason)), unsigned(DestroyFlags));
 
 	using namespace UE::Net::Private;
 
@@ -635,7 +635,7 @@ void UObjectReplicationBridge::RegisterRemoteInstance(FNetRefHandle RefHandle, U
 		GetObjectReferenceCache()->AddRemoteReference(RefHandle, Instance);
 	}
 
-	UE_LOG_OBJECTREPLICATIONBRIDGE(Verbose, TEXT("RegisterRemoteInstance %s %s with ProtocolId:0x%" UINT64_x_FMT), *RefHandle.ToString(), ToCStr(Instance->GetName()), Protocol->ProtocolIdentifier);
+	UE_LOG_OBJECTREPLICATIONBRIDGE(Verbose, TEXT("RegisterRemoteInstance %s %s with ProtocolId:0x%" UINT64_x_FMT), *PrintObjectFromNetRefHandle(RefHandle), ToCStr(Instance->GetName()), Protocol->ProtocolIdentifier);
 }
 
 FReplicationBridgeCreateNetRefHandleResult UObjectReplicationBridge::CreateNetRefHandleFromRemote(FNetRefHandle RootObjectOfSubObject, FNetRefHandle WantedNetHandle, FReplicationBridgeSerializationContext& Context)
@@ -1458,7 +1458,7 @@ bool UObjectReplicationBridge::IsClassCritical(const UClass* Class)
 	return false;
 }
 
-FString UObjectReplicationBridge::PrintConnectionInfo(uint32 ConnectionId)
+FString UObjectReplicationBridge::PrintConnectionInfo(uint32 ConnectionId) const
 {
 	return FString::Printf(TEXT("ConnectionId:%u"), ConnectionId);
 }
@@ -1561,7 +1561,7 @@ void UObjectReplicationBridge::SetClassDynamicFilterConfig(FName ClassPathName,
 
 	if (FilterName != NAME_None)
 	{
-		FNetObjectFilterHandle FilterHandle = GetReplicationSystem()->GetFilterHandle(FilterName);;
+		FNetObjectFilterHandle FilterHandle = GetReplicationSystem()->GetFilterHandle(FilterName);
 
 		if (ensureMsgf(FilterHandle != InvalidNetObjectFilterHandle, TEXT("SetClassDynamicFilterConfig for %s received invalid filter named %s"), *ClassPathName.ToString(), *FilterName.ToString()))
 		{
diff --git a/Engine/Source/Runtime/Experimental/Iris/Core/Private/Iris/ReplicationSystem/ObjectReplicationBridgeDebugging.cpp b/Engine/Source/Runtime/Experimental/Iris/Core/Private/Iris/ReplicationSystem/ObjectReplicationBridgeDebugging.cpp
--- a/Engine/Source/Runtime/Experimental/Iris/Core/Private/Iris/ReplicationSystem/ObjectReplicationBridgeDebugging.cpp
+++ b/Engine/Source/Runtime/Experimental/Iris/Core/Private/Iris/ReplicationSystem/ObjectReplicationBridgeDebugging.cpp
@@ -17,9 +17,117 @@
 #include "Net/Core/NetBitArrayPrinter.h"
 #include "Net/Core/Trace/NetDebugName.h"
 
+/**
+ * This class contains misc console commands that log the state of different Iris systems.
+ * 
+ * Most cmds support common optional parameters that are listed here:
+ *		RepSystemId=X => Execute the cmd on a specific ReplicationSystem. Useful in PIE
+ *		WithSubObjects => Print the subobjects attached to each RootObject
+ *		SortByClass => Log the rootobjects alphabetically by ClassName (usually the default)
+ *		SortByNetRefHandle => Log the rootobjects by their NetRefHandle Id starting with static objects (odd Id) then dynamic objects (even Id)
+ */
+
 namespace UE::Net::Private::ObjectBridgeDebugging
 {
 
+enum class EPrintTraits : uint32
+{
+	Default					= 0x0000,
+	LogSubObjects			= 0x0001, // log the subobjects of each rootobject
+	LogTraits				= EPrintTraits::LogSubObjects,
+
+	SortByClass				= 0x0100, // log objects sorted by their class name
+	SortByNetRefHandle		= 0x0200, // log objects sorted by netrefhandle (odd (static) first, even (dynamic) second)
+	SortTraits				= EPrintTraits::SortByNetRefHandle | EPrintTraits::SortByClass
+};
+ENUM_CLASS_FLAGS(EPrintTraits);
+
+EPrintTraits FindPrintTraitsFromArgs(const TArray<FString>& Args)
+{
+	EPrintTraits Traits = EPrintTraits::Default;
+
+	if (Args.FindByPredicate([](const FString& Str) { return Str.Contains(TEXT("WithSubObjects")); }))
+	{
+		Traits = Traits | EPrintTraits::LogSubObjects;
+	}
+
+	if (Args.FindByPredicate([](const FString& Str) { return Str.Contains(TEXT("SortByClass")); }) )
+	{
+		Traits = Traits | EPrintTraits::SortByClass;
+	}
+	else if (Args.FindByPredicate([](const FString& Str) { return Str.Contains(TEXT("SortByNetRefHandle")); }))
+	{
+		Traits = Traits | EPrintTraits::SortByNetRefHandle;
+	}
+
+	return Traits;
+}
+
+/** Holds information about root objects sortable by class name */
+struct FRootObjectData
+{
+	FInternalNetRefIndex ObjectIndex = 0;
+	FNetRefHandle NetHandle;
+	UObject* Instance = nullptr;
+	UClass* Class = nullptr;
+};
+
+// Transform a bit array of root object indexes into an array of RootObjectData struct
+void FillRootObjectArrayFromBitArray(TArray<FRootObjectData>& OutRootObjects, const FNetBitArrayView RootObjectList, FNetRefHandleManager* NetRefHandleManager)
+{
+	RootObjectList.ForAllSetBits([&](uint32 RootObjectIndex)
+	{
+		FRootObjectData Data;
+		Data.ObjectIndex = RootObjectIndex;
+		Data.NetHandle = NetRefHandleManager->GetNetRefHandleFromInternalIndex(RootObjectIndex);
+		Data.Instance = NetRefHandleManager->GetReplicatedObjectInstance(RootObjectIndex);
+		Data.Class = Data.Instance ? Data.Instance->GetClass() : nullptr;
+
+		OutRootObjects.Emplace(MoveTemp(Data));
+	});
+}
+
+void SortByClassName(TArray<FRootObjectData>& OutArray)
+{
+	Algo::Sort(OutArray, [](const FRootObjectData& lhs, const FRootObjectData& rhs)
+	{
+		if (lhs.Class == rhs.Class) { return false; }
+		if (!lhs.Class) { return false; }
+		if (!rhs.Class) { return true; }
+		return lhs.Class->GetName() < rhs.Class->GetName();
+	});
+}
+
+void SortByNetRefHandle(TArray<FRootObjectData>& OutArray)
+{
+	// Sort static objects first (odds) then dynamic ones second (evens)
+	Algo::Sort(OutArray, [](const FRootObjectData& lhs, const FRootObjectData& rhs)
+	{
+		if (lhs.NetHandle == rhs.NetHandle) { return false; }
+		if (!lhs.NetHandle.IsValid()) { return false; }
+		if (!rhs.NetHandle.IsValid()) { return true; }
+		if (lhs.NetHandle.IsStatic() && rhs.NetHandle.IsDynamic()) { return false; }
+		if (lhs.NetHandle.IsDynamic() && rhs.NetHandle.IsStatic()) { return true; }
+		return lhs.NetHandle < rhs.NetHandle;
+	});
+}
+
+/** Sort the array with the selected trait. If no traits were selected, sort via the default one */
+void SortViaTrait(TArray<FRootObjectData>& OutArray, EPrintTraits ArgTraits, EPrintTraits DefaultTraits)
+{
+	EPrintTraits SelectedTrait = ArgTraits & EPrintTraits::SortTraits;
+	if (SelectedTrait == EPrintTraits::Default)
+	{
+		SelectedTrait = DefaultTraits;
+	}
+
+	switch(SelectedTrait)
+	{
+		case EPrintTraits::SortByClass: SortByClassName(OutArray);  break;
+		case EPrintTraits::SortByNetRefHandle: SortByNetRefHandle(OutArray); break;
+	}
+}
+
 void PrintDefaultNetObjectState(UReplicationSystem* ReplicationSystem, uint32 ConnectionId, const FReplicationFragments& RegisteredFragments, FStringBuilderBase& StringBuilder)
 {
 	FReplicationSystemInternal* ReplicationSystemInternal = ReplicationSystem->GetReplicationSystemInternal();
@@ -91,31 +199,57 @@ FString PrintNetObject(FNetRefHandleManager* NetRefHandleManager, FInternalNetRe
 	);
 }
 
-void LogRootObjectListWithSubObjects(FNetRefHandleManager* NetRefHandleManager, const FNetBitArrayView RootObjectList, TTuple<uint32, uint32>& OutStats, TFunction<FString(FInternalNetRefIndex ObjectIndex)> OptionalObjectPrint)
+struct FLogContext
+{
+	// Mandatory parameters
+	FNetRefHandleManager* NetRefHandleManager = nullptr;
+	const TArray<FRootObjectData>& RootObjectArray;
+
+	// Optional parameters
+	TFunction<FString(FInternalNetRefIndex ObjectIndex)> OptionalObjectPrint;
+
+	// Stats
+	uint32 NumRootObjects = 0;
+	uint32 NumSubObjects = 0;
+};
+
+void LogRootObjectList(FLogContext& LogContext, bool bLogSubObjects)
 {
 	using namespace UE::Net;
 	using namespace UE::Net::Private;
 
-	auto PrintRootObjectInfo = [&](FInternalNetRefIndex RootObjectIndex)
+	FNetRefHandleManager* NetRefHandleManager = LogContext.NetRefHandleManager;
+
+	for (const FRootObjectData& RootObject : LogContext.RootObjectArray)
 	{
-		UE_LOG(LogIrisBridge, Display, TEXT("%s %s"), *PrintNetObject(NetRefHandleManager, RootObjectIndex), OptionalObjectPrint ? *OptionalObjectPrint(RootObjectIndex) : TEXT(""));
+		UE_LOG(LogIrisBridge, Display, TEXT("%s %s"), *PrintNetObject(NetRefHandleManager, RootObject.ObjectIndex), LogContext.OptionalObjectPrint ? *LogContext.OptionalObjectPrint(RootObject.ObjectIndex) : TEXT(""));
 
-		// TotalRootObjects
-		OutStats.Get<0>()++;
+		LogContext.NumRootObjects++;
 
-		TArrayView<const FInternalNetRefIndex> SubObjects = NetRefHandleManager->GetSubObjects(RootObjectIndex);
-		for (FInternalNetRefIndex SubObjectIndex : SubObjects)
+		if (bLogSubObjects)
 		{
-			UE_LOG(LogIrisBridge, Display, TEXT("\t%s %s"), *PrintNetObject(NetRefHandleManager, SubObjectIndex), OptionalObjectPrint ? *OptionalObjectPrint(RootObjectIndex) : TEXT(""));
+			TArrayView<const FInternalNetRefIndex> SubObjects = NetRefHandleManager->GetSubObjects(RootObject.ObjectIndex);
+			for (FInternalNetRefIndex SubObjectIndex : SubObjects)
+			{
+				UE_LOG(LogIrisBridge, Display, TEXT("\t%s %s"), *PrintNetObject(NetRefHandleManager, SubObjectIndex), LogContext.OptionalObjectPrint ? *LogContext.OptionalObjectPrint(SubObjectIndex) : TEXT(""));
 
-			// TotalSubObjects
-			OutStats.Get<1>()++;
+				LogContext.NumSubObjects++;
+			}
 		}
 	};
-
-	RootObjectList.ForAllSetBits(PrintRootObjectInfo);
 }
 
+void LogViaTrait(FLogContext& LogContext, EPrintTraits ArgTraits, EPrintTraits DefaultTraits)
+{
+	EPrintTraits SelectedTrait = ArgTraits & EPrintTraits::LogTraits;
+	if (SelectedTrait == EPrintTraits::Default)
+	{
+		SelectedTrait = DefaultTraits & EPrintTraits::LogTraits;
+	}
+
+	const bool bLogSubObjects = (SelectedTrait & EPrintTraits::LogSubObjects) != EPrintTraits::Default;
+	LogRootObjectList(LogContext, bLogSubObjects);	
+}
 
 } // end namespace UE::Net::Private::ObjectBridgeDebugging
 
@@ -181,12 +315,13 @@ FAutoConsoleCommand ObjectBridgePrintReplicatedObjects(
 	{
 		if (UObjectReplicationBridge* ObjectBridge = CastChecked<UObjectReplicationBridge>(RepSystem->GetReplicationBridge()))
 		{
-			ObjectBridge->PrintReplicatedObjects();
+			EPrintTraits ArgTraits = FindPrintTraitsFromArgs(Args);
+			ObjectBridge->PrintReplicatedObjects((uint32)ArgTraits);
 		}
 	}
 }));
 
-void UObjectReplicationBridge::PrintReplicatedObjects() const
+void UObjectReplicationBridge::PrintReplicatedObjects(uint32 ArgTraits) const
 {
 	using namespace UE::Net;
 	using namespace UE::Net::Private;
@@ -199,8 +334,15 @@ void UObjectReplicationBridge::PrintReplicatedObjects() const
 	uint32 TotalSubObjects = 0;
 
 	FNetBitArray RootObjects;
+	FNetBitArrayView RootObjectsView = MakeNetBitArrayView(RootObjects);
 	RootObjects.Init(NetRefHandleManager->GetMaxActiveObjectCount());
-	MakeNetBitArrayView(RootObjects).Set(NetRefHandleManager->GetGlobalScopableInternalIndices(), FNetBitArrayView::AndNotOp, NetRefHandleManager->GetSubObjectInternalIndicesView());
+	RootObjectsView.Set(NetRefHandleManager->GetGlobalScopableInternalIndices(), FNetBitArrayView::AndNotOp, NetRefHandleManager->GetSubObjectInternalIndicesView());
+
+	TArray<FRootObjectData> RootObjectArray;
+	{
+		FillRootObjectArrayFromBitArray(RootObjectArray, RootObjectsView, NetRefHandleManager);
+		SortViaTrait(RootObjectArray, (EPrintTraits)ArgTraits, EPrintTraits::Default);
+	}
 
 	auto PrintClassOrProtocol = [&](FInternalNetRefIndex ObjectIndex) -> FString
 	{
@@ -209,12 +351,12 @@ void UObjectReplicationBridge::PrintReplicatedObjects() const
 
 		return FString::Printf(TEXT("Class %s"), ObjectPtr ? *(ObjectPtr->GetClass()->GetName()) : NetObjectData.Protocol->DebugName->Name);
 	};
-	
-	TTuple<uint32, uint32> Stats(0,0);
-	LogRootObjectListWithSubObjects(NetRefHandleManager, MakeNetBitArrayView(RootObjects), Stats, PrintClassOrProtocol);
+
+	FLogContext LogContext = {.NetRefHandleManager = NetRefHandleManager, .RootObjectArray = RootObjectArray, .OptionalObjectPrint = PrintClassOrProtocol};
+	LogViaTrait(LogContext, (EPrintTraits)ArgTraits, EPrintTraits::Default);
 
 	UE_LOG(LogIrisBridge, Display, TEXT(""));
-	UE_LOG(LogIrisBridge, Display, TEXT("Printed %u root objects and %u sub objects"), Stats.Get<0>(), Stats.Get<1>());
+	UE_LOG(LogIrisBridge, Display, TEXT("Printed %u root objects and %u sub objects"), LogContext.NumRootObjects, LogContext.NumSubObjects);
 	UE_LOG(LogIrisBridge, Display, TEXT("################ Stop Printing ALL Replicated Objects ################"));
 }
 
@@ -230,12 +372,13 @@ FAutoConsoleCommand ObjectBridgePrintRelevantObjects(
 	{
 		if (UObjectReplicationBridge* ObjectBridge = CastChecked<UObjectReplicationBridge>(RepSystem->GetReplicationBridge()))
 		{
-			ObjectBridge->PrintRelevantObjects();
+			EPrintTraits ArgTraits = FindPrintTraitsFromArgs(Args);
+			ObjectBridge->PrintRelevantObjects((uint32)ArgTraits);
 		}
 	}
 }));
 
-void UObjectReplicationBridge::PrintRelevantObjects() const
+void UObjectReplicationBridge::PrintRelevantObjects(uint32 ArgTraits) const
 {
 	using namespace UE::Net;
 	using namespace UE::Net::Private;
@@ -244,19 +387,22 @@ void UObjectReplicationBridge::PrintRelevantObjects() const
 	UE_LOG(LogIrisBridge, Display, TEXT("################ Start Printing Relevant Objects ################"));
 	UE_LOG(LogIrisBridge, Display, TEXT(""));
 
-	uint32 TotalRootObjects = 0;
-	uint32 TotalSubObjects = 0;
-
 	FNetBitArray RootObjects;
+	FNetBitArrayView RootObjectsView = MakeNetBitArrayView(RootObjects);
 	RootObjects.Init(NetRefHandleManager->GetMaxActiveObjectCount());
-	MakeNetBitArrayView(RootObjects).Set(NetRefHandleManager->GetRelevantObjectsInternalIndices(), FNetBitArrayView::AndNotOp, NetRefHandleManager->GetSubObjectInternalIndicesView());
+	RootObjectsView.Set(NetRefHandleManager->GetRelevantObjectsInternalIndices(), FNetBitArrayView::AndNotOp, NetRefHandleManager->GetSubObjectInternalIndicesView());
 
-	TTuple<uint32, uint32> Stats(0, 0);
-	LogRootObjectListWithSubObjects(NetRefHandleManager, MakeNetBitArrayView(RootObjects), Stats, nullptr);
-	
+	TArray<FRootObjectData> RootObjectArray;
+	{
+		FillRootObjectArrayFromBitArray(RootObjectArray, RootObjectsView, NetRefHandleManager);
+		SortViaTrait(RootObjectArray, (EPrintTraits)ArgTraits, EPrintTraits::Default);
+	}
+
+	FLogContext LogContext = {.NetRefHandleManager = NetRefHandleManager, .RootObjectArray = RootObjectArray };
+	LogViaTrait(LogContext, (EPrintTraits)ArgTraits, EPrintTraits::Default);
 
 	UE_LOG(LogIrisBridge, Display, TEXT(""));
-	UE_LOG(LogIrisBridge, Display, TEXT("Printed %u root objects and %u sub objects"), Stats.Get<0>(), Stats.Get<1>());
+	UE_LOG(LogIrisBridge, Display, TEXT("Printed %u root objects and %u sub objects"), LogContext.NumRootObjects, LogContext.NumSubObjects);
 	UE_LOG(LogIrisBridge, Display, TEXT("################ Stop Printing Relevant Objects ################"));
 }
 
@@ -272,12 +418,13 @@ FAutoConsoleCommand ObjectBridgePrintAlwaysRelevantObjects(
 	{
 		if (UObjectReplicationBridge* ObjectBridge = CastChecked<UObjectReplicationBridge>(RepSystem->GetReplicationBridge()))
 		{
-			ObjectBridge->PrintAlwaysRelevantObjects();
+			EPrintTraits ArgTraits = FindPrintTraitsFromArgs(Args);
+			ObjectBridge->PrintAlwaysRelevantObjects((uint32)ArgTraits);
 		}
 	}
 }));
 
-void UObjectReplicationBridge::PrintAlwaysRelevantObjects() const
+void UObjectReplicationBridge::PrintAlwaysRelevantObjects(uint32 ArgTraits) const
 {
 	using namespace UE::Net;
 	using namespace UE::Net::Private;
@@ -296,18 +443,25 @@ void UObjectReplicationBridge::PrintAlwaysRelevantObjects() const
 	// Remove subobjects from the list.
 	MakeNetBitArrayView(AlwaysRelevantList).Combine(NetRefHandleManager->GetSubObjectInternalIndicesView(), FNetBitArrayView::AndNotOp);
 
-	TTuple<uint32, uint32> Stats(0, 0);
-	LogRootObjectListWithSubObjects(NetRefHandleManager, MakeNetBitArrayView(AlwaysRelevantList), Stats, nullptr);
+	TArray<FRootObjectData> AlwaysRelevantObjects;
+	{
+		FillRootObjectArrayFromBitArray(AlwaysRelevantObjects, MakeNetBitArrayView(AlwaysRelevantList), NetRefHandleManager);
+		SortViaTrait(AlwaysRelevantObjects, (EPrintTraits)ArgTraits, EPrintTraits::SortByClass);
+	}
+
+	FLogContext LogContext = {.NetRefHandleManager=NetRefHandleManager, .RootObjectArray=AlwaysRelevantObjects };
+	LogViaTrait(LogContext, (EPrintTraits)ArgTraits, EPrintTraits::Default);
 
 	UE_LOG(LogIrisBridge, Display, TEXT(""));
-	UE_LOG(LogIrisBridge, Display, TEXT("Printed %u root objects and %u sub objects"), Stats.Get<0>(), Stats.Get<1>());
+	UE_LOG(LogIrisBridge, Display, TEXT("Printed %u root objects and %u subobjects"), LogContext.NumRootObjects, LogContext.NumSubObjects);
 	UE_LOG(LogIrisBridge, Display, TEXT("################ Stop Printing Always Relevant Objects ################"));
 }
 
 //-----------------------------------------------
 FAutoConsoleCommand ObjectBridgePrintRelevantObjectsToConnection(
 TEXT("Net.Iris.PrintRelevantObjectsToConnection"),
-TEXT("Prints the list of replicated objects relevant to a specific connection"),
+TEXT("Prints the list of replicated objects relevant to a specific connection.")
+TEXT(" OptionalParams: WithFilter"),
 FConsoleCommandWithArgsDelegate::CreateLambda([](const TArray< FString >& Args)
 {
 	using namespace UE::Net;
@@ -332,13 +486,13 @@ void UObjectReplicationBridge::PrintRelevantObjectsForConnections(const TArray<F
 	using namespace UE::Net::Private;
 	using namespace UE::Net::Private::ObjectBridgeDebugging;
 
-	UE_LOG(LogIrisBridge, Display, TEXT("################ Start Printing Relevant Objects of %d Connections ################"), 1);
-	UE_LOG(LogIrisBridge, Display, TEXT(""));
-
 	FReplicationSystemInternal* ReplicationSystemInternal = GetReplicationSystem()->GetReplicationSystemInternal();
+
 	FReplicationConnections& Connections = ReplicationSystemInternal->GetConnections();
 	const FNetBitArray& ValidConnections = Connections.GetValidConnections();
 
+	const FReplicationFiltering& Filtering = ReplicationSystemInternal->GetFiltering();
+
 	// Default to all connections
 	FNetBitArray ConnectionsToPrint;
 	ConnectionsToPrint.InitAndCopy(ValidConnections);
@@ -361,57 +515,176 @@ void UObjectReplicationBridge::PrintRelevantObjectsForConnections(const TArray<F
 		}
 	}
 
+	UE_LOG(LogIrisBridge, Display, TEXT("################ Start Printing Relevant Objects of %d Connections ################"), ConnectionsToPrint.CountSetBits());
+	UE_LOG(LogIrisBridge, Display, TEXT(""));
+
+	const bool bWithFilterInfo = nullptr != Args.FindByPredicate([](const FString& Str) { return Str.Contains(TEXT("WithFilter")); });
+
+	EPrintTraits ArgTraits = FindPrintTraitsFromArgs(Args);
+
 	ConnectionsToPrint.ForAllSetBits([&](uint32 ConnectionId)
 	{
+		const FReplicationView& ConnectionViews = Connections.GetReplicationView(ConnectionId);
+		FString ViewLocs;
+		for (const FReplicationView::FView& UserView : ConnectionViews.Views )
+		{
+			ViewLocs += FString::Printf(TEXT("%s "), *UserView.Pos.ToCompactString());
+		}
+
 		UE_LOG(LogIrisBridge, Display, TEXT(""));
-		UE_LOG(LogIrisBridge, Display, TEXT("###### Begin Relevant list of Connection:%u ######"), ConnectionId);
+		UE_LOG(LogIrisBridge, Display, TEXT("###### Begin Relevant list of Connection:%u ViewPos:%s Named: %s ######"), ConnectionId, *ViewLocs, *PrintConnectionInfo(ConnectionId));
 		UE_LOG(LogIrisBridge, Display, TEXT(""));
 
-		struct FRootObjectData
-		{
-			FInternalNetRefIndex ObjectIndex = 0;
-			UObject* Instance = nullptr;
-			UClass* Class = nullptr;
-
-			// Sort by class names
-			bool operator<(const FRootObjectData& rhs) const 
-			{ 
-				if (Class==rhs.Class) { return false; }
-				if (!Class) { return false; }
-				if (!rhs.Class) { return true; }
-				return Class->GetName() < rhs.Class->GetName();
-			}
-		};
-		TArray<FRootObjectData> RelevantObjects;
-
 		FNetBitArray RootObjects;
 		RootObjects.Init(NetRefHandleManager->GetMaxActiveObjectCount());
 		MakeNetBitArrayView(RootObjects).Set(GetReplicationSystem()->GetReplicationSystemInternal()->GetFiltering().GetRelevantObjectsInScope(ConnectionId), FNetBitArrayView::AndNotOp, NetRefHandleManager->GetSubObjectInternalIndicesView());
-		
-		RootObjects.ForAllSetBits([&](uint32 RootObjectIndex)
+
+		TArray<FRootObjectData> RelevantObjects;
 		{
-			FRootObjectData Data;
-			Data.ObjectIndex = RootObjectIndex;
-			Data.Instance = NetRefHandleManager->GetReplicatedObjectInstance(RootObjectIndex);
-			Data.Class = Data.Instance ? Data.Instance->GetClass() : nullptr;
+			FillRootObjectArrayFromBitArray(RelevantObjects, MakeNetBitArrayView(RootObjects), NetRefHandleManager);
+			SortViaTrait(RelevantObjects, ArgTraits, EPrintTraits::SortByClass);
+		}
+
+		auto AddFilterInfo = [&](FInternalNetRefIndex ObjectIndex) -> FString
+		{
+            // TODO: When printing with subobjects. Try to tell if they are relevant or not to the connection.
+			return FString::Printf(TEXT("\t%s"), *Filtering.PrintFilterObjectInfo(ObjectIndex, ConnectionId));
+		};
 
-			RelevantObjects.Emplace(MoveTemp(Data));
-		});
+		FLogContext LogContext = {.NetRefHandleManager=NetRefHandleManager, .RootObjectArray=RelevantObjects, .OptionalObjectPrint=AddFilterInfo};
+		LogViaTrait(LogContext, ArgTraits, EPrintTraits::Default);
 
-		RelevantObjects.Sort();
+		UE_LOG(LogIrisBridge, Display, TEXT(""));
+		UE_LOG(LogIrisBridge, Display, TEXT("###### Stop Relevant list of Connection:%u | Total: %u root objects relevant ######"), ConnectionId, LogContext.NumRootObjects);
+		UE_LOG(LogIrisBridge, Display, TEXT(""));
+	});
+
+	UE_LOG(LogIrisBridge, Display, TEXT(""));
+	UE_LOG(LogIrisBridge, Display, TEXT("################ Stop Printing Relevant Objects of %d Connections ################"), ConnectionsToPrint.CountSetBits());
+}
+
+//-----------------------------------------------
+FAutoConsoleCommand ObjectBridgePrintNetCullDistances(
+TEXT("Net.Iris.PrintNetCullDistances"),
+TEXT("Prints the list of replicated objects and their current netculldistance."),
+FConsoleCommandWithArgsDelegate::CreateLambda([](const TArray< FString >& Args)
+{
+	using namespace UE::Net;
+	using namespace UE::Net::Private;
+	using namespace UE::Net::Private::ObjectBridgeDebugging;
 
-		for (const FRootObjectData& ObjectData : RelevantObjects)
+	UReplicationSystem* RepSystem = FindReplicationSystemFromArg(Args);
+	if (RepSystem)
+	{
+		if (UObjectReplicationBridge* ObjectBridge = CastChecked<UObjectReplicationBridge>(RepSystem->GetReplicationBridge()))
 		{
-			UE_LOG(LogIrisBridge, Display, TEXT("%s"), *PrintNetObject(NetRefHandleManager, ObjectData.ObjectIndex));
+			FReplicationSystemInternal* ReplicationSystemInternal = RepSystem->GetReplicationSystemInternal();
 
-			// TODO: Output the subobjects and tell if they are relevant or not to the connection.
+			ObjectBridge->PrintNetCullDistances(Args);
 		}
+	}
+}));
 
-		UE_LOG(LogIrisBridge, Display, TEXT(""));
-		UE_LOG(LogIrisBridge, Display, TEXT("###### Stop Relevant list of Connection:%u | Total: %u root objects relevant ######"), ConnectionId, RelevantObjects.Num());
-		UE_LOG(LogIrisBridge, Display, TEXT(""));
+void UObjectReplicationBridge::PrintNetCullDistances(const TArray<FString>& Args) const
+{
+	using namespace UE::Net;
+	using namespace UE::Net::Private;
+	using namespace UE::Net::Private::ObjectBridgeDebugging;
+
+	FReplicationSystemInternal* ReplicationSystemInternal = GetReplicationSystem()->GetReplicationSystemInternal();
+	const FWorldLocations& WorldLocations = ReplicationSystemInternal->GetWorldLocations();
+
+	struct FCullDistanceInfo
+	{
+		UClass* Class = nullptr;
+		float CDOCullDistance = 0.0f;
+
+		uint32 NumTotal = 0; // Total replicated rootobjects of this class
+		uint32 NumCDOCullDistance = 0; // Total replicated rootobjects
+
+		// Track culldistance values for actors that are different from the CDO
+		TMap<float /*CullDistance*/, uint32 /*ActorCount with culldistance value*/> DivergentCullDistances; 
+	};
+
+	TMap<UClass*, FCullDistanceInfo> ClassCullDistanceMap;
+
+	FNetBitArray RootObjects;
+	RootObjects.InitAndCopy(NetRefHandleManager->GetGlobalScopableInternalIndices());
+	
+	// Remove objects that didn't register world location info
+	MakeNetBitArrayView(RootObjects).Combine(WorldLocations.GetObjectsWithWorldInfo(), FNetBitArrayView::AndOp);
+
+	// Filter down to objects in the GridFilter. Other filters do not use net culling
+	{
+		FNetBitArray GridFilterList;
+		GridFilterList.Init(NetRefHandleManager->GetMaxActiveObjectCount());
+		ReplicationSystemInternal->GetFiltering().BuildObjectsInFilterList(MakeNetBitArrayView(GridFilterList), TEXT("Spatial"));
+		RootObjects.Combine(GridFilterList, FNetBitArray::AndOp);
+	}
+
+	RootObjects.ForAllSetBits([&](uint32 RootObjectIndex)
+	{
+		const float CurrentCullDistance = WorldLocations.GetCullDistance(RootObjectIndex);
+
+		if( UObject* RepObj = NetRefHandleManager->GetReplicatedObjectInstance(RootObjectIndex) )
+		{
+			UClass* RepObjClass = RepObj->GetClass();
+			UObject* RepClassCDO = RepObjClass->GetDefaultObject();
+
+			float CDOCullDistance = 0.0;
+
+			// Try to find the object's actual culldistance and that of the CDO.
+			if (GetInstanceWorldObjectInfoFunction)
+			{
+				FVector Loc;
+				GetInstanceWorldObjectInfoFunction(NetRefHandleManager->GetNetRefHandleFromInternalIndex(RootObjectIndex), RepClassCDO, Loc, CDOCullDistance);
+			}
+
+			FCullDistanceInfo& Info = ClassCullDistanceMap.FindOrAdd(RepObjClass);
+
+			if (Info.Class == nullptr)
+			{
+				Info.Class = RepObjClass;
+				Info.CDOCullDistance = CDOCullDistance;
+			}
+
+			Info.NumTotal++;
+			
+			// Check if Obj has diverged from the CDO
+			if(CurrentCullDistance != CDOCullDistance)
+			{
+				uint32& NumDivergent = Info.DivergentCullDistances.FindOrAdd(CurrentCullDistance, 0);
+				++NumDivergent;
+			}
+			else
+			{
+				Info.NumCDOCullDistance++;
+			}
+		}
 	});
 
+	// Sort from highest to lowest
+	ClassCullDistanceMap.ValueSort([](const FCullDistanceInfo& lhs, const FCullDistanceInfo& rhs) { return lhs.CDOCullDistance >= rhs.CDOCullDistance; });
+
+	UE_LOG(LogIrisBridge, Display, TEXT("################ Start Printing NetCullDistance Values ################"));
+	UE_LOG(LogIrisBridge, Display, TEXT(""));
+
+	for (auto ClassIt = ClassCullDistanceMap.CreateIterator(); ClassIt; ++ClassIt)
+	{
+		FCullDistanceInfo& Info = ClassIt.Value();
+		UClass* Class = Info.Class;
+
+		UE_LOG(LogIrisBridge, Display, TEXT("NetCullDistance: %f | Class: %s | ReplicatedCount: %u | Using CDO CullDistance: %u (%.2f%%)"),
+			Info.CDOCullDistance, *Info.Class->GetName(), Info.NumTotal, Info.NumCDOCullDistance, (Info.NumCDOCullDistance/(float)Info.NumTotal)*100.f);
+
+		Info.DivergentCullDistances.KeySort([](const float& lhs, const float& rhs){ return lhs >= rhs; });
+
+		for (auto DivergentIt = Info.DivergentCullDistances.CreateConstIterator(); DivergentIt; ++DivergentIt)
+		{
+			UE_LOG(LogIrisBridge, Display, TEXT("\tNetCullDistance: %f | UseCount: %d (%.2f%%)"), DivergentIt.Key(), DivergentIt.Value(), (DivergentIt.Value()/(float)Info.NumTotal)*100.f);
+		}
+	}
+	
 	UE_LOG(LogIrisBridge, Display, TEXT(""));
-	UE_LOG(LogIrisBridge, Display, TEXT("################ Stop Printing Relevant Objects of %d Connections ################"), 1);
+	UE_LOG(LogIrisBridge, Display, TEXT("################ Stop Printing NetCullDistance Values ################"));
 }
\ No newline at end of file
diff --git a/Engine/Source/Runtime/Experimental/Iris/Core/Public/Iris/ReplicationSystem/Filtering/NetObjectFilter.h b/Engine/Source/Runtime/Experimental/Iris/Core/Public/Iris/ReplicationSystem/Filtering/NetObjectFilter.h
--- a/Engine/Source/Runtime/Experimental/Iris/Core/Public/Iris/ReplicationSystem/Filtering/NetObjectFilter.h
+++ b/Engine/Source/Runtime/Experimental/Iris/Core/Public/Iris/ReplicationSystem/Filtering/NetObjectFilter.h
@@ -238,6 +238,21 @@ public:
 	/** Returns the filter's traits. */
 	ENetFilterTraits GetFilterTraits() const { return FilterTraits; }
 
+	struct FDebugInfoParams
+	{
+		FName FilterName;
+
+		const FNetObjectFilteringInfo* FilteringInfos = nullptr;
+
+		/** ID of the connection that the filtering applies to. */
+		uint32 ConnectionId = 0;
+
+		/** The view associated with the connection and its sub-connections that objects are filtered for. */
+		UE::Net::FReplicationView View;
+	};
+
+	IRISCORE_API virtual FString PrintDebugInfoForObject(const FDebugInfoParams& Params, uint32 ObjectIndex) const { return Params.FilterName.ToString(); };
+
 protected:
 	IRISCORE_API UNetObjectFilter();
 
diff --git a/Engine/Source/Runtime/Experimental/Iris/Core/Public/Iris/ReplicationSystem/Filtering/NetObjectGridFilter.h b/Engine/Source/Runtime/Experimental/Iris/Core/Public/Iris/ReplicationSystem/Filtering/NetObjectGridFilter.h
--- a/Engine/Source/Runtime/Experimental/Iris/Core/Public/Iris/ReplicationSystem/Filtering/NetObjectGridFilter.h
+++ b/Engine/Source/Runtime/Experimental/Iris/Core/Public/Iris/ReplicationSystem/Filtering/NetObjectGridFilter.h
@@ -15,6 +15,11 @@ namespace UE::Net
 	class FNetCullDistanceOverrides;
 	class FWorldLocations;
 	struct FRepTagFindInfo;
+
+	namespace Private
+	{
+		class FNetRefHandleManager;
+	}
 }
 
 UCLASS(transient, config=Engine, MinimalAPI)
@@ -35,7 +40,7 @@ public:
 
 	/** Objects with larger sqrt(NetCullDistanceSqr) will be rejected. Disabled when value is zero. */
 	UPROPERTY(Config)
-	float MaxCullDistance = 0;
+	float MaxCullDistance = 0.0f;
 
 	/** Objects without a NetCullDistanceSquared property will assume to have this value but squared unless there's a cull distance override. */
 	UPROPERTY(Config)
@@ -69,6 +74,7 @@ protected:
 	IRISCORE_API virtual void PreFilter(FNetObjectPreFilteringParams&) override;
 	IRISCORE_API virtual void Filter(FNetObjectFilteringParams&) override;
 	IRISCORE_API virtual void PostFilter(FNetObjectPostFilteringParams&) override;
+	IRISCORE_API virtual FString PrintDebugInfoForObject(const FDebugInfoParams& Params, uint32 ObjectIndex) const override;
 
 protected:
 	struct FObjectLocationInfo : public FNetObjectFilteringInfo
@@ -233,6 +239,7 @@ private:
 	TMap<FCellCoord, FCellObjects> Cells;
 	uint32 FrameIndex = 0;
 
+	const UE::Net::Private::FNetRefHandleManager*  NetRefHandleManager = nullptr;
 	const UE::Net::FNetCullDistanceOverrides* NetCullDistanceOverrides = nullptr;
 };
 
diff --git a/Engine/Source/Runtime/Experimental/Iris/Core/Public/Iris/ReplicationSystem/ObjectReplicationBridge.h b/Engine/Source/Runtime/Experimental/Iris/Core/Public/Iris/ReplicationSystem/ObjectReplicationBridge.h
--- a/Engine/Source/Runtime/Experimental/Iris/Core/Public/Iris/ReplicationSystem/ObjectReplicationBridge.h
+++ b/Engine/Source/Runtime/Experimental/Iris/Core/Public/Iris/ReplicationSystem/ObjectReplicationBridge.h
@@ -181,12 +181,13 @@ public:
 
 public:
 
-	// Debug functions exposed via console commands
+	// Debug functions exposed that are triggered via console commands
 	void PrintDynamicFilterClassConfig() const;
-	void PrintReplicatedObjects() const;
-	void PrintAlwaysRelevantObjects() const;
-	void PrintRelevantObjects() const;
+	void PrintReplicatedObjects(uint32 ArgTraits) const;
+	void PrintAlwaysRelevantObjects(uint32 ArgTraits) const;
+	void PrintRelevantObjects(uint32 ArgTraits) const;
 	void PrintRelevantObjectsForConnections(const TArray<FString>& Args) const;
+	void PrintNetCullDistances(const TArray<FString>& Args) const;
 
 protected:
 	IRISCORE_API virtual ~UObjectReplicationBridge();
@@ -306,7 +307,7 @@ protected:
 	IRISCORE_API bool IsClassCritical(const UClass* Class);
 
 	/** Returns the most relevant description of the client tied to this connection id. */
-	[[nodiscard]] IRISCORE_API virtual FString PrintConnectionInfo(uint32 ConnectionId);
+	[[nodiscard]] IRISCORE_API virtual FString PrintConnectionInfo(uint32 ConnectionId) const;
 
 	/** Current max tick rate set by the engine */
 	float GetMaxTickRate() const { return MaxTickRate; }
diff --git a/Engine/Source/Runtime/Experimental/Iris/Core/Public/Iris/ReplicationSystem/WorldLocations.h b/Engine/Source/Runtime/Experimental/Iris/Core/Public/Iris/ReplicationSystem/WorldLocations.h
--- a/Engine/Source/Runtime/Experimental/Iris/Core/Public/Iris/ReplicationSystem/WorldLocations.h
+++ b/Engine/Source/Runtime/Experimental/Iris/Core/Public/Iris/ReplicationSystem/WorldLocations.h
@@ -64,6 +64,9 @@ public:
 	void ResetObjectsWithDirtyInfo();
 	FNetBitArrayView GetObjectsWithDirtyInfo() const;
 
+	/** Returns the list of objects that registered world location information */
+	const FNetBitArrayView GetObjectsWithWorldInfo() const { return MakeNetBitArrayView(ValidInfoIndexes); }
+
 private:
 	enum : uint32
 	{
