Fix for float conversion warning

#rnx

diff --git a/Engine/Source/Runtime/Experimental/Iris/Core/Private/Iris/ReplicationSystem/ObjectReplicationBridgeDebugging.cpp b/Engine/Source/Runtime/Experimental/Iris/Core/Private/Iris/ReplicationSystem/ObjectReplicationBridgeDebugging.cpp
--- a/Engine/Source/Runtime/Experimental/Iris/Core/Private/Iris/ReplicationSystem/ObjectReplicationBridgeDebugging.cpp
+++ b/Engine/Source/Runtime/Experimental/Iris/Core/Private/Iris/ReplicationSystem/ObjectReplicationBridgeDebugging.cpp
@@ -675,13 +675,13 @@ void UObjectReplicationBridge::PrintNetCullDistances(const TArray<FString>& Args
 		UClass* Class = Info.Class;
 
 		UE_LOG(LogIrisBridge, Display, TEXT("NetCullDistance: %f | Class: %s | ReplicatedCount: %u | Using CDO CullDistance: %u (%.2f%%)"),
-			Info.CDOCullDistance, *Info.Class->GetName(), Info.NumTotal, Info.NumCDOCullDistance, (Info.NumCDOCullDistance/(float)Info.NumTotal)*100.f);
+			Info.CDOCullDistance, *Info.Class->GetName(), Info.NumTotal, Info.NumCDOCullDistance, ((float)Info.NumCDOCullDistance/(float)Info.NumTotal)*100.f);
 
 		Info.DivergentCullDistances.KeySort([](const float& lhs, const float& rhs){ return lhs >= rhs; });
 
 		for (auto DivergentIt = Info.DivergentCullDistances.CreateConstIterator(); DivergentIt; ++DivergentIt)
 		{
-			UE_LOG(LogIrisBridge, Display, TEXT("\tNetCullDistance: %f | UseCount: %d (%.2f%%)"), DivergentIt.Key(), DivergentIt.Value(), (DivergentIt.Value()/(float)Info.NumTotal)*100.f);
+			UE_LOG(LogIrisBridge, Display, TEXT("\tNetCullDistance: %f | UseCount: %d (%.2f%%)"), DivergentIt.Key(), DivergentIt.Value(), ((float)DivergentIt.Value()/(float)Info.NumTotal)*100.f);
 		}
 	}
 	
