implement UPrimitiveComponent::AddImpulse and FWritePhysicsObjectInterface::SetLinearImpulseVelocity

diff --git a/Engine/Source/Runtime/Engine/Private/PrimitiveComponentPhysics.cpp b/Engine/Source/Runtime/Engine/Private/PrimitiveComponentPhysics.cpp
--- a/Engine/Source/Runtime/Engine/Private/PrimitiveComponentPhysics.cpp
+++ b/Engine/Source/Runtime/Engine/Private/PrimitiveComponentPhysics.cpp
@@ -174,6 +174,24 @@ void UPrimitiveComponent::AddImpulse(FVector Impulse, FName BoneName, bool bVelC
 		WarnInvalidPhysicsOperations(LOCTEXT("AddImpulse", "AddImpulse"), BI, BoneName);
 		BI->AddImpulse(Impulse, bVelChange);
 	}
+	else if (BoneName == NAME_None)
+	{
+		TArray<Chaos::FPhysicsObjectHandle> PhysicsObjects = GetAllPhysicsObjects();
+		FLockedWritePhysicsObjectExternalInterface Interface = FPhysicsObjectExternalInterface::LockWrite(PhysicsObjects);
+		
+		PhysicsObjects = PhysicsObjects.FilterByPredicate(
+			[&Interface](Chaos::FPhysicsObject* Object) {
+			return !Interface->AreAllDisabled({ &Object, 1 });
+		}
+		);
+		Interface->SetLinearImpulseVelocity(PhysicsObjects, Impulse, bVelChange);
+	}
+	else if (Chaos::FPhysicsObject* Object = GetPhysicsObjectByName(BoneName))
+	{
+		FLockedWritePhysicsObjectExternalInterface Interface = FPhysicsObjectExternalInterface::LockWrite({ &Object, 1 });
+		Interface->SetLinearImpulseVelocity({ &Object, 1 }, Impulse, bVelChange);
+	}
+
 }
 
 void UPrimitiveComponent::AddAngularImpulseInRadians(FVector Impulse, FName BoneName, bool bVelChange)
diff --git a/Engine/Source/Runtime/Experimental/Chaos/Private/Chaos/PhysicsObjectInterface.cpp b/Engine/Source/Runtime/Experimental/Chaos/Private/Chaos/PhysicsObjectInterface.cpp
--- a/Engine/Source/Runtime/Experimental/Chaos/Private/Chaos/PhysicsObjectInterface.cpp
+++ b/Engine/Source/Runtime/Experimental/Chaos/Private/Chaos/PhysicsObjectInterface.cpp
@@ -1015,6 +1015,51 @@ namespace Chaos
 		}
 	}
 
+	template<EThreadContext Id>
+	void FWritePhysicsObjectInterface<Id>::SetLinearImpulseVelocity(TArrayView<const FPhysicsObjectHandle> InObjects, const FVector& Impulse, bool bVelChange)
+	{
+		if (InObjects.IsEmpty())
+		{
+			return;
+		}
+
+		for (const FPhysicsObjectHandle Object : InObjects)
+		{
+			if (!Object)
+			{
+				continue;
+			}
+
+			if (TThreadParticle<Id>* Particle = Object->GetParticle<Id>())
+			{
+				if (Chaos::TThreadRigidParticle<Id>* Rigid = Particle->CastToRigidParticle())
+				{
+					if (bVelChange)
+					{
+						Rigid->SetLinearImpulseVelocity(Impulse * Rigid->M(), false);
+					}
+					else
+					{
+						Rigid->SetLinearImpulseVelocity(Impulse, false);
+					}
+				}
+			}
+		}
+
+		if constexpr (Id == EThreadContext::External)
+		{
+			if (Chaos::FPhysicsSolverBase* Solver = FPhysicsObjectInterface::GetSolver(InObjects[0]))
+			{
+				Solver->EnqueueCommandImmediate(
+					[AllObjects = TArray<FPhysicsObjectHandle>{InObjects}, Impulse, bVelChange]() {
+						Chaos::FWritePhysicsObjectInterface_Internal Interface = Chaos::FPhysicsObjectInternalInterface::GetWrite();
+						Interface.SetLinearImpulseVelocity(AllObjects, Impulse, bVelChange);
+					}
+				);
+			}
+		}
+	}
+
 	template<EThreadContext Id>
 	void FWritePhysicsObjectInterface<Id>::AddRadialImpulse(TArrayView<const FPhysicsObjectHandle> InObjects, FVector Origin, float Radius, float Strength, enum ERadialImpulseFalloff Falloff, bool bApplyStrain, bool bInvalidate, bool bVelChange)
 	{
diff --git a/Engine/Source/Runtime/Experimental/Chaos/Public/Chaos/PhysicsObjectInterface.h b/Engine/Source/Runtime/Experimental/Chaos/Public/Chaos/PhysicsObjectInterface.h
--- a/Engine/Source/Runtime/Experimental/Chaos/Public/Chaos/PhysicsObjectInterface.h
+++ b/Engine/Source/Runtime/Experimental/Chaos/Public/Chaos/PhysicsObjectInterface.h
@@ -131,6 +131,7 @@ namespace Chaos
 		CHAOS_API void ForceKinematic(TArrayView<const FPhysicsObjectHandle> InObjects);
 		CHAOS_API void AddForce(TArrayView<const FPhysicsObjectHandle> InObjects, const FVector& Force, bool bInvalidate);
 		CHAOS_API void AddTorque(TArrayView<const FPhysicsObjectHandle> InObjects, const FVector& Torque, bool bInvalidate);
+		CHAOS_API void SetLinearImpulseVelocity(TArrayView<const FPhysicsObjectHandle> InObjects, const FVector& Impulse, bool bVelChange);
 
 		UE_DEPRECATED(5.4, "This version AddRadialImpulse has been deprecated. Please use the version where the strain value is passed explicitly")
 		CHAOS_API void AddRadialImpulse(TArrayView<const FPhysicsObjectHandle> InObjects, FVector Origin, float Radius, float Strength, enum ERadialImpulseFalloff Falloff, bool bApplyStrain, bool bInvalidate, bool bVelChange = false);
