particles with the OneWayInteraction flag set should not contribute to strain

diff --git a/Engine/Source/Runtime/Experimental/Chaos/Private/Chaos/Collision/PBDCollisionConstraint.cpp b/Engine/Source/Runtime/Experimental/Chaos/Private/Chaos/Collision/PBDCollisionConstraint.cpp
--- a/Engine/Source/Runtime/Experimental/Chaos/Private/Chaos/Collision/PBDCollisionConstraint.cpp
+++ b/Engine/Source/Runtime/Experimental/Chaos/Private/Chaos/Collision/PBDCollisionConstraint.cpp
@@ -339,12 +339,12 @@ namespace Chaos
 		const FRealSingle InitialOverlapDepenetrationVelocity1 = FConstGenericParticleHandle(GetParticle1())->InitialOverlapDepenetrationVelocity();
 		InitialOverlapDepenetrationVelocity = FMath::Max(InitialOverlapDepenetrationVelocity0, InitialOverlapDepenetrationVelocity1);
 
-		// Is this a one-way interaction?
+		// Is this a one-way interaction? A dynamic one-way interaction particle that hits a kinematic non one-way interaction particle should still be considered one-way.
 		const bool bDynamic0 = FConstGenericParticleHandle(GetParticle0())->IsDynamic();
 		const bool bDynamic1 = FConstGenericParticleHandle(GetParticle1())->IsDynamic();
 		const bool bOneWay0 = FConstGenericParticleHandle(GetParticle0())->OneWayInteraction();
 		const bool bOneWay1 = FConstGenericParticleHandle(GetParticle1())->OneWayInteraction();
-		Flags.bIsOneWayInteraction = bDynamic0 && bDynamic1 && (bOneWay0 || bOneWay1);
+		Flags.bIsOneWayInteraction = (bDynamic0 || bDynamic1) && (bOneWay0 || bOneWay1);
 
 		// Only levelsets use incremental collision manifolds
 		if (bInUseManifold && ((ImplicitType0 == ImplicitObjectType::LevelSet) || (ImplicitType1 == ImplicitObjectType::LevelSet)))
diff --git a/Engine/Source/Runtime/Experimental/Chaos/Private/Chaos/PBDRigidClustering.cpp b/Engine/Source/Runtime/Experimental/Chaos/Private/Chaos/PBDRigidClustering.cpp
--- a/Engine/Source/Runtime/Experimental/Chaos/Private/Chaos/PBDRigidClustering.cpp
+++ b/Engine/Source/Runtime/Experimental/Chaos/Private/Chaos/PBDRigidClustering.cpp
@@ -2212,6 +2212,11 @@ namespace Chaos
 				continue;
 			}
 
+			if (ContactHandle->GetContact().GetIsOneWayInteraction())
+			{
+				continue;
+			}
+
 			TVector<FGeometryParticleHandle*, 2> ConstrainedParticles = ContactHandle->GetConstrainedParticles();
 			
 			// make sure we only compute things if one of the two particle is clustered
diff --git a/Engine/Source/Runtime/Experimental/Chaos/Public/Chaos/Collision/PBDCollisionConstraint.h b/Engine/Source/Runtime/Experimental/Chaos/Public/Chaos/Collision/PBDCollisionConstraint.h
--- a/Engine/Source/Runtime/Experimental/Chaos/Public/Chaos/Collision/PBDCollisionConstraint.h
+++ b/Engine/Source/Runtime/Experimental/Chaos/Public/Chaos/Collision/PBDCollisionConstraint.h
@@ -405,6 +405,8 @@ namespace Chaos
 		// Whether this constraint was disabled by the user (e.g., via a collision callback)
 		bool GetDisabled() const { return Flags.bDisabled; }
 
+		bool GetIsOneWayInteraction() const { return Flags.bIsOneWayInteraction; }
+
 		void SetIsProbe(bool bInProbe) { Flags.bIsProbe = bInProbe; }
 		bool GetIsProbe() const { return Flags.bIsProbe; }
 
