Simple -Only option so I can target the precise test I want

#rnx

diff --git a/Engine/Source/Runtime/Experimental/IoStoreOnDemand/Private/Tool/Test.cpp b/Engine/Source/Runtime/Experimental/IoStoreOnDemand/Private/Tool/Test.cpp
--- a/Engine/Source/Runtime/Experimental/IoStoreOnDemand/Private/Tool/Test.cpp
+++ b/Engine/Source/Runtime/Experimental/IoStoreOnDemand/Private/Tool/Test.cpp
@@ -41,9 +41,23 @@ static void CacheTests(const FContext& Context)
 ////////////////////////////////////////////////////////////////////////////////
 static int32 TestCommandEntry(const FContext& Context)
 {
-	CommandTest();
-	CacheTests(Context);
-	HttpTests(Context);
+	FStringView Only = Context.Get<FStringView>(TEXT("-Only"));
+
+	if (Only.IsEmpty())
+	{
+		CommandTest();
+	}
+
+	if (Only.IsEmpty() || Only == TEXT("cache"))
+	{
+		CacheTests(Context);
+	}
+
+	if (Only.IsEmpty() || Only == TEXT("http"))
+	{
+		HttpTests(Context);
+	}
+
 	return 0;
 }
 
@@ -55,6 +69,7 @@ static FCommand TestCommand(
 	{
 		TArgument<FStringView>(TEXT("-Host"), TEXT("Host of the HTTP test server")),
 		TArgument<FStringView>(TEXT("-Dir"), TEXT("Primary directory to use for cache tests")),
+		TArgument<FStringView>(TEXT("-Only"), TEXT("Only run a particular test (http|cache)")),
 	}
 );
 
