SetBlocking() didn't actually work. Now it does

#rnx

diff --git a/Engine/Source/Runtime/Experimental/IoStoreOnDemand/Private/CoreHttp/Http.cpp b/Engine/Source/Runtime/Experimental/IoStoreOnDemand/Private/CoreHttp/Http.cpp
--- a/Engine/Source/Runtime/Experimental/IoStoreOnDemand/Private/CoreHttp/Http.cpp
+++ b/Engine/Source/Runtime/Experimental/IoStoreOnDemand/Private/CoreHttp/Http.cpp
@@ -900,28 +900,23 @@ int32 FSocket::Recv(char* Dest, uint32 Size)
 ////////////////////////////////////////////////////////////////////////////////
 bool FSocket::SetBlocking(bool bBlocking)
 {
-	bool bSuccess = false;
 #if defined(IAS_HTTP_HAS_NONBLOCK_IMPL)
-	bSuccess = SetNonBlockingSocket(Socket);
+	return SetNonBlockingSocket(Socket, !bBlocking);
 #elif PLATFORM_MICROSOFT
-	unsigned long NonBlockingMode = 1;
-	if (ioctlsocket(Socket, FIONBIO, &NonBlockingMode) != SOCKET_ERROR)
-	{
-		bSuccess = true;
-	}
+	unsigned long NonBlockingMode = (bBlocking != true);
+	return (ioctlsocket(Socket, FIONBIO, &NonBlockingMode) != SOCKET_ERROR);
 #else
 	int32 Flags = fcntl(Socket, F_GETFL, 0);
-	if (Flags != -1)
+	if (Flags == -1)
 	{
-		Flags |= Flags | int32(O_NONBLOCK);
-		if (fcntl(Socket, F_SETFL, Flags) >= 0)
-		{
-			bSuccess = true;
-		}
+		return false;
 	}
-#endif
 
-	return bSuccess;
+	int32 NewFlags = bBlocking
+		? (Flags & ~int32(O_NONBLOCK))
+		: (Flags |	int32(O_NONBLOCK));
+	return (Flags == NewFlags) || (fcntl(Socket, F_SETFL, Flags) >= 0);
+#endif
 }
 
 ////////////////////////////////////////////////////////////////////////////////
