ChaosFlesh: Viewport component scaling fix

diff --git a/Engine/Source/Runtime/Experimental/Chaos/Private/GeometryCollection/GeometryCollection.cpp b/Engine/Source/Runtime/Experimental/Chaos/Private/GeometryCollection/GeometryCollection.cpp
--- a/Engine/Source/Runtime/Experimental/Chaos/Private/GeometryCollection/GeometryCollection.cpp
+++ b/Engine/Source/Runtime/Experimental/Chaos/Private/GeometryCollection/GeometryCollection.cpp
@@ -1610,23 +1610,16 @@ void FGeometryCollection::Init(FGeometryCollection* Collection, const TArray<flo
 
 		// set the vertex information
 		TManagedArray<FVector2f>* UV0 = Collection->FindUVLayer(0);
-		FVector3f TempVertices(0.f, 0.f, 0.f);
 		for (int32 Idx = 0; Idx < NumNewVertices; ++Idx)
 		{
 			Vertices[Idx] = FVector3f(RawVertexArray[3 * Idx], RawVertexArray[3 * Idx + 1], RawVertexArray[3 * Idx + 2]);
-			TempVertices += Vertices[Idx];
 			(*UV0)[Idx] = FVector2f::ZeroVector;
 
 			Colors[Idx] = Collection->Defaults.DefaultVertexColor;
 			BoneMap[Idx] = 0;
 		}
 
-		
-
-		// set the particle information
-		TempVertices /= (float)NumNewVertices;
-		Transform[0] = FTransform3f(TempVertices);
-		Transform[0].NormalizeRotation();
+		Transform[0] = FTransform3f(FVector3f(0.f,0.f,0.f));
 
 		// set the index information
 		TArray<FVector3f> FaceNormals;
