Route all test requests through the test server proxy

Prior to this change test requests would go either directly to the server, or via the proxy to inject faults, delays, etc. Now all requests go through the proxy. This will allow us to fire up easily reuse all tests against a different proxy later.

#rnx

diff --git a/Engine/Source/Runtime/Experimental/IoStoreOnDemand/Private/CoreHttp/Http.cpp b/Engine/Source/Runtime/Experimental/IoStoreOnDemand/Private/CoreHttp/Http.cpp
--- a/Engine/Source/Runtime/Experimental/IoStoreOnDemand/Private/CoreHttp/Http.cpp
+++ b/Engine/Source/Runtime/Experimental/IoStoreOnDemand/Private/CoreHttp/Http.cpp
@@ -3648,7 +3648,8 @@ IOSTOREONDEMAND_API void IasHttpTest(const ANSICHAR* TestHost="localhost")
 #endif
 
 	TAnsiStringBuilder<64> Ret;
-	auto BuildUrl = [&] (const ANSICHAR* Suffix=nullptr, int32 Port=9493) -> const auto& {
+	auto BuildUrl = [&] (const ANSICHAR* Suffix=nullptr, uint32 Port=9493) -> const auto&
+	{
 		Ret.Reset();
 		Ret << "http://";
 		Ret << TestHost;
@@ -3824,11 +3825,11 @@ IOSTOREONDEMAND_API void IasHttpTest(const ANSICHAR* TestHost="localhost")
 		};
 
 		FConnectionPool::FParams Params;
-		Params.SetHostFromUrl(BuildUrl("", 9494));
+		Params.SetHostFromUrl(BuildUrl());
 		FConnectionPool Pool(Params);
 
 		FEventLoop Loop2;
-		Loop2.Send(Loop2.Get("/data?stall", Pool), Sink);
+		Loop2.Send(Loop2.Get("/data?stall=1", Pool), Sink);
 
 		// Requests are pipelined. The second one will get went during the stall so
 		// we expect it to fail. The subsequent ones are expected to succeed.
@@ -3907,8 +3908,8 @@ IOSTOREONDEMAND_API void IasHttpTest(const ANSICHAR* TestHost="localhost")
 			const ANSICHAR* Uri;
 			bool Disconnect;
 		} StressUrls[] = {
-			{ "/data",				false },
-			{ "/data?disconnect",	true },
+			{ "/data?slowly=1",		false },
+			{ "/data?disconnect=1",	true },
 		};
 
 		uint64 Errors = 0;
@@ -3938,7 +3939,7 @@ IOSTOREONDEMAND_API void IasHttpTest(const ANSICHAR* TestHost="localhost")
 		{
 			FTicketSink Sink = ExpectDisconnect ? FTicketSink(ErrorSink) : FTicketSink(HashSink);
 
-			const auto& StressUrl = BuildUrl(StressUri, 9494);
+			const auto& StressUrl = BuildUrl(StressUri);
 			for (bool AddDelay : {false, true})
 			{
 				FTicket Tickets[StressLoad];
@@ -3994,7 +3995,7 @@ IOSTOREONDEMAND_API void IasHttpTest(const ANSICHAR* TestHost="localhost")
 	{
 		TAnsiStringBuilder<32> TamperUrl;
 		TamperUrl << "/data?tamper=" << i;
-		FAnsiStringView Url = BuildUrl(TamperUrl.ToString(), 9494);
+		FAnsiStringView Url = BuildUrl(TamperUrl.ToString());
 
 		for (int j = 0; j < 48; ++j)
 		{
diff --git a/Engine/Source/Runtime/Experimental/IoStoreOnDemand/Private/CoreHttp/TestServer.py b/Engine/Source/Runtime/Experimental/IoStoreOnDemand/Private/CoreHttp/TestServer.py
--- a/Engine/Source/Runtime/Experimental/IoStoreOnDemand/Private/CoreHttp/TestServer.py
+++ b/Engine/Source/Runtime/Experimental/IoStoreOnDemand/Private/CoreHttp/TestServer.py
@@ -1,6 +1,5 @@
 import io
 import os
-import re
 import time
 import base64
 import socket
@@ -8,12 +7,17 @@ import random
 import threading
 import http.server
 import http.client
+from pathlib import Path
+from urllib.parse import parse_qsl
 
 # {{{1 proxied .................................................................
 
 def intercept(fd):
     line = fd.readline()
-    headers = http.client.parse_headers(fd)
+    try:
+        headers = http.client.parse_headers(fd)
+    except http.client.HTTPException:
+        return 413
     return line, headers
 
 def make_preamble(line, headers):
@@ -25,7 +29,12 @@ def make_preamble(line, headers):
 
 def proxy_impl(client, httpd):
     # get request
-    line, headers = intercept(client)
+    req = intercept(client)
+    if isinstance(req, int):
+        client.write(f"HTTP/1.1 {req} IasTestServerProxyError\r\n".encode())
+        client.write(b"Content-Length: 0\r\n\r\n")
+        return False
+    line, headers = req
     if not (line or headers):
         return False
     close = (headers.get("Connection", "").lower() == "close")
@@ -33,21 +42,23 @@ def proxy_impl(client, httpd):
     httpd.write(msg)
     httpd.flush()
 
-    # establish behaviour
-    disconnect = False
-    if b"?disconnect HTTP" in line:
-        disconnect = True
+    # extract request
+    method, path, proto = line.split(b" ")
+    assert method == b"GET"
 
-    stall = False
-    if b"?stall HTTP" in line:
-        stall = disconnect = True
+    if b"?" in path:
+        _, query = path.split(b"?", 1)
+        query = {k:v for k,v in parse_qsl(query)}
+    else:
+        query = {}
 
-    tamper = 0
-    if m := re.search(b"\?tamper=(\d+) HTTP", line):
-        tamper = int(m.group(1))
-        tamper = tamper / 100.0
+    # establish behaviour
+    tamper = float(query.get(b"tamper", 0)) / 100.0
+    disconnect = b"disconnect" in query
+    stall = b"stall" in query
+    slowly = b"slowly" in query
 
-    # get repsonse
+    # get response
     line, headers = intercept(httpd)
     close = close or (headers.get("Connection", "").lower() == "close")
     content_len = int(headers["Content-Length"])
@@ -69,9 +80,7 @@ def proxy_impl(client, httpd):
             c = data[i] if random.random() > tamper else (int(random.random() * 0x4567) & 0xff)
             data[i] = c
 
-
     # retransmit
-    slowly = not (stall or tamper > 0)
     send_time = (0.75 + (random.random() * 0.75)) if slowly else 0
     while data:
         percent = 0.02 + (random.random() * 0.08)
@@ -82,18 +91,19 @@ def proxy_impl(client, httpd):
         client.write(piece)
         client.flush()
 
-        time.sleep(send_time * percent)
+        if send_time:
+            time.sleep(send_time * percent)
 
     # we're done
     if stall:
         time.sleep(2)
 
-    return not (close or disconnect)
+    return not (close or disconnect or stall)
 
-def proxy_loop(client):
+def proxy_client(client, httpd_port):
     client_fd = client.makefile("rwb")
     with socket.socket(socket.AF_INET, socket.SOCK_STREAM) as httpd:
-        httpd.connect(("127.0.0.1", 9493))
+        httpd.connect(("127.0.0.1", httpd_port))
         httpd_fd = httpd.makefile("rwb")
         try:
             while proxy_impl(client_fd, httpd_fd):
@@ -102,15 +112,15 @@ def proxy_loop(client):
             pass
     client.close()
 
-def proxy():
+def proxy_loop(httpd_port):
     with socket.socket(socket.AF_INET, socket.SOCK_STREAM) as sock:
-        sock.bind(("", 9494))
-        sock.listen()
+        sock.bind(("", 9493))
+        sock.listen(16)
         sock.setblocking(True)
         while True:
             client, address = sock.accept()
             try:
-                threading.Thread(target=proxy_loop, args=(client,), daemon=True).start()
+                threading.Thread(target=proxy_client, args=(client, httpd_port), daemon=True).start()
             except (ConnectionResetError, ConnectionAbortedError):
                 pass
 
@@ -191,6 +201,13 @@ class Handler(http.server.BaseHTTPRequestHandler):
         if query := parts[-1].split("?"):
             parts[-1] = query[0]
 
+        if parts[0] == "hello":
+            self.send_response(200)
+            self.send_header("Content-Length", 5)
+            self.end_headers()
+            self.wfile.write(b"hello")
+            return
+
         if parts[0] == "data":
             size = -1
             if len(parts) > 1:
@@ -203,18 +220,32 @@ class Handler(http.server.BaseHTTPRequestHandler):
 
         return self.send_error(404, f"not found '{self.path}'")
 
+def plain_httpd_loop(port):
+    server = http.server.ThreadingHTTPServer(("", port), Handler)
+    server.serve_forever()
+
 
 
 # {{{1 main ....................................................................
 
 def main():
-    os.chdir("c:/")
+    for item in Path(__file__).parents:
+        if (item / "GenerateProjectFiles.bat").is_file():
+            os.chdir(item)
+            break
+    else:
+        assert False
 
-    proxy_thread = threading.Thread(target=proxy, daemon=True)
-    proxy_thread.start()
+    httpd_port = int(random.random() * 0x8000) + 0x4000
 
-    server = http.server.ThreadingHTTPServer(("", 9493), Handler)
-    server.serve_forever()
+    def start_svc(target, *args):
+        threading.Thread(target=target, args=args, daemon=True).start()
+
+    start_svc(proxy_loop, httpd_port)
+    start_svc(plain_httpd_loop, httpd_port)
+
+    while True:
+        time.sleep(3600)
 
 if __name__ == "__main__":
     main()
