Object for homogenising results from IAS components

In a few places in IAS where an action can have multiple outcomes; failure,
pending, disconnect, N bytes received, etc. Some times we use integers which
works well but is not at all clear. There was a prior attempt to tidy this up
(FResult) but it was never seen through to completion. Here we go again, only
with a different name to make life a little less stressful. This submit is just
the object itself.

#rnx

diff --git a/Engine/Source/Runtime/Experimental/IoStoreOnDemand/Private/CoreHttp/Http.cpp b/Engine/Source/Runtime/Experimental/IoStoreOnDemand/Private/CoreHttp/Http.cpp
--- a/Engine/Source/Runtime/Experimental/IoStoreOnDemand/Private/CoreHttp/Http.cpp
+++ b/Engine/Source/Runtime/Experimental/IoStoreOnDemand/Private/CoreHttp/Http.cpp
@@ -208,6 +208,71 @@ private:
 };
 static_assert(sizeof(FResult) == sizeof(void*));
 
+////////////////////////////////////////////////////////////////////////////////
+class FOutcome
+{
+public:
+	static FOutcome Ok(uint32 Result=0);
+	static FOutcome Waiting();
+	static FOutcome Error(const char* Message, int32 Code=-1);
+	static FOutcome None()				{ return Error(""); }
+	bool			IsError() const		{ return Message < 0x8000'0000'0000; }
+	bool			IsWaiting() const	{ return Tag == WaitTag; }
+	bool			IsOk() const		{ return Tag == OkTag; }
+	FAnsiStringView GetMessage() const	{ check(IsError()); return (const char*)(Message); }
+	int32			GetErrorCode() const{ check(IsError()); return int32(Code); }
+	uint32			GetResult() const	{ check(IsOk()); return Result; }
+
+private:
+					FOutcome() = default;
+
+	static uint32 const OkTag	= 0x0000'8000;
+	static uint32 const WaitTag = 0x0001'8000;
+
+	union {
+		struct {
+			UPTRINT	Message : 48;
+			PTRINT	Code	: 16;
+		};
+		struct {
+			uint32	Result;
+			uint32	Tag;
+		};
+	};
+};
+static_assert(sizeof(FOutcome) == sizeof(void*));
+
+////////////////////////////////////////////////////////////////////////////////
+FOutcome FOutcome::Ok(uint32 Result)
+{
+	FOutcome Outcome;
+	Outcome.Tag = OkTag;
+	Outcome.Result = Result;
+	return Outcome;
+}
+
+////////////////////////////////////////////////////////////////////////////////
+FOutcome FOutcome::Waiting()
+{
+	FOutcome Outcome;
+	Outcome.Tag = WaitTag;
+	return Outcome;
+}
+
+////////////////////////////////////////////////////////////////////////////////
+FOutcome FOutcome::Error(const char* Message, int32 Code)
+{
+	check(Message != nullptr);
+	check(Code <= 0xffff && Code >= -0xffff);
+
+	FOutcome Outcome;
+	Outcome.Message = UPTRINT(Message);
+	Outcome.Code = int16(Code);
+	return Outcome;
+}
+
+
+
 ////////////////////////////////////////////////////////////////////////////////
 template <typename LambdaType>
 static void EnumerateHeaders(FAnsiStringView Headers, LambdaType&& Lambda)
@@ -3513,6 +3578,22 @@ static void MiscTest()
 	check(UrlOut.Path == 25);
 #undef CRLF
 
+	static const char* OutcomeMsg = "\x4d\x52";
+	check(FOutcome::Error(OutcomeMsg, -5).IsOk() == false);
+	check(FOutcome::Error(OutcomeMsg, -5).IsWaiting() == false);
+	check(FOutcome::Error(OutcomeMsg, -5).IsError());
+	check(FOutcome::Error(OutcomeMsg, -5).GetErrorCode() == -5);
+	check(FOutcome::Error(OutcomeMsg,  5).GetErrorCode() ==  5);
+	check(FOutcome::Error(OutcomeMsg, -5).GetMessage() == OutcomeMsg);
+
+	check(FOutcome::Ok(  0).IsOk());
+	check(FOutcome::Ok(-13).IsWaiting() == false);
+	check(FOutcome::Ok( 13).IsError() == false);
+
+	check(FOutcome::Waiting().IsOk() == false);
+	check(FOutcome::Waiting().IsWaiting());
+	check(FOutcome::Waiting().IsError() == false);
+
 	check(FResult(-5).GetValue()		== -5);
 	check(FResult(-1).GetValue()		== -1);
 	check(FResult( 0).GetValue()		==  0);
