Properly reserve space in arrays when resolving the distributed endpoint.

#rnx

diff --git a/Engine/Source/Runtime/Experimental/IoStoreOnDemand/Private/DistributionEndpoints.cpp b/Engine/Source/Runtime/Experimental/IoStoreOnDemand/Private/DistributionEndpoints.cpp
--- a/Engine/Source/Runtime/Experimental/IoStoreOnDemand/Private/DistributionEndpoints.cpp
+++ b/Engine/Source/Runtime/Experimental/IoStoreOnDemand/Private/DistributionEndpoints.cpp
@@ -100,6 +100,8 @@ FDistributionEndpoints::EResult FDistributionEndpoints::ParseResponse(FHttpRespo
 			}
 
 			TArray<FJsonValuePtr> JsonValues = JsonObj->GetArrayField(TEXT("distributions"));
+			OutUrls.Reserve(JsonValues.Num());
+
 			for (const FJsonValuePtr& JsonValue : JsonValues)
 			{
 				FString ServiceUrl = JsonValue->AsString();
diff --git a/Engine/Source/Runtime/Experimental/IoStoreOnDemand/Private/OnDemandIoDispatcherBackend.cpp b/Engine/Source/Runtime/Experimental/IoStoreOnDemand/Private/OnDemandIoDispatcherBackend.cpp
--- a/Engine/Source/Runtime/Experimental/IoStoreOnDemand/Private/OnDemandIoDispatcherBackend.cpp
+++ b/Engine/Source/Runtime/Experimental/IoStoreOnDemand/Private/OnDemandIoDispatcherBackend.cpp
@@ -2016,6 +2016,8 @@ bool FOnDemandIoBackend::ResolveDistributedEndpoint(const FDistributedEndpointUr
 		if (Result == FDistributionEndpoints::EResult::Success)
 		{
 			FWriteScopeLock _(Lock);
+			AvailableEps.Urls.Reserve(AvailableEps.Urls.Num() + ServiceUrls.Num());
+
 			for (const FString& Url : ServiceUrls)
 			{
 				AvailableEps.Urls.Add(Url.Replace(TEXT("https"), TEXT("http")));
