Conforming local variable type

#rnx

diff --git a/Engine/Source/Runtime/Experimental/IoStoreOnDemand/Private/CoreHttp/Http.cpp b/Engine/Source/Runtime/Experimental/IoStoreOnDemand/Private/CoreHttp/Http.cpp
--- a/Engine/Source/Runtime/Experimental/IoStoreOnDemand/Private/CoreHttp/Http.cpp
+++ b/Engine/Source/Runtime/Experimental/IoStoreOnDemand/Private/CoreHttp/Http.cpp
@@ -882,7 +882,7 @@ bool FSocket::Connect(uint32 IpAddress, uint32 Port)
 	AddrInet.sin_port = htons(uint16(Port));
 	memcpy(&(AddrInet.sin_addr), &IpAddress, sizeof(IpAddress));
 
-	int Result = connect(Socket, &(sockaddr&)AddrInet, sizeof(AddrInet));
+	int32 Result = connect(Socket, &(sockaddr&)AddrInet, sizeof(AddrInet));
 
 	if (IsSocketResult(EWOULDBLOCK) | IsSocketResult(EINPROGRESS))
 	{
