Fix for ReplicationSystemTest compile errors.

#rnx

diff --git a/Engine/Plugins/Runtime/ReplicationSystemTestPlugin/Source/Private/Tests/ReplicationSystem/Filtering/MockNetObjectFilter.cpp b/Engine/Plugins/Runtime/ReplicationSystemTestPlugin/Source/Private/Tests/ReplicationSystem/Filtering/MockNetObjectFilter.cpp
--- a/Engine/Plugins/Runtime/ReplicationSystemTestPlugin/Source/Private/Tests/ReplicationSystem/Filtering/MockNetObjectFilter.cpp
+++ b/Engine/Plugins/Runtime/ReplicationSystemTestPlugin/Source/Private/Tests/ReplicationSystem/Filtering/MockNetObjectFilter.cpp
@@ -191,7 +191,7 @@ void UMockNetObjectFilterUsingFragmentData::UpdateObjects(FNetObjectFilterUpdate
 		if (UPTRINT* InternalStateOffset = ObjectToFilterOutOffset.Find(ObjectIndex))
 		{
 			// Warning: This is generally not safe, not even for primitive types.
-			const bool bFilterOut = *reinterpret_cast<const bool*>(Params.StateBuffers[ObjectIndex] + *InternalStateOffset);
+			const bool bFilterOut = *reinterpret_cast<const bool*>((*Params.StateBuffers)[ObjectIndex] + *InternalStateOffset);
 			ObjectToFilterOut.Emplace(ObjectIndex, bFilterOut);
 		}
 	}
diff --git a/Engine/Plugins/Runtime/ReplicationSystemTestPlugin/Source/Private/Tests/ReplicationSystem/Prioritization/MockNetObjectPrioritizer.cpp b/Engine/Plugins/Runtime/ReplicationSystemTestPlugin/Source/Private/Tests/ReplicationSystem/Prioritization/MockNetObjectPrioritizer.cpp
--- a/Engine/Plugins/Runtime/ReplicationSystemTestPlugin/Source/Private/Tests/ReplicationSystem/Prioritization/MockNetObjectPrioritizer.cpp
+++ b/Engine/Plugins/Runtime/ReplicationSystemTestPlugin/Source/Private/Tests/ReplicationSystem/Prioritization/MockNetObjectPrioritizer.cpp
@@ -94,7 +94,7 @@ void UMockNetObjectPrioritizer::UpdateObjects(FNetObjectPrioritizerUpdateParams&
 		if (UPTRINT* InternalStateOffset = ObjectToPriorityOffset.Find(ObjectIndex))
 		{
 			// Warning: This is generally not safe, not even for primitive types.
-			const float Priority = *reinterpret_cast<const float*>(Params.StateBuffers[ObjectIndex] + *InternalStateOffset);
+			const float Priority = *reinterpret_cast<const float*>((*Params.StateBuffers)[ObjectIndex] + *InternalStateOffset);
 			ObjectToPriority.Emplace(ObjectIndex, Priority);
 		}
 	}
diff --git a/Engine/Source/Runtime/Experimental/Iris/Core/Private/Iris/ReplicationSystem/Filtering/ReplicationFiltering.cpp b/Engine/Source/Runtime/Experimental/Iris/Core/Private/Iris/ReplicationSystem/Filtering/ReplicationFiltering.cpp
--- a/Engine/Source/Runtime/Experimental/Iris/Core/Private/Iris/ReplicationSystem/Filtering/ReplicationFiltering.cpp
+++ b/Engine/Source/Runtime/Experimental/Iris/Core/Private/Iris/ReplicationSystem/Filtering/ReplicationFiltering.cpp
@@ -1319,6 +1319,7 @@ void FReplicationFiltering::UpdateDynamicFiltering(ENetFilterType FilterType)
 	const uint32* SubObjectsData = NetRefHandleManager->GetSubObjectInternalIndices().GetData();
 	const uint32* DependentObjectsData = NetRefHandleManager->GetDependentObjectInternalIndices().GetData();
 	const uint32* ObjectsRequiringDynamicFilterUpdateData = ObjectsRequiringDynamicFilterUpdate.GetData();
+	const TNetChunkedArray<uint8*>* ObjectsStateBuffers = FilterType == ENetFilterType::PostPoll_FragmentBased ? &NetRefHandleManager->GetReplicatedObjectStateBuffers() : nullptr;
 
 	uint32* ConnectionIds = static_cast<uint32*>(FMemory_Alloca(ValidConnections.GetNumBits() * sizeof(uint32)));
 	uint32 ConnectionCount = 0;
@@ -1357,6 +1358,7 @@ void FReplicationFiltering::UpdateDynamicFiltering(ENetFilterType FilterType)
 				FNetObjectFilteringParams FilteringParams(MakeNetBitArrayView(Info.FilteredObjects));
 				FilteringParams.OutAllowedObjects = AllowedObjects;
 				FilteringParams.FilteringInfos = NetObjectFilteringInfos.GetData();
+				FilteringParams.StateBuffers = ObjectsStateBuffers;
 				FilteringParams.ConnectionId = ConnId;
 				FilteringParams.View = Connections->GetReplicationView(ConnId);
 
@@ -1541,6 +1543,11 @@ void FReplicationFiltering::BatchNotifyFiltersOfDirtyObjects(FUpdateDirtyObjects
 	FNetObjectFilterUpdateParams UpdateParameters;
 	UpdateParameters.FilteringInfos = NetObjectFilteringInfos.GetData();
 
+	if (BatchHelper.HasProtocolBuffers())
+	{
+		UpdateParameters.StateBuffers = &NetRefHandleManager->GetReplicatedObjectStateBuffers();
+	}
+
 	// $IRIS TODO: We should probably have a trait asking if the Filter needs to receive UpdateObjects.
 	for (const FUpdateDirtyObjectsBatchHelper::FPerFilterInfo& PerFilterInfo : BatchHelper.PerFilterInfos)
 	{
diff --git a/Engine/Source/Runtime/Experimental/Iris/Core/Private/Iris/ReplicationSystem/Prioritization/ReplicationPrioritization.cpp b/Engine/Source/Runtime/Experimental/Iris/Core/Private/Iris/ReplicationSystem/Prioritization/ReplicationPrioritization.cpp
--- a/Engine/Source/Runtime/Experimental/Iris/Core/Private/Iris/ReplicationSystem/Prioritization/ReplicationPrioritization.cpp
+++ b/Engine/Source/Runtime/Experimental/Iris/Core/Private/Iris/ReplicationSystem/Prioritization/ReplicationPrioritization.cpp
@@ -801,6 +801,7 @@ void FReplicationPrioritization::BatchNotifyPrioritizersOfDirtyObjects(FUpdateDi
 	BatchHelper.PrepareBatch(ObjectIndices, ObjectCount, ObjectIndexToPrioritizer.GetData());
 
 	FNetObjectPrioritizerUpdateParams UpdateParameters;
+	UpdateParameters.StateBuffers = &NetRefHandleManager->GetReplicatedObjectStateBuffers();
 	UpdateParameters.PrioritizationInfos = NetObjectPrioritizationInfos.GetData();
 
 	for (const FUpdateDirtyObjectsBatchHelper::FPerPrioritizerInfo& PerPrioritizerInfo : BatchHelper.PerPrioritizerInfos)
diff --git a/Engine/Source/Runtime/Experimental/Iris/Core/Public/Iris/ReplicationSystem/Filtering/NetObjectFilter.h b/Engine/Source/Runtime/Experimental/Iris/Core/Public/Iris/ReplicationSystem/Filtering/NetObjectFilter.h
--- a/Engine/Source/Runtime/Experimental/Iris/Core/Public/Iris/ReplicationSystem/Filtering/NetObjectFilter.h
+++ b/Engine/Source/Runtime/Experimental/Iris/Core/Public/Iris/ReplicationSystem/Filtering/NetObjectFilter.h
@@ -3,6 +3,7 @@
 #pragma once
 
 #include "Net/Core/NetBitArray.h"
+#include "Iris/Core/NetChunkedArray.h"
 #include "Iris/ReplicationSystem/ReplicationView.h"
 #include "UObject/ObjectMacros.h"
 #include "NetObjectFilter.generated.h"
@@ -56,6 +57,9 @@ struct FNetObjectFilteringParams
 	/** FilteringInfos for all objects. Index using the set bit indices in FilteredObjects. */
 	const FNetObjectFilteringInfo* FilteringInfos = nullptr;
 
+	/** State buffers for all objects. Index using the set bit indices in FilteredObjects. */
+	const UE::Net::TNetChunkedArray<uint8*>* StateBuffers = nullptr;
+
 	/** ID of the connection that the filtering applies to. */
 	uint32 ConnectionId;
 
@@ -193,6 +197,12 @@ struct FNetObjectFilterUpdateParams
 	* NOTE: Only for filters of type FragmentBased; null for Raw types
 	*/
 	UE::Net::FReplicationInstanceProtocol const* const* InstanceProtocols = nullptr;
+
+	/**
+	* State buffers for all objects. Index using ObjectIndices[0..ObjectCount-1].
+	* NOTE: Only for filters of type FragmentBased; null for Raw types
+	*/
+	const UE::Net::TNetChunkedArray<uint8*>* StateBuffers = nullptr;
 };
 
 UCLASS(Abstract)
diff --git a/Engine/Source/Runtime/Experimental/Iris/Core/Public/Iris/ReplicationSystem/Prioritization/NetObjectPrioritizer.h b/Engine/Source/Runtime/Experimental/Iris/Core/Public/Iris/ReplicationSystem/Prioritization/NetObjectPrioritizer.h
--- a/Engine/Source/Runtime/Experimental/Iris/Core/Public/Iris/ReplicationSystem/Prioritization/NetObjectPrioritizer.h
+++ b/Engine/Source/Runtime/Experimental/Iris/Core/Public/Iris/ReplicationSystem/Prioritization/NetObjectPrioritizer.h
@@ -2,6 +2,7 @@
 
 #pragma once
 
+#include "Iris/Core/NetChunkedArray.h"
 #include "Iris/ReplicationSystem/ReplicationView.h"
 #include "UObject/ObjectMacros.h"
 #include "NetObjectPrioritizer.generated.h"
@@ -121,6 +122,9 @@ struct FNetObjectPrioritizerUpdateParams
 	/** InstanceProtocols for updated objects. Index using 0..ObjectCount-1. */
 	UE::Net::FReplicationInstanceProtocol const*const* InstanceProtocols;
 
+	/** State buffers for all objects. Index using ObjectIndices[0..ObjectCount-1]. */
+	const UE::Net::TNetChunkedArray<uint8*>* StateBuffers = nullptr;
+
 	/** Infos for all objects. Index using ObjectIndices[0..ObjectCount-1]. */
 	FNetObjectPrioritizationInfo* PrioritizationInfos;
 };
