Chaos - particle correction supports moving joint-connected bodies

diff --git a/Engine/Source/Runtime/Experimental/Chaos/Private/Chaos/PBDRigidsEvolutionGBF.cpp b/Engine/Source/Runtime/Experimental/Chaos/Private/Chaos/PBDRigidsEvolutionGBF.cpp
--- a/Engine/Source/Runtime/Experimental/Chaos/Private/Chaos/PBDRigidsEvolutionGBF.cpp
+++ b/Engine/Source/Runtime/Experimental/Chaos/Private/Chaos/PBDRigidsEvolutionGBF.cpp
@@ -3,6 +3,7 @@
 #include "Chaos/Collision/SimSweep.h"
 #include "Chaos/Defines.h"
 #include "Chaos/Evolution/SolverBodyContainer.h"
+#include "Chaos/Framework/HashMappedArray.h"
 #include "Chaos/Framework/Parallel.h"
 #include "Chaos/ImplicitObjectTransformed.h"
 #include "Chaos/ImplicitObjectUnion.h"
@@ -1411,41 +1412,6 @@ void FPBDRigidsEvolutionGBF::SetParticleKinematicTarget(FGeometryParticleHandle*
 	}
 }
 
-void FPBDRigidsEvolutionGBF::SetParticleTransformCorrection(FGeometryParticleHandle* InParticle, const FVec3& InPosDelta, const FVec3& InRotDelta)
-{
-	FGenericParticleHandle Particle = InParticle;
-	const bool bHavePositionDelta = !InPosDelta.IsZero();
-	const bool bHaveRotationDelta = !InRotDelta.IsZero();
-
-	if (bHavePositionDelta || bHaveRotationDelta)
-	{
-		const FVec3 PrevX = Particle->GetX();
-		const FRotation3 PrevR = Particle->GetR();
-
-		// Apply the position correction
-		if (bHavePositionDelta)
-		{
-			const FVec3 NewX = Particle->GetX() + InPosDelta;
-			Particle->SetX(NewX);
-		}
-
-		// Apply the rotation correction
-		if (bHaveRotationDelta)
-		{
-			const FRotation3 NewR = FRotation3::IntegrateRotationWithAngularVelocity(Particle->GetR(), InRotDelta, 1.0f);
-			Particle->SetR(NewR);
-		}
-
-		// If we applied a correction we must fix the collision anchors so that friction doesn't undo prevent our move or rotation
-		InParticle->ParticleCollisions().VisitCollisions(
-			[this, InParticle](FPBDCollisionConstraint& Collision)
-			{
-				Collision.UpdateParticleTransform(InParticle);
-				return ECollisionVisitorResult::Continue;
-			});
-	}
-}
-
 void FPBDRigidsEvolutionGBF::OnParticleMoved(FGeometryParticleHandle* InParticle, const FVec3& PrevX, const FRotation3& PrevR, const bool bIsTeleport)
 {
 	// When a particle is moved, we need to 
@@ -1491,6 +1457,112 @@ void FPBDRigidsEvolutionGBF::OnParticleMoved(FGeometryParticleHandle* InParticle
 	}
 }
 
+void FPBDRigidsEvolutionGBF::ApplyParticleTransformCorrectionDelta(FGeometryParticleHandle* InParticle, const FVec3& InPosDelta, const FVec3& InRotDelta, const bool bApplyToConnectedBodies)
+{
+	ApplyParticleTransformCorrection(
+		InParticle, 
+		InParticle->GetX() + InPosDelta, 
+		FRotation3::IntegrateRotationWithAngularVelocity(InParticle->GetR(), InRotDelta, FReal(1.0)),
+		bApplyToConnectedBodies);
+}
+
+void FPBDRigidsEvolutionGBF::ApplyParticleTransformCorrection(FGeometryParticleHandle* InParticle, const FVec3& InPos, const FRotation3& InRot, const bool bApplyToConnectedBodies)
+{
+	const FRigidTransform3 OldParticleTransform = InParticle->GetTransformXR();
+	const FRigidTransform3 NewParticleTransform = FRigidTransform3(InPos, InRot);
+
+	// Move the root particle
+	ApplyParticleTransformCorrectionImpl(InParticle, NewParticleTransform);
+
+	if (bApplyToConnectedBodies)
+	{
+		// Find all the connected particles and move them to retain their relative transform
+		// NOTE: ConnectedParticles will not include InParticle
+		TArray<FGeometryParticleHandle*> ConnectedParticles = GetConnectedParticles(InParticle);
+		for (FGeometryParticleHandle* ConnectedParticle : ConnectedParticles)
+		{
+			if (FGenericParticleHandle(ConnectedParticle)->IsDynamic())
+			{
+				const FRigidTransform3 RelativeTransform = ConnectedParticle->GetTransformXR().GetRelativeTransformNoScale(OldParticleTransform);
+				const FRigidTransform3 NewOtherParticleTransform = RelativeTransform * NewParticleTransform;
+				ApplyParticleTransformCorrectionImpl(ConnectedParticle, NewOtherParticleTransform);
+			}
+		}
+	}
+}
+
+void FPBDRigidsEvolutionGBF::ApplyParticleTransformCorrectionImpl(FGeometryParticleHandle* InParticle, const FRigidTransform3& InTransform)
+{
+	FGenericParticleHandle Particle = InParticle;
+
+	Particle->SetX(InTransform.GetTranslation());
+	Particle->SetP(InTransform.GetTranslation());
+	Particle->SetR(InTransform.GetRotation());
+	Particle->SetQ(InTransform.GetRotation());
+
+	// We must fix the collision anchors so that friction doesn't undo our move or rotation
+	InParticle->ParticleCollisions().VisitCollisions(
+		[this, InParticle](FPBDCollisionConstraint& Collision)
+		{
+			Collision.UpdateParticleTransform(InParticle);
+			return ECollisionVisitorResult::Continue;
+		});
+}
+
+TArray<FGeometryParticleHandle*> FPBDRigidsEvolutionGBF::GetConnectedParticles(FGeometryParticleHandle* InParticle)
+{
+	if (InParticle->ParticleConstraints().IsEmpty())
+	{
+		return {};
+	}
+
+	// Use a hashmap to prevent O(N^2) loop below. It will also contains the output list of connected particles.
+	struct FHashMapTraits
+	{
+		static uint32 GetIDHash(const FGeometryParticleHandle* Particle) { return MurmurFinalize32(uint32(Particle->UniqueIdx().Idx)); }
+		static bool ElementHasID(const FGeometryParticleHandle* A, const FGeometryParticleHandle* B) { return A == B; }
+	};
+	Private::THashMappedArray<FGeometryParticleHandle*, FGeometryParticleHandle*, FHashMapTraits> ConnectedParticles(256);
+
+	FGeometryParticleHandle* NextParticle = InParticle;
+	int32 NextParticleIndex = 0;
+	while (true)
+	{
+		// Loop over the joints on the next particle and add the other particle into the queue (if not already in the queue, or already processed)
+		for (FConstraintHandle* Constraint : NextParticle->ParticleConstraints())
+		{
+			if (FPBDJointConstraintHandle* Joint = Constraint->As<FPBDJointConstraintHandle>())
+			{
+				const TVec3<EJointMotionType>& JointLinearMotion = Joint->GetSettings().LinearMotionTypes;
+				if ((JointLinearMotion[0] == EJointMotionType::Locked) && (JointLinearMotion[1] == EJointMotionType::Locked) && (JointLinearMotion[2] == EJointMotionType::Locked))
+				{
+					FParticlePair JointParticles = Joint->GetConstrainedParticles();
+					FGeometryParticleHandle* OtherParticle = (JointParticles[0] != NextParticle) ? JointParticles[0] : JointParticles[1];
+					if ((OtherParticle != InParticle) && (ConnectedParticles.Find(OtherParticle) == nullptr))
+					{
+						ConnectedParticles.Add(OtherParticle, OtherParticle);
+					}
+				}
+			}
+		}
+
+		// We are done if we did not add any particles in the most recent loop above
+		check(NextParticleIndex <= ConnectedParticles.Num());
+		if (NextParticleIndex >= ConnectedParticles.Num())
+		{
+			break;
+		}
+
+		// Move to the next particle in the queue
+		NextParticle = ConnectedParticles.At(NextParticleIndex);
+		++NextParticleIndex;
+	}
+	check(NextParticleIndex == ConnectedParticles.Num());
+
+	// NOTE: This moves the internal array to the output - no duplication
+	return ConnectedParticles.ExtractElements();
+}
+
 void FPBDRigidsEvolutionGBF::SetParticleVelocities(FGeometryParticleHandle* InParticle, const FVec3& InV, const FVec3f& InW)
 {
 	if (FKinematicGeometryParticleHandle* Kinematic = InParticle->CastToKinematicParticle())
diff --git a/Engine/Source/Runtime/Experimental/Chaos/Public/Chaos/Collision/MeshContactGenerator.h b/Engine/Source/Runtime/Experimental/Chaos/Public/Chaos/Collision/MeshContactGenerator.h
--- a/Engine/Source/Runtime/Experimental/Chaos/Public/Chaos/Collision/MeshContactGenerator.h
+++ b/Engine/Source/Runtime/Experimental/Chaos/Public/Chaos/Collision/MeshContactGenerator.h
@@ -11,17 +11,6 @@
 
 namespace Chaos::Private
 {
-	// Adapter for FContactEdgeID use in THashMappedArray
-	template<> struct THashMappedArrayIDTraits<FContactEdgeID>
-	{
-		using FIDType = FContactEdgeID;
-
-		static uint32 HashID(const FIDType& ID)
-		{
-			return uint32(MurmurFinalize64(ID.EdgeID));
-		}
-	};
-
 	class FMeshContactGeneratorSettings
 	{
 	public:
@@ -487,10 +476,20 @@ namespace Chaos::Private
 		TArray<FTriangleContactPointData> ContactDatas;
 
 		// A map of EdgeID to the two triangles (indices) that use the Edge
-		THashMappedArray<FContactEdgeID, FEdgeTriangleIndices> EdgeTriangleIndicesMap;
+		struct FEdgeTriangleIndicesMapTraits
+		{
+			static uint32 GetIDHash(const FContactEdgeID& EdgeID) { return uint32(MurmurFinalize64(EdgeID.EdgeID)); }
+			static bool ElementHasID(const FEdgeTriangleIndices& TriangleIndices, const FContactEdgeID& EdgeID) { return TriangleIndices.ID == EdgeID; }
+		};
+		THashMappedArray<FContactEdgeID, FEdgeTriangleIndices, FEdgeTriangleIndicesMapTraits> EdgeTriangleIndicesMap;
 
 		// A map of VertexID to contact index on that vertex - we only ever keep one contact per vertex
-		THashMappedArray<FContactVertexID, FVertexContactIndex> VertexContactIndicesMap;
+		struct VertexContactIndicesMapTraits
+		{
+			static uint32 GetIDHash(const FContactVertexID& VertexID) { return MurmurFinalize32(VertexID); }
+			static bool ElementHasID(const FVertexContactIndex& VertexIndex, const FContactVertexID& VertexID) { return VertexIndex.ID == VertexID; }
+		};
+		THashMappedArray<FContactVertexID, FVertexContactIndex, VertexContactIndicesMapTraits> VertexContactIndicesMap;
 	};
 
 }
\ No newline at end of file
diff --git a/Engine/Source/Runtime/Experimental/Chaos/Public/Chaos/Framework/HashMappedArray.h b/Engine/Source/Runtime/Experimental/Chaos/Public/Chaos/Framework/HashMappedArray.h
--- a/Engine/Source/Runtime/Experimental/Chaos/Public/Chaos/Framework/HashMappedArray.h
+++ b/Engine/Source/Runtime/Experimental/Chaos/Public/Chaos/Framework/HashMappedArray.h
@@ -5,30 +5,24 @@
 
 namespace Chaos::Private
 {
-	// Default ID traits for use with THashMappedArray that works for all types that can be cast to int32
-	template<typename TIDType>
-	struct THashMappedArrayIDTraits
+	// Default traits for THashMappedArray that works for all ID/Element pairs where the 
+	// ID has a MurmurFinalize32 implementation and we can compare equality of Elements and IDs.
+	template<typename TIDType, typename TElementType>
+	struct THashMappedArrayTraits
 	{
 		using FIDType = TIDType;
+		using FElementType = TElementType;
 
 		// Hash the ID to a 32 bit unsigned int for use with FHashTable
-		static uint32 HashID(const FIDType& ID)
+		static uint32 GetIDHash(const FIDType& ID)
 		{
-			return MurmurFinalize32(uint32(ID));
+			return MurmurFinalize32(ID);
 		}
-	};
-
-	// Default Element traits for THashMappedArray that works for all types that have an member variable ID of type TIDType
-	template<typename TElementType, typename TIDType>
-	struct THashMappedArrayElementTraits
-	{
-		using FIDType = TIDType;
-		using FElementType = TElementType;
 
 		// Return true if the element is the one with the specified ID
-		static bool Is(const FElementType& Element, const FIDType& ID)
+		static bool ElemenatHasID(const FElementType& Element, const FIDType& ID)
 		{
-			return Element.ID == ID;
+			return Element == ID;
 		}
 	};
 
@@ -39,7 +33,7 @@ namespace Chaos::Private
 	*	using FMyDataID = int32;
 	*	struct FMyData
 	*	{
-	*		FMyDataID ID;	// Every FMyData will require a unique ID if using the default THashMappedArrayElementTraits
+	*		FMyDataID ID;	// Every FMyData will require a unique ID if using the default THashMappedArrayTraits
 	*		float MyValue;
 	*	};
 	* 
@@ -52,16 +46,15 @@ namespace Chaos::Private
 	*	const FMyData* MyData2 = MyDataMap.Find(2);		// MyData2->MyValue == 2.0
 	* 
 	*/
-	template<typename TIDType, typename TElementType>
+	template<typename TIDType, typename TElementType, typename TTraits = THashMappedArrayTraits<TIDType, TElementType>>
 	class THashMappedArray
 	{
 	public:
 		using FIDType = TIDType;
 		using FElementType = TElementType;
-		using FIDTraits = THashMappedArrayIDTraits<FIDType>;
-		using FElementTraits = THashMappedArrayElementTraits<FElementType, FIDType>;
+		using FTraits = TTraits;
 		using FHashType = uint32;
-		using FType = THashMappedArray<FIDType, FElementType>;
+		using FType = THashMappedArray<FIDType, FElementType, FTraits>;
 
 		// Initialize the hash table. InHashSize must be a power of two (asserted)
 		THashMappedArray(const int32 InHashSize)
@@ -83,7 +76,7 @@ namespace Chaos::Private
 			checkSlow(Find(ID) == nullptr);
 
 			const int32 Index = Elements.Add(Element);
-			const FHashType Key = HashMapKey(ID);
+			const FHashType Key = FTraits::GetIDHash(ID);
 
 			HashTable.Add(Key, Index);
 		}
@@ -97,7 +90,7 @@ namespace Chaos::Private
 			checkSlow(Find(ID) == nullptr);
 
 			const int32 Index = Elements.Emplace(Forward<ArgsType>(Args)...);
-			const FHashType Key = FIDTraits::HashID(ID);
+			const FHashType Key = FTraits::GetIDHash(ID);
 
 			HashTable.Add(Key, Index);
 		}
@@ -111,10 +104,10 @@ namespace Chaos::Private
 		// Find the element with the specified ID. Roughly O(Max(1,N/M)) for N elements with a hash table of size M
 		FElementType* Find(const FIDType ID)
 		{
-			const FHashType Key = FIDTraits::HashID(ID);
+			const FHashType Key = FTraits::GetIDHash(ID);
 			for (uint32 Index = HashTable.First(Key); HashTable.IsValid(Index); Index = HashTable.Next(Index))
 			{
-				if (FElementTraits::Is(Elements[Index], ID))
+				if (FTraits::ElementHasID(Elements[Index], ID))
 				{
 					return &Elements[Index];
 				}
@@ -122,6 +115,32 @@ namespace Chaos::Private
 			return nullptr;
 		}
 
+		// The number of elements that have been added to the map
+		int32 Num() const
+		{
+			return Elements.Num();
+		}
+
+		// Get the element at ElementIndex (indexed by order in which they were added)
+		FElementType& At(const int32 ElementIndex)
+		{
+			return Elements[ElementIndex];
+		}
+
+		// Get the element at ElementIndex (indexed by order in which they were added)
+		const FElementType& At(const int32 ElementIndex) const
+		{
+			return Elements[ElementIndex];
+		}
+
+		// Move the array elements into an external array and reset
+		TArray<FElementType> ExtractElements()
+		{
+			HashTable.Clear();
+			return MoveTemp(Elements);
+		}
+
+
 	private:
 		FHashTable HashTable;
 		TArray<FElementType> Elements;
diff --git a/Engine/Source/Runtime/Experimental/Chaos/Public/Chaos/PBDRigidsEvolutionGBF.h b/Engine/Source/Runtime/Experimental/Chaos/Public/Chaos/PBDRigidsEvolutionGBF.h
--- a/Engine/Source/Runtime/Experimental/Chaos/Public/Chaos/PBDRigidsEvolutionGBF.h
+++ b/Engine/Source/Runtime/Experimental/Chaos/Public/Chaos/PBDRigidsEvolutionGBF.h
@@ -172,11 +172,19 @@ namespace Chaos
 		CHAOS_API void SetParticleKinematicTarget(FGeometryParticleHandle* ParticleHandle, const FKinematicTarget& NewKinematicTarget);
 
 		/*
-		 * Apply a momentumless correction to the particle transform, usually as a result of a server correction.
+		 * [EXPERIMENTAL] Apply a momentumless correction to the particle transform, usually as a result of a server correction.
 		 * This will shift the particle by the supplied delta and handle updating of friction anchors or anything else that might prevent or undo the shift.
+		 * If bApplyToConnectedBodies is true, any particle attached by a joint with locked linear limits will also get moved.
 		 * NOTE: must be called prior to Integrate() to be effective.
+		 * NOTE: be careful with bApplyToConnectedBodies - only one particle in the connected graph should have ApplyParticleTransformCorrectionDelta called on 
+		 * it, otherwise you will get multiple particles trying to recorrect each other leading to very strange behaviour.
 		 */
-		CHAOS_API void SetParticleTransformCorrection(FGeometryParticleHandle* InParticle, const FVec3& InPosDelta, const FVec3& InRotDelta);
+		CHAOS_API void ApplyParticleTransformCorrectionDelta(FGeometryParticleHandle* InParticle, const FVec3& InPosDelta, const FVec3& InRotDelta, const bool bApplyToConnectedBodies);
+
+		/*
+		 * [EXPERIMENTAL] Similar to SetParticleTransformCorrectionDelta, but supplied an absolute transform to jump to. This is used for snaps.
+		 */
+		CHAOS_API void ApplyParticleTransformCorrection(FGeometryParticleHandle* InParticle, const FVec3& InPos, const FRotation3& InRot, const bool bApplyToConnectedBodies);
 
 		/**
 		 * Called when a particle is moved. We need to reset some friction properties, sleeping properties, etc
@@ -258,6 +266,12 @@ namespace Chaos
 
 		CHAOS_API void AdvanceOneTimeStepImpl(const FReal dt, const FSubStepInfo& SubStepInfo);
 
+		// Update the particle transform and fix collision anchors (used by client corrections)
+		void ApplyParticleTransformCorrectionImpl(FGeometryParticleHandle* InParticle, const FRigidTransform3& InTransform);
+
+		// Get all the particles that are connected to InParticle by a joint with locked position limits
+		TArray<FGeometryParticleHandle*> GetConnectedParticles(FGeometryParticleHandle* InParticle);
+
 		void UpdateInertiaConditioning();
 
 		FEvolutionResimCache* GetCurrentStepResimCache()
