Fix CIS

diff --git a/Engine/Source/Runtime/Experimental/ChaosVehicles/ChaosVehiclesCore/Private/SimModule/SimulationModuleBase.cpp b/Engine/Source/Runtime/Experimental/ChaosVehicles/ChaosVehiclesCore/Private/SimModule/SimulationModuleBase.cpp
--- a/Engine/Source/Runtime/Experimental/ChaosVehicles/ChaosVehiclesCore/Private/SimModule/SimulationModuleBase.cpp
+++ b/Engine/Source/Runtime/Experimental/ChaosVehicles/ChaosVehiclesCore/Private/SimModule/SimulationModuleBase.cpp
@@ -123,7 +123,10 @@ const FTransform& ISimulationModuleBase::GetParentRelativeTransform() const
 
 void FSimOutputData::FillOutputState(const ISimulationModuleBase* SimModule)
 {
-	AnimationSetupIndex = SimModule->AnimationSetupIndex;
+	if (SimModule)
+	{
+		AnimationSetupIndex = SimModule->AnimationSetupIndex;
+	}
 
 #if !(UE_BUILD_SHIPPING || UE_BUILD_TEST)
 	if (SimModule)
