Iris
- Enabled pruning of stale objects in shipping builds

diff --git a/Engine/Source/Runtime/Experimental/Iris/Core/Private/Iris/ReplicationSystem/ObjectReplicationBridge.cpp b/Engine/Source/Runtime/Experimental/Iris/Core/Private/Iris/ReplicationSystem/ObjectReplicationBridge.cpp
--- a/Engine/Source/Runtime/Experimental/Iris/Core/Private/Iris/ReplicationSystem/ObjectReplicationBridge.cpp
+++ b/Engine/Source/Runtime/Experimental/Iris/Core/Private/Iris/ReplicationSystem/ObjectReplicationBridge.cpp
@@ -835,9 +835,10 @@ void UObjectReplicationBridge::PruneStaleObjects()
 				UE_LOG(LogIrisBridge, Warning, TEXT("UObjectReplicationBridge::PruneStaleObjects ObjectInstance replicated as: %s of Type named:%s has been destroyed without notifying the ReplicationSystem %u"), *ObjectData.RefHandle.ToString(), ToCStr(DebugName), ObjectData.RefHandle.GetReplicationSystemId());
 
 				// If the instance protocol is bound, then this is an error and we cannot safely cleanup as unbinding abound instance protocol will modify bound states
-				checkf(!EnumHasAnyFlags(ObjectData.InstanceProtocol->InstanceTraits, EReplicationInstanceProtocolTraits::IsBound), TEXT("UObjectReplicationBridge::PruneStaleObjects Bound ObjectInstance replicated as: %s has been destroyed without notifying the ReplicationSystem."), *ObjectData.RefHandle.ToString(), ToCStr(DebugName), ObjectData.RefHandle.GetReplicationSystemId());
-
-				StaleObjects.Push(ObjectData.RefHandle);
+				if (ensureMsgf(!EnumHasAnyFlags(ObjectData.InstanceProtocol->InstanceTraits, EReplicationInstanceProtocolTraits::IsBound), TEXT("UObjectReplicationBridge::PruneStaleObjects Bound ObjectInstance replicated as: %s has been destroyed without notifying the ReplicationSystem."), *ObjectData.RefHandle.ToString()))
+				{
+					StaleObjects.Push(ObjectData.RefHandle);
+				}
 			}
 		}
 	};
diff --git a/Engine/Source/Runtime/Experimental/Iris/Core/Private/Iris/ReplicationSystem/ReplicationSystem.cpp b/Engine/Source/Runtime/Experimental/Iris/Core/Private/Iris/ReplicationSystem/ReplicationSystem.cpp
--- a/Engine/Source/Runtime/Experimental/Iris/Core/Private/Iris/ReplicationSystem/ReplicationSystem.cpp
+++ b/Engine/Source/Runtime/Experimental/Iris/Core/Private/Iris/ReplicationSystem/ReplicationSystem.cpp
@@ -29,10 +29,9 @@
 
 namespace ReplicationSystemCVars
 {
-#if !UE_BUILD_SHIPPING
+
 static bool bForcePruneBeforeUpdate = false;
 static FAutoConsoleVariableRef CVarForcePruneBeforeUpdate(TEXT("net.Iris.ForcePruneBeforeUpdate"), bForcePruneBeforeUpdate, TEXT("Verify integrity of all tracked instances at the start of every update."));
-#endif
 
 static bool bAllowAttachmentSendPolicyFlags = true;
 static FAutoConsoleVariableRef CVarAllowAttachmentSendPolicyFlags(TEXT("net.Iris.Attachments.AllowSendPolicyFlags"), bAllowAttachmentSendPolicyFlags, TEXT("Allow use of ENetObjectAttachmentSendPolicyFlags to specify behavior of RPCs."));
@@ -663,13 +662,11 @@ void UReplicationSystem::PreSendUpdate(const FSendUpdateParams& Params)
 		}
 #endif
 	
-#if !UE_BUILD_SHIPPING
 		// Force a integrity check of all replicated instances
 		if (bDoCollectGarbage || ReplicationSystemCVars::bForcePruneBeforeUpdate)
 		{
 			CollectGarbage();
 		}
-#endif
 
 		if (bAllowObjectReplication)
 		{
