Cluster union : when processing damage from contact, find the actual child particle in contact instead of using the one that has the closest center of mass
This fixes the issue where the wrong child may be processed for destruction
Also fix the indexing of shapes ( ShapeIdx ) after removing some of them

diff --git a/Engine/Source/Runtime/Experimental/Chaos/Private/Chaos/GeometryParticles.cpp b/Engine/Source/Runtime/Experimental/Chaos/Private/Chaos/GeometryParticles.cpp
--- a/Engine/Source/Runtime/Experimental/Chaos/Private/Chaos/GeometryParticles.cpp
+++ b/Engine/Source/Runtime/Experimental/Chaos/Private/Chaos/GeometryParticles.cpp
@@ -50,6 +50,17 @@ namespace Chaos
 	void TGeometryParticlesImp<T, d, SimType>::RemoveShapesAtSortedIndices(const int32 ParticleIndex, const TArrayView<const int32>& InIndices)
 	{
 		RemoveArrayItemsAtSortedIndices(MShapesArray[ParticleIndex], InIndices);
+
+		// ShapeIdx need to be be updated 
+		const int32 NumShapes = MShapesArray[ParticleIndex].Num();
+		for (int32 ShapeIndex = 0; ShapeIndex < NumShapes; ++ShapeIndex)
+		{
+			const FShapeInstancePtr& Shape = MShapesArray[ParticleIndex][ShapeIndex];
+			if (Shape)
+			{
+				Shape->ModifyShapeIndex(ShapeIndex);
+			}
+		}
 	}
 
 	template <typename T, int d, EGeometryParticlesSimType SimType>
diff --git a/Engine/Source/Runtime/Experimental/Chaos/Private/Chaos/PBDRigidClustering.cpp b/Engine/Source/Runtime/Experimental/Chaos/Private/Chaos/PBDRigidClustering.cpp
--- a/Engine/Source/Runtime/Experimental/Chaos/Private/Chaos/PBDRigidClustering.cpp
+++ b/Engine/Source/Runtime/Experimental/Chaos/Private/Chaos/PBDRigidClustering.cpp
@@ -2266,15 +2266,43 @@ namespace Chaos
 							return;
 						}
 
-						// At the moment, we don't want to apply strains to children of ClusterUnions, we want instead
-						// to apply the strains to GRANDchildren of ClusterUnions.
-						FPBDRigidParticleHandle* ClosestChild = FindClosestChild(Cluster, ContactWorldLocation);
+						FPBDRigidParticleHandle* ChildParticleInContact = nullptr;
+						if (Chaos::FClusterUnion* ClusterUnion = ClusterUnionManager.FindClusterUnionFromParticle(Cluster))
+						{
+							// At the moment, we don't want to apply strains to children of ClusterUnions, we want instead
+							// to apply the strains to GRANDchildren of ClusterUnions.
+							// We are looking for the child particle that match the shape in contact 
+							// We cannot use a proximity method because this may return the wrong child if its center of mass is closer than actual child in contact
+							// Like in the example below where A would be wrongly reported:
+							//               +-------+
+							//               |   x A |
+							//            +--+-------+--------------------------------+
+							// contact-> *|                    x B                    |
+							//            +-------------------------------------------+
+							int32 ShapeInContactIndex = ContactHandle->GetContact().GetShape0()->GetShapeIndex();
+							if (Cluster == ContactHandle->GetContact().GetParticle1())
+							{
+								ShapeInContactIndex = ContactHandle->GetContact().GetShape1()->GetShapeIndex();
+							}
 
+							// cluster union garantee a one to one mapping between shapes index and children
+							// ( note : we cannot use MChildren because it can be oput of order, so we need to use the ChildParticles from FClusterUnion ) 
+							if (ensureAlways(ClusterUnion->ChildParticles.IsValidIndex(ShapeInContactIndex)))
+							{
+								ChildParticleInContact = ClusterUnion->ChildParticles[ShapeInContactIndex];
+							}
+
+							if (ChildParticleInContact == nullptr)
+							{
+								// if anything above failed to find a particle fall back to the distance based method
+								ChildParticleInContact = FindClosestChild(Cluster, ContactWorldLocation);
+							}
+						}
 						// If Closest child is not a clustered, then there is no substructure to apply strain to,
 						// so null Cluster
 						Cluster
-							= ClosestChild
-							? ClosestChild->CastToClustered()
+							= ChildParticleInContact
+							? ChildParticleInContact->CastToClustered()
 							: nullptr;
 					}
 
