Constraints: fixed issue with constraints not being added to the level using UConstraintsScriptingLibrary

#rnx

diff --git a/Engine/Source/Runtime/Experimental/Animation/Constraints/Private/ConstraintsScripting.cpp b/Engine/Source/Runtime/Experimental/Animation/Constraints/Private/ConstraintsScripting.cpp
--- a/Engine/Source/Runtime/Experimental/Animation/Constraints/Private/ConstraintsScripting.cpp
+++ b/Engine/Source/Runtime/Experimental/Animation/Constraints/Private/ConstraintsScripting.cpp
@@ -46,8 +46,29 @@ UTickableTransformConstraint* UConstraintsScriptingLibrary::CreateFromType(
 bool UConstraintsScriptingLibrary::AddConstraint(UWorld* InWorld, UTransformableHandle* InParentHandle, UTransformableHandle* InChildHandle,
 	UTickableTransformConstraint *InConstraint, const bool bMaintainOffset)
 {
-	bool Val =  FTransformConstraintUtils::AddConstraint(InWorld, InParentHandle, InChildHandle,InConstraint, bMaintainOffset);
-	return Val;
+	if (!InWorld)
+	{
+		UE_LOG(LogTemp, Error, TEXT("AddConstraint: Need Valid World."));
+		return false;
+	}
+	
+	if (!InConstraint)
+	{
+		UE_LOG(LogTemp, Error, TEXT("AddConstraint: InConstraint is null."));
+		return false;
+	}
+	
+	const bool bAdded = FTransformConstraintUtils::AddConstraint(InWorld, InParentHandle, InChildHandle,InConstraint, bMaintainOffset);
+	if (!bAdded)
+	{
+		UE_LOG(LogTemp, Error, TEXT("AddConstraint: Constraint not added"));
+		return false;
+	}
+	
+	FConstraintsManagerController& Controller = FConstraintsManagerController::Get(InWorld);
+	Controller.StaticConstraintCreated(InWorld, InConstraint);
+	
+	return true;
 }
 
 TArray<UTickableConstraint*> UConstraintsScriptingLibrary::GetConstraintsArray(UWorld* InWorld)
diff --git a/Engine/Source/Runtime/Experimental/Animation/Constraints/Public/ConstraintsScripting.h b/Engine/Source/Runtime/Experimental/Animation/Constraints/Public/ConstraintsScripting.h
--- a/Engine/Source/Runtime/Experimental/Animation/Constraints/Public/ConstraintsScripting.h
+++ b/Engine/Source/Runtime/Experimental/Animation/Constraints/Public/ConstraintsScripting.h
@@ -70,6 +70,7 @@ public:
 	* @param InConsrtaint The constraint
 	* @return return If constraint added correctly
 	*/
+	UE_DEPRECATED(5.4, "UConstraintsScriptingLibrary::AddConstraint is deprecated. use UControlRigSequencerEditorLibrary::AddConstraint instead.")
 	UFUNCTION(BlueprintCallable, Category = "Editor Scripting | Sequencer Tools | Control Rig | Constraints")
 	static CONSTRAINTS_API bool AddConstraint(UWorld* InWorld, UTransformableHandle* InParentHandle, UTransformableHandle* InChildHandle, UTickableTransformConstraint *InConstraint,
 			const bool bMaintainOffset);
