Constraints: Use weak pointer to hold World in Constraint Subystem

diff --git a/Engine/Source/Runtime/Experimental/Animation/Constraints/Private/ConstraintSubsystem.cpp b/Engine/Source/Runtime/Experimental/Animation/Constraints/Private/ConstraintSubsystem.cpp
--- a/Engine/Source/Runtime/Experimental/Animation/Constraints/Private/ConstraintSubsystem.cpp
+++ b/Engine/Source/Runtime/Experimental/Animation/Constraints/Private/ConstraintSubsystem.cpp
@@ -44,7 +44,7 @@ void UConstraintSubsystem::Deinitialize()
 {
 	for (int32 Index = ConstraintsInWorld.Num() - 1; Index >= 0; --Index)
 	{
-		ConstraintsInWorld[Index].RemoveConstraints(ConstraintsInWorld[Index].World);
+		ConstraintsInWorld[Index].RemoveConstraints(ConstraintsInWorld[Index].World.Get());
 	}
 	ConstraintsInWorld.Reset();
 
@@ -320,7 +320,6 @@ FConstraintsEvaluationGraph& FConstraintsInWorld::GetEvaluationGraph()
 {
 	if (!EvaluationGraph)
 	{
-		ensure(World);
 		EvaluationGraph = MakeShared<FConstraintsEvaluationGraph>(*this);
 	}
 	return *EvaluationGraph;
diff --git a/Engine/Source/Runtime/Experimental/Animation/Constraints/Private/ConstraintSubsystem.h b/Engine/Source/Runtime/Experimental/Animation/Constraints/Private/ConstraintSubsystem.h
--- a/Engine/Source/Runtime/Experimental/Animation/Constraints/Private/ConstraintSubsystem.h
+++ b/Engine/Source/Runtime/Experimental/Animation/Constraints/Private/ConstraintSubsystem.h
@@ -21,7 +21,7 @@ public:
 	GENERATED_BODY()
 
 	UPROPERTY(transient)
-	TObjectPtr<UWorld> World;
+	TWeakObjectPtr<UWorld> World;
 
 	UPROPERTY(transient)
 	mutable TArray<TWeakObjectPtr<UTickableConstraint>> Constraints;
diff --git a/Engine/Source/Runtime/Experimental/Animation/Constraints/Private/ConstraintsEvaluationGraph.cpp b/Engine/Source/Runtime/Experimental/Animation/Constraints/Private/ConstraintsEvaluationGraph.cpp
--- a/Engine/Source/Runtime/Experimental/Animation/Constraints/Private/ConstraintsEvaluationGraph.cpp
+++ b/Engine/Source/Runtime/Experimental/Animation/Constraints/Private/ConstraintsEvaluationGraph.cpp
@@ -41,7 +41,7 @@ FConstraintNode& FConstraintsEvaluationGraph::GetNode(const TWeakObjectPtr<UTick
 
 	FConstraintNode Node;
 	Node.ConstraintID = InConstraint->ConstraintID;
-	Node.ConstraintTick = &InConstraint->GetTickFunction(ConstraintsInWorld.World);
+	Node.ConstraintTick = &InConstraint->GetTickFunction(ConstraintsInWorld.World.Get());
 	return Nodes.Emplace_GetRef(MoveTemp(Node)); 
 }
 
