Reduce the frequency at which unresolved references are resolved.

diff --git a/Engine/Source/Runtime/Experimental/Iris/Core/Private/Iris/ReplicationSystem/ReplicationReader.cpp b/Engine/Source/Runtime/Experimental/Iris/Core/Private/Iris/ReplicationSystem/ReplicationReader.cpp
--- a/Engine/Source/Runtime/Experimental/Iris/Core/Private/Iris/ReplicationSystem/ReplicationReader.cpp
+++ b/Engine/Source/Runtime/Experimental/Iris/Core/Private/Iris/ReplicationSystem/ReplicationReader.cpp
@@ -42,6 +42,24 @@ CSV_DEFINE_CATEGORY(IrisClient, true);
 namespace UE::Net::Private
 {
 
+static bool bUseResolvingHandleCache = true;
+static FAutoConsoleVariableRef CVarUseResolvingHandleCache(
+	TEXT("net.Iris.UseResolvingHandleCache"),
+	bUseResolvingHandleCache,
+	TEXT("Enable the use of a hot and cold cache when resolving unresolved caches to reduce the time spent resolving references."));
+
+static int32 HotResolvingLifetimeMS = 1000;
+static FAutoConsoleVariableRef CVarHotResolvingLifetimeMS(
+	TEXT("net.Iris.HotResolvingLifetimeMS"),
+	HotResolvingLifetimeMS,
+	TEXT("An unresolved reference is considered hot if it was created within this many milliseconds, and cold otherwise."));
+
+static int32 ColdResolvingRetryTimeMS = 200;
+static FAutoConsoleVariableRef CVarColdResolvingRetryTimeMS(
+	TEXT("net.Iris.ColdResolvingRetryTimeMS"),
+	ColdResolvingRetryTimeMS,
+	TEXT("Resolve unresolved cold references after this many milliseconds."));
+
 static bool bExecuteReliableRPCsBeforeApplyState = true;
 static FAutoConsoleVariableRef CVarExecuteReliableRPCsBeforeApplyState(
 		TEXT("net.Iris.ExecuteReliableRPCsBeforeApplyState"),
@@ -850,6 +868,9 @@ void FReplicationReader::ReadObjectInBatch(FNetSerializationContext& Context, FN
 			goto ErrorHandling;
 		}
 
+		// If this handle is considered unresolved, add it to the hot cache to force a resolve.
+		RemoveFromUnresolvedCache(NetRefHandle);
+
 		InternalIndex = NetRefHandleManager->GetInternalIndex(NetRefHandle);
 		FNetRefHandleManager::FReplicatedObjectData& ObjectData = NetRefHandleManager->GetReplicatedObjectDataNoCheck(InternalIndex);
 		ObjectData.bAllowDestroyInstanceFromRemote = EnumHasAnyFlags(CreateResult.Flags, EReplicationBridgeCreateNetRefHandleResultFlags::AllowDestroyInstanceFromRemote);
@@ -1100,6 +1121,7 @@ void FReplicationReader::UpdateObjectReferenceTracking(FReplicatedObjectInfo* Re
 
 				// Remove from tracking
 				UnresolvedHandleToDependents.RemoveSingle(Handle, OwnerInternalIndex);
+				RemoveFromUnresolvedCache(Handle);
 				UE_LOG(LogIris, Verbose, TEXT("FReplicationReader::UpdateObjectReferenceTracking Removing unresolved reference %s for %s"), ToCStr(Handle.ToString()), ToCStr(NetRefHandleManager->GetNetRefHandleFromInternalIndex(OwnerInternalIndex).ToString()));
 			}
 		}
@@ -1258,6 +1280,7 @@ void FReplicationReader::CleanupReferenceTracking(FReplicatedObjectInfo* ObjectI
 		// Remove from tracking
 		FNetRefHandle Handle = Element.Value;
 		UnresolvedHandleToDependents.RemoveSingle(Handle, ObjectIndex);
+		RemoveFromUnresolvedCache(Handle);
 		UE_LOG(LogIris, Verbose, TEXT("FReplicationReader::CleanupReferenceTracking Removing unresolved reference %s for %s"), *Handle.ToString(), *(NetRefHandleManager->GetNetRefHandleFromInternalIndex(ObjectIndex).ToString()));
 	}
 	ObjectInfo->UnresolvedObjectReferences.Reset();
@@ -1674,6 +1697,7 @@ void FReplicationReader::DispatchStateData(FNetSerializationContext& Context)
 void FReplicationReader::ResolveAndDispatchUnresolvedReferences()
 {
 	IRIS_PROFILER_SCOPE(FReplicationReader_ResolveAndDispatchUnresolvedReferences);
+	CSV_SCOPED_TIMING_STAT(IrisClient, ResolveAndDispatchUnresolvedReferences);
 
 	// Setup context for dispatch
 	FInternalNetSerializationContext InternalContext;
@@ -1688,14 +1712,55 @@ void FReplicationReader::ResolveAndDispatchUnresolvedReferences()
 	Context.SetInternalContext(&InternalContext);
 
 	// Currently we brute force this by iterating over all handles pending resolve and update all objects pending resolve
-	TArray<FNetRefHandle> UpdatedHandles;
-	UpdatedHandles.Reserve(128);
-	UnresolvedHandleToDependents.GetKeys(UpdatedHandles);
-	
-	TSet<uint32> InternalObjectsToResolve;
-	InternalObjectsToResolve.Reserve(UnresolvedHandleToDependents.Num());
+	VisitedUnresolvedHandles.Reset();
+	InternalObjectsToResolve.Reset();
+
+	const uint32 CurrTimeMS = static_cast<uint32>(FPlatformTime::Seconds() * 1000.0f);
+	const uint32 HotLifetimeMS = HotResolvingLifetimeMS > 0 ? static_cast<uint32>(HotResolvingLifetimeMS) : 0;
+	const uint32 ColdRetryTimeMS = ColdResolvingRetryTimeMS > 0 ? static_cast<uint32>(ColdResolvingRetryTimeMS) : 0;
+
+	for (const TPair<FNetRefHandle, uint32>& It : UnresolvedHandleToDependents)
+	{
+		const FNetRefHandle& Handle = It.Key;
+
+		if (!VisitedUnresolvedHandles.Contains(Handle))
+		{
+			// Determine if the handle should be resolved.
+			if (bUseResolvingHandleCache)
+			{
+				// If the handle is in the hot cache it should be resolved every time ResolveAndDispatchUnresolvedReferences() is called
+				// and will be moved to the cold cache after a fixed period of time.
+				if (const uint32* LifetimeMS = HotUnresolvedHandleCache.Find(Handle))
+				{
+					if ((CurrTimeMS - *LifetimeMS) > HotLifetimeMS)
+					{
+						HotUnresolvedHandleCache.Remove(Handle);
+						ColdUnresolvedHandleCache.Add(Handle);
+					}
+				}
+				// If the handle is in the cold cache it will only be resolved at a fixed interval and will remain in this cache indefinitely.
+				else if (uint32* LastResolvedMS = ColdUnresolvedHandleCache.Find(Handle))
+				{
+					if ((CurrTimeMS - *LastResolvedMS) < ColdRetryTimeMS)
+					{
+						continue;
+					}
 
-	for (FNetRefHandle Handle : UpdatedHandles)
+					*LastResolvedMS = CurrTimeMS;
+				}
+				// If the handle is in neither the hot or cold cache, put it in the hot cache.
+				else
+				{
+					HotUnresolvedHandleCache.Add(Handle, CurrTimeMS);
+				}
+			}
+
+			// Only check this handle once per call.
+			VisitedUnresolvedHandles.Add(Handle);
+		}
+	}
+
+	for (FNetRefHandle Handle : VisitedUnresolvedHandles)
 	{
 		// Only make sense to update dependant objects if handle is resolvable
 		if (ObjectReferenceCache->ResolveObjectReferenceHandle(Handle, ResolveContext) != nullptr)
@@ -1716,13 +1781,23 @@ void FReplicationReader::ResolveAndDispatchUnresolvedReferences()
 		ResolveAndDispatchUnresolvedReferencesForObject(Context, InternalIndex);
 	}
 
-	CSV_CUSTOM_STAT(IrisClient, UnresolvedHandlesToResolve, UpdatedHandles.Num(), ECsvCustomStatOp::Accumulate);
+	CSV_CUSTOM_STAT(IrisClient, HotUnresolvedHandleCache, HotUnresolvedHandleCache.Num(), ECsvCustomStatOp::Set);
+	CSV_CUSTOM_STAT(IrisClient, ColdUnresolvedHandleCache, ColdUnresolvedHandleCache.Num(), ECsvCustomStatOp::Set);
+
+	CSV_CUSTOM_STAT(IrisClient, UnresolvedHandlesToResolve, VisitedUnresolvedHandles.Num(), ECsvCustomStatOp::Accumulate);
 	CSV_CUSTOM_STAT(IrisClient, UnresolvedObjectsToResolve, InternalObjectsToResolve.Num(), ECsvCustomStatOp::Accumulate);
 
-	if (NumHandlesPendingResolveLastUpdate != UpdatedHandles.Num() || ObjectsWithAttachmentPendingResolve.Num() > 0)
+	const int32 TotalCacheSize = static_cast<int32>(
+		HotUnresolvedHandleCache.GetAllocatedSize() +
+		ColdUnresolvedHandleCache.GetAllocatedSize() +
+		VisitedUnresolvedHandles.GetAllocatedSize() +
+		InternalObjectsToResolve.GetAllocatedSize());
+	CSV_CUSTOM_STAT(IrisClient, UnresolvedHandleBufferSizes, TotalCacheSize, ECsvCustomStatOp::Set);
+
+	if (NumHandlesPendingResolveLastUpdate != VisitedUnresolvedHandles.Num() || ObjectsWithAttachmentPendingResolve.Num() > 0)
 	{
 		UE_LOG_REPLICATIONREADER(TEXT("FReplicationReader::ResolveAndDispatchUnresolvedReferences NetHandles pending: %u Attachments pending: %u)"), UpdatedHandles.Num(), ObjectsWithAttachmentPendingResolve.Num());
-		NumHandlesPendingResolveLastUpdate = UpdatedHandles.Num();
+		NumHandlesPendingResolveLastUpdate = VisitedUnresolvedHandles.Num();
 	}
 }
 
@@ -1750,6 +1825,7 @@ void FReplicationReader::UpdateUnresolvableReferenceTracking()
 			Dependents.Reset();
 			UnresolvedHandleToDependents.MultiFind(DestroyedHandle, Dependents, bMaintainOrder);
 			UnresolvedHandleToDependents.Remove(DestroyedHandle);
+			RemoveFromUnresolvedCache(DestroyedHandle);
 			for (const uint32 DependentObjectIndex : Dependents)
 			{
 				FReplicatedObjectInfo* ReplicationInfo = GetReplicatedObjectInfo(DependentObjectIndex);
@@ -1916,6 +1992,15 @@ bool FReplicationReader::EnqueueEndReplication(FPendingBatchData* PendingBatchDa
 	return true;
 }
 
+void FReplicationReader::RemoveFromUnresolvedCache(const FNetRefHandle Handle)
+{
+	if (bUseResolvingHandleCache && !UnresolvedHandleToDependents.Contains(Handle))
+	{
+		HotUnresolvedHandleCache.Remove(Handle);
+		ColdUnresolvedHandleCache.Remove(Handle);
+	}
+}
+
 void FReplicationReader::ProcessQueuedBatches()
 {
 	UE_NET_TRACE_FRAME_STATSCOUNTER(Parameters.ReplicationSystem->GetId(), ReplicationReader.PendingQueuedBatches, PendingBatches.PendingBatches.Num(), ENetTraceVerbosity::Trace);
diff --git a/Engine/Source/Runtime/Experimental/Iris/Core/Private/Iris/ReplicationSystem/ReplicationReader.h b/Engine/Source/Runtime/Experimental/Iris/Core/Private/Iris/ReplicationSystem/ReplicationReader.h
--- a/Engine/Source/Runtime/Experimental/Iris/Core/Private/Iris/ReplicationSystem/ReplicationReader.h
+++ b/Engine/Source/Runtime/Experimental/Iris/Core/Private/Iris/ReplicationSystem/ReplicationReader.h
@@ -203,6 +203,10 @@ private:
 	// If we are queuing data for a batch we must also defer calls to EndReplication
 	// This method writes this method in the form of a QueuedChunk
 	bool EnqueueEndReplication(FPendingBatchData* PendingBatchData, bool bShouldDestroyInstance, FNetRefHandle NetRefHandleToEndReplication);
+
+	// Remove a handle from the hot and cold unresolved caches used by ResolveAndDispatchUnresolvedReferences(). If this handle is marked as unresolved
+	// again, it will be added to the hot cache.
+	void RemoveFromUnresolvedCache(const FNetRefHandle Handle);
 	
 private:
 
@@ -220,6 +224,16 @@ private:
 	FReplicationStateStorage* StateStorage;
 	UReplicationBridge* ReplicationBridge;
 
+	// A cache holding unresolved handles that should be resolved each time ResolveAndDispatchUnresolvedReferences() is called.
+	TMap<FNetRefHandle, uint32> HotUnresolvedHandleCache;
+
+	// A cache holding unresolved handles that should be resolved by ResolveAndDispatchUnresolvedReferences() at fixed intervals.
+	TMap<FNetRefHandle, uint32> ColdUnresolvedHandleCache;
+
+	// Temporary buffers used by ResolveAndDispatchUnresolvedReferences().
+	TSet<FNetRefHandle> VisitedUnresolvedHandles;
+	TSet<uint32> InternalObjectsToResolve;
+
 	// We track some data about incoming objects
 	// Stored in a map for now
 	TMap<uint32, FReplicatedObjectInfo> ReplicatedObjects;
