Remove clamping from IAS cache usage statistics

#rnx

- The use of clamp here is wrong, as the function is (value, min,max) but we are assuming (min, value, max)
- We probably don't want to clamp here at all as it would hide errors that we should investigate. At the moment due to the incorrect use of clamp we are seeing percentages way in excess of 100% so we know there is something to work on but if the clamp was working correctly we'd never have spotted the problem.

diff --git a/Engine/Source/Runtime/Experimental/IoStoreOnDemand/Private/Statistics.cpp b/Engine/Source/Runtime/Experimental/IoStoreOnDemand/Private/Statistics.cpp
--- a/Engine/Source/Runtime/Experimental/IoStoreOnDemand/Private/Statistics.cpp
+++ b/Engine/Source/Runtime/Experimental/IoStoreOnDemand/Private/Statistics.cpp
@@ -481,8 +481,9 @@ void FOnDemandIoBackendStats::ReportEndPointAnalytics(TArray<FAnalyticsEventAttr
 
 	if (GIasReportCacheAnalyticsEnabled)
 	{
-		int64 CacheTotalCount = GCacheGetCount.Get() + GCachePutCount.Get();
-		float CacheUsagePercent = GCacheMaxBytes > 0 ? FMath::Clamp(0.f, (100.f*float(GCacheCachedBytes.Get())) / float(GCacheMaxBytes), 100.f) : 0.f;
+		const int64 CacheTotalCount = GCacheGetCount.Get() + GCachePutCount.Get();
+		const float CacheUsagePercent = GCacheMaxBytes > 0 ? 100.f * (float(GCacheCachedBytes.Get()) / float(GCacheMaxBytes)) : 0.f;
+
 		AppendAnalyticsEventAttributeArray(OutAnalyticsArray
 
 			,TRACK_DELTA("IasCacheTotalCount", CacheTotalCount)
