Adding CSV stats for Iris PreUpdate, Polling and Prioritization
#tests Wind64 Server+Client against Cooked Data

diff --git a/Engine/Source/Runtime/Experimental/Iris/Core/Private/Iris/ReplicationSystem/ObjectReplicationBridge.cpp b/Engine/Source/Runtime/Experimental/Iris/Core/Private/Iris/ReplicationSystem/ObjectReplicationBridge.cpp
--- a/Engine/Source/Runtime/Experimental/Iris/Core/Private/Iris/ReplicationSystem/ObjectReplicationBridge.cpp
+++ b/Engine/Source/Runtime/Experimental/Iris/Core/Private/Iris/ReplicationSystem/ObjectReplicationBridge.cpp
@@ -7,6 +7,7 @@
 
 #include "Iris/IrisConfigInternal.h"
 
+#include "Iris/Core/IrisCsv.h"
 #include "Iris/Core/IrisLog.h"
 #include "Iris/Core/IrisMemoryTracker.h"
 #include "Iris/Core/IrisProfiler.h"
@@ -1028,6 +1029,9 @@ void UObjectReplicationBridge::PreUpdate(const UE::Net::FNetBitArrayView Objects
 {
 	using namespace UE::Net::Private;
 
+#if UE_NET_IRIS_CSV_STATS
+	CSV_SCOPED_TIMING_STAT(Iris, ReplicationBridge_PreUpdate);
+#endif
 	IRIS_PROFILER_SCOPE(UObjectReplicationBridge_PreUpdate);
 
 	// TODO: Get rid of Poller
@@ -1046,6 +1050,9 @@ void UObjectReplicationBridge::PollAndCopy(const UE::Net::FNetBitArrayView Objec
 {
 	using namespace UE::Net::Private;
 
+#if UE_NET_IRIS_CSV_STATS
+	CSV_SCOPED_TIMING_STAT(Iris, ReplicationBridge_PollAndCopy);
+#endif
 	IRIS_PROFILER_SCOPE(UObjectReplicationBridge_PollAndCopy);
 
 	FObjectPoller::FInitParams PollerInitParams;
diff --git a/Engine/Source/Runtime/Experimental/Iris/Core/Private/Iris/ReplicationSystem/ReplicationSystem.cpp b/Engine/Source/Runtime/Experimental/Iris/Core/Private/Iris/ReplicationSystem/ReplicationSystem.cpp
--- a/Engine/Source/Runtime/Experimental/Iris/Core/Private/Iris/ReplicationSystem/ReplicationSystem.cpp
+++ b/Engine/Source/Runtime/Experimental/Iris/Core/Private/Iris/ReplicationSystem/ReplicationSystem.cpp
@@ -1,6 +1,7 @@
 // Copyright Epic Games, Inc. All Rights Reserved.
 
 #include "Iris/ReplicationSystem/ReplicationSystem.h"
+#include "Iris/Core/IrisCsv.h"
 #include "Iris/Core/IrisDebugging.h"
 #include "Iris/Core/IrisLog.h"
 #include "Iris/Core/IrisMemoryTracker.h"
@@ -368,6 +369,9 @@ public:
 	// Runs after filtering
 	void UpdatePrioritization(const FNetBitArrayView& ReplicatingConnections)
 	{
+#if UE_NET_IRIS_CSV_STATS
+		CSV_SCOPED_TIMING_STAT(Iris, ReplicationSystem_UpdatePrioritization);
+#endif
 		IRIS_PROFILER_SCOPE(FReplicationSystem::FImpl::UpdatePrioritization);
 		LLM_SCOPE_BYTAG(Iris);
 
