Iris
* Dynamic FilterConfigs can now specify a FilterProfile which can be used to set filter settings unique to specific classes.
* GridFilter now supports filter profiles. Can be used to set a specific value to the FrameCountBeforeCulling of an actor who is spatially filtered.

#[review]-32282100
#rnx

diff --git a/Engine/Source/Runtime/Engine/Private/Net/Iris/ReplicationSystem/ActorReplicationBridge.cpp b/Engine/Source/Runtime/Engine/Private/Net/Iris/ReplicationSystem/ActorReplicationBridge.cpp
--- a/Engine/Source/Runtime/Engine/Private/Net/Iris/ReplicationSystem/ActorReplicationBridge.cpp
+++ b/Engine/Source/Runtime/Engine/Private/Net/Iris/ReplicationSystem/ActorReplicationBridge.cpp
@@ -283,7 +283,9 @@ UE::Net::FNetRefHandle UActorReplicationBridge::BeginReplication(AActor* Actor,
 		{
 			// Only apply owner filter if we haven't force enabled a dynamic filter.
 			constexpr bool bRequireForceEnabled = true;
-			FNetObjectFilterHandle FilterHandle = GetDynamicFilter(Actor->GetClass(), bRequireForceEnabled);
+			FName FilterProfile;
+			FNetObjectFilterHandle FilterHandle = GetDynamicFilter(Actor->GetClass(), bRequireForceEnabled, FilterProfile);
+
 			if (FilterHandle == InvalidNetObjectFilterHandle)
 			{
 				GetReplicationSystem()->SetFilter(ActorRefHandle, ToOwnerFilterHandle);
diff --git a/Engine/Source/Runtime/Experimental/Iris/Core/Private/Iris/ReplicationSystem/Filtering/NetObjectGridFilter.cpp b/Engine/Source/Runtime/Experimental/Iris/Core/Private/Iris/ReplicationSystem/Filtering/NetObjectGridFilter.cpp
--- a/Engine/Source/Runtime/Experimental/Iris/Core/Private/Iris/ReplicationSystem/Filtering/NetObjectGridFilter.cpp
+++ b/Engine/Source/Runtime/Experimental/Iris/Core/Private/Iris/ReplicationSystem/Filtering/NetObjectGridFilter.cpp
@@ -37,6 +37,26 @@ void UNetObjectGridFilter::RemoveConnection(uint32 ConnectionId)
 	PerConnectionInfos[ConnectionId] = {};
 }
 
+uint16 UNetObjectGridFilter::GetFrameCountBeforeCulling(FName ProfileName) const
+{
+	uint16 FrameCount = Config->DefaultFrameCountBeforeCulling;
+	
+	if (ProfileName.IsNone())
+	{
+		return FrameCount;
+	}
+
+	FNetObjectGridFilterProfile* FilterProfile = Config->FilterProfiles.FindByKey(ProfileName);
+	ensureMsgf(FilterProfile, TEXT("UNetObjectGridFilterConfig does not hold any profile named %s"), *ProfileName.ToString());
+
+	if (FilterProfile)
+	{
+		FrameCount = FilterProfile->FrameCountBeforeCulling;
+	}
+
+	return FrameCount;
+}
+
 bool UNetObjectGridFilter::AddObject(uint32 ObjectIndex, FNetObjectFilterAddObjectParams& Params)
 {
 	// We support either a world location in the state, tagged with RepTag_WorldLocation, or via the WorldLocations instance.
@@ -53,7 +73,9 @@ bool UNetObjectGridFilter::AddObject(uint32 ObjectIndex, FNetObjectFilterAddObje
 	ObjectLocationInfo.SetInfoIndex(InfoIndex);
 
 	FPerObjectInfo& PerObjectInfo = ObjectInfos[InfoIndex];
+
 	PerObjectInfo.ObjectIndex = ObjectIndex;
+	PerObjectInfo.FrameCountBeforeCulling = GetFrameCountBeforeCulling(Params.ProfileName);
 
 	AddCellInfoForObject(ObjectLocationInfo, Params.InstanceProtocol);
 	
@@ -175,7 +197,7 @@ void UNetObjectGridFilter::Filter(FNetObjectFilteringParams& Params)
 
 						if (ObjectToViewDistSq <= DistSq)
 						{
-							ConnectionInfo.RecentObjectFrameCount.Add(ObjectIndex, Config->ViewPosRelevancyFrameCount);
+							ConnectionInfo.RecentObjectFrameCount.Add(ObjectIndex, PerObjectInfo.FrameCountBeforeCulling);
 							break; // Don't need to test all connections once one is within Distance
 						}
 					}
@@ -191,7 +213,7 @@ void UNetObjectGridFilter::Filter(FNetObjectFilteringParams& Params)
 #endif
 
 		// Set the AllowedObjects and decrease their frame count
-		for (TMap<uint32, uint32>::TIterator It = ConnectionInfo.RecentObjectFrameCount.CreateIterator(); It; ++It)
+		for (TMap<uint32, uint16>::TIterator It = ConnectionInfo.RecentObjectFrameCount.CreateIterator(); It; ++It)
 		{
 			if (It->Value > 0)
 			{
@@ -250,12 +272,20 @@ void UNetObjectGridFilter::AddCellInfoForObject(const FObjectLocationInfo& Objec
 {
 	// Called for completely new objects
 	FPerObjectInfo& PerObjectInfo = ObjectInfos[ObjectLocationInfo.GetInfoIndex()];
-	PerObjectInfo.SetCullDistance(Config->DefaultCullDistance);
-	UpdatePositionAndCullDistance(ObjectLocationInfo, PerObjectInfo, InstanceProtocol);
+	
+	// Set cull distance
+	{
+		PerObjectInfo.SetCullDistance(Config->DefaultCullDistance);
+		UpdatePositionAndCullDistance(ObjectLocationInfo, PerObjectInfo, InstanceProtocol);
+	}
 
 	FCellBox NewCellBox;
-	CalculateCellBox(PerObjectInfo, NewCellBox);
-	PerObjectInfo.CellBox = NewCellBox;
+
+	// Set cellbox info
+	{
+		CalculateCellBox(PerObjectInfo, NewCellBox);
+		PerObjectInfo.CellBox = NewCellBox;
+	}
 
 	const uint32 ObjectIndex = PerObjectInfo.ObjectIndex;
 
@@ -485,8 +515,8 @@ FString UNetObjectGridFilter::PrintDebugInfoForObject(const FDebugInfoParams& Pa
 	const FPerObjectInfo& PerObjectInfo = ObjectInfos[ObjectLocationInfo.GetInfoIndex()];
 
 	const FPerConnectionInfo* const PerConnectionInfo = PerConnectionInfos.IsValidIndex(Params.ConnectionId) ? &PerConnectionInfos[Params.ConnectionId] : nullptr;
-	const uint32* CulledFrameCountPtr = PerConnectionInfo ? PerConnectionInfo->RecentObjectFrameCount.Find(ObjectIndex) : nullptr;
-	const uint32 CulledFrameCount = CulledFrameCountPtr ? *CulledFrameCountPtr : 0;
+	const uint16* CulledFrameCountPtr = PerConnectionInfo ? PerConnectionInfo->RecentObjectFrameCount.Find(ObjectIndex) : nullptr;
+	const uint16 CulledFrameCount = CulledFrameCountPtr ? *CulledFrameCountPtr : 0;
 
 	double Dist2d = DOUBLE_BIG_NUMBER;
 	double DistZ = DOUBLE_BIG_NUMBER;
diff --git a/Engine/Source/Runtime/Experimental/Iris/Core/Private/Iris/ReplicationSystem/Filtering/ReplicationFiltering.cpp b/Engine/Source/Runtime/Experimental/Iris/Core/Private/Iris/ReplicationSystem/Filtering/ReplicationFiltering.cpp
--- a/Engine/Source/Runtime/Experimental/Iris/Core/Private/Iris/ReplicationSystem/Filtering/ReplicationFiltering.cpp
+++ b/Engine/Source/Runtime/Experimental/Iris/Core/Private/Iris/ReplicationSystem/Filtering/ReplicationFiltering.cpp
@@ -432,7 +432,7 @@ void FReplicationFiltering::SetOwningConnection(FInternalNetRefIndex ObjectIndex
 	}
 }
 
-bool FReplicationFiltering::SetFilter(FInternalNetRefIndex ObjectIndex, FNetObjectFilterHandle Filter)
+bool FReplicationFiltering::SetFilter(FInternalNetRefIndex ObjectIndex, FNetObjectFilterHandle Filter, FName FilterConfigProfile)
 {
 	if (Filter == ConnectionFilterHandle)
 	{
@@ -440,7 +440,7 @@ bool FReplicationFiltering::SetFilter(FInternalNetRefIndex ObjectIndex, FNetObje
 		return false;
 	}
 
-	UE_LOG(LogIrisFiltering, Verbose, TEXT("Setting filter %s to %s"), *GetFilterName(Filter).ToString(), *NetRefHandleManager->PrintObjectFromIndex(ObjectIndex));
+	UE_LOG(LogIrisFiltering, Verbose, TEXT("Setting filter %s to %s (profile %s)"), *GetFilterName(Filter).ToString(), *NetRefHandleManager->PrintObjectFromIndex(ObjectIndex), *FilterConfigProfile.ToString());
 
 	const bool bWantsToUseDynamicFilter = FNetObjectFilterHandleUtil::IsDynamicFilter(Filter);
 	const uint8 OldDynamicFilterIndex = ObjectIndexToDynamicFilterIndex[ObjectIndex];
@@ -477,11 +477,11 @@ bool FReplicationFiltering::SetFilter(FInternalNetRefIndex ObjectIndex, FNetObje
 		this->FreePerObjectInfoForObject(ObjIndex);
 	};
 
-	auto TrySetDynamicFilter = [this, &ObjectData, &ReplicatedObjectsStateBuffers](uint32 ObjIndex, uint32 FilterIndex)
+	auto TrySetDynamicFilter = [this, &ObjectData, &ReplicatedObjectsStateBuffers, FilterConfigProfile](uint32 ObjIndex, uint32 FilterIndex)
 	{
 		FNetObjectFilteringInfo& NetObjectFilteringInfo = this->NetObjectFilteringInfos[ObjIndex];
 		NetObjectFilteringInfo = {};
-		FNetObjectFilterAddObjectParams AddParams = { NetObjectFilteringInfo, ObjectData.InstanceProtocol, ObjectData.Protocol, ReplicatedObjectsStateBuffers[ObjIndex] };
+		FNetObjectFilterAddObjectParams AddParams = { .OutInfo=NetObjectFilteringInfo, .ProfileName=FilterConfigProfile, .InstanceProtocol=ObjectData.InstanceProtocol, .Protocol=ObjectData.Protocol, .StateBuffer=ReplicatedObjectsStateBuffers[ObjIndex] };
 		FFilterInfo& FilterInfo = this->DynamicFilterInfos[FilterIndex];
 		if (FilterInfo.Filter->AddObject(ObjIndex, AddParams))
 		{
diff --git a/Engine/Source/Runtime/Experimental/Iris/Core/Private/Iris/ReplicationSystem/Filtering/ReplicationFiltering.h b/Engine/Source/Runtime/Experimental/Iris/Core/Private/Iris/ReplicationSystem/Filtering/ReplicationFiltering.h
--- a/Engine/Source/Runtime/Experimental/Iris/Core/Private/Iris/ReplicationSystem/Filtering/ReplicationFiltering.h
+++ b/Engine/Source/Runtime/Experimental/Iris/Core/Private/Iris/ReplicationSystem/Filtering/ReplicationFiltering.h
@@ -85,7 +85,18 @@ public:
 	void SetOwningConnection(FInternalNetRefIndex ObjectIndex, uint32 ConnectionId);
 	uint32 GetOwningConnection(FInternalNetRefIndex ObjectIndex) const { return !bHasDirtyOwner ? ObjectIndexToOwningConnection[ObjectIndex] : GetOwningConnectionIfDirty(ObjectIndex); }
 
-	bool SetFilter(FInternalNetRefIndex ObjectIndex, FNetObjectFilterHandle Filter);
+	/**
+	 * Setup an object to be filtered by the passed filter handle.
+	 * @param ObjectIndex The object that wants to be filtered
+	 * @param Filter The handle to the filter to add the object into
+	 * @param FilterConfigProfile Optional profile name that can be used to specialize the filter parameters
+	 */
+	bool SetFilter(FInternalNetRefIndex ObjectIndex, FNetObjectFilterHandle Filter, FName FilterConfigProfile);
+
+	/** 
+	 * Returns true if the object is part of a filter with the ENetFilterTraits::Spatial trait. 
+	 * Indicates that the object is location filtered.
+	 */
 	bool IsUsingSpatialFilter(FInternalNetRefIndex ObjectIndex) const;
 
 	// Set whether an object is allowed to be replicated to certain connections or not.
diff --git a/Engine/Source/Runtime/Experimental/Iris/Core/Private/Iris/ReplicationSystem/ObjectReplicationBridge.cpp b/Engine/Source/Runtime/Experimental/Iris/Core/Private/Iris/ReplicationSystem/ObjectReplicationBridge.cpp
--- a/Engine/Source/Runtime/Experimental/Iris/Core/Private/Iris/ReplicationSystem/ObjectReplicationBridge.cpp
+++ b/Engine/Source/Runtime/Experimental/Iris/Core/Private/Iris/ReplicationSystem/ObjectReplicationBridge.cpp
@@ -412,6 +412,7 @@ void UObjectReplicationBridge::AssignDynamicFilter(UObject* Instance, const FCre
 	}
 
 	FNetObjectFilterHandle FilterHandle = InvalidNetObjectFilterHandle;
+	FName FilterConfigProfile;
 
 	if (Params.bUseExplicitDynamicFilter)
 	{
@@ -426,13 +427,13 @@ void UObjectReplicationBridge::AssignDynamicFilter(UObject* Instance, const FCre
 	else if (Params.bUseClassConfigDynamicFilter)
 	{
 		constexpr bool bRequireForceEnabled = false;
-		FilterHandle = GetDynamicFilter(Instance->GetClass(), bRequireForceEnabled);
+		FilterHandle = GetDynamicFilter(Instance->GetClass(), bRequireForceEnabled, FilterConfigProfile);
 	}
 
 	if (FilterHandle != InvalidNetObjectFilterHandle)
 	{
-		UE_LOG_OBJECTREPLICATIONBRIDGE(Verbose, TEXT("BeginReplication Filter: %s will be used for %s "), *(ReplicationSystem->GetFilterName(FilterHandle).ToString()), *NetRefHandleManager->PrintObjectFromNetRefHandle(RefHandle));
-		ReplicationSystem->SetFilter(RefHandle, FilterHandle);
+		UE_LOG_OBJECTREPLICATIONBRIDGE(Verbose, TEXT("BeginReplication Filter: %s will be used for %s. (FilterProfile: %s)"), *(ReplicationSystem->GetFilterName(FilterHandle).ToString()), *NetRefHandleManager->PrintObjectFromNetRefHandle(RefHandle), *FilterConfigProfile.ToString());
+		ReplicationSystem->SetFilter(RefHandle, FilterHandle, FilterConfigProfile);
 	}
 	
 }
@@ -1525,7 +1526,7 @@ void UObjectReplicationBridge::SetClassTypeStatsConfig(const FString& ClassPathN
 	SetClassTypeStatsConfig(FName(ClassPathName), FName(TypeStatsName));
 }
 
-void UObjectReplicationBridge::SetClassDynamicFilterConfig(FName ClassPathName, const UE::Net::FNetObjectFilterHandle FilterHandle)
+void UObjectReplicationBridge::SetClassDynamicFilterConfig(FName ClassPathName, const UE::Net::FNetObjectFilterHandle FilterHandle, FName FilterProfile)
 {
 	if (ClassPathName.IsNone())
 	{
@@ -1538,8 +1539,8 @@ void UObjectReplicationBridge::SetClassDynamicFilterConfig(FName ClassPathName,
 		{
 			if (OldFilterInfo->FilterHandle != FilterHandle)
 			{
-				UE_LOG(LogIrisFilterConfig, Log, TEXT("SetClassDynamicFilterConfig assigned %s to use filter %s. Previously using filter %s."),
-					*ClassPathName.ToString(), *(ReplicationSystem->GetFilterName(FilterHandle).ToString()), *(ReplicationSystem->GetFilterName(OldFilterInfo->FilterHandle).ToString()));
+				UE_LOG(LogIrisFilterConfig, Log, TEXT("SetClassDynamicFilterConfig assigned %s to use filter %s (Profile %s). Previously using filter %s."),
+					*ClassPathName.ToString(), *(ReplicationSystem->GetFilterName(FilterHandle).ToString()), *FilterProfile.ToString(), *(ReplicationSystem->GetFilterName(OldFilterInfo->FilterHandle).ToString()));
 			}
 			else
 			{
@@ -1548,17 +1549,18 @@ void UObjectReplicationBridge::SetClassDynamicFilterConfig(FName ClassPathName,
 		}
 		else
 		{
-			UE_LOG(LogIrisFilterConfig, Log, TEXT("SetClassDynamicFilterConfig assigned %s to use filter %s."), *ClassPathName.ToString(), *(ReplicationSystem->GetFilterName(FilterHandle).ToString()));
+			UE_LOG(LogIrisFilterConfig, Log, TEXT("SetClassDynamicFilterConfig assigned %s to use filter %s (Profile %s)."), *ClassPathName.ToString(), *(ReplicationSystem->GetFilterName(FilterHandle).ToString()), *FilterProfile.ToString());
 		}
 	}
 
 	FClassFilterInfo FilterInfo;
 	FilterInfo.FilterHandle = FilterHandle;
+	FilterInfo.FilterProfile = FilterProfile;
 	FilterInfo.bForceEnable = false;
 	ClassesWithDynamicFilter.Add(ClassPathName, FilterInfo);
 }
 
-void UObjectReplicationBridge::SetClassDynamicFilterConfig(FName ClassPathName, FName FilterName)
+void UObjectReplicationBridge::SetClassDynamicFilterConfig(FName ClassPathName, FName FilterName, FName FilterProfile)
 {
 	using namespace UE::Net;
 
@@ -1573,72 +1575,90 @@ void UObjectReplicationBridge::SetClassDynamicFilterConfig(FName ClassPathName,
 
 		if (ensureMsgf(FilterHandle != InvalidNetObjectFilterHandle, TEXT("SetClassDynamicFilterConfig for %s received invalid filter named %s"), *ClassPathName.ToString(), *FilterName.ToString()))
 		{
-			SetClassDynamicFilterConfig(ClassPathName, FilterHandle);
+			SetClassDynamicFilterConfig(ClassPathName, FilterHandle, FilterProfile);
 		}
 	}
 	else
 	{
 		// Reset the filter so the class does not get assigned a dynamic filter anymore.
-		SetClassDynamicFilterConfig(ClassPathName, InvalidNetObjectFilterHandle);
+		SetClassDynamicFilterConfig(ClassPathName, InvalidNetObjectFilterHandle, FilterProfile);
 	}
 }
 
-UE::Net::FNetObjectFilterHandle UObjectReplicationBridge::GetDynamicFilter(const UClass* Class, bool bRequireForceEnabled)
+UE::Net::FNetObjectFilterHandle UObjectReplicationBridge::GetDynamicFilter(const UClass* Class, bool bRequireForceEnabled, FName& OutFilterProfile)
 {
 	using namespace UE::Net;
 
-	if (ClassesWithDynamicFilter.Num() > 0)
+	if (ClassesWithDynamicFilter.IsEmpty())
 	{
-		const FName ClassName = GetConfigClassPathName(Class);
+		/**
+		 * For the cases when there are no configured filter mappings we just check whether to use a spatial filter or not.
+		 * We don't add anything to the filter mapping.
+		 */
+		const FNetObjectFilterHandle FilterHandle = ShouldUseDefaultSpatialFilterFunction(Class) ? DefaultSpatialFilterHandle : InvalidNetObjectFilterHandle;
+		return FilterHandle;
+	}
 
-		// Try exact match first.
-		if (FClassFilterInfo* FilterInfoPtr = ClassesWithDynamicFilter.Find(ClassName))
+	const FName ClassName = GetConfigClassPathName(Class);
+
+	// Try exact match first.
+	if (FClassFilterInfo* FilterInfoPtr = ClassesWithDynamicFilter.Find(ClassName))
+	{
+		const bool bUseFilter = !bRequireForceEnabled || FilterInfoPtr->bForceEnable;
+		if (bUseFilter)
 		{
-			const bool bUseFilter = !bRequireForceEnabled || FilterInfoPtr->bForceEnable;
-			return bUseFilter ? FilterInfoPtr->FilterHandle : InvalidNetObjectFilterHandle;
+			OutFilterProfile = FilterInfoPtr->FilterProfile;
+			return FilterInfoPtr->FilterHandle;
 		}
-
-		/**
-		 * Try to find superclass. If we find it and the classes are considered equal we copy the filter setting.
-		 * If it's not equal we check whether it can be spatialized or not and use the result of that.
-		 * In all cases we add the result to the mapping for faster lookup next time.
-		 */
-		for (const UClass* SuperClass = Class->GetSuperClass(); SuperClass != nullptr; SuperClass = SuperClass->GetSuperClass())
+		else
 		{
-			const FName SuperClassName = GetConfigClassPathName(SuperClass);
+			return InvalidNetObjectFilterHandle;
+		}
+	}
+
+	/**
+	* Try to find superclass. If we find it and the classes are considered equal we copy the filter setting.
+	* If it's not equal we check whether it can be spatialized or not and use the result of that.
+	* In all cases we add the result to the mapping for faster lookup next time.
+	*/
+	for (const UClass* SuperClass = Class->GetSuperClass(); SuperClass != nullptr; SuperClass = SuperClass->GetSuperClass())
+	{
+		const FName SuperClassName = GetConfigClassPathName(SuperClass);
 
-			// Try to get exact match first.
-			if (FClassFilterInfo* FilterInfoPtr = ClassesWithDynamicFilter.Find(SuperClassName))
+		// Try to get exact match first.
+		if (FClassFilterInfo* FilterInfoPtr = ClassesWithDynamicFilter.Find(SuperClassName))
+		{
+			if (ShouldSubclassUseSameFilterFunction(SuperClass, Class))
 			{
-				if (ShouldSubclassUseSameFilterFunction(SuperClass, Class))
-				{
-					FClassFilterInfo FilterInfo = *FilterInfoPtr;
-					ClassesWithDynamicFilter.Add(ClassName, FilterInfo);
+				FClassFilterInfo FilterInfo = *FilterInfoPtr;
+				ClassesWithDynamicFilter.Add(ClassName, FilterInfo);
+
+				const bool bUseFilter = !bRequireForceEnabled || FilterInfo.bForceEnable;
 
-					const bool bUseFilter = !bRequireForceEnabled || FilterInfo.bForceEnable;
-					return bUseFilter ? FilterInfo.FilterHandle : InvalidNetObjectFilterHandle;
+				if (bUseFilter)
+				{
+					OutFilterProfile = FilterInfoPtr->FilterProfile;
+					return FilterInfo.FilterHandle;
+				}
+				else
+				{
+					return InvalidNetObjectFilterHandle;
 				}
+					
+			}
 
-				// Here's a good place to put a line of code and set a breakpoint to debug inheritance issues.
+			// Here's a good place to put a line of code and set a breakpoint to debug inheritance issues.
 
-				break;
-			}
+			break;
 		}
-
-		// Either super class wasn't found or it wasn't considered equal. Let's add a new filter mapping.
-		FClassFilterInfo FilterInfo;
-		FilterInfo.FilterHandle = ShouldUseDefaultSpatialFilterFunction(Class) ? DefaultSpatialFilterHandle : InvalidNetObjectFilterHandle;
-		FilterInfo.bForceEnable = false;
-		ClassesWithDynamicFilter.Add(ClassName, FilterInfo);
-		return FilterInfo.FilterHandle;
 	}
 
-	/**
-	 * For the cases when there are no configured filter mappings we just check whether to use a spatial filter or not.
-	 * We don't add anything to the filter mapping.
-	 */
-	const FNetObjectFilterHandle FilterHandle = ShouldUseDefaultSpatialFilterFunction(Class) ? DefaultSpatialFilterHandle : InvalidNetObjectFilterHandle;
-	return FilterHandle;
+	// Either super class wasn't found or it wasn't considered equal. Let's add a new filter mapping.
+	FClassFilterInfo FilterInfo;
+	FilterInfo.FilterHandle = ShouldUseDefaultSpatialFilterFunction(Class) ? DefaultSpatialFilterHandle : InvalidNetObjectFilterHandle;
+	FilterInfo.bForceEnable = false;
+	ClassesWithDynamicFilter.Add(ClassName, FilterInfo);
+	return FilterInfo.FilterHandle;
 }
 
 UE::Net::FNetObjectPrioritizerHandle UObjectReplicationBridge::GetPrioritizer(const UClass* Class, bool bRequireForceEnabled)
@@ -1737,6 +1757,7 @@ void UObjectReplicationBridge::LoadConfig()
 		{
 			FClassFilterInfo FilterInfo;
 			FilterInfo.FilterHandle = ReplicationSystem->GetFilterHandle(FilterConfig.DynamicFilterName);
+			FilterInfo.FilterProfile = FilterConfig.FilterProfile;
 			FilterInfo.bForceEnable = FilterConfig.bForceEnableOnAllInstances;
 			ClassesWithDynamicFilter.Add(FilterConfig.ClassName, FilterInfo);
 		}
diff --git a/Engine/Source/Runtime/Experimental/Iris/Core/Private/Iris/ReplicationSystem/ReplicationSystem.cpp b/Engine/Source/Runtime/Experimental/Iris/Core/Private/Iris/ReplicationSystem/ReplicationSystem.cpp
--- a/Engine/Source/Runtime/Experimental/Iris/Core/Private/Iris/ReplicationSystem/ReplicationSystem.cpp
+++ b/Engine/Source/Runtime/Experimental/Iris/Core/Private/Iris/ReplicationSystem/ReplicationSystem.cpp
@@ -1089,7 +1089,7 @@ uint32 UReplicationSystem::GetOwningNetConnection(FNetRefHandle Handle) const
 	return Filtering.GetOwningConnection(ObjectInternalIndex);
 }
 
-bool UReplicationSystem::SetFilter(FNetRefHandle Handle, UE::Net::FNetObjectFilterHandle Filter)
+bool UReplicationSystem::SetFilter(FNetRefHandle Handle, UE::Net::FNetObjectFilterHandle Filter, FName FilterConfigProfile /*= NAME_None*/)
 {
 	using namespace UE::Net::Private;
 
@@ -1107,7 +1107,7 @@ bool UReplicationSystem::SetFilter(FNetRefHandle Handle, UE::Net::FNetObjectFilt
 	}
 
 	FReplicationFiltering& Filtering = Impl->ReplicationSystemInternal.GetFiltering();
-	return Filtering.SetFilter(ObjectInternalIndex, Filter);
+	return Filtering.SetFilter(ObjectInternalIndex, Filter, FilterConfigProfile);
 }
 
 UE::Net::FNetObjectFilterHandle UReplicationSystem::GetFilterHandle(const FName FilterName) const
diff --git a/Engine/Source/Runtime/Experimental/Iris/Core/Public/Iris/ReplicationSystem/Filtering/NetObjectFilter.h b/Engine/Source/Runtime/Experimental/Iris/Core/Public/Iris/ReplicationSystem/Filtering/NetObjectFilter.h
--- a/Engine/Source/Runtime/Experimental/Iris/Core/Public/Iris/ReplicationSystem/Filtering/NetObjectFilter.h
+++ b/Engine/Source/Runtime/Experimental/Iris/Core/Public/Iris/ReplicationSystem/Filtering/NetObjectFilter.h
@@ -167,6 +167,9 @@ struct FNetObjectFilterAddObjectParams
 	/** The info is zeroed before the AddObject() call. Fill in with filter specifics, like offsets to tags. */
 	FNetObjectFilteringInfo& OutInfo;
 
+	/** Name of a specialized configuration profile. When none, the default settings are expected. */
+	FName ProfileName;
+
 	/** The FReplicationInstanceProtocol which describes the source state data. */
 	const UE::Net::FReplicationInstanceProtocol* InstanceProtocol;
 
diff --git a/Engine/Source/Runtime/Experimental/Iris/Core/Public/Iris/ReplicationSystem/Filtering/NetObjectGridFilter.h b/Engine/Source/Runtime/Experimental/Iris/Core/Public/Iris/ReplicationSystem/Filtering/NetObjectGridFilter.h
--- a/Engine/Source/Runtime/Experimental/Iris/Core/Public/Iris/ReplicationSystem/Filtering/NetObjectGridFilter.h
+++ b/Engine/Source/Runtime/Experimental/Iris/Core/Public/Iris/ReplicationSystem/Filtering/NetObjectGridFilter.h
@@ -22,16 +22,45 @@ namespace UE::Net
 	}
 }
 
+/**
+ * Specialized template that configures unique properties.
+ * Useful when you need to specialize a behavior per class or object type
+ */
+USTRUCT()
+struct FNetObjectGridFilterProfile
+{
+	GENERATED_BODY()
+
+	/** The config name used to map to this profile */
+	UPROPERTY()
+	FName FilterProfileName;
+
+	/** Number of frames we keep the object relevant until it is officially culled out.*/
+	UPROPERTY()
+	uint16 FrameCountBeforeCulling = 4;
+
+	bool operator==(FName Key) const { return FilterProfileName == Key; }
+};
+
+/**
+ * Common settings used to configure how the GridFilter behaves
+ */
 UCLASS(transient, config=Engine, MinimalAPI)
 class UNetObjectGridFilterConfig : public UNetObjectFilterConfig
 {
 	GENERATED_BODY()
 
 public:
-	/** How many frames a view position should be considered relevant. To avoid culling issues when player borders cells. */
+	/** 
+	 * How many frames a previous grid cell should continue to be considered relevant. To avoid culling issues when player borders cells. 
+	 * Only used when bUseExactCullDistance is false.
+	 */
 	UPROPERTY(Config)
 	uint32 ViewPosRelevancyFrameCount = 2;
 
+	UPROPERTY(Config)
+	uint16 DefaultFrameCountBeforeCulling = 4;
+
 	UPROPERTY(Config)
 	float CellSizeX = 20000.0f;
 
@@ -54,9 +83,16 @@ public:
 	UPROPERTY(Config)
 	FVector MaxPos = {+0.5f*2097152.0f, +0.5f*2097152.0f, +0.5f*2097152.0f};
 
-	/** If true, use the exact cull distance to determine the objects to create/destroy on the client. Otherwise only use the net grid. */
+	/** 
+	 * If true: use the exact distance between an object and the viewer to determine if the object is relevant or should be culled out.
+	 * When false: consider all objects within a grid cell to be relevant when a viewer is located within the cell. This can extend the relevant distance of objects beyond their cull distance. 
+	 */
 	UPROPERTY(Config)
 	bool bUseExactCullDistance = true;
+
+	/** Map of specialized configuration profiles */
+	UPROPERTY(Config)
+	TArray<FNetObjectGridFilterProfile> FilterProfiles;
 };
 
 UCLASS(abstract)
@@ -79,17 +115,25 @@ protected:
 protected:
 	struct FObjectLocationInfo : public FNetObjectFilteringInfo
 	{
-		bool IsUsingWorldLocations() const { return GetLocationStateIndex() == InvalidStateIndex; }
+		bool IsUsingWorldLocations() const  { return GetLocationStateIndex() == InvalidStateIndex; }
 		bool IsUsingLocationInState() const { return GetLocationStateIndex() != InvalidStateIndex; }
 
-		void SetLocationStateOffset(uint16 Offset) { Data[0] = Offset; }
-		uint16 GetLocationStateOffset() const { return Data[0]; }
+		/**
+		* Data mapping:
+		* uint16 Data[0] = LocationState offset
+		* uint16 Data[1] = LocationState index
+		* uint16 Data[2] = FPerObjectInfo index (low bytes)
+		* uint16 Data[3] = FPerObjectInfo index (high bytes)
+		*/
+
+		void SetLocationStateOffset(uint16 Offset)	{ Data[0] = Offset; }
+		uint16 GetLocationStateOffset() const		{ return Data[0]; }
 
 		void SetLocationStateIndex(uint16 Index) { Data[1] = Index; }
-		uint16 GetLocationStateIndex() const { return Data[1]; }
+		uint16 GetLocationStateIndex() const	 { return Data[1]; }
 
-		void SetInfoIndex(uint32 Index) { Data[2] = Index & 65535U; Data[3] = Index >> 16U; }
-		uint32 GetInfoIndex() const { return (uint32(Data[3]) << 16U) | uint32(Data[2]); }
+		void SetInfoIndex(uint32 Index)	{ Data[2] = Index & 65535U; Data[3] = Index >> 16U; }
+		uint32 GetInfoIndex() const		{ return (uint32(Data[3]) << 16U) | uint32(Data[2]); }
 	};
 
 	IRISCORE_API void AddCellInfoForObject(const FObjectLocationInfo& ObjectInfo, const UE::Net::FReplicationInstanceProtocol* InstanceProtocol);
@@ -107,10 +151,10 @@ protected:
 
 	struct FCellBox
 	{
-		int32 MinX;
-		int32 MaxX;
-		int32 MinY;
-		int32 MaxY;
+		int32 MinX = 0;
+		int32 MaxX = 0;
+		int32 MinY = 0;
+		int32 MaxY = 0;
 
 		bool operator==(const FCellBox&) const;
 		bool operator!=(const FCellBox&) const;
@@ -119,9 +163,10 @@ protected:
 	// We can't fit all info we need in 4x16bits.
 	struct FPerObjectInfo
 	{
-		FVector Position = { 0.f,0.f,0.f };
+		FVector Position = FVector::ZeroVector;
 		FCellBox CellBox = {};
 		uint32 ObjectIndex = 0U;
+		uint16 FrameCountBeforeCulling = 0U;
 
 		float GetCullDistance() const
 		{
@@ -201,7 +246,7 @@ private:
 		TArray<FCellAndTimestamp, TInlineAllocator<32>> RecentCells;
 		
 		// Objects that have been recently visible to the connection and the frame countdown.
-		TMap<uint32, uint32> RecentObjectFrameCount;
+		TMap<uint32, uint16> RecentObjectFrameCount;
 	};
 
 	/** Aggregator for stats */
@@ -224,6 +269,8 @@ private:
 	void CalculateCellBox(const FPerObjectInfo& PerObjectInfo, FCellBox& OutCellBox);
 	void CalculateCellCoord(FCellCoord& OutCoord, const FVector& Pos);
 
+	uint16 GetFrameCountBeforeCulling(FName ProfileName) const;
+
 	static bool AreCellsDisjoint(const FCellBox& A, const FCellBox& B);
 	static bool DoesCellContainCoord(const FCellBox& Cell, const FCellCoord& Coord);
 
diff --git a/Engine/Source/Runtime/Experimental/Iris/Core/Public/Iris/ReplicationSystem/ObjectReplicationBridge.h b/Engine/Source/Runtime/Experimental/Iris/Core/Public/Iris/ReplicationSystem/ObjectReplicationBridge.h
--- a/Engine/Source/Runtime/Experimental/Iris/Core/Public/Iris/ReplicationSystem/ObjectReplicationBridge.h
+++ b/Engine/Source/Runtime/Experimental/Iris/Core/Public/Iris/ReplicationSystem/ObjectReplicationBridge.h
@@ -172,8 +172,8 @@ public:
 	IRISCORE_API void SetPollFrequency(FNetRefHandle RootHandle, float PollFrequency);
 
 	/** Set the object filter to use for objects of this class and any derived classes without an explicit config. */
-	IRISCORE_API void SetClassDynamicFilterConfig(FName ClassPathName, const UE::Net::FNetObjectFilterHandle FilterHandle);
-	IRISCORE_API void SetClassDynamicFilterConfig(FName ClassPathName, FName FilterName);
+	IRISCORE_API void SetClassDynamicFilterConfig(FName ClassPathName, const UE::Net::FNetObjectFilterHandle FilterHandle, FName FilterProfile=NAME_None);
+	IRISCORE_API void SetClassDynamicFilterConfig(FName ClassPathName, FName FilterName, FName FilterProfile=NAME_None);
 
 	/** Set the TypeStats to use for specified class and any derived classes without explicit config */
 	IRISCORE_API void SetClassTypeStatsConfig(FName ClassPathName, FName TypeStatsName);
@@ -353,7 +353,7 @@ private:
 
 protected:
 	/** Retrieves the dynamic filter to set for the given class. Will return an invalid handle if no dynamic filter should be set. */
-	IRISCORE_API UE::Net::FNetObjectFilterHandle GetDynamicFilter(const UClass* Class, bool bRequireForceEnabled);
+	IRISCORE_API UE::Net::FNetObjectFilterHandle GetDynamicFilter(const UClass* Class, bool bRequireForceEnabled, FName& OutFilterProfile);
 
 private:
 	/** Retrieves the prioritizer to set for the given class. If bRequireForceEnabled the config needs to have bForceEnableOnAllInstances set in order for this method to return the configured prioritizer. Returns an invalid handle if no prioritizer should be set. */
@@ -395,6 +395,7 @@ private:
 	struct FClassFilterInfo
 	{
 		UE::Net::FNetObjectFilterHandle FilterHandle;
+		FName FilterProfile;
 		bool bForceEnable = false;
 	};
 
diff --git a/Engine/Source/Runtime/Experimental/Iris/Core/Public/Iris/ReplicationSystem/ObjectReplicationBridgeConfig.h b/Engine/Source/Runtime/Experimental/Iris/Core/Public/Iris/ReplicationSystem/ObjectReplicationBridgeConfig.h
--- a/Engine/Source/Runtime/Experimental/Iris/Core/Public/Iris/ReplicationSystem/ObjectReplicationBridgeConfig.h
+++ b/Engine/Source/Runtime/Experimental/Iris/Core/Public/Iris/ReplicationSystem/ObjectReplicationBridgeConfig.h
@@ -47,6 +47,10 @@ struct FObjectReplicationBridgeFilterConfig
 	UPROPERTY()
 	FName DynamicFilterName;
 
+	/** Optional name to a profile that can further specialize the settings within a dynamic filter */
+	UPROPERTY()
+	FName FilterProfile;
+
 	/** Whether this filter should be used for all instances of this class and subclasses, regardless of bAlwaysRelevant and bOnlyRelevantToOwner settings on instance. */
 	UPROPERTY()
 	bool bForceEnableOnAllInstances = false;
diff --git a/Engine/Source/Runtime/Experimental/Iris/Core/Public/Iris/ReplicationSystem/ReplicationSystem.h b/Engine/Source/Runtime/Experimental/Iris/Core/Public/Iris/ReplicationSystem/ReplicationSystem.h
--- a/Engine/Source/Runtime/Experimental/Iris/Core/Public/Iris/ReplicationSystem/ReplicationSystem.h
+++ b/Engine/Source/Runtime/Experimental/Iris/Core/Public/Iris/ReplicationSystem/ReplicationSystem.h
@@ -380,16 +380,20 @@ public:
 	 * Sets a filter for a replicated object which will be used until the next call to SetFilter or SetConnectionFilter.
 	 * Filters are used prevent objects from being replicated to certain connections. An object that is filtered
 	 * out will cause the object to be destroyed on the remote side.
+	 * 
 	 * @param Handle A valid handle to an object.
 	 * @param FilterHandle A valid handle to a filter, retrieved via a call to GetFilter or one of the two special handles
-	 * InvalidNetObjectFilterHandle to clear filtering or ToOwnerFilterHandle for owner filtering.
+	 *		  InvalidNetObjectFilterHandle to clear filtering or ToOwnerFilterHandle for owner filtering.
+	 * @param FilterConfigProfile Optional name of a specialized profile to use as the object's configuration. When none filters are expected to use default settings.
+	 * 
 	 * @return true if the filter was successfully set and false if it was not. It can fail for various reasons,
 	 * such as the filter not supporting the object in question for implementation defined reasons. If the function fails
 	 * the filter of the object is unspecified, it could be using a previous filter or use no filtering.
+	 * 
 	 * @see GetFilterHandle
 	 * @see SetConnectionFilter
 	 */
-	IRISCORE_API bool SetFilter(FNetRefHandle Handle, UE::Net::FNetObjectFilterHandle FilterHandle);
+	IRISCORE_API bool SetFilter(FNetRefHandle Handle, UE::Net::FNetObjectFilterHandle FilterHandle, FName FilterConfigProfile=NAME_None);
 
 	/**
 	 * Gets the handle for a filter with a given name. The handle can be used in subsequent calls to SetFilter.
