Fix potential flicker when marking render state dirty on geometry collection component.
Issue is that geometry collection scene proxies were always created at the rest transform state, and then the current transforms were set at the first update. This would could cause a reset flicker.
Also the first update wouldn't happen if the component was at rest, so visually it would go back to the rest state and stay there.
The fix is to get the current DynamicData during scene proxy construction.

diff --git a/Engine/Source/Runtime/Experimental/GeometryCollectionEngine/Private/GeometryCollection/GeometryCollectionSceneProxy.cpp b/Engine/Source/Runtime/Experimental/GeometryCollectionEngine/Private/GeometryCollection/GeometryCollectionSceneProxy.cpp
--- a/Engine/Source/Runtime/Experimental/GeometryCollectionEngine/Private/GeometryCollection/GeometryCollectionSceneProxy.cpp
+++ b/Engine/Source/Runtime/Experimental/GeometryCollectionEngine/Private/GeometryCollection/GeometryCollectionSceneProxy.cpp
@@ -211,6 +211,8 @@ FGeometryCollectionSceneProxy::FGeometryCollectionSceneProxy(UGeometryCollection
 	// changed from the prev to curr frame, but this is expensive.  We should revisit this if the draw calls for velocity
 	// rendering become a problem. One solution could be to use internal solver sleeping state to drive motion blur.
 	bAlwaysHasVelocity = true;
+
+	DynamicData = Component->InitDynamicData(true);
 }
 
 FGeometryCollectionSceneProxy::~FGeometryCollectionSceneProxy()
@@ -302,16 +304,16 @@ void FGeometryCollectionSceneProxy::CreateRenderThreadResources(FRHICommandListB
 		const EResourceLockMode LockMode = bLocalGeometryCollectionTripleBufferUploads ? RLM_WriteOnly_NoOverwrite : RLM_WriteOnly;
 
 		FGeometryCollectionTransformBuffer& TransformBuffer = GetCurrentTransformBuffer();
-		TransformBuffer.UpdateDynamicData(RHICmdList, RestTransforms, LockMode);
+		TransformBuffer.UpdateDynamicData(RHICmdList, DynamicData->Transforms, LockMode);
 		FGeometryCollectionTransformBuffer& PrevTransformBuffer = GetCurrentPrevTransformBuffer();
-		PrevTransformBuffer.UpdateDynamicData(RHICmdList, RestTransforms, LockMode);
+		PrevTransformBuffer.UpdateDynamicData(RHICmdList, DynamicData->PrevTransforms, LockMode);
 	}
 	else
 	{
 		// Initialize CPU skinning buffer with rest transforms.
 		SkinnedPositionVertexBuffer.Init(MeshResource.PositionVertexBuffer.GetNumVertices(), false);
 		SkinnedPositionVertexBuffer.InitResource(RHICmdList);
-		UpdateSkinnedPositions(RHICmdList, RestTransforms);
+		UpdateSkinnedPositions(RHICmdList, DynamicData->Transforms);
 	}
 
 	SetupVertexFactory(RHICmdList, VertexFactory);
@@ -1133,10 +1135,7 @@ FNaniteGeometryCollectionSceneProxy::FNaniteGeometryCollectionSceneProxy(UGeomet
 	TArray<FMatrix44f> RestTransforms;
 	Component->GetRestTransforms(RestTransforms);
 
-	FGeometryCollectionDynamicData* DynamicData = GDynamicDataPool.Allocate();
-	DynamicData->IsDynamic = true;
-	DynamicData->Transforms = RestTransforms;
-	DynamicData->PrevTransforms = RestTransforms;
+	FGeometryCollectionDynamicData* DynamicData = Component->InitDynamicData(true);
 	SetDynamicData_RenderThread(DynamicData, Component->GetRenderMatrix());
 }
 
diff --git a/Engine/Source/Runtime/Experimental/GeometryCollectionEngine/Public/GeometryCollection/GeometryCollectionComponent.h b/Engine/Source/Runtime/Experimental/GeometryCollectionEngine/Public/GeometryCollection/GeometryCollectionComponent.h
--- a/Engine/Source/Runtime/Experimental/GeometryCollectionEngine/Public/GeometryCollection/GeometryCollectionComponent.h
+++ b/Engine/Source/Runtime/Experimental/GeometryCollectionEngine/Public/GeometryCollection/GeometryCollectionComponent.h
@@ -590,6 +590,8 @@ class UGeometryCollectionComponent : public UMeshComponent, public IChaosNotifyH
 	friend class FScopedColorEdit;
 #endif
 	friend class FGeometryCollectionCommands;
+	friend class FGeometryCollectionSceneProxy;
+	friend class FNaniteGeometryCollectionSceneProxy;
 
 public:
 
