Fix linker error: There's no need to dll export TSimModuleSettings as the class is templated and will be instanciated on use. The export confuses the linker.
Also make single argument constructor explicit while in there

diff --git a/Engine/Source/Runtime/Experimental/ChaosVehicles/ChaosVehiclesCore/Public/SimModule/SimulationModuleBase.h b/Engine/Source/Runtime/Experimental/ChaosVehicles/ChaosVehiclesCore/Public/SimModule/SimulationModuleBase.h
--- a/Engine/Source/Runtime/Experimental/ChaosVehicles/ChaosVehiclesCore/Public/SimModule/SimulationModuleBase.h
+++ b/Engine/Source/Runtime/Experimental/ChaosVehicles/ChaosVehiclesCore/Public/SimModule/SimulationModuleBase.h
@@ -99,11 +99,11 @@ namespace Chaos
 	 * Code common between all simulation building blocks settings
 	 */
 	template <typename T>
-	class CHAOSVEHICLESCORE_API TSimModuleSettings
+	class TSimModuleSettings
 	{
 	public:
 
-		TSimModuleSettings(const T& SetupIn) : SetupData(SetupIn)
+		explicit TSimModuleSettings(const T& SetupIn) : SetupData(SetupIn)
 		{
 			SetupData = SetupIn; // deliberate copy for now
 		}
