Dataflow : make sure DataflowInputOutput header only forward declare FDataflowNode type
(this avoid potential circular header dependency because of the template metod used in FDataflowNode and FDataflowInput / output classes

diff --git a/Engine/Source/Runtime/Experimental/Dataflow/Core/Private/Dataflow/DataflowConnection.cpp b/Engine/Source/Runtime/Experimental/Dataflow/Core/Private/Dataflow/DataflowConnection.cpp
--- a/Engine/Source/Runtime/Experimental/Dataflow/Core/Private/Dataflow/DataflowConnection.cpp
+++ b/Engine/Source/Runtime/Experimental/Dataflow/Core/Private/Dataflow/DataflowConnection.cpp
@@ -26,3 +26,18 @@ uint32 FDataflowConnection::GetOffset() const
 }
 
 
+bool FDataflowConnection::IsOwningNodeEnabled() const
+{
+	return (OwningNode && OwningNode->bActive);
+}
+
+FGuid FDataflowConnection::GetOwningNodeGuid() const
+{
+	return OwningNode ? OwningNode->GetGuid() : FGuid();
+}
+
+uint32 FDataflowConnection::GetOwningNodeValueHash() const
+{
+	return OwningNode ? OwningNode->GetValueHash() : 0;
+}
+
diff --git a/Engine/Source/Runtime/Experimental/Dataflow/Core/Private/Dataflow/DataflowInputOutput.cpp b/Engine/Source/Runtime/Experimental/Dataflow/Core/Private/Dataflow/DataflowInputOutput.cpp
--- a/Engine/Source/Runtime/Experimental/Dataflow/Core/Private/Dataflow/DataflowInputOutput.cpp
+++ b/Engine/Source/Runtime/Experimental/Dataflow/Core/Private/Dataflow/DataflowInputOutput.cpp
@@ -155,3 +155,9 @@ bool FDataflowOutput::EvaluateImpl(Dataflow::FContext& Context) const
 
 	return true;
 }
+
+const FDataflowInput* FDataflowOutput::GetPassthroughInput() const
+{
+	return OwningNode ? OwningNode->FindInput(GetPassthroughRealAddress()) : nullptr;
+}
+
diff --git a/Engine/Source/Runtime/Experimental/Dataflow/Core/Public/Dataflow/DataflowConnection.h b/Engine/Source/Runtime/Experimental/Dataflow/Core/Public/Dataflow/DataflowConnection.h
--- a/Engine/Source/Runtime/Experimental/Dataflow/Core/Public/Dataflow/DataflowConnection.h
+++ b/Engine/Source/Runtime/Experimental/Dataflow/Core/Public/Dataflow/DataflowConnection.h
@@ -53,6 +53,11 @@ protected:
 
 	friend struct FDataflowNode;
 
+protected:
+	DATAFLOWCORE_API bool IsOwningNodeEnabled() const;
+	DATAFLOWCORE_API FGuid GetOwningNodeGuid() const;
+	DATAFLOWCORE_API uint32 GetOwningNodeValueHash() const;
+
 public:
 	FDataflowConnection() {};
 	DATAFLOWCORE_API FDataflowConnection(Dataflow::FPin::EDirection Direction, FName InType, FName InName, FDataflowNode* OwningNode = nullptr, const FProperty* InProperty = nullptr, FGuid InGuid = FGuid::NewGuid());
@@ -87,5 +92,4 @@ public:
 	}
 
 	virtual void Invalidate(const Dataflow::FTimestamp& ModifiedTimestamp = Dataflow::FTimestamp::Current()) {};
-
 };
diff --git a/Engine/Source/Runtime/Experimental/Dataflow/Core/Public/Dataflow/DataflowInputOutput.h b/Engine/Source/Runtime/Experimental/Dataflow/Core/Public/Dataflow/DataflowInputOutput.h
--- a/Engine/Source/Runtime/Experimental/Dataflow/Core/Public/Dataflow/DataflowInputOutput.h
+++ b/Engine/Source/Runtime/Experimental/Dataflow/Core/Public/Dataflow/DataflowInputOutput.h
@@ -5,7 +5,6 @@
 #include "CoreMinimal.h"
 #include "Dataflow/DataflowNodeParameters.h"
 #include "Dataflow/DataflowConnection.h"
-#include "Dataflow/DataflowNode.h"
 #include "Templates/Function.h"
 #include "Async/Async.h"
 #include "GenericPlatform/GenericPlatformCriticalSection.h"
@@ -130,7 +129,7 @@ public:
 	{
 		if (Property)
 		{
-			Context.SetData(CacheKey(), { OwningNode->GetGuid(), Property, OwningNode->GetValueHash(), Dataflow::FTimestamp::Current() }, Forward<T>(InVal));
+			Context.SetData(CacheKey(), { GetOwningNodeGuid(), Property, GetOwningNodeValueHash(), Dataflow::FTimestamp::Current() }, Forward<T>(InVal));
 		}
 	}
 
@@ -138,7 +137,7 @@ public:
 	{
 		if (!this->Evaluate<T>(Context))
 		{
-			Context.SetData(CacheKey(), { OwningNode->GetGuid(), Property, OwningNode->GetValueHash(), Dataflow::FTimestamp::Current() }, Default);
+			Context.SetData(CacheKey(), { GetOwningNodeGuid(), Property,GetOwningNodeValueHash(), Dataflow::FTimestamp::Current() }, Default);
 		}
 
 		if (Context.HasData(CacheKey()))
@@ -159,6 +158,8 @@ public:
 
 	DATAFLOWCORE_API virtual void Invalidate(const Dataflow::FTimestamp& ModifiedTimestamp = Dataflow::FTimestamp::Current()) override;
 
+private:
+	DATAFLOWCORE_API const FDataflowInput* GetPassthroughInput() const;
 };
  
 template<class T>
@@ -171,7 +172,7 @@ const T& FDataflowInput::GetValue(Dataflow::FContext& Context, const T& Default)
 		{
 			if (!ConnectionOut->Evaluate<T>(Context))
 			{
-				Context.SetData(ConnectionOut->CacheKey(), { OwningNode->GetGuid(), Property, OwningNode->GetValueHash(), Dataflow::FTimestamp::Current() }, Default);
+				Context.SetData(ConnectionOut->CacheKey(), { GetOwningNodeGuid(), Property, GetOwningNodeValueHash(), Dataflow::FTimestamp::Current() }, Default);
 			}
 			if (Context.HasData(ConnectionOut->CacheKey()))
 			{
@@ -194,11 +195,11 @@ bool FDataflowOutput::Evaluate(Dataflow::FContext& Context) const
 {
 	check(OwningNode);
  
-	if (OwningNode->bActive)
+	if (IsOwningNodeEnabled())
 	{
 		return Context.Evaluate(*this);
 	}
-	else if(const FDataflowInput* PassthroughInput = OwningNode->FindInput(GetPassthroughRealAddress()))
+	else if(const FDataflowInput* PassthroughInput = GetPassthroughInput())
 	{
 		// @todo(dataflow) would be nice if the passthrough does not overwrite the existing cache value.
 		const T& PassthroughData = PassthroughInput->GetValue<T>(Context, *reinterpret_cast<const T*>(PassthroughInput->RealAddress()));
diff --git a/Engine/Source/Runtime/Experimental/Dataflow/Core/Public/Dataflow/DataflowNode.h b/Engine/Source/Runtime/Experimental/Dataflow/Core/Public/Dataflow/DataflowNode.h
--- a/Engine/Source/Runtime/Experimental/Dataflow/Core/Public/Dataflow/DataflowNode.h
+++ b/Engine/Source/Runtime/Experimental/Dataflow/Core/Public/Dataflow/DataflowNode.h
@@ -5,14 +5,13 @@
 #include "ChaosLog.h"
 #include "CoreMinimal.h"
 #include "Dataflow/DataflowConnection.h"
+#include "Dataflow/DataflowInputOutput.h"
 #include "Dataflow/DataflowNodeParameters.h"
 #include "UObject/StructOnScope.h"
 #include "Dataflow/DataflowSettings.h"
 
 #include "DataflowNode.generated.h"
 
-struct FDataflowInput;
-struct FDataflowOutput;
 class UScriptStruct;
 
 namespace Dataflow {
