Iris
- EndReplication now handles ending replication for objects marked as garbage.
- Fixed ensure when executing DestructionInfos for objects with SubObjects

diff --git a/Engine/Source/Runtime/Engine/Private/Net/Iris/ReplicationSystem/ActorReplicationBridge.cpp b/Engine/Source/Runtime/Engine/Private/Net/Iris/ReplicationSystem/ActorReplicationBridge.cpp
--- a/Engine/Source/Runtime/Engine/Private/Net/Iris/ReplicationSystem/ActorReplicationBridge.cpp
+++ b/Engine/Source/Runtime/Engine/Private/Net/Iris/ReplicationSystem/ActorReplicationBridge.cpp
@@ -419,10 +419,11 @@ void UActorReplicationBridge::EndReplication(AActor* Actor, EEndPlayReason::Type
 {
 	using namespace UE::Net;
 
-	FNetRefHandle RefHandle = GetReplicatedRefHandle(Actor);
+	FNetRefHandle RefHandle = GetReplicatedRefHandle(Actor, EGetRefHandleFlags::EvenIfGarbage);
 	if (RefHandle.IsValid())
 	{
-		UE_LOG(LogIrisBridge, Verbose, TEXT("EndReplication for %s %s. Reason %s "), *Actor->GetName(), *RefHandle.ToString(), *UEnum::GetValueAsString(TEXT("Engine.EEndPlayReason"), EndPlayReason));
+		UE_LOG(LogIrisBridge, Verbose, TEXT("EndReplication for %s %s. Reason %s "), *GetNameSafe(Actor), *RefHandle.ToString(), *UEnum::GetValueAsString(TEXT("Engine.EEndPlayReason"), EndPlayReason));
+		ensureMsgf(IsValid(Actor), TEXT("Calling EndReplication for Invalid Object for %s %s."), *GetNameSafe(Actor), *RefHandle.ToString());
 	
 		EEndReplicationFlags Flags = EEndReplicationFlags::None;
 		const bool bShouldDestroyObject = EndPlayReason == EEndPlayReason::Destroyed;		
@@ -460,11 +461,11 @@ void UActorReplicationBridge::EndReplicationForActorComponent(UActorComponent* A
 {
 	using namespace UE::Net;
 
-	FNetRefHandle ComponentHandle = GetReplicatedRefHandle(ActorComponent);
-	const AActor* Actor = ActorComponent->GetOwner();
-	if (ComponentHandle.IsValid() && Actor)
+	FNetRefHandle ComponentHandle = GetReplicatedRefHandle(ActorComponent, EGetRefHandleFlags::EvenIfGarbage);
+	if (ComponentHandle.IsValid())
 	{
-		UE_LOG(LogIrisBridge, Verbose, TEXT("EndReplicationForActorComponent for %s %s."), *ActorComponent->GetName(), *ComponentHandle.ToString());
+		UE_LOG(LogIrisBridge, Verbose, TEXT("EndReplicationForActorComponent for %s %s."), *GetNameSafe(ActorComponent), *ComponentHandle.ToString());
+		ensureMsgf(IsValid(ActorComponent), TEXT("Calling EndReplication for Invalid Object for %s %s."), *GetNameSafe(ActorComponent), *ComponentHandle.ToString());
 		
 		UObjectReplicationBridge::EndReplication(ComponentHandle, EndReplicationFlags, nullptr);
 	}
diff --git a/Engine/Source/Runtime/Engine/Private/Net/Iris/ReplicationSystem/ReplicationSystemUtil.cpp b/Engine/Source/Runtime/Engine/Private/Net/Iris/ReplicationSystem/ReplicationSystemUtil.cpp
--- a/Engine/Source/Runtime/Engine/Private/Net/Iris/ReplicationSystem/ReplicationSystemUtil.cpp
+++ b/Engine/Source/Runtime/Engine/Private/Net/Iris/ReplicationSystem/ReplicationSystemUtil.cpp
@@ -285,25 +285,20 @@ void FReplicationSystemUtil::BeginReplicationForActorComponentSubObject(UActorCo
 	}
 }
 
-void FReplicationSystemUtil::EndReplicationForActorComponent(UActorComponent* SubObject)
+void FReplicationSystemUtil::EndReplicationForActorComponent(UActorComponent* ActorComponent)
 {
-	ReplicationSystemUtil::ForEachReplicationSystem([SubObject](UReplicationSystem* ReplicationSystem)
+	ReplicationSystemUtil::ForEachReplicationSystem([ActorComponent](UReplicationSystem* ReplicationSystem)
 	{
 		if (UActorReplicationBridge* Bridge = Cast<UActorReplicationBridge>(ReplicationSystem->GetReplicationBridge()))
 		{
 			constexpr EEndReplicationFlags EndReplicationFlags = EEndReplicationFlags::DestroyNetHandle | EEndReplicationFlags::ClearNetPushId;
-			Bridge->EndReplicationForActorComponent(SubObject, EndReplicationFlags);
+			Bridge->EndReplicationForActorComponent(ActorComponent, EndReplicationFlags);
 		}
 	});
 }
 
 void FReplicationSystemUtil::EndReplicationForActorSubObject(const AActor* Actor, UObject* SubObject)
 {
-	if (!IsValid(SubObject))
-	{
-		return;
-	}
-
 	ReplicationSystemUtil::ForEachReplicationSystem([SubObject](UReplicationSystem* ReplicationSystem)
 	{
 		if (UActorReplicationBridge* Bridge = Cast<UActorReplicationBridge>(ReplicationSystem->GetReplicationBridge()))
diff --git a/Engine/Source/Runtime/Experimental/Iris/Core/Private/Iris/ReplicationSystem/ObjectReferenceCache.cpp b/Engine/Source/Runtime/Experimental/Iris/Core/Private/Iris/ReplicationSystem/ObjectReferenceCache.cpp
--- a/Engine/Source/Runtime/Experimental/Iris/Core/Private/Iris/ReplicationSystem/ObjectReferenceCache.cpp
+++ b/Engine/Source/Runtime/Experimental/Iris/Core/Private/Iris/ReplicationSystem/ObjectReferenceCache.cpp
@@ -419,7 +419,7 @@ bool FObjectReferenceCache::CreateObjectReferenceInternal(const UObject* Object,
 	return true;
 }
 
-FNetRefHandle FObjectReferenceCache::GetObjectReferenceHandleFromObject(const UObject* Object) const
+FNetRefHandle FObjectReferenceCache::GetObjectReferenceHandleFromObject(const UObject* Object, EGetRefHandleFlags GetRefHandleFlags) const
 {
 	if (Object == nullptr)
 	{
@@ -441,7 +441,8 @@ FNetRefHandle FObjectReferenceCache::GetObjectReferenceHandleFromObject(const UO
 
 		const FCachedNetObjectReference& CachedObject = *CachedObjectPtr;
 
-		const UObject* ExistingObject = CachedObject.Object.Get();
+		const bool bEvenIfGarbage = EnumHasAnyFlags(GetRefHandleFlags, EGetRefHandleFlags::EvenIfGarbage);
+		const UObject* ExistingObject = CachedObject.Object.Get(bEvenIfGarbage);
 		if (ExistingObject == Object)
 		{
 			UE_LOG_REFERENCECACHE(VeryVerbose, TEXT("ObjectReferenceCache::GetObjectReferenceHandleFromObject Found existing %s for Object %s, OuterNetRefHandle: %s"), *CachedObject.NetRefHandle.ToString(), ToCStr(Object->GetPathName()), *CachedObject.OuterNetRefHandle.ToString());
diff --git a/Engine/Source/Runtime/Experimental/Iris/Core/Private/Iris/ReplicationSystem/ObjectReferenceCache.h b/Engine/Source/Runtime/Experimental/Iris/Core/Private/Iris/ReplicationSystem/ObjectReferenceCache.h
--- a/Engine/Source/Runtime/Experimental/Iris/Core/Private/Iris/ReplicationSystem/ObjectReferenceCache.h
+++ b/Engine/Source/Runtime/Experimental/Iris/Core/Private/Iris/ReplicationSystem/ObjectReferenceCache.h
@@ -5,6 +5,7 @@
 #include "CoreTypes.h"
 #include "Iris/Core/NetObjectReference.h"
 #include "Iris/ReplicationSystem/NetRefHandle.h"
+#include "Iris/ReplicationSystem/ReplicationSystemTypes.h"
 #include "UObject/WeakObjectPtr.h"
 #include "Containers/Map.h"
 #include "ObjectReferenceCacheFwd.h"
@@ -48,7 +49,7 @@ public:
 	FNetRefHandle CreateObjectReferenceHandle(const UObject* Object);
 
 	// Get existing handle for object
-	FNetRefHandle GetObjectReferenceHandleFromObject(const UObject* Object) const;
+	FNetRefHandle GetObjectReferenceHandleFromObject(const UObject* Object, EGetRefHandleFlags GetRefHandleFlags = EGetRefHandleFlags::None) const;
 
 	// Get object from handle, only if the object is in the cache.
 	UObject* GetObjectFromReferenceHandle(FNetRefHandle RefHandle);
diff --git a/Engine/Source/Runtime/Experimental/Iris/Core/Private/Iris/ReplicationSystem/ObjectReplicationBridge.cpp b/Engine/Source/Runtime/Experimental/Iris/Core/Private/Iris/ReplicationSystem/ObjectReplicationBridge.cpp
--- a/Engine/Source/Runtime/Experimental/Iris/Core/Private/Iris/ReplicationSystem/ObjectReplicationBridge.cpp
+++ b/Engine/Source/Runtime/Experimental/Iris/Core/Private/Iris/ReplicationSystem/ObjectReplicationBridge.cpp
@@ -202,10 +202,9 @@ UObject* UObjectReplicationBridge::GetReplicatedObject(FNetRefHandle Handle) con
 	return IsReplicatedHandle(Handle) ? GetObjectFromReferenceHandle(Handle) : nullptr;
 };
 
-UE::Net::FNetRefHandle UObjectReplicationBridge::GetReplicatedRefHandle(const UObject* Object) const
+UE::Net::FNetRefHandle UObjectReplicationBridge::GetReplicatedRefHandle(const UObject* Object, EGetRefHandleFlags GetRefHandleFlags) const
 {
-	FNetRefHandle Handle = GetObjectReferenceCache()->GetObjectReferenceHandleFromObject(Object);
-
+	const FNetRefHandle Handle = GetObjectReferenceCache()->GetObjectReferenceHandleFromObject(Object, GetRefHandleFlags);
 	return IsReplicatedHandle(Handle) ? Handle : FNetRefHandle::GetInvalid();
 }
 
@@ -577,8 +576,12 @@ bool UObjectReplicationBridge::WriteNetRefHandleCreationInfo(FReplicationBridgeS
 
 void UObjectReplicationBridge::EndReplication(UObject* Instance, EEndReplicationFlags EndReplicationFlags, FEndReplicationParameters* Parameters)
 {
-	const FNetRefHandle Handle = GetReplicatedRefHandle(Instance);
-	UReplicationBridge::EndReplication(Handle, EndReplicationFlags, Parameters);
+	const FNetRefHandle RefHandle = GetReplicatedRefHandle(Instance, EGetRefHandleFlags::EvenIfGarbage);
+	if (RefHandle.IsValid())
+	{
+		ensureMsgf(IsValid(Instance), TEXT("Calling EndReplication for Invalid Object for %s %s."), *GetNameSafe(Instance), *RefHandle.ToString());
+		UReplicationBridge::EndReplication(RefHandle, EndReplicationFlags, Parameters);
+	}
 }
 
 void UObjectReplicationBridge::DetachInstanceFromRemote(FNetRefHandle Handle, EReplicationBridgeDestroyInstanceReason DestroyReason, EReplicationBridgeDestroyInstanceFlags DestroyFlags)
@@ -827,7 +830,7 @@ void UObjectReplicationBridge::PruneStaleObjects()
 	// Detect stale references and try to kill/report them
 	auto DetectStaleObjectsFunc = [&LocalNetRefHandleManager, &StaleObjects, &ReplicatedInstances](uint32 InternalNetHandleIndex)
 	{
-		if (ReplicatedInstances[InternalNetHandleIndex] == nullptr)
+		if (!IsValid(ReplicatedInstances[InternalNetHandleIndex]))
 		{
 			const FNetRefHandleManager::FReplicatedObjectData& ObjectData = LocalNetRefHandleManager.GetReplicatedObjectDataNoCheck(InternalNetHandleIndex);
 			if (ObjectData.InstanceProtocol)
@@ -1395,7 +1398,7 @@ bool UObjectReplicationBridge::FindOrCachePollFrequency(const UClass* Class, flo
 	}
 	else
 	{
-		// Add the class hiearchy to our set of classes without overrides.
+		// Add the class hierarchy to our set of classes without overrides.
 		for (const UClass* ClassToAdd = Class; ClassToAdd != nullptr; ClassToAdd = ClassToAdd->GetSuperClass())
 		{
 			// We avoid adding classes that are in the exact match container, even though it's not strictly necessary.
diff --git a/Engine/Source/Runtime/Experimental/Iris/Core/Private/Iris/ReplicationSystem/ReplicationBridge.cpp b/Engine/Source/Runtime/Experimental/Iris/Core/Private/Iris/ReplicationSystem/ReplicationBridge.cpp
--- a/Engine/Source/Runtime/Experimental/Iris/Core/Private/Iris/ReplicationSystem/ReplicationBridge.cpp
+++ b/Engine/Source/Runtime/Experimental/Iris/Core/Private/Iris/ReplicationSystem/ReplicationBridge.cpp
@@ -155,7 +155,7 @@ void UReplicationBridge::DestroyNetObjectFromRemote(FNetRefHandle Handle, ERepli
 		ObjectData.bTearOff = (DestroyReason == EReplicationBridgeDestroyInstanceReason::TearOff);
 		ObjectData.bPendingEndReplication = 1U;
 
-		// if the a subobject owner is to be destroyed we want to detach all subobjects before doing so to ensure we execute expected callbacks
+		// if a subobject owner is to be destroyed we want to detach all subobjects before doing so to ensure we execute expected callbacks
 		// We keep tracking them internally
 		DetachSubObjectInstancesFromRemote(Handle, DestroyReason, DestroyFlags);
 		
@@ -280,6 +280,11 @@ void UReplicationBridge::ReadAndExecuteDestructionInfoFromRemote(FReplicationBri
 
 		constexpr EReplicationBridgeDestroyInstanceReason DestroyReason = EReplicationBridgeDestroyInstanceReason::Destroy;
 		const EReplicationBridgeDestroyInstanceFlags DestroyFlags = IsAllowedToDestroyInstance(Instance) ? EReplicationBridgeDestroyInstanceFlags::AllowDestroyInstanceFromRemote : EReplicationBridgeDestroyInstanceFlags::None;
+
+		// if a subobject owner is to be destroyed we want to detach all subobjects before doing so to ensure we execute expected callbacks
+		// We keep tracking them internally
+		DetachSubObjectInstancesFromRemote(ReferenceToDestroy.GetRefHandle(), DestroyReason, DestroyFlags);
+
 		CallDetachInstanceFromRemote(ReferenceToDestroy.GetRefHandle(), DestroyReason, DestroyFlags);
 	}
 }
diff --git a/Engine/Source/Runtime/Experimental/Iris/Core/Public/Iris/ReplicationSystem/ObjectReplicationBridge.h b/Engine/Source/Runtime/Experimental/Iris/Core/Public/Iris/ReplicationSystem/ObjectReplicationBridge.h
--- a/Engine/Source/Runtime/Experimental/Iris/Core/Public/Iris/ReplicationSystem/ObjectReplicationBridge.h
+++ b/Engine/Source/Runtime/Experimental/Iris/Core/Public/Iris/ReplicationSystem/ObjectReplicationBridge.h
@@ -44,6 +44,7 @@ class UObjectReplicationBridge : public UReplicationBridge
 
 public:
 	using UReplicationBridge::EndReplication;
+	using EGetRefHandleFlags = UE::Net::EGetRefHandleFlags;
 
 	struct FCreateNetRefHandleParams
 	{
@@ -89,7 +90,7 @@ public:
 	IRISCORE_API UObject* GetReplicatedObject(FNetRefHandle Handle) const;
 
 	/** Get NetRefHandle from a replicated UObject. */
-	IRISCORE_API FNetRefHandle GetReplicatedRefHandle(const UObject* Object) const;
+	IRISCORE_API FNetRefHandle GetReplicatedRefHandle(const UObject* Object, EGetRefHandleFlags GetRefHandleFlags = EGetRefHandleFlags::None) const;
 
 	/** Get NetRefHandle from a NetHandle. */
 	IRISCORE_API FNetRefHandle GetReplicatedRefHandle(FNetHandle Handle) const;
diff --git a/Engine/Source/Runtime/Experimental/Iris/Core/Public/Iris/ReplicationSystem/ReplicationSystemTypes.h b/Engine/Source/Runtime/Experimental/Iris/Core/Public/Iris/ReplicationSystem/ReplicationSystemTypes.h
--- a/Engine/Source/Runtime/Experimental/Iris/Core/Public/Iris/ReplicationSystem/ReplicationSystemTypes.h
+++ b/Engine/Source/Runtime/Experimental/Iris/Core/Public/Iris/ReplicationSystem/ReplicationSystemTypes.h
@@ -4,6 +4,7 @@
 
 #include "CoreTypes.h"
 #include "Delegates/Delegate.h"
+#include "Misc/EnumClassFlags.h"
 
 namespace UE::Net
 {
@@ -14,6 +15,13 @@ enum class ENetObjectDeltaCompressionStatus : unsigned
 	Allow,
 };
 
+enum class EGetRefHandleFlags : uint32
+{
+	None,
+	EvenIfGarbage
+};
+ENUM_CLASS_FLAGS(EGetRefHandleFlags);
+
 enum class EReplicationSystemSendPass : unsigned
 {
 	Invalid,
