Fix force offset missing functionality

diff --git a/Engine/Source/Runtime/Experimental/ChaosVehicles/ChaosVehiclesCore/Private/SimModule/DeferredForcesModular.cpp b/Engine/Source/Runtime/Experimental/ChaosVehicles/ChaosVehiclesCore/Private/SimModule/DeferredForcesModular.cpp
--- a/Engine/Source/Runtime/Experimental/ChaosVehicles/ChaosVehiclesCore/Private/SimModule/DeferredForcesModular.cpp
+++ b/Engine/Source/Runtime/Experimental/ChaosVehicles/ChaosVehiclesCore/Private/SimModule/DeferredForcesModular.cpp
@@ -393,7 +393,9 @@ void FDeferredForcesModular::AddForceAtPosition(Chaos::FPBDRigidParticleHandle*
 {
 	if (ensure(RigidHandle))
 	{
-		AddForce_Implementation(RigidHandle, DataIn.OffsetTransform, DataIn.Force, OffsetTransform, (DataIn.Flags & EForceFlags::LevelSlope) == EForceFlags::LevelSlope, DataIn.DebugColor);
+		FTransform FinalOffsetTransform = OffsetTransform;
+		FinalOffsetTransform.AddToTranslation(DataIn.Position);
+		AddForce_Implementation(RigidHandle, FinalOffsetTransform, DataIn.Force, FinalOffsetTransform, (DataIn.Flags & EForceFlags::LevelSlope) == EForceFlags::LevelSlope, DataIn.DebugColor);
 	}
 }
 
