Add additional solver stage stats to Chaos advance tasks, particularly for marshalling manager updates and rewind

diff --git a/Engine/Source/Runtime/Experimental/Chaos/Private/Chaos/Framework/PhysicsSolverBase.cpp b/Engine/Source/Runtime/Experimental/Chaos/Private/Chaos/Framework/PhysicsSolverBase.cpp
--- a/Engine/Source/Runtime/Experimental/Chaos/Private/Chaos/Framework/PhysicsSolverBase.cpp
+++ b/Engine/Source/Runtime/Experimental/Chaos/Private/Chaos/Framework/PhysicsSolverBase.cpp
@@ -169,10 +169,16 @@ namespace Chaos
 
 		Solver.AdvanceSolverBy(FSubStepInfo{ PseudoFraction, PushData->IntervalStep, PushData->IntervalNumSteps, PushData->bSolverSubstepped });
 
-		Solver.GetMarshallingManager().FreeDataToHistory_Internal(PushData);	//cannot use push data after this point
-		PushData = nullptr;
+		{
+			SCOPE_CYCLE_COUNTER(STAT_ResetMarshallingData);
+			Solver.GetMarshallingManager().FreeDataToHistory_Internal(PushData);	//cannot use push data after this point
+			PushData = nullptr;
+		}
 
-		Solver.ConditionalApplyRewind_Internal();
+		{
+			SCOPE_CYCLE_COUNTER(STAT_ConditionalApplyRewind);
+			Solver.ConditionalApplyRewind_Internal();
+		}
 
 		Solver.NumPendingSolverAdvanceTasks--;
 	}
diff --git a/Engine/Source/Runtime/Experimental/Chaos/Private/ChaosStats.cpp b/Engine/Source/Runtime/Experimental/Chaos/Private/ChaosStats.cpp
--- a/Engine/Source/Runtime/Experimental/Chaos/Private/ChaosStats.cpp
+++ b/Engine/Source/Runtime/Experimental/Chaos/Private/ChaosStats.cpp
@@ -67,6 +67,14 @@ DEFINE_STAT(STAT_GatherSleepingEvent)
 DEFINE_STAT(STAT_GatherCrumblingEvent)
 DEFINE_STAT(STAT_AccelerationStructureReset);
 
+DEFINE_STAT(STAT_FinalizeCallbacks);
+DEFINE_STAT(STAT_ResetClusteringEvents);
+DEFINE_STAT(STAT_RewindFinishFrame);
+DEFINE_STAT(STAT_ResetMarshallingData);
+DEFINE_STAT(STAT_ConditionalApplyRewind);
+DEFINE_STAT(STAT_FinalizePullData);
+DEFINE_STAT(STAT_DestroyPendingProxies);
+
 DEFINE_STAT(STAT_ParamUpdateObject);
 DEFINE_STAT(STAT_ParamUpdateField);
 DEFINE_STAT(STAT_SyncEvents_GameThread);
diff --git a/Engine/Source/Runtime/Experimental/Chaos/Private/PBDRigidsSolver.cpp b/Engine/Source/Runtime/Experimental/Chaos/Private/PBDRigidsSolver.cpp
--- a/Engine/Source/Runtime/Experimental/Chaos/Private/PBDRigidsSolver.cpp
+++ b/Engine/Source/Runtime/Experimental/Chaos/Private/PBDRigidsSolver.cpp
@@ -579,7 +579,10 @@ namespace Chaos
 				MSolver->GetEvolution()->EndFrame(MDeltaTime);
 			}
 
-			MSolver->FinalizeCallbackData_Internal();
+			{
+				SCOPE_CYCLE_COUNTER(STAT_FinalizeCallbacks);
+				MSolver->FinalizeCallbackData_Internal();
+			}
 
 			MSolver->SetSolverTime(MSolver->GetSolverTime() + MDeltaTime );
 			MSolver->GetCurrentFrame()++;
@@ -594,10 +597,14 @@ namespace Chaos
 				MSolver->CompleteSceneSimulation();
 			}
 
-			// reset all clustering events needs to be after CompleteSceneSimulation to make sure the cache recording gets them before they get removed
-			// they cannot be right after the presolve callback ( as they were before ) because they will cause geometry collection replicated clients to miss them
-			// Todo(chaos) we should probably move all of the solver event reset here in the future
-			MSolver->GetEvolution()->GetRigidClustering().ResetAllEvents();
+			{
+				SCOPE_CYCLE_COUNTER(STAT_ResetClusteringEvents);
+
+				// reset all clustering events needs to be after CompleteSceneSimulation to make sure the cache recording gets them before they get removed
+				// they cannot be right after the presolve callback ( as they were before ) because they will cause geometry collection replicated clients to miss them
+				// Todo(chaos) we should probably move all of the solver event reset here in the future
+				MSolver->GetEvolution()->GetRigidClustering().ResetAllEvents();
+			}
 
 			// Recover unused memory from particle arrays, based on the array shrink policy
 			if (bChaosSolverShrinkArrays)
@@ -608,6 +615,7 @@ namespace Chaos
 
 			if (FRewindData* RewindData = MSolver->GetRewindData())
 			{
+				SCOPE_CYCLE_COUNTER(STAT_RewindFinishFrame);
 				RewindData->FinishFrame();
 			}
 		}
@@ -1380,6 +1388,7 @@ namespace Chaos
 
 		if (MLastDt > 0)
 		{
+			SCOPE_CYCLE_COUNTER(STAT_FinalizePullData);
 			//pass information back to external thread
 			//we skip dt=0 case because sync data should be identical if dt = 0
 			MarshallingManager.FinalizePullData_Internal(MEvolution->LatestExternalTimestampConsumed_Internal, StartSimTime, MLastDt);
@@ -1387,6 +1396,7 @@ namespace Chaos
 
 		if(SubStepInfo.Step == SubStepInfo.NumSteps - 1)
 		{
+			SCOPE_CYCLE_COUNTER(STAT_DestroyPendingProxies);
 			//final step so we can destroy proxies
 			DestroyPendingProxies_Internal();
 		}
diff --git a/Engine/Source/Runtime/Experimental/Chaos/Public/ChaosStats.h b/Engine/Source/Runtime/Experimental/Chaos/Public/ChaosStats.h
--- a/Engine/Source/Runtime/Experimental/Chaos/Public/ChaosStats.h
+++ b/Engine/Source/Runtime/Experimental/Chaos/Public/ChaosStats.h
@@ -85,6 +85,14 @@ DECLARE_CYCLE_STAT_EXTERN(TEXT("Gathering Sleeping Event Data"), STAT_GatherSlee
 DECLARE_CYCLE_STAT_EXTERN(TEXT("Acceleration Structure Reset"), STAT_AccelerationStructureReset, STATGROUP_Chaos, CHAOS_API);
 DECLARE_CYCLE_STAT_EXTERN(TEXT("Gathering Crumbling Event Data"), STAT_GatherCrumblingEvent, STATGROUP_Chaos, CHAOS_API);
 
+DECLARE_CYCLE_STAT_EXTERN(TEXT("Finalize Sim Callbacks"), STAT_FinalizeCallbacks, STATGROUP_Chaos, CHAOS_API);
+DECLARE_CYCLE_STAT_EXTERN(TEXT("Reset Clustering Events"), STAT_ResetClusteringEvents, STATGROUP_Chaos, CHAOS_API);
+DECLARE_CYCLE_STAT_EXTERN(TEXT("Rewind Finish Frame"), STAT_RewindFinishFrame, STATGROUP_Chaos, CHAOS_API);
+DECLARE_CYCLE_STAT_EXTERN(TEXT("Reset Marshalling Data"), STAT_ResetMarshallingData, STATGROUP_Chaos, CHAOS_API);
+DECLARE_CYCLE_STAT_EXTERN(TEXT("Conditional Apply Rewind"), STAT_ConditionalApplyRewind, STATGROUP_Chaos, CHAOS_API);
+DECLARE_CYCLE_STAT_EXTERN(TEXT("Finalize Pull Data"), STAT_FinalizePullData, STATGROUP_Chaos, CHAOS_API);
+DECLARE_CYCLE_STAT_EXTERN(TEXT("Destroy Pending Proxies"), STAT_DestroyPendingProxies, STATGROUP_Chaos, CHAOS_API);
+
 DECLARE_CYCLE_STAT_EXTERN(TEXT("Reset Collision Rule"), STAT_ResetCollisionRule, STATGROUP_Chaos, CHAOS_API);
 
 DECLARE_CYCLE_STAT_EXTERN(TEXT("Object Parameter Update"), STAT_ParamUpdateObject, STATGROUP_Chaos, CHAOS_API);
