Dataflow : Add support for cache references useful for fast passthrough or reroute nodes

diff --git a/Engine/Plugins/Experimental/Dataflow/Source/DataflowEditor/Private/Dataflow/AssetDefinition_DataflowContext.cpp b/Engine/Plugins/Experimental/Dataflow/Source/DataflowEditor/Private/Dataflow/AssetDefinition_DataflowContext.cpp
--- a/Engine/Plugins/Experimental/Dataflow/Source/DataflowEditor/Private/Dataflow/AssetDefinition_DataflowContext.cpp
+++ b/Engine/Plugins/Experimental/Dataflow/Source/DataflowEditor/Private/Dataflow/AssetDefinition_DataflowContext.cpp
@@ -50,15 +50,15 @@ namespace DataflowContextDefinitionHelpers
 
 			if (TUniquePtr<FContextCacheElementBase>* Data = Context->GetBaseData(Key))
 			{
-				if ((*Data) && !(*Data)->Property)
+				if ((*Data) && !(*Data)->GetProperty())
 				{
-					if (TSharedPtr<FDataflowNode> Node = Dataflow->FindBaseNode((*Data)->NodeGuid))
+					if (TSharedPtr<FDataflowNode> Node = Dataflow->FindBaseNode((*Data)->GetNodeGuid()))
 					{
 						if (FDataflowOutput* Output = Node->FindOutput(Key))
 						{
 							if (const FProperty* Property = Output->GetProperty())
 							{
-								(*Data)->Property = Property;
+								(*Data)->SetProperty(Property);
 								bValidKey = true;
 							}
 						}
@@ -105,11 +105,10 @@ namespace DataflowContextDefinitionHelpers
 			if (TUniquePtr<FContextCacheElementBase>* Data = Context->GetBaseData(Key))
 			{
 				if (!(*Data)) return false;
-				if (!(*Data)->Property) return false;
-				if(TSharedPtr<FDataflowNode> Node = Dataflow->FindBaseNode((*Data)->NodeGuid))
+				if (!(*Data)->GetProperty()) return false;
+				if(TSharedPtr<FDataflowNode> Node = Dataflow->FindBaseNode((*Data)->GetNodeGuid()))
 				{
-					uint32 Hash = Node->GetValueHash();
-					if ((*Data)->NodeHash != Node->GetValueHash())
+					if ((*Data)->GetNodeHash() != Node->GetValueHash())
 					{
 						return false;
 					}
diff --git a/Engine/Plugins/Experimental/Dataflow/Source/DataflowEnginePlugin/Private/Dataflow/DataflowEngineContextCaching.cpp b/Engine/Plugins/Experimental/Dataflow/Source/DataflowEnginePlugin/Private/Dataflow/DataflowEngineContextCaching.cpp
--- a/Engine/Plugins/Experimental/Dataflow/Source/DataflowEnginePlugin/Private/Dataflow/DataflowEngineContextCaching.cpp
+++ b/Engine/Plugins/Experimental/Dataflow/Source/DataflowEnginePlugin/Private/Dataflow/DataflowEngineContextCaching.cpp
@@ -26,7 +26,9 @@ namespace Dataflow
 				{
 					check(Element!=nullptr); // read from a explicit type and return a null element
 					
-					FManagedArrayCollection& Collection = FContextCachingFactory::GetTypedElement<FManagedArrayCollection>(Element);
+					const FManagedArrayCollection EmptyCollection;
+					// cache always return const data , because they are immutable, however the serialize method is not const and require a const_cast
+					FManagedArrayCollection& Collection = const_cast<FManagedArrayCollection&>(FContextCachingFactory::GetTypedElement<FManagedArrayCollection>(Element, EmptyCollection));
 					Chaos::FChaosArchive ChaosAr(Ar);
 					Collection.Serialize(ChaosAr);
 
diff --git a/Engine/Source/Runtime/Experimental/Dataflow/Core/Private/Dataflow/DataflowInputOutput.cpp b/Engine/Source/Runtime/Experimental/Dataflow/Core/Private/Dataflow/DataflowInputOutput.cpp
--- a/Engine/Source/Runtime/Experimental/Dataflow/Core/Private/Dataflow/DataflowInputOutput.cpp
+++ b/Engine/Source/Runtime/Experimental/Dataflow/Core/Private/Dataflow/DataflowInputOutput.cpp
@@ -161,3 +161,19 @@ const FDataflowInput* FDataflowOutput::GetPassthroughInput() const
 	return OwningNode ? OwningNode->FindInput(GetPassthroughRealAddress()) : nullptr;
 }
 
+void FDataflowOutput::ForwardInput(const void* InputReference, Dataflow::FContext& Context) const
+{
+	if (Property && OwningNode)
+	{
+		const FDataflowInput* InputToForward = OwningNode->FindInput(InputReference);
+		if (InputToForward->GetConnectedOutputs().Num())
+		{
+			ensure(InputToForward->GetType() == GetType());
+			ensure(InputToForward->GetConnectedOutputs().Num() == 1);
+			if (const FDataflowOutput* ConnectionOut = InputToForward->GetConnection())
+			{
+				Context.SetDataReference(CacheKey(), Property, ConnectionOut->CacheKey());
+			}
+		}
+	}
+}
diff --git a/Engine/Source/Runtime/Experimental/Dataflow/Core/Private/Dataflow/DataflowNodeParameters.cpp b/Engine/Source/Runtime/Experimental/Dataflow/Core/Private/Dataflow/DataflowNodeParameters.cpp
--- a/Engine/Source/Runtime/Experimental/Dataflow/Core/Private/Dataflow/DataflowNodeParameters.cpp
+++ b/Engine/Source/Runtime/Experimental/Dataflow/Core/Private/Dataflow/DataflowNodeParameters.cpp
@@ -113,7 +113,8 @@ namespace Dataflow
 
 			for (TPair<FContextCacheKey, TUniquePtr<FContextCacheElementBase>>& Elem : Pairs)
 			{
-				if (Elem.Value && Elem.Value->Property)
+				// note : we only serialize typed cache element and ignore the reference ones ( since they don't hold data per say )
+				if (Elem.Value && Elem.Value->Property && Elem.Value->Type == FContextCacheElementBase::EType::CacheElementTyped)
 				{
 					FProperty* Property = (FProperty*)Elem.Value->Property;
 					FName TypeName(Elem.Value->Property->GetCPPType());
diff --git a/Engine/Source/Runtime/Experimental/Dataflow/Core/Public/Dataflow/DataflowContextCachingFactory.h b/Engine/Source/Runtime/Experimental/Dataflow/Core/Public/Dataflow/DataflowContextCachingFactory.h
--- a/Engine/Source/Runtime/Experimental/Dataflow/Core/Public/Dataflow/DataflowContextCachingFactory.h
+++ b/Engine/Source/Runtime/Experimental/Dataflow/Core/Public/Dataflow/DataflowContextCachingFactory.h
@@ -48,9 +48,16 @@ namespace Dataflow
 		void RegisterSerializeFunction(const FName& Type, FSerializeFunction InSerializeFunc);
 
 		template<class T>
-		static T& GetTypedElement(FContextCacheElementBase* InElement)
+		static const T& GetTypedElement(const FContextCacheElementBase* InElement, const T& Default)
 		{
-			return (T&)InElement->GetTypedData<T>(InElement->Property);
+			// we only support typed cache element and not cache reference
+			if (InElement && InElement->GetType() == FContextCacheElementBase::EType::CacheElementTyped)
+			{
+				// it is assumed the type requested matches the cache entry 
+				const FContextCacheElement<T>* TypedElement = static_cast<const FContextCacheElement<T>*>(InElement);
+				return TypedElement->GetDataDirect();
+			}
+			return Default;
 		}
 
 		template<class T>
diff --git a/Engine/Source/Runtime/Experimental/Dataflow/Core/Public/Dataflow/DataflowInputOutput.h b/Engine/Source/Runtime/Experimental/Dataflow/Core/Public/Dataflow/DataflowInputOutput.h
--- a/Engine/Source/Runtime/Experimental/Dataflow/Core/Public/Dataflow/DataflowInputOutput.h
+++ b/Engine/Source/Runtime/Experimental/Dataflow/Core/Public/Dataflow/DataflowInputOutput.h
@@ -129,7 +129,7 @@ public:
 	{
 		if (Property)
 		{
-			Context.SetData(CacheKey(), { GetOwningNodeGuid(), Property, GetOwningNodeValueHash(), Dataflow::FTimestamp::Current() }, Forward<T>(InVal));
+			Context.SetData(CacheKey(), Property, Forward<T>(InVal), GetOwningNodeGuid(), GetOwningNodeValueHash(), Dataflow::FTimestamp::Current());
 		}
 	}
 
@@ -137,7 +137,7 @@ public:
 	{
 		if (!this->Evaluate<T>(Context))
 		{
-			Context.SetData(CacheKey(), { GetOwningNodeGuid(), Property,GetOwningNodeValueHash(), Dataflow::FTimestamp::Current() }, Default);
+			Context.SetData(CacheKey(), Property, Default, GetOwningNodeGuid(), GetOwningNodeValueHash(), Dataflow::FTimestamp::Current());
 		}
 
 		if (Context.HasData(CacheKey()))
@@ -148,6 +148,10 @@ public:
 		return Default;
 	}
 
+	// there's no need for a templatized version as the parameter will not be used
+	// the method do check if the type of the input is the same as the output type though 
+	DATAFLOWCORE_API void ForwardInput(const void* InputReference, Dataflow::FContext& Context) const;
+
 	DATAFLOWCORE_API bool EvaluateImpl(Dataflow::FContext& Context) const;
 	
 	template<class T>
@@ -172,7 +176,7 @@ const T& FDataflowInput::GetValue(Dataflow::FContext& Context, const T& Default)
 		{
 			if (!ConnectionOut->Evaluate<T>(Context))
 			{
-				Context.SetData(ConnectionOut->CacheKey(), { GetOwningNodeGuid(), Property, GetOwningNodeValueHash(), Dataflow::FTimestamp::Current() }, Default);
+				Context.SetData(ConnectionOut->CacheKey(), Property, Default, GetOwningNodeGuid(), GetOwningNodeValueHash(), Dataflow::FTimestamp::Current());
 			}
 			if (Context.HasData(ConnectionOut->CacheKey()))
 			{
diff --git a/Engine/Source/Runtime/Experimental/Dataflow/Core/Public/Dataflow/DataflowNode.h b/Engine/Source/Runtime/Experimental/Dataflow/Core/Public/Dataflow/DataflowNode.h
--- a/Engine/Source/Runtime/Experimental/Dataflow/Core/Public/Dataflow/DataflowNode.h
+++ b/Engine/Source/Runtime/Experimental/Dataflow/Core/Public/Dataflow/DataflowNode.h
@@ -229,8 +229,47 @@ struct FDataflowNode
 	*/
 	template<class T> void SetValue(Dataflow::FContext& Context, T&& Value, const typename TDecay<T>::Type* Reference) const
 	{
-		checkSlow(FindOutput(Reference));
-		FindOutput(Reference)->SetValue(Forward<T>(Value), Context);
+		if (const FDataflowOutput* Output = FindOutput(Reference))
+		{
+			Output->template SetValue<T>(Forward<T>(Value), Context);
+		}
+		else
+		{
+			checkfSlow(false, TEXT("This output could not be found within this node, check this has been properly registered in the node constructor"));
+		}
+	}
+
+	/**
+	*   ForwardInput(...)
+	*
+	*   Forward an input to this output
+	*	this will not cache the value itself but cache a reference to the input connection cache entry
+	*	this is memory efficient and do not require a runtime copy of the data
+	*	input and output references must match in type
+	*
+	*   @param Context : The evaluation context that holds the data store.
+	*   @param InputReference : Pointer to a input member of this node that needs to be forwarded
+	*   @param Reference : Pointer to a member of this node that corresponds with the output to set.
+	*/
+	template<class T> void ForwardInput(Dataflow::FContext& Context, const typename TDecay<T>::Type* InputReference, const typename TDecay<T>::Type* Reference) const
+	{
+		if (const FDataflowOutput* Output = FindOutput(Reference))
+		{
+			if (const FDataflowInput* Input = FindInput(Reference))
+			{
+				// we need to pull the value first so the upstream of the graph evaluate 
+				Input->template GetValue<T>(Context, *Reference);
+				Output->ForwardInput(InputReference, Context);
+			}
+			else
+			{
+				checkfSlow(false, TEXT("This input could not be found within this node, check this has been properly registered in the node constructor"));
+			}
+		}
+		else
+		{
+			checkfSlow(false, TEXT("This output could not be found within this node, check this has been properly registered in the node constructor"));
+		}
 	}
 
 	/**
diff --git a/Engine/Source/Runtime/Experimental/Dataflow/Core/Public/Dataflow/DataflowNodeParameters.h b/Engine/Source/Runtime/Experimental/Dataflow/Core/Public/Dataflow/DataflowNodeParameters.h
--- a/Engine/Source/Runtime/Experimental/Dataflow/Core/Public/Dataflow/DataflowNodeParameters.h
+++ b/Engine/Source/Runtime/Experimental/Dataflow/Core/Public/Dataflow/DataflowNodeParameters.h
@@ -18,6 +18,10 @@ struct FDataflowConnection;
 
 namespace Dataflow
 {
+	class FContext;
+
+	typedef uint32 FContextCacheKey;
+
 	struct FTimestamp
 	{
 		typedef uint64 Type;
@@ -45,17 +49,42 @@ namespace Dataflow
 
 	struct FContextCacheElementBase 
 	{
-		FContextCacheElementBase(FGuid InNodeGuid = FGuid(), const FProperty* InProperty = nullptr, uint32 InNodeHash = 0, FTimestamp InTimestamp = FTimestamp::Invalid)
-			: NodeGuid(InNodeGuid)
+		enum EType
+		{
+			CacheElementTyped,
+			CacheElementReference
+		};
+
+		FContextCacheElementBase(EType CacheElementType, FGuid InNodeGuid = FGuid(), const FProperty* InProperty = nullptr, uint32 InNodeHash = 0, FTimestamp InTimestamp = FTimestamp::Invalid)
+			: Type(CacheElementType)
+			, NodeGuid(InNodeGuid)
 			, Property(InProperty)
 			, NodeHash(InNodeHash)
 			, Timestamp(InTimestamp)
 		{}
 		virtual ~FContextCacheElementBase() {}
 
+		// InReferenceDataKey is the key of the cache element this function is called on 
+		virtual TUniquePtr<FContextCacheElementBase> CreateReference(FContextCacheKey InReferenceDataKey) const = 0;
+
 		template<typename T>
-		const T& GetTypedData(const FProperty* PropertyIn) const;
+		const T& GetTypedData(FContext& Context, const FProperty* PropertyIn, const T& Default) const;
 		
+		EType GetType() const {	return Type; }
+
+		const FProperty* GetProperty() const { return Property; }
+		const FTimestamp& GetTimestamp() const { return Timestamp; }
+
+		const FGuid& GetNodeGuid() const { return NodeGuid; }
+		const uint32 GetNodeHash() const { return NodeHash; }
+
+		// use this with caution: setting the property of a wrong type may cause problems
+		void SetProperty(const FProperty* NewProperty) { Property = NewProperty; }
+
+	private:
+		friend struct FContextCache;
+
+		EType Type;
 		FGuid NodeGuid;
 		const FProperty* Property = nullptr;
 		uint32 NodeHash = 0;
@@ -66,24 +95,55 @@ namespace Dataflow
 	struct FContextCacheElement : public FContextCacheElementBase 
 	{
 		FContextCacheElement(FGuid InNodeGuid, const FProperty* InProperty, T&& InData, uint32 InNodeHash, FTimestamp Timestamp)
-			: FContextCacheElementBase(InNodeGuid, InProperty, InNodeHash, Timestamp)
+			: FContextCacheElementBase(EType::CacheElementTyped, InNodeGuid, InProperty, InNodeHash, Timestamp)
 			, Data(Forward<T>(InData))
 		{}
 		
+		const T& GetData(FContext& Context, const FProperty* PropertyIn, const T& Default) const;
+
+		const T& GetDataDirect() const { return Data; }
+
+		virtual TUniquePtr<FContextCacheElementBase> CreateReference(FContextCacheKey InReferenceDataKey) const override;
+
+	private:
 		typedef typename TDecay<T>::Type FDataType;  // Using universal references here means T could be either const& or an rvalue reference
 		const FDataType Data;                        // Decaying T removes any reference and gets the correct underlying storage data type
 	};
 
 	template<class T>
-	const T& FContextCacheElementBase::GetTypedData(const FProperty* InProperty) const
+	struct FContextCacheElementReference : public FContextCacheElementBase
 	{
-		check(InProperty);
-		check(Property->SameType(InProperty));
-		return static_cast<const FContextCacheElement<T>&>(*this).Data;
-	}
+		FContextCacheElementReference(FGuid InNodeGuid, const FProperty* InProperty, FContextCacheKey InDataKey, uint32 InNodeHash, FTimestamp Timestamp)
+			: FContextCacheElementBase(EType::CacheElementReference, InNodeGuid, InProperty, InNodeHash, Timestamp)
+			, DataKey(InDataKey)
+		{}
 
-	typedef uint32 FContextCacheKey;
-	
+		const T& GetData(FContext& Context, const FProperty* PropertyIn, const T& Default) const;
+
+		virtual TUniquePtr<FContextCacheElementBase> CreateReference(FContextCacheKey InReferenceDataKey) const override;
+
+	private:
+		const FContextCacheKey DataKey; // this is a key to another cache element
+	};
+
+	// cache element method implementation 
+	template<class T>
+	const T& FContextCacheElementBase::GetTypedData(FContext& Context, const FProperty* PropertyIn, const T& Default) const
+	{
+		check(PropertyIn);
+		// check(PropertyIn->IsA<T>()); // @todo(dataflow) compile error for non-class T; find alternatives
+		check(Property->SameType(PropertyIn));
+		if (Type == EType::CacheElementTyped)
+		{
+			return static_cast<const FContextCacheElement<T>&>(*this).GetData(Context, PropertyIn, Default);
+		}
+		if (Type == EType::CacheElementReference)
+		{
+			return static_cast<const FContextCacheElementReference<T>&>(*this).GetData(Context, PropertyIn, Default);
+		}
+		check(false); // should never happen
+		return Default;
+	}
 
 	struct FContextCache : public TMap<FContextCacheKey, TUniquePtr<FContextCacheElementBase>>
 	{
@@ -149,12 +209,25 @@ namespace Dataflow
 		virtual void SetDataImpl(FContextCacheKey Key, TUniquePtr<FContextCacheElementBase>&& DataStoreEntry) = 0;
 		
 		template<typename T>
-		void SetData(FContextCacheKey Key, FContextCacheElementBase&& Data, T&& Value)
+		void SetData(FContextCacheKey InKey, const FProperty* InProperty, T&& InValue, const FGuid& InNodeGuid, uint32 InNodeHash, const FTimestamp& InTimestamp)
 		{
-			FContextCacheKey IntKey = (FContextCacheKey)Key;
-			TUniquePtr<FContextCacheElement<T>> DataStoreEntry = MakeUnique<FContextCacheElement<T>>(Data.NodeGuid, Data.Property, Forward<T>(Value), Data.NodeHash, Data.Timestamp);
+			TUniquePtr<FContextCacheElement<T>> DataStoreEntry = MakeUnique<FContextCacheElement<T>>(InNodeGuid, InProperty, Forward<T>(InValue), InNodeHash, InTimestamp);
+
+			SetDataImpl(InKey, MoveTemp(DataStoreEntry));
+		}
 
-			SetDataImpl(IntKey, MoveTemp(DataStoreEntry));
+		void SetDataReference(FContextCacheKey Key, const FProperty* Property, FContextCacheKey ReferenceKey)
+		{
+			// find the reference key to get 
+			if (TUniquePtr<FContextCacheElementBase>* CacheElement = GetDataImpl(ReferenceKey))
+			{
+				TUniquePtr<FContextCacheElementBase> CacheReferenceElement = (*CacheElement)->CreateReference(ReferenceKey);
+				SetDataImpl(Key, MoveTemp(CacheReferenceElement));
+			}
+			else
+			{
+				ensure(false); // could not find the original cache element 
+			}
 		}
 
 		virtual TUniquePtr<FContextCacheElementBase>* GetDataImpl(FContextCacheKey Key) = 0;
@@ -164,7 +237,7 @@ namespace Dataflow
 		{
 			if (TUniquePtr<FContextCacheElementBase>* Cache = GetDataImpl(Key))
 			{
-				return (*Cache)->GetTypedData<T>(InProperty);
+				return (*Cache)->GetTypedData<T>(*this, InProperty, Default);
 			}
 			return Default;
 		}
@@ -249,7 +322,7 @@ namespace Dataflow
 
 		virtual bool HasDataImpl(FContextCacheKey Key, FTimestamp InTimestamp = FTimestamp::Invalid) override
 		{
-			return DataStore.Contains(Key) && DataStore[Key]->Timestamp >= InTimestamp;
+			return DataStore.Contains(Key) && DataStore[Key]->GetTimestamp() >= InTimestamp;
 		}
 
 		virtual bool IsEmptyImpl() const override
@@ -294,7 +367,7 @@ namespace Dataflow
 		{
 			CacheLock->Lock(); ON_SCOPE_EXIT { CacheLock->Unlock(); };
 			
-			return DataStore.Contains(Key) && DataStore[Key]->Timestamp >= InTimestamp;
+			return DataStore.Contains(Key) && DataStore[Key]->GetTimestamp() >= InTimestamp;
 		}
 
 		virtual bool IsEmptyImpl() const override
@@ -307,4 +380,30 @@ namespace Dataflow
 
 	};
 
+	// cache classes implemetation 
+	// this needs to be after the FContext definition because they access its methods
+
+	template<class T>
+	const T& FContextCacheElement<T>::GetData(FContext& Context, const FProperty* PropertyIn, const T& Default) const
+	{
+		return Data;
+	}
+
+	template<class T>
+	TUniquePtr<FContextCacheElementBase> FContextCacheElement<T>::CreateReference(FContextCacheKey InReferenceDataKey) const
+	{
+		return MakeUnique<FContextCacheElementReference<T>>(GetNodeGuid(), GetProperty(), InReferenceDataKey, GetNodeHash(), GetTimestamp());
+	}
+
+	template<class T>
+	const T& FContextCacheElementReference<T>::GetData(FContext& Context, const FProperty* PropertyIn, const T& Default) const
+	{
+		return Context.GetData(DataKey, PropertyIn, Default);
+	}
+
+	template<class T>
+	TUniquePtr<FContextCacheElementBase> FContextCacheElementReference<T>::CreateReference(FContextCacheKey InReferenceDataKey) const
+	{
+		return MakeUnique<FContextCacheElementReference<T>>(GetNodeGuid(), GetProperty(), InReferenceDataKey, GetNodeHash(), GetTimestamp());
+	}
 }
