Chaos - TCArray API changed to match TArray
- Add SignedSquare() utility function

diff --git a/Engine/Source/Runtime/Experimental/Chaos/Private/Chaos/Collision/CapsuleTriangleContactPoint.cpp b/Engine/Source/Runtime/Experimental/Chaos/Private/Chaos/Collision/CapsuleTriangleContactPoint.cpp
--- a/Engine/Source/Runtime/Experimental/Chaos/Private/Chaos/Collision/CapsuleTriangleContactPoint.cpp
+++ b/Engine/Source/Runtime/Experimental/Chaos/Private/Chaos/Collision/CapsuleTriangleContactPoint.cpp
@@ -211,7 +211,7 @@ namespace Chaos
 		const bool bInsideAll1 = ((EdgeD1s[0] <= DistanceTolerance) & (EdgeD1s[1] <= DistanceTolerance) & (EdgeD1s[2] <= DistanceTolerance)) != 0;
 		if ((bInsideAll0 & (FaceD0 < RejectDistance) & (FaceD0 < FaceD1 + DistanceTolerance)) != 0)
 		{
-			FContactPoint& ContactPoint = OutContactPoints[OutContactPoints.Add()];
+			FContactPoint& ContactPoint = OutContactPoints[OutContactPoints.AddUninitialized()];
 			ContactPoint.ShapeContactPoints[0] = P0 - R * FaceN;
 			ContactPoint.ShapeContactPoints[1] = P0 - FaceD0 * FaceN;
 			ContactPoint.ShapeContactNormal = FaceN;
@@ -222,7 +222,7 @@ namespace Chaos
 		}
 		if ((bInsideAll1 & (FaceD1 < RejectDistance) & (FaceD1 < FaceD0 + DistanceTolerance)) != 0)
 		{
-			FContactPoint& ContactPoint = OutContactPoints[OutContactPoints.Add()];
+			FContactPoint& ContactPoint = OutContactPoints[OutContactPoints.AddUninitialized()];
 			ContactPoint.ShapeContactPoints[0] = P1 - R * FaceN;
 			ContactPoint.ShapeContactPoints[1] = P1 - FaceD1 * FaceN;
 			ContactPoint.ShapeContactNormal = FaceN;
@@ -300,7 +300,7 @@ namespace Chaos
 					{
 						const FReal CylinderFaceD = FVec3::DotProduct(CylinderP - FaceP, FaceN);
 
-						FContactPoint& ContactPoint = OutContactPoints[OutContactPoints.Add()];
+						FContactPoint& ContactPoint = OutContactPoints[OutContactPoints.AddUninitialized()];
 						ContactPoint.ShapeContactPoints[0] = CylinderP;
 						ContactPoint.ShapeContactPoints[1] = CylinderP - CylinderFaceD * FaceN;
 						ContactPoint.ShapeContactNormal = FaceN;
@@ -429,7 +429,7 @@ namespace Chaos
 					const FVec3 CapsuleP = SegmentP - R * SegmentEdgeN;
 					const FReal CapsuleDist = FVec3::DotProduct(CapsuleP - FaceP, FaceN);
 
-					FContactPoint& ContactPoint = OutContactPoints[OutContactPoints.Add()];
+					FContactPoint& ContactPoint = OutContactPoints[OutContactPoints.AddUninitialized()];
 					ContactPoint.ShapeContactPoints[0] = CapsuleP;
 					ContactPoint.ShapeContactPoints[1] = CapsuleP - CapsuleDist * FaceN;
 					ContactPoint.ShapeContactNormal = FaceN;
@@ -439,7 +439,7 @@ namespace Chaos
 				}
 				else
 				{
-					FContactPoint& ContactPoint = OutContactPoints[OutContactPoints.Add()];
+					FContactPoint& ContactPoint = OutContactPoints[OutContactPoints.AddUninitialized()];
 					ContactPoint.ShapeContactPoints[0] = SegmentP - R * SegmentEdgeN;
 					ContactPoint.ShapeContactPoints[1] = EdgeP;
 					ContactPoint.ShapeContactNormal = SegmentEdgeN;
@@ -465,7 +465,7 @@ namespace Chaos
 				const FReal SegmentEdgeDist = FMath::Sqrt(SegmentEdgeDistSq);
 				const FVec3 SegmentEdgeN = SegmentEdgeDelta / SegmentEdgeDist;
 
-				FContactPoint& ContactPoint = OutContactPoints[OutContactPoints.Add()];
+				FContactPoint& ContactPoint = OutContactPoints[OutContactPoints.AddUninitialized()];
 				ContactPoint.ShapeContactPoints[0] = SegmentEdgeC + (SegmentEdgeDist - R) * SegmentEdgeN;
 				ContactPoint.ShapeContactPoints[1] = SegmentEdgeC;
 				ContactPoint.ShapeContactNormal = SegmentEdgeN;
diff --git a/Engine/Source/Runtime/Experimental/Chaos/Private/Chaos/Collision/ConvexTriangleContactPoint.cpp b/Engine/Source/Runtime/Experimental/Chaos/Private/Chaos/Collision/ConvexTriangleContactPoint.cpp
--- a/Engine/Source/Runtime/Experimental/Chaos/Private/Chaos/Collision/ConvexTriangleContactPoint.cpp
+++ b/Engine/Source/Runtime/Experimental/Chaos/Private/Chaos/Collision/ConvexTriangleContactPoint.cpp
@@ -466,7 +466,7 @@ namespace Chaos
 		{
 			if (Distance < CullDistance)
 			{
-				FContactPoint& ContactPoint = OutContactPoints[OutContactPoints.Add()];
+				FContactPoint& ContactPoint = OutContactPoints[OutContactPoints.AddUninitialized()];
 				ContactPoint.ShapeContactPoints[0] = ConvexX;
 				ContactPoint.ShapeContactPoints[1] = TriX;
 				ContactPoint.ShapeContactNormal = SeparatingAxis;
diff --git a/Engine/Source/Runtime/Experimental/Chaos/Private/Chaos/CollisionOneShotManifolds.cpp b/Engine/Source/Runtime/Experimental/Chaos/Private/Chaos/CollisionOneShotManifolds.cpp
--- a/Engine/Source/Runtime/Experimental/Chaos/Private/Chaos/CollisionOneShotManifolds.cpp
+++ b/Engine/Source/Runtime/Experimental/Chaos/Private/Chaos/CollisionOneShotManifolds.cpp
@@ -1151,7 +1151,7 @@ namespace Chaos
 							const FReal ConvexContactPointDistance = FVec3::DotProduct((ConvexPlanePosition - CapsuleContactPoint), ConvexPlaneNormal) * ConvexDistanceMultiplier;
 							const FVec3 ConvexContactPoint = CapsuleContactPoint + ConvexContactPointDistance * CapsulePlaneNormal;
 
-							FContactPoint& ContactPoint = OutContactPoints[OutContactPoints.Add()];
+							FContactPoint& ContactPoint = OutContactPoints[OutContactPoints.AddUninitialized()];
 							ContactPoint.ShapeContactPoints[0] = CapsuleToConvexTransform.InverseTransformPositionNoScale(CapsuleContactPoint);
 							ContactPoint.ShapeContactPoints[1] = ConvexContactPoint;
 							ContactPoint.ShapeContactNormal = CapsulePlaneNormal;
@@ -1205,7 +1205,7 @@ namespace Chaos
 						{
 							const FVec3 ConvexContactPoint = CapsuleContactPoint - FVec3::DotProduct(CapsuleContactPoint - ConvexPlanePosition, ConvexPlaneNormal) * ConvexPlaneNormal;
 
-							FContactPoint& ContactPoint = OutContactPoints[OutContactPoints.Add()];
+							FContactPoint& ContactPoint = OutContactPoints[OutContactPoints.AddUninitialized()];
 							ContactPoint.ShapeContactPoints[0] = CapsuleToConvexTransform.InverseTransformPositionNoScale(CapsuleContactPoint);
 							ContactPoint.ShapeContactPoints[1] = ConvexContactPoint;
 							ContactPoint.ShapeContactNormal = ConvexPlaneNormal;
@@ -1444,7 +1444,7 @@ namespace Chaos
 				// Triangle face contact (clipped vertices are convex vertices)
 				for (int32 ContactPointIndex = 0; ContactPointIndex < ContactPointCount; ++ContactPointIndex)
 				{
-					FContactPoint& ContactPoint = OutContactPoints[OutContactPoints.Add()];
+					FContactPoint& ContactPoint = OutContactPoints[OutContactPoints.AddUninitialized()];
 					const FVec3& ConvexContactPoint = ClippedVertices[ContactPointIndex];
 					const FVec3 TriangleContactPoint = ConvexContactPoint - FVec3::DotProduct(ConvexContactPoint - TrianglePlanePosition, TrianglePlaneNormal) * TrianglePlaneNormal;
 
@@ -1460,7 +1460,7 @@ namespace Chaos
 				// Convex face contact (clipped vertices are triangle vertices)
 				for (int32 ContactPointIndex = 0; ContactPointIndex < ContactPointCount; ++ContactPointIndex)
 				{
-					FContactPoint& ContactPoint = OutContactPoints[OutContactPoints.Add()];
+					FContactPoint& ContactPoint = OutContactPoints[OutContactPoints.AddUninitialized()];
 					const FVec3& TriangleContactPoint = ClippedVertices[ContactPointIndex];
 					const FVec3 ConvexContactPoint = TriangleContactPoint - FVec3::DotProduct(TriangleContactPoint - ConvexPlanePosition, ConvexPlaneNormal) * ConvexPlaneNormal;
 
diff --git a/Engine/Source/Runtime/Experimental/Chaos/Private/Chaos/CollisionOneShotManifoldsMiscShapes.cpp b/Engine/Source/Runtime/Experimental/Chaos/Private/Chaos/CollisionOneShotManifoldsMiscShapes.cpp
--- a/Engine/Source/Runtime/Experimental/Chaos/Private/Chaos/CollisionOneShotManifoldsMiscShapes.cpp
+++ b/Engine/Source/Runtime/Experimental/Chaos/Private/Chaos/CollisionOneShotManifoldsMiscShapes.cpp
@@ -268,7 +268,7 @@ namespace Chaos
 						const FReal ConvexContactDistance = Utilities::RaySphereIntersectionDistance(ConvexPlaneVertex, ClosestContactPoint.ShapeContactNormal, SpherePos, SphereRadius);
 						if (ConvexContactDistance < CullDistance)
 						{
-							FContactPoint& ConvexContactPoint = ContactPoints[ContactPoints.Add()];
+							FContactPoint& ConvexContactPoint = ContactPoints[ContactPoints.AddUninitialized()];
 							ConvexContactPoint.ShapeContactPoints[0] = SphereToConvexTransform.InverseTransformPositionNoScale(ConvexPlaneVertex + ClosestContactPoint.ShapeContactNormal * ConvexContactDistance);
 							ConvexContactPoint.ShapeContactPoints[1] = ConvexPlaneVertex;
 							ConvexContactPoint.ShapeContactNormal = ClosestContactPoint.ShapeContactNormal;
diff --git a/Engine/Source/Runtime/Experimental/Chaos/Public/Chaos/Collision/PBDCollisionSolverJacobi.h b/Engine/Source/Runtime/Experimental/Chaos/Public/Chaos/Collision/PBDCollisionSolverJacobi.h
--- a/Engine/Source/Runtime/Experimental/Chaos/Public/Chaos/Collision/PBDCollisionSolverJacobi.h
+++ b/Engine/Source/Runtime/Experimental/Chaos/Public/Chaos/Collision/PBDCollisionSolverJacobi.h
@@ -60,7 +60,7 @@ namespace Chaos
 
 			FORCEINLINE_DEBUGGABLE int32 AddManifoldPoint()
 			{
-				return ManifoldPoints.Add();
+				return ManifoldPoints.AddUninitialized();
 			}
 
 			FORCEINLINE_DEBUGGABLE void SetWorldContact(
diff --git a/Engine/Source/Runtime/Experimental/Chaos/Public/Chaos/Framework/UncheckedArray.h b/Engine/Source/Runtime/Experimental/Chaos/Public/Chaos/Framework/UncheckedArray.h
--- a/Engine/Source/Runtime/Experimental/Chaos/Public/Chaos/Framework/UncheckedArray.h
+++ b/Engine/Source/Runtime/Experimental/Chaos/Public/Chaos/Framework/UncheckedArray.h
@@ -144,7 +144,16 @@ namespace Chaos
 		/**
 		 * @brief Increase the size of the array without re-initializing the new element
 		*/
+		UE_DEPRECATED(5.5, "Renamed to AddUinitialized() to match TArray API")
 		inline int32 Add()
+		{
+			return AddUninitialized();
+		}
+
+		/**
+		 * @brief Increase the size of the array without re-initializing the new element
+		*/
+		inline int32 AddUninitialized()
 		{
 			check(NumElements < MaxElements);
 			CheckSentinels();
diff --git a/Engine/Source/Runtime/Experimental/Chaos/Public/Chaos/Utilities.h b/Engine/Source/Runtime/Experimental/Chaos/Public/Chaos/Utilities.h
--- a/Engine/Source/Runtime/Experimental/Chaos/Public/Chaos/Utilities.h
+++ b/Engine/Source/Runtime/Experimental/Chaos/Public/Chaos/Utilities.h
@@ -15,6 +15,15 @@ namespace Chaos
 {
 	namespace Utilities
 	{
+		// Get the signed square of the input, i.e. (A * Abs(A)) or (Sign(A) * Square(A))
+		// SignedSquare(2) -> 4
+		// SignedSquare(-2) -> -4
+		template<typename TRealType>
+		inline TRealType SignedSquare(TRealType A)
+		{
+			return A * FMath::Abs(A);
+		}
+
 		//! Take the factorial of \p Num, which should be of integral type.
 		template<class TINT = uint64>
 		TINT Factorial(TINT Num)
