Chaos - set solver config before the first Advance

Any particles that exist in the scene at the start were not seeing the correct settings.

diff --git a/Engine/Source/Runtime/Experimental/Chaos/Private/Chaos/Collision/PBDCollisionConstraint.cpp b/Engine/Source/Runtime/Experimental/Chaos/Private/Chaos/Collision/PBDCollisionConstraint.cpp
--- a/Engine/Source/Runtime/Experimental/Chaos/Private/Chaos/Collision/PBDCollisionConstraint.cpp
+++ b/Engine/Source/Runtime/Experimental/Chaos/Private/Chaos/Collision/PBDCollisionConstraint.cpp
@@ -332,12 +332,10 @@ namespace Chaos
 		Flags.bUseManifold = bInUseManifold;
 		Flags.bUseIncrementalManifold = false;
 
-		// The initial overlap depenetration velocity is that max of the two body settings. A negative value means use the system default.
-		// If either body has a zero of positive value, it will be used instead of the system setting. If both have a value, the max will be used.
-		// If both bodies have negative values, the solver settings will be used (search for GetInitialOverlapDepentrationVelocity to see where)
+		// The initial overlap depenetration velocity is that max of the two body settings. I.e., if either body wants to depenetrate, they will.
 		const FRealSingle InitialOverlapDepenetrationVelocity0 = FConstGenericParticleHandle(GetParticle0())->InitialOverlapDepenetrationVelocity();
 		const FRealSingle InitialOverlapDepenetrationVelocity1 = FConstGenericParticleHandle(GetParticle1())->InitialOverlapDepenetrationVelocity();
-		InitialOverlapDepenetrationVelocity = FMath::Max(InitialOverlapDepenetrationVelocity0, InitialOverlapDepenetrationVelocity1);
+		InitialOverlapDepenetrationVelocity = FMath::Max3(InitialOverlapDepenetrationVelocity0, InitialOverlapDepenetrationVelocity1, 0.0f);
 
 		// Is this a one-way interaction? A dynamic one-way interaction particle that hits a kinematic non one-way interaction particle should still be considered one-way.
 		const bool bDynamic0 = FConstGenericParticleHandle(GetParticle0())->IsDynamic();
@@ -1213,7 +1211,7 @@ namespace Chaos
 					ManifoldPoints[PointIndex].Flags.bHasStaticFrictionAnchor = true;
 					ManifoldPoints[PointIndex].ShapeAnchorPoints[0] = SavedManifoldPoints[SavedManifoldPointIndex].ShapeContactPoints[0];
 					ManifoldPoints[PointIndex].ShapeAnchorPoints[1] = SavedManifoldPoints[SavedManifoldPointIndex].ShapeContactPoints[1];
-					ManifoldPoints[PointIndex].InitialPhi = SavedManifoldPoints[SavedManifoldPointIndex].InitialPhi;
+					//ManifoldPoints[PointIndex].InitialPhi = SavedManifoldPoints[SavedManifoldPointIndex].InitialPhi;
 					ManifoldPoints[PointIndex].Flags.bInitialContact = false;
 				}
 				// Nothing to do if no saved friction point because we already set the achor to the most recently detected contact point
diff --git a/Engine/Source/Runtime/Experimental/Chaos/Private/ChaosSolversModule.cpp b/Engine/Source/Runtime/Experimental/Chaos/Private/ChaosSolversModule.cpp
--- a/Engine/Source/Runtime/Experimental/Chaos/Private/ChaosSolversModule.cpp
+++ b/Engine/Source/Runtime/Experimental/Chaos/Private/ChaosSolversModule.cpp
@@ -180,11 +180,7 @@ void FChaosSolversModule::SyncTask(bool bForceBlockingSync /*= false*/)
 #endif
 }
 
-Chaos::FPBDRigidsSolver* FChaosSolversModule::CreateSolver(UObject* InOwner, Chaos::FReal InAsyncDt, Chaos::EThreadingMode InThreadingMode
-#if CHAOS_DEBUG_NAME
-	, const FName& DebugName
-#endif
-)
+Chaos::FPBDRigidsSolver* FChaosSolversModule::CreateSolver(UObject* InOwner, Chaos::FReal InAsyncDt, Chaos::EThreadingMode InThreadingMode, const FName& DebugName)
 {
 	LLM_SCOPE(ELLMTag::Chaos);
 	using namespace Chaos;
@@ -192,7 +188,7 @@ Chaos::FPBDRigidsSolver* FChaosSolversModule::CreateSolver(UObject* InOwner, Cha
 	FChaosScopeSolverLock SolverScopeLock;
 	
 	EMultiBufferMode SolverBufferMode = InThreadingMode == EThreadingMode::SingleThread ? EMultiBufferMode::Single : EMultiBufferMode::Double;
-	auto* NewSolver = new FPBDRigidsSolver(SolverBufferMode,InOwner, InAsyncDt);
+	FPBDRigidsSolver* NewSolver = new FPBDRigidsSolver(SolverBufferMode, InOwner, InAsyncDt);
 	AllSolvers.Add(NewSolver);
 
 	// Add The solver to the owner list
diff --git a/Engine/Source/Runtime/Experimental/Chaos/Private/PBDRigidsSolver.cpp b/Engine/Source/Runtime/Experimental/Chaos/Private/PBDRigidsSolver.cpp
--- a/Engine/Source/Runtime/Experimental/Chaos/Private/PBDRigidsSolver.cpp
+++ b/Engine/Source/Runtime/Experimental/Chaos/Private/PBDRigidsSolver.cpp
@@ -1281,90 +1281,8 @@
 
 	void FPBDRigidsSolver::PrepareAdvanceBy(const FReal DeltaTime)
 	{
-		MEvolution->GetCollisionConstraints().SetCollisionsEnabled(bChaosSolverCollisionEnabled);
+		// Handle runtime cvar changes for debugging
+		ApplyCVars();
 
-		FCollisionDetectorSettings CollisionDetectorSettings = MEvolution->GetCollisionConstraints().GetDetectorSettings();
-		CollisionDetectorSettings.bAllowManifoldReuse = (ChaosSolverCollisionAllowManifoldUpdate != 0);
-		CollisionDetectorSettings.bDeferNarrowPhase = (ChaosSolverCollisionDeferNarrowPhase != 0);
-		CollisionDetectorSettings.bAllowManifolds = (ChaosSolverCollisionUseManifolds != 0);
-		CollisionDetectorSettings.bAllowCCD = bChaosUseCCD;
-		CollisionDetectorSettings.bAllowMACD = bChaosUseMACD;
-		MEvolution->GetCollisionConstraints().SetDetectorSettings(CollisionDetectorSettings);
-		
-		FPBDJointSolverSettings JointsSettings = MEvolution->GetJointConstraints().GetSettings();
-		JointsSettings.MinSolverStiffness = ChaosSolverJointMinSolverStiffness;
-		JointsSettings.MaxSolverStiffness = ChaosSolverJointMaxSolverStiffness;
-		JointsSettings.NumIterationsAtMaxSolverStiffness = ChaosSolverJointNumIterationsAtMaxSolverStiffness;
-		JointsSettings.PositionTolerance = ChaosSolverJointPositionTolerance;
-		JointsSettings.AngleTolerance = ChaosSolverJointAngleTolerance;
-		JointsSettings.MinParentMassRatio = ChaosSolverJointMinParentMassRatio;
-		JointsSettings.MaxInertiaRatio = ChaosSolverJointMaxInertiaRatio;
-		JointsSettings.bSolvePositionLast = bChaosSolverJointSolvePositionLast;
-		JointsSettings.bUsePositionBasedDrives = bChaosSolverJointUsePositionBasedDrives;
-		JointsSettings.NumShockPropagationIterations = ChaosSolverJointNumShockProagationIterations;
-		JointsSettings.ShockPropagationOverride = ChaosSolverJointShockPropagation;
-		JointsSettings.bUseLinearSolver = bChaosSolverJointUseLinearSolver;
-		JointsSettings.bSortEnabled = false;
-		MEvolution->GetJointConstraints().SetSettings(JointsSettings);
-
-		// Apply CVAR overrides if set
-		{
-			// To enable runtime support for switching collision features on/off we need to update existing constraints when config changes.
-			if (bChaosCollisionConfigChanged)
-			{
-				// For now destroy the collisions. This is a bit over the top and causes problems for sleeping islands, but it's only for debugging/testing.
-				GetEvolution()->DestroyTransientConstraints();
-				bChaosCollisionConfigChanged = false;
-			}
-
-			if (ChaosSolverCollisionPositionFrictionIterations >= 0)
-			{
-				MEvolution->GetCollisionConstraints().SetPositionFrictionIterations(ChaosSolverCollisionPositionFrictionIterations);
-			}
-			if (ChaosSolverCollisionVelocityFrictionIterations >= 0)
-			{
-				MEvolution->GetCollisionConstraints().SetVelocityFrictionIterations(ChaosSolverCollisionVelocityFrictionIterations);
-			}
-			{
-				MEvolution->SetShockPropagationIterations(ChaosSolverCollisionPositionShockPropagationIterations, ChaosSolverCollisionVelocityShockPropagationIterations);
-			}
-			if (ChaosSolverPositionIterations >= 0)
-			{
-				SetPositionIterations(ChaosSolverPositionIterations);
-			}
-			if (ChaosSolverVelocityIterations >= 0)
-			{
-				SetVelocityIterations(ChaosSolverVelocityIterations);
-			}
-			if (ChaosSolverProjectionIterations >= 0)
-			{
-				SetProjectionIterations(ChaosSolverProjectionIterations);
-			}
-			if (ChaosSolverCullDistance >= 0.0f)
-			{
-				SetCollisionCullDistance(ChaosSolverCullDistance);
-			}
-			if ((ChaosSolverVelocityBoundsMultiplier >= 0.0f) && (ChaosSolverMaxVelocityBoundsExpansion >= 0.0f))
-			{
-				SetVelocityBoundsExpansion(ChaosSolverVelocityBoundsMultiplier, ChaosSolverMaxVelocityBoundsExpansion);
-			}
-			if ((ChaosSolverVelocityBoundsMultiplierMACD >= 0.0f) && (ChaosSolverMaxVelocityBoundsExpansionMACD >= 0.0f))
-			{
-				SetVelocityBoundsExpansionMACD(ChaosSolverVelocityBoundsMultiplierMACD, ChaosSolverMaxVelocityBoundsExpansionMACD);
-			}
-			if (ChaosSolverMaxPushOutVelocity >= 0.0f)
-			{
-				SetCollisionMaxPushOutVelocity(ChaosSolverMaxPushOutVelocity);
-			}
-			if (ChaosSolverDepenetrationVelocity >= 0.0f)
-			{
-				SetCollisionDepenetrationVelocity(ChaosSolverDepenetrationVelocity);
-			}
-			if (ChaosSolverDeterministic >= 0)
-			{
-				UpdateIsDeterministic();
-			}
-		}
-
 		UE_LOG(LogPBDRigidsSolver, Verbose, TEXT("PBDRigidsSolver::Tick(%3.5f)"), DeltaTime);
 		MLastDt = DeltaTime;
@@ -2758,11 +2676,18 @@
 			Rigid->SetInertiaConditioningDirty();
 		}
 
+		// A negative depenetration velocity means use the solver setting
+		FRealSingle DepenetrationVelocity = DynamicMisc.InitialOverlapDepenetrationVelocity();
+		if (DepenetrationVelocity < 0)
+		{
+			DepenetrationVelocity = GetEvolution()->GetCollisionConstraints().GetSolverSettings().DepenetrationVelocity;
+		}
+
 		Rigid->SetLinearEtherDrag(DynamicMisc.LinearEtherDrag());
 		Rigid->SetAngularEtherDrag(DynamicMisc.AngularEtherDrag());
 		Rigid->SetMaxLinearSpeedSq(DynamicMisc.MaxLinearSpeedSq());
 		Rigid->SetMaxAngularSpeedSq(DynamicMisc.MaxAngularSpeedSq());
-		Rigid->SetInitialOverlapDepenetrationVelocity(DynamicMisc.InitialOverlapDepenetrationVelocity());
+		Rigid->SetInitialOverlapDepenetrationVelocity(DepenetrationVelocity);
 		Rigid->SetSleepThresholdMultiplier(DynamicMisc.SleepThresholdMultiplier());
 		Rigid->SetCollisionGroup(DynamicMisc.CollisionGroup());
 		Rigid->SetDisabled(DynamicMisc.Disabled());
@@ -2810,6 +2735,96 @@
 		SetCollisionFilterSettings(InConfig.CollisionFilterSettings);
 		SetBreakingFilterSettings(InConfig.BreakingFilterSettings);
 		SetTrailingFilterSettings(InConfig.TrailingFilterSettings);
+
+		ApplyCVars();
 	}
 
+	void FPBDRigidsSolver::ApplyCVars()
+	{
+		MEvolution->GetCollisionConstraints().SetCollisionsEnabled(bChaosSolverCollisionEnabled);
+
+		FCollisionDetectorSettings CollisionDetectorSettings = MEvolution->GetCollisionConstraints().GetDetectorSettings();
+		CollisionDetectorSettings.bAllowManifoldReuse = (ChaosSolverCollisionAllowManifoldUpdate != 0);
+		CollisionDetectorSettings.bDeferNarrowPhase = (ChaosSolverCollisionDeferNarrowPhase != 0);
+		CollisionDetectorSettings.bAllowManifolds = (ChaosSolverCollisionUseManifolds != 0);
+		CollisionDetectorSettings.bAllowCCD = bChaosUseCCD;
+		CollisionDetectorSettings.bAllowMACD = bChaosUseMACD;
+		MEvolution->GetCollisionConstraints().SetDetectorSettings(CollisionDetectorSettings);
+
+		FPBDJointSolverSettings JointsSettings = MEvolution->GetJointConstraints().GetSettings();
+		JointsSettings.MinSolverStiffness = ChaosSolverJointMinSolverStiffness;
+		JointsSettings.MaxSolverStiffness = ChaosSolverJointMaxSolverStiffness;
+		JointsSettings.NumIterationsAtMaxSolverStiffness = ChaosSolverJointNumIterationsAtMaxSolverStiffness;
+		JointsSettings.PositionTolerance = ChaosSolverJointPositionTolerance;
+		JointsSettings.AngleTolerance = ChaosSolverJointAngleTolerance;
+		JointsSettings.MinParentMassRatio = ChaosSolverJointMinParentMassRatio;
+		JointsSettings.MaxInertiaRatio = ChaosSolverJointMaxInertiaRatio;
+		JointsSettings.bSolvePositionLast = bChaosSolverJointSolvePositionLast;
+		JointsSettings.bUsePositionBasedDrives = bChaosSolverJointUsePositionBasedDrives;
+		JointsSettings.NumShockPropagationIterations = ChaosSolverJointNumShockProagationIterations;
+		JointsSettings.ShockPropagationOverride = ChaosSolverJointShockPropagation;
+		JointsSettings.bUseLinearSolver = bChaosSolverJointUseLinearSolver;
+		JointsSettings.bSortEnabled = false;
+		MEvolution->GetJointConstraints().SetSettings(JointsSettings);
+
+		// Apply CVAR overrides if set
+		{
+			// To enable runtime support for switching collision features on/off we need to update existing constraints when config changes.
+			if (bChaosCollisionConfigChanged)
+			{
+				// For now destroy the collisions. This is a bit over the top and causes problems for sleeping islands, but it's only for debugging/testing.
+				GetEvolution()->DestroyTransientConstraints();
+				bChaosCollisionConfigChanged = false;
+			}
+
+			if (ChaosSolverCollisionPositionFrictionIterations >= 0)
+			{
+				MEvolution->GetCollisionConstraints().SetPositionFrictionIterations(ChaosSolverCollisionPositionFrictionIterations);
+			}
+			if (ChaosSolverCollisionVelocityFrictionIterations >= 0)
+			{
+				MEvolution->GetCollisionConstraints().SetVelocityFrictionIterations(ChaosSolverCollisionVelocityFrictionIterations);
+			}
+			{
+				MEvolution->SetShockPropagationIterations(ChaosSolverCollisionPositionShockPropagationIterations, ChaosSolverCollisionVelocityShockPropagationIterations);
+			}
+			if (ChaosSolverPositionIterations >= 0)
+			{
+				SetPositionIterations(ChaosSolverPositionIterations);
+			}
+			if (ChaosSolverVelocityIterations >= 0)
+			{
+				SetVelocityIterations(ChaosSolverVelocityIterations);
+			}
+			if (ChaosSolverProjectionIterations >= 0)
+			{
+				SetProjectionIterations(ChaosSolverProjectionIterations);
+			}
+			if (ChaosSolverCullDistance >= 0.0f)
+			{
+				SetCollisionCullDistance(ChaosSolverCullDistance);
+			}
+			if ((ChaosSolverVelocityBoundsMultiplier >= 0.0f) && (ChaosSolverMaxVelocityBoundsExpansion >= 0.0f))
+			{
+				SetVelocityBoundsExpansion(ChaosSolverVelocityBoundsMultiplier, ChaosSolverMaxVelocityBoundsExpansion);
+			}
+			if ((ChaosSolverVelocityBoundsMultiplierMACD >= 0.0f) && (ChaosSolverMaxVelocityBoundsExpansionMACD >= 0.0f))
+			{
+				SetVelocityBoundsExpansionMACD(ChaosSolverVelocityBoundsMultiplierMACD, ChaosSolverMaxVelocityBoundsExpansionMACD);
+			}
+			if (ChaosSolverMaxPushOutVelocity >= 0.0f)
+			{
+				SetCollisionMaxPushOutVelocity(ChaosSolverMaxPushOutVelocity);
+			}
+			if (ChaosSolverDepenetrationVelocity >= 0.0f)
+			{
+				SetCollisionDepenetrationVelocity(ChaosSolverDepenetrationVelocity);
+			}
+			if (ChaosSolverDeterministic >= 0)
+			{
+				UpdateIsDeterministic();
+			}
+		}
+	}
+
 	FPBDRigidsSolver::~FPBDRigidsSolver()
diff --git a/Engine/Source/Runtime/Experimental/Chaos/Public/ChaosSolversModule.h b/Engine/Source/Runtime/Experimental/Chaos/Public/ChaosSolversModule.h
--- a/Engine/Source/Runtime/Experimental/Chaos/Public/ChaosSolversModule.h
+++ b/Engine/Source/Runtime/Experimental/Chaos/Public/ChaosSolversModule.h
@@ -142,11 +142,7 @@ public:
 	 * @param ThreadingMode The desired threading mode the solver will use
 	 * @param bStandalone Whether the solver is standalone (not sent to physics thread - updating left to caller)
 	 */
-	CHAOS_API Chaos::FPBDRigidsSolver* CreateSolver(UObject* InOwner, Chaos::FReal InAsyncDt, Chaos::EThreadingMode ThreadingMode = Chaos::EThreadingMode::SingleThread
-#if CHAOS_DEBUG_NAME
-		, const FName& DebugName = NAME_None
-#endif
-	);
+	CHAOS_API Chaos::FPBDRigidsSolver* CreateSolver(UObject* InOwner, Chaos::FReal InAsyncDt, Chaos::EThreadingMode ThreadingMode = Chaos::EThreadingMode::SingleThread, const FName& DebugName = NAME_None);
 
 	CHAOS_API void MigrateSolver(Chaos::FPhysicsSolverBase* InSolver, const UObject* InNewOwner);
 
diff --git a/Engine/Source/Runtime/Experimental/Chaos/Public/PBDRigidsSolver.h b/Engine/Source/Runtime/Experimental/Chaos/Public/PBDRigidsSolver.h
--- a/Engine/Source/Runtime/Experimental/Chaos/Public/PBDRigidsSolver.h
+++ b/Engine/Source/Runtime/Experimental/Chaos/Public/PBDRigidsSolver.h
@@ -359,6 +359,8 @@ namespace Chaos
 		CHAOS_API virtual void PushPhysicsState(const FReal ExternalDt, const int32 NumSteps, const int32 NumExternalSteps) override;
 		CHAOS_API virtual void SetExternalTimestampConsumed_Internal(const int32 Timestamp) override;
 
+		CHAOS_API void ApplyCVars();
+
 		CHAOS_API void UpdateIsDeterministic();
 
 		CHAOS_API void DebugDrawShapes(const bool bShowStatic, const bool bShowKinematic, const bool bShowDynamic) const;
diff --git a/Engine/Source/Runtime/PhysicsCore/Private/ChaosScene.cpp b/Engine/Source/Runtime/PhysicsCore/Private/ChaosScene.cpp
--- a/Engine/Source/Runtime/PhysicsCore/Private/ChaosScene.cpp
+++ b/Engine/Source/Runtime/PhysicsCore/Private/ChaosScene.cpp
@@ -54,9 +54,7 @@ TAutoConsoleVariable<int32> CVar_ApplyProjectSettings(TEXT("p.Chaos.Simulation.A
 FChaosScene::FChaosScene(
 	UObject* OwnerPtr
 	, Chaos::FReal InAsyncDt
-#if CHAOS_DEBUG_NAME
 	, const FName& DebugName
-#endif
 )
 	: SolverAccelerationStructure(nullptr)
 	, ChaosModule(nullptr)
@@ -72,11 +70,7 @@ FChaosScene::FChaosScene(
 
 	Chaos::EThreadingMode ThreadingMode = bForceSingleThread ? Chaos::EThreadingMode::SingleThread : Chaos::EThreadingMode::TaskGraph;
 
-	SceneSolver = ChaosModule->CreateSolver(OwnerPtr, InAsyncDt, ThreadingMode
-#if CHAOS_DEBUG_NAME
-		,DebugName
-#endif
-		);
+	SceneSolver = ChaosModule->CreateSolver(OwnerPtr, InAsyncDt, ThreadingMode, DebugName);
 	check(SceneSolver);
 
 #if WITH_CHAOS_VISUAL_DEBUGGER
@@ -89,14 +83,11 @@ FChaosScene::FChaosScene(
 	SimCallback = SceneSolver->CreateAndRegisterSimCallbackObject_External<FChaosSceneSimCallback>();
 
 	// Apply project settings to the solver
-	if(CVar_ApplyProjectSettings.GetValueOnAnyThread() != 0)
+	if (CVar_ApplyProjectSettings.GetValueOnAnyThread() != 0)
 	{
 		UPhysicsSettingsCore* Settings = UPhysicsSettingsCore::Get();
-		SceneSolver->RegisterSimOneShotCallback([InSolver = SceneSolver, SolverConfigCopy = Settings->SolverOptions, bIsDeterministic = Settings->bEnableEnhancedDeterminism]()
-		{
-			InSolver->ApplyConfig(SolverConfigCopy);
-			InSolver->SetIsDeterministic(bIsDeterministic);
-		});
+		SceneSolver->ApplyConfig(Settings->SolverOptions);
+		SceneSolver->SetIsDeterministic(Settings->bEnableEnhancedDeterminism);
 	}
 
 	// Make sure we have initialized structure on game thread, evolution has already initialized structure, just need to copy.
diff --git a/Engine/Source/Runtime/PhysicsCore/Public/Chaos/ChaosScene.h b/Engine/Source/Runtime/PhysicsCore/Public/Chaos/ChaosScene.h
--- a/Engine/Source/Runtime/PhysicsCore/Public/Chaos/ChaosScene.h
+++ b/Engine/Source/Runtime/PhysicsCore/Public/Chaos/ChaosScene.h
@@ -88,10 +88,8 @@ public:
 	PHYSICSCORE_API FChaosScene(
 		UObject* OwnerPtr
 		, Chaos::FReal InAsyncDt
-#if CHAOS_DEBUG_NAME
-	, const FName& DebugName=NAME_None
-#endif
-);
+		, const FName& DebugName = NAME_None
+	);
 
 	PHYSICSCORE_API virtual ~FChaosScene();
 
