Iris
- Fixed BitStream corruptions due to sending unexpected data to not yet confirmed objects
- Fixed issue where PendingFlush would not properly flush attachments if posted on the same frame as the object was destroyed if it was waiting for cration confirmation.

diff --git a/Engine/Plugins/Runtime/ReplicationSystemTestPlugin/Source/Private/Tests/ReplicationSystem/TestFlushBeforeDestroy.cpp b/Engine/Plugins/Runtime/ReplicationSystemTestPlugin/Source/Private/Tests/ReplicationSystem/TestFlushBeforeDestroy.cpp
--- a/Engine/Plugins/Runtime/ReplicationSystemTestPlugin/Source/Private/Tests/ReplicationSystem/TestFlushBeforeDestroy.cpp
+++ b/Engine/Plugins/Runtime/ReplicationSystemTestPlugin/Source/Private/Tests/ReplicationSystem/TestFlushBeforeDestroy.cpp
@@ -106,6 +106,125 @@ UE_NET_TEST_FIXTURE(FTestFlushBeforeDestroyFixture, TestReliableAttachmentFlushe
 	UE_NET_ASSERT_TRUE(Client->GetReplicationBridge()->GetReplicatedObject(ObjectHandle) == nullptr);
 }
 
+// This test exercises what was a bad case where we was posting RPC:s to a not yet confirmed objects which also was marked for destroy
+// This put the replication system in a state where it wrote data that the client could not process.
+// Currently we will just drop the data if the initial create packet is lost as we cannot yet send creation info for
+// destroyed objects.
+UE_NET_TEST_FIXTURE(FTestFlushBeforeDestroyFixture, TestReliableAttachmentFlushedWithPendingCreationLostBeforeDestroy)
+{
+	FReplicationSystemTestClient* Client = CreateClient();
+	RegisterNetBlobHandlers(Client);
+
+	UReplicatedTestObject* ServerObject = Server->CreateObject(0, 0);
+	FNetRefHandle ObjectHandle = ServerObject->NetRefHandle;
+
+	// Setup a situation where we have creation info in flight when the object is destroyed
+
+	// Send creation info
+	Server->PreSendUpdate();
+	Server->SendTo(Client, TEXT("WaitOnCreateConfirmation"));
+	Server->PostSendUpdate();
+
+	// Create attachment
+	{
+		constexpr uint32 PayloadBitCount = 24;
+		const TRefCountPtr<FNetObjectAttachment>& Attachment = MockNetObjectAttachmentHandler->CreateReliableNetObjectAttachment(PayloadBitCount);
+		FNetObjectReference AttachmentTarget = FObjectReferenceCache::MakeNetObjectReference(ServerObject->NetRefHandle);
+		Server->GetReplicationSystem()->QueueNetObjectAttachment(Client->ConnectionIdOnServer, AttachmentTarget, Attachment);
+	}
+
+	// Destroy object, on server
+	Server->DestroyObject(ServerObject);
+
+	// Previously this would issue a flush and send the attachment data even though creation was not yet confirmed leading to a client disconnect..
+	Server->PreSendUpdate();
+	const bool bDataWasSent = Server->SendTo(Client, TEXT("State should still be WaitOnCreateConfirmation"));
+	Server->PostSendUpdate();
+
+	// We do not expect any data to be in this packet.
+	UE_NET_ASSERT_FALSE(bDataWasSent);
+
+	// Drop the data and notify server
+	Server->DeliverTo(Client, false);
+
+	// Deliver data
+	if (bDataWasSent)
+	{
+		// Caused bitstream error on client.
+		Server->DeliverTo(Client, true);
+	}
+
+	// Update to drive the last transition which we expect to be
+	Server->UpdateAndSend({Client});
+
+	// Verify that the attachment has not received
+	UE_NET_ASSERT_EQ(ClientMockNetObjectAttachmentHandler->GetFunctionCallCounts().OnNetBlobReceived, 0U);
+
+	// Verify that object does not exist
+	UE_NET_ASSERT_TRUE(Client->GetReplicationBridge()->GetReplicatedObject(ObjectHandle) == nullptr);
+}
+
+UE_NET_TEST_FIXTURE(FTestFlushBeforeDestroyFixture, TestReliableAttachmentFlushedWithPendingCreationInflightBeforeDestroy)
+{
+	FReplicationSystemTestClient* Client = CreateClient();
+	RegisterNetBlobHandlers(Client);
+
+	UReplicatedTestObject* ServerObject = Server->CreateObject(0, 0);
+	FNetRefHandle ObjectHandle = ServerObject->NetRefHandle;
+
+	// Setup a situation where we have creation info in flight when the object is destroyed
+
+	// Send creation info
+	Server->PreSendUpdate();
+	Server->SendTo(Client, TEXT("WaitOnCreateConfirmation"));
+	Server->PostSendUpdate();
+
+	// Create attachment
+	{
+		constexpr uint32 PayloadBitCount = 24;
+		const TRefCountPtr<FNetObjectAttachment>& Attachment = MockNetObjectAttachmentHandler->CreateReliableNetObjectAttachment(PayloadBitCount);
+		FNetObjectReference AttachmentTarget = FObjectReferenceCache::MakeNetObjectReference(ServerObject->NetRefHandle);
+		Server->GetReplicationSystem()->QueueNetObjectAttachment(Client->ConnectionIdOnServer, AttachmentTarget, Attachment);
+	}
+
+	// Destroy object, on server
+	Server->DestroyObject(ServerObject);
+
+	// Previously this would issue a flush and send data before creation is confirmed.
+	Server->PreSendUpdate();
+	const bool bDataWasSentInError = Server->SendTo(Client, TEXT("State should still be WaitOnCreateConfirmation"));
+	Server->PostSendUpdate();
+	
+	// We do not expect any data to be in this packet.
+	UE_NET_ASSERT_FALSE(bDataWasSentInError);
+
+	// Deliver the packet with CreationInfo
+	Server->DeliverTo(Client, true);
+
+	// Deliver data if we sent data.
+	if (bDataWasSentInError)
+	{
+		// Caused bitstream error on client.
+		Server->DeliverTo(Client, true);
+	}
+
+	// Expected to write the attachment
+	Server->PreSendUpdate();
+	Server->SendAndDeliverTo(Client, DeliverPacket, TEXT("WaitOnFlush"));
+	Server->PostSendUpdate();
+
+	// Verify that the attachment has not received
+	UE_NET_ASSERT_EQ(ClientMockNetObjectAttachmentHandler->GetFunctionCallCounts().OnNetBlobReceived, 1U);
+
+	// Expected to destroy the object
+	Server->PreSendUpdate();
+	Server->SendAndDeliverTo(Client, DeliverPacket, TEXT("Destroy"));
+	Server->PostSendUpdate();
+
+	// Verify that object does not exist
+	UE_NET_ASSERT_TRUE(Client->GetReplicationBridge()->GetReplicatedObject(ObjectHandle) == nullptr);
+}
+
 UE_NET_TEST_FIXTURE(FTestFlushBeforeDestroyFixture, TestReliableAttachmentSubObjectFlushedBeforeDestroy)
 {
 	FReplicationSystemTestClient* Client = CreateClient();
diff --git a/Engine/Source/Runtime/Experimental/Iris/Core/Private/Iris/ReplicationSystem/ReplicationWriter.cpp b/Engine/Source/Runtime/Experimental/Iris/Core/Private/Iris/ReplicationSystem/ReplicationWriter.cpp
--- a/Engine/Source/Runtime/Experimental/Iris/Core/Private/Iris/ReplicationSystem/ReplicationWriter.cpp
+++ b/Engine/Source/Runtime/Experimental/Iris/Core/Private/Iris/ReplicationSystem/ReplicationWriter.cpp
@@ -677,6 +677,17 @@ void FReplicationWriter::UpdateScope(const FNetBitArrayView& UpdatedScope)
 				ObjectsWithDirtyChanges.SetBitValue(Index, Info.HasDirtyChangeMask);
 			}
 		}
+		else if (State == EReplicatedObjectState::WaitOnCreateConfirmation)
+		{
+			// Need to restore as we might have been in case where we was pending destroy
+			ObjectsPendingDestroy.ClearBit(Index);
+			Info.FlushFlags = GetDefaultFlushFlags();
+
+			// If we have accumulated changes while waiting on flush, we should send them now
+			Info.SubObjectPendingDestroy = 0U;
+			Info.HasDirtyChangeMask |= FNetBitArrayView(Info.GetChangeMaskStoragePointer(), Info.ChangeMaskBitCount).IsAnyBitSet();
+			ObjectsWithDirtyChanges.SetBitValue(Index, Info.HasDirtyChangeMask);
+		}
 		else if (State == EReplicatedObjectState::WaitOnDestroyConfirmation || State == EReplicatedObjectState::CancelPendingDestroy)
 		{
 			// Need to clear the pending destroy bit or else the object will be masked out of ObjectsInScope.
@@ -773,12 +784,16 @@ void FReplicationWriter::UpdateScope(const FNetBitArrayView& UpdatedScope)
 			{
 				// Store info about what we need to flush
 				Info.FlushFlags = FlushFlags;
-				SetState(Index, EReplicatedObjectState::WaitOnFlush);
 
-				// If we do not have any state data to flush we can clear the has dirty states flag
-				if ((FlushFlags & FlushFlags_FlushState) == 0U)
+				if (State != EReplicatedObjectState::WaitOnCreateConfirmation)
 				{
-					Info.HasDirtyChangeMask = 0U;
+					SetState(Index, EReplicatedObjectState::WaitOnFlush);
+
+					// If we do not have any state data to flush we can clear the has dirty states flag
+					if ((FlushFlags & FlushFlags_FlushState) == 0U)
+					{
+						Info.HasDirtyChangeMask = 0U;
+					}
 				}
 
 				// Mark object as pending destroy so that we can poll the flush status in WriteObjectPendingDestroy
@@ -1042,6 +1057,11 @@ void FReplicationWriter::HandleDeliveredRecord(const FReplicationRecord::FRecord
 					{
 						SetState(InternalIndex, EReplicatedObjectState::WaitOnFlush);
 					}
+					// so are objects marked for destroy requiring flush
+					else if (ObjectsPendingDestroy.GetBit(InternalIndex))
+					{
+						SetState(InternalIndex, EReplicatedObjectState::WaitOnFlush);
+					}
 				}
 			}
 			Info.IsCreationConfirmed = 1U;
@@ -1109,12 +1129,12 @@ void FReplicationWriter::HandleDeliveredRecord(const FReplicationRecord::FRecord
 	if (Info.GetState() == EReplicatedObjectState::WaitOnFlush)
 	{
 		bool bStillPendingFlush = false;
-		if (RecordInfo.HasChangeMask && !!(Info.FlushFlags & EFlushFlags::FlushFlags_FlushState))
+		if ((RecordInfo.HasChangeMask || Info.HasDirtyChangeMask) && !!(Info.FlushFlags & EFlushFlags::FlushFlags_FlushState))
 		{
 			bStillPendingFlush |= (Info.HasDirtyChangeMask || HasInFlightStateChanges(ReplicationRecord.GetInfoForIndex(RecordInfo.NextIndex)) || IsObjectPartOfActiveHugeObject(InternalIndex));
 		}
 
-		if (RecordInfo.HasAttachments && !!(Info.FlushFlags & FlushFlags_FlushReliable))
+		if ((RecordInfo.HasAttachments || Info.HasAttachments) && !!(Info.FlushFlags & FlushFlags_FlushReliable))
 		{
 			bStillPendingFlush |= !Attachments.IsAllReliableSentAndAcked(ENetObjectAttachmentType::Normal, InternalIndex);
 		}
@@ -1208,45 +1228,56 @@ void FReplicationWriter::HandleDroppedRecord<FReplicationWriter::EReplicatedObje
 
 	if (CurrentState < EReplicatedObjectState::Created)
 	{
-		// Mark object as having dirty changes
-		MarkObjectDirty(InternalIndex, "DroppedWaitOnCreate");
-
-		// Resend creation data
-		SetState(InternalIndex, EReplicatedObjectState::PendingCreate);
+		// Until we have implemented cached creation info we cannot send creation info for destroyed objects
+		// So we just have to StopReplication
+		const bool bCanSendCreationInfo = !ObjectsPendingDestroy.GetBit(InternalIndex);
+		if (bCanSendCreationInfo)
+		{
+			// Mark object as having dirty changes
+			MarkObjectDirty(InternalIndex, "DroppedWaitOnCreate");
 
-		// Must also restore changemask
-		FNetBitArrayView ChangeMask(Info.GetChangeMaskStoragePointer(), Info.ChangeMaskBitCount);
-		FNetBitArrayView LostChangeMask = FChangeMaskUtil::MakeChangeMask(RecordInfo.ChangeMaskOrPtr, Info.ChangeMaskBitCount);
-		ChangeMask.Combine(LostChangeMask, FNetBitArrayView::OrOp);
+			// Resend creation data
+			SetState(InternalIndex, EReplicatedObjectState::PendingCreate);
 
-		// Mark changemask dirty
-		Info.HasDirtyChangeMask = 1U;
+			// Must also restore changemask
+			FNetBitArrayView ChangeMask(Info.GetChangeMaskStoragePointer(), Info.ChangeMaskBitCount);
+			FNetBitArrayView LostChangeMask = FChangeMaskUtil::MakeChangeMask(RecordInfo.ChangeMaskOrPtr, Info.ChangeMaskBitCount);
+			ChangeMask.Combine(LostChangeMask, FNetBitArrayView::OrOp);
 
-		// Indicate that we have dirty subobjects
-		Info.HasDirtySubObjects = 1U;
+			// Mark changemask dirty
+			Info.HasDirtyChangeMask = 1U;
 
-		// Mark attachments as dirty
-		Info.HasAttachments |= RecordInfo.HasAttachments;
+			// Indicate that we have dirty subobjects
+			Info.HasDirtySubObjects = 1U;
 
-		if (Info.IsSubObject)
-		{
-			// Mark owner dirty as well as subobjects only are scheduled together with owner
-			const FNetRefHandleManager::FReplicatedObjectData& ObjectData = NetRefHandleManager->GetReplicatedObjectData(InternalIndex);
-			uint32 SubObjectOwnerInternalIndex = ObjectData.SubObjectRootIndex;
+			// Mark attachments as dirty
+			Info.HasAttachments |= RecordInfo.HasAttachments;
 
-			FReplicationInfo& SubObjectOwnerReplicationInfo = GetReplicationInfo(SubObjectOwnerInternalIndex);
-			if (ensure(SubObjectOwnerReplicationInfo.GetState() < EReplicatedObjectState::PendingDestroy))
+			if (Info.IsSubObject)
 			{
-				// Mark owner as dirty
-				MarkObjectDirty(SubObjectOwnerInternalIndex, "DroppedWaitOnCreate2");
+				// Mark owner dirty as well as subobjects only are scheduled together with owner
+				const FNetRefHandleManager::FReplicatedObjectData& ObjectData = NetRefHandleManager->GetReplicatedObjectData(InternalIndex);
+				uint32 SubObjectOwnerInternalIndex = ObjectData.SubObjectRootIndex;
 
-				// Indicate that we have dirty subobjects
-				SubObjectOwnerReplicationInfo.HasDirtySubObjects = 1U;
+				FReplicationInfo& SubObjectOwnerReplicationInfo = GetReplicationInfo(SubObjectOwnerInternalIndex);
+				if (ensure(SubObjectOwnerReplicationInfo.GetState() < EReplicatedObjectState::PendingDestroy))
+				{
+					// Mark owner as dirty
+					MarkObjectDirty(SubObjectOwnerInternalIndex, "DroppedWaitOnCreate2");
+
+					// Indicate that we have dirty subobjects
+					SubObjectOwnerReplicationInfo.HasDirtySubObjects = 1U;
 
-				// Give slight priority bump to owner
-				SchedulingPriorities[SubObjectOwnerInternalIndex] += FReplicationWriter::LostStatePriorityBump;
+					// Give slight priority bump to owner
+					SchedulingPriorities[SubObjectOwnerInternalIndex] += FReplicationWriter::LostStatePriorityBump;
+				}
 			}
 		}
+		else
+		{
+			SetState(InternalIndex, EReplicatedObjectState::PendingCreate);
+			StopReplication(InternalIndex);
+		}
 	}
 	else if (CurrentState == EReplicatedObjectState::SubObjectPendingDestroy)
 	{
