[ChaosVD] Improving feedback when a recording is started but the ChaosVD editor is not available (either because you are in a Client/Server build or you are in the editor but with CVD closed).

- Now when a recording starts, a message on screen will be shown while playing or simulation indicating a recording is active
- Improved error pop up to indicate when an error was caused by using the wrong arguments

diff --git a/Engine/Source/Runtime/Experimental/ChaosSolverEngine/Private/Chaos/ChaosSolverEnginePlugin.cpp b/Engine/Source/Runtime/Experimental/ChaosSolverEngine/Private/Chaos/ChaosSolverEnginePlugin.cpp
--- a/Engine/Source/Runtime/Experimental/ChaosSolverEngine/Private/Chaos/ChaosSolverEnginePlugin.cpp
+++ b/Engine/Source/Runtime/Experimental/ChaosSolverEngine/Private/Chaos/ChaosSolverEnginePlugin.cpp
@@ -9,6 +9,7 @@
 #include "Chaos/ChaosDebugDrawComponent.h"
 #include "Chaos/ChaosSolverActor.h"
 #include "ChaosSolversModule.h"
+#include "ChaosVDRecordingStateScreenMessageHandler.h"
 
 class FChaosSolverEnginePlugin : public IChaosSolverEnginePlugin
 {
@@ -26,11 +27,20 @@ void FChaosSolverEnginePlugin::StartupModule()
 	ChaosModule->SetSolverActorClass(AChaosSolverActor::StaticClass(), AChaosSolverActor::StaticClass());
 
 	UChaosDebugDrawComponent::BindWorldDelegates();
+
+#if WITH_CHAOS_VISUAL_DEBUGGER
+	FChaosVDRecordingStateScreenMessageHandler::Get().Initialize();
+#endif
+
 }
 
 void FChaosSolverEnginePlugin::ShutdownModule()
 {
-	
+
+#if WITH_CHAOS_VISUAL_DEBUGGER
+	FChaosVDRecordingStateScreenMessageHandler::Get().TearDown();
+#endif
+
 }
 
 
diff --git a/Engine/Source/Runtime/Experimental/ChaosSolverEngine/Private/Chaos/ChaosVDRecordingStateScreenMessageHandler.cpp b/Engine/Source/Runtime/Experimental/ChaosSolverEngine/Private/Chaos/ChaosVDRecordingStateScreenMessageHandler.cpp
new file mode 100644
--- /dev/null
+++ b/Engine/Source/Runtime/Experimental/ChaosSolverEngine/Private/Chaos/ChaosVDRecordingStateScreenMessageHandler.cpp
@@ -0,0 +1,106 @@
+// Copyright Epic Games, Inc. All Rights Reserved.
+
+#include "ChaosVDRecordingStateScreenMessageHandler.h"
+
+#if WITH_CHAOS_VISUAL_DEBUGGER
+#include "ChaosVDRuntimeModule.h"
+#include "Engine/Engine.h"
+
+void FChaosVDRecordingStateScreenMessageHandler::HandleCVDRecordingStarted()
+{
+	if (!GEngine)
+	{
+		return;
+	}
+	
+	static FText ChaosVDRecordingStartedMessage = NSLOCTEXT("ChaosVisualDebugger", "OnScreenChaosVDRecordingStartedMessage", "Chaos Visual Debugger recording in progress...");
+
+	if (CVDRecordingMessageKey == 0)
+	{
+		CVDRecordingMessageKey = GetTypeHash(ChaosVDRecordingStartedMessage.ToString());
+	}
+
+	// Add a long duration value, we will remove the message manually when the recording stops
+	constexpr float MessageDurationSeconds = 3600.0f;
+	GEngine->AddOnScreenDebugMessage(CVDRecordingMessageKey, MessageDurationSeconds, FColor::Red,ChaosVDRecordingStartedMessage.ToString());
+}
+
+void FChaosVDRecordingStateScreenMessageHandler::HandleCVDRecordingStopped() const
+{
+	if (!GEngine)
+	{
+		return;
+	}
+
+	if (CVDRecordingMessageKey != 0)
+	{
+		GEngine->RemoveOnScreenDebugMessage(CVDRecordingMessageKey);
+	}
+}
+
+void FChaosVDRecordingStateScreenMessageHandler::HandleCVDRecordingStartFailed(const FText& InFailureReason) const
+{
+#if !WITH_EDITOR
+	// In non-editor builds we don't have an error pop-up, therefore we want to show the error message on screen
+	FText ErrorMessage = FText::FormatOrdered(NSLOCTEXT("ChaosVisualDebugger","StartRecordingFailedOnScreenMessage", "Failed to start CVD recording. {0}"), InFailureReason);
+
+	constexpr float MessageDurationSeconds = 4.0f;
+	GEngine->AddOnScreenDebugMessage(CVDRecordingMessageKey, MessageDurationSeconds, FColor::Red, ErrorMessage.ToString());
+#endif
+}
+
+void FChaosVDRecordingStateScreenMessageHandler::HandlePIEStarted(UGameInstance* GameInstance)
+{
+	// If we were already recording, show the message
+	if (FChaosVDRuntimeModule::Get().IsRecording())
+	{
+		HandleCVDRecordingStarted();
+	}
+}
+
+FChaosVDRecordingStateScreenMessageHandler& FChaosVDRecordingStateScreenMessageHandler::Get()
+{
+	static FChaosVDRecordingStateScreenMessageHandler MessageHandler;
+	return MessageHandler;
+}
+
+void FChaosVDRecordingStateScreenMessageHandler::Initialize()
+{
+	RecordingStartedHandle = FChaosVDRuntimeModule::Get().RegisterRecordingStartedCallback(FChaosVDRecordingStateChangedDelegate::FDelegate::CreateRaw(this, &FChaosVDRecordingStateScreenMessageHandler::HandleCVDRecordingStarted));
+	RecordingStoppedHandle = FChaosVDRuntimeModule::Get().RegisterRecordingStopCallback(FChaosVDRecordingStateChangedDelegate::FDelegate::CreateRaw(this, &FChaosVDRecordingStateScreenMessageHandler::HandleCVDRecordingStopped));
+	RecordingStartFailedHandle = FChaosVDRuntimeModule::Get().RegisterRecordingStartFailedCallback(FChaosVDRecordingStartFailedDelegate::FDelegate::CreateRaw(this, &FChaosVDRecordingStateScreenMessageHandler::HandleCVDRecordingStartFailed));
+
+#if WITH_EDITOR
+	PIEStartedHandle = FWorldDelegates::OnPIEStarted.AddRaw(this, &FChaosVDRecordingStateScreenMessageHandler::HandlePIEStarted);
+#endif
+
+	// If we were already recording, show the message
+	if (FChaosVDRuntimeModule::Get().IsRecording())
+	{
+		HandleCVDRecordingStarted();
+	}
+}
+
+void FChaosVDRecordingStateScreenMessageHandler::TearDown()
+{
+	// Note: This works during engine shutdown because the Module Manager doesn't free the dll on module unload to account for use cases like this
+	// If this appears in a callstack crash it means that assumption changed or was not correct to begin with.
+	// A possible solution is just check if the module is loaded querying the module manager just using the module's name
+	if (FChaosVDRuntimeModule::IsLoaded())
+	{
+		FChaosVDRuntimeModule::Get().RemoveRecordingStartedCallback(RecordingStartedHandle);
+		FChaosVDRuntimeModule::Get().RemoveRecordingStopCallback(RecordingStoppedHandle);
+
+#if WITH_EDITOR
+		 FWorldDelegates::OnPIEStarted.Remove(PIEStartedHandle);
+#endif
+
+		// Make sure of removing the message from the screen in case the recording didn't quite stopped yet
+		if (FChaosVDRuntimeModule::Get().IsRecording())
+		{
+			HandleCVDRecordingStopped();
+		}
+	}
+
+}
+#endif
diff --git a/Engine/Source/Runtime/Experimental/ChaosSolverEngine/Private/Chaos/ChaosVDRecordingStateScreenMessageHandler.h b/Engine/Source/Runtime/Experimental/ChaosSolverEngine/Private/Chaos/ChaosVDRecordingStateScreenMessageHandler.h
new file mode 100644
--- /dev/null
+++ b/Engine/Source/Runtime/Experimental/ChaosSolverEngine/Private/Chaos/ChaosVDRecordingStateScreenMessageHandler.h
@@ -0,0 +1,38 @@
+// Copyright Epic Games, Inc. All Rights Reserved.
+
+#pragma once
+
+#if WITH_CHAOS_VISUAL_DEBUGGER
+
+#include "Delegates/IDelegateInstance.h"
+#include "HAL/Platform.h"
+
+class UGameInstance;
+
+/** Helper object that adds a Recording in progress message in the viewport when a CVD recording is active and we are Playing or Simulation*/
+class FChaosVDRecordingStateScreenMessageHandler
+{
+public:
+
+	static FChaosVDRecordingStateScreenMessageHandler& Get();
+
+	void Initialize();
+	void TearDown();
+
+private:
+	void HandleCVDRecordingStarted();
+	void HandleCVDRecordingStopped() const;
+	void HandleCVDRecordingStartFailed(const FText& InFailureReason) const;
+	void HandlePIEStarted(UGameInstance* GameInstance);
+
+	FDelegateHandle RecordingStartedHandle;
+	FDelegateHandle RecordingStoppedHandle;
+	FDelegateHandle RecordingStartFailedHandle;
+	uint64 CVDRecordingMessageKey = 0;
+
+#if WITH_EDITOR
+	FDelegateHandle PIEStartedHandle;
+#endif
+};
+
+#endif
diff --git a/Engine/Source/Runtime/Experimental/ChaosVisualDebugger/Private/ChaosVDRuntimeModule.cpp b/Engine/Source/Runtime/Experimental/ChaosVisualDebugger/Private/ChaosVDRuntimeModule.cpp
--- a/Engine/Source/Runtime/Experimental/ChaosVisualDebugger/Private/ChaosVDRuntimeModule.cpp
+++ b/Engine/Source/Runtime/Experimental/ChaosVisualDebugger/Private/ChaosVDRuntimeModule.cpp
@@ -43,6 +43,7 @@ static FAutoConsoleVariable CVarChaosVDGTimeBetweenFullCaptures(
 
 FChaosVDRecordingStateChangedDelegate FChaosVDRuntimeModule::RecordingStartedDelegate = FChaosVDRecordingStateChangedDelegate();
 FChaosVDRecordingStateChangedDelegate FChaosVDRuntimeModule::RecordingStopDelegate = FChaosVDRecordingStateChangedDelegate();
+FChaosVDRecordingStartFailedDelegate FChaosVDRuntimeModule::RecordingStartFailedDelegate = FChaosVDRecordingStartFailedDelegate();
 FChaosVDCaptureRequestDelegate FChaosVDRuntimeModule::PerformFullCaptureDelegate = FChaosVDCaptureRequestDelegate();
 FRWLock FChaosVDRuntimeModule::DelegatesRWLock = FRWLock();
 
@@ -143,6 +144,9 @@ void FChaosVDRuntimeModule::StartRecording(TConstArrayView<FString> Args)
 	// Start Listening for Trace Stopped events, in case Trace is stopped outside our control so we can gracefully stop CVD recording and log a warning 
 	FTraceAuxiliary::OnTraceStopped.AddRaw(this, &FChaosVDRuntimeModule::HandleTraceStopRequest);
 
+	// Start with a generic Failure reason
+	FText FailureReason = LOCTEXT("SeeLogsForErrorDetailsText","Please see the logs for more details...");
+
 #if UE_TRACE_ENABLED
 
 	// Other tools could bee using trace
@@ -194,6 +198,10 @@ void FChaosVDRuntimeModule::StartRecording(TConstArrayView<FString> Args)
 		*Target,
 		nullptr, &TracingOptions);
 	}
+	else
+	{
+		FailureReason = LOCTEXT("WrongCommandArgumentsError", "The start recording command was called with invalid arguments");
+	}
 #endif
 	
 	AccumulatedRecordingTime = 0.0f;
@@ -216,16 +224,16 @@ void FChaosVDRuntimeModule::StartRecording(TConstArrayView<FString> Args)
 			FMath::Clamp(ConfiguredTimeBetweenCaptures, MinAllowedTimeInSecondsBetweenCaptures, TNumericLimits<int32>::Max()));
 
 		RecordingTimerHandle = FTSTicker::GetCoreTicker().AddTicker(FTickerDelegate::CreateRaw(this, &FChaosVDRuntimeModule::RecordingTimerTick));
-
 	}
 	else
 	{
-		UE_LOG(LogChaosVDRuntime, Error, TEXT("[%s] Failed to start CVD recording..."), ANSI_TO_TCHAR(__FUNCTION__));
+		UE_LOG(LogChaosVDRuntime, Error, TEXT("[%s] Failed to start CVD recording | Reason: [%s]"), ANSI_TO_TCHAR(__FUNCTION__), *FailureReason.ToString());
 
 #if WITH_EDITOR
-		FMessageDialog::Open(EAppMsgType::Ok, LOCTEXT("StartRecordingFailedMessage", "Failed to start CVD recording. Please see the logs for more details... "));
+		FMessageDialog::Open(EAppMsgType::Ok, FText::FormatOrdered(LOCTEXT("StartRecordingFailedMessage", "Failed to start CVD recording. \n\n{0}"), FailureReason));
 #endif
 
+		RecordingStartFailedDelegate.Broadcast(FailureReason);
 	}
 
 }
diff --git a/Engine/Source/Runtime/Experimental/ChaosVisualDebugger/Public/ChaosVDRuntimeModule.h b/Engine/Source/Runtime/Experimental/ChaosVisualDebugger/Public/ChaosVDRuntimeModule.h
--- a/Engine/Source/Runtime/Experimental/ChaosVisualDebugger/Public/ChaosVDRuntimeModule.h
+++ b/Engine/Source/Runtime/Experimental/ChaosVisualDebugger/Public/ChaosVDRuntimeModule.h
@@ -9,6 +9,7 @@
 #include "Misc/ScopeRWLock.h"
 
 struct FChaosVDRecording;
+class FText;
 
 /* Option flags that controls what should be recorded when doing a full capture **/
 enum class EChaosVDFullCaptureFlags : int32
@@ -20,6 +21,7 @@ ENUM_CLASS_FLAGS(EChaosVDFullCaptureFlags)
 
 DECLARE_MULTICAST_DELEGATE(FChaosVDRecordingStateChangedDelegate)
 DECLARE_MULTICAST_DELEGATE_OneParam(FChaosVDCaptureRequestDelegate, EChaosVDFullCaptureFlags)
+DECLARE_MULTICAST_DELEGATE_OneParam(FChaosVDRecordingStartFailedDelegate, const FText&)
 
 class CHAOSVDRUNTIME_API FChaosVDRuntimeModule : public IModuleInterface
 {
@@ -56,6 +58,12 @@ public:
 		FWriteScopeLock WriteLock(DelegatesRWLock);
 		return RecordingStopDelegate.Add(InCallback);
 	}
+
+	static FDelegateHandle RegisterRecordingStartFailedCallback(const FChaosVDRecordingStartFailedDelegate::FDelegate& InCallback)
+	{
+		FWriteScopeLock WriteLock(DelegatesRWLock);
+		return RecordingStartFailedDelegate.Add(InCallback);
+	}
 	
 	static FDelegateHandle RegisterFullCaptureRequestedCallback(const FChaosVDCaptureRequestDelegate::FDelegate& InCallback)
 	{
@@ -74,6 +82,12 @@ public:
 		FWriteScopeLock WriteLock(DelegatesRWLock);
 		return RecordingStopDelegate.Remove(InDelegateToRemove);
 	}
+
+	static bool RemoveRecordingStartFailedCallback(const FDelegateHandle& InDelegateToRemove)
+	{
+		FWriteScopeLock WriteLock(DelegatesRWLock);
+		return RecordingStartFailedDelegate.Remove(InDelegateToRemove);
+	}
 	
 	static bool RemoveFullCaptureRequestedCallback(const FDelegateHandle& InDelegateToRemove)
 	{
@@ -115,6 +129,7 @@ private:
 
 	static FChaosVDRecordingStateChangedDelegate RecordingStartedDelegate;
 	static FChaosVDRecordingStateChangedDelegate RecordingStopDelegate;
+	static FChaosVDRecordingStartFailedDelegate RecordingStartFailedDelegate;
 	static FChaosVDCaptureRequestDelegate PerformFullCaptureDelegate;
 
 	FThreadSafeCounter LastGeneratedID;
