Boolean on the geometry collection to enable the convex optimization

diff --git a/Engine/Source/Runtime/Experimental/Chaos/Private/Chaos/ConvexOptimizer.cpp b/Engine/Source/Runtime/Experimental/Chaos/Private/Chaos/ConvexOptimizer.cpp
--- a/Engine/Source/Runtime/Experimental/Chaos/Private/Chaos/ConvexOptimizer.cpp
+++ b/Engine/Source/Runtime/Experimental/Chaos/Private/Chaos/ConvexOptimizer.cpp
@@ -12,7 +12,7 @@ namespace Chaos
 namespace CVars
 {
 	// Replace all the convexes within an implicit hierarchy with a simplified one (kdop18 tribox for now) for collision
-	bool bChaosConvexSimplifyUnion = true;
+	CHAOS_API bool bChaosConvexSimplifyUnion = false;
 	FAutoConsoleVariableRef CVarChaosSimplifyUnion(TEXT("p.Chaos.Convex.SimplifyUnion"), bChaosConvexSimplifyUnion, TEXT("If true replace all the convexes within an implcit hierarchy with a simplified one (kdop18 tribox for now) for collision"));
 
 	// Max number of convex LODs used during simplification  for dynamic particles
@@ -117,11 +117,11 @@ void FConvexOptimizer::BuildConvexShapes(const FShapesArray& UnionShapes)
 }
 
 void FConvexOptimizer::SimplifyRootConvexes(const Chaos::FImplicitObjectUnionPtr& UnionGeometry, const FShapesArray& UnionShapes,
-		const EObjectStateType ObjectState)
+		const EObjectStateType ObjectState, const TBitArray<>& bOptimizeConvexes)
 {
 	SCOPE_CYCLE_COUNTER(STAT_Collisions_SimplifyConvexes);
 
-	if(CVars::bChaosConvexSimplifyUnion && UnionGeometry && (UnionShapes.Num() > 0) &&
+	if(CVars::bChaosConvexSimplifyUnion && UnionGeometry && (UnionShapes.Num() > 0) && (UnionShapes.Num() == bOptimizeConvexes.Num()) &&
 		(UnionShapes.Num() <= UnionGeometry->GetObjects().Num()))
 	{
 		SimplifiedConvexes.Reset();
@@ -131,13 +131,13 @@ void FConvexOptimizer::SimplifyRootConvexes(const Chaos::FImplicitObjectUnionPtr
 		if((UnionShapes.Num() == 1) || ((ObjectState == EObjectStateType::Dynamic) && (CVars::ChaosConvexDynamicMode == 0)) 
 				|| ((ObjectState != EObjectStateType::Dynamic) && (CVars::ChaosConvexKinematicMode == 0)))
 		{
-			BuildSingleConvex(UnionGeometry, UnionShapes);
+			BuildSingleConvex(UnionGeometry, UnionShapes, bOptimizeConvexes);
 		}
 		else
 		{
 			const bool bEnableMerging = ((ObjectState == EObjectStateType::Dynamic) && (CVars::ChaosConvexDynamicMode == 2))
 				|| ((ObjectState != EObjectStateType::Dynamic) && (CVars::ChaosConvexKinematicMode == 2));
-			BuildMultipleConvex(UnionGeometry, UnionShapes, bEnableMerging);
+			BuildMultipleConvex(UnionGeometry, UnionShapes, bEnableMerging, bOptimizeConvexes);
 		}
 	}
 
@@ -199,7 +199,7 @@ FORCEINLINE void ResizeCachedTriboxes(FConvexOptimizer::FTriboxNodes& RootTribox
 
 DECLARE_CYCLE_STAT(TEXT("Collisions::BuildConvexTriboxes"), STAT_BuildConvexTriboxes, STATGROUP_ChaosCollision);
 FORCEINLINE void BuildConvexTriboxes(const Chaos::FImplicitObjectUnionPtr& UnionGeometry, const FShapesArray& UnionShapes,
-	TUniquePtr<Private::FCollisionObjects>& CollisionObjects, FConvexOptimizer::FTriboxNodes& RootTriboxes, int32& NextConvexId)
+	TUniquePtr<Private::FCollisionObjects>& CollisionObjects, FConvexOptimizer::FTriboxNodes& RootTriboxes, int32& NextConvexId, const TBitArray<>& bOptimizeConvexes)
 {
 	SCOPE_CYCLE_COUNTER(STAT_BuildConvexTriboxes);
 
@@ -220,54 +220,64 @@ FORCEINLINE void BuildConvexTriboxes(const Chaos::FImplicitObjectUnionPtr& Union
 				FTribox LocalTribox;
 				int32 LeafObjectIndex = 0, ObjectIndex = 0;
 				RootObject->VisitLeafObjectsImpl(FRigidTransform3::Identity,RootObjectIndex, ObjectIndex, LeafObjectIndex,
-					[&CollisionObjects, &LocalTribox, &bHasRootTribox, &ShapeIndex, &NextConvexId](
+					[&CollisionObjects, &LocalTribox, &bHasRootTribox, &ShapeIndex, &NextConvexId, &bOptimizeConvexes](
 					const FImplicitObject* ImplicitObject, const FRigidTransform3& RelativeTransform,
 					const int32 RootIndex, const int32, const int32)
 				{
-					if(const FConvex* Convex = ImplicitObject->AsA<FConvex>())
+					if (!bOptimizeConvexes[ShapeIndex])
 					{
-						if(!bHasRootTribox)
+						// Add non convex implicits to the list of collision objects
+						Private::FImplicitBVH::CollectLeafObject(ImplicitObject,
+							RelativeTransform, ShapeIndex, CollisionObjects->ImplicitObjects, NextConvexId++);
+					}
+					else
+					{ 
+						if(const FConvex* Convex = ImplicitObject->AsA<FConvex>())
 						{
-							const FTribox::FRigidTransform3Type ConvexTransform(RelativeTransform);
-							// Disable collision for all the convexes that are going to be used to build the tribox
-							// For now only used for debug draw
-							LocalTribox.AddConvex(Convex, ConvexTransform);
+							if(!bHasRootTribox)
+							{
+								const FTribox::FRigidTransform3Type ConvexTransform(RelativeTransform);
+								// Disable collision for all the convexes that are going to be used to build the tribox
+								// For now only used for debug draw
+								LocalTribox.AddConvex(Convex, ConvexTransform);
+							}
 						}
-					}
-					else if(const TImplicitObjectScaled<FConvex>* ScaledObject = TImplicitObjectScaled<FConvex>::AsScaled(*ImplicitObject))
-					{
-						if(!bHasRootTribox)
+						else if(const TImplicitObjectScaled<FConvex>* ScaledObject = TImplicitObjectScaled<FConvex>::AsScaled(*ImplicitObject))
 						{
-							const FTribox::FRigidTransform3Type ConvexTransform(RelativeTransform.GetTranslation(),
-								RelativeTransform.GetRotation(), RelativeTransform.GetScale3D() * ScaledObject->GetScale());
-							// Disable collision for all the convexes that are going to be used to build the tribox
-							// For now only used for debug draw
-							LocalTribox.AddConvex(ScaledObject->GetUnscaledObject(), ConvexTransform);
+							if(!bHasRootTribox)
+							{
+								const FTribox::FRigidTransform3Type ConvexTransform(RelativeTransform.GetTranslation(),
+									RelativeTransform.GetRotation(), RelativeTransform.GetScale3D() * ScaledObject->GetScale());
+								// Disable collision for all the convexes that are going to be used to build the tribox
+								// For now only used for debug draw
+								LocalTribox.AddConvex(ScaledObject->GetUnscaledObject(), ConvexTransform);
+							}
 						}
-					}
-					else if(const TImplicitObjectInstanced<FConvex>* InstancedObject = TImplicitObjectInstanced<FConvex>::AsInstanced(*ImplicitObject))
-					{
-						if(!bHasRootTribox)
+						else if(const TImplicitObjectInstanced<FConvex>* InstancedObject = TImplicitObjectInstanced<FConvex>::AsInstanced(*ImplicitObject))
 						{
-							const FTribox::FRigidTransform3Type ConvexTransform(RelativeTransform);
-							// Disable collision for all the convexes that are going to be used to build the tribox
-							// For now only used for debug draw
-							LocalTribox.AddConvex(InstancedObject->Object(), ConvexTransform);
+							if(!bHasRootTribox)
+							{
+								const FTribox::FRigidTransform3Type ConvexTransform(RelativeTransform);
+								// Disable collision for all the convexes that are going to be used to build the tribox
+								// For now only used for debug draw
+								LocalTribox.AddConvex(InstancedObject->Object(), ConvexTransform);
+							}
+						}
+						else
+						{
+							// Add non convex implicits to the list of collision objects
+							Private::FImplicitBVH::CollectLeafObject(ImplicitObject,
+								RelativeTransform, ShapeIndex, CollisionObjects->ImplicitObjects, NextConvexId++);
 						}
-					}
-					else
-					{
-						// Add non convex implicits to the list of collision objects
-						Private::FImplicitBVH::CollectLeafObject(ImplicitObject,
-							RelativeTransform, ShapeIndex, CollisionObjects->ImplicitObjects, NextConvexId++);
 					}
 				});
-				if(!bHasRootTribox && LocalTribox.HasDatas())
+
+				if (!bHasRootTribox && LocalTribox.HasDatas())
 				{
 					LocalTribox.BuildTribox();
-					RootTribox = &(RootTriboxes.Add(RootObject, {LocalTribox, nullptr, ShapeIndex, LocalTribox.ComputeVolume(),true}));
+					RootTribox = &(RootTriboxes.Add(RootObject, { LocalTribox, nullptr, ShapeIndex, LocalTribox.ComputeVolume(),true }));
 				}
-				if(RootTribox)
+				if (RootTribox)
 				{
 					RootTribox->NodeTribox.SetValid(true);
 					RootTribox->ShapeIndex = ShapeIndex;
@@ -278,9 +288,9 @@ FORCEINLINE void BuildConvexTriboxes(const Chaos::FImplicitObjectUnionPtr& Union
 	ResizeCachedTriboxes(RootTriboxes);
 }
 
-void FConvexOptimizer::BuildSingleConvex(const Chaos::FImplicitObjectUnionPtr& UnionGeometry, const FShapesArray& UnionShapes)
+void FConvexOptimizer::BuildSingleConvex(const Chaos::FImplicitObjectUnionPtr& UnionGeometry, const FShapesArray& UnionShapes, const TBitArray<>& bOptimizeConvexes)
 {
-	BuildConvexTriboxes(UnionGeometry, UnionShapes, CollisionObjects, RootTriboxes, NextConvexId);
+	BuildConvexTriboxes(UnionGeometry, UnionShapes, CollisionObjects, RootTriboxes, NextConvexId, bOptimizeConvexes);
 	
 	MainTribox.NodeTribox = FTribox(); 
 	for(auto& RootTribox : RootTriboxes)
@@ -441,9 +451,9 @@ void FConvexOptimizer::MergeConnectedShapes(const Chaos::FImplicitObjectUnionPtr
 }
 
 	
-void FConvexOptimizer::BuildMultipleConvex(const Chaos::FImplicitObjectUnionPtr& UnionGeometry, const FShapesArray& UnionShapes, const bool bEnableMerging)
+void FConvexOptimizer::BuildMultipleConvex(const Chaos::FImplicitObjectUnionPtr& UnionGeometry, const FShapesArray& UnionShapes, const bool bEnableMerging, const TBitArray<>& bOptimizeConvexes)
 {
-	BuildConvexTriboxes(UnionGeometry, UnionShapes, CollisionObjects, RootTriboxes, NextConvexId);
+	BuildConvexTriboxes(UnionGeometry, UnionShapes, CollisionObjects, RootTriboxes, NextConvexId, bOptimizeConvexes);
 
 	if(!bEnableMerging)
 	{
diff --git a/Engine/Source/Runtime/Experimental/Chaos/Private/Chaos/PBDRigidClustering.cpp b/Engine/Source/Runtime/Experimental/Chaos/Private/Chaos/PBDRigidClustering.cpp
--- a/Engine/Source/Runtime/Experimental/Chaos/Private/Chaos/PBDRigidClustering.cpp
+++ b/Engine/Source/Runtime/Experimental/Chaos/Private/Chaos/PBDRigidClustering.cpp
@@ -74,7 +74,7 @@ namespace Chaos
 
 	namespace CVars
 	{
-		extern bool bChaosConvexSimplifyUnion;
+		extern CHAOS_API bool bChaosConvexSimplifyUnion;
 	}
 
 	template <typename TProxy=FGeometryCollectionPhysicsProxy>
@@ -398,6 +398,33 @@ namespace Chaos
 		{ 
 			if(FImplicitObjectUnion* Union = Particle->GetGeometry()->template AsA<FImplicitObjectUnion>())
 			{
+				TBitArray<> bOptimizeConvexes;
+				bOptimizeConvexes.Init(true, Particle->ShapesArray().Num());
+				if (FGeometryCollectionPhysicsProxy* ConcreteProxy = GetConcreteProxy(Particle))
+				{
+					if(Particle->ShapesArray().Num() == 1)
+					{ 
+						bOptimizeConvexes[0] = ConcreteProxy->GetSimParameters().bOptimizeConvexes;
+					}
+				}
+				else if (Chaos::FClusterUnion* ClusterUnion = ClusterUnionManager.FindClusterUnionFromParticle(Particle))
+				{
+					if(Particle->ShapesArray().Num() == ClusterUnion->ChildParticles.Num())
+					{
+						int32 ShapeIndex = 0;
+						for (FPBDRigidParticleHandle* ChildHandle : ClusterUnion->ChildParticles)
+						{
+							if (FPBDRigidClusteredParticleHandle* ChildClustered = ChildHandle->CastToClustered())
+							{
+								if (FGeometryCollectionPhysicsProxy* ChildProxy = GetConcreteProxy(ChildClustered))
+								{
+									bOptimizeConvexes[ShapeIndex] = ChildProxy->GetSimParameters().bOptimizeConvexes;
+								}
+							}
+							++ShapeIndex;
+						}
+					}
+				}
 				if(Union->GetNumLeafObjects() > 1)
 				{ 
 					if (!Particle->ConvexOptimizer())
@@ -405,7 +432,7 @@ namespace Chaos
 						Particle->ConvexOptimizer() = MakePimpl<Private::FConvexOptimizer>();
 					}
 					Particle->ConvexOptimizer()->SimplifyRootConvexes(Union,
-						Particle->ShapesArray(), Particle->ObjectState());
+						Particle->ShapesArray(), Particle->ObjectState(), bOptimizeConvexes);
 					bHasOptimizer = Particle->ConvexOptimizer()->IsValid();
 				}
 			}
diff --git a/Engine/Source/Runtime/Experimental/Chaos/Public/Chaos/ConvexOptimizer.h b/Engine/Source/Runtime/Experimental/Chaos/Public/Chaos/ConvexOptimizer.h
--- a/Engine/Source/Runtime/Experimental/Chaos/Public/Chaos/ConvexOptimizer.h
+++ b/Engine/Source/Runtime/Experimental/Chaos/Public/Chaos/ConvexOptimizer.h
@@ -60,7 +60,7 @@ namespace Private
 
 		// Simplify all the convexes in the hierarchy 
 		CHAOS_API void SimplifyRootConvexes(const Chaos::FImplicitObjectUnionPtr& UnionGeometry, const FShapesArray& UnionShapes,
-			const EObjectStateType ObjectState);
+			const EObjectStateType ObjectState, const TBitArray<>& bOptimizeConvexes);
 
 		// Check if the manager is valid or not
 		CHAOS_API bool IsValid() const {return !SimplifiedConvexes.IsEmpty();}
@@ -86,10 +86,10 @@ namespace Private
 		void MergeConnectedShapes(const Chaos::FImplicitObjectUnionPtr& UnionGeometry, TArray<FTriboxNode>& MergedNodes);
 
 		// Build a single convex
-		void BuildSingleConvex(const Chaos::FImplicitObjectUnionPtr& UnionGeometry, const FShapesArray& UnionShapes);
+		void BuildSingleConvex(const Chaos::FImplicitObjectUnionPtr& UnionGeometry, const FShapesArray& UnionShapes, const TBitArray<>& bOptimizeConvexes);
 
 		// Build several convexes 
-		void BuildMultipleConvex(const Chaos::FImplicitObjectUnionPtr& UnionGeometry, const FShapesArray& UnionShapes, const bool bEnableMerging);
+		void BuildMultipleConvex(const Chaos::FImplicitObjectUnionPtr& UnionGeometry, const FShapesArray& UnionShapes, const bool bEnableMerging, const TBitArray<>& bOptimizeConvexes);
 
 		// Build the simplified shapes
 		void BuildConvexShapes(const FShapesArray& UnionShapes);
diff --git a/Engine/Source/Runtime/Experimental/Chaos/Public/GeometryCollection/GeometryCollectionSimulationCoreTypes.h b/Engine/Source/Runtime/Experimental/Chaos/Public/GeometryCollection/GeometryCollectionSimulationCoreTypes.h
--- a/Engine/Source/Runtime/Experimental/Chaos/Public/GeometryCollection/GeometryCollectionSimulationCoreTypes.h
+++ b/Engine/Source/Runtime/Experimental/Chaos/Public/GeometryCollection/GeometryCollectionSimulationCoreTypes.h
@@ -227,6 +227,7 @@ struct FSimulationParameters
 		, UserData(nullptr)
 		, bEnableStrainOnCollision(true)
 		, bUseStaticMeshCollisionForTraces(false)
+		, bOptimizeConvexes(true)
 	{}
 
 	FSimulationParameters(const FSimulationParameters& Other)
@@ -291,6 +292,7 @@ struct FSimulationParameters
 		, UserData(Other.UserData)
 		, bEnableStrainOnCollision(Other.bEnableStrainOnCollision)
 		, bUseStaticMeshCollisionForTraces(Other.bUseStaticMeshCollisionForTraces)
+		, bOptimizeConvexes(Other.bOptimizeConvexes)
 	{
 	}
 
@@ -388,4 +390,6 @@ struct FSimulationParameters
 	bool bEnableStrainOnCollision;
 
 	bool bUseStaticMeshCollisionForTraces;
+
+	bool bOptimizeConvexes = true;
 };
diff --git a/Engine/Source/Runtime/Experimental/GeometryCollectionEngine/Private/GeometryCollection/GeometryCollectionComponent.cpp b/Engine/Source/Runtime/Experimental/GeometryCollectionEngine/Private/GeometryCollection/GeometryCollectionComponent.cpp
--- a/Engine/Source/Runtime/Experimental/GeometryCollectionEngine/Private/GeometryCollection/GeometryCollectionComponent.cpp
+++ b/Engine/Source/Runtime/Experimental/GeometryCollectionEngine/Private/GeometryCollection/GeometryCollectionComponent.cpp
@@ -3965,6 +3965,7 @@ void UGeometryCollectionComponent::RegisterAndInitializePhysicsProxy()
 		SimulationParameters.UserData = static_cast<void*>(&PhysicsUserData);
 		SimulationParameters.bEnableStrainOnCollision = bEnableDamageFromCollision;
 		SimulationParameters.bUseStaticMeshCollisionForTraces = bUseStaticMeshCollisionForTraces;
+		SimulationParameters.bOptimizeConvexes = RestCollection ? RestCollection->bOptimizeConvexes : true;
 
 		UPhysicalMaterial* EnginePhysicalMaterial = GetPhysicalMaterial();
 		if (ensure(EnginePhysicalMaterial))
diff --git a/Engine/Source/Runtime/Experimental/GeometryCollectionEngine/Private/GeometryCollection/GeometryCollectionObject.cpp b/Engine/Source/Runtime/Experimental/GeometryCollectionEngine/Private/GeometryCollection/GeometryCollectionObject.cpp
--- a/Engine/Source/Runtime/Experimental/GeometryCollectionEngine/Private/GeometryCollection/GeometryCollectionObject.cpp
+++ b/Engine/Source/Runtime/Experimental/GeometryCollectionEngine/Private/GeometryCollection/GeometryCollectionObject.cpp
@@ -68,6 +68,14 @@ FAutoConsoleVariableRef CVarGeometryCollectionAlwaysRecreateSimulationData(
 	bGeometryCollectionAlwaysRecreateSimulationData,
 	TEXT("always recreate the simulation data even if the simulation data is not marked as dirty - this has runtime cost in editor - only use as a last resort if default has issues [def:false]"));
 
+namespace Chaos
+{
+	namespace CVars
+	{
+		extern CHAOS_API bool bChaosConvexSimplifyUnion;
+	}
+}
+
 
 #if ENABLE_COOK_STATS
 namespace GeometryCollectionCookStats
@@ -116,6 +124,7 @@ UGeometryCollection::UGeometryCollection(const FObjectInitializer& ObjectInitial
 	, Mass(2500.0f)
 	, MinimumMassClamp(0.1f)
 	, bImportCollisionFromSource(false)
+	, bOptimizeConvexes(Chaos::CVars::bChaosConvexSimplifyUnion)
 	, bScaleOnRemoval(true)
 	, bRemoveOnMaxSleep(false)
 	, MaximumSleepTime(5.0, 10.0)
@@ -1986,4 +1995,23 @@ void UGeometryCollection::RemoveUserDataOfClass(TSubclassOf<UAssetUserData> InUs
 const TArray<UAssetUserData*>* UGeometryCollection::GetAssetUserDataArray() const
 {
 	return &ToRawPtrTArrayUnsafe(AssetUserData);
-}
\ No newline at end of file
+}
+
+#if WITH_EDITOR
+bool UGeometryCollection::CanEditChange(const FProperty* InProperty) const
+{
+	if (!Super::CanEditChange(InProperty))
+	{
+		return false;
+	}
+
+	const FName& Name = InProperty->GetFName();
+
+	if (Name == GET_MEMBER_NAME_CHECKED(ThisClass, bOptimizeConvexes))
+	{
+		return Chaos::CVars::bChaosConvexSimplifyUnion == true;
+	}
+
+	return true;
+}
+#endif
\ No newline at end of file
diff --git a/Engine/Source/Runtime/Experimental/GeometryCollectionEngine/Public/GeometryCollection/GeometryCollectionObject.h b/Engine/Source/Runtime/Experimental/GeometryCollectionEngine/Public/GeometryCollection/GeometryCollectionObject.h
--- a/Engine/Source/Runtime/Experimental/GeometryCollectionEngine/Public/GeometryCollection/GeometryCollectionObject.h
+++ b/Engine/Source/Runtime/Experimental/GeometryCollectionEngine/Public/GeometryCollection/GeometryCollectionObject.h
@@ -703,6 +703,14 @@ public:
 	*/
 	UPROPERTY(EditAnywhere, BlueprintReadWrite, Category = "Collisions")
 	bool bImportCollisionFromSource;
+
+	/**
+	* whether to optimize convexes for collisions. If true the convex optimizer will generate at runtime one 
+	* single convex shape for physics collisions ignoring all the user defined ones. 
+	* Enable p.Chaos.Convex.SimplifyUnion cvar to be able to use it (experimental)
+	*/
+	UPROPERTY(EditAnywhere, BlueprintReadWrite, Category = "Collisions")
+	bool bOptimizeConvexes = true;
 	
 #if WITH_EDITORONLY_DATA
 	/**
@@ -822,6 +830,10 @@ public:
 
 	GEOMETRYCOLLECTIONENGINE_API const TArray<int32>& GetAutoInstanceTransformRemapIndices() const { return AutoInstanceTransformRemapIndices; }
 
+#if WITH_EDITOR
+	virtual bool CanEditChange(const FProperty* InProperty) const override;
+#endif // WITH_EDITOR
+
 private:
 #if WITH_EDITOR
 	GEOMETRYCOLLECTIONENGINE_API void CreateSimulationDataImp(bool bCopyFromDDC);
