Cache physics particle in the modular vehicle modules.

diff --git a/Engine/Source/Runtime/Experimental/ChaosVehicles/ChaosVehiclesCore/Private/SimModule/SimulationModuleBase.cpp b/Engine/Source/Runtime/Experimental/ChaosVehicles/ChaosVehiclesCore/Private/SimModule/SimulationModuleBase.cpp
--- a/Engine/Source/Runtime/Experimental/ChaosVehicles/ChaosVehiclesCore/Private/SimModule/SimulationModuleBase.cpp
+++ b/Engine/Source/Runtime/Experimental/ChaosVehicles/ChaosVehiclesCore/Private/SimModule/SimulationModuleBase.cpp
@@ -56,8 +56,12 @@ ISimulationModuleBase* ISimulationModuleBase::GetFirstChild()
 
 FPBDRigidClusteredParticleHandle* ISimulationModuleBase::GetClusterParticle(Chaos::FClusterUnionPhysicsProxy* Proxy)
 { 
-	// TODO: should store what we need rather than search for it all the time
-	FPBDRigidClusteredParticleHandle* ClusterChild = nullptr;
+	if (ParticleIdx.IsValid() && CachedParticle && (CachedParticle->UniqueIdx() == ParticleIdx))
+	{
+		return CachedParticle;
+	}
+
+	CachedParticle = nullptr;
 
 	FPBDRigidsEvolutionGBF& Evolution = *static_cast<FPBDRigidsSolver*>(Proxy->GetSolver<FPBDRigidsSolver>())->GetEvolution();
 	FClusterUnionManager& ClusterUnionManager = Evolution.GetRigidClustering().GetClusterUnionManager();
@@ -70,11 +74,11 @@ FPBDRigidClusteredParticleHandle* ISimulationModuleBase::GetClusterParticle(Chao
 
 		if (FPBDRigidParticleHandle* Particle = GetParticleFromUniqueIndex(ParticleIdx.Idx, Particles))
 		{
-			ClusterChild = Particle->CastToClustered();
+			CachedParticle = Particle->CastToClustered();
 		}
 	}
 
-	return ClusterChild;
+	return CachedParticle;
 }
 
 FPBDRigidParticleHandle* ISimulationModuleBase::GetParticleFromUniqueIndex(int32 ParticleUniqueIdx, TArray<FPBDRigidParticleHandle*>& Particles)
diff --git a/Engine/Source/Runtime/Experimental/ChaosVehicles/ChaosVehiclesCore/Public/SimModule/SimulationModuleBase.h b/Engine/Source/Runtime/Experimental/ChaosVehicles/ChaosVehiclesCore/Public/SimModule/SimulationModuleBase.h
--- a/Engine/Source/Runtime/Experimental/ChaosVehicles/ChaosVehiclesCore/Public/SimModule/SimulationModuleBase.h
+++ b/Engine/Source/Runtime/Experimental/ChaosVehicles/ChaosVehiclesCore/Public/SimModule/SimulationModuleBase.h
@@ -187,6 +187,7 @@ namespace Chaos
 			, bAnimationEnabled(true)
 			, AppliedForce(FVector::ZeroVector)
 			, Guid(INDEX_NONE)
+			, CachedParticle(nullptr)
 		{}
 		virtual ~ISimulationModuleBase() {}
 
@@ -371,7 +372,7 @@ namespace Chaos
 		FVector AppliedForce;
 		int Guid; // needed a way of associating internal module with game thread.
 
-		//FPBDRigidClusteredParticleHandle* ClusterParticle;
+		FPBDRigidClusteredParticleHandle* CachedParticle;
 	};
 
 	/**
