ChaosFlesh: Active muscle contraction

diff --git a/Engine/Plugins/Experimental/ChaosFlesh/Source/ChaosFleshEngine/Private/ChaosFlesh/ChaosDeformableSolverComponent.cpp b/Engine/Plugins/Experimental/ChaosFlesh/Source/ChaosFleshEngine/Private/ChaosFlesh/ChaosDeformableSolverComponent.cpp
--- a/Engine/Plugins/Experimental/ChaosFlesh/Source/ChaosFleshEngine/Private/ChaosFlesh/ChaosDeformableSolverComponent.cpp
+++ b/Engine/Plugins/Experimental/ChaosFlesh/Source/ChaosFleshEngine/Private/ChaosFlesh/ChaosDeformableSolverComponent.cpp
@@ -202,7 +202,6 @@ void UDeformableSolverComponent::Reset()
 			, SolverDebugging.CacheToFile
 			, SolverConstraints.bEnableKinematics
 			, SolverCollisions.bUseFloor
-			, SolverCollisions.bDoSelfCollision
 			, false /*SolverCollisions.SolverGridBasedCollisions.bUseGridBasedConstraints*/
 			, 25. /*SolverCollisions.SolverGridBasedCollisions.GridDx*/
 			, SolverEvolution.SolverQuasistatics.bDoQuasistatics
@@ -217,9 +216,13 @@ void UDeformableSolverComponent::Reset()
 			, SolverConstraints.GaussSeidelConstraints.bUseSOR
 			, SolverConstraints.GaussSeidelConstraints.OmegaSOR
 			, SolverConstraints.GaussSeidelConstraints.bUseGSNeohookean
+			, SolverConstraints.GaussSeidelConstraints.CollisionSpring.bDoSelfCollision
+			, SolverConstraints.GaussSeidelConstraints.CollisionSpring.bDoInComponentSelfCollision
+			, SolverConstraints.GaussSeidelConstraints.CollisionSpring.NRingExcluded
 			, SolverConstraints.GaussSeidelConstraints.CollisionSpring.CollisionSearchRadius
 			, SolverConstraints.GaussSeidelConstraints.CollisionSpring.CollisionSpringStiffness
 			, SolverConstraints.GaussSeidelConstraints.CollisionSpring.bAllowSliding
+			, SolverMuscleActivation.bDoMuscleActivation
 		}));
 
 		for (TObjectPtr<UDeformablePhysicsComponent>& DeformableComponent : ConnectedObjects.DeformableComponents)
diff --git a/Engine/Plugins/Experimental/ChaosFlesh/Source/ChaosFleshEngine/Public/ChaosFlesh/ChaosDeformableSolverComponent.h b/Engine/Plugins/Experimental/ChaosFlesh/Source/ChaosFleshEngine/Public/ChaosFlesh/ChaosDeformableSolverComponent.h
--- a/Engine/Plugins/Experimental/ChaosFlesh/Source/ChaosFleshEngine/Public/ChaosFlesh/ChaosDeformableSolverComponent.h
+++ b/Engine/Plugins/Experimental/ChaosFlesh/Source/ChaosFleshEngine/Public/ChaosFlesh/ChaosDeformableSolverComponent.h
@@ -99,10 +99,25 @@ struct FCollisionSpringGroup
 {
 	GENERATED_USTRUCT_BODY()
 	/**
+	* If uses spring self-collision
+	*/
+	UPROPERTY(EditAnywhere, Category = "CollisionSpring")
+	bool bDoSelfCollision = false;
+	/**
+	* If uses in-component self-collision
+	*/
+	UPROPERTY(EditAnywhere, Category = "CollisionSpring")
+	bool bDoInComponentSelfCollision = false;
+	/**
+	* N ring to exclude for in-component self-collision 
+	*/
+	UPROPERTY(EditAnywhere, Category = "CollisionSpring")
+	int32 NRingExcluded = 1;
+	/**
 	* Search radius for point triangle pairs
 	*/
 	UPROPERTY(EditAnywhere, Category = "CollisionSpring")
-	float CollisionSearchRadius = 10.f;
+	float CollisionSearchRadius = 0.f;
 	/**
 	* Collision spring stiffness; larger value will stop penetration better
 	*/
@@ -156,9 +171,6 @@ struct FSolverCollisionsGroup
 {
 	GENERATED_USTRUCT_BODY()
 
-	UPROPERTY(EditAnywhere, Category = "Collisions")
-	bool bDoSelfCollision = false;
-
 	UPROPERTY(EditAnywhere, Category = "Physics")
 	bool bUseFloor = true;
 
@@ -202,7 +214,6 @@ struct FSolverConstraintsGroup
 	FSolverGaussSeidelConstraintsGroup GaussSeidelConstraints;
 };
 
-
 USTRUCT(BlueprintType)
 struct FSolverForcesGroup
 {
@@ -218,6 +229,15 @@ struct FSolverForcesGroup
 	bool bEnableGravity = true;
 };
 
+USTRUCT(BlueprintType)
+struct FSolverMuscleActivationGroup
+{
+	GENERATED_USTRUCT_BODY()
+
+	UPROPERTY(EditAnywhere, Category = "MuscleActivation")
+	bool bDoMuscleActivation = false;
+};
+
 /**
 *	UDeformableSolverComponent
 */
@@ -278,6 +298,9 @@ public:
 	UPROPERTY(EditAnywhere, BlueprintReadOnly, Category = "Physics")
 	FSolverDebuggingGroup SolverDebugging;
 
+	UPROPERTY(EditAnywhere, BlueprintReadOnly, Category = "Physics")
+	FSolverMuscleActivationGroup SolverMuscleActivation;
+
 	// Simulation Variables
 	TUniquePtr<FDeformableSolver> Solver;
 
diff --git a/Engine/Plugins/Experimental/ChaosFlesh/Source/ChaosFleshNodes/Private/Dataflow/ChaosFleshEngineAssetNodes.cpp b/Engine/Plugins/Experimental/ChaosFlesh/Source/ChaosFleshNodes/Private/Dataflow/ChaosFleshEngineAssetNodes.cpp
--- a/Engine/Plugins/Experimental/ChaosFlesh/Source/ChaosFleshNodes/Private/Dataflow/ChaosFleshEngineAssetNodes.cpp
+++ b/Engine/Plugins/Experimental/ChaosFlesh/Source/ChaosFleshNodes/Private/Dataflow/ChaosFleshEngineAssetNodes.cpp
@@ -4,6 +4,7 @@
 #include "Dataflow/ChaosFleshTetrahedralNodes.h"
 
 #include "Chaos/Math/Poisson.h"
+#include "Chaos/Matrix.h"
 #include "Chaos/Utilities.h"
 #include "ChaosFlesh/ChaosFlesh.h"
 #include "ChaosFlesh/FleshAsset.h"
@@ -12,8 +13,8 @@
 #include "Dataflow/DataflowInputOutput.h"
 #include "Dataflow/DataflowObjectInterface.h"
 #include "Dataflow/DataflowNodeFactory.h"
+#include "GeometryCollection/Facades/CollectionMuscleActivationFacade.h"
 #include "GeometryCollection/TransformCollection.h"
-
 #include UE_INLINE_GENERATED_CPP_BY_NAME(ChaosFleshEngineAssetNodes)
 
 namespace Dataflow
@@ -24,6 +25,7 @@ namespace Dataflow
 		DATAFLOW_NODE_REGISTER_CREATION_FACTORY(FFleshAssetTerminalDataflowNode);
 		DATAFLOW_NODE_REGISTER_CREATION_FACTORY(FSetFleshDefaultPropertiesNode);
 		DATAFLOW_NODE_REGISTER_CREATION_FACTORY(FComputeFiberFieldNode);
+		DATAFLOW_NODE_REGISTER_CREATION_FACTORY(FComputeMuscleActivationDataNode);
 		DATAFLOW_NODE_REGISTER_CREATION_FACTORY(FVisualizeFiberFieldNode);
 		DATAFLOW_NODE_REGISTER_CREATION_FACTORY(FComputeIslandsNode);
 		DATAFLOW_NODE_REGISTER_CREATION_FACTORY(FGenerateOriginInsertionNode);
@@ -425,6 +427,107 @@ FComputeFiberFieldNode::ComputeFiberField(
 	return Directions;
 }
 
+void FComputeMuscleActivationDataNode::Evaluate(Dataflow::FContext& Context, const FDataflowOutput* Out) const
+{
+	if (Out->IsA<FManagedArrayCollection>(&Collection))
+	{
+		FManagedArrayCollection InCollection = GetValue<FManagedArrayCollection>(Context, &Collection);
+		TArray<int32> InOriginIndices = GetValue<TArray<int32>>(Context, &OriginIndicesIn);
+		TArray<int32> InInsertionIndices = GetValue<TArray<int32>>(Context, &InsertionIndicesIn);
+		TArray<float> OriginInsertionRestLength;
+		if (InOriginIndices.Num() > 0 && InInsertionIndices.Num() > 0)
+		{
+			if (TManagedArray<FVector3f>* Vertex = InCollection.FindAttribute<FVector3f>("Vertex", "Vertices"))
+			{
+				if (TManagedArray<FIntVector4>* Elements = InCollection.FindAttribute<FIntVector4>(FTetrahedralCollection::TetrahedronAttribute, FTetrahedralCollection::TetrahedralGroup))
+				{
+					if (TManagedArray<FVector3f>* FiberDirections = InCollection.FindAttribute<FVector3f>("FiberDirection", FTetrahedralCollection::TetrahedralGroup))
+					{
+						GeometryCollection::Facades::FMuscleActivationFacade FMuscleActivation(InCollection);
+						TArray<TArray<int32>> MuscleActivationElements;
+						TArray<int32> OriginNodes;
+						TArray<int32> InsertionNodes;
+						if (TManagedArray<int32>* ComponentIndex = InCollection.FindAttribute<int32>("ComponentIndex", FGeometryCollection::VerticesGroup))
+						{
+							TMap<int32, int32> ComponentToIndex; //Component index to muscle index
+							for (int32 i = 0; i < InOriginIndices.Num(); i++)
+							{
+								if (!ComponentToIndex.Contains((*ComponentIndex)[InOriginIndices[i]]))
+								{
+									//TODO: use one pair of separately defined origin and insertion for each muscle instead of choosing the first from kinematic origins and last from insertions
+									ComponentToIndex.Add((*ComponentIndex)[InOriginIndices[i]], OriginNodes.Num());
+									OriginNodes.Add(InOriginIndices[i]);
+								}
+							}
+							InsertionNodes.Init(INDEX_NONE, OriginNodes.Num());
+							for (int32 i = 0; i < InInsertionIndices.Num(); i++)
+							{
+								if (!ComponentToIndex.Contains((*ComponentIndex)[InInsertionIndices[i]]))
+								{
+									ensureMsgf(false, TEXT("No origin in this component"));
+								}
+								else if (InsertionNodes[ComponentToIndex[(*ComponentIndex)[InInsertionIndices[i]]]] == INDEX_NONE)
+								{
+									InsertionNodes[ComponentToIndex[(*ComponentIndex)[InInsertionIndices[i]]]] = InInsertionIndices[i];
+								}
+							}
+							MuscleActivationElements.SetNum(OriginNodes.Num());
+							for (int32 ElemIdx = 0; ElemIdx < Elements->Num(); ElemIdx++)
+							{
+								if (ComponentToIndex.Contains((*ComponentIndex)[(*Elements)[ElemIdx][0]]))
+								{
+									MuscleActivationElements[ComponentToIndex[(*ComponentIndex)[(*Elements)[ElemIdx][0]]]].Add(ElemIdx);
+								}
+							}
+						}
+						else
+						{
+							MuscleActivationElements.SetNum(1);
+							for (int32 i = 0; i < Elements->Num(); i++)
+							{
+								MuscleActivationElements[0].Add(i);
+							}
+							OriginNodes.Add(InOriginIndices[0]);
+							InsertionNodes.Add(InInsertionIndices[0]);
+						}
+						OriginInsertionRestLength.SetNum(OriginNodes.Num());
+
+						for (int32 MuscleComponentIdx = 0; MuscleComponentIdx < OriginNodes.Num(); MuscleComponentIdx++)
+						{
+							if (ensureMsgf(InsertionNodes[MuscleComponentIdx] != INDEX_NONE, TEXT("InsertionNodes[%d] is not a pair"), MuscleComponentIdx))
+							{
+								GeometryCollection::Facades::FMuscleActivationData MuscleActivationData;
+								MuscleActivationData.MuscleActivationElement = MuscleActivationElements[MuscleComponentIdx];
+								MuscleActivationData.OriginInsertionPair = FIntVector2(OriginNodes[MuscleComponentIdx], InsertionNodes[MuscleComponentIdx]);
+								MuscleActivationData.OriginInsertionRestLength = ((*Vertex)[OriginNodes[MuscleComponentIdx]] - (*Vertex)[InsertionNodes[MuscleComponentIdx]]).Size();
+								MuscleActivationData.FiberDirectionMatrix.SetNum(MuscleActivationElements[MuscleComponentIdx].Num());
+								for (int32 LocalElemIdx = 0; LocalElemIdx < MuscleActivationElements[MuscleComponentIdx].Num(); LocalElemIdx++)
+								{
+									FVector3f V = (*FiberDirections)[MuscleActivationElements[MuscleComponentIdx][LocalElemIdx]];
+									// QR decomposition on vvT
+									FVector3f W = V;
+									if (V.X < V.Y)
+									{
+										W.X += 1.f;
+									}
+									else
+									{
+										W.Y += 1.f;
+									}
+									FVector3f U = (V ^ W).GetSafeNormal();
+									W = (U ^ V).GetSafeNormal();
+									MuscleActivationData.FiberDirectionMatrix[LocalElemIdx] = Chaos::PMatrix33d(V, W, U);
+								}
+								FMuscleActivation.AddMuscleActivationData(MuscleActivationData);
+							}
+						}
+					}
+				}
+			}
+		}
+		Out->SetValue(MoveTemp(InCollection), Context);
+	}
+}
 
 void FComputeIslandsNode::Evaluate(Dataflow::FContext& Context, const FDataflowOutput* Out) const
 {
diff --git a/Engine/Plugins/Experimental/ChaosFlesh/Source/ChaosFleshNodes/Private/Dataflow/ChaosFleshPositionTargetInitializationNodes.cpp b/Engine/Plugins/Experimental/ChaosFlesh/Source/ChaosFleshNodes/Private/Dataflow/ChaosFleshPositionTargetInitializationNodes.cpp
--- a/Engine/Plugins/Experimental/ChaosFlesh/Source/ChaosFleshNodes/Private/Dataflow/ChaosFleshPositionTargetInitializationNodes.cpp
+++ b/Engine/Plugins/Experimental/ChaosFlesh/Source/ChaosFleshNodes/Private/Dataflow/ChaosFleshPositionTargetInitializationNodes.cpp
@@ -612,7 +612,7 @@ void FSetFleshBonePositionTargetBindingDataflowNode::Evaluate(Dataflow::FContext
 							}
 						}
 						Chaos::TVec3<float> CoordDiff = (CoordMaxs - CoordMins) * VertexRadiusRatio;
-						SphereRadius = Chaos::FReal(FGenericPlatformMath::Min(CoordDiff[0], FGenericPlatformMath::Min(CoordDiff[1], CoordDiff[2])));
+						SphereRadius = Chaos::FReal(FGenericPlatformMath::Max(CoordDiff[0], FGenericPlatformMath::Max(CoordDiff[1], CoordDiff[2])));
 
 						TArray<Chaos::TSphere<Chaos::FReal, 3>*> VertexSpherePtrs;
 						TArray<Chaos::TSphere<Chaos::FReal, 3>> VertexSpheres;
diff --git a/Engine/Plugins/Experimental/ChaosFlesh/Source/ChaosFleshNodes/Public/Dataflow/ChaosFleshEngineAssetNodes.h b/Engine/Plugins/Experimental/ChaosFlesh/Source/ChaosFleshNodes/Public/Dataflow/ChaosFleshEngineAssetNodes.h
--- a/Engine/Plugins/Experimental/ChaosFlesh/Source/ChaosFleshNodes/Public/Dataflow/ChaosFleshEngineAssetNodes.h
+++ b/Engine/Plugins/Experimental/ChaosFlesh/Source/ChaosFleshNodes/Public/Dataflow/ChaosFleshEngineAssetNodes.h
@@ -196,6 +196,38 @@ public:
 	virtual void Evaluate(Dataflow::FContext& Context, const FDataflowOutput* Out) const override;
 };
 
+/**
+* Computes an orthogonal matrix for each element
+* M = [v,w,u], where v is the fiber direction of that element, w and u are chosen to be orthogonal to v and each other.
+*/
+USTRUCT(meta = (DataflowFlesh))
+struct FComputeMuscleActivationDataNode : public FDataflowNode
+{
+	GENERATED_USTRUCT_BODY()
+	DATAFLOW_NODE_DEFINE_INTERNAL(FComputeMuscleActivationDataNode, "ComputeMuscleActivationData", "Flesh", "")
+
+public:
+	UPROPERTY(meta = (DataflowInput, DataflowOutput, DisplayName = "Collection", DataflowPassthrough = "Collection"))
+	FManagedArrayCollection Collection;
+
+	UPROPERTY(meta = (DataflowInput, DisplayName = "OriginIndicesIn"))
+		TArray<int32> OriginIndicesIn;
+
+	UPROPERTY(meta = (DataflowInput, DisplayName = "InsertionIndicesIn"))
+		TArray<int32> InsertionIndicesIn;
+
+	FComputeMuscleActivationDataNode(const Dataflow::FNodeParameters& InParam, FGuid InGuid = FGuid::NewGuid())
+		: FDataflowNode(InParam, InGuid)
+	{
+		RegisterInputConnection(&Collection);
+		RegisterInputConnection(&OriginIndicesIn);
+		RegisterInputConnection(&InsertionIndicesIn);
+		RegisterOutputConnection(&Collection, &Collection);
+	}
+
+	virtual void Evaluate(Dataflow::FContext& Context, const FDataflowOutput* Out) const override;
+};
+
 USTRUCT(meta = (DataflowFlesh))
 struct FComputeIslandsNode : public FDataflowNode
 {
diff --git a/Engine/Source/Runtime/Experimental/Chaos/Private/Chaos/Deformable/ChaosDeformableSolver.cpp b/Engine/Source/Runtime/Experimental/Chaos/Private/Chaos/Deformable/ChaosDeformableSolver.cpp
--- a/Engine/Source/Runtime/Experimental/Chaos/Private/Chaos/Deformable/ChaosDeformableSolver.cpp
+++ b/Engine/Source/Runtime/Experimental/Chaos/Private/Chaos/Deformable/ChaosDeformableSolver.cpp
@@ -25,11 +25,13 @@
 #include "Chaos/Utilities.h"
 #include "Chaos/PBDEvolution.h"
 #include "Containers/StringConv.h"
+#include "Containers/Set.h"
 #include "CoreMinimal.h"
 #include "GeometryCollection/Facades/CollectionKinematicBindingFacade.h"
 #include "GeometryCollection/Facades/CollectionVertexBoneWeightsFacade.h"
 #include "GeometryCollection/Facades/CollectionPositionTargetFacade.h"
 #include "GeometryCollection/Facades/CollectionConstraintOverrideFacade.h"
+#include "GeometryCollection/Facades/CollectionMuscleActivationFacade.h"
 #include "Misc/FileHelper.h"
 #include "HAL/PlatformFileManager.h"
 #include "HAL/IConsoleManager.h"
@@ -61,6 +63,7 @@ DECLARE_CYCLE_STAT(TEXT("Chaos.Deformable.Solver.InitializeSelfCollisionVariable
 DECLARE_CYCLE_STAT(TEXT("Chaos.Deformable.Solver.InitializeGridBasedConstraintVariables"), STAT_ChaosDeformableSolver_InitializeGridBasedConstraintVariables, STATGROUP_Chaos);
 DECLARE_CYCLE_STAT(TEXT("Chaos.Deformable.Solver.InitializeGaussSeidelConstraintVariables"), STAT_ChaosDeformableSolver_InitializeGaussSeidelConstraintVariables, STATGROUP_Chaos);
 DECLARE_CYCLE_STAT(TEXT("Chaos.Deformable.Solver.InitializeGaussSeidelConstraint"), STAT_ChaosDeformableSolver_InitializeGaussSeidelConstraint, STATGROUP_Chaos);
+DECLARE_CYCLE_STAT(TEXT("Chaos.Deformable.Solver.InitializeMuscleActivation"), STAT_ChaosDeformableSolver_InitializeMuscleActivation, STATGROUP_Chaos);
 DECLARE_CYCLE_STAT(TEXT("Chaos.Deformable.Solver.RemoveSimulationObjects"), STAT_ChaosDeformableSolver_RemoveSimulationObjects, STATGROUP_Chaos);
 DECLARE_CYCLE_STAT(TEXT("Chaos.Deformable.Solver.RemoveProxy"), STAT_ChaosDeformableSolver_RemoveProxy, STATGROUP_Chaos);
 DECLARE_CYCLE_STAT(TEXT("Chaos.Deformable.Solver.AddProxy"), STAT_ChaosDeformableSolver_AddProxy, STATGROUP_Chaos);
@@ -142,6 +145,7 @@ namespace Chaos::Softs
 		if (Property.bDoSelfCollision)
 		{
 			SurfaceTriangleMesh.Reset(new Chaos::FTriangleMesh());
+			SurfaceVertices.Reset(new TArray<int32>);
 		}
 		if (Property.bUseGridBasedConstraints)
 		{
@@ -156,6 +160,7 @@ namespace Chaos::Softs
 			AllTetNuMeshArray.Reset(new TArray<FSolverReal>());
 			AllTetAlphaJArray.Reset(new TArray<FSolverReal>());
 			GSWeakConstraints.Reset(new FGaussSeidelWeakConstraints<FSolverReal, FSolverParticles>({}, {}, {}, {}, {}, GDeformableXPBDWeakConstraintParams));
+			MuscleActivationConstraints.Reset(new FMuscleActivationConstraints<FSolverReal, FSolverParticles>());
 		}
 
 		InitializeKinematicConstraint();
@@ -335,6 +340,10 @@ namespace Chaos::Softs
 					{
 						InitializeGaussSeidelConstraintVariables();
 					}
+					if (Property.bDoMuscleActivation)
+					{
+						InitializeMuscleActivationVariables();
+					}
 				}
 				UninitializedProxys_Internal.SetNum(0, EAllowShrinking::Yes);
 			}
@@ -374,6 +383,7 @@ namespace Chaos::Softs
 					InitializeDeformableParticles(*Proxy);
 					InitializeKinematicParticles(*Proxy);
 					InitializeWeakConstraint(*Proxy);
+					InitializeMuscleActivation(*Proxy);
 					InitializeTetrahedralConstraint(*Proxy);
 					InitializeGidBasedConstraints(*Proxy);
 					InitializeGaussSeidelConstraints(*Proxy);
@@ -1137,37 +1147,75 @@ namespace Chaos::Softs
 								}
 							}
 						}
-						int32 Offset = ParticleComponentIndex.Num();
-						ParticleComponentIndex.SetNum(ParticleComponentIndex.Num() + Vertex->Num());
-						for (int32 i = 0; i < Vertex->Num(); i++) {
-							ParticleComponentIndex[i + Offset] = ComponentOffset;
-						}
-						int NewComponentOffset = ComponentOffset;
-						if (const TManagedArray<int32>* ComponentIndex = Rest->FindAttribute<int32>("ComponentIndex", FGeometryCollection::VerticesGroup))
+						VertexOffset += Vertex->Num();
+
+						
+						if (!Property.bDoInComponentSelfCollision) //Component-Component collisions
 						{
-							ensureMsgf(ComponentIndex->Num() == Vertex->Num(), TEXT("ComponentIndex size is not equal to vertex size"));
-							for (int32 i = 0; i < ComponentIndex->Num(); i++) {
-								if ((*ComponentIndex)[i] < 0)
-								{
-									ParticleComponentIndex[i + Offset] = (*ComponentIndex)[i]; //Isolated Nodes
-								}
-								else
-								{
-									ParticleComponentIndex[i + Offset] = ComponentOffset + (*ComponentIndex)[i];
-									NewComponentOffset = NewComponentOffset < ParticleComponentIndex[i + Offset] ? ParticleComponentIndex[i + Offset] : NewComponentOffset;
+							int32 Offset = ParticleComponentIndex.Num();
+							ParticleComponentIndex.SetNum(ParticleComponentIndex.Num() + Vertex->Num());
+							for (int32 i = 0; i < Vertex->Num(); i++) {
+								ParticleComponentIndex[i + Offset] = ComponentOffset;
+							}
+							int NewComponentOffset = ComponentOffset;
+							if (const TManagedArray<int32>* ComponentIndex = Rest->FindAttribute<int32>("ComponentIndex", FGeometryCollection::VerticesGroup))
+							{
+								ensureMsgf(ComponentIndex->Num() == Vertex->Num(), TEXT("ComponentIndex size is not equal to vertex size"));
+								for (int32 i = 0; i < ComponentIndex->Num(); i++) {
+									if ((*ComponentIndex)[i] < 0)
+									{
+										ParticleComponentIndex[i + Offset] = (*ComponentIndex)[i]; //Isolated Nodes
+									}
+									else
+									{
+										ParticleComponentIndex[i + Offset] = ComponentOffset + (*ComponentIndex)[i];
+										NewComponentOffset = NewComponentOffset < ParticleComponentIndex[i + Offset] ? ParticleComponentIndex[i + Offset] : NewComponentOffset;
+									}
 								}
 							}
+							ComponentOffset = NewComponentOffset + 1;
 						}
-						ComponentOffset = NewComponentOffset + 1;
-						VertexOffset += Vertex->Num();
 					}
 				}
 			}
 		}
 
 		SurfaceTriangleMesh->Init(*SurfaceElements);
+		SurfaceTriangleMesh->GetVertexSetAsArray(*SurfaceVertices);
 		TriangleMeshCollisions.Reset(new FPBDTriangleMeshCollisions(
 			0, Evolution->Particles().Size(), *SurfaceTriangleMesh, false, false));
+		ParticleTriangleExclusionMap.Reset();
+		if (Property.bDoInComponentSelfCollision)
+		{	
+			int32 NRadius = Property.NRingExcluded;
+			SurfaceTriangleMesh->GetPointToNeighborsMap();
+			SurfaceTriangleMesh->GetPointToTriangleMap(); //Initialize Maps before parallel
+			ParticleTriangleExclusionMap.Reserve(SurfaceVertices->Num());
+			PhysicsParallelFor(SurfaceVertices->Num(), 
+				[this, &NRadius](int32 i)
+				{
+					int32 VertexId = (*SurfaceVertices)[i];
+					TSet<int32>& TriangleSet = ParticleTriangleExclusionMap.FindOrAdd(VertexId);
+					TSet<int32> NRing;
+					if (NRadius > 1)
+					{
+						NRing = SurfaceTriangleMesh->GetNRing(VertexId, NRadius-1);
+					}
+					else
+					{
+						NRing.Add(VertexId);
+					}
+					for (TSet<int32>::TConstIterator It = NRing.CreateConstIterator(); It; ++It)
+					{
+						TArray<int32> CoincidentTriangles = SurfaceTriangleMesh->GetCoincidentTriangles(*It);
+						for (int32 j = 0; j < CoincidentTriangles.Num(); ++j)
+						{
+							TriangleSet.Add(CoincidentTriangles[j]);
+						}
+					}
+				}
+			);
+		}
 	}
 
 	void FDeformableSolver::InitializeGridBasedConstraintVariables()
@@ -1274,18 +1322,18 @@ namespace Chaos::Softs
 			int32 ConstraintIndex1 = Evolution->AddConstraintRuleRange(1, true);
 			Evolution->ConstraintRules()[ConstraintIndex1] =
 				[this](FSolverParticles& InParticles, const FSolverReal Dt)
-			{
-				this->GSMainConstraint->Apply(InParticles, Dt);
-			};
+				{
+					this->GSMainConstraint->Apply(InParticles, Dt);
+				};
 
 
 			int32 StaticIndex = GSMainConstraint->AddStaticConstraintResidualAndHessianRange(1);
 
 			GSMainConstraint->StaticConstraintResidualAndHessian()[StaticIndex] = [this](const FSolverParticles& Particles, const int32 ElementIndex, const int32 ElementIndexLocal, const FSolverReal Dt, TVec3<FSolverReal>& ParticleResidual, Chaos::PMatrix<FSolverReal, 3, 3>& ParticleHessian)
-			{
-				this->GSCorotatedConstraints->AddHyperelasticResidualAndHessian(Particles, ElementIndex, ElementIndexLocal, Dt, ParticleResidual, ParticleHessian);
-			};
-				
+				{
+					this->GSCorotatedConstraints->AddHyperelasticResidualAndHessian(Particles, ElementIndex, ElementIndexLocal, Dt, ParticleResidual, ParticleHessian);
+				};
+
 
 			if (Property.bEnablePositionTargets)
 			{
@@ -1298,15 +1346,15 @@ namespace Chaos::Softs
 
 				int32 StaticIndex1 = GSMainConstraint->AddStaticConstraintResidualAndHessianRange(1);
 				GSMainConstraint->StaticConstraintResidualAndHessian()[StaticIndex1] = [this](const FSolverParticles& Particles, const int32 ConstraintIndex, const int32 ConstraintIndexLocal, const FSolverReal Dt, TVec3<FSolverReal>& ParticleResidual, Chaos::PMatrix<FSolverReal, 3, 3>& ParticleHessian)
-				{
-					this->GSWeakConstraints->AddWCResidualAndHessian(Particles, ConstraintIndex, ConstraintIndexLocal, Dt, ParticleResidual, ParticleHessian);
-				};
+					{
+						this->GSWeakConstraints->AddWCResidualAndHessian(Particles, ConstraintIndex, ConstraintIndexLocal, Dt, ParticleResidual, ParticleHessian);
+					};
 
 				int32 PerNodeIndex = GSMainConstraint->AddPerNodeHessianRange(1);
 				GSMainConstraint->PerNodeHessian()[PerNodeIndex] = [this](const int32 p, const FSolverReal Dt, Chaos::PMatrix<FSolverReal, 3, 3>& ParticleHessian)
-				{
-					this->GSWeakConstraints->AddWCHessian(p, Dt, ParticleHessian);
-				};
+					{
+						this->GSWeakConstraints->AddWCHessian(p, Dt, ParticleHessian);
+					};
 
 			}
 		}
@@ -1318,58 +1366,99 @@ namespace Chaos::Softs
 			int32 InitIndex1 = Evolution->AddConstraintInitRange(1, true);
 			Evolution->ConstraintInits()[InitIndex1] =
 				[this](FSolverParticles& InParticles, const FSolverReal Dt)
-			{
-				this->GSWeakConstraints->Init(InParticles, Dt);
-			};
+				{
+					this->GSWeakConstraints->Init(InParticles, Dt);
+				};
 		}
 
 		if (Property.bDoSelfCollision)
 		{
 			int32 DynamicIndex = GSMainConstraint->AddDynamicConstraintResidualAndHessianRange(1);
 			GSMainConstraint->DynamicConstraintResidualAndHessian()[DynamicIndex] = [this](const FSolverParticles& Particles, const int32 ConstraintIndex, const int32 ConstraintIndexLocal, const FSolverReal Dt, TVec3<FSolverReal>& ParticleResidual, Chaos::PMatrix<FSolverReal, 3, 3>& ParticleHessian)
-			{
-				this->GSWeakConstraints->AddWCResidualAndHessian(Particles, ConstraintIndex + this->GSWeakConstraints->InitialWCSize, ConstraintIndexLocal, Dt, ParticleResidual, ParticleHessian);
-			};
+				{
+					this->GSWeakConstraints->AddWCResidualAndHessian(Particles, ConstraintIndex + this->GSWeakConstraints->InitialWCSize, ConstraintIndexLocal, Dt, ParticleResidual, ParticleHessian);
+				};
 
 			if (!Property.bEnablePositionTargets)
 			{
 				int32 PerNodeIndex = GSMainConstraint->AddPerNodeHessianRange(1);
 				GSMainConstraint->PerNodeHessian()[PerNodeIndex] = [this](const int32 p, const FSolverReal Dt, Chaos::PMatrix<FSolverReal, 3, 3>& ParticleHessian)
-				{
-					this->GSWeakConstraints->AddWCHessian(p, Dt, ParticleHessian);
-				};
+					{
+						this->GSWeakConstraints->AddWCHessian(p, Dt, ParticleHessian);
+					};
 			}
 
-			this->GSWeakConstraints->UpdateBoundaryVertices(*SurfaceElements);
-			
 			int32 InitIndex = Evolution->AddConstraintInitRange(1, true);
 			Evolution->ConstraintInits()[InitIndex] =
 				[this](FSolverParticles& InParticles, const FSolverReal Dt)
-			{
-				//TODO(Yushan & Joey): Add the collision detection code in here:
-				//this->GSWeakConstraints->CollisionDetection(this->Evolution->Particles(), *SurfaceElements, ParticleComponentIndex, Property.CollisionSearchRadius, Property.CollisionSpringStiffness, Property.bAllowSliding);
-				
-				this->TriangleMeshCollisions->Init(InParticles);
-				TArray<FPBDTriangleMeshCollisions::FGIAColor> EmptyGIAColors;
-				this->GSWeakConstraints->CollisionDetectionSpatialHash(this->Evolution->Particles(), *SurfaceTriangleMesh, ParticleComponentIndex, TriangleMeshCollisions->GetDynamicSpatialHash(), Property.CollisionSearchRadius, Property.CollisionSpringStiffness, Property.bAllowSliding);
-				
-			};
+				{
+					this->TriangleMeshCollisions->Init(InParticles, Property.CollisionSearchRadius);
+					TArray<FPBDTriangleMeshCollisions::FGIAColor> EmptyGIAColors;
+					if (Property.bDoInComponentSelfCollision)
+					{
+						this->GSWeakConstraints->CollisionDetectionSpatialHashInComponent(this->Evolution->Particles(), *SurfaceVertices, *SurfaceTriangleMesh, ParticleTriangleExclusionMap, TriangleMeshCollisions->GetDynamicSpatialHash(), Property.CollisionSearchRadius, Property.CollisionSpringStiffness, Property.bAllowSliding);
+					}
+					else
+					{
+						this->GSWeakConstraints->CollisionDetectionSpatialHash(this->Evolution->Particles(), *SurfaceVertices, *SurfaceTriangleMesh, ParticleComponentIndex, TriangleMeshCollisions->GetDynamicSpatialHash(), Property.CollisionSearchRadius, Property.CollisionSpringStiffness, Property.bAllowSliding);
+					}
+				};
 
 			int32 InitIndex1 = Evolution->AddConstraintInitRange(1, true);
 
 
 			Evolution->ConstraintInits()[InitIndex1] =
 				[this](FSolverParticles& InParticles, const FSolverReal Dt)
-			{
-				TArray<TArray<int32>> WCCollisionConstraints, WCCollisionIncidentElements, WCCollisionIncidentElementsLocal;
-				this->GSWeakConstraints->ComputeCollisionWCDataSimplified(WCCollisionConstraints, WCCollisionIncidentElements, WCCollisionIncidentElementsLocal);
-				this->GSMainConstraint->AddDynamicConstraints(WCCollisionConstraints, WCCollisionIncidentElements, WCCollisionIncidentElementsLocal);
-				this->GSMainConstraint->InitDynamicColor(InParticles);
-			}; 
-			
+				{
+					TArray<TArray<int32>> WCCollisionConstraints, WCCollisionIncidentElements, WCCollisionIncidentElementsLocal;
+					this->GSWeakConstraints->ComputeCollisionWCDataSimplified(WCCollisionConstraints, WCCollisionIncidentElements, WCCollisionIncidentElementsLocal);
+					this->GSMainConstraint->AddDynamicConstraints(WCCollisionConstraints, WCCollisionIncidentElements, WCCollisionIncidentElementsLocal);
+					this->GSMainConstraint->InitDynamicColor(InParticles);
+				};
+
 		}
 	}
+	void FDeformableSolver::InitializeMuscleActivationVariables()
+	{
+		int32 InitIndex = Evolution->AddConstraintInitRange(1, true);
+		Evolution->ConstraintInits()[InitIndex] =
+			[this](FSolverParticles& InParticles, const FSolverReal Dt)
+			{
+				this->MuscleActivationConstraints->UpdateLengthBasedMuscleActivation(InParticles);
+				if (Property.bUseGaussSeidelConstraints)
+				{
+					if (Property.bUseGSNeohookean)
+					{
+						this->MuscleActivationConstraints->ApplyMuscleActivation(*this->GSNeohookeanConstraints);
+					}
+					else
+					{
+						this->MuscleActivationConstraints->ApplyMuscleActivation(*this->GSCorotatedConstraints);
+					}
+				}
+				else
+				{
+					this->MuscleActivationConstraints->ApplyMuscleActivation(*this->CorotatedConstraints[0]);
+				}
+			};
+	}
 
+	void FDeformableSolver::InitializeMuscleActivation(FFleshThreadingProxy& Proxy)
+	{
+		PERF_SCOPE(STAT_ChaosDeformableSolver_InitializeMuscleActivation);
+		if (Property.bDoMuscleActivation)
+		{
+			//Grab MuscleActivationElements, OriginInsertion, FiberDirectionMatrix
+			const FManagedArrayCollection& Rest = Proxy.GetRestCollection();
+			GeometryCollection::Facades::FMuscleActivationFacade MuscleActivationFacade(Rest);
+			if (MuscleActivationFacade.IsValid())
+			{
+				int32 VertexOffset = Proxy.GetSolverParticleRange()[0];
+				int32 ElementOffset = AllElements->Num(); //InitializeMuscleActivation has to go before adding tetrahedrons
+				MuscleActivationConstraints->AddMuscles(MuscleActivationFacade, VertexOffset, ElementOffset);
+			}
+		}
+	}
 
 	void FDeformableSolver::RemoveSimulationObjects()
 	{
@@ -1432,6 +1521,7 @@ namespace Chaos::Softs
 			{
 				if (FFleshThreadingProxy* Proxy = BaseProxyPair.Value->As<FFleshThreadingProxy>())
 				{
+					InitializeMuscleActivation(*Proxy);
 					InitializeTetrahedralConstraint(*Proxy);
 					InitializeGidBasedConstraints(*Proxy);
 					InitializeGaussSeidelConstraints(*Proxy);
diff --git a/Engine/Source/Runtime/Experimental/Chaos/Private/Chaos/TriangleMesh.cpp b/Engine/Source/Runtime/Experimental/Chaos/Private/Chaos/TriangleMesh.cpp
--- a/Engine/Source/Runtime/Experimental/Chaos/Private/Chaos/TriangleMesh.cpp
+++ b/Engine/Source/Runtime/Experimental/Chaos/Private/Chaos/TriangleMesh.cpp
@@ -1877,6 +1877,75 @@ bool FTriangleMesh::PointProximityQuery(const TSpatialHashType<Softs::FSolverRea
 
 }
 
+template<typename T>
+bool FTriangleMesh::PointClosestTriangleQuery(const TSpatialHashType<T>& SpatialHash, const TConstArrayView<TVec3<T>>& Points, const int32 PointIndex, const TVec3<T>& PointPosition, const T PointThickness, const T ThisThickness,
+	TFunctionRef<bool(const int32 PointIndex, const int32 TriangleIndex)> BroadphaseTest, TArray<TTriangleCollisionPoint<T>>& Result) const
+{
+	TRACE_CPUPROFILER_EVENT_SCOPE(FTriangleMesh_PointClosestTriangleQuery);
+	const T TotalThickness = ThisThickness + PointThickness;
+	const T TotalThicknessSq = TotalThickness * TotalThickness;
+	typename TSpatialHashType<T>::FVectorAABB QueryBounds(PointPosition);
+	QueryBounds.Thicken(TotalThickness);
+
+	const TArray<int32> PotentialIntersections = SpatialHash.FindAllIntersections(QueryBounds,
+		[PointIndex, &BroadphaseTest](int32 Payload)
+		{
+			return BroadphaseTest(PointIndex, Payload);
+		});
+	TRACE_CPUPROFILER_EVENT_SCOPE(FTriangleMesh_PointClosestTriangleQueryAddTriangle);
+	Result.Reset(0);
+	T MinDistanceSq = TotalThicknessSq;
+	int32 MinTriId = INDEX_NONE;
+	for (int32 TriIdx : PotentialIntersections)
+	{
+		const TVec3<T>& A = Points[MElements[TriIdx][0]];
+		const TVec3<T>& B = Points[MElements[TriIdx][1]];
+		const TVec3<T>& C = Points[MElements[TriIdx][2]];
+		TVec3<T> Bary;
+		const TVec3<T> ClosestPoint = FindClosestPointAndBaryOnTriangle(A, B, C, PointPosition, Bary);
+
+		const T DistSq = (PointPosition - ClosestPoint).SizeSquared();
+		if (DistSq > TotalThicknessSq || DistSq > MinDistanceSq)
+		{
+			// Failed narrow test.
+			continue;
+		}
+		else
+		{
+			MinDistanceSq = DistSq;
+			MinTriId = TriIdx;
+		}
+	}
+	if (MinTriId != INDEX_NONE)
+	{
+		const TVec3<T>& A = Points[MElements[MinTriId][0]];
+		const TVec3<T>& B = Points[MElements[MinTriId][1]];
+		const TVec3<T>& C = Points[MElements[MinTriId][2]];
+		TVec3<T> Bary;
+		const TVec3<T> ClosestPoint = FindClosestPointAndBaryOnTriangle(A, B, C, PointPosition, Bary);
+		const T DistSq = (PointPosition - ClosestPoint).SizeSquared();
+		TVec3<T> Normal = TVec3<T>::CrossProduct(B - A, C - A).GetSafeNormal();
+		if (TVec3<T>::DotProduct(Normal, PointPosition - A) < 0) //Point is outside of triangle
+		{
+			return false;
+		}	
+
+		TTriangleCollisionPoint<T> CollisionPoint;
+		CollisionPoint.ContactType = TTriangleCollisionPoint<T>::EContactType::PointFace;
+		CollisionPoint.Indices[0] = PointIndex;
+		CollisionPoint.Indices[1] = MinTriId;
+		CollisionPoint.Bary = TVec4<T>((T)1., Bary.X, Bary.Y, Bary.Z);
+		CollisionPoint.Location = ClosestPoint;
+		CollisionPoint.Normal = Normal;
+		CollisionPoint.Phi = FMath::Sqrt(DistSq);
+		Result.Add(CollisionPoint);
+	}
+	return Result.Num() > 0;
+}
+template bool FTriangleMesh::PointClosestTriangleQuery<FRealSingle>(const TSpatialHashType<FRealSingle>& SpatialHash, const TConstArrayView<TVector<FRealSingle, 3>>& Points, const int32 PointIndex, const TVector<FRealSingle, 3>& PointPosition, const FRealSingle PointThickness, const FRealSingle ThisThickness, TFunctionRef<bool(const int32 PointIndex, const int32 TriangleIndex)> BroadphaseTest, TArray<TTriangleCollisionPoint<FRealSingle>>& Result) const;
+template bool FTriangleMesh::PointClosestTriangleQuery<FRealDouble>(const TSpatialHashType<FRealDouble>& SpatialHash, const TConstArrayView<TVector<FRealDouble, 3>>& Points, const int32 PointIndex, const TVector<FRealDouble, 3>& PointPosition, const FRealDouble PointThickness, const FRealDouble ThisThickness, TFunctionRef<bool(const int32 PointIndex, const int32 TriangleIndex)> BroadphaseTest, TArray<TTriangleCollisionPoint<FRealDouble>>& Result) const;
+
+
 template<typename T>
 bool FTriangleMesh::EdgeIntersectionQuery(const TSpatialHashType<T>& SpatialHash, const TConstArrayView<TVec3<T>>& Points, const int32 EdgeIndex, const TVec3<T>& EdgePosition1, const TVec3<T>& EdgePosition2,
 	TFunctionRef<bool(const int32 EdgeIndex, const int32 TriangleIndex)> BroadphaseTest, TArray<TTriangleCollisionPoint<T>>& Result) const
diff --git a/Engine/Source/Runtime/Experimental/Chaos/Private/GeometryCollection/Facades/CollectionMuscleActivationFacade.cpp b/Engine/Source/Runtime/Experimental/Chaos/Private/GeometryCollection/Facades/CollectionMuscleActivationFacade.cpp
new file mode 100644
--- /dev/null
+++ b/Engine/Source/Runtime/Experimental/Chaos/Private/GeometryCollection/Facades/CollectionMuscleActivationFacade.cpp
@@ -0,0 +1,86 @@
+// Copyright Epic Games, Inc. All Rights Reserved.
+
+/*=============================================================================
+	FleshCollection.cpp: FFleshCollection methods.
+=============================================================================*/
+
+#include "GeometryCollection/Facades/CollectionMuscleActivationFacade.h"
+#include "GeometryCollection/GeometryCollection.h"
+
+namespace GeometryCollection::Facades
+{
+	// Attributes
+	const FName FMuscleActivationFacade::GroupName("MuscleActivation");
+	const FName FMuscleActivationFacade::MuscleActivationElement("MuscleActivationElement");
+	const FName FMuscleActivationFacade::OriginInsertionPair("OriginInsertionPair");
+	const FName FMuscleActivationFacade::OriginInsertionRestLength("OriginInsertionRestLength");
+	const FName FMuscleActivationFacade::FiberDirectionMatrix("FiberDirectionMatrix");
+
+
+	FMuscleActivationFacade::FMuscleActivationFacade(FManagedArrayCollection& InCollection)
+		: ConstCollection(InCollection)
+		, Collection(&InCollection)
+		, MuscleActivationElementAttribute(InCollection, MuscleActivationElement, GroupName)
+		, OriginInsertionPairAttribute(InCollection, OriginInsertionPair, GroupName)
+		, OriginInsertionRestLengthAttribute(InCollection, OriginInsertionRestLength, GroupName)
+		, FiberDirectionMatrixAttribute(InCollection, FiberDirectionMatrix, GroupName)
+	{
+		DefineSchema();
+	}
+
+	FMuscleActivationFacade::FMuscleActivationFacade(const FManagedArrayCollection& InCollection)
+		: ConstCollection(InCollection)
+		, Collection(nullptr)
+		, MuscleActivationElementAttribute(InCollection, MuscleActivationElement, GroupName)
+		, OriginInsertionPairAttribute(InCollection, OriginInsertionPair, GroupName)
+		, OriginInsertionRestLengthAttribute(InCollection, OriginInsertionRestLength, GroupName)
+		, FiberDirectionMatrixAttribute(InCollection, FiberDirectionMatrix, GroupName)
+	{
+		
+	}
+
+	bool FMuscleActivationFacade::IsValid() const
+	{
+		return MuscleActivationElementAttribute.IsValid() && OriginInsertionPairAttribute.IsValid() && OriginInsertionRestLengthAttribute.IsValid() && FiberDirectionMatrixAttribute.IsValid();
+	}
+
+	void FMuscleActivationFacade::DefineSchema()
+	{
+		check(!IsConst());
+		MuscleActivationElementAttribute.Add(ManageArrayAccessor::EPersistencePolicy::MakePersistent, "Tetrahedral");
+		OriginInsertionPairAttribute.Add(ManageArrayAccessor::EPersistencePolicy::MakePersistent, FGeometryCollection::VerticesGroup);
+		OriginInsertionRestLengthAttribute.Add();
+		FiberDirectionMatrixAttribute.Add();
+	}
+
+	int32 FMuscleActivationFacade::AddMuscleActivationData(const FMuscleActivationData& InputData)
+	{
+		check(!IsConst());
+		if (IsValid())
+		{
+			int32 NewIndex = MuscleActivationElementAttribute.AddElements(1);
+			MuscleActivationElementAttribute.Modify()[NewIndex] = InputData.MuscleActivationElement;
+			OriginInsertionPairAttribute.Modify()[NewIndex] = InputData.OriginInsertionPair;
+			OriginInsertionRestLengthAttribute.Modify()[NewIndex] = InputData.OriginInsertionRestLength;
+			FiberDirectionMatrixAttribute.Modify()[NewIndex] = InputData.FiberDirectionMatrix;
+			return NewIndex;
+		}
+		return INDEX_NONE;
+	}
+
+	const FMuscleActivationData FMuscleActivationFacade::GetMuscleActivationData(const int32 DataIndex) const
+	{
+		FMuscleActivationData ReturnData;
+		if (IsValid())
+		{
+			if (NumMuscles() > DataIndex && DataIndex > -1)
+			{
+				ReturnData.MuscleActivationElement = MuscleActivationElementAttribute.Get()[DataIndex];
+				ReturnData.OriginInsertionPair = OriginInsertionPairAttribute.Get()[DataIndex];
+				ReturnData.OriginInsertionRestLength = OriginInsertionRestLengthAttribute.Get()[DataIndex];
+				ReturnData.FiberDirectionMatrix = FiberDirectionMatrixAttribute.Get()[DataIndex];
+			}
+		}
+		return ReturnData;
+	}
+}
diff --git a/Engine/Source/Runtime/Experimental/Chaos/Public/Chaos/Deformable/ChaosDeformableSolver.h b/Engine/Source/Runtime/Experimental/Chaos/Public/Chaos/Deformable/ChaosDeformableSolver.h
--- a/Engine/Source/Runtime/Experimental/Chaos/Public/Chaos/Deformable/ChaosDeformableSolver.h
+++ b/Engine/Source/Runtime/Experimental/Chaos/Public/Chaos/Deformable/ChaosDeformableSolver.h
@@ -15,6 +15,7 @@
 #include "Chaos/Deformable/GaussSeidelNeohookeanConstraints.h"
 #include "Chaos/Deformable/GaussSeidelWeakConstraints.h"
 #include "Chaos/Deformable/GaussSeidelMainConstraint.h"
+#include "Chaos/Deformable/MuscleActivationConstraints.h"
 #include "Chaos/XPBDWeakConstraints.h"
 #include "Chaos/BlendedXPBDCorotatedConstraints.h"
 #include "Chaos/XPBDGridBasedCorotatedConstraints.h"
@@ -132,6 +133,8 @@ namespace Chaos::Softs
 		CHAOS_API void InitializeSelfCollisionVariables();
 		CHAOS_API void InitializeGridBasedConstraintVariables();
 		CHAOS_API void InitializeGaussSeidelConstraintVariables();
+		CHAOS_API void InitializeMuscleActivationVariables();
+		CHAOS_API void InitializeMuscleActivation(FFleshThreadingProxy& Proxy);
 		CHAOS_API void UpdateCollisionBodies(FCollisionManagerProxy&, FThreadingProxy::FKey, FSolverReal DeltaTime);
 		CHAOS_API void RemoveSimulationObjects();
 		CHAOS_API TArray<Chaos::TVec3<FSolverReal>> ComputeParticleTargets(const TArray<TArray<int32>>& ParticleIndices);
@@ -191,6 +194,7 @@ namespace Chaos::Softs
 		TUniquePtr <TArray<TVec3<int32>>> SurfaceElements;
 		TUniquePtr <TArray<Chaos::TVec4<int32>>> AllElements;
 		TUniquePtr <FTriangleMesh> SurfaceTriangleMesh;
+		TUniquePtr <TArray<int32>> SurfaceVertices;
 		TUniquePtr <TArray<TArray<int32>>> AllIncidentElements;
 		TUniquePtr <TArray<TArray<int32>>> AllIncidentElementsLocal;
 		TUniquePtr <TArray<FSolverReal>> AllTetEMeshArray;
@@ -201,6 +205,10 @@ namespace Chaos::Softs
 		TUniquePtr <TArray<FSolverReal>> AllWeights;
 		TUniquePtr <TArray<FSolverReal>> AllSecondWeights;
 		TArray<int32> ParticleComponentIndex;
+		TMap<int32, TSet<int32>> ParticleTriangleExclusionMap;
+		//Muscle Activation Variables
+		TUniquePtr<Softs::FMuscleActivationConstraints<Softs::FSolverReal, Softs::FSolverParticles>> MuscleActivationConstraints;
+
 		//typedef TMap<int32, TTuple<float, Chaos::Softs::FPAndInvM, FVector3f>> TransientConstraintBufferMap;
 		typedef TMap<int32, TTuple<float, float, FVector3f>> TransientConstraintBufferMap;
 		TransientConstraintBufferMap TransientConstraintBuffer;
diff --git a/Engine/Source/Runtime/Experimental/Chaos/Public/Chaos/Deformable/ChaosDeformableSolverTypes.h b/Engine/Source/Runtime/Experimental/Chaos/Public/Chaos/Deformable/ChaosDeformableSolverTypes.h
--- a/Engine/Source/Runtime/Experimental/Chaos/Public/Chaos/Deformable/ChaosDeformableSolverTypes.h
+++ b/Engine/Source/Runtime/Experimental/Chaos/Public/Chaos/Deformable/ChaosDeformableSolverTypes.h
@@ -23,24 +23,27 @@ namespace Chaos::Softs
 			bool InCacheToFile = false,
 			bool InbEnableKinematics = true,
 			bool InbUseFloor = true,
-			bool InbDoSelfCollision = false,
 			bool InbUseGridBasedConstraints = false,
-			FSolverReal InGridDx = (FSolverReal)1. ,
+			FSolverReal InGridDx = (FSolverReal)1.,
 			bool InbDoQuasistatics = false,
 			FSolverReal InEMesh = (FSolverReal)100000.,
 			bool InbDoBlended = false,
 			FSolverReal InBlendedZeta = (FSolverReal).1,
 			FSolverReal InDamping = (FSolverReal)0,
-			bool InbEnableGravity = true, 
-			bool InbEnableCorotatedConstraints = true, 
-			bool InbEnablePositionTargets = true, 
-			bool InbUseGaussSeidelConstraints = false, 
+			bool InbEnableGravity = true,
+			bool InbEnableCorotatedConstraints = true,
+			bool InbEnablePositionTargets = true,
+			bool InbUseGaussSeidelConstraints = false,
 			bool InbUseSOR = true,
-			FSolverReal InOmegaSOR = (FSolverReal)1.6, 
+			FSolverReal InOmegaSOR = (FSolverReal)1.6,
 			bool InbUseGSNeohookean = false,
-			FSolverReal InCollisionSearchRadius = (FSolverReal)10.,
+			bool InbDoSelfCollision = false,
+			bool InbDoInComponentSelfCollision = false,
+			int32 InNRingExcluded = 1,
+			FSolverReal InCollisionSearchRadius = (FSolverReal)0,
 			FSolverReal InCollisionSpringStiffness = (FSolverReal)500.,
-			bool InbAllowSliding = true)
+			bool InbAllowSliding = true,
+			bool InbDoMuscleActivation = false)
 			: NumSolverSubSteps(InNumSolverSubSteps)
 			, NumSolverIterations(InNumSolverIterations)
 			, FixTimeStep(InFixTimeStep)
@@ -48,7 +51,6 @@ namespace Chaos::Softs
 			, CacheToFile(InCacheToFile)
 			, bEnableKinematics(InbEnableKinematics)
 			, bUseFloor(InbUseFloor)
-			, bDoSelfCollision(InbDoSelfCollision)
 			, bUseGridBasedConstraints(InbUseGridBasedConstraints)
 			, GridDx(InGridDx)
 			, bDoQuasistatics(InbDoQuasistatics)
@@ -63,9 +65,13 @@ namespace Chaos::Softs
 			, bUseSOR(InbUseSOR)
 			, OmegaSOR(InOmegaSOR)
 			, bUseGSNeohookean(InbUseGSNeohookean)
+			, bDoSelfCollision(InbDoSelfCollision)
+			, bDoInComponentSelfCollision(InbDoInComponentSelfCollision)
+			, NRingExcluded(InNRingExcluded)
 			, CollisionSearchRadius(InCollisionSearchRadius)
 			, CollisionSpringStiffness(InCollisionSpringStiffness)
 			, bAllowSliding(InbAllowSliding)
+			, bDoMuscleActivation(InbDoMuscleActivation)
 		{}
 
 		int32 NumSolverSubSteps = 5;
@@ -75,7 +81,6 @@ namespace Chaos::Softs
 		bool CacheToFile = false;
 		bool bEnableKinematics = true;
 		bool bUseFloor = true;
-		bool bDoSelfCollision = false;
 		bool bUseGridBasedConstraints = false;
 		FSolverReal GridDx = (FSolverReal)1.;
 		bool bDoQuasistatics = false;
@@ -90,9 +95,13 @@ namespace Chaos::Softs
 		bool bUseSOR = true;
 		FSolverReal OmegaSOR = (FSolverReal)1.6;
 		bool bUseGSNeohookean = false;
-		FSolverReal CollisionSearchRadius = (FSolverReal)10.;
+		bool bDoSelfCollision = false;
+		bool bDoInComponentSelfCollision = false;
+		int32 NRingExcluded = 1;
+		FSolverReal CollisionSearchRadius = (FSolverReal)0;
 		FSolverReal CollisionSpringStiffness = (FSolverReal)500.;
 		bool bAllowSliding = true;
+		bool bDoMuscleActivation = false;
 	};
 
 
diff --git a/Engine/Source/Runtime/Experimental/Chaos/Public/Chaos/Deformable/GaussSeidelWeakConstraints.h b/Engine/Source/Runtime/Experimental/Chaos/Public/Chaos/Deformable/GaussSeidelWeakConstraints.h
--- a/Engine/Source/Runtime/Experimental/Chaos/Public/Chaos/Deformable/GaussSeidelWeakConstraints.h
+++ b/Engine/Source/Runtime/Experimental/Chaos/Public/Chaos/Deformable/GaussSeidelWeakConstraints.h
@@ -15,7 +15,7 @@
 #include "Chaos/Triangle.h"
 #include "Chaos/TriangleCollisionPoint.h"
 #include "Chaos/TriangleMesh.h"
-
+#include <unordered_map>
 namespace Chaos::Softs
 {
 
@@ -336,17 +336,6 @@ namespace Chaos::Softs
 
 		}
 
-		void UpdateBoundaryVertices(const TArray<TVec3<int32>> Elements)
-		{
-			for (int32 i = 0; i < Elements.Num(); i++)
-			{
-				for (int32 j = 0; j < 3; j++)
-				{
-					BoundaryVertices.AddUnique(Elements[i][j]);
-				}
-			}
-		}
-
 		//CollisionDetectionSpatialHash should be faster than CollisionDetectionBVH
 		void CollisionDetectionBVH(const FSolverParticles& Particles, const TArray<TVec3<int32>>& SurfaceElements, const TArray<int32>& ComponentIndex, float DetectRadius = 1.f, float PositionTargetStiffness = 10000.f, bool UseAnisotropicSpring = true)
 		{
@@ -491,82 +480,105 @@ namespace Chaos::Softs
 		}
 
 		template<typename SpatialAccelerator>
-		void CollisionDetectionSpatialHash(const FSolverParticles& Particles, const FTriangleMesh& TriangleMesh, const TArray<int32>& ComponentIndex, const SpatialAccelerator& Spatial, float DetectRadius = 1.f, float PositionTargetStiffness = 10000.f, bool UseAnisotropicSpring = true)
+		void CollisionDetectionSpatialHash(const FSolverParticles& Particles, const TArray<int32>& SurfaceVertices, const FTriangleMesh& TriangleMesh, const TArray<int32>& ComponentIndex, const SpatialAccelerator& Spatial, float DetectRadius = 1.f, float PositionTargetStiffness = 10000.f, bool UseAnisotropicSpring = true)
 		{
 			TRACE_CPUPROFILER_EVENT_SCOPE(STAT_ChaosGaussSeidelWeakConstraintsCollisionDetectionSpatialHash);
-			constexpr int32 MaxConnectionsPerPoint = 1;
-			Resize(InitialWCSize + Particles.Size() * MaxConnectionsPerPoint);
-			
+			Resize(InitialWCSize + Particles.Size());
 			std::atomic<int32> ConstraintIndex(InitialWCSize);
 			const TArray<TVec3<int32>>& Elements = TriangleMesh.GetSurfaceElements();
-			PhysicsParallelFor(BoundaryVertices.Num(),
-				[this, &Spatial, &Particles, &ConstraintIndex, MaxConnectionsPerPoint, &TriangleMesh, &Elements, &DetectRadius, &ComponentIndex, &PositionTargetStiffness, &UseAnisotropicSpring](int32 i)
+			PhysicsParallelFor(SurfaceVertices.Num(),
+				[this, &Spatial, &Particles, &SurfaceVertices, &ConstraintIndex, &TriangleMesh, &Elements, &DetectRadius, &ComponentIndex, &PositionTargetStiffness, &UseAnisotropicSpring](int32 i)
 				{
-					const int32 Index = BoundaryVertices[i];
-
-					TArray< TTriangleCollisionPoint<FSolverReal> > Result;
-					if (TriangleMesh.PointProximityQuery(Spatial, static_cast<const TArrayView<const FSolverVec3>&>(Particles.XArray()), Index, Particles.GetX(Index), DetectRadius, DetectRadius,
+					const int32 Index = SurfaceVertices[i];
+					TArray< TTriangleCollisionPoint<FSolverReal>> Result;
+					//PointProximityQuery
+					if (TriangleMesh.PointClosestTriangleQuery(Spatial, static_cast<const TArrayView<const FSolverVec3>&>(Particles.XArray()), Index, Particles.GetX(Index), DetectRadius, DetectRadius,
 						[this, &ComponentIndex, &Elements](const int32 PointIndex, const int32 TriangleIndex)->bool
 						{
-							const TVector<int32, 3>& Elem = Elements[TriangleIndex];
-
 							if (ComponentIndex[PointIndex] == ComponentIndex[Elements[TriangleIndex][0]])
 							{
 								return false;
 							}
-
 							return true;
 						},
 						Result))
 					{
-
-						if (Result.Num() > MaxConnectionsPerPoint)
+						for (const TTriangleCollisionPoint<FSolverReal>& CollisionPoint : Result)
 						{
-							// TODO: once we have a PartialSort, use that instead here.
-							Result.Sort(
-								[](const TTriangleCollisionPoint<FSolverReal>& First, const TTriangleCollisionPoint<FSolverReal>& Second)->bool
-								{
-									return First.Phi < Second.Phi;
-								}
-							);
-							Result.SetNum(MaxConnectionsPerPoint, EAllowShrinking::No);
+							const TVector<int32, 3>& Elem = Elements[CollisionPoint.Indices[1]];
+							// // NOTE: CollisionPoint.Normal has already been flipped to point toward the Point, so need to recalculate here.
+							// const TTriangle<FSolverReal> Triangle(Particles.GetX(Elem[0]), Particles.GetX(Elem[1]), Particles.GetX(Elem[2]));
+							// if ((Particles.GetX(Index) - CollisionPoint.Location).Dot(-Triangle.GetNormal()) < 0) //Is point inside boundary? Normal should point outwards
+							// {
+							const int32 IndexToWrite = ConstraintIndex.fetch_add(1);
+							Indices[IndexToWrite] = { Elem[0], Elem[1] ,Elem[2] };
+							SecondIndices[IndexToWrite] = { Index };
+							Weights[IndexToWrite] = { CollisionPoint.Bary[1], CollisionPoint.Bary[2], CollisionPoint.Bary[3] };
+							SecondWeights[IndexToWrite] = { 1.f };
+							float SpringStiffness = 0.f;
+							for (int32 k = 0; k < 3; k++)
+							{
+								SpringStiffness += Weights[IndexToWrite][k] * PositionTargetStiffness * Particles.M(Elem[k]);
+							}
+							SpringStiffness += PositionTargetStiffness * Particles.M(Index);
+							Stiffness[IndexToWrite] = SpringStiffness;
+							IsAnisotropic[IndexToWrite] = UseAnisotropicSpring;
+							Normals[IndexToWrite] = CollisionPoint.Normal;						
 						}
+					}
+				}
+			);
+
+			// Shrink the arrays to the actual number of found constraints.
+			const int32 ConstraintNum = ConstraintIndex.load();
+			Resize(ConstraintNum);
+		}
 
+		template<typename SpatialAccelerator>
+		void CollisionDetectionSpatialHashInComponent(const FSolverParticles& Particles, const TArray<int32>& SurfaceVertices, const FTriangleMesh& TriangleMesh, const TMap<int32, TSet<int32>>& ExcludeMap, const SpatialAccelerator& Spatial, float DetectRadius = 1.f, float PositionTargetStiffness = 10000.f, bool UseAnisotropicSpring = true)
+		{
+			TRACE_CPUPROFILER_EVENT_SCOPE(STAT_ChaosGaussSeidelWeakConstraintsCollisionDetectionSpatialHashInComponent);
+			Resize(InitialWCSize + Particles.Size());
+			std::atomic<int32> ConstraintIndex(InitialWCSize);
+			const TArray<TVec3<int32>>& Elements = TriangleMesh.GetSurfaceElements();
+			PhysicsParallelFor(SurfaceVertices.Num(),
+				[this, &Spatial, &Particles, &SurfaceVertices, &ConstraintIndex, &TriangleMesh, &ExcludeMap, &Elements, &DetectRadius, &PositionTargetStiffness, &UseAnisotropicSpring](int32 i)
+				{
+					const int32 Index = SurfaceVertices[i];
+					TArray< TTriangleCollisionPoint<FSolverReal>> Result;
+					//PointProximityQuery
+					if (TriangleMesh.PointClosestTriangleQuery(Spatial, static_cast<const TArrayView<const FSolverVec3>&>(Particles.XArray()), Index, Particles.GetX(Index), DetectRadius, DetectRadius,
+						[this, &Elements, &ExcludeMap](const int32 PointIndex, const int32 TriangleIndex)->bool
+						{	
+							if (ExcludeMap.Find(PointIndex) && ExcludeMap[PointIndex].Contains(TriangleIndex))
+							{
+								return false;
+							}
+							return true;
+						},
+						Result))
+					{
 						for (const TTriangleCollisionPoint<FSolverReal>& CollisionPoint : Result)
 						{
 							const TVector<int32, 3>& Elem = Elements[CollisionPoint.Indices[1]];
-						
-
-							// NOTE: CollisionPoint.Normal has already been flipped to point toward the Point, so need to recalculate here.
-							const TTriangle<FSolverReal> Triangle(Particles.GetX(Elem[0]), Particles.GetX(Elem[1]), Particles.GetX(Elem[2]));
-							if ((Particles.GetX(Index) - CollisionPoint.Location).Dot(-Triangle.GetNormal()) < 0) //Is point inside boundary? Normal should point outwards
+							// // NOTE: CollisionPoint.Normal has already been flipped to point toward the Point, so need to recalculate here.
+							// const TTriangle<FSolverReal> Triangle(Particles.GetX(Elem[0]), Particles.GetX(Elem[1]), Particles.GetX(Elem[2]));
+							// if ((Particles.GetX(Index) - CollisionPoint.Location).Dot(-Triangle.GetNormal()) < 0) //Is point inside boundary? Normal should point outwards
+							// {
+							const int32 IndexToWrite = ConstraintIndex.fetch_add(1);
+							Indices[IndexToWrite] = { Elem[0], Elem[1] ,Elem[2] };
+							SecondIndices[IndexToWrite] = { Index };
+							Weights[IndexToWrite] = { CollisionPoint.Bary[1], CollisionPoint.Bary[2], CollisionPoint.Bary[3] };
+							SecondWeights[IndexToWrite] = { 1.f };
+							float SpringStiffness = 0.f;
+							for (int32 k = 0; k < 3; k++)
 							{
-								const int32 IndexToWrite = ConstraintIndex.fetch_add(1);
-
-								Indices[IndexToWrite] = { Elem[0], Elem[1] ,Elem[2] };
-								SecondIndices[IndexToWrite] = { Index };
-								Weights[IndexToWrite] = { CollisionPoint.Bary[1], CollisionPoint.Bary[2], CollisionPoint.Bary[3] };
-								SecondWeights[IndexToWrite] = { 1.f };
-								/*
-								Indices[IndexToWrite].Add(Elem[0]);
-								Indices[IndexToWrite].Add(Elem[1]);
-								Indices[IndexToWrite].Add(Elem[2]);
-								SecondIndices[IndexToWrite].Add(Index);
-								Weights[IndexToWrite].Add(ClosestBary[0]);
-								Weights[IndexToWrite].Add(ClosestBary[1]);
-								Weights[IndexToWrite].Add(ClosestBary[2]);
-								SecondWeights[SecondWeights.Num() - 1].Add(1.f);
-								*/
-								float SpringStiffness = 0.f;
-								for (int32 k = 0; k < 3; k++)
-								{
-									SpringStiffness += Weights[IndexToWrite][k] * PositionTargetStiffness * Particles.M(Elem[k]);
-								}
-								SpringStiffness += PositionTargetStiffness * Particles.M(Index);
-								Stiffness[IndexToWrite] = SpringStiffness;
-								IsAnisotropic[IndexToWrite] = UseAnisotropicSpring;
-								Normals[IndexToWrite] = Triangle.GetNormal();
-							}							
+								SpringStiffness += Weights[IndexToWrite][k] * PositionTargetStiffness * Particles.M(Elem[k]);
+							}
+							SpringStiffness += PositionTargetStiffness * Particles.M(Index);
+							Stiffness[IndexToWrite] = SpringStiffness;
+							IsAnisotropic[IndexToWrite] = UseAnisotropicSpring;
+							Normals[IndexToWrite] = CollisionPoint.Normal;						
 						}
 					}
 				}
@@ -854,7 +866,6 @@ namespace Chaos::Softs
 		//TArray<TVector<int32, 4>> NoCollisionConstraints;
 		//For debugging
 		bool Detected = false;
-		TArray<int32> BoundaryVertices;
 
 		TArray<TArray<int32>> NoCollisionConstraints;
 		TArray<TArray<int32>> NoCollisionWCIncidentElements;
diff --git a/Engine/Source/Runtime/Experimental/Chaos/Public/Chaos/Deformable/MuscleActivationConstraints.h b/Engine/Source/Runtime/Experimental/Chaos/Public/Chaos/Deformable/MuscleActivationConstraints.h
new file mode 100644
--- /dev/null
+++ b/Engine/Source/Runtime/Experimental/Chaos/Public/Chaos/Deformable/MuscleActivationConstraints.h
@@ -0,0 +1,84 @@
+// Copyright Epic Games, Inc. All Rights Reserved.
+#pragma once
+
+#include "Chaos/Framework/Parallel.h"
+#include "Chaos/XPBDCorotatedConstraints.h"
+#include "GeometryCollection/Facades/CollectionMuscleActivationFacade.h"
+
+namespace Chaos::Softs
+{
+
+	using Chaos::TVec3;
+	typedef GeometryCollection::Facades::FMuscleActivationFacade Facade;
+	typedef GeometryCollection::Facades::FMuscleActivationData Data;
+	template <typename T, typename ParticleType>
+	struct FMuscleActivationConstraints
+	{
+		//Handles muscle activation data
+		FMuscleActivationConstraints(){}
+
+		virtual ~FMuscleActivationConstraints() {}
+
+		void Init(const FSolverParticles& InParticles, const T Dt) const
+		{
+		}
+
+		void AddMuscles(const Facade& FMuscleActivation, int32 VertexOffset = 0, int32 ElementOffset = 0)
+		{
+			for (int MuscleIdx = 0; MuscleIdx < FMuscleActivation.NumMuscles(); MuscleIdx++)
+			{
+				Data MuscleActivationData = FMuscleActivation.GetMuscleActivationData(MuscleIdx);
+				for (int32 i = 0; i < MuscleActivationData.MuscleActivationElement.Num(); i++)
+				{
+					MuscleActivationData.MuscleActivationElement[i] += ElementOffset;
+				}
+				MuscleActivationElement.Add(MuscleActivationData.MuscleActivationElement);
+				MuscleActivationData.OriginInsertionPair[0] += VertexOffset;
+				MuscleActivationData.OriginInsertionPair[1] += VertexOffset;
+				OriginInsertionPair.Add(MuscleActivationData.OriginInsertionPair);
+				OriginInsertionRestLength.Add(MuscleActivationData.OriginInsertionRestLength);
+				Activation.Add(0.f);
+				FiberDirectionMatrix.Add(MuscleActivationData.FiberDirectionMatrix);
+			}
+		}
+
+		void UpdateLengthBasedMuscleActivation(const ParticleType& InParticles)
+		{
+			float AverageFiberLengthScale = 0;
+			for (int32 MuscleIdx = 0; MuscleIdx < MuscleActivationElement.Num(); MuscleIdx++)
+			{
+				if (ensureMsgf(MuscleActivationElement[MuscleIdx].Num() == FiberDirectionMatrix[MuscleIdx].Num(),
+					TEXT("MuscleActivationElement[%d].Num() = %d, not equal to FiberDirectionMatrix[%d].Num() = %d"), MuscleIdx, MuscleActivationElement[MuscleIdx].Num(), MuscleIdx, FiberDirectionMatrix[MuscleIdx].Num()))
+				{
+					//Infer activation
+					float CurrentFiberLengthScale = (InParticles.X(OriginInsertionPair[MuscleIdx][0]) - InParticles.X(OriginInsertionPair[MuscleIdx][1])).Size() / OriginInsertionRestLength[MuscleIdx];
+					AverageFiberLengthScale += CurrentFiberLengthScale;
+					Activation[MuscleIdx] = CurrentFiberLengthScale * CurrentFiberLengthScale;
+				}
+			}
+		}
+
+		void ApplyMuscleActivation(FXPBDCorotatedConstraints<T,ParticleType>& Constraints) const
+		{
+			for (int32 MuscleIdx = 0; MuscleIdx < MuscleActivationElement.Num(); MuscleIdx++)
+			{
+				if (ensureMsgf(MuscleActivationElement[MuscleIdx].Num() == FiberDirectionMatrix[MuscleIdx].Num(), 
+					TEXT("MuscleActivationElement[%d].Num() = %d, not equal to FiberDirectionMatrix[%d].Num() = %d"), MuscleIdx, MuscleActivationElement[MuscleIdx].Num(), MuscleIdx, FiberDirectionMatrix[MuscleIdx].Num()))
+				{
+					for (int32 ElemIdx = 0; ElemIdx < MuscleActivationElement[MuscleIdx].Num(); ElemIdx++)
+					{
+						Constraints.ModifyDmInverseFromFiberLength(MuscleActivationElement[MuscleIdx][ElemIdx], Activation[MuscleIdx], FiberDirectionMatrix[MuscleIdx][ElemIdx]);
+					}
+				}
+			}
+		}
+		
+		TArray<TArray<int32>> MuscleActivationElement;
+		TArray<FIntVector2> OriginInsertionPair;
+		TArray<float> OriginInsertionRestLength;
+		TArray<float> Activation;
+		TArray<TArray<Chaos::PMatrix33d>> FiberDirectionMatrix;
+	};
+
+
+}// End namespace Chaos::Softs
diff --git a/Engine/Source/Runtime/Experimental/Chaos/Public/Chaos/TriangleMesh.h b/Engine/Source/Runtime/Experimental/Chaos/Public/Chaos/TriangleMesh.h
--- a/Engine/Source/Runtime/Experimental/Chaos/Public/Chaos/TriangleMesh.h
+++ b/Engine/Source/Runtime/Experimental/Chaos/Public/Chaos/TriangleMesh.h
@@ -324,6 +324,10 @@ namespace Chaos
 		bool PointProximityQuery(const TSpatialHashType<FRealSingle>& SpatialHash, const TConstArrayView<TVec3<FRealSingle>>& Points, const int32 PointIndex, const TVec3<FRealSingle>& PointPosition, const FRealSingle PointThickness, const Softs::FPBDFlatWeightMap& ThisThicknesses,
 			const FRealSingle ThisThicknessExtraMultiplier, int32 ThicknessMapIndexOffset, TFunctionRef<bool(const int32 PointIndex, const int32 TriangleIndex)> BroadphaseTest, TArray<TTriangleCollisionPoint<FRealSingle>>& Result) const;
 
+		template<typename T>
+		bool PointClosestTriangleQuery(const TSpatialHashType<T>& SpatialHash, const TConstArrayView<TVec3<T>>& Points, const int32 PointIndex, const TVec3<T>& PointPosition, const T PointThickness, const T ThisThickness,
+			TFunctionRef<bool(const int32 PointIndex, const int32 TriangleIndex)> BroadphaseTest, TArray<TTriangleCollisionPoint<T>>& Result) const;
+
 		template<typename T>
 		bool EdgeIntersectionQuery(const TSpatialHashType<T>& SpatialHash, const TConstArrayView<TVec3<T>>& Points, const int32 EdgeIndex, const TVec3<T>& EdgePosition1, const TVec3<T>& EdgePosition2,
 			TFunctionRef<bool(const int32 EdgeIndex, const int32 TriangleIndex)> BroadphaseTest, TArray<TTriangleCollisionPoint<T>>& Result) const;
diff --git a/Engine/Source/Runtime/Experimental/Chaos/Public/Chaos/XPBDCorotatedConstraints.h b/Engine/Source/Runtime/Experimental/Chaos/Public/Chaos/XPBDCorotatedConstraints.h
--- a/Engine/Source/Runtime/Experimental/Chaos/Public/Chaos/XPBDCorotatedConstraints.h
+++ b/Engine/Source/Runtime/Experimental/Chaos/Public/Chaos/XPBDCorotatedConstraints.h
@@ -57,7 +57,7 @@ namespace Chaos::Softs
 				}
 
 			}
-
+			DmInverseSave = DmInverse;
 			InitColor(InParticles);
 		}
 
@@ -105,7 +105,7 @@ namespace Chaos::Softs
 					Measure[e] = -Measure[e];
 				}
 			}
-
+			DmInverseSave = DmInverse;
 			InitColor(InParticles);
 		}
 
@@ -163,6 +163,7 @@ namespace Chaos::Softs
 					Measure[e] = -Measure[e];
 				}
 			}
+			DmInverseSave = DmInverse;
 			if (bDoColoring) {
 				InitColor(InParticles);
 			}
@@ -199,6 +200,7 @@ namespace Chaos::Softs
 					Measure[e] = -Measure[e];
 				}
 			}
+			DmInverseSave = DmInverse;
 		}
 
 		virtual ~FXPBDCorotatedConstraints() {}
@@ -238,6 +240,16 @@ namespace Chaos::Softs
 			}
 			return DmInv;
 		}
+
+		PMatrix<T, 3, 3> ElementDmInvSave(const int e) const {
+			PMatrix<T, 3, 3> DmInv((T)0.);
+			for (int r = 0; r < 3; r++) {
+				for (int c = 0; c < 3; c++) {
+					DmInv.SetAt(r, c, DmInverseSave[(3 * 3) * e + 3 * r + c]);
+				}
+			}
+			return DmInv;
+		}
 		
 		virtual void Init() const 
 		{
@@ -394,6 +406,24 @@ namespace Chaos::Softs
 			return dC2;
 		}
 
+		void ModifyDmInverseFromFiberLength(const int32 e, const float L, const PMatrix<T, 3, 3>& MFiberDir) const
+		{
+			TRACE_CPUPROFILER_EVENT_SCOPE(TEXT("STAT_ChaosXPBDCorotatedModifyDmInverseFromFiberLength"));
+			PMatrix<T, 3, 3> DmInv = ElementDmInvSave(e);
+			if (L < 1)
+			{
+				PMatrix<T, 3, 3> D(1 / L, sqrt(L), sqrt(L));
+				PMatrix<T, 3, 3> Factor = MFiberDir * D * MFiberDir.GetTransposed();
+				DmInv = Factor * DmInv; //Matrix multiplication convention: AB = B*A
+			}
+			for (int r = 0; r < 3; r++) 
+			{
+				for (int c = 0; c < 3; c++) 
+				{
+					DmInverse[(3 * 3) * e + 3 * r + c] = DmInv.GetAt(r, c);
+				}
+			}
+		}
 
 	protected:
 
@@ -587,7 +617,7 @@ namespace Chaos::Softs
 	protected:
 		mutable TArray<T> LambdaArray;
 		mutable TArray<T> DmInverse;
-
+		TArray<T> DmInverseSave;
 		//parallel data:
 		FDeformableXPBDCorotatedParams CorotatedParams;
 
diff --git a/Engine/Source/Runtime/Experimental/Chaos/Public/GeometryCollection/Facades/CollectionMuscleActivationFacade.h b/Engine/Source/Runtime/Experimental/Chaos/Public/GeometryCollection/Facades/CollectionMuscleActivationFacade.h
new file mode 100644
--- /dev/null
+++ b/Engine/Source/Runtime/Experimental/Chaos/Public/GeometryCollection/Facades/CollectionMuscleActivationFacade.h
@@ -0,0 +1,59 @@
+// Copyright Epic Games, Inc. All Rights Reserved.
+
+#pragma once
+
+#include "CoreMinimal.h"
+#include "GeometryCollection/ManagedArrayAccessor.h"
+#include "GeometryCollection/ManagedArrayCollection.h"
+#include "GeometryCollection/Facades/CollectionSelectionFacade.h"
+
+namespace GeometryCollection::Facades
+{
+	struct FMuscleActivationData
+	{
+		TArray<int32> MuscleActivationElement;
+		FIntVector2 OriginInsertionPair;
+		float OriginInsertionRestLength;
+		TArray<Chaos::PMatrix33d> FiberDirectionMatrix;
+	};
+
+	/** Kinematic Facade */
+	class FMuscleActivationFacade
+	{
+	public:
+
+		static CHAOS_API const FName GroupName;
+		static CHAOS_API const FName MuscleActivationElement;
+		static CHAOS_API const FName OriginInsertionPair;
+		static CHAOS_API const FName OriginInsertionRestLength;
+		static CHAOS_API const FName FiberDirectionMatrix;
+
+		CHAOS_API FMuscleActivationFacade(FManagedArrayCollection& InCollection);
+		CHAOS_API FMuscleActivationFacade(const FManagedArrayCollection& InCollection);
+
+		/** Create the facade attributes. */
+		CHAOS_API void DefineSchema();
+
+		/** Is the facade defined constant. */
+		bool IsConst() const { return Collection == nullptr; }
+
+		/** Is the Facade defined on the collection? */
+		CHAOS_API bool IsValid() const;
+
+		//
+		//  Skeletal Mesh Bone Bindings
+		//
+		CHAOS_API int32 AddMuscleActivationData(const FMuscleActivationData& InputData);
+		CHAOS_API const FMuscleActivationData GetMuscleActivationData(const int32 DataIndex) const;
+		int32 NumMuscles() const { return MuscleActivationElementAttribute.Num(); }
+
+	private:
+		const FManagedArrayCollection& ConstCollection;
+		FManagedArrayCollection* Collection = nullptr;
+
+		TManagedArrayAccessor<TArray<int32>> MuscleActivationElementAttribute;
+		TManagedArrayAccessor<FIntVector2> OriginInsertionPairAttribute;
+		TManagedArrayAccessor<float> OriginInsertionRestLengthAttribute;
+		TManagedArrayAccessor<TArray<Chaos::PMatrix33d>> FiberDirectionMatrixAttribute;
+	};
+}
diff --git a/Engine/Source/Runtime/Experimental/Chaos/Public/GeometryCollection/ManagedArray.h b/Engine/Source/Runtime/Experimental/Chaos/Public/GeometryCollection/ManagedArray.h
--- a/Engine/Source/Runtime/Experimental/Chaos/Public/GeometryCollection/ManagedArray.h
+++ b/Engine/Source/Runtime/Experimental/Chaos/Public/GeometryCollection/ManagedArray.h
@@ -11,6 +11,7 @@
 #include "Chaos/ParticleHandle.h"
 #include "Chaos/BVHParticles.h"
 #include "Math/Vector.h"
+#include "Chaos/Matrix.h"
 #include "Templates/TypeHash.h"
 
 struct FManagedArrayCollection;
diff --git a/Engine/Source/Runtime/Experimental/Chaos/Public/GeometryCollection/ManagedArrayTypeValues.inl b/Engine/Source/Runtime/Experimental/Chaos/Public/GeometryCollection/ManagedArrayTypeValues.inl
--- a/Engine/Source/Runtime/Experimental/Chaos/Public/GeometryCollection/ManagedArrayTypeValues.inl
+++ b/Engine/Source/Runtime/Experimental/Chaos/Public/GeometryCollection/ManagedArrayTypeValues.inl
@@ -56,6 +56,8 @@ MANAGED_ARRAY_TYPE(Chaos::FConvexPtr, FConvexRefCountedPtr)
 MANAGED_ARRAY_TYPE(FTransform3f, Transform3f)
 MANAGED_ARRAY_TYPE(TArray<FIntVector3>, IntVector3Array)
 MANAGED_ARRAY_TYPE(TArray<FVector4f>, Vector4fArray)
+MANAGED_ARRAY_TYPE(Chaos::PMatrix33d, PMatrix33d)
+MANAGED_ARRAY_TYPE(TArray<Chaos::PMatrix33d>, PMatrix33dArray)
 
 // NOTE: new types must be added at the bottom to keep serialization from breaking
 
diff --git a/Engine/Source/Runtime/Experimental/Chaos/Public/GeometryCollection/ManagedArrayTypes.h b/Engine/Source/Runtime/Experimental/Chaos/Public/GeometryCollection/ManagedArrayTypes.h
--- a/Engine/Source/Runtime/Experimental/Chaos/Public/GeometryCollection/ManagedArrayTypes.h
+++ b/Engine/Source/Runtime/Experimental/Chaos/Public/GeometryCollection/ManagedArrayTypes.h
@@ -17,6 +17,7 @@
 #include "Chaos/ParticleHandle.h"
 #include "Chaos/ParticleHandleFwd.h"
 #include "Chaos/Convex.h"
+#include "Chaos/Matrix.h"
 
 inline FArchive& operator<<(FArchive& Ar, TArray<FVector3f>*& ValueIn)
 {
diff --git a/Engine/Source/Runtime/Experimental/ChaosCore/Public/Chaos/Matrix.h b/Engine/Source/Runtime/Experimental/ChaosCore/Public/Chaos/Matrix.h
--- a/Engine/Source/Runtime/Experimental/ChaosCore/Public/Chaos/Matrix.h
+++ b/Engine/Source/Runtime/Experimental/ChaosCore/Public/Chaos/Matrix.h
@@ -1255,4 +1255,6 @@ namespace Chaos
 		static const PMatrix<FRealSingle, 3, 3> Zero;
 		static const PMatrix<FRealSingle, 3, 3> Identity;
 	};
+
+	typedef PMatrix<double, 3, 3> PMatrix33d;
 }
\ No newline at end of file
