Chaos Cloth Solver: fix bug where we could potentially miss some self collisions due to not doing a compare and swap in the hash map code.

diff --git a/Engine/Source/Runtime/Experimental/Chaos/Public/Chaos/HierarchicalSpatialHash.h b/Engine/Source/Runtime/Experimental/Chaos/Public/Chaos/HierarchicalSpatialHash.h
--- a/Engine/Source/Runtime/Experimental/Chaos/Public/Chaos/HierarchicalSpatialHash.h
+++ b/Engine/Source/Runtime/Experimental/Chaos/Public/Chaos/HierarchicalSpatialHash.h
@@ -147,21 +147,25 @@ private:
 		{
 			const int32 NumElementsToAdd = (MaxCellIdx[0] - MinCellIdx[0] + 1) * (MaxCellIdx[1] - MinCellIdx[1] + 1) * (MaxCellIdx[2] - MinCellIdx[2] + 1);
 
-			int32 IndexToWrite = ConcurrentElementAddIdx.fetch_add(NumElementsToAdd, std::memory_order_relaxed);
-
-			for (int32 XIdx = MinCellIdx[0]; XIdx <= MaxCellIdx[0]; ++XIdx)
+			if (ensureAlways(NumElementsToAdd <= 8))
 			{
-				for (int32 YIdx = MinCellIdx[1]; YIdx <= MaxCellIdx[1]; ++YIdx)
+				int32 IndexToWrite = ConcurrentElementAddIdx.fetch_add(NumElementsToAdd, std::memory_order_relaxed);
+
+				for (int32 XIdx = MinCellIdx[0]; XIdx <= MaxCellIdx[0]; ++XIdx)
 				{
-					for (int32 ZIdx = MinCellIdx[2]; ZIdx <= MaxCellIdx[2]; ++ZIdx)
+					for (int32 YIdx = MinCellIdx[1]; YIdx <= MaxCellIdx[1]; ++YIdx)
 					{
-						FElement Element;
-						Element.Key = FHashIndex(XIdx, YIdx, ZIdx, Lod);
-						const uint32 HashKey = Element.Key.GetTypeHash() & (Hash.Num() - 1);
-						Element.NextIndex = Hash[HashKey].exchange(IndexToWrite);
-						Element.Value = Value;
-						Elements[IndexToWrite] = MoveTemp(Element);
-						++IndexToWrite;
+						for (int32 ZIdx = MinCellIdx[2]; ZIdx <= MaxCellIdx[2]; ++ZIdx)
+						{
+							FElement Element;
+							Element.Key = FHashIndex(XIdx, YIdx, ZIdx, Lod);
+							const uint32 HashKey = Element.Key.GetTypeHash() & (Hash.Num() - 1);
+							Element.NextIndex = Hash[HashKey].load(std::memory_order_relaxed);
+							while (!Hash[HashKey].compare_exchange_weak(Element.NextIndex, IndexToWrite, std::memory_order_release, std::memory_order_relaxed));
+							Element.Value = Value;
+							Elements[IndexToWrite] = MoveTemp(Element);
+							++IndexToWrite;
+						}
 					}
 				}
 			}
@@ -173,7 +177,8 @@ private:
 			FElement Element;
 			Element.Key = FHashIndex( CellIdx, Lod );
 			const uint32 HashKey = Element.Key.GetTypeHash() & (Hash.Num() - 1);
-			Element.NextIndex = Hash[HashKey].exchange(IndexToWrite);
+			Element.NextIndex = Hash[HashKey].load(std::memory_order_relaxed);
+			while (!Hash[HashKey].compare_exchange_weak(Element.NextIndex, IndexToWrite, std::memory_order_release, std::memory_order_relaxed));
 			Element.Value = Value;
 			Elements[IndexToWrite] = MoveTemp(Element);
 
