Stale data could cause IAS cache to load incorrectly.

When the IAS cache loads, the journal file is compared against the data file to get them back in sync (the former is generally written at a lower frequency than the latter - they can have different views of the state). If the data file cursor was nearing the maximum data file size it would reset the cursor back to zero. This would leave older data at the end of the data file that the journal may still reference. Said false reference would cause cache startup to select the wrong range of the journal to add as items. Now data is written right to the end of the data file so age of bytes in the file always progresses forward.

#rnx

diff --git a/Engine/Source/Runtime/Experimental/IoStoreOnDemand/Private/JournaledCache.cpp b/Engine/Source/Runtime/Experimental/IoStoreOnDemand/Private/JournaledCache.cpp
--- a/Engine/Source/Runtime/Experimental/IoStoreOnDemand/Private/JournaledCache.cpp
+++ b/Engine/Source/Runtime/Experimental/IoStoreOnDemand/Private/JournaledCache.cpp
@@ -117,7 +117,7 @@ public:
 	uint32			GetMax() const		{ return MaxSize; }
 	const FIoBuffer*Get(uint64 Key) const;
 	bool			Put(uint64 Key, FIoBuffer&& Data);
-	int32			Peel(int32 TargetPeelSize, PeelItems& Out, int32 MaxPeelSize);
+	int32			Peel(int32 PeelThreshold, PeelItems& Out);
 	uint32			DebugVisit(void* Param, FDebugCacheEntry::Callback* Callback);
 
 private:
@@ -186,7 +186,7 @@ bool FMemCache::Put(uint64 Key, FIoBuffer&& Data)
 }
 
 ////////////////////////////////////////////////////////////////////////////////
-int32 FMemCache::Peel(int32 TargetPeelSize, PeelItems& Out, int32 MaxPeelSize)
+int32 FMemCache::Peel(int32 PeelThreshold, PeelItems& Out)
 {
 	TRACE_CPUPROFILER_EVENT_SCOPE(IasCache::Peel);
 
@@ -197,22 +197,14 @@ int32 FMemCache::Peel(int32 TargetPeelSize, PeelItems& Out, int32 MaxPeelSize)
 	// Add large items
 	int32 NumItems = Items.Num();
 	int32 DropSize = 0;
-	for (int32 i=Items.Num() - 1; i >= 0 && DropSize < TargetPeelSize; --i)
+	for (int32 i=Items.Num() - 1; i >= 0 && DropSize < PeelThreshold; --i)
 	{
 		FItem& Item = Items[i];
 
-		int32 NextDropSize = DropSize + int32(Item.Data.GetSize());
-		
-		if (NextDropSize > MaxPeelSize)
-		{
-			continue;
-		}
+		DropSize += int32(Item.Data.GetSize());
 
 		Out.Add(MoveTemp(Item));
-		Items.Swap(i, NumItems - 1);
 		--NumItems;
-
-		DropSize = NextDropSize;
 	}
 
 	Items.SetNum(NumItems);
@@ -276,7 +268,7 @@ int32 FMemCache::Drop(uint32 Size)
 // {{{1 phrase .................................................................
 
 ////////////////////////////////////////////////////////////////////////////////
-static const uint32 MAGIC = 0x04930002;
+static const uint32 MAGIC = 0x04930003;
 static const uint32 SIZE_BITS = 25;
 static const uint32 MARKER_MAX = 0x3fffffff;
 static const uint32	HASH_CHECKSUM_SIZE = 64;
@@ -387,6 +379,7 @@ public:
 	uint32					GetMaxSize() const	{ return MaxSize; }
 	uint32					GetCursor() const	{ return Cursor; }
 	uint32					GetMarker() const	{ return Marker; }
+	void					ClearPartial()		{ PreviousPartial = 0; }
 	uint32					GetPreviousPartial() const { return PreviousPartial; };
 
 private:
@@ -532,6 +525,7 @@ void FDiskJournal::ClosePhrase(FDiskPhrase&& Phrase, uint64 DataCursor)
 	// Potentially adjust the data cursor to the start of the previous
 	// partial write.
 	Desc.DataCursor = DataCursor - PreviousPartial;
+	check(DataCursor >= PreviousPartial);
 
 	// If we have already written a partial fragment of this phrase
 	// we use the hash of the start
@@ -613,8 +607,7 @@ public:
 	EIoErrorCode			Materialize(uint64 Key, FIoBuffer& Out, uint32 Offset=0) const;
 	int32					Flush();
 	void					Drop();
-	void					Wrap();
-	uint64					RemainingUntilWrap() const;
+	uint64					RemainingUntilWrap();
 	uint32					DebugVisit(void* Param, FDebugCacheEntry::Callback* Callback);
 
 private:
@@ -629,6 +622,7 @@ private:
 	friend int32			LoadCache(FDiskCache&);
 	void					OpenDataFile();
 	using					FDataMap = TMap<uint64, FMapEntry>;
+	void					Wrap();
 	void					Spam();
 	uint64					Insert(uint64 DataBase, const FDataEntry& Entry);
 	uint64					Insert(uint64 DataBase, const FDataEntry* Entries, uint32 EntryCount);
@@ -687,8 +681,10 @@ FDiskPhrase FDiskCache::OpenPhrase(uint32 DataSize)
 ////////////////////////////////////////////////////////////////////////////////
 void FDiskCache::Wrap()
 {
+	check(DataCursor >= MaxDataSize - HASH_CHECKSUM_SIZE);
 	OverRemoval = 0;
 	DataCursor = 0;
+	Journal.ClearPartial();
 }
 
 ////////////////////////////////////////////////////////////////////////////////
@@ -711,10 +707,7 @@ void FDiskCache::ClosePhrase(FDiskPhrase&& Phrase)
 	}
 
 	uint32 WriteSize = Phrase.GetWriteSize();
-	if (DataCursor + WriteSize > MaxDataSize)
-	{
-		Wrap();
-	}
+	check(DataCursor + WriteSize <= MaxDataSize);
 
 	bool bWriteOk;
 	{
@@ -893,8 +886,13 @@ void FDiskCache::Drop()
 }
 
 ////////////////////////////////////////////////////////////////////////////////
-uint64 FDiskCache::RemainingUntilWrap() const
+uint64 FDiskCache::RemainingUntilWrap()
 {
+	if (DataCursor >= MaxDataSize - HASH_CHECKSUM_SIZE)
+	{
+		Wrap();
+	}
+
 	return MaxDataSize - DataCursor;
 }
 
@@ -1361,7 +1359,13 @@ uint32 FCache::WriteMemToDisk(int32 Allowance)
 {
 	TRACE_CPUPROFILER_EVENT_SCOPE(IasCache::Flush_MemCache);
 
-	bool bCloseToWrap = false;
+	bool bEof = false;
+	if (int64 UntilWrap = DiskCache.RemainingUntilWrap(); UntilWrap <= int64(Allowance))
+	{
+		bEof = true;
+		Allowance = int32(UntilWrap);
+	}
+
 	FMemCache::PeelItems PeelItems;
 	int32 WriteSize = 0;
 	{
@@ -1382,17 +1386,16 @@ uint32 FCache::WriteMemToDisk(int32 Allowance)
 			{
 				Partial.Reset();
 			}
+			else if (bEof)
+			{
+				Partial->Remaining = uint32(Partial->Data.GetSize());
+			}
 		}
 		
-		// We don't want to deal with wrapping partials around the end of the buffer
-		const int32 UntilWrap = int32(DiskCache.RemainingUntilWrap()) - WriteSize;
-		check(UntilWrap >= 0);
-		bCloseToWrap = UntilWrap <= Allowance;
-
 		// If there is any allowance left start peeling of buffers from the memcache
 		if (WriteSize < Allowance)
 		{
-			WriteSize += MemCache.Peel(Allowance - WriteSize, PeelItems, UntilWrap);
+			WriteSize += MemCache.Peel(Allowance - WriteSize, PeelItems);
 		}
 
 		// Finally split any overshooting buffers into a partial slice and save the
@@ -1402,7 +1405,6 @@ uint32 FCache::WriteMemToDisk(int32 Allowance)
 		if (WriteSize > Allowance)
 		{
 			auto [Key, Data] = PeelItems.Pop();
-			check(Data.GetSize() <= UntilWrap);
 			
 			int32 RemainderSize = WriteSize - Allowance;
 			const int32 PartialSize = int32(Data.GetSize()) - RemainderSize;
@@ -1415,8 +1417,9 @@ uint32 FCache::WriteMemToDisk(int32 Allowance)
 				const FMemoryView PartialSlice = Data.GetView().Left(PartialSize);
 				PeelItems.Push(FMemCache::FItem{0, FIoBuffer(PartialSlice, Data)});
 			}
-			Partial.Emplace(FPartialItem{Key, MoveTemp(Data), uint32(RemainderSize)});
 			WriteSize -= RemainderSize;
+			RemainderSize = bEof ? int32(Data.GetSize()) : RemainderSize;
+			Partial.Emplace(FPartialItem{Key, MoveTemp(Data), uint32(RemainderSize)});
 		}
 
 		uint32 NewDemand = MemCache.GetDemand();
@@ -1446,16 +1449,6 @@ uint32 FCache::WriteMemToDisk(int32 Allowance)
 
 	DiskCache.ClosePhrase(MoveTemp(Phrase));
 
-	// If we are close to the end of the file and cannot find any suitable
-	// buffers to peel off, reset file to beginning. There should be no partials
-	// at this point, but release it if any exists (it will be lost)
-	if (bCloseToWrap && PeelItems.IsEmpty())
-	{
-		check(!Partial.IsSet());
-		Partial.Reset();
-		DiskCache.Wrap();
-	}
-
 	return WriteSize;
 }
 
@@ -2432,9 +2425,9 @@ static void MemCacheTests(FSupport& Support)
 
 	FMemCache MemCache(64);
 	MemCache.Put(1, Support.DummyData(1));
-	check(MemCache.Peel(0, Peeled, 0) == 0);
+	check(MemCache.Peel(0, Peeled) == 0);
 	check(Peeled.Num() == 0);
-	check(MemCache.Peel(64, Peeled, 64) == 1);
+	check(MemCache.Peel(64, Peeled) == 1);
 	check(Peeled.Num() == 1);
 	check(MemCache.GetUsed() == 0);
 	Peeled.Reset();
@@ -2445,7 +2438,7 @@ static void MemCacheTests(FSupport& Support)
 		MemCache.Put(i + 1, Support.DummyData(1));
 	}
 
-	check(MemCache.Peel(32, Peeled, 32) == 32);
+	check(MemCache.Peel(32, Peeled) == 32);
 	check(Peeled.Num() == 32);
 	check(MemCache.GetUsed() == 32);
 	for (auto& [Key, _] : Peeled)
