Updated IAS journal tests

#rnx

diff --git a/Engine/Source/Runtime/Experimental/IoStoreOnDemand/Private/JournaledCache.cpp b/Engine/Source/Runtime/Experimental/IoStoreOnDemand/Private/JournaledCache.cpp
--- a/Engine/Source/Runtime/Experimental/IoStoreOnDemand/Private/JournaledCache.cpp
+++ b/Engine/Source/Runtime/Experimental/IoStoreOnDemand/Private/JournaledCache.cpp
@@ -2540,30 +2540,36 @@ static void CacheTests(FSupport& Support)
 	check(Cache->WriteMemToDisk(WriteAllowance) > 0);
 	check(Cache->Flush() > 0);
 
-#if 0
 	auto Validate = [&] () {
 		struct FVisitState {
 			FCache& Cache;
 			const uint8* WorkRange[2];
 		};
 		auto Visitor = [] (void* Param, const FDebugCacheEntry& Entry) {
-			auto* State = (FVisitState*)Param;
+			auto& State = *(FVisitState*)Param;
 			
-			auto GetEntry = State->Cache.Get(Entry.Key);
-			check(GetEntry.IsHit());
 			FIoBuffer Data;
-			check(GetEntry.Materialize(Data));
+			auto Token = State.Cache.Get(Entry.Key, Data);
+			if (const uint8* Ptr = Data.GetData(); Ptr == nullptr)
+			{
+				check(Token == Entry.Key);
+				check(State.Cache.Materialize(Token, Data, 0) == EIoErrorCode::Ok);
+			}
+
 			check(Data.GetSize() == Entry.Size);
 			check(KeyGen(Data) == Entry.Key);
 
+			const uint8* Ptr = Data.GetData();
 			int32 IsFromDisk = 0;
-			IsFromDisk |= Data.GetData() >= State->WorkRange[1];
-			IsFromDisk |= (Data.GetData() + Data.GetSize()) <= State->WorkRange[0];
+			IsFromDisk |= Ptr >= State.WorkRange[1];
+			IsFromDisk |= (Ptr + Data.GetSize()) <= State.WorkRange[0];
 			check(Entry.IsMemCache != IsFromDisk);
 		};
-		FVisitState State = {Cache, {Working, Working + WorkingSize}};
-		return Cache.DebugVisit(&State, Visitor);
+		FVisitState State = {*Cache, {Support.Working, Support.Working + Support.WorkingSize}};
+		return Cache->DebugVisit(&State, Visitor);
 	};
+	
+	NewCache();
 	check(Validate() == 0);
 
 	WriteAllowance = uint32(512_Ki);
@@ -2574,36 +2580,39 @@ static void CacheTests(FSupport& Support)
 		for (uint32 j = 0; j < i; ++j)
 		{
 			FIoBuffer Data = Support.DummyData(32);
-			Cache.Put(KeyGen(Data), Data);
+			Cache->Put(KeyGen(Data), Data);
 		}
-		Cache.Flush(WriteAllowance);
+		Cache->WriteMemToDisk(WriteAllowance);
 		check(Validate() == i);
 
-		Cache.Reset();
-		check(Cache.Load());
-		check(Validate() == 0); // not enough flushes to write a journal
+		NewCache(false);
+		check(Cache->Load());
+		check(Validate() == 0); // no flushes to write a journal
 	}
-	Cache.Reset();
 
 	// general
 	for (int32 i : {1, 4, 136, 137})
 	{
+		enum { Jfi = 4 };
 		for (int32 j = 0; j < i; ++j)
 		{
 			PrimePuts(WriteAllowance);
-			Cache.Flush(WriteAllowance);
+			Cache->WriteMemToDisk(WriteAllowance);
+			if (j % Jfi)
+			{
+				Cache->Flush();
+			}
 		}
 		uint32 PreCount = Validate();
 
-		Cache.Reset();
-		check(Cache.Load());
+		NewCache();
+		check(Cache->Load());
 
 		uint32 PostCount = Validate();
-		check(!PostCount == !(i / 4)); // JournalFlushInterval
+		check(!PostCount == !(i / Jfi));
 		check(PostCount <= PreCount);
-		check((PostCount == PreCount) == ((i & 3) == 0));
-	}
-	Cache.Reset();
+ 	}
+	NewCache();
 
 	// power 2
 	for (int32 i : {74, 75})
@@ -2612,16 +2621,17 @@ static void CacheTests(FSupport& Support)
 		{
 			for (int32 j = 0; j < 3; ++j)
 			{
-				FIoBuffer Data = Support.DummyData(64_Ki - 4);
-				Cache.Put(KeyGen(Data), Data);
-				Cache.Flush(WriteAllowance);
+				FIoBuffer Data = Support.DummyData(64_Ki);
+				Cache->Put(KeyGen(Data), Data);
+				Cache->WriteMemToDisk(WriteAllowance);
 			}
-			Cache.Flush(WriteAllowance);
+			Cache->Flush();
 		}
 		Validate();
-		Cache.Reset();
+		NewCache();
 	}
 
+#if 0
 	// marker wrap
 	// one-phrase journal
 	// journal paragraphs that are all the same size
