Intersection tests would occasionally include four bytes too much

One upon a time a four byte identifier was prefixed to a "phrase" of data (a bundle of cache items). This identifier was used to detect discontinuities between the journal and the data as they are not written to attomically. Now we hash the start of a "phrase" so the four byte adjustment for the first item is not applicable.

#rnx

diff --git a/Engine/Source/Runtime/Experimental/IoStoreOnDemand/Private/JournaledCache.cpp b/Engine/Source/Runtime/Experimental/IoStoreOnDemand/Private/JournaledCache.cpp
--- a/Engine/Source/Runtime/Experimental/IoStoreOnDemand/Private/JournaledCache.cpp
+++ b/Engine/Source/Runtime/Experimental/IoStoreOnDemand/Private/JournaledCache.cpp
@@ -613,9 +613,8 @@ public:
 private:
 	struct FMapEntry
 	{
-		uint64				DataCursor : 38;
+		uint64				DataCursor : 39;
 		uint64				Size : SIZE_BITS;
-		uint64				First : 1;
 	};
 	static_assert(sizeof(FMapEntry) == sizeof(uint64));
 
@@ -794,7 +793,6 @@ uint64 FDiskCache::Insert(uint64 DataBase, const FDataEntry& Entry)
 	Value.DataCursor = DataBase + Entry.Offset;
 	check(Value.DataCursor < MaxDataSize);
 	Value.Size = Entry.Size;
-	Value.First = (Entry.Offset == 0);
 	DataMap.Add(Entry.Key, Value);
 	return Entry.Size;
 }
@@ -834,7 +832,7 @@ void FDiskCache::Prune(uint64 DataBase, uint32 Size)
 		const FMapEntry& Candidate = Iter.Value();
 
 		int64 B[] = {
-			int64(Candidate.DataCursor) - (int64(Candidate.First) << 2),
+			int64(Candidate.DataCursor),
 			int64(Candidate.DataCursor + Candidate.Size)
 		};
 
