Converted UGeometryCollection rest collections to shared ownership as the physics proxy now relies on the rest collection for immutable data.

- Proxy lifetime can exceed the lifetime of the UGeoemtryCollection so previously we can have the rest collection destroyed out from under the proxy causing it to crash

diff --git a/Engine/Plugins/Experimental/ChaosCaching/Source/ChaosCaching/Private/Chaos/Adapters/GeometryCollectionComponentCacheAdapter.cpp b/Engine/Plugins/Experimental/ChaosCaching/Source/ChaosCaching/Private/Chaos/Adapters/GeometryCollectionComponentCacheAdapter.cpp
--- a/Engine/Plugins/Experimental/ChaosCaching/Source/ChaosCaching/Private/Chaos/Adapters/GeometryCollectionComponentCacheAdapter.cpp
+++ b/Engine/Plugins/Experimental/ChaosCaching/Source/ChaosCaching/Private/Chaos/Adapters/GeometryCollectionComponentCacheAdapter.cpp
@@ -72,8 +72,8 @@ namespace Chaos
 
 		const FCachedEventData& ProxyCachedEventData = CachedData[Proxy];
 
-		const Chaos::FPhysicsSolver* Solver         = Proxy->GetSolver<Chaos::FPhysicsSolver>();
-		const FGeometryCollection*   RestCollection = Proxy->GetSimParameters().RestCollection;
+		const Chaos::FPhysicsSolver* Solver = Proxy->GetSolver<Chaos::FPhysicsSolver>();
+		TSharedPtr<const FGeometryCollection> RestCollection = Proxy->GetSimParameters().RestCollectionShared;
 
 		if(!RestCollection || !Solver)
 		{
@@ -253,8 +253,8 @@ namespace Chaos
 		}
 
 		const FTransform ComponentToWorld = Proxy->GetSimParameters().WorldTransform;
-		const FGeometryCollection* RestCollection = Proxy->GetSimParameters().RestCollection;
-		Chaos::FPhysicsSolver*     Solver         = Proxy->GetSolver<Chaos::FPhysicsSolver>();
+		TSharedPtr<const FGeometryCollection> RestCollection = Proxy->GetSimParameters().RestCollectionShared;
+		Chaos::FPhysicsSolver* Solver = Proxy->GetSolver<Chaos::FPhysicsSolver>();
 
 		if(!RestCollection || !Solver)
 		{
diff --git a/Engine/Source/Runtime/Experimental/Chaos/Private/Chaos/StrainModification.cpp b/Engine/Source/Runtime/Experimental/Chaos/Private/Chaos/StrainModification.cpp
--- a/Engine/Source/Runtime/Experimental/Chaos/Private/Chaos/StrainModification.cpp
+++ b/Engine/Source/Runtime/Experimental/Chaos/Private/Chaos/StrainModification.cpp
@@ -205,7 +205,7 @@ const TSet<int32>* Chaos::FStrainedProxyModifier::InitRestChildren(FGeometryColl
 	if (RootIndex == INDEX_NONE) { return nullptr; }
 
 	// Return a ptr to the set of child indices
-	return &Parameters.RestCollection->Children[RootIndex];
+	return &Parameters.RestCollectionShared->Children[RootIndex];
 }
 
 Chaos::FStrainedProxyModifier Chaos::FStrainedProxyIterator::operator*()
diff --git a/Engine/Source/Runtime/Experimental/Chaos/Private/GeometryCollectionProxyData.cpp b/Engine/Source/Runtime/Experimental/Chaos/Private/GeometryCollectionProxyData.cpp
--- a/Engine/Source/Runtime/Experimental/Chaos/Private/GeometryCollectionProxyData.cpp
+++ b/Engine/Source/Runtime/Experimental/Chaos/Private/GeometryCollectionProxyData.cpp
@@ -14,10 +14,27 @@ GeometryCollectionProxyData.cpp:
 
 FTransformDynamicCollection::FTransformDynamicCollection(const FGeometryCollection* InRestCollection)
 	: FManagedArrayCollection()
+PRAGMA_DISABLE_DEPRECATION_WARNINGS
 	, RestCollection(InRestCollection)
+PRAGMA_ENABLE_DEPRECATION_WARNINGS
 	, bTransformHasChanged(false)
 {
+PRAGMA_DISABLE_DEPRECATION_WARNINGS
 	check(RestCollection != nullptr);
+PRAGMA_ENABLE_DEPRECATION_WARNINGS
+	Construct();
+}
+
+FTransformDynamicCollection::FTransformDynamicCollection(TSharedPtr<const FGeometryCollection> InRestCollection)
+	: FManagedArrayCollection()
+PRAGMA_DISABLE_DEPRECATION_WARNINGS
+	// To be removed with RestCollection post-deprecation
+	, RestCollection(InRestCollection.Get())
+PRAGMA_ENABLE_DEPRECATION_WARNINGS
+	, RestCollectionShared(InRestCollection)
+	, bTransformHasChanged(false)
+{
+	check(RestCollectionShared);
 	Construct();
 }
 
@@ -34,7 +51,7 @@ void FTransformDynamicCollection::InitializeTransforms()
 	if (bTransformHasChanged == false)
 	{
 		AddExternalAttribute<FTransform3f>(FTransformCollection::TransformAttribute, FTransformCollection::TransformGroup, Transform);
-		CopyAttribute(*RestCollection, FTransformCollection::TransformAttribute, FTransformCollection::TransformGroup);
+		CopyAttribute(*RestCollectionShared, FTransformCollection::TransformAttribute, FTransformCollection::TransformGroup);
 
 		bTransformHasChanged = true;
 	}
@@ -44,7 +61,7 @@ const FTransform3f& FTransformDynamicCollection::GetTransform(int32 Index) const
 {
 	if (bTransformHasChanged == false)
 	{
-		return RestCollection->Transform[Index];
+		return RestCollectionShared->Transform[Index];
 	}
 	return Transform[Index];
 }
@@ -57,8 +74,8 @@ void FTransformDynamicCollection::SetTransform(int32 Index, const FTransform3f&
 
 int32 FTransformDynamicCollection::GetNumTransforms() const
 {
-	ensure(!bTransformHasChanged || RestCollection->Transform.Num() == Transform.Num());
-	return RestCollection->Transform.Num();
+	ensure(!bTransformHasChanged || RestCollectionShared->Transform.Num() == Transform.Num());
+	return RestCollectionShared->Transform.Num();
 }
 
 void FTransformDynamicCollection::ResetInitialTransforms()
@@ -88,13 +105,13 @@ void FTransformDynamicCollection::SetHasParent(int32 Index, bool Value)
 
 int32 FTransformDynamicCollection::GetParent(int32 Index) const 
 {
-	check(RestCollection != nullptr);
-	return HasParent[Index] ? RestCollection->Parent[Index] : INDEX_NONE;
+	check(RestCollectionShared);
+	return HasParent[Index] ? RestCollectionShared->Parent[Index] : INDEX_NONE;
 }
 
 bool FTransformDynamicCollection::HasChildren(int32 Index) const
 {
-	const TSet<int32>& Children(RestCollection->Children[Index]);
+	const TSet<int32>& Children(RestCollectionShared->Children[Index]);
 	for (int32 Child : Children)
 	{
 		if (HasParent[Child])
@@ -125,6 +142,22 @@ const FName FGeometryDynamicCollection::CollisionMaskAttribute("CollisionMask");
 const FName FGeometryDynamicCollection::CollisionGroupAttribute("CollisionGroup");
 
 FGeometryDynamicCollection::FGeometryDynamicCollection(const FGeometryCollection* InRestCollection)
+PRAGMA_DISABLE_DEPRECATION_WARNINGS
+	: FTransformDynamicCollection(InRestCollection)
+PRAGMA_ENABLE_DEPRECATION_WARNINGS
+	, OptionalLinearVelocityAttribute(nullptr)
+	, OptionalAngularVelocityAttribute(nullptr)
+	, OptionalAnimateTransformAttribute(nullptr)
+{
+	// Transform Group
+	AddExternalAttribute<bool>(FGeometryDynamicCollection::ActiveAttribute, FTransformCollection::TransformGroup, Active);
+	AddExternalAttribute<uint8>(FGeometryDynamicCollection::DynamicStateAttribute, FTransformCollection::TransformGroup, DynamicState);
+	AddExternalAttribute(SimplicialsAttribute, FTransformCollection::TransformGroup, Simplicials);
+	AddExternalAttribute(SimulatableParticlesAttribute, FGeometryCollection::TransformGroup, SimulatableParticles);
+	AddExternalAttribute(InternalClusterParentTypeAttribute, FGeometryCollection::TransformGroup, InternalClusterParentType);
+}
+
+FGeometryDynamicCollection::FGeometryDynamicCollection(TSharedPtr<const FGeometryCollection> InRestCollection)
 	: FTransformDynamicCollection(InRestCollection)
 	, OptionalLinearVelocityAttribute(nullptr)
 	, OptionalAngularVelocityAttribute(nullptr)
@@ -141,7 +174,7 @@ FGeometryDynamicCollection::FGeometryDynamicCollection(const FGeometryCollection
 const TManagedArrayAccessor<int32> FGeometryDynamicCollection::GetInitialLevels() const
 {
 	static const FName LevelAttributeName = "Level";
-	return TManagedArrayAccessor<int32>(*RestCollection, LevelAttributeName, FGeometryCollection::TransformGroup);
+	return TManagedArrayAccessor<int32>(*RestCollectionShared, LevelAttributeName, FGeometryCollection::TransformGroup);
 }
 
 void FGeometryDynamicCollection::AddVelocitiesAttributes()
diff --git a/Engine/Source/Runtime/Experimental/Chaos/Private/PhysicsProxy/GeometryCollectionPhysicsProxy.cpp b/Engine/Source/Runtime/Experimental/Chaos/Private/PhysicsProxy/GeometryCollectionPhysicsProxy.cpp
--- a/Engine/Source/Runtime/Experimental/Chaos/Private/PhysicsProxy/GeometryCollectionPhysicsProxy.cpp
+++ b/Engine/Source/Runtime/Experimental/Chaos/Private/PhysicsProxy/GeometryCollectionPhysicsProxy.cpp
@@ -524,7 +524,7 @@ FGeometryCollectionPhysicsProxy::FGeometryCollectionPhysicsProxy(
 	, LastSyncCountGT(MAX_uint32)
 #endif
 	, CollisionParticlesPerObjectFraction(CollisionParticlesPerObjectFractionDefault)
-	, PhysicsThreadCollection(Parameters.RestCollection)
+	, PhysicsThreadCollection(Parameters.RestCollectionShared)
 	, GameThreadCollection(GameThreadCollectionIn)
 	, WorldTransform_External(FTransform::Identity)
 	, bIsGameThreadWorldTransformDirty(false)
@@ -570,7 +570,7 @@ void FGeometryCollectionPhysicsProxy::Initialize(Chaos::FPBDRigidsEvolutionBase
 	//
 	FGeometryDynamicCollection& DynamicCollection = GameThreadCollection;
 
-	InitializeDynamicCollection(DynamicCollection, *Parameters.RestCollection, Parameters);
+	InitializeDynamicCollection(DynamicCollection, *Parameters.RestCollectionShared, Parameters);
 
 	NumTransforms = DynamicCollection.NumElements(FGeometryCollection::TransformGroup);
 	BaseParticleIndex = 0; // Are we always zero indexed now?
@@ -594,7 +594,7 @@ void FGeometryCollectionPhysicsProxy::Initialize(Chaos::FPBDRigidsEvolutionBase
 		{
 			FromTransformToParticleIndex[Index] = RedirectIndex;
 			FromParticleToTransformIndex[RedirectIndex] = Index;
-			PhysicsObjects.Emplace(Chaos::FPhysicsObjectFactory::CreatePhysicsObject(this, Index, FName(Parameters.RestCollection->BoneName[Index])));
+			PhysicsObjects.Emplace(Chaos::FPhysicsObjectFactory::CreatePhysicsObject(this, Index, FName(Parameters.RestCollectionShared->BoneName[Index])));
 			RedirectIndex++;
 		}
 		else
@@ -636,7 +636,7 @@ void FGeometryCollectionPhysicsProxy::Initialize(Chaos::FPBDRigidsEvolutionBase
 	PhysicsThreadCollectionAnchoringFacade.CopyAnchoredAttribute(DynamicCollectionAnchoringFacade);
 
 	const FVector Scale = Parameters.WorldTransform.GetScale3D();
-	const TManagedArray<float>& Mass = Parameters.RestCollection->GetAttribute<float>(MassAttributeName, FTransformCollection::TransformGroup);
+	const TManagedArray<float>& Mass = Parameters.RestCollectionShared->GetAttribute<float>(MassAttributeName, FTransformCollection::TransformGroup);
 
 	TManagedArray<Chaos::FImplicitObjectPtr>& Implicits = GameThreadCollection.ModifyAttribute<Chaos::FImplicitObjectPtr>(FGeometryDynamicCollection::ImplicitsAttribute, FTransformCollection::TransformGroup);
 	if (ensure(NumTransforms == Implicits.Num() && NumEffectiveParticles == GTParticles.Num())) // Implicits are in the transform group so this invariant should always hold
@@ -719,9 +719,9 @@ void FGeometryCollectionPhysicsProxy::CreateGTParticles(TManagedArray<Chaos::FIm
 {
 	const Chaos::Facades::FCollectionAnchoringFacade DynamicCollectionAnchoringFacade(GameThreadCollection);
 	const FVector Scale = Parameters.WorldTransform.GetScale3D();
-	const TManagedArray<float>& Mass = Parameters.RestCollection->GetAttribute<float>(MassAttributeName, FTransformCollection::TransformGroup);
-	const TManagedArray<FTransform>& MassToLocal = Parameters.RestCollection->GetAttribute<FTransform>(MassToLocalAttributeName, FTransformCollection::TransformGroup);
-	const TManagedArray<int32>& Level = Parameters.RestCollection->GetAttribute<int32>(LevelAttributeName, FTransformCollection::TransformGroup);
+	const TManagedArray<float>& Mass = Parameters.RestCollectionShared->GetAttribute<float>(MassAttributeName, FTransformCollection::TransformGroup);
+	const TManagedArray<FTransform>& MassToLocal = Parameters.RestCollectionShared->GetAttribute<FTransform>(MassToLocalAttributeName, FTransformCollection::TransformGroup);
+	const TManagedArray<int32>& Level = Parameters.RestCollectionShared->GetAttribute<int32>(LevelAttributeName, FTransformCollection::TransformGroup);
 
 	TArray<int32> ChildrenToCheckForParentFix;
 	if (!bInitializationTime && !bCreateGTParticleForChildren && !bBuildGeometryForChildrenOnGT)
@@ -935,7 +935,7 @@ void FGeometryCollectionPhysicsProxy::CreateChildrenGeometry_External()
 			if (bRemoveImplicitsInDynamicCollections)
 			{
 				GameThreadCollection.AddAttribute<Chaos::FImplicitObjectPtr>(FGeometryDynamicCollection::ImplicitsAttribute, FTransformCollection::TransformGroup);
-				GameThreadCollection.CopyAttribute(*Parameters.RestCollection, FGeometryDynamicCollection::ImplicitsAttribute, FTransformCollection::TransformGroup);
+				GameThreadCollection.CopyAttribute(*Parameters.RestCollectionShared, FGeometryDynamicCollection::ImplicitsAttribute, FTransformCollection::TransformGroup);
 			}
 
 			TManagedArray<Chaos::FImplicitObjectPtr>& Implicits = GameThreadCollection.ModifyAttribute<Chaos::FImplicitObjectPtr>(FGeometryDynamicCollection::ImplicitsAttribute, FTransformCollection::TransformGroup);
@@ -1215,7 +1215,7 @@ void FGeometryCollectionPhysicsProxy::UpdateDamageThreshold_Internal()
 
 void FGeometryCollectionPhysicsProxy::InitializeBodiesPT(Chaos::FPBDRigidsSolver* RigidsSolver, typename Chaos::FPBDRigidsSolver::FParticlesType& Particles)
 {
-	const FGeometryCollection* RestCollection = Parameters.RestCollection;
+	const FGeometryCollection* RestCollection = Parameters.RestCollectionShared.Get();
 	const FGeometryDynamicCollection& DynamicCollection = PhysicsThreadCollection;
 
 	bIsInitializedOnPhysicsThread = true;
@@ -1698,14 +1698,14 @@ void FGeometryCollectionPhysicsProxy::CreateChildrenGeometry_Internal()
 {
 	if (!bHasBuiltGeometryOnPT)
 	{
-		const FGeometryCollection* RestCollection = Parameters.RestCollection;
+		const FGeometryCollection* RestCollection = Parameters.RestCollectionShared.Get();
 
 		if (Parameters.Simulating && ensure(RestCollection))
 		{
 			if (bRemoveImplicitsInDynamicCollections)
 			{
 				PhysicsThreadCollection.AddAttribute<Chaos::FImplicitObjectPtr>(FGeometryDynamicCollection::ImplicitsAttribute, FTransformCollection::TransformGroup);
-				PhysicsThreadCollection.CopyAttribute(*Parameters.RestCollection, FGeometryDynamicCollection::ImplicitsAttribute, FTransformCollection::TransformGroup);
+				PhysicsThreadCollection.CopyAttribute(*RestCollection, FGeometryDynamicCollection::ImplicitsAttribute, FTransformCollection::TransformGroup);
 			}
 
 			check(NumTransforms == PhysicsThreadCollection.NumElements(FTransformCollection::TransformGroup));
@@ -1866,7 +1866,7 @@ float FGeometryCollectionPhysicsProxy::ComputeUserDefinedDamageThreshold_Interna
 				LocalBoundingBox = FBox(ImplicitBoundingBox.Min(), ImplicitBoundingBox.Max());
 			}
 
-			const FSharedSimulationSizeSpecificData& SizeSpecificData = GetSizeSpecificData(Parameters.Shared.SizeSpecificData, *Parameters.RestCollection, TransformIndex, LocalBoundingBox);
+			const FSharedSimulationSizeSpecificData& SizeSpecificData = GetSizeSpecificData(Parameters.Shared.SizeSpecificData, *Parameters.RestCollectionShared, TransformIndex, LocalBoundingBox);
 			DamageThreshold = SizeSpecificData.DamageThreshold;
 		}
 		else
@@ -1907,7 +1907,7 @@ Chaos::TPBDGeometryCollectionParticleHandle<Chaos::FReal, 3>* FGeometryCollectio
 {
 	FGeometryDynamicCollection& DynamicCollection = PhysicsThreadCollection;
 	TManagedArray<uint8>& DynamicState = DynamicCollection.DynamicState;
-	const TManagedArray<FTransform>& MassToLocal = Parameters.RestCollection->GetAttribute<FTransform>(MassToLocalAttributeName, FTransformCollection::TransformGroup);
+	const TManagedArray<FTransform>& MassToLocal = Parameters.RestCollectionShared->GetAttribute<FTransform>(MassToLocalAttributeName, FTransformCollection::TransformGroup);
 	const TManagedArray<Chaos::FImplicitObjectPtr>& Implicits = DynamicCollection.GetAttribute<Chaos::FImplicitObjectPtr>(FGeometryDynamicCollection::ImplicitsAttribute, FTransformCollection::TransformGroup);
 	const TManagedArray<TUniquePtr<FCollisionStructureManager::FSimplicial>>& Simplicials = DynamicCollection.Simplicials;
 	Chaos::Facades::FCollectionAnchoringFacade AnchoringFacade(DynamicCollection);
@@ -1951,7 +1951,7 @@ Chaos::TPBDGeometryCollectionParticleHandle<Chaos::FReal, 3>* FGeometryCollectio
 		}
 	}
 
-	check(Parameters.RestCollection);
+	check(Parameters.RestCollectionShared);
 
 	// NOTE: The particle creation call above (CreateClusterParticle) activates the particle, which does various things including adding
 	// it to the constraint graph. It seems a bit dodgy to be completely changing what the particle is after it has been enabled. Maybe we should fix this.
@@ -2025,7 +2025,7 @@ FGeometryCollectionPhysicsProxy::BuildClusters_Internal(
 
 	FGeometryDynamicCollection& DynamicCollection = PhysicsThreadCollection;
 	TManagedArray<uint8>& DynamicState = DynamicCollection.DynamicState;
-	const TManagedArray<FTransform>& MassToLocal = Parameters.RestCollection->GetAttribute<FTransform>(MassToLocalAttributeName, FTransformCollection::TransformGroup);
+	const TManagedArray<FTransform>& MassToLocal = Parameters.RestCollectionShared->GetAttribute<FTransform>(MassToLocalAttributeName, FTransformCollection::TransformGroup);
 	const TManagedArray<Chaos::FImplicitObjectPtr>& Implicits = DynamicCollection.GetAttribute<Chaos::FImplicitObjectPtr>(FGeometryDynamicCollection::ImplicitsAttribute, FTransformCollection::TransformGroup);
 
 	//If we are a root particle use the world transform, otherwise set the relative transform
@@ -2084,9 +2084,9 @@ FGeometryCollectionPhysicsProxy::BuildClusters_Internal(
 		}
 	}
 
-	check(Parameters.RestCollection);
-	const TManagedArray<float>& Mass = Parameters.RestCollection->GetAttribute<float>(MassAttributeName, FTransformCollection::TransformGroup);
-	const TManagedArray<FVector3f>& InertiaTensor = Parameters.RestCollection->GetAttribute<FVector3f>(InertiaTensorAttributeName, FTransformCollection::TransformGroup);
+	check(Parameters.RestCollectionShared);
+	const TManagedArray<float>& Mass = Parameters.RestCollectionShared->GetAttribute<float>(MassAttributeName, FTransformCollection::TransformGroup);
+	const TManagedArray<FVector3f>& InertiaTensor = Parameters.RestCollectionShared->GetAttribute<FVector3f>(InertiaTensorAttributeName, FTransformCollection::TransformGroup);
 
 	const float ScaledMass = AdjustMassForScale(Mass[CollectionClusterIndex]);
 	const Chaos::FVec3f ScaledInertia = AdjustInertiaForScale((Chaos::FVec3f)InertiaTensor[CollectionClusterIndex]);
@@ -2337,7 +2337,7 @@ void FGeometryCollectionPhysicsProxy::GetRelevantParticleHandles(
 
 FName FGeometryCollectionPhysicsProxy::GetTransformName_External(const FGeometryCollectionItemIndex ItemIndex) const
 {
-	if (Parameters.RestCollection)
+	if (Parameters.RestCollectionShared)
 	{
 		if (!ItemIndex.IsInternalCluster())
 		{
@@ -3279,7 +3279,7 @@ void FGeometryCollectionPhysicsProxy::SetWorldTransform_Internal(const FTransfor
 		TArray<FPBDRigidParticleHandle*> DeferredClusterUnionParticleUpdates;
 		TArray<FTransform> DeferredClusterUnionChildToParentUpdates;
 
-		const TManagedArray<FTransform>& MassToLocal = Parameters.RestCollection->GetAttribute<FTransform>(MassToLocalAttributeName, FTransformCollection::TransformGroup);
+		const TManagedArray<FTransform>& MassToLocal = Parameters.RestCollectionShared->GetAttribute<FTransform>(MassToLocalAttributeName, FTransformCollection::TransformGroup);
 
 		for (int32 ParticleIndex = 0; ParticleIndex < NumEffectiveParticles; ++ParticleIndex)
 		{
@@ -3375,7 +3375,7 @@ void FGeometryCollectionPhysicsProxy::SetUseStaticMeshCollisionForTraces_Externa
 			{
 				if (bUseSMCollisionForTraces)
 				{
-					const TManagedArray<FTransform>& MassToLocal = Parameters.RestCollection->GetAttribute<FTransform>("MassToLocal", FGeometryCollection::TransformGroup);
+					const TManagedArray<FTransform>& MassToLocal = Parameters.RestCollectionShared->GetAttribute<FTransform>("MassToLocal", FGeometryCollection::TransformGroup);
 
 					const FTransform ToLocal = MassToLocal[TransformIndex].Inverse();
 					TArray<Chaos::FImplicitObjectPtr> Geoms;
@@ -3534,7 +3534,7 @@ void FGeometryCollectionPhysicsProxy::SetGravityGroupIndex_Internal(int32 Gravit
 void FGeometryCollectionPhysicsProxy::SetOneWayInteractionLevel_External(int32 InOneWayInteractionLevel)
 {
 	// @todo(chaos): not sure we need the OneWayInteraction flag on the GT side for geometry collection particles?
-	const TManagedArray<int32>& Level = Parameters.RestCollection->GetAttribute<int32>(LevelAttributeName, FTransformCollection::TransformGroup);
+	const TManagedArray<int32>& Level = Parameters.RestCollectionShared->GetAttribute<int32>(LevelAttributeName, FTransformCollection::TransformGroup);
 	for (int32 ParticleIndex = 0; ParticleIndex < GTParticles.Num(); ++ParticleIndex)
 	{
 		TUniquePtr<FParticle>& Particle = GTParticles[ParticleIndex];
@@ -3771,7 +3771,7 @@ void FGeometryCollectionPhysicsProxy::BufferPhysicsResults_Internal(Chaos::FPBDR
 	{ 
 		SCOPE_CYCLE_COUNTER(STAT_CalcParticleToWorld);
 
-		const TManagedArray<FTransform>& MassToLocalArray = Parameters.RestCollection->GetAttribute<FTransform>(MassToLocalAttributeName, FTransformCollection::TransformGroup);
+		const TManagedArray<FTransform>& MassToLocalArray = Parameters.RestCollectionShared->GetAttribute<FTransform>(MassToLocalAttributeName, FTransformCollection::TransformGroup);
 
 		// Explicitly handle root breaking to ensure that the Root GT particle Active state matches the PT particle.
 		// If the root was in a cluster union and the cluster gets broken immediately on removal, we
@@ -4317,7 +4317,7 @@ bool FGeometryCollectionPhysicsProxy::PullFromPhysicsState(const Chaos::FDirtyGe
 		bIsCollectionDirty |= PullNonInterpolatableDataFromSinglePhysicsState(PullData, !bNeedInterpolation, NextPullData ? &NextPullData->Results().GetModifiedTransformIndices() : nullptr);
 
 		const TManagedArray<bool>* AnimationsActive = GameThreadCollection.GetAnimateTransformAttribute();
-		const TManagedArray<FTransform>& MassToLocal = Parameters.RestCollection->GetAttribute<FTransform>(MassToLocalAttributeName, FTransformCollection::TransformGroup);
+		const TManagedArray<FTransform>& MassToLocal = Parameters.RestCollectionShared->GetAttribute<FTransform>(MassToLocalAttributeName, FTransformCollection::TransformGroup);
 
 		// second : interpolate-able ones
 		if (bNeedInterpolation)
@@ -5900,7 +5900,7 @@ void FGeometryCollectionPhysicsProxy::RebaseAllGameThreadCollectionTransformsOnN
 {
 	check(IsInGameThread());
 	const TManagedArray<bool>* AnimationsActive = GameThreadCollection.GetAnimateTransformAttribute();
-	const TManagedArray<FTransform>& MassToLocal = Parameters.RestCollection->GetAttribute<FTransform>(MassToLocalAttributeName, FTransformCollection::TransformGroup);
+	const TManagedArray<FTransform>& MassToLocal = Parameters.RestCollectionShared->GetAttribute<FTransform>(MassToLocalAttributeName, FTransformCollection::TransformGroup);
 	const bool bIsComponentTransformScaled = !WorldTransform_External.GetScale3D().Equals(FVector::OneVector);
 	const FTransform ComponentScaleTransform(FQuat::Identity, FVector::ZeroVector, WorldTransform_External.GetScale3D());
 
diff --git a/Engine/Source/Runtime/Experimental/Chaos/Public/GeometryCollection/GeometryCollectionSimulationCoreTypes.h b/Engine/Source/Runtime/Experimental/Chaos/Public/GeometryCollection/GeometryCollectionSimulationCoreTypes.h
--- a/Engine/Source/Runtime/Experimental/Chaos/Public/GeometryCollection/GeometryCollectionSimulationCoreTypes.h
+++ b/Engine/Source/Runtime/Experimental/Chaos/Public/GeometryCollection/GeometryCollectionSimulationCoreTypes.h
@@ -170,7 +170,10 @@ struct FSimulationParameters
 {
 	FSimulationParameters()
 		: Name("")
+PRAGMA_DISABLE_DEPRECATION_WARNINGS
 		, RestCollection(nullptr)
+PRAGMA_ENABLE_DEPRECATION_WARNINGS
+		, RestCollectionShared(nullptr)
 		, InitialRootIndex(INDEX_NONE)
 		, RecordedTrack(nullptr)
 		, bOwnsTrack(false)
@@ -232,7 +235,10 @@ struct FSimulationParameters
 
 	FSimulationParameters(const FSimulationParameters& Other)
 		: Name(Other.Name)
+PRAGMA_DISABLE_DEPRECATION_WARNINGS
 		, RestCollection(Other.RestCollection)
+PRAGMA_ENABLE_DEPRECATION_WARNINGS
+		, RestCollectionShared(Other.RestCollectionShared)
 		, InitialRootIndex(Other.InitialRootIndex)
 		, InitializationCommands(Other.InitializationCommands)
 		, RecordedTrack(Other.RecordedTrack)
@@ -308,7 +314,12 @@ struct FSimulationParameters
 	bool IsCachePlaying() { return CacheType == EGeometryCollectionCacheType::Play || CacheType == EGeometryCollectionCacheType::RecordAndPlay; }
 
 	FString Name;
+	
+	UE_DEPRECATED(5.4, "Raw pointer no longer in use, instead prefer RestCollectionShared")
 	const FGeometryCollection* RestCollection;
+
+	TSharedPtr<const FGeometryCollection> RestCollectionShared;
+
 	int32 InitialRootIndex;
 	TArray<FFieldSystemCommand> InitializationCommands;
 	const FRecordedTransformTrack* RecordedTrack;
diff --git a/Engine/Source/Runtime/Experimental/Chaos/Public/GeometryCollectionProxyData.h b/Engine/Source/Runtime/Experimental/Chaos/Public/GeometryCollectionProxyData.h
--- a/Engine/Source/Runtime/Experimental/Chaos/Public/GeometryCollectionProxyData.h
+++ b/Engine/Source/Runtime/Experimental/Chaos/Public/GeometryCollectionProxyData.h
@@ -22,7 +22,9 @@ class FTransformDynamicCollection : public FManagedArrayCollection
 public:
 	typedef FManagedArrayCollection Super;
 
+	UE_DEPRECATED(5.4, "No longer handles a raw ptr, use the version taking TSharedPtr")
 	CHAOS_API FTransformDynamicCollection(const FGeometryCollection* InRestCollection);
+	CHAOS_API FTransformDynamicCollection(TSharedPtr<const FGeometryCollection> InRestCollection);
 	FTransformDynamicCollection(FTransformDynamicCollection&) = delete;
 	FTransformDynamicCollection& operator=(const FTransformDynamicCollection&) = delete;
 	FTransformDynamicCollection(FTransformDynamicCollection&&) = delete;
@@ -43,7 +45,7 @@ public:
 	template<typename Lambda>
 	void IterateThroughChildren(int32 Index, Lambda&& LambdaIt) const
 	{
-		const TSet<int32>& Children = RestCollection->Children[Index];
+		const TSet<int32>& Children = RestCollectionShared->Children[Index];
 		for (const int32 Child : Children)
 		{
 			if (GetHasParent(Child))
@@ -58,8 +60,9 @@ public:
 	}
 
 protected:
+	UE_DEPRECATED(5.4, "Raw pointer no longer in use, instead access the shared version")
 	const FGeometryCollection* RestCollection;
-
+	TSharedPtr<const FGeometryCollection> RestCollectionShared;
 private:
 	TManagedArray<bool>         HasParent;
 	TManagedArray<FTransform3f> Transform;
@@ -80,7 +83,9 @@ private:
 class FGeometryDynamicCollection : public FTransformDynamicCollection
 {
 public:
+	UE_DEPRECATED(5.4, "No longer handles a raw ptr, use the version taking TSharedPtr")
 	CHAOS_API FGeometryDynamicCollection(const FGeometryCollection* InRestCollection);
+	CHAOS_API FGeometryDynamicCollection(TSharedPtr<const FGeometryCollection> InRestCollection);
 	FGeometryDynamicCollection(FGeometryDynamicCollection&) = delete;
 	FGeometryDynamicCollection& operator=(const FGeometryDynamicCollection&) = delete;
 	FGeometryDynamicCollection(FGeometryDynamicCollection&&) = delete;
diff --git a/Engine/Source/Runtime/Experimental/GeometryCollectionEngine/Private/GeometryCollection/GeometryCollectionComponent.cpp b/Engine/Source/Runtime/Experimental/GeometryCollectionEngine/Private/GeometryCollection/GeometryCollectionComponent.cpp
--- a/Engine/Source/Runtime/Experimental/GeometryCollectionEngine/Private/GeometryCollection/GeometryCollectionComponent.cpp
+++ b/Engine/Source/Runtime/Experimental/GeometryCollectionEngine/Private/GeometryCollection/GeometryCollectionComponent.cpp
@@ -3735,7 +3735,7 @@ void UGeometryCollectionComponent::ResetDynamicCollection()
 #endif
 	if (bCreateDynamicCollection && RestCollection && RestCollection->GetGeometryCollection())
 	{
-		DynamicCollection = MakeUnique<FGeometryDynamicCollection>(RestCollection->GetGeometryCollection().Get());
+		DynamicCollection = MakeUnique<FGeometryDynamicCollection>(RestCollection->GetGeometryCollection());
 
 		IndirectParentArray = nullptr;
 		GetParentArrayCopyOnWrite();
@@ -3911,7 +3911,11 @@ void UGeometryCollectionComponent::RegisterAndInitializePhysicsProxy()
 		if (RestCollection)
 		{
 			RestCollection->GetSharedSimulationParams(SimulationParameters.Shared);
-			SimulationParameters.RestCollection = RestCollection->GetGeometryCollection().Get();
+			SimulationParameters.RestCollectionShared = RestCollection->GetGeometryCollection();
+PRAGMA_DISABLE_DEPRECATION_WARNINGS
+			// To be removed when RestCollection is removed post-deprecation. Here for back compat
+			SimulationParameters.RestCollection = SimulationParameters.RestCollectionShared.Get();
+PRAGMA_ENABLE_DEPRECATION_WARNINGS
 			SimulationParameters.InitialRootIndex = RestCollection->GetRootIndex();
 			ClusterCollectionType = RestCollection->ClusterConnectionType;
 			ConnectionGraphBoundsFilteringMargin = RestCollection->ConnectionGraphBoundsFilteringMargin;
