Constraints: use Algo::Sort instead of Algo::StableSort to have the right constraints ordering

#rnx

diff --git a/Engine/Source/Runtime/Experimental/Animation/Constraints/Private/ConstraintSubsystem.cpp b/Engine/Source/Runtime/Experimental/Animation/Constraints/Private/ConstraintSubsystem.cpp
--- a/Engine/Source/Runtime/Experimental/Animation/Constraints/Private/ConstraintSubsystem.cpp
+++ b/Engine/Source/Runtime/Experimental/Animation/Constraints/Private/ConstraintSubsystem.cpp
@@ -320,15 +320,15 @@ FConstraintsEvaluationGraph& FConstraintsInWorld::GetEvaluationGraph()
 {
 	if (!EvaluationGraph)
 	{
-		EvaluationGraph = MakeShared<FConstraintsEvaluationGraph>(*this);
+		EvaluationGraph = MakeShared<FConstraintsEvaluationGraph>(this);
 	}
 	return *EvaluationGraph;
 }
 
-void FConstraintsInWorld::InvalidateGraph() const
+void FConstraintsInWorld::InvalidateGraph()
 {
 	if (EvaluationGraph)
 	{
-		EvaluationGraph->InvalidateData();
+		EvaluationGraph.Reset();
 	}
 }
\ No newline at end of file
diff --git a/Engine/Source/Runtime/Experimental/Animation/Constraints/Private/ConstraintSubsystem.h b/Engine/Source/Runtime/Experimental/Animation/Constraints/Private/ConstraintSubsystem.h
--- a/Engine/Source/Runtime/Experimental/Animation/Constraints/Private/ConstraintSubsystem.h
+++ b/Engine/Source/Runtime/Experimental/Animation/Constraints/Private/ConstraintSubsystem.h
@@ -30,7 +30,7 @@ public:
 	CONSTRAINTS_API void RemoveConstraints(UWorld* World);
 
 	FConstraintsEvaluationGraph& GetEvaluationGraph();
-	void InvalidateGraph() const;
+	void InvalidateGraph();
 
 private:
 	TSharedPtr<FConstraintsEvaluationGraph> EvaluationGraph = nullptr;
diff --git a/Engine/Source/Runtime/Experimental/Animation/Constraints/Private/ConstraintsEvaluationGraph.cpp b/Engine/Source/Runtime/Experimental/Animation/Constraints/Private/ConstraintsEvaluationGraph.cpp
--- a/Engine/Source/Runtime/Experimental/Animation/Constraints/Private/ConstraintsEvaluationGraph.cpp
+++ b/Engine/Source/Runtime/Experimental/Animation/Constraints/Private/ConstraintsEvaluationGraph.cpp
@@ -130,7 +130,7 @@ void FConstraintsEvaluationGraph::Rebuild()
 		});
 		return bIsLHSAPrerexOfRHS;
 	};
-	Algo::StableSort(Nodes, EvaluationOrderPredicate);
+	Algo::Sort(Nodes, EvaluationOrderPredicate);
 
 	// store node index after re-ordering
 	for (int32 NodeIndex = 0; NodeIndex < Nodes.Num(); NodeIndex++)
diff --git a/Engine/Source/Runtime/Experimental/Animation/Constraints/Private/ConstraintsManager.cpp b/Engine/Source/Runtime/Experimental/Animation/Constraints/Private/ConstraintsManager.cpp
--- a/Engine/Source/Runtime/Experimental/Animation/Constraints/Private/ConstraintsManager.cpp
+++ b/Engine/Source/Runtime/Experimental/Animation/Constraints/Private/ConstraintsManager.cpp
@@ -766,7 +766,7 @@ bool FConstraintsManagerController::RemoveAllConstraints(bool bDoNotCompensate)
 	return false;
 }
 
-static void SortConstraints(UWorld* InWorld, TArray< TWeakObjectPtr<UTickableConstraint> >& SortedConstraints)
+static void SortConstraints(UWorld* InWorld, TArray< TWeakObjectPtr<UTickableConstraint> >& InOutSortedConstraints)
 {
 	using ConstraintPtr = TWeakObjectPtr<UTickableConstraint>;
 	// LHS ticks before RHS = LHS is a prerex of RHS 
@@ -785,13 +785,13 @@ static void SortConstraints(UWorld* InWorld, TArray< TWeakObjectPtr<UTickableCon
 		return bIsLHSAPrerexOfRHS;
 	};
 
-	Algo::StableSort(SortedConstraints, [TicksBefore](const ConstraintPtr& LHS, const ConstraintPtr& RHS)
-		{
-			return TicksBefore(*LHS, *RHS);
-		});
+	Algo::Sort(InOutSortedConstraints, [TicksBefore](const ConstraintPtr& LHS, const ConstraintPtr& RHS)
+	{
+		return TicksBefore(*LHS, *RHS);
+	});
 }
 
-static void SortConstraints(UWorld* InWorld, TArray< TObjectPtr<UTickableConstraint> > &SortedConstraints)
+static void SortConstraints(UWorld* InWorld, TArray< TObjectPtr<UTickableConstraint> >& InOutSortedConstraints)
 {
 	using ConstraintPtr = TObjectPtr<UTickableConstraint>;
 	// LHS ticks before RHS = LHS is a prerex of RHS 
@@ -806,10 +806,10 @@ static void SortConstraints(UWorld* InWorld, TArray< TObjectPtr<UTickableConstra
 		return bIsLHSAPrerexOfRHS;
 	};
 
-	Algo::StableSort(SortedConstraints, [TicksBefore](const ConstraintPtr& LHS, const ConstraintPtr& RHS)
-		{
-			return TicksBefore(*LHS, *RHS);
-		});
+	Algo::Sort(InOutSortedConstraints, [TicksBefore](const ConstraintPtr& LHS, const ConstraintPtr& RHS)
+	{
+		return TicksBefore(*LHS, *RHS);
+	});
 }
 
 TArray< TObjectPtr<UTickableConstraint> > FConstraintsManagerController::GetStaticConstraints(const bool bSorted) const
@@ -868,12 +868,15 @@ TArray< TWeakObjectPtr<UTickableConstraint> > FConstraintsManagerController::Get
 
 void FConstraintsManagerController::EvaluateAllConstraints() const
 {
-	TArray< TWeakObjectPtr<UTickableConstraint>>Constraints = GetAllConstraints(true);
-	for (const TWeakObjectPtr<UTickableConstraint>& InConstraint : Constraints)
+	using ConstraintPtr = TWeakObjectPtr<UTickableConstraint>;
+
+	static constexpr bool bSorted = true, bTickHandles = true;
+	const TArray<ConstraintPtr> Constraints = GetAllConstraints(bSorted);
+	for (const ConstraintPtr& Constraint : Constraints)
 	{
-		if (InConstraint.IsValid())
+		if (Constraint.IsValid())
 		{
-			InConstraint->Evaluate(true);
+			Constraint->Evaluate(bTickHandles);
 		}
 	}
 }
diff --git a/Engine/Source/Runtime/Experimental/Animation/Constraints/Public/ConstraintsEvaluationGraph.h b/Engine/Source/Runtime/Experimental/Animation/Constraints/Public/ConstraintsEvaluationGraph.h
--- a/Engine/Source/Runtime/Experimental/Animation/Constraints/Public/ConstraintsEvaluationGraph.h
+++ b/Engine/Source/Runtime/Experimental/Animation/Constraints/Public/ConstraintsEvaluationGraph.h
@@ -40,8 +40,8 @@ struct FConstraintNode
 struct FConstraintsEvaluationGraph: public TSharedFromThis<FConstraintsEvaluationGraph>
 {
 public:
-	explicit FConstraintsEvaluationGraph(FConstraintsInWorld& InConstraintsInWorld)
-		: ConstraintsInWorld(InConstraintsInWorld)
+	explicit FConstraintsEvaluationGraph(const FConstraintsInWorld* InConstraintsInWorld)
+		: ConstraintsInWorld(*InConstraintsInWorld)
 	{}
 
 	/**
@@ -97,7 +97,7 @@ private:
 	TArray<FConstraintNode> Nodes;
 
 	/* References to the actual stored data. */
-	FConstraintsInWorld& ConstraintsInWorld;
+	const FConstraintsInWorld& ConstraintsInWorld;
 	
 	/* Current graph state. */
 	EGraphState	State = InvalidData;
