Modular Vehicles:
Add deferred force API for applying forces relative to cluster COM.
Overload Simulate method with a version that has a Proxy parameter.

diff --git a/Engine/Source/Runtime/Experimental/ChaosVehicles/ChaosVehiclesCore/Private/SimModule/DeferredForcesModular.cpp b/Engine/Source/Runtime/Experimental/ChaosVehicles/ChaosVehiclesCore/Private/SimModule/DeferredForcesModular.cpp
--- a/Engine/Source/Runtime/Experimental/ChaosVehicles/ChaosVehiclesCore/Private/SimModule/DeferredForcesModular.cpp
+++ b/Engine/Source/Runtime/Experimental/ChaosVehicles/ChaosVehiclesCore/Private/SimModule/DeferredForcesModular.cpp
@@ -178,6 +178,14 @@ Chaos::FPBDRigidParticleHandle* FDeferredForcesModular::GetParticle(const FTrans
 	return ClusterParticles[SingleChassisIndex];
 }
 
+Chaos::FPBDRigidParticleHandle* FDeferredForcesModular::GetClusterParticle(TArray<Chaos::FPBDRigidClusteredParticleHandle*>& ClusterParticles)
+{
+	using namespace Chaos;
+	ensure(ClusterParticles.Num() == 1);
+	const int SingleChassisIndex = 0;
+	return ClusterParticles[SingleChassisIndex];
+}
+
 template<typename TransformType>
 void FDeferredForcesModular::ApplyTemplate(FGeometryCollectionPhysicsProxy* Proxy
 		, const TManagedArray<TransformType>& Transforms
@@ -290,6 +298,16 @@ void FDeferredForcesModular::Apply(TArray<Chaos::FPBDRigidParticleHandle*>& Part
 			}
 		}
 
+		for (const FApplyForceAtPositionData& Data : ApplyForceAtCOMDatas)
+		{
+			Chaos::FPBDRigidParticleHandle* RigidHandle = GetClusterParticle(ClusterParticles);
+			if (RigidHandle)
+			{
+				AddForceAtCOM(RigidHandle, Data);
+			}
+		}
+
+
 		for (const FApplyForceAtPositionData& Data : ApplyForceAtPositionDatas)
 		{
 			Chaos::FPBDRigidParticleHandle* RigidHandle = GetParticle(Data.OffsetTransform, Particles, ClusterParticles, Data.ParticleIdx, Data.Position, RelativeTransform);
@@ -310,6 +328,7 @@ void FDeferredForcesModular::Apply(TArray<Chaos::FPBDRigidParticleHandle*>& Part
 	}
 
 	ApplyForceDatas.Empty();
+	ApplyForceAtCOMDatas.Empty();
 	ApplyForceAtPositionDatas.Empty();
 	ApplyTorqueDatas.Empty();
 }
@@ -389,6 +408,16 @@ void FDeferredForcesModular::AddForce(Chaos::FPBDRigidParticleHandle* RigidHandl
 
 }
 
+void FDeferredForcesModular::AddForceAtCOM(Chaos::FPBDRigidParticleHandle* RigidHandle, const FApplyForceAtPositionData& DataIn)
+{
+	if (ensure(RigidHandle))
+	{
+		FTransform OffsetTransform = FTransform::Identity;
+		OffsetTransform.SetTranslation(RigidHandle->CenterOfMass() + DataIn.Position);
+		AddForce_Implementation(RigidHandle, FTransform::Identity, DataIn.Force, OffsetTransform, (DataIn.Flags & EForceFlags::LevelSlope) == EForceFlags::LevelSlope, DataIn.DebugColor);
+	}
+}
+
 void FDeferredForcesModular::AddForceAtPosition(Chaos::FPBDRigidParticleHandle* RigidHandle, const FApplyForceAtPositionData& DataIn, const FTransform& OffsetTransform)
 {
 	if (ensure(RigidHandle))
diff --git a/Engine/Source/Runtime/Experimental/ChaosVehicles/ChaosVehiclesCore/Private/SimModule/SimModuleTree.cpp b/Engine/Source/Runtime/Experimental/ChaosVehicles/ChaosVehiclesCore/Private/SimModule/SimModuleTree.cpp
--- a/Engine/Source/Runtime/Experimental/ChaosVehicles/ChaosVehiclesCore/Private/SimModule/SimModuleTree.cpp
+++ b/Engine/Source/Runtime/Experimental/ChaosVehicles/ChaosVehiclesCore/Private/SimModule/SimModuleTree.cpp
@@ -241,7 +241,7 @@ void FSimModuleTree::SimulateNode(float DeltaTime, FAllInputs& Inputs, int NodeI
 	{
 		if (Module->IsEnabled())
 		{
-			Module->Simulate(DeltaTime, Inputs, *this);
+			Module->Simulate(PhysicsProxy, DeltaTime, Inputs, *this);
 
 			if (IsAnimationEnabled() && Module->IsAnimationEnabled())
 			{
diff --git a/Engine/Source/Runtime/Experimental/ChaosVehicles/ChaosVehiclesCore/Private/SimModule/SimulationModuleBase.cpp b/Engine/Source/Runtime/Experimental/ChaosVehicles/ChaosVehiclesCore/Private/SimModule/SimulationModuleBase.cpp
--- a/Engine/Source/Runtime/Experimental/ChaosVehicles/ChaosVehiclesCore/Private/SimModule/SimulationModuleBase.cpp
+++ b/Engine/Source/Runtime/Experimental/ChaosVehicles/ChaosVehiclesCore/Private/SimModule/SimulationModuleBase.cpp
@@ -19,6 +19,15 @@ void ISimulationModuleBase::AddLocalForceAtPosition(const FVector& Force, const
 	}
 }
 
+void ISimulationModuleBase::AddForceAtCOMPosition(const FVector& Force, const FVector& OffsetFromCOM /*= = FVector::ZeroVector*/, bool bAllowSubstepping /*= true*/, bool bLevelSlope /*= false*/, const FColor& DebugColorIn /*= FColor::Blue*/)
+{
+	AppliedForce = Force;
+	if (SimModuleTree)
+	{
+		SimModuleTree->AccessDeferredForces().AddCOM(FDeferredForcesModular::FApplyForceAtPositionData(ComponentTransform, TransformIndex, ParticleIdx.Idx, Force, OffsetFromCOM, bAllowSubstepping, false, bLevelSlope, DebugColorIn));
+	}
+}
+
 void ISimulationModuleBase::AddLocalForce(const FVector& Force, bool bAllowSubstepping, bool bIsLocalForce, bool bLevelSlope, const FColor& DebugColorIn)
 {
 	AppliedForce = Force;
diff --git a/Engine/Source/Runtime/Experimental/ChaosVehicles/ChaosVehiclesCore/Public/SimModule/DeferredForcesModular.h b/Engine/Source/Runtime/Experimental/ChaosVehicles/ChaosVehiclesCore/Public/SimModule/DeferredForcesModular.h
--- a/Engine/Source/Runtime/Experimental/ChaosVehicles/ChaosVehiclesCore/Public/SimModule/DeferredForcesModular.h
+++ b/Engine/Source/Runtime/Experimental/ChaosVehicles/ChaosVehiclesCore/Public/SimModule/DeferredForcesModular.h
@@ -144,6 +144,11 @@ public:
 		ApplyForceAtPositionDatas.Add(ApplyForceAtPositionDataIn);
 	}
 
+	void AddCOM(const FApplyForceAtPositionData& ApplyForceAtPositionDataIn)
+	{
+		ApplyForceAtCOMDatas.Add(ApplyForceAtPositionDataIn);
+	}
+
 	void Add(const FAddTorqueInRadiansData& ApplyTorqueDataIn)
 	{
 		ApplyTorqueDatas.Add(ApplyTorqueDataIn);
@@ -192,6 +197,8 @@ public:
 		, const FVector& PositionalOffset
 		, FTransform& TransformOut);
 
+	Chaos::FPBDRigidParticleHandle* GetClusterParticle(TArray<Chaos::FPBDRigidClusteredParticleHandle*>& ClusterParticles);
+
 	void Apply(FGeometryCollectionPhysicsProxy* Proxy
 		, const TManagedArray<FTransform>& Transforms
 		, const TManagedArray<FTransform>& CollectionMassToLocal
@@ -218,8 +225,10 @@ private:
 	void AddForceAtPosition(Chaos::FPBDRigidParticleHandle* RigidHandle, const FApplyForceAtPositionData& DataIn, const FTransform& OffsetTransform);
 	void AddTorque(Chaos::FPBDRigidParticleHandle* RigidHandle, const FAddTorqueInRadiansData& DataIn, const FTransform& OffsetTransform);
 	void AddForce(Chaos::FPBDRigidParticleHandle* RigidHandle, const FApplyForceData& DataIn, const FTransform& OffsetTransform);
+	void AddForceAtCOM(Chaos::FPBDRigidParticleHandle* RigidHandle, const FApplyForceAtPositionData& DataIn);
 
 	TArray<FApplyForceData> ApplyForceDatas;
+	TArray<FApplyForceAtPositionData> ApplyForceAtCOMDatas;
 	TArray<FApplyForceAtPositionData> ApplyForceAtPositionDatas;
 	TArray<FAddTorqueInRadiansData> ApplyTorqueDatas;
 	TArray<FAddImpulseData> ApplyImpulseDatas;
diff --git a/Engine/Source/Runtime/Experimental/ChaosVehicles/ChaosVehiclesCore/Public/SimModule/SimulationModuleBase.h b/Engine/Source/Runtime/Experimental/ChaosVehicles/ChaosVehiclesCore/Public/SimModule/SimulationModuleBase.h
--- a/Engine/Source/Runtime/Experimental/ChaosVehicles/ChaosVehiclesCore/Public/SimModule/SimulationModuleBase.h
+++ b/Engine/Source/Runtime/Experimental/ChaosVehicles/ChaosVehiclesCore/Public/SimModule/SimulationModuleBase.h
@@ -222,6 +222,14 @@ namespace Chaos
 		/**
 		 * The main Simulation function that is called from the physics async callback thread
 		 */
+		virtual void Simulate(Chaos::FClusterUnionPhysicsProxy* Proxy, float DeltaTime, const FAllInputs& Inputs, FSimModuleTree& VehicleModuleSystem) 
+		{
+			Simulate(DeltaTime, Inputs, VehicleModuleSystem);
+		}
+
+		/**
+		* The main Simulation function that is called from the physics async callback thread
+		*/
 		virtual void Simulate(float DeltaTime, const FAllInputs& Inputs, FSimModuleTree& VehicleModuleSystem) {}
 
 		/**
@@ -276,6 +284,9 @@ namespace Chaos
 		 */
 		void AddLocalForceAtPosition(const FVector& Force, const FVector& Position, bool bAllowSubstepping = true, bool bIsLocalForce = false, bool bLevelSlope = false, const FColor& DebugColorIn = FColor::Blue);
 
+
+		void AddForceAtCOMPosition(const FVector& Force, const FVector& OffsetFromCOM = FVector::ZeroVector, bool bAllowSubstepping = true, bool bLevelSlope = false, const FColor& DebugColorIn = FColor::Blue);
+
 		/**
 		 * Force application function, handles deferred force application and applying the force at the collect location based on whether the GC cluster is intact or fractured
 		 * Note: forces are applied in local coordinates of the module
