[ChaosVD] Fix for data race ensure in OnRecordingUpdate delegate

The issue could happen if a recording was unloaded while new data was being written to it (for example if you open a CVD file, you then go and open another one before it finishes loading, or just close CVD while it is being unloaded).

As that delegate broadcast is a hot path (specially during live debugging), I changed the implementation to not depend on a delegate. Now we use a timestamp the UI can check against to know if it needs to update its data

#tests I tested that I can load a CVD file and the UI updates properly as new frames are loaded in. I also tested that live recording (which constantly adds frames to a recording) works as expected.

diff --git a/Engine/Plugins/ChaosVD/Source/ChaosVD/Private/ChaosVDPlaybackController.cpp b/Engine/Plugins/ChaosVD/Source/ChaosVD/Private/ChaosVDPlaybackController.cpp
--- a/Engine/Plugins/ChaosVD/Source/ChaosVD/Private/ChaosVDPlaybackController.cpp
+++ b/Engine/Plugins/ChaosVD/Source/ChaosVD/Private/ChaosVDPlaybackController.cpp
@@ -76,8 +76,6 @@ bool FChaosVDPlaybackController::LoadChaosVDRecordingFromTraceSession(const FCha
 
 	LoadedRecording->SetIsLive(InSessionDescriptor.bIsLiveSession);
 
-	LoadedRecording->OnRecordingUpdated().AddRaw(this, &FChaosVDPlaybackController::HandleCurrentRecordingUpdated);
-
 	HandleCurrentRecordingUpdated();
 
 	LoadedRecording->OnGeometryDataLoaded().AddRaw(this, &FChaosVDPlaybackController::EnqueueGeometryDataUpdate);
@@ -94,11 +92,12 @@ bool FChaosVDPlaybackController::LoadChaosVDRecordingFromTraceSession(const FCha
 
 void FChaosVDPlaybackController::UnloadCurrentRecording(EChaosVDUnloadRecordingFlags UnloadOptions)
 {
+	RecordingLastSeenTimeUpdatedAsCycle = 0;
+
 	TrackInfoUpdateGTQueue.Empty();
 	
 	if (LoadedRecording.IsValid())
 	{
-		LoadedRecording->OnRecordingUpdated().RemoveAll(this);
 		LoadedRecording.Reset();
 	}
 
@@ -633,6 +632,14 @@ bool FChaosVDPlaybackController::Tick(float DeltaTime)
 
 	if (LoadedRecording.IsValid())
 	{
+		uint64 CurrentLastUpdatedTime = LoadedRecording->GetLastUpdatedTimeAsCycle();
+		if (CurrentLastUpdatedTime != RecordingLastSeenTimeUpdatedAsCycle)
+		{
+			RecordingLastSeenTimeUpdatedAsCycle = CurrentLastUpdatedTime;
+
+			HandleCurrentRecordingUpdated();
+		}
+
 		// Load at least the first frame
 		if (!bPlayedFirstFrame)
 		{
diff --git a/Engine/Plugins/ChaosVD/Source/ChaosVD/Private/ChaosVDPlaybackController.h b/Engine/Plugins/ChaosVD/Source/ChaosVD/Private/ChaosVDPlaybackController.h
--- a/Engine/Plugins/ChaosVD/Source/ChaosVD/Private/ChaosVDPlaybackController.h
+++ b/Engine/Plugins/ChaosVD/Source/ChaosVD/Private/ChaosVDPlaybackController.h
@@ -283,6 +283,9 @@ protected:
 	/** Set to true when the recording data controlled by this Playback Controller is updated, the update delegate will be called on the GT */
 	std::atomic<bool> bHasPendingGTUpdateBroadcast;
 
+	/** Last seen Platform Cycle on which the loaded recording was updated */
+	uint64 RecordingLastSeenTimeUpdatedAsCycle = 0;
+
 	/** Queue with a copy of all Track Info Updates that needs to be done in the Game thread */
 	TQueue<FChaosVDQueuedTrackInfoUpdate, EQueueMode::Mpsc> TrackInfoUpdateGTQueue;
 
diff --git a/Engine/Source/Runtime/Experimental/ChaosVDData/Private/ChaosVDRecording.cpp b/Engine/Source/Runtime/Experimental/ChaosVDData/Private/ChaosVDRecording.cpp
--- a/Engine/Source/Runtime/Experimental/ChaosVDData/Private/ChaosVDRecording.cpp
+++ b/Engine/Source/Runtime/Experimental/ChaosVDData/Private/ChaosVDRecording.cpp
@@ -233,7 +233,6 @@ void FChaosVDRecording::GenerateAndStoreKeyframeForSolver_AssumesLocked(const in
 	}
 }
 
-
 void FChaosVDRecording::AddFrameForSolver(const int32 SolverID, FChaosVDSolverFrameData&& InFrameData)
 {
 	int32 CurrentFrameNumber;
@@ -288,7 +287,7 @@ void FChaosVDRecording::AddFrameForSolver(const int32 SolverID, FChaosVDSolverFr
 		AddKeyFrameNumberForSolver(SolverID, CurrentFrameNumber);
 	}
 
-	OnRecordingUpdated().Broadcast();
+	LastUpdatedTimeAsCycle = FPlatformTime::Cycles64();
 }
 
 void FChaosVDRecording::AddGameFrameData(const FChaosVDGameFrameData& InFrameData)
diff --git a/Engine/Source/Runtime/Experimental/ChaosVDData/Public/ChaosVDRecording.h b/Engine/Source/Runtime/Experimental/ChaosVDData/Public/ChaosVDRecording.h
--- a/Engine/Source/Runtime/Experimental/ChaosVDData/Public/ChaosVDRecording.h
+++ b/Engine/Source/Runtime/Experimental/ChaosVDData/Public/ChaosVDRecording.h
@@ -13,13 +13,13 @@
 #include "ChaosVisualDebugger/ChaosVDMemWriterReader.h"
 #include "DataWrappers/ChaosVDJointDataWrappers.h"
 #include "DataWrappers/ChaosVDQueryDataWrappers.h"
+#include <atomic>
 
 namespace Chaos::VisualDebugger
 {
-class FChaosVDSerializableNameTable;
+	class FChaosVDSerializableNameTable;
 }
 
-DECLARE_MULTICAST_DELEGATE(FChaosVDRecordingUpdated)
 DECLARE_MULTICAST_DELEGATE_TwoParams(FChaosVDGeometryDataLoaded, const Chaos::FConstImplicitObjectPtr&, const uint32 GeometryID)
 
 struct FChaosVDStepData
@@ -197,9 +197,6 @@ struct CHAOSVDDATA_API FChaosVDRecording
 	 */
 	void AddGameFrameData(const FChaosVDGameFrameData& InFrameData);
 
-	/** Called each time the recording changes - Mainly when a new frame is added from the Trace analysis */
-	FChaosVDRecordingUpdated& OnRecordingUpdated() { return RecordingUpdatedDelegate; };
-
 	/** Called each time new geometry data becomes available in the recording - Mainly when a new frame is added from the Trace analysis */
 	FChaosVDGeometryDataLoaded& OnGeometryDataLoaded() { return GeometryDataLoaded; };
 
@@ -280,6 +277,9 @@ struct CHAOSVDDATA_API FChaosVDRecording
 	/** Returns true if this recording does not have any usable data */
 	bool IsEmpty() const;
 
+	/** Returns the last Platform Cycle on which this recording was updated (A new frame was added) */
+	uint64 GetLastUpdatedTimeAsCycle() { return LastUpdatedTimeAsCycle; }
+
 protected:
 
 	/** Adds an Implicit Object to the recording and takes ownership of it */
@@ -298,7 +298,7 @@ protected:
 	TMap<int32, TMap<int32, FChaosVDSolverFrameData>> GeneratedKeyFrameDataPerSolver;
 	TMap<int32, TArray<int32>> RecordedKeyFramesNumberPerSolver;
 	TArray<FChaosVDGameFrameData> GameFrames;
-	FChaosVDRecordingUpdated RecordingUpdatedDelegate;
+
 	FChaosVDGeometryDataLoaded GeometryDataLoaded;
 
 	/** Id to Ptr map of all shared geometry data required to visualize */
@@ -311,6 +311,9 @@ protected:
 	/** True if this recording is being populated from a live session */
 	bool bIsLive = false;
 
+	/** Last Platform Cycle on which this recording was updated */
+	std::atomic<uint64> LastUpdatedTimeAsCycle;
+
 	/** Map that temporary holds generated particle data during the key frame generation process, keeping its memory allocation between generated frames*/
 	TMap<int32, TSharedPtr<FChaosVDParticleDataWrapper>> ParticlesOnCurrentGeneratedKeyframe;
 
