Partially written data would be considered a cache miss

#rnx

diff --git a/Engine/Source/Runtime/Experimental/IoStoreOnDemand/Private/JournaledCache.cpp b/Engine/Source/Runtime/Experimental/IoStoreOnDemand/Private/JournaledCache.cpp
--- a/Engine/Source/Runtime/Experimental/IoStoreOnDemand/Private/JournaledCache.cpp
+++ b/Engine/Source/Runtime/Experimental/IoStoreOnDemand/Private/JournaledCache.cpp
@@ -1291,18 +1291,22 @@ uint32 FCache::GetDemand() const
 ////////////////////////////////////////////////////////////////////////////////
 bool FCache::Has(uint64 Key) const
 {
+	check(Key);
+
 	if (DiskCache.Has(Key))
 	{
 		return true;
 	}
 
 	FReadScopeLock _(MemLock);
-	return (MemCache.Get(Key) != nullptr);
+	return (MemCache.Get(Key) != nullptr) || (Partial && Partial->Key == Key);
 }
 
 ////////////////////////////////////////////////////////////////////////////////
 FCache::FGetToken FCache::Get(uint64 Key, FIoBuffer& OutData) const
 {
+	check(Key);
+
 	// Disk first as that will have more data and is more likely to hit
 	if (DiskCache.Has(Key))
 	{
@@ -1316,6 +1320,11 @@ FCache::FGetToken FCache::Get(uint64 Key, FIoBuffer& OutData) const
 		OutData = *Data;
 	}
 
+	if (Partial && Partial->Key == Key)
+	{
+		OutData = Partial->Data;
+	}
+
 	return 0;
 }
 
@@ -1457,6 +1466,11 @@ uint32 FCache::DebugVisit(void* Param, FDebugCacheEntry::Callback* Callback)
 	uint32 Count = 0;
 	Count += MemCache.DebugVisit(Param, Callback);
 	Count += DiskCache.DebugVisit(Param, Callback);
+	if (Partial)
+	{
+		Count++;
+		Callback(Param, {Partial->Key, uint32(Partial->Data.GetSize()), 1});
+	}
 	return Count;
 }
 
@@ -2610,7 +2624,7 @@ static void CacheTests(FSupport& Support)
 		check(!PostCount == !(i / Jfi));
 		check(PostCount <= PreCount);
  	}
-	NewCache();
+	NewCache(true);
 
 	// power 2
 	for (int32 i : {74, 75})
@@ -2626,7 +2640,36 @@ static void CacheTests(FSupport& Support)
 			Cache->Flush();
 		}
 		Validate();
+		NewCache(true);
+	}
+
+	// partial get
+	{
+		const uint32 Key = 0x493;
+		const uint32 DataSize = uint32(64_Ki);
+		FIoBuffer Data = Support.DummyData(DataSize);
+		Cache->Put(Key, Data);
+		for (int32 i : {0, 1})
+		{
+			Cache->WriteMemToDisk(uint32(16_Ki + (16_Ki * i)));
+			Data = FIoBuffer();
+			Cache->Get(Key, Data);
+			check(Data.GetData() != nullptr);
+			check(Data.GetSize() == DataSize);
+		}
+
+		Cache->WriteMemToDisk(uint32(16_Ki));
+		Data = FIoBuffer();
+		check(Cache->Get(Key, Data) == Key);
+		check(Data.GetData() == nullptr);
+
+		Cache->Flush();
 		NewCache();
+		Cache->Load();
+		check(Cache->Get(Key, Data) == Key);
+		check(Data.GetData() == nullptr);
+
+		NewCache(true);
 	}
 
 #if 0
@@ -2645,6 +2688,7 @@ static void CacheTests(FSupport& Support)
 	// don't load-and-sort so many paragraphs (only need max-data size)
 #endif // 0
 
+	// journal wrap
 	{
 		FCache::FConfig Config;
 		Config.Path = Support.TestDir / "cache_jrn_wrap";
@@ -2710,6 +2754,7 @@ static void CacheTests(FSupport& Support)
 		check(Found > 40);
 	}
 	
+	// journal wrap 2
 	{
 		FCache::FConfig Config;
 		Config.Path = Support.TestDir / "cache_jrn_wrap2";
