[ChaosVD] Making the Particle Data Wrapper properties visible only

The initial plan was to leave them editable, but just disable the property editor (to allow later on enabled edition in specific properties), but there is a limitation with the feature to disable the property editor ( it disables the name field widget making it difficult to read).

For now we can just make these properties visible only which will keep the name field enabled in each property. The team owning the property editor will look into removing that limitation in a future release.

diff --git a/Engine/Source/Runtime/Experimental/ChaosVisualDebugger/Public/DataWrappers/ChaosVDParticleDataWrapper.h b/Engine/Source/Runtime/Experimental/ChaosVisualDebugger/Public/DataWrappers/ChaosVDParticleDataWrapper.h
--- a/Engine/Source/Runtime/Experimental/ChaosVisualDebugger/Public/DataWrappers/ChaosVDParticleDataWrapper.h
+++ b/Engine/Source/Runtime/Experimental/ChaosVisualDebugger/Public/DataWrappers/ChaosVDParticleDataWrapper.h
@@ -86,17 +86,17 @@ struct FChaosVDFRigidParticleControlFlags : public FChaosVDWrapperDataBase
 		bHasValidData = true;
 	}
 
-	UPROPERTY(EditAnywhere, Category= "Particle Control Flags")
+	UPROPERTY(VisibleAnywhere, Category= "Particle Control Flags")
 	bool bGravityEnabled;
-	UPROPERTY(EditAnywhere, Category= "Particle Control Flags")
+	UPROPERTY(VisibleAnywhere, Category= "Particle Control Flags")
 	bool bCCDEnabled;
-	UPROPERTY(EditAnywhere, Category= "Particle Control Flags")
+	UPROPERTY(VisibleAnywhere, Category= "Particle Control Flags")
 	bool bOneWayInteractionEnabled;
-	UPROPERTY(EditAnywhere, Category = "Particle Control Flags")
+	UPROPERTY(VisibleAnywhere, Category = "Particle Control Flags")
 	bool bInertiaConditioningEnabled;
-	UPROPERTY(EditAnywhere, Category= "Particle Control Flags")
+	UPROPERTY(VisibleAnywhere, Category= "Particle Control Flags")
 	int32 GravityGroupIndex;
-	UPROPERTY(EditAnywhere, Category = "Particle Control Flags")
+	UPROPERTY(VisibleAnywhere, Category = "Particle Control Flags")
 	bool bMACDEnabled;
 };
 
@@ -138,10 +138,10 @@ struct FChaosVDParticlePositionRotation : public FChaosVDWrapperDataBase
 
 	CHAOSVDRUNTIME_API bool Serialize(FArchive& Ar);
 
-	UPROPERTY(EditAnywhere, Category= "Particle Velocities")
+	UPROPERTY(VisibleAnywhere, Category= "Particle Velocities")
 	FVector MX = FVector(ForceInit);
 
-	UPROPERTY(EditAnywhere, Category= "Particle Velocities")
+	UPROPERTY(VisibleAnywhere, Category= "Particle Velocities")
 	FQuat MR = FQuat(ForceInit);
 };
 
@@ -182,10 +182,10 @@ struct FChaosVDParticleVelocities : public FChaosVDWrapperDataBase
 		bHasValidData = true;
 	}
 	
-	UPROPERTY(EditAnywhere, Category= "Particle Velocities")
+	UPROPERTY(VisibleAnywhere, Category= "Particle Velocities")
 	FVector MV = FVector(ForceInit);
 
-	UPROPERTY(EditAnywhere, Category= "Particle Velocities")
+	UPROPERTY(VisibleAnywhere, Category= "Particle Velocities")
 	FVector MW = FVector(ForceInit);
 };
 
@@ -229,16 +229,16 @@ struct FChaosVDParticleDynamics : public FChaosVDWrapperDataBase
 		bHasValidData = true;
 	}
 
-	UPROPERTY(EditAnywhere, Category= "Particle Dynamics")
+	UPROPERTY(VisibleAnywhere, Category= "Particle Dynamics")
 	FVector MAcceleration = FVector(ForceInit);
 
-	UPROPERTY(EditAnywhere, Category= "Particle Dynamics")
+	UPROPERTY(VisibleAnywhere, Category= "Particle Dynamics")
 	FVector MAngularAcceleration = FVector(ForceInit);
 
-	UPROPERTY(EditAnywhere, Category= "Particle Dynamics")
+	UPROPERTY(VisibleAnywhere, Category= "Particle Dynamics")
 	FVector MLinearImpulseVelocity = FVector(ForceInit);
 
-	UPROPERTY(EditAnywhere, Category= "Particle Dynamics")
+	UPROPERTY(VisibleAnywhere, Category= "Particle Dynamics")
 	FVector MAngularImpulseVelocity = FVector(ForceInit);
 };
 
@@ -285,22 +285,22 @@ struct FChaosVDParticleMassProps : public FChaosVDWrapperDataBase
 		bHasValidData = true;
 	}
 	
-	UPROPERTY(EditAnywhere, Category= "Particle Mass Props")
+	UPROPERTY(VisibleAnywhere, Category= "Particle Mass Props")
 	FVector MCenterOfMass = FVector(ForceInit);
 
-	UPROPERTY(EditAnywhere, Category= "Particle Mass Props")
+	UPROPERTY(VisibleAnywhere, Category= "Particle Mass Props")
 	FQuat MRotationOfMass = FQuat(ForceInit);
 
-	UPROPERTY(EditAnywhere, Category= "Particle Mass Props")
+	UPROPERTY(VisibleAnywhere, Category= "Particle Mass Props")
 	FVector MI = FVector(ForceInit);
 
-	UPROPERTY(EditAnywhere, Category= "Particle Mass Props")
+	UPROPERTY(VisibleAnywhere, Category= "Particle Mass Props")
 	FVector MInvI = FVector(ForceInit);
 
-	UPROPERTY(EditAnywhere, Category= "Particle Mass Props")
+	UPROPERTY(VisibleAnywhere, Category= "Particle Mass Props")
 	double MM = 0.0;
 
-	UPROPERTY(EditAnywhere, Category= "Particle Mass Props")
+	UPROPERTY(VisibleAnywhere, Category= "Particle Mass Props")
 	double MInvM = 0.0;
 };
 
@@ -355,37 +355,37 @@ struct FChaosVDParticleDynamicMisc : public FChaosVDWrapperDataBase
 		bHasValidData = true;
 	}
 
-	UPROPERTY(EditAnywhere, Category= "Particle Dynamic Misc")
+	UPROPERTY(VisibleAnywhere, Category= "Particle Dynamic Misc")
 	double MAngularEtherDrag;
 
-	UPROPERTY(EditAnywhere, Category= "Particle Dynamic Misc")
+	UPROPERTY(VisibleAnywhere, Category= "Particle Dynamic Misc")
 	double MMaxLinearSpeedSq;
 	
-	UPROPERTY(EditAnywhere, Category= "Particle Dynamic Misc")
+	UPROPERTY(VisibleAnywhere, Category= "Particle Dynamic Misc")
 	double MMaxAngularSpeedSq;
 
-	UPROPERTY(EditAnywhere, Category = "Particle Dynamic Misc")
+	UPROPERTY(VisibleAnywhere, Category = "Particle Dynamic Misc")
 	float MInitialOverlapDepenetrationVelocity;
 
-	UPROPERTY(EditAnywhere, Category = "Particle Dynamic Misc")
+	UPROPERTY(VisibleAnywhere, Category = "Particle Dynamic Misc")
 	float MSleepThresholdMultiplier;
 
-	UPROPERTY(EditAnywhere, Category= "Particle Dynamic Misc")
+	UPROPERTY(VisibleAnywhere, Category= "Particle Dynamic Misc")
 	int32 MCollisionGroup;
 
-	UPROPERTY(EditAnywhere, Category= "Particle Dynamic Misc")
+	UPROPERTY(VisibleAnywhere, Category= "Particle Dynamic Misc")
 	EChaosVDObjectStateType MObjectState;
 
-	UPROPERTY(EditAnywhere, Category= "Particle Dynamic Misc")
+	UPROPERTY(VisibleAnywhere, Category= "Particle Dynamic Misc")
 	EChaosVDSleepType MSleepType;
 
-	UPROPERTY(EditAnywhere, Category= "Particle Dynamic Misc")
+	UPROPERTY(VisibleAnywhere, Category= "Particle Dynamic Misc")
 	uint32 MCollisionConstraintFlag = 0;
 
-	UPROPERTY(EditAnywhere, Category= "Particle Dynamic Misc")
+	UPROPERTY(VisibleAnywhere, Category= "Particle Dynamic Misc")
 	FChaosVDFRigidParticleControlFlags MControlFlags;
 
-	UPROPERTY(EditAnywhere, Category= "Particle Dynamic Misc")
+	UPROPERTY(VisibleAnywhere, Category= "Particle Dynamic Misc")
 	bool bDisabled;
 };
 
@@ -410,10 +410,10 @@ struct FChaosVDConnectivityEdge
 {
 	GENERATED_BODY()
 
-	UPROPERTY(EditAnywhere, Category=ConnectivityEdge)
+	UPROPERTY(VisibleAnywhere, Category=ConnectivityEdge)
 	int32 SiblingParticleID = INDEX_NONE;
 
-	UPROPERTY(EditAnywhere, Category=ConnectivityEdge)
+	UPROPERTY(VisibleAnywhere, Category=ConnectivityEdge)
 	float Strain = 0.0f;
 
 	CHAOSVDRUNTIME_API bool Serialize(FArchive& Ar)
@@ -484,43 +484,43 @@ struct FChaosVDParticleCluster : public FChaosVDWrapperDataBase
 		bHasValidData = true;
 	}
 
-	UPROPERTY(EditAnywhere, Category= "Particle Cluster")
+	UPROPERTY(VisibleAnywhere, Category= "Particle Cluster")
 	int32 ParentParticleID = INDEX_NONE;
 
-	UPROPERTY(EditAnywhere, Category= "Particle Cluster | Cluster Id")
+	UPROPERTY(VisibleAnywhere, Category= "Particle Cluster | Cluster Id")
 	int32 NumChildren = INDEX_NONE;
 	
-	UPROPERTY(EditAnywhere, Category= "Particle Cluster")
+	UPROPERTY(VisibleAnywhere, Category= "Particle Cluster")
 	FTransform ChildToParent;
 
-	UPROPERTY(EditAnywhere, Category= "Particle Cluster")
+	UPROPERTY(VisibleAnywhere, Category= "Particle Cluster")
 	int32 ClusterGroupIndex = INDEX_NONE;
 
-	UPROPERTY(EditAnywhere, Category= "Particle Cluster")
+	UPROPERTY(VisibleAnywhere, Category= "Particle Cluster")
 	bool bInternalCluster = false;
 
-	UPROPERTY(EditAnywhere, Category= "Particle Cluster")
+	UPROPERTY(VisibleAnywhere, Category= "Particle Cluster")
 	float CollisionImpulse = 0.0f;
 
-	UPROPERTY(EditAnywhere, Category= "Particle Cluster")
+	UPROPERTY(VisibleAnywhere, Category= "Particle Cluster")
 	float ExternalStrains = 0.0f;
 	
-	UPROPERTY(EditAnywhere, Category= "Particle Cluster")
+	UPROPERTY(VisibleAnywhere, Category= "Particle Cluster")
 	float InternalStrains = 0.0f;
 
-	UPROPERTY(EditAnywhere, Category= "Particle Cluster")
+	UPROPERTY(VisibleAnywhere, Category= "Particle Cluster")
 	float Strain = 0.0f;
 
-	UPROPERTY(EditAnywhere, Category= "Particle Cluster")
+	UPROPERTY(VisibleAnywhere, Category= "Particle Cluster")
 	TArray<FChaosVDConnectivityEdge> ConnectivityEdges;
 
-	UPROPERTY(EditAnywhere, Category= "Particle Cluster")
+	UPROPERTY(VisibleAnywhere, Category= "Particle Cluster")
 	bool bIsAnchored = false;
 
-	UPROPERTY(EditAnywhere, Category= "Particle Cluster")
+	UPROPERTY(VisibleAnywhere, Category= "Particle Cluster")
 	bool bUnbreakable = false;
 
-	UPROPERTY(EditAnywhere, Category= "Particle Cluster")
+	UPROPERTY(VisibleAnywhere, Category= "Particle Cluster")
 	bool bIsChildToParentLocked = false;
 };
 
@@ -568,25 +568,25 @@ struct FChaosVDParticleDataWrapper : public FChaosVDWrapperDataBase
 	UPROPERTY(VisibleAnywhere, Category= "Particle Non Frequent Data")
 	int32 SolverID = INDEX_NONE;
 
-	UPROPERTY(EditAnywhere, Category= "Particle Non Frequent Data")
+	UPROPERTY(VisibleAnywhere, Category= "Particle Non Frequent Data")
 	EChaosVDParticleType Type = EChaosVDParticleType::Unknown;
 
-	UPROPERTY(EditAnywhere, Category= "Particle Position Rotation")
+	UPROPERTY(VisibleAnywhere, Category= "Particle Position Rotation")
 	FChaosVDParticlePositionRotation ParticlePositionRotation;
 
-	UPROPERTY(EditAnywhere, Category= "Particle Particle Velocities")
+	UPROPERTY(VisibleAnywhere, Category= "Particle Particle Velocities")
 	FChaosVDParticleVelocities ParticleVelocities;
 
-	UPROPERTY(EditAnywhere, Category= "Particle Particle Dynamics")
+	UPROPERTY(VisibleAnywhere, Category= "Particle Particle Dynamics")
 	FChaosVDParticleDynamics ParticleDynamics;
 
-	UPROPERTY(EditAnywhere, Category= "Particle Particle Dynamics Misc")
+	UPROPERTY(VisibleAnywhere, Category= "Particle Particle Dynamics Misc")
 	FChaosVDParticleDynamicMisc ParticleDynamicsMisc;
 
-	UPROPERTY(EditAnywhere, Category= "Particle Particle Mass Props")
+	UPROPERTY(VisibleAnywhere, Category= "Particle Particle Mass Props")
 	FChaosVDParticleMassProps ParticleMassProps;
 
-	UPROPERTY(EditAnywhere, Category= "Particle Cluster Data")
+	UPROPERTY(VisibleAnywhere, Category= "Particle Cluster Data")
 	FChaosVDParticleCluster ParticleCluster;
 
 	UPROPERTY()
