Iris
* Added official NetRefHandleError enum
* Added net.Iris.EnsureNetRefHandleError cvar. Used to control if NetRefHandle error should cause an ensure or not.
 -1=never ensure. 0=ensure on any error. 1..X=ensure only for specific error type
 Defaults to ensure on any error

#rnx

diff --git a/Engine/Source/Runtime/Engine/Private/DataChannel.cpp b/Engine/Source/Runtime/Engine/Private/DataChannel.cpp
--- a/Engine/Source/Runtime/Engine/Private/DataChannel.cpp
+++ b/Engine/Source/Runtime/Engine/Private/DataChannel.cpp
@@ -17,6 +17,7 @@
 #include "Net/Core/Trace/Private/NetTraceInternal.h"
 #include "Net/Core/Misc/GuidReferences.h"
 #include "Net/Core/NetCoreModule.h"
+#include "Net/Core/Connection/NetEnums.h"
 #include "UObject/UObjectIterator.h"
 #include "EngineStats.h"
 #include "Engine/Engine.h"
@@ -1873,14 +1874,14 @@ void UControlChannel::ReceivedBunch( FInBunch& Bunch )
 		}
 		else if (MessageType == NMT_IrisNetRefHandleError)
 		{
-			uint32 ErrorType = 0; //TBD
+			UE::Net::ENetRefHandleError NetRefHandleError = UE::Net::ENetRefHandleError::None;
 			uint64 NetRefHandleId = 0;
-			if (FNetControlMessage<NMT_IrisNetRefHandleError>::Receive(Bunch, ErrorType, NetRefHandleId))
+			if (FNetControlMessage<NMT_IrisNetRefHandleError>::Receive(Bunch, NetRefHandleError, NetRefHandleId))
 			{
 #if UE_WITH_IRIS
 				if (UReplicationSystem* IrisRepSystem = Connection->Driver->GetReplicationSystem())
 				{
-					IrisRepSystem->ReportErrorWithNetRefHandle(ErrorType, NetRefHandleId, Connection->GetConnectionId());
+					IrisRepSystem->ReportErrorWithNetRefHandle(NetRefHandleError, NetRefHandleId, Connection->GetConnectionId());
 				}
 #endif
 			}
diff --git a/Engine/Source/Runtime/Engine/Private/Net/Experimental/Iris/DataStreamChannel.cpp b/Engine/Source/Runtime/Engine/Private/Net/Experimental/Iris/DataStreamChannel.cpp
--- a/Engine/Source/Runtime/Engine/Private/Net/Experimental/Iris/DataStreamChannel.cpp
+++ b/Engine/Source/Runtime/Engine/Private/Net/Experimental/Iris/DataStreamChannel.cpp
@@ -23,6 +23,9 @@
 #include "Net/Core/Trace/NetTrace.h"
 #include "Net/Core/Connection/NetResult.h"
 #include "Net/DataChannel.h"
+#include "Net/Iris/ReplicationSystem/ActorReplicationBridge.h"
+#include "Net/Iris/ReplicationSystem/ReplicationSystemUtil.h"
+
 
 #include "PacketHandler.h"
 #include "ProfilingDebugging/CsvProfiler.h"
@@ -136,9 +139,11 @@ void UDataStreamChannel::ReceivedBunch(FInBunch& Bunch)
 			TNetResult<ENetCloseResult> NetResult(ENetCloseResult::IrisNetRefHandleError, FString::Printf(TEXT("IrisNetRefHandleError=%s"), *SerializationContext.GetErrorHandleContext().ToString()));
 			AddToChainResultPtr(Bunch.ExtendedError, MoveTemp(NetResult));
 
-			uint32 ErrorType = 0; //TBD
-			uint64 RawHandleId = SerializationContext.GetErrorHandleContext().GetId();
-			FNetControlMessage<NMT_IrisNetRefHandleError>::Send(Connection, ErrorType, RawHandleId);
+			UActorReplicationBridge* ActorBridge = FReplicationSystemUtil::GetActorReplicationBridge(Connection);
+			if (ActorBridge)
+			{
+				ActorBridge->SendErrorWithNetRefHandle(UE::Net::ENetRefHandleError::BitstreamCorrupted, SerializationContext.GetErrorHandleContext(), Connection->GetConnectionId());
+			}
 		}
 		Bunch.SetError();
 	}
diff --git a/Engine/Source/Runtime/Engine/Private/Net/Iris/ReplicationSystem/ActorReplicationBridge.cpp b/Engine/Source/Runtime/Engine/Private/Net/Iris/ReplicationSystem/ActorReplicationBridge.cpp
--- a/Engine/Source/Runtime/Engine/Private/Net/Iris/ReplicationSystem/ActorReplicationBridge.cpp
+++ b/Engine/Source/Runtime/Engine/Private/Net/Iris/ReplicationSystem/ActorReplicationBridge.cpp
@@ -34,6 +34,7 @@
 #include "HAL/LowLevelMemStats.h"
 #include "Net/DataBunch.h"
 #include "Net/DataChannel.h"
+#include "Net/Core/Connection/NetEnums.h"
 #include "Net/Core/Connection/NetCloseResult.h"
 #include "Net/Core/Misc/NetSubObjectRegistry.h"
 #include "Net/Core/Trace/NetDebugName.h"
@@ -1212,7 +1213,7 @@ void UActorReplicationBridge::OnProtocolMismatchReported(FNetRefHandle RefHandle
 	}
 }
 
-void UActorReplicationBridge::ReportErrorWithNetRefHandle(uint32 ErrorType, FNetRefHandle RefHandle, uint32 ConnectionId)
+void UActorReplicationBridge::SendErrorWithNetRefHandle(UE::Net::ENetRefHandleError ErrorType, FNetRefHandle RefHandle, uint32 ConnectionId)
 {
 	if (NetDriver)
 	{
@@ -1223,7 +1224,7 @@ void UActorReplicationBridge::ReportErrorWithNetRefHandle(uint32 ErrorType, FNet
 		}
 		else
 		{
-			UE_LOG(LogIrisBridge, Error, TEXT("UActorReplicationBridge::ReportErrorWithNetRefHandle could not find Connection for id:%u"), ConnectionId);
+			UE_LOG(LogIrisBridge, Error, TEXT("UActorReplicationBridge::SendErrorWithNetRefHandle could not find Connection for id:%u"), ConnectionId);
 		}
 	}
 }
diff --git a/Engine/Source/Runtime/Engine/Public/Net/Iris/ReplicationSystem/ActorReplicationBridge.h b/Engine/Source/Runtime/Engine/Public/Net/Iris/ReplicationSystem/ActorReplicationBridge.h
--- a/Engine/Source/Runtime/Engine/Public/Net/Iris/ReplicationSystem/ActorReplicationBridge.h
+++ b/Engine/Source/Runtime/Engine/Public/Net/Iris/ReplicationSystem/ActorReplicationBridge.h
@@ -14,11 +14,16 @@ class UWorld;
 class UNetDriver;
 class UIrisObjectReferencePackageMap;
 
-namespace UE::Net::Private
+namespace UE::Net
 {
-	struct FActorReplicationBridgeCreationHeader;
-	struct FActorCreationHeader;
-	struct FSubObjectCreationHeader;
+	enum class ENetRefHandleError : uint32;
+
+	namespace Private
+	{
+		struct FActorReplicationBridgeCreationHeader;
+		struct FActorCreationHeader;
+		struct FSubObjectCreationHeader;
+	}
 }
 
 namespace UE::Net
@@ -79,7 +84,7 @@ public:
 	UIrisObjectReferencePackageMap* GetObjectReferencePackageMap() const { return ObjectReferencePackageMap; }
 
 	/** Tell the remote connection that we detected a reading error with a specific replicated object */
-	ENGINE_API virtual void ReportErrorWithNetRefHandle(uint32 ErrorType, FNetRefHandle RefHandle, uint32 ConnectionId) override;
+	ENGINE_API virtual void SendErrorWithNetRefHandle(UE::Net::ENetRefHandleError ErrorType, FNetRefHandle RefHandle, uint32 ConnectionId) override;
 	
 	/** Updates the level group for an actor that changed levels */
 	void ActorChangedLevel(const AActor* Actor, const ULevel* PreviousLevel);
diff --git a/Engine/Source/Runtime/Experimental/Iris/Core/Private/Iris/ReplicationSystem/NetRefHandleManager.cpp b/Engine/Source/Runtime/Experimental/Iris/Core/Private/Iris/ReplicationSystem/NetRefHandleManager.cpp
--- a/Engine/Source/Runtime/Experimental/Iris/Core/Private/Iris/ReplicationSystem/NetRefHandleManager.cpp
+++ b/Engine/Source/Runtime/Experimental/Iris/Core/Private/Iris/ReplicationSystem/NetRefHandleManager.cpp
@@ -909,4 +909,18 @@ FString FNetRefHandleManager::PrintObjectFromIndex(FInternalNetRefIndex ObjectIn
 	}
 }
 
+FString FNetRefHandleManager::PrintObjectFromNetRefHandle(FNetRefHandle ObjectHandle) const
+{ 
+	const FInternalNetRefIndex ObjectIndex = GetInternalIndex(ObjectHandle);
+	if (ObjectIndex != InvalidInternalIndex)
+	{
+		return PrintObjectFromIndex(ObjectIndex);
+	}
+	else
+	{
+		return FString::Printf(TEXT("NetObject None (InternalIndex: None) (%s)"), *ObjectHandle.ToString());
+	}
+	
+}
+
 } // end namespace UE::Net::Private
diff --git a/Engine/Source/Runtime/Experimental/Iris/Core/Private/Iris/ReplicationSystem/NetRefHandleManager.h b/Engine/Source/Runtime/Experimental/Iris/Core/Private/Iris/ReplicationSystem/NetRefHandleManager.h
--- a/Engine/Source/Runtime/Experimental/Iris/Core/Private/Iris/ReplicationSystem/NetRefHandleManager.h
+++ b/Engine/Source/Runtime/Experimental/Iris/Core/Private/Iris/ReplicationSystem/NetRefHandleManager.h
@@ -290,7 +290,7 @@ public:
 
 	/** Return a string to identify the object linked to an index in logs */
 	FString PrintObjectFromIndex(FInternalNetRefIndex ObjectIndex) const;
-	FString PrintObjectFromNetRefHandle(FNetRefHandle ObjectHandle) const { return PrintObjectFromIndex(GetInternalIndex(ObjectHandle)); }
+	FString PrintObjectFromNetRefHandle(FNetRefHandle ObjectHandle) const;
 
 	DECLARE_MULTICAST_DELEGATE_OneParam(FOnLargestIndexIncrease, uint32 LargestIndex);
 
diff --git a/Engine/Source/Runtime/Experimental/Iris/Core/Private/Iris/ReplicationSystem/ObjectReplicationBridge.cpp b/Engine/Source/Runtime/Experimental/Iris/Core/Private/Iris/ReplicationSystem/ObjectReplicationBridge.cpp
--- a/Engine/Source/Runtime/Experimental/Iris/Core/Private/Iris/ReplicationSystem/ObjectReplicationBridge.cpp
+++ b/Engine/Source/Runtime/Experimental/Iris/Core/Private/Iris/ReplicationSystem/ObjectReplicationBridge.cpp
@@ -16,6 +16,7 @@
 
 #include "Net/Core/NetBitArrayPrinter.h"
 #include "Net/Core/NetHandle/NetHandleManager.h"
+#include "Net/Core/Connection/NetEnums.h"
 #include "Net/Core/PropertyConditions/PropertyConditionsDelegates.h"
 #include "Net/Core/Trace/NetTrace.h"
 #include "Net/Core/Trace/NetDebugName.h"
@@ -89,6 +90,13 @@ static FAutoConsoleVariableRef CVarEnableForceNetUpdate(
 	TEXT("When true the system only allows ForceNetUpdate to skip the poll frequency of objects. When false any MarkDirty object will be immediately polled.")
 );
 
+static int32 GEnsureNetRefHandleError = 0;
+static FAutoConsoleVariableRef CvarEnsureNetRefHandleError(
+	TEXT("net.Iris.EnsureNetRefHandleError"),
+	GEnsureNetRefHandleError,
+	TEXT("Chooses if we should ensure when a NetRefHandleError was reported. -1=never ensure. 0=always ensure. 1..X=ensure only for specific error type")
+);
+
 UObjectReplicationBridge::FCreateNetRefHandleParams UObjectReplicationBridge::DefaultCreateNetRefHandleParams =
 {
 	.bCanReceive=false, 
@@ -141,6 +149,8 @@ void UObjectReplicationBridge::Initialize(UReplicationSystem* InReplicationSyste
 {
 	LLM_SCOPE_BYTAG(Iris);
 
+	ensureMsgf((GEnsureNetRefHandleError == -1 || GEnsureNetRefHandleError < (int32)UE::Net::ENetRefHandleError::Max), TEXT("GEnsureNetRefHandleError is set to an invalid value: %u"), GEnsureNetRefHandleError);
+
 	Super::Initialize(InReplicationSystem);
 
 	const uint32 MaxActiveObjectCount = NetRefHandleManager->GetMaxActiveObjectCount();
@@ -2025,7 +2035,7 @@ void UObjectReplicationBridge::OnProtocolMismatchReported(FNetRefHandle RefHandl
 	}
 }
 
-void UObjectReplicationBridge::OnErrorWithNetRefHandleReported(uint32 ErrorType, FNetRefHandle RefHandle, uint32 ConnectionId)
+void UObjectReplicationBridge::OnErrorWithNetRefHandleReported(UE::Net::ENetRefHandleError ErrorType, FNetRefHandle RefHandle, uint32 ConnectionId)
 {
 	using namespace UE::Net;
 	using namespace UE::Net::Private;
@@ -2033,20 +2043,35 @@ void UObjectReplicationBridge::OnErrorWithNetRefHandleReported(uint32 ErrorType,
 	// Ensure at the end so the log contains all the relevant information
 	ON_SCOPE_EXIT
 	{
-		ensureMsgf(false, TEXT("NetRefHandle error: %u reported for %s. Look at the log for important information on the object tied to the handle."), ErrorType, *RefHandle.ToString());
+		if (GEnsureNetRefHandleError==0 || GEnsureNetRefHandleError==(int32)ErrorType)
+		{
+			// Use different calls for every error type to prevent only reporting the first error that occured.
+			switch(ErrorType)
+			{
+				case ENetRefHandleError::BitstreamCorrupted:
+				{
+					ensureMsgf(false, TEXT("NetRefHandle error: Bitstream corrupted reported with %s. Look at the log for important information on the object tied to the handle."), *NetRefHandleManager->PrintObjectFromNetRefHandle(RefHandle));
+				} break;
+				case ENetRefHandleError::ReplicationDisabled:
+				{
+					ensureMsgf(false, TEXT("NetRefHandle error: Replication disabled for %s. Client was unable to read the creation data and did not spawn the netobject."), *NetRefHandleManager->PrintObjectFromNetRefHandle(RefHandle));
+				} break;
+				default:
+				{
+					ensureMsgf(false, TEXT("NetRefHandle error: %s reported for %s. This should not occur. Missing switch case or wrong error type used."), LexToString(ErrorType), *NetRefHandleManager->PrintObjectFromNetRefHandle(RefHandle));
+				} break;
+			}
+		}
 	};
 
 	const FInternalNetRefIndex ObjectInternalIndex = NetRefHandleManager->GetInternalIndex(RefHandle);
 	if (ObjectInternalIndex == FNetRefHandleManager::InvalidInternalIndex)
 	{
-		UE_LOG(LogIris, Warning, TEXT("OnErrorWithNetRefHandleReported: %u from client %s for %s but object has no InternalIndex."), ErrorType, *PrintConnectionInfo(ConnectionId), *RefHandle.ToString());
+		UE_LOG(LogIris, Warning, TEXT("OnErrorWithNetRefHandleReported: %s from client %s for %s but object has no InternalIndex."), LexToString(ErrorType), *PrintConnectionInfo(ConnectionId), *RefHandle.ToString());
 		return;
 	}
 
-	UE_LOG(LogIris, Error, TEXT("OnErrorWithNetRefHandleReported: %u from client %s. Problematic object was %s"), 
-		ErrorType, *PrintConnectionInfo(ConnectionId),
-		*NetRefHandleManager->PrintObjectFromIndex(ObjectInternalIndex)
-	);
+	UE_LOG(LogIris, Error, TEXT("OnErrorWithNetRefHandleReported: %s from client %s. Problematic object was %s"),  LexToString(ErrorType), *PrintConnectionInfo(ConnectionId), *NetRefHandleManager->PrintObjectFromIndex(ObjectInternalIndex));
 }
 
 TArray<uint32> UObjectReplicationBridge::FindConnectionsFromArgs(const TArray<FString>& Args) const
diff --git a/Engine/Source/Runtime/Experimental/Iris/Core/Private/Iris/ReplicationSystem/ReplicationReader.cpp b/Engine/Source/Runtime/Experimental/Iris/Core/Private/Iris/ReplicationSystem/ReplicationReader.cpp
--- a/Engine/Source/Runtime/Experimental/Iris/Core/Private/Iris/ReplicationSystem/ReplicationReader.cpp
+++ b/Engine/Source/Runtime/Experimental/Iris/Core/Private/Iris/ReplicationSystem/ReplicationReader.cpp
@@ -745,11 +745,8 @@ uint32 FReplicationReader::ReadObjectBatch(FNetSerializationContext& Context)
 		{
 			if (Context.GetError() == GNetError_BrokenNetHandle)
 			{
-				const uint32 ErrorType = 1; //TBD
-				ReplicationBridge->ReportErrorWithNetRefHandle(ErrorType, IncompleteHandle, Parameters.ConnectionId);
+				ReplicationBridge->SendErrorWithNetRefHandle(UE::Net::ENetRefHandleError::ReplicationDisabled, IncompleteHandle, Parameters.ConnectionId);
  
-				// $TODO: Report this to the server so it knows that the state of data in the batch is unknown
-
 				// Log error and try to recover, if get more incoming data for an object in the broken state we will skip it.
 				UE_LOG(LogIris, Error, TEXT("FReplicationReader::ReadObject Failed to read object batch handle: %s skipping batch data"), ToCStr(IncompleteHandle.ToString()));
 					
@@ -2369,4 +2366,4 @@ void FReplicationReader::SetRemoteNetTokenStoreState(FNetTokenStoreState* InRemo
 	ResolveContext.RemoteNetTokenStoreState = InRemoteTokenStoreState;
 }
 
-}
+} // end namespace UE::Net::Private
diff --git a/Engine/Source/Runtime/Experimental/Iris/Core/Private/Iris/ReplicationSystem/ReplicationSystem.cpp b/Engine/Source/Runtime/Experimental/Iris/Core/Private/Iris/ReplicationSystem/ReplicationSystem.cpp
--- a/Engine/Source/Runtime/Experimental/Iris/Core/Private/Iris/ReplicationSystem/ReplicationSystem.cpp
+++ b/Engine/Source/Runtime/Experimental/Iris/Core/Private/Iris/ReplicationSystem/ReplicationSystem.cpp
@@ -7,8 +7,11 @@
 #include "Iris/Core/IrisMemoryTracker.h"
 #include "Iris/Core/IrisProfiler.h"
 #include "Iris/Core/NetObjectReference.h"
+
 #include "Net/Core/Trace/NetTrace.h"
 #include "Net/Core/Misc/NetConditionGroupManager.h"
+#include "Net/Core/Connection/NetEnums.h"
+
 #include "Iris/ReplicationState/ReplicationStateUtil.h"
 #include "Iris/ReplicationSystem/ChangeMaskCache.h"
 #include "Iris/ReplicationSystem/Filtering/NetObjectFilter.h"
@@ -24,6 +27,7 @@
 #include "Iris/Serialization/NetBitStreamWriter.h"
 #include "Iris/Serialization/NetSerializer.h"
 #include "Iris/Serialization/IrisObjectReferencePackageMap.h"
+
 #include "HAL/IConsoleManager.h"
 #include "ProfilingDebugging/CsvProfiler.h"
 #include "UObject/UObjectGlobals.h"
@@ -1661,7 +1665,7 @@ void UReplicationSystem::ReportProtocolMismatch(uint64 NetRefHandleId, uint32 Co
 	Impl->ReplicationSystemInternal.GetReplicationBridge()->OnProtocolMismatchReported(NetRefHandle, ConnectionId);
 }
 
-void UReplicationSystem::ReportErrorWithNetRefHandle(uint32 ErrorType, uint64 NetRefHandleId, uint32 ConnectionId)
+void UReplicationSystem::ReportErrorWithNetRefHandle(UE::Net::ENetRefHandleError ErrorType, uint64 NetRefHandleId, uint32 ConnectionId)
 {
 	using namespace UE::Net::Private;
 	const FNetRefHandle NetRefHandle = FNetRefHandleManager::MakeNetRefHandle(NetRefHandleId, GetId());
diff --git a/Engine/Source/Runtime/Experimental/Iris/Core/Public/Iris/ReplicationSystem/ObjectReplicationBridge.h b/Engine/Source/Runtime/Experimental/Iris/Core/Public/Iris/ReplicationSystem/ObjectReplicationBridge.h
--- a/Engine/Source/Runtime/Experimental/Iris/Core/Public/Iris/ReplicationSystem/ObjectReplicationBridge.h
+++ b/Engine/Source/Runtime/Experimental/Iris/Core/Public/Iris/ReplicationSystem/ObjectReplicationBridge.h
@@ -15,6 +15,8 @@ DECLARE_LOG_CATEGORY_EXTERN(LogIrisFilterConfig, Log, All);
 
 namespace UE::Net
 {
+	enum class UE::Net::ENetRefHandleError : uint32;
+
 	struct FNetObjectResolveContext;
 	typedef uint32 FNetObjectFilterHandle;
 	typedef uint32 FNetObjectPrioritizerHandle;
@@ -210,7 +212,7 @@ protected:
 	IRISCORE_API virtual void DetachInstanceFromRemote(FNetRefHandle Handle, EReplicationBridgeDestroyInstanceReason DestroyReason, EReplicationBridgeDestroyInstanceFlags DestroyFlags) override;
 	IRISCORE_API virtual void DetachInstance(FNetRefHandle Handle) override;
 	IRISCORE_API virtual void OnProtocolMismatchReported(FNetRefHandle RefHandle, uint32 ConnectionId) override;
-	IRISCORE_API virtual void OnErrorWithNetRefHandleReported(uint32 ErrorType, FNetRefHandle RefHandle, uint32 ConnectionId) override;
+	IRISCORE_API virtual void OnErrorWithNetRefHandleReported(UE::Net::ENetRefHandleError ErrorType, FNetRefHandle RefHandle, uint32 ConnectionId) override;
 	
 protected:
 	/**
diff --git a/Engine/Source/Runtime/Experimental/Iris/Core/Public/Iris/ReplicationSystem/ReplicationBridge.h b/Engine/Source/Runtime/Experimental/Iris/Core/Public/Iris/ReplicationSystem/ReplicationBridge.h
--- a/Engine/Source/Runtime/Experimental/Iris/Core/Public/Iris/ReplicationSystem/ReplicationBridge.h
+++ b/Engine/Source/Runtime/Experimental/Iris/Core/Public/Iris/ReplicationSystem/ReplicationBridge.h
@@ -21,6 +21,8 @@ class UNetDriver;
 
 namespace UE::Net
 {
+	enum class UE::Net::ENetRefHandleError : uint32;
+
 	struct FNetDependencyInfo;
 	class FNetTokenStoreState;
 	class FReplicationFragment;
@@ -269,10 +271,10 @@ protected:
 	IRISCORE_API virtual void OnProtocolMismatchReported(FNetRefHandle RefHandle, uint32 ConnectionId) {}
 
 	/** Called when a remote connection has a critical error caused by a specific NetRefHandle */
-	IRISCORE_API virtual void OnErrorWithNetRefHandleReported(uint32 ErrorType, FNetRefHandle RefHandle, uint32 ConnectionId) {}
+	IRISCORE_API virtual void OnErrorWithNetRefHandleReported(UE::Net::ENetRefHandleError ErrorType, FNetRefHandle RefHandle, uint32 ConnectionId) {}
 
 	/** Tell the remote connection that we detected a reading error with a specific replicated object */
-	IRISCORE_API virtual void ReportErrorWithNetRefHandle(uint32 ErrorType, FNetRefHandle RefHandle, uint32 ConnectionId) {}
+	IRISCORE_API virtual void SendErrorWithNetRefHandle(UE::Net::ENetRefHandleError ErrorType, FNetRefHandle RefHandle, uint32 ConnectionId) {}
 
 private:
 
diff --git a/Engine/Source/Runtime/Experimental/Iris/Core/Public/Iris/ReplicationSystem/ReplicationSystem.h b/Engine/Source/Runtime/Experimental/Iris/Core/Public/Iris/ReplicationSystem/ReplicationSystem.h
--- a/Engine/Source/Runtime/Experimental/Iris/Core/Public/Iris/ReplicationSystem/ReplicationSystem.h
+++ b/Engine/Source/Runtime/Experimental/Iris/Core/Public/Iris/ReplicationSystem/ReplicationSystem.h
@@ -27,6 +27,7 @@ namespace UE::Net
 	class FNetBitArray;
 	class FNetCullDistanceOverrides;
 	enum class ENetFilterStatus : uint32;
+	enum class ENetRefHandleError : uint32;
 	class FNetObjectAttachment;
 	enum class ENetObjectDeltaCompressionStatus : unsigned;
 	typedef uint32 FNetObjectFilterHandle;
@@ -604,7 +605,7 @@ public:
 	IRISCORE_API void ReportProtocolMismatch(uint64 NetRefHandleId, uint32 ConnectionId);
 
 	/** Called when a connection reports a critical error with a netrefhandle object */
-	IRISCORE_API void ReportErrorWithNetRefHandle(uint32 ErrorType, uint64 NetRefHandleId, uint32 ConnectionId);
+	IRISCORE_API void ReportErrorWithNetRefHandle(UE::Net::ENetRefHandleError ErrorType, uint64 NetRefHandleId, uint32 ConnectionId);
 
 public:
 	// For internal use and not exported.
diff --git a/Engine/Source/Runtime/Net/Core/Public/Net/Core/Connection/NetEnums.h b/Engine/Source/Runtime/Net/Core/Public/Net/Core/Connection/NetEnums.h
--- a/Engine/Source/Runtime/Net/Core/Public/Net/Core/Connection/NetEnums.h
+++ b/Engine/Source/Runtime/Net/Core/Public/Net/Core/Connection/NetEnums.h
@@ -113,3 +113,32 @@ namespace ESecurityEvent
 #if UE_ENABLE_INCLUDE_ORDER_DEPRECATED_IN_5_2
 #include "CoreMinimal.h"
 #endif
+
+
+namespace UE::Net
+{
+
+/**
+ * Error types used when reporting an error on a specific Iris handle
+ */
+enum class ENetRefHandleError : uint32
+{
+	None = 0,
+	BitstreamCorrupted, // An error while reading detected an irrecoverable bitstream corruption. Client was forced to disconnect
+	ReplicationDisabled, // An error while reading was recoverable but prevents this NetObject from replicating.
+	Max, // always keep last
+};
+
+inline const TCHAR* LexToString(ENetRefHandleError NetRefHandleError)
+{
+	switch (NetRefHandleError)
+	{
+		case ENetRefHandleError::None: return TEXT("None");
+		case ENetRefHandleError::BitstreamCorrupted: return TEXT("BitstreamCorrupted");
+		case ENetRefHandleError::ReplicationDisabled: return TEXT("ReplicationDisabled");
+		case ENetRefHandleError::Max: return TEXT("Max");
+		default: return TEXT("missing");
+	}
+}
+
+} // end namespace UE::Net
