Iris - Fix storage calculation.

#rnx

diff --git a/Engine/Source/Runtime/Experimental/Iris/Core/Private/Iris/ReplicationSystem/ObjectPollFrequencyLimiter.cpp b/Engine/Source/Runtime/Experimental/Iris/Core/Private/Iris/ReplicationSystem/ObjectPollFrequencyLimiter.cpp
--- a/Engine/Source/Runtime/Experimental/Iris/Core/Private/Iris/ReplicationSystem/ObjectPollFrequencyLimiter.cpp
+++ b/Engine/Source/Runtime/Experimental/Iris/Core/Private/Iris/ReplicationSystem/ObjectPollFrequencyLimiter.cpp
@@ -4,6 +4,7 @@
 #include "Iris/Core/IrisProfiler.h"
 #include "HAL/IConsoleManager.h"
 #include "Math/VectorRegister.h"
+#include "Templates/AlignmentTemplates.h"
 
 namespace UE::Net::Private
 {
@@ -15,7 +16,7 @@ FObjectPollFrequencyLimiter::FObjectPollFrequencyLimiter()
 void FObjectPollFrequencyLimiter::Init(uint32 MaxActiveObjectCount)
 {
 	// We want to be able to easily update 32 objects at a time.
-	const uint32 StorageObjectCount = (MaxActiveObjectCount + 31U) & ~(MaxActiveObjectCount - 1U);
+	const uint32 StorageObjectCount = Align(MaxActiveObjectCount, 32U);
 
 	// Allocate max amount of memory required.
 	FramesBetweenUpdates.SetNumZeroed(StorageObjectCount);
