Fix bad nested namespace in forward declaration

#rnx

diff --git a/Engine/Source/Runtime/Experimental/Iris/Core/Public/Iris/ReplicationSystem/ObjectReplicationBridge.h b/Engine/Source/Runtime/Experimental/Iris/Core/Public/Iris/ReplicationSystem/ObjectReplicationBridge.h
--- a/Engine/Source/Runtime/Experimental/Iris/Core/Public/Iris/ReplicationSystem/ObjectReplicationBridge.h
+++ b/Engine/Source/Runtime/Experimental/Iris/Core/Public/Iris/ReplicationSystem/ObjectReplicationBridge.h
@@ -15,7 +15,7 @@ DECLARE_LOG_CATEGORY_EXTERN(LogIrisFilterConfig, Log, All);
 
 namespace UE::Net
 {
-	enum class UE::Net::ENetRefHandleError : uint32;
+	enum class ENetRefHandleError : uint32;
 
 	struct FNetObjectResolveContext;
 	typedef uint32 FNetObjectFilterHandle;
diff --git a/Engine/Source/Runtime/Experimental/Iris/Core/Public/Iris/ReplicationSystem/ReplicationBridge.h b/Engine/Source/Runtime/Experimental/Iris/Core/Public/Iris/ReplicationSystem/ReplicationBridge.h
--- a/Engine/Source/Runtime/Experimental/Iris/Core/Public/Iris/ReplicationSystem/ReplicationBridge.h
+++ b/Engine/Source/Runtime/Experimental/Iris/Core/Public/Iris/ReplicationSystem/ReplicationBridge.h
@@ -21,7 +21,7 @@ class UNetDriver;
 
 namespace UE::Net
 {
-	enum class UE::Net::ENetRefHandleError : uint32;
+	enum class ENetRefHandleError : uint32;
 
 	struct FNetDependencyInfo;
 	class FNetTokenStoreState;
