[ChaosVD] Changing the file name generation to be a name based on a time stamp and let UE trace define the full path.

The old implementation tried to create a name by appending a prefix "ChaosVD" to a number. That number was determined by trying file names until hitting a combination that does not exists (that implementation was taken as is from day one from the old serialization code for the Solver Evolution).

It seems there is a rare case where that will return a name UE trace cannot use because it thinks it already exists. Looking at the UE trace code, it does several validations to the path we are not making. Also, other tools seems to just provide the file name and let UE trace handle build the full path.

Therefore I am changing the file name generation code to just provide a base filename (with ChaosVD- appended), and using a timestamp to make it unique, like UE trace does. With this change the UE Trace code will handle building the full path.

The CVD files will still be written to the Profiling folder

#tests I tested that I can start a recording to file, and the file will be written to the profiling folder as expected, and that I can also auto load them in CVD when the recording stops

diff --git a/Engine/Plugins/ChaosVD/Source/ChaosVD/Private/Widgets/SChaosVDRecordingControls.cpp b/Engine/Plugins/ChaosVD/Source/ChaosVD/Private/Widgets/SChaosVDRecordingControls.cpp
--- a/Engine/Plugins/ChaosVD/Source/ChaosVD/Private/Widgets/SChaosVDRecordingControls.cpp
+++ b/Engine/Plugins/ChaosVD/Source/ChaosVD/Private/Widgets/SChaosVDRecordingControls.cpp
@@ -262,7 +262,7 @@ void SChaosVDRecordingControls::HandleRecordingStop()
 		}
 		else
 		{
-			const FText RecordingPathMessage = FText::Format(LOCTEXT("RecordingSavedPathMessage"," Recoring saved at {0} "), FText::AsCultureInvariant(FChaosVDRuntimeModule::Get().GetActiveRecordingFileName()));
+			const FText RecordingPathMessage = FText::Format(LOCTEXT("RecordingSavedPathMessage"," Recoring saved at {0} "), FText::AsCultureInvariant(FChaosVDRuntimeModule::Get().GetLastRecordingFileNamePath()));
 			RecordingPathMessageHandle = StatusBarSubsystem->PushStatusBarMessage(StatusBarID, RecordingPathMessage);
 		}
 	}
@@ -271,7 +271,7 @@ void SChaosVDRecordingControls::HandleRecordingStop()
 	{
 		if (FMessageDialog::Open(EAppMsgType::YesNo, LOCTEXT("OpenLastRecordingMessage", "Do you want to load the recorded file now? ")) == EAppReturnType::Yes)
 		{
-			MainTabSharedPtr->GetChaosVDEngineInstance()->LoadRecording(FChaosVDRuntimeModule::Get().GetActiveRecordingFileName());
+			MainTabSharedPtr->GetChaosVDEngineInstance()->LoadRecording(FChaosVDRuntimeModule::Get().GetLastRecordingFileNamePath());
 		}
 	}
 }
diff --git a/Engine/Source/Runtime/Experimental/ChaosVisualDebugger/Private/ChaosVDRuntimeModule.cpp b/Engine/Source/Runtime/Experimental/ChaosVisualDebugger/Private/ChaosVDRuntimeModule.cpp
--- a/Engine/Source/Runtime/Experimental/ChaosVisualDebugger/Private/ChaosVDRuntimeModule.cpp
+++ b/Engine/Source/Runtime/Experimental/ChaosVisualDebugger/Private/ChaosVDRuntimeModule.cpp
@@ -98,6 +98,11 @@ int32 FChaosVDRuntimeModule::GenerateUniqueID()
 	return LastGeneratedID.Increment();
 }
 
+FString FChaosVDRuntimeModule::GetLastRecordingFileNamePath() const
+{
+	return LastRecordingFileNamePath;
+}
+
 void FChaosVDRuntimeModule::StopTrace()
 {
 	bRequestedStop = true;
@@ -106,14 +111,7 @@ void FChaosVDRuntimeModule::StopTrace()
 
 void FChaosVDRuntimeModule::GenerateRecordingFileName(FString& OutFileName)
 {
-	const TCHAR* FilePrefix = TEXT("ChaosVD");
-	const FString FullPathPrefix = FPaths::ProfilingDir() / FilePrefix;
-
-	int32 Tries = 0;
-	do
-	{
-		OutFileName = FString::Printf(TEXT("%s_%d.utrace"), *FullPathPrefix, Tries++);
-	} while (IFileManager::Get().FileExists(*OutFileName));
+	OutFileName = TEXT("ChaosVD-") + FDateTime::Now().ToString(TEXT("%Y%m%d_%H%M%S.utrace"));
 }
 
 bool FChaosVDRuntimeModule::RequestFullCapture(float DeltaTime)
@@ -179,10 +177,14 @@ void FChaosVDRuntimeModule::StartRecording(TConstArrayView<FString> Args)
 
 	if (Args.Num() == 0 || Args[0] == TEXT("File"))
 	{
-		ActiveRecordingFileName.Empty();
-		GenerateRecordingFileName(ActiveRecordingFileName);
+		LastRecordingFileNamePath.Empty();
+		GenerateRecordingFileName(LastRecordingFileNamePath);
 
-		bIsRecording = FTraceAuxiliary::Start(FTraceAuxiliary::EConnectionType::File, *ActiveRecordingFileName, nullptr, &TracingOptions);
+		UE_LOG(LogChaosVDRuntime, Log, TEXT("[%s] Generated trace file name [%s]"), ANSI_TO_TCHAR(__FUNCTION__), *LastRecordingFileNamePath);
+
+		bIsRecording = FTraceAuxiliary::Start(FTraceAuxiliary::EConnectionType::File, *LastRecordingFileNamePath, nullptr, &TracingOptions);
+
+		LastRecordingFileNamePath = bIsRecording ? FTraceAuxiliary::GetTraceDestinationString() : TEXT("");
 	}
 	else if(Args[0] == TEXT("Server"))
 	{
@@ -233,7 +235,10 @@ void FChaosVDRuntimeModule::StartRecording(TConstArrayView<FString> Args)
 		FMessageDialog::Open(EAppMsgType::Ok, FText::FormatOrdered(LOCTEXT("StartRecordingFailedMessage", "Failed to start CVD recording. \n\n{0}"), FailureReason));
 #endif
 
-		RecordingStartFailedDelegate.Broadcast(FailureReason);
+		{
+			FReadScopeLock ReadLock(DelegatesRWLock);
+			RecordingStartFailedDelegate.Broadcast(FailureReason);
+		}	
 	}
 
 }
diff --git a/Engine/Source/Runtime/Experimental/ChaosVisualDebugger/Public/ChaosVDRuntimeModule.h b/Engine/Source/Runtime/Experimental/ChaosVisualDebugger/Public/ChaosVDRuntimeModule.h
--- a/Engine/Source/Runtime/Experimental/ChaosVisualDebugger/Public/ChaosVDRuntimeModule.h
+++ b/Engine/Source/Runtime/Experimental/ChaosVisualDebugger/Public/ChaosVDRuntimeModule.h
@@ -98,8 +98,8 @@ public:
 	/** Returns the accumulated recording time in seconds since the recording started */
 	float GetAccumulatedRecordingTime() const { return AccumulatedRecordingTime; }
 
-	/** Returns the relative file name path of the active recording */
-	const FString& GetActiveRecordingFileName() const { return ActiveRecordingFileName; }
+	/** Returns the full path of the active recording file*/
+	FString GetLastRecordingFileNamePath() const;
 
 private:
 
@@ -134,7 +134,7 @@ private:
 
 	FThreadSafeCounter LastGeneratedID;
 
-	FString ActiveRecordingFileName;
+	FString LastRecordingFileNamePath;
 
 	static FRWLock DelegatesRWLock;
 };
