Fix remove on sleep not working properly anymore with geometry collection
- make sure the dynamic collection is mark dirty after updating the scaled transform
- make sure we only scale relative to the previouly scaled transform

diff --git a/Engine/Source/Runtime/Experimental/GeometryCollectionEngine/Private/GeometryCollection/GeometryCollectionComponent.cpp b/Engine/Source/Runtime/Experimental/GeometryCollectionEngine/Private/GeometryCollection/GeometryCollectionComponent.cpp
--- a/Engine/Source/Runtime/Experimental/GeometryCollectionEngine/Private/GeometryCollection/GeometryCollectionComponent.cpp
+++ b/Engine/Source/Runtime/Experimental/GeometryCollectionEngine/Private/GeometryCollection/GeometryCollectionComponent.cpp
@@ -4563,13 +4563,19 @@ void UGeometryCollectionComponent::UpdateRemovalIfNeeded()
 							ShrinkRadius = -AccumulatedSphere.W;
 						}
 
+						const FTransform3f& CurrentTransform = DynamicCollection->GetTransform(TransformIndex);
+						const float CurrentScale = CurrentTransform.GetScale3D().X; // always uniform scale
+
+						const float ScaleToApply = (CurrentScale >  SMALL_NUMBER) ? (Scale / CurrentScale): Scale;
+
 						// InverseComponentTransform calculation must in double precision for LWC
 						const FQuat LocalRotation = (InverseComponentTransform * FTransform((*CompSpaceTransform)[TransformIndex].Inverse())).GetRotation();
 						const FVector LocalDown = LocalRotation.RotateVector(FVector(0.f, 0.f, ShrinkRadius));
 						const FVector CenterOfMass = (*MassToLocal)[TransformIndex].GetTranslation();
 						const FVector3f ScaleCenter = FVector3f(LocalDown + CenterOfMass);
-						const FTransform3f ScaleTransform(FQuat4f::Identity, ScaleCenter * FVector3f::FReal(1.f - Scale), FVector3f(Scale));
-						DynamicCollection->SetTransform(TransformIndex, ScaleTransform * DynamicCollection->GetTransform(TransformIndex));
+						const FTransform3f ScaleTransform(FQuat4f::Identity, ScaleCenter * FVector3f::FReal(1.f - ScaleToApply), FVector3f(ScaleToApply));
+						const FTransform3f FinalTransform = ScaleTransform * CurrentTransform;
+						DynamicCollection->SetTransform(TransformIndex, FinalTransform);
 						bTransformsChanged = true;
 					}
 				}
@@ -4577,6 +4583,7 @@ void UGeometryCollectionComponent::UpdateRemovalIfNeeded()
 
 			if (bTransformsChanged)
 			{
+				DynamicCollection->MakeDirty();
 				OnTransformsDirty();
 			}
 		}
