Fix level 1 clusters not breaking into level 2 clusters
- make sure the level1 cluster do not get removed from internal structure

diff --git a/Engine/Source/Runtime/Experimental/Chaos/Private/Chaos/PBDRigidClustering.cpp b/Engine/Source/Runtime/Experimental/Chaos/Private/Chaos/PBDRigidClustering.cpp
--- a/Engine/Source/Runtime/Experimental/Chaos/Private/Chaos/PBDRigidClustering.cpp
+++ b/Engine/Source/Runtime/Experimental/Chaos/Private/Chaos/PBDRigidClustering.cpp
@@ -1510,14 +1510,7 @@ namespace Chaos
 							// just got re-enabled in RemoveParticlesFromCluster.
 							for (FPBDRigidParticleHandle* Particle : Island)
 							{
-								if (FPBDRigidClusteredParticleHandle* ClusterParticle = Particle->CastToClustered())
-								{
-									DisableCluster(ClusterParticle);
-								}
-								else
-								{
-									MEvolution.DisableParticle(Particle);
-								}
+								MEvolution.DisableParticle(Particle);
 							}
 						}
 					}
