Tentative fix for server crash in geometry collection physics proxy initializtion on the physics thread
- make sure indices ate properly tested

diff --git a/Engine/Source/Runtime/Experimental/Chaos/Private/PhysicsProxy/GeometryCollectionPhysicsProxy.cpp b/Engine/Source/Runtime/Experimental/Chaos/Private/PhysicsProxy/GeometryCollectionPhysicsProxy.cpp
--- a/Engine/Source/Runtime/Experimental/Chaos/Private/PhysicsProxy/GeometryCollectionPhysicsProxy.cpp
+++ b/Engine/Source/Runtime/Experimental/Chaos/Private/PhysicsProxy/GeometryCollectionPhysicsProxy.cpp
@@ -1411,7 +1411,7 @@ void FGeometryCollectionPhysicsProxy::InitializeBodiesPT(Chaos::FPBDRigidsSolver
 					// Cluster parent
 					GameThreadCollection.IterateThroughChildren(TransformGroupIndex, [&](int32 ChildIndex)
 					{
-						if (SubTreeContainsSimulatableParticle[ChildIndex])
+						if (SubTreeContainsSimulatableParticle.IsValidIndex(ChildIndex) && SubTreeContainsSimulatableParticle[ChildIndex])
 						{
 							SubTreeContainsSimulatableParticle[TransformGroupIndex] = true;
 							return false;
diff --git a/Engine/Source/Runtime/Experimental/Chaos/Public/GeometryCollectionProxyData.h b/Engine/Source/Runtime/Experimental/Chaos/Public/GeometryCollectionProxyData.h
--- a/Engine/Source/Runtime/Experimental/Chaos/Public/GeometryCollectionProxyData.h
+++ b/Engine/Source/Runtime/Experimental/Chaos/Public/GeometryCollectionProxyData.h
@@ -45,15 +45,18 @@ public:
 	template<typename Lambda>
 	void IterateThroughChildren(int32 Index, Lambda&& LambdaIt) const
 	{
-		const TSet<int32>& Children = RestCollectionShared->Children[Index];
-		for (const int32 Child : Children)
-		{
-			if (GetHasParent(Child))
+		if (RestCollectionShared && RestCollectionShared->Children.IsValidIndex(Index))
+		{	
+			const TSet<int32>& Children = RestCollectionShared->Children[Index];
+			for (const int32 Child : Children)
 			{
-				bool bContinue = LambdaIt(Child);
-				if (!bContinue)
+				if (GetHasParent(Child))
 				{
-					break;
+					bool bContinue = LambdaIt(Child);
+					if (!bContinue)
+					{
+						break;
+					}
 				}
 			}
 		}
