GeometryCollection : SkeletalMeshConversion
- Update the transform attributes to reflect the SkeletalMesh transform hierarchy.

diff --git a/Engine/Source/Runtime/Experimental/GeometryCollectionEngine/Private/GeometryCollection/GeometryCollectionEngineConversion.cpp b/Engine/Source/Runtime/Experimental/GeometryCollectionEngine/Private/GeometryCollection/GeometryCollectionEngineConversion.cpp
--- a/Engine/Source/Runtime/Experimental/GeometryCollectionEngine/Private/GeometryCollection/GeometryCollectionEngineConversion.cpp
+++ b/Engine/Source/Runtime/Experimental/GeometryCollectionEngine/Private/GeometryCollection/GeometryCollectionEngineConversion.cpp
@@ -1159,7 +1159,7 @@ bool FGeometryCollectionEngineConversion::AppendSkeletalMesh(const USkeletalMesh
 		}
 	}
 	const FSkelMeshRenderSection& RenderSection = SkeletalMeshLODRenderData.RenderSections[0];
-	const TArray<FBoneIndexType>& SkeletalBoneMap = RenderSection.BoneMap;
+	//const TArray<FBoneIndexType>& SkeletalBoneMap = RenderSection.BoneMap;
 
 	//
 	// The Component transform for each Mesh will become the FTransform that drives
@@ -1170,15 +1170,35 @@ bool FGeometryCollectionEngineConversion::AppendSkeletalMesh(const USkeletalMesh
 	// the first render section, this will need to be expanded to include all render
 	// sections.
 	//
+	int32 RootIndex = INDEX_NONE;
 	const USkeleton* Skeleton = SkeletalMesh->GetSkeleton();
-	TManagedArray<FTransform>& Transform = InCollection->ModifyAttribute<FTransform>(FTransformCollection::TransformAttribute, FTransformCollection::TransformGroup);
-	int32 TransformBaseIndex = InCollection->AddElements(SkeletalBoneMap.Num(), FGeometryCollection::TransformGroup);
-	const FReferenceSkeleton & ReferenceSkeletion = Skeleton->GetReferenceSkeleton();
-	const TArray<FTransform> & RestArray = Skeleton->GetRefLocalPoses();
-	for (int32 BoneIndex = 0; BoneIndex < SkeletalBoneMap.Num(); BoneIndex++)
+	const TArray<FTransform>& RestArray = Skeleton->GetRefLocalPoses();
+	const FReferenceSkeleton& ReferenceSkeletion = Skeleton->GetReferenceSkeleton();
+
+	TManagedArray<FTransform3f>& LocalSpaceTransform = InCollection->ModifyAttribute<FTransform3f>(FTransformCollection::TransformAttribute, FTransformCollection::TransformGroup);
+	TManagedArray<int32>& Parent = InCollection->ModifyAttribute<int32>(FTransformCollection::ParentAttribute, FTransformCollection::TransformGroup);
+	TManagedArray<FLinearColor>& BoneColor = InCollection->ModifyAttribute<FLinearColor>("BoneColor", FTransformCollection::TransformGroup);
+	TManagedArray<FString>& BoneName = InCollection->ModifyAttribute<FString>("BoneName", FTransformCollection::TransformGroup);	
+	int32 TransformBaseIndex = InCollection->AddElements(ReferenceSkeletion.GetNum(), FGeometryCollection::TransformGroup);
+	TManagedArray<int32>& SimulationType = InCollection->ModifyAttribute<int32>("SimulationType", FTransformCollection::TransformGroup);
+
+	for (int32 BoneIndex = 0; BoneIndex < ReferenceSkeletion.GetNum(); BoneIndex++)
 	{
-		FTransform BoneTransform = FAnimationRuntime::GetComponentSpaceTransformRefPose(ReferenceSkeletion, SkeletalBoneMap[BoneIndex]);
-		Transform[TransformBaseIndex + BoneIndex] = BoneTransform;
+		LocalSpaceTransform[TransformBaseIndex + BoneIndex] = FTransform3f(RestArray[BoneIndex]);
+		// For validation against the component space position use
+		// FTransform ComponentSpaceTransform =FAnimationRuntime::GetComponentSpaceTransformRefPose(ReferenceSkeletion, SkeletalBoneMap[BoneIndex]);
+
+		BoneName[TransformBaseIndex + BoneIndex] = ReferenceSkeletion.GetRefBoneInfo()[BoneIndex].Name.ToString();
+		const FColor RandBoneColor(FMath::Rand() % 100 + 5, FMath::Rand() % 100 + 5, FMath::Rand() % 100 + 5, 255);
+		BoneColor[TransformBaseIndex + BoneIndex] = FLinearColor(RandBoneColor);
+
+		Parent[TransformBaseIndex + BoneIndex] = ReferenceSkeletion.GetRefBoneInfo()[BoneIndex].ParentIndex;
+		SimulationType[TransformBaseIndex + BoneIndex] = FGeometryCollection::ESimulationTypes::FST_None;
+
+		if (Parent[TransformBaseIndex + BoneIndex] == INDEX_NONE)
+		{
+			RootIndex = TransformBaseIndex + BoneIndex;
+		}
 	}
 
 
@@ -1218,14 +1238,15 @@ bool FGeometryCollectionEngineConversion::AppendSkeletalMesh(const USkeletalMesh
 	TManagedArray<FVector3f>& TangentU = InCollection->ModifyAttribute<FVector3f>("TangentU", FGeometryCollection::VerticesGroup);
 	TManagedArray<FVector3f>& TangentV = InCollection->ModifyAttribute<FVector3f>("TangentV", FGeometryCollection::VerticesGroup);
 	TManagedArray<int32>& BoneMap = InCollection->ModifyAttribute<int32>("BoneMap", FGeometryCollection::VerticesGroup);
-	TManagedArray<FLinearColor>& BoneColor = InCollection->ModifyAttribute<FLinearColor>("BoneColor", FTransformCollection::TransformGroup);
-	TManagedArray<FString>& BoneName = InCollection->ModifyAttribute<FString>("BoneName", FTransformCollection::TransformGroup);
+
 
 	// 
 	// Transform Attributes 
 	// 
-	TManagedArray<int32>& Parent = InCollection->ModifyAttribute<int32>(FTransformCollection::ParentAttribute, FTransformCollection::TransformGroup);
-	TManagedArray<int32>& SimulationType = InCollection->ModifyAttribute<int32>("SimulationType", FTransformCollection::TransformGroup);
+
+	TArray<FTransform> ComponentTransform;
+	GeometryCollectionAlgo::GlobalMatrices(LocalSpaceTransform, Parent, ComponentTransform);
+
 	int InitialNumVertices = InCollection->NumElements(FGeometryCollection::VerticesGroup);
 	int VertexBaseIndex = InCollection->AddElements(VertexCount, FGeometryCollection::VerticesGroup);
 	const int32 NumUVLayers = VertexBuffers.StaticMeshVertexBuffer.GetNumTexCoords();
@@ -1239,8 +1260,11 @@ bool FGeometryCollectionEngineConversion::AppendSkeletalMesh(const USkeletalMesh
 		SkinWeightVertexBuffer.GetRigidWeightBone(VertexIndex, SkeletalBoneIndex);
 		if (SkeletalBoneIndex > -1)
 		{
-			BoneMap[VertexOffset] = SkeletalBoneIndex + TransformBaseIndex;
-			Vertex[VertexOffset] = (FVector4f)Transform[BoneMap[VertexOffset]].ToInverseMatrixWithScale().TransformPosition((FVector)PositionVertexBuffer.VertexPosition(VertexIndex));
+			BoneMap[VertexOffset] = RootIndex;
+			Vertex[VertexOffset] = (FVector4f)ComponentTransform[BoneMap[VertexOffset]].ToInverseMatrixWithScale().TransformPosition((FVector)PositionVertexBuffer.VertexPosition(VertexIndex));
+			//@todo(GeometryCollectionConversion) : Support nesting rigid geometry within a transform space
+			//BoneMap[VertexOffset] = SkeletalBoneIndex + TransformBaseIndex;
+			//Vertex[VertexOffset] = (FVector3f)PositionVertexBuffer.VertexPosition(VertexIndex);
 		}
 		check(BoneMap[VertexOffset] != -1);
 		TangentU[VertexOffset] = VertexBuffers.StaticMeshVertexBuffer.VertexTangentX(VertexIndex);
@@ -1258,29 +1282,6 @@ bool FGeometryCollectionEngineConversion::AppendSkeletalMesh(const USkeletalMesh
 			Color[VertexOffset] = FLinearColor::White;
 	}
 
-	int32 InitialIndex = -1;
-	int32 LastParentIndex = -1;
-	int32 CurrentLevel = 0;
-	for (int32 BoneIndex = 0; BoneIndex < SkeletalBoneMap.Num(); BoneIndex++)
-	{
-		// transform based on position of the actor. 
-		Transform[TransformBaseIndex + BoneIndex] = SkeletalMeshTransform * Transform[TransformBaseIndex + BoneIndex];
-
-		// bone attributes
-		BoneName[TransformBaseIndex + BoneIndex] = ReferenceSkeletion.GetBoneName(SkeletalBoneMap[BoneIndex]).ToString();
-		const FColor RandBoneColor(FMath::Rand() % 100 + 5, FMath::Rand() % 100 + 5, FMath::Rand() % 100 + 5, 255);
-		BoneColor[TransformBaseIndex + BoneIndex] = FLinearColor(RandBoneColor);
-
-		// Bone Hierarchy - Added at root with no common parent
-		int32 ParentIndex = ReferenceSkeletion.GetParentIndex(SkeletalBoneMap[BoneIndex]);
-		int32 UseParentIndex = ParentIndex + InitialIndex;
-		if (LastParentIndex != UseParentIndex)
-		{
-			LastParentIndex = UseParentIndex;
-		}
-		Parent[TransformBaseIndex + BoneIndex] = UseParentIndex;
-		SimulationType[TransformBaseIndex + BoneIndex] = FGeometryCollection::ESimulationTypes::FST_Rigid;
-	}
 
 	// Geometry Group
 	TArray<int32> GeometryIndices;
