Hash phrases in phrase phase instead of entry phase.

Given the following two contiguous phrases in the data file A and B;

|AAAAaaa|bbbBBBB|

"aaabbb" represents a single cache item that was persisted to disk in two
writes - a "partial". Prior to this change the hash was taken at "aaa" (start
of entry). Now it is taken at "bbb" - the start of the phrase.

The journal will get two entries; one for "aaa" and one for "bbb". The first
will have a key of zero - entries with a null key are no longer added into the
DiskCache map. The entry for "bbb" will have the full key. It will also have an
non-zero offset. As the first entry of a phrase should have a zero offset, when
a non-zero offset is found this indicates a bias to find the start of a
partial.

This also works for huge partials that span three or more phrases;
"|AAAaa|bbb|ccCCCC|". "aa" and "bb" will have a zero key. "cc" will have a
non-zero offset that takes shows where "aa" is.

When loading, if the oldest detect phrase starts with a partial, that partial
is simply skipped.

#rnx

diff --git a/Engine/Source/Runtime/Experimental/IoStoreOnDemand/Private/JournaledCache.cpp b/Engine/Source/Runtime/Experimental/IoStoreOnDemand/Private/JournaledCache.cpp
--- a/Engine/Source/Runtime/Experimental/IoStoreOnDemand/Private/JournaledCache.cpp
+++ b/Engine/Source/Runtime/Experimental/IoStoreOnDemand/Private/JournaledCache.cpp
@@ -298,9 +298,9 @@ static_assert(sizeof(FPhraseDesc) == sizeof(FDataEntry));
 ////////////////////////////////////////////////////////////////////////////////
 struct FDiskPhrase
 {
-						FDiskPhrase(TArray<FDataEntry>& InEntries, int32 InMaxEntries, uint32 DataSize, uint32 PreviousPartial);
+						FDiskPhrase(TArray<FDataEntry>& InEntries, int32 InMaxEntries, uint32 DataSize);
 						FDiskPhrase(FDiskPhrase&&) = default;
-	bool				Add(uint64 Key, FIoBuffer&& Data);
+	bool				Add(uint64 Key, FIoBuffer&& Data, uint32 PartialBias=0);
 	void				Drop()					{ return Entries.SetNumUninitialized(Index); }
 	const FDataEntry*	GetEntries() const		{ return Entries.GetData() + Index; }
 	int32				GetEntryCount() const	{ return Entries.Num() - Index; }
@@ -317,7 +317,6 @@ private:
 	uint32				EntriesSize = 0;
 	uint32				Index;
 	int32				MaxEntries;
-	uint32				PreviousPartial;
 
 private:
 						FDiskPhrase(const FDiskPhrase&) = delete;
@@ -326,37 +325,30 @@ private:
 };
 
 ////////////////////////////////////////////////////////////////////////////////
-FDiskPhrase::FDiskPhrase(TArray<FDataEntry>& InEntries, int32 InMaxEntries, uint32 DataSize, uint32 InPreviousPartial)
+FDiskPhrase::FDiskPhrase(TArray<FDataEntry>& InEntries, int32 InMaxEntries, uint32 DataSize)
 : Entries(InEntries)
 , Index(Entries.Num())
 , MaxEntries(InMaxEntries)
-, PreviousPartial(InPreviousPartial)
 {
 	Buffer = TUniquePtr<uint8[]>(new uint8[DataSize]);
 }
 
 ////////////////////////////////////////////////////////////////////////////////
-bool FDiskPhrase::Add(uint64 Key, FIoBuffer&& Data)
+bool FDiskPhrase::Add(uint64 Key, FIoBuffer&& Data, uint32 PartialBias)
 {
 	check(MaxEntries > 0);
 	const uint32 DataSize = uint32(Data.GetSize());
-	// Actual entries (Key!=0) 
-	// Padding (Key=0, Size=0) can appear at the end of journal file
-	// Partials (Key=0) write data but does not add entry
-	if (Key || Data.GetSize() == 0)
-	{
-		const uint32 FullSize = DataSize + PreviousPartial;
-		PreviousPartial = 0;
-		check(FullSize < (1 << SIZE_BITS));
-		Entries.Add({Key, CurrentOffset, FullSize});
-		EntriesSize += FullSize;
-		CurrentOffset += FullSize;
-		--MaxEntries;
-	}
-	else
-	{
-		PreviousPartial += DataSize;
-	}
+	check(DataSize < (1 << SIZE_BITS));
+	Entries.Add(FDataEntry{
+		.Key = Key,
+		.Offset = CurrentOffset + PartialBias,
+		.Size = DataSize + PartialBias,
+		.EntryCount = 0,
+	});
+	EntriesSize += DataSize;
+	CurrentOffset += DataSize;
+	--MaxEntries;
+
 	std::memcpy(Buffer.Get() + Cursor, Data.GetData(), DataSize);
 	Cursor += DataSize;
 	return MaxEntries > 0;
@@ -379,8 +371,6 @@ public:
 	uint32					GetMaxSize() const	{ return MaxSize; }
 	uint32					GetCursor() const	{ return Cursor; }
 	uint32					GetMarker() const	{ return Marker; }
-	void					ClearPartial()		{ PreviousPartial = 0; }
-	uint32					GetPreviousPartial() const { return PreviousPartial; };
 
 private:
 	friend int32			LoadCache(FDiskCache&);
@@ -389,12 +379,10 @@ private:
 	static uint32			HashBytes(const uint8* Data, uint32 Size, uint32 Seed);
 	TArray<FDataEntry>		Entries;
 	FStringView				RootPath;
-	uint32					Marker = 0;
 	TUniquePtr<IFileHandle>	JrnHandle;
+	uint32					Marker = 0;
 	uint32					Cursor = 0;
 	uint32					MaxSize;
-	uint32					PreviousPartial = 0;
-	uint32					PreviousHash = 0;
 };
 
 ////////////////////////////////////////////////////////////////////////////////
@@ -470,7 +458,7 @@ FDiskPhrase FDiskJournal::OpenPhrase(uint32 DataSize)
 	Entries.Add(FDataEntry{});
 	int32 MaxEntries = int32((MaxSize - Cursor) / sizeof(FDataEntry)) - Entries.Num();
 
-	FDiskPhrase Ret(Entries, FMath::Min(MaxEntries, int32(UINT16_MAX)), DataSize, PreviousPartial);
+	FDiskPhrase Ret(Entries, FMath::Min(MaxEntries, int32(UINT16_MAX)), DataSize);
 
 	return Ret;
 }
@@ -478,27 +466,10 @@ FDiskPhrase FDiskJournal::OpenPhrase(uint32 DataSize)
 ////////////////////////////////////////////////////////////////////////////////
 void FDiskJournal::ClosePhrase(FDiskPhrase&& Phrase, uint64 DataCursor)
 {
-	auto SavePreviousHash = [&]
-	{
-		// Save the hash for the next phrase, if we have just written the (first) partial
-		if (PreviousPartial && !PreviousHash)
-		{
-			const uint64 PartialOffset = Phrase.GetWriteSize() - PreviousPartial;
-			check(PreviousPartial >= HASH_CHECKSUM_SIZE);
-			PreviousHash = HashBytes(Phrase.GetPhraseData() + PartialOffset, HASH_CHECKSUM_SIZE, Marker);
-		}
-	};
-	
 	uint32 EntryCount = Phrase.GetEntryCount();
 	if (EntryCount == 0)
 	{
 		Entries.Pop();
-		// We could write partial data but not complete any entries
-		if (const uint32 WriteSize = Phrase.GetWriteSize(); WriteSize > 0)
-		{
-			PreviousPartial += WriteSize;
-			SavePreviousHash();
-		}
 		return;
 	}
 	
@@ -522,36 +493,16 @@ void FDiskJournal::ClosePhrase(FDiskPhrase&& Phrase, uint64 DataCursor)
 	auto& Desc = (FPhraseDesc&)(Phrase.GetEntries()[-1]);
 	Desc.Magic = MAGIC;
 	Desc.Marker = Marker;
-	// Potentially adjust the data cursor to the start of the previous
-	// partial write.
-	Desc.DataCursor = DataCursor - PreviousPartial;
-	check(DataCursor >= PreviousPartial);
+	Desc.DataCursor = DataCursor;
 
-	// If we have already written a partial fragment of this phrase
-	// we use the hash of the start
-	if (PreviousHash)
-	{
-		Desc.Hash = PreviousHash;
-		PreviousHash = 0;
-	}
-	else
-	{
-		const uint32 HashSize = FMath::Min(HASH_CHECKSUM_SIZE, Phrase.GetDataSize());
-		Desc.Hash = HashBytes(Phrase.GetPhraseData(), HashSize, Desc.Marker);
-	}
-
-	// If we wrote a partial entry at the end we need adjust the
-	// next phrase accordingly.
-	check((PreviousPartial + Phrase.GetWriteSize()) >= Phrase.GetDataSize());
-	PreviousPartial = (PreviousPartial + Phrase.GetWriteSize()) - Phrase.GetDataSize();
+	const uint32 HashSize = FMath::Min(HASH_CHECKSUM_SIZE, Phrase.GetDataSize());
+	Desc.Hash = HashBytes(Phrase.GetPhraseData(), HashSize, Desc.Marker);
 
 	// Increment and wrap marker
 	if (++Marker > MARKER_MAX)
 	{
 		Marker = 0;
 	}
-
-	SavePreviousHash();
 }
 
 ////////////////////////////////////////////////////////////////////////////////
@@ -643,6 +594,8 @@ FDiskCache::FDiskCache(FString&& Path, uint64 InMaxDataSize, uint32 InJournalSiz
 , MaxDataSize(InMaxDataSize)
 , Journal(BinPath, InJournalSize)
 {
+	check(MaxDataSize >= HASH_CHECKSUM_SIZE);
+
 	// Align down to keep to some assumptions
 	MaxDataSize = (MaxDataSize - Journal.GetMaxSize()) & ~((1ull << 20) - 1);
 
@@ -683,7 +636,6 @@ void FDiskCache::Wrap()
 	check(DataCursor >= MaxDataSize - HASH_CHECKSUM_SIZE);
 	OverRemoval = 0;
 	DataCursor = 0;
-	Journal.ClearPartial();
 }
 
 ////////////////////////////////////////////////////////////////////////////////
@@ -728,7 +680,7 @@ void FDiskCache::ClosePhrase(FDiskPhrase&& Phrase)
 	{
 		FWriteScopeLock _(Lock);
 		Prune(DataCursor, WriteSize);
-		Insert(DataCursor - Journal.GetPreviousPartial(), Phrase.GetEntries(), EntryCount);
+		Insert(DataCursor, Phrase.GetEntries(), EntryCount);
 	}
 
 	Journal.ClosePhrase(MoveTemp(Phrase), DataCursor);
@@ -789,6 +741,11 @@ EIoErrorCode FDiskCache::Materialize(uint64 Key, FIoBuffer& Out, uint32 Offset)
 ////////////////////////////////////////////////////////////////////////////////
 uint64 FDiskCache::Insert(uint64 DataBase, const FDataEntry& Entry)
 {
+	if (Entry.Key == 0)
+	{
+		return 0;
+	}
+
 	FMapEntry Value;
 	Value.DataCursor = DataBase + Entry.Offset;
 	check(Value.DataCursor < MaxDataSize);
@@ -800,10 +757,22 @@ uint64 FDiskCache::Insert(uint64 DataBase, const FDataEntry& Entry)
 ////////////////////////////////////////////////////////////////////////////////
 uint64 FDiskCache::Insert(uint64 DataBase, const FDataEntry* Entries, uint32 EntryCount)
 {
+	uint32 PartialBias = 0;
+	if (EntryCount > 0)
+	{
+		PartialBias = Entries[0].Offset;
+	}
+
 	uint64 TotalSize = 0;
 	for (uint32 i = 0; i < EntryCount; ++i)
 	{
-		TotalSize += Insert(DataBase, Entries[i]);
+		check(int64(DataBase - PartialBias) >= 0);
+
+		FDataEntry Entry = Entries[i];
+		Entry.Offset -= PartialBias;
+
+		TotalSize += Insert(DataBase - PartialBias, Entry);
+		PartialBias = 0;
 	}
 
 	MappedBytes += TotalSize;
@@ -1157,13 +1126,6 @@ static int32 LoadCache(FDiskCache& DiskCache)
 			return -1;
 		}
 
-		// The last entry may be a padded entry which should be skipped. See
-		// FDiskPhrase::ClosePhrase for details.
-		if (LastEntry->Key == 0)
-		{
-			check(LastEntry->Size == 0);
-			--EntryCount;
-		}
 		MappedItems += EntryCount;
 		MappedBytes += DiskCache.Insert(Holm.DataCursor, Holm.Entries, EntryCount);
 	}
@@ -1234,7 +1196,7 @@ private:
 	{
 		uint64		Key;
 		FIoBuffer	Data;
-		uint32		Remaining;
+		uint32		Cursor;
 	};
 	using FPartial = TOptional<FPartialItem>;
 	
@@ -1373,6 +1335,7 @@ uint32 FCache::WriteMemToDisk(int32 Allowance)
 		Allowance = int32(UntilWrap);
 	}
 
+	uint32 PartialBias = 0;
 	FMemCache::PeelItems PeelItems;
 	int32 WriteSize = 0;
 	{
@@ -1382,55 +1345,56 @@ uint32 FCache::WriteMemToDisk(int32 Allowance)
 		// and peel off as much of that buffer as possible.
 		if (Partial)
 		{
-			const uint32 BufferOffset = uint32(Partial->Data.GetSize()) - Partial->Remaining;
-			const FMemoryView PartialSlice = Partial->Data.GetView().Mid(BufferOffset, Allowance);
-			const int32 PartialSize = int32(PartialSlice.GetSize());
-			Partial->Remaining -= PartialSize;
-			WriteSize += PartialSize;
-			const uint64 Key = Partial->Remaining ? 0 : Partial->Key;
-			PeelItems.Push(FMemCache::FItem {Key, FIoBuffer(PartialSlice, Partial->Data)});
-			if (Partial->Remaining == 0)
+			FMemoryView View = Partial->Data.GetView();
+			View = View.Mid(Partial->Cursor, Allowance);
+
+			uint32 ViewSize = uint32(View.GetSize());
+			Partial->Cursor += ViewSize;
+			WriteSize += ViewSize;
+
+			FMemCache::FItem PeelItem = {
+				.Key	= 0,
+				.Data	= FIoBuffer(View, Partial->Data),
+			};
+			if (Partial->Cursor >= uint32(Partial->Data.GetSize()))
 			{
+				PeelItem.Key = Partial->Key;
+				PartialBias = uint32(Partial->Data.GetSize()) - ViewSize;
 				Partial.Reset();
 			}
 			else if (bEof)
 			{
-				Partial->Remaining = uint32(Partial->Data.GetSize());
+				Partial->Cursor = 0;
 			}
+			PeelItems.Push(PeelItem);
 		}
-		
+
 		// If there is any allowance left start peeling of buffers from the memcache
 		if (WriteSize < Allowance)
 		{
 			WriteSize += MemCache.Peel(Allowance - WriteSize, PeelItems);
+
+			uint32 NewDemand = MemCache.GetDemand();
+			Demand.store(NewDemand, std::memory_order_relaxed);
 		}
 
 		// Finally split any overshooting buffers into a partial slice and save the
 		// buffer in the Partial member. While being dropped this buffer exists neither
 		// in the memcache or the disk cache. The partial fragment needs to be at least
 		// large enough for the hash checksum.
-		if (WriteSize > Allowance)
+		if (int32 Overshoot = WriteSize - Allowance; Overshoot > 0)
 		{
 			auto [Key, Data] = PeelItems.Pop();
-			
-			int32 RemainderSize = WriteSize - Allowance;
-			const int32 PartialSize = int32(Data.GetSize()) - RemainderSize;
-			if (PartialSize < HASH_CHECKSUM_SIZE)
-			{
-				RemainderSize = int32(Data.GetSize());
-			}
-			else
-			{
-				const FMemoryView PartialSlice = Data.GetView().Left(PartialSize);
-				PeelItems.Push(FMemCache::FItem{0, FIoBuffer(PartialSlice, Data)});
-			}
-			WriteSize -= RemainderSize;
-			RemainderSize = bEof ? int32(Data.GetSize()) : RemainderSize;
-			Partial.Emplace(FPartialItem{Key, MoveTemp(Data), uint32(RemainderSize)});
-		}
 
-		uint32 NewDemand = MemCache.GetDemand();
-		Demand.store(NewDemand, std::memory_order_relaxed);
+			int32 PartialSize = int32(Data.GetSize()) - Overshoot;
+			const FMemoryView PartialSlice = Data.GetView().Left(PartialSize);
+			PeelItems.Push(FMemCache::FItem{0, FIoBuffer(PartialSlice, Data)});
+
+			WriteSize -= Overshoot;
+
+			uint32 Cursor = bEof ? 0 : PartialSize;
+			Partial.Emplace(FPartialItem{Key, MoveTemp(Data), uint32(Cursor)});
+		}
 	}
 
 	check(WriteSize >= 0 && WriteSize <= Allowance);
@@ -1440,11 +1404,12 @@ uint32 FCache::WriteMemToDisk(int32 Allowance)
 	{
 		auto& [Key, Data] = PeelItems[i];
 		check(Key || i == (PeelItems.Num() - 1)); // Partials must be last
-		if (Phrase.GetRemainingEntries() < 1 || !Phrase.Add(Key, MoveTemp(Data)))
+		if (Phrase.GetRemainingEntries() < 1 || !Phrase.Add(Key, MoveTemp(Data), PartialBias))
 		{
 			PeelIndex = i;
 			break;
 		}
+		PartialBias = 0;
 	}
 
 	if (PeelIndex >= 0)
