A benign tidy.

#rnx

diff --git a/Engine/Source/Runtime/Experimental/IoStoreOnDemand/Private/JournaledCache.cpp b/Engine/Source/Runtime/Experimental/IoStoreOnDemand/Private/JournaledCache.cpp
--- a/Engine/Source/Runtime/Experimental/IoStoreOnDemand/Private/JournaledCache.cpp
+++ b/Engine/Source/Runtime/Experimental/IoStoreOnDemand/Private/JournaledCache.cpp
@@ -569,9 +569,10 @@ private:
 	};
 	static_assert(sizeof(FMapEntry) == sizeof(uint64));
 
+	using					FDataMap = TMap<uint64, FMapEntry>;
+
 	friend int32			LoadCache(FDiskCache&);
 	void					OpenDataFile();
-	using					FDataMap = TMap<uint64, FMapEntry>;
 	void					Wrap();
 	void					Spam();
 	uint64					Insert(uint64 DataBase, const FDataEntry& Entry);
@@ -1114,7 +1115,8 @@ static int32 LoadCache(FDiskCache& DiskCache)
 	}
 
 	// Add known entries into the tree.
-	uint64 MappedBytes = 0; uint32 MappedItems = 0;
+	uint64 MappedBytes = 0;
+	uint32 MappedItems = 0;
 	for (uint32 i = BasisIndex, n = Paragraphs.Num(); i < n; ++i)
 	{
 		const FPhraseDesc& Holm = Paragraphs[i].Phrase[0];
@@ -2637,22 +2639,6 @@ static void CacheTests(FSupport& Support)
 		NewCache(true);
 	}
 
-#if 0
-	// marker wrap
-	// one-phrase journal
-	// journal paragraphs that are all the same size
-	// phrases with no entries
-
-	// cache items larger than pending memory
-	// cache items larger than write allowance
-
-	// journal wrapping without truncation
-
-	// changes in max data/journal size
-
-	// don't load-and-sort so many paragraphs (only need max-data size)
-#endif // 0
-
 	// journal wrap
 	{
 		FCache::FConfig Config;
@@ -2906,6 +2892,22 @@ static void CacheTests(FSupport& Support)
 
 	// Make sure we delete the cache at the end to release the file handle
 	delete Cache;
+
+
+
+	// marker wrap
+	// one-phrase journal
+	// journal paragraphs that are all the same size
+	// phrases with no entries
+
+	// cache items larger than pending memory
+	// cache items larger than write allowance
+
+	// journal wrapping without truncation
+
+	// changes in max data/journal size
+
+	// don't load-and-sort so many paragraphs (only need max-data size
 }
 
 ////////////////////////////////////////////////////////////////////////////////
