Fixed OOB file read test and moved a block of notes to the end of a function.

#rnx

diff --git a/Engine/Source/Runtime/Experimental/IoStoreOnDemand/Private/JournaledCache.cpp b/Engine/Source/Runtime/Experimental/IoStoreOnDemand/Private/JournaledCache.cpp
--- a/Engine/Source/Runtime/Experimental/IoStoreOnDemand/Private/JournaledCache.cpp
+++ b/Engine/Source/Runtime/Experimental/IoStoreOnDemand/Private/JournaledCache.cpp
@@ -1077,7 +1077,7 @@ static int32 LoadCache(FDiskCache& DiskCache)
 	// Detect data writes that are newer than any journal flushes.
 	auto ReadHash = [File] (uint64 Cursor, uint32& OutHash, uint32 MaxHashSize, uint32 Seed)
 	{
-		if (Cursor + sizeof(OutHash) > uint64(File->Size()))
+		if (Cursor + MaxHashSize > uint64(File->Size()))
 		{
 			return false;
 		}
