Added tests for small phrases and removed unrequired hash-size bias.

#rnx

diff --git a/Engine/Source/Runtime/Experimental/IoStoreOnDemand/Private/JournaledCache.cpp b/Engine/Source/Runtime/Experimental/IoStoreOnDemand/Private/JournaledCache.cpp
--- a/Engine/Source/Runtime/Experimental/IoStoreOnDemand/Private/JournaledCache.cpp
+++ b/Engine/Source/Runtime/Experimental/IoStoreOnDemand/Private/JournaledCache.cpp
@@ -634,7 +634,7 @@ FDiskPhrase FDiskCache::OpenPhrase(uint32 DataSize)
 ////////////////////////////////////////////////////////////////////////////////
 void FDiskCache::Wrap()
 {
-	check(DataCursor >= MaxDataSize - HASH_CHECKSUM_SIZE);
+	check(DataCursor >= MaxDataSize);
 	OverRemoval = 0;
 	DataCursor = 0;
 }
@@ -856,7 +856,7 @@ void FDiskCache::Drop()
 ////////////////////////////////////////////////////////////////////////////////
 uint64 FDiskCache::RemainingUntilWrap()
 {
-	if (DataCursor >= MaxDataSize - HASH_CHECKSUM_SIZE)
+	if (DataCursor >= MaxDataSize)
 	{
 		Wrap();
 	}
@@ -1085,11 +1085,10 @@ static int32 LoadCache(FDiskCache& DiskCache)
 		File->Seek(Cursor);
 		
 		const uint32 HashSize = FMath::Min(HASH_CHECKSUM_SIZE, MaxHashSize);
-		TArray<uint8, TInlineAllocator<HASH_CHECKSUM_SIZE>> BytesToHash;
-		BytesToHash.AddZeroed(HashSize);
-		if (const bool Result = File->Read(BytesToHash.GetData(), HashSize))
+		uint8 Buffer[HASH_CHECKSUM_SIZE];
+		if (const bool Result = File->Read(Buffer, HashSize))
 		{
-			OutHash = FDiskJournal::HashBytes(BytesToHash.GetData(), BytesToHash.Num(), Seed);
+			OutHash = FDiskJournal::HashBytes(Buffer, HashSize, Seed);
 			return true;
 		}
 		return false;
@@ -2639,6 +2638,26 @@ static void CacheTests(FSupport& Support)
 		NewCache(true);
 	}
 
+	// little phrases
+	{
+		FIoBuffer Data;
+		Data = Support.DummyData(2);	Cache->Put(1, Data);
+		Data = Support.DummyData(126);	Cache->Put(2, Data);
+
+		Cache->WriteMemToDisk(32);
+		Cache->WriteMemToDisk(32);
+		Cache->WriteMemToDisk(64);
+		Cache->Flush();
+
+		NewCache();
+		Cache->Load();
+		Data = FIoBuffer();
+		check(Cache->Get(1, Data) == 1); check(Data.GetData() == nullptr);
+		check(Cache->Get(2, Data) == 2); check(Data.GetData() == nullptr);
+
+		NewCache(true);
+	}
+
 	// journal wrap
 	{
 		FCache::FConfig Config;
