Support ISMPool geometry collection custom renderer working in editor.
It does that through each geom collection creating its own ISPool. This way we get the runtime visuals and hit proxies still work.
Removed root proxy static component creation in geometry collection, and added a custom renderer that does this.

diff --git a/Engine/Plugins/Experimental/ChaosEditor/Source/FractureEditor/Private/FractureEditorMode.cpp b/Engine/Plugins/Experimental/ChaosEditor/Source/FractureEditor/Private/FractureEditorMode.cpp
--- a/Engine/Plugins/Experimental/ChaosEditor/Source/FractureEditor/Private/FractureEditorMode.cpp
+++ b/Engine/Plugins/Experimental/ChaosEditor/Source/FractureEditor/Private/FractureEditorMode.cpp
@@ -564,7 +564,7 @@ void UFractureEditorMode::OnActorSelectionChanged(const TArray<UObject*>& NewSel
 		for(UGeometryCollectionComponent* GeometryCollectionComponent : GeometryCollectionComponents)
 		{
 			GeometryCollectionComponent->SetEmbeddedGeometrySelectable(true);
-			GeometryCollectionComponent->EnableRootProxyStaticMeshComponents(false);
+			GeometryCollectionComponent->ForceNativeRendering(true);
 			
 			FGeometryCollectionEdit RestCollectionEdit = GeometryCollectionComponent->EditRestCollection(GeometryCollection::EEditUpdate::None);
 			if (!RestCollectionEdit.GetRestCollection())
@@ -602,9 +602,7 @@ void UFractureEditorMode::OnActorSelectionChanged(const TArray<UObject*>& NewSel
 			ShowBoneColorsEdit.SetEnableBoneSelection(false);
 
 			ExistingSelection->SetEmbeddedGeometrySelectable(false);
-
-			// should we restore the previous state ? 
-			ExistingSelection->EnableRootProxyStaticMeshComponents(true);
+			ExistingSelection->ForceNativeRendering(false);
 
 			// If we have a Hide array on the collection, remove it.
 			if (const UGeometryCollection* RestCollection = ExistingSelection->GetRestCollection())
diff --git a/Engine/Source/Runtime/Experimental/GeometryCollectionEngine/Private/GeometryCollection/GeometryCollectionComponent.cpp b/Engine/Source/Runtime/Experimental/GeometryCollectionEngine/Private/GeometryCollection/GeometryCollectionComponent.cpp
--- a/Engine/Source/Runtime/Experimental/GeometryCollectionEngine/Private/GeometryCollection/GeometryCollectionComponent.cpp
+++ b/Engine/Source/Runtime/Experimental/GeometryCollectionEngine/Private/GeometryCollection/GeometryCollectionComponent.cpp
@@ -131,9 +131,6 @@ FAutoConsoleVariableRef CVarGeometryCollectionEmitRootBreakingEvent(TEXT("p.Chao
 bool GeometryCollectionCreatePhysicsStateInEditor = false;
 FAutoConsoleVariableRef CVarGeometryCollectionCreatePhysicsStateInEditor(TEXT("p.Chaos.GC.CreatePhysicsStateInEditor"), GeometryCollectionCreatePhysicsStateInEditor, TEXT("when on , physics state for a GC will be create in editor ( non PIE )"));
 
-bool GeometryCollectionEnableRootProxyComponents = true;
-FAutoConsoleVariableRef CVarGeometryCollectionRootProxyComponents(TEXT("p.Chaos.GC.EnableRootProxyComponents"), GeometryCollectionEnableRootProxyComponents, TEXT("when on ( by default ) , create root proxy components"));
-
 bool GeometryCollectionUseReplicationV2 = true;
 FAutoConsoleVariableRef CVarGeometryCollectionUseReplicationV2(TEXT("p.Chaos.GC.UseReplicationV2"), GeometryCollectionUseReplicationV2, TEXT("When true use new replication data model"));
 
@@ -632,6 +629,11 @@ UGeometryCollectionComponent::UGeometryCollectionComponent(const FObjectInitiali
 	, bEnableRunTimeDataCollection(false)
 	, RunTimeDataCollectionGuid(FGuid::NewGuid())
 #endif
+	, bForceBrokenForCustomRenderer(false)
+	, bUpdateCustomRendererOnPostPhysicsSync(true)
+	, bCustomRendererCanUseNativeFallback(false)
+	, bCustomRendererShouldUseNativeFallback(false)
+	, bForceNativeRenderer(false)
 	, bEnableReplication(false)
 	, bEnableAbandonAfterLevel(true)
 	, AbandonedCollisionProfileName(UCollisionProfile::CustomCollisionProfileName)
@@ -653,8 +655,6 @@ UGeometryCollectionComponent::UGeometryCollectionComponent(const FObjectInitiali
 	, bIsTransformSelectionModeEnabled(false)
 #endif  // #if GEOMETRYCOLLECTION_EDITOR_SELECTION
 	, bIsMoving(false)
-	, bUpdateCustomRenderer(true)
-	, bUpdateCustomRendererOnPostPhysicsSync(true)
 {
 	// by default tick is registered but disabled, we only need it when we need to update the removal timers
 	// tick will be then enabled only when the root is broken from OnPostPhysicsSync callback
@@ -1091,7 +1091,13 @@ void UGeometryCollectionComponent::UpdateCachedBounds()
 
 bool UGeometryCollectionComponent::ShouldCreateRenderState() const
 {
-	return !CanUseCustomRenderer();
+	// Called once on registration
+	// Returning true here means that we pay overhead of render thread tasks for maintaining render state.
+	// Try and return false whenever possible, but ONLY if there is no chance that we will want to create a scene proxy.
+	return 
+		!IsCustomRendererAvailable() ||  // Custom renderer replaces need for a scene proxy.
+		bCustomRendererCanUseNativeFallback || // Except if custom renderer can require native rendering according to state.
+		!GetWorld()->IsGameWorld(); // Editor worlds need to support fracture editor which forces native rendering.
 }
 
 void UGeometryCollectionComponent::CreateRenderState_Concurrent(FRegisterComponentContext* Context)
@@ -1106,8 +1112,7 @@ FPrimitiveSceneProxy* UGeometryCollectionComponent::CreateSceneProxy()
 
 	FPrimitiveSceneProxy* LocalSceneProxy = nullptr;
 
-	const bool bNotUsingRootProxyComponents = RootProxyStaticMeshComponents.IsEmpty();
-	if (RestCollection && !CanUseCustomRenderer() && bNotUsingRootProxyComponents)
+	if (RestCollection && !IsUsingCustomRenderer())
 	{
 		if (UseNanite(GetScene()->GetShaderPlatform()) &&
 			RestCollection->EnableNanite &&
@@ -1697,6 +1702,17 @@ void UGeometryCollectionComponent::SetEmbeddedGeometrySelectable(bool bSelectabl
 	}
 }
 
+void UGeometryCollectionComponent::ForceNativeRendering(bool bForce)
+{
+	if (bForce != bForceNativeRenderer)
+	{
+		bForceNativeRenderer = bForce;
+		UnregisterCustomRenderer();
+		RegisterCustomRenderer();
+		MarkRenderStateDirty();
+	}
+}
+
 int32 UGeometryCollectionComponent::EmbeddedIndexToTransformIndex(const UInstancedStaticMeshComponent* ISMComponent, int32 InstanceIndex) const
 {
 	for (int32 ISMIdx = 0; ISMIdx < EmbeddedGeometryComponents.Num(); ++ISMIdx)
@@ -3270,6 +3286,12 @@ void UGeometryCollectionComponent::OnUpdateTransform(EUpdateTransformFlags Updat
 	{
 		PhysicsProxy->SetWorldTransform_External(GetComponentTransform());
 	}
+
+	// Make sure that custom renderer is updated at in editor.
+	if (!GetWorld()->IsGameWorld())
+	{
+		RefreshCustomRenderer();
+	}
 }
 
 bool UGeometryCollectionComponent::HasAnySockets() const
@@ -3543,130 +3565,35 @@ void UGeometryCollectionComponent::OnRegister()
 	}
 
 	Super::OnRegister();
-
-	CreateRootProxyComponentsIfNeeded();
 }
 
 void UGeometryCollectionComponent::OnUnregister()
 {
-	ClearRootProxyComponents();
-
 	Super::OnUnregister();
 
 	// Remove any custom renderer.
-	if (CustomRenderer)
-	{
-		UnregisterCustomRenderer();
-		CustomRenderer = nullptr;
-	}
-}
-
-void UGeometryCollectionComponent::EnableRootProxyStaticMeshComponents(bool bEnabled)
-{
-	if (bEnabled != bEnableRootProxyStaticMeshComponents)
-	{
-		bEnableRootProxyStaticMeshComponents = bEnabled;
-
-		if (bEnabled)
-		{
-			CreateRootProxyComponentsIfNeeded();
-		}
-		else
-		{
-			ClearRootProxyComponents();
-		}	
-		MarkRenderStateDirty();
-	}
-	
+	UnregisterCustomRenderer();
+	CustomRenderer = nullptr;
 }
 
-bool UGeometryCollectionComponent::ShouldCreateRootProxyComponents() const
-{
-	const bool bHasRootProxyMeshes = RestCollection && RestCollection->RootProxyData.ProxyMeshes.Num() > 0;
-	const bool bHasCustomRenderer = CanUseCustomRenderer();
-	return bHasRootProxyMeshes && !bHasCustomRenderer && bEnableRootProxyStaticMeshComponents && GeometryCollectionEnableRootProxyComponents;
-}
-
-void UGeometryCollectionComponent::CreateRootProxyComponentsIfNeeded()
+void UGeometryCollectionComponent::RegisterCustomRenderer()
 {
-	ClearRootProxyComponents();
+	bCustomRendererCanUseNativeFallback = false;
 
-	if (ShouldCreateRootProxyComponents() && GetOwner())
+	if (IsCustomRendererAvailable())
 	{
-		for (const TObjectPtr<UStaticMesh>& ProxyMesh : RestCollection->RootProxyData.ProxyMeshes)
+		if (IGeometryCollectionExternalRenderInterface* RendererInterface = CustomRenderer.GetInterface())
 		{
-			if (UStaticMesh* StaticMesh = ProxyMesh.Get())
-			{
-				const FName UniqueName = MakeUniqueObjectName(GetOwner(), UStaticMeshComponent::StaticClass(), TEXT("GC_RootProxyMesh"));
-				UStaticMeshComponent* MeshComponent = NewObject<UStaticMeshComponent>(GetOwner(), UniqueName, RF_DuplicateTransient | RF_Transient | RF_TextExportTransient);
-
-				MeshComponent->SetStaticMesh(StaticMesh);
-				//MeshComponent->SetRelativeTransform(GetComponentTransform());
-				MeshComponent->SetCanEverAffectNavigation(false);
-				MeshComponent->SetCollisionProfileName(UCollisionProfile::NoCollision_ProfileName);
-				MeshComponent->SetMobility(this->Mobility);
-				MeshComponent->SetupAttachment(this);
-				MeshComponent->RegisterComponent();
+			bCustomRendererCanUseNativeFallback = RendererInterface->CanEverUseNativeFallback();
 
-				RootProxyStaticMeshComponents.Add(MeshComponent);
-			}
-		}
-	}
-}
-
-void UGeometryCollectionComponent::UpdateRootProxyComponentsIfNeeded()
-{
-	if (RootProxyStaticMeshComponents.Num() > 0 && !bUpdateComponentTransformToRootBone)
-	{
-		if (RestCollection && !IsRootBroken())
-		{
-			const TSharedPtr<FGeometryCollection, ESPMode::ThreadSafe> AssetCollection = RestCollection->GetGeometryCollection();
-			if (AssetCollection)
+			if (IsUsingCustomRenderer())
 			{
-				const int32 RootIndex = RestCollection->GetRootIndex();
-				const FTransform3f RootRestTransform = AssetCollection->Transform[RootIndex];
-				const FTransform3f RootCurrentTransform = ComponentSpaceTransforms.RequestRootTransform();
-				const FTransform RootTransformFromRestPose = FTransform(RootCurrentTransform.GetRelativeTransform(RootRestTransform));
-
-				for (int32 MeshIndex = 0; MeshIndex < RootProxyStaticMeshComponents.Num(); MeshIndex++)
-				{
-					if (TObjectPtr<UStaticMeshComponent> StaticMeshComponent = RootProxyStaticMeshComponents[MeshIndex])
-					{
-						if (RootProxyLocalTransforms.IsValidIndex(MeshIndex))
-						{
-							StaticMeshComponent->SetRelativeTransform(FTransform(RootProxyLocalTransforms[MeshIndex]) * RootTransformFromRestPose);
-						}
-						else
-						{
-							StaticMeshComponent->SetRelativeTransform(RootTransformFromRestPose);
-						}
-					}
-				}
+				RendererInterface->OnRegisterGeometryCollection(*this);
 			}
 		}
 	}
 }
 
-void UGeometryCollectionComponent::ClearRootProxyComponents()
-{
-	for (TObjectPtr<UStaticMeshComponent> StaticMeshComponent : RootProxyStaticMeshComponents)
-	{
-		if (StaticMeshComponent)
-		{
-			StaticMeshComponent->DestroyComponent();
-		}
-	}
-	RootProxyStaticMeshComponents.Empty();
-}
-
-void UGeometryCollectionComponent::RegisterCustomRenderer()
-{
-	if (IGeometryCollectionExternalRenderInterface* RendererInterface = CustomRenderer.GetInterface())
-	{
-		RendererInterface->OnRegisterGeometryCollection(*this);
-	}
-}
-
 void UGeometryCollectionComponent::UnregisterCustomRenderer()
 {
 	if (IGeometryCollectionExternalRenderInterface* RendererInterface = CustomRenderer.GetInterface())
@@ -3878,7 +3805,12 @@ void UGeometryCollectionComponent::OnCreatePhysicsState()
 				// We're skipping over the primitive component so we need to make sure this event gets fired.
 				OnComponentPhysicsStateChanged.Broadcast(this, EComponentPhysicsStateChange::Created);
 			}
-			
+		}
+
+		// Make sure that custom renderer is updated at least once in editor.
+		if (!GetWorld()->IsGameWorld())
+		{
+			RefreshCustomRenderer();
 		}
 	}
 }
@@ -4360,12 +4292,6 @@ void UGeometryCollectionComponent::OnPostPhysicsSync()
 			InitializeRemovalDynamicAttributesIfNeeded();
 		}
 
-		if (RootProxyStaticMeshComponents.Num() > 0)
-		{
-			ClearRootProxyComponents();
-			RecreateRenderState_Concurrent();
-		}
-
 		UpdateRemovalIfNeeded();
 
 		CheckFullyDecayed();
@@ -4463,8 +4389,6 @@ void UGeometryCollectionComponent::UpdateRenderSystemsIfNeeded(bool bDynamicColl
 			RefreshCustomRenderer();
 		}
 
-		UpdateRootProxyComponentsIfNeeded();
-
 		if (SceneProxy != nullptr)
 		{
 			if (SceneProxy->IsNaniteMesh())
@@ -5046,8 +4970,6 @@ void UGeometryCollectionComponent::SetRestCollection(const UGeometryCollection*
 			RegisterCustomRenderer();
 		}
 
-		CreateRootProxyComponentsIfNeeded();
-
 		ClearRootProxyLocalTransforms();
 
 		if (bApplyAssetDefaults)
@@ -6301,9 +6223,9 @@ void UGeometryCollectionComponent::InitializeEmbeddedGeometry()
 	}
 }
 
-void UGeometryCollectionComponent::EnableRootProxyForCustomRenderer(bool bEnable)
+void UGeometryCollectionComponent::ForceBrokenForCustomRenderer(bool bForceBroken)
 { 
-	bEnableRootProxyForCustomRenderer = bEnable;
+	bForceBrokenForCustomRenderer = bForceBroken;
 	RefreshCustomRenderer();
 }
 
@@ -6328,66 +6250,74 @@ void UGeometryCollectionComponent::ClearRootProxyLocalTransforms()
 void UGeometryCollectionComponent::RefreshRootProxies()
 {
 	RefreshCustomRenderer();
-	UpdateRootProxyComponentsIfNeeded();
 }
 
-bool UGeometryCollectionComponent::CanUseCustomRenderer() const 
+bool UGeometryCollectionComponent::IsCustomRendererAvailable() const
 {
-	return bChaos_GC_UseCustomRenderer && CustomRenderer != nullptr && GetWorld()->IsGameWorld();
+	return bChaos_GC_UseCustomRenderer && CustomRenderer != nullptr && FApp::CanEverRender();
+}
+
+bool UGeometryCollectionComponent::IsUsingCustomRenderer() const
+{
+	return IsCustomRendererAvailable() && !(bCustomRendererCanUseNativeFallback && bCustomRendererShouldUseNativeFallback) && !bForceNativeRenderer;
 }
 
 void UGeometryCollectionComponent::RefreshCustomRenderer()
 {
-	if (CanUseCustomRenderer())
+	if (RestCollection == nullptr || !IsCustomRendererAvailable())
 	{
-		// Don't refresh the custom renderer on the server but we still need to do the work of computing component space transforms.
-		const bool bUpdateRenderer = !IsNetMode(NM_DedicatedServer) && bUpdateCustomRenderer;
+		return;
+	}
+	
+	if (IGeometryCollectionExternalRenderInterface* RendererInterface = CustomRenderer.GetInterface())
+	{
+		const int32 RootIndex = GetRootIndex();
+		bool bIsBroken = DynamicCollection ? !DynamicCollection->Active[RootIndex] : false;
 
-		if (IGeometryCollectionExternalRenderInterface* RendererInterface = CustomRenderer.GetInterface())
+#if !(UE_BUILD_SHIPPING)
+		if (CVarNumToForceBreak->GetInt() >= 0)
 		{
-			if (RestCollection != nullptr)
-			{
-				if (bUpdateRenderer)
-				{
-					const FTransform ComponentTransform = GetComponentTransform();
-					const int32 RootIndex = GetRootIndex();
-
-					bool bIsBroken = DynamicCollection ? !DynamicCollection->Active[RootIndex] : false;
+			bIsBroken = ForcedBroken.Break(this, ComponentSpaceTransforms.RequestAllTransforms());
+		}
+#endif						
 
-				#if !(UE_BUILD_SHIPPING)
-					if (CVarNumToForceBreak->GetInt() >= 0)
-					{
-						bIsBroken = ForcedBroken.Break(this, ComponentSpaceTransforms.RequestAllTransforms());
-					}
-				#endif						
+		const bool bRenderRootProxy = !bForceBrokenForCustomRenderer && !bIsBroken && (RestCollection->RootProxyData.ProxyMeshes.Num() > 0);
 
-					const bool bRenderRootProxy = bEnableRootProxyForCustomRenderer && !bIsBroken && (RestCollection->RootProxyData.ProxyMeshes.Num() > 0);
+		uint32 StateFlags = 0;
+		StateFlags |= bHiddenInGame || !IsVisible() ? 0 : IGeometryCollectionExternalRenderInterface::EState_Visible;
+		StateFlags |= bRenderRootProxy ? 0 : IGeometryCollectionExternalRenderInterface::EState_Broken;
+		StateFlags |= bForceBrokenForCustomRenderer ? IGeometryCollectionExternalRenderInterface::EState_ForcedBroken : 0;
 
-					uint32 StateFlags = 0;
-					StateFlags |= bHiddenInGame || !IsVisible() ? 0 : IGeometryCollectionExternalRenderInterface::EState_Visible;
-					StateFlags |= bRenderRootProxy ? 0 : IGeometryCollectionExternalRenderInterface::EState_Broken;
-					StateFlags |= bEnableRootProxyForCustomRenderer ? 0 : IGeometryCollectionExternalRenderInterface::EState_ForcedBroken;
+		// Test if we should change the native rendering fallback state before updating custom renderer.
+		const bool bShouldUseNativeFallback = bCustomRendererCanUseNativeFallback && RendererInterface->ShouldUseNativeFallback(StateFlags);
+		if (bCustomRendererShouldUseNativeFallback != bShouldUseNativeFallback)
+		{
+			bCustomRendererShouldUseNativeFallback = bShouldUseNativeFallback;
+			UnregisterCustomRenderer();
+			RegisterCustomRenderer();
+			MarkRenderStateDirty();
+		}
 
-					RendererInterface->UpdateState(*RestCollection, ComponentTransform, StateFlags);
+		if (IsUsingCustomRenderer())
+		{
+			RendererInterface->UpdateState(*RestCollection, GetComponentTransform(), StateFlags);
 
-					if (bRenderRootProxy)
-					{
-						const FTransform CompSpaceRootTransform(ComponentSpaceTransforms.RequestRootTransform());
-						if (RootProxyLocalTransforms.IsEmpty())
-						{
-							RendererInterface->UpdateRootTransform(*RestCollection, CompSpaceRootTransform);
-						}
-						else
-						{
-							RendererInterface->UpdateRootTransforms(*RestCollection, CompSpaceRootTransform, MakeArrayView(RootProxyLocalTransforms));
-						}
-					}
-					else
-					{
-						const TArray<FTransform3f>& CompSpaceTransforms = ComponentSpaceTransforms.RequestAllTransforms();
-						RendererInterface->UpdateTransforms(*RestCollection, CompSpaceTransforms);
-					}
+			if (bRenderRootProxy)
+			{
+				const FTransform CompSpaceRootTransform(ComponentSpaceTransforms.RequestRootTransform());
+				if (RootProxyLocalTransforms.IsEmpty())
+				{
+					RendererInterface->UpdateRootTransform(*RestCollection, CompSpaceRootTransform);
 				}
+				else
+				{
+					RendererInterface->UpdateRootTransforms(*RestCollection, CompSpaceRootTransform, MakeArrayView(RootProxyLocalTransforms));
+				}
+			}
+			else
+			{
+				const TArray<FTransform3f>& CompSpaceTransforms = ComponentSpaceTransforms.RequestAllTransforms();
+				RendererInterface->UpdateTransforms(*RestCollection, CompSpaceTransforms);
 			}
 		}
 	}
@@ -7144,7 +7074,7 @@ void UGeometryCollectionComponent::PostLoad()
 	Super::PostLoad();
 
 	// If there is a rest collection and no custom renderer then precache PSOs required to render the mesh
-	if (RestCollection && !CanUseCustomRenderer())
+	if (RestCollection && !IsCustomRendererAvailable())
 	{
 		PrecachePSOs();
 	}
diff --git a/Engine/Source/Runtime/Experimental/GeometryCollectionEngine/Private/GeometryCollection/GeometryCollectionISMPoolRenderer.cpp b/Engine/Source/Runtime/Experimental/GeometryCollectionEngine/Private/GeometryCollection/GeometryCollectionISMPoolRenderer.cpp
--- a/Engine/Source/Runtime/Experimental/GeometryCollectionEngine/Private/GeometryCollection/GeometryCollectionISMPoolRenderer.cpp
+++ b/Engine/Source/Runtime/Experimental/GeometryCollectionEngine/Private/GeometryCollection/GeometryCollectionISMPoolRenderer.cpp
@@ -13,9 +13,19 @@
 
 #include UE_INLINE_GENERATED_CPP_BY_NAME(GeometryCollectionISMPoolRenderer)
 
-void UGeometryCollectionISMPoolRenderer::OnRegisterGeometryCollection(UGeometryCollectionComponent const& InComponent)
+void UGeometryCollectionISMPoolRenderer::OnRegisterGeometryCollection(UGeometryCollectionComponent& InComponent)
 {
 	OwningLevel = InComponent.GetComponentLevel();
+
+	// In editor we create our own ISMPool. 
+	// This guarantees the same look in editor/game, and allows editor hit proxies to continue working.
+	if (!InComponent.GetWorld()->IsGameWorld())
+	{
+		CachedISMPoolComponent = NewObject<UGeometryCollectionISMPoolComponent>(InComponent.GetOwner());
+		CachedISMPoolComponent->SetTickablePoolManagement(false);
+		CachedISMPoolComponent->SetupAttachment(&InComponent);
+		CachedISMPoolComponent->RegisterComponent();
+	}
 }
 
 void UGeometryCollectionISMPoolRenderer::OnUnregisterGeometryCollection()
@@ -23,7 +33,7 @@ void UGeometryCollectionISMPoolRenderer::OnUnregisterGeometryCollection()
 	ReleaseGroup(MergedMeshGroup);
 	ReleaseGroup(InstancesGroup);
 
-	ISMPoolActor = nullptr;
+	CachedISMPoolComponent = nullptr;
 }
 
 void UGeometryCollectionISMPoolRenderer::UpdateState(UGeometryCollection const& InGeometryCollection, FTransform const& InComponentTransform, uint32 InStateFlags)
@@ -75,17 +85,26 @@ void UGeometryCollectionISMPoolRenderer::UpdateTransforms(UGeometryCollection co
 	UpdateInstanceTransforms(InGeometryCollection, ComponentTransform, InTransforms);
 }
 
+UGeometryCollectionISMPoolComponent* UGeometryCollectionISMPoolRenderer::GetISMPoolComponent() const
+{
+	return CachedISMPoolComponent;
+}
+
 UGeometryCollectionISMPoolComponent* UGeometryCollectionISMPoolRenderer::GetOrCreateISMPoolComponent()
 {
-	if (ISMPoolActor == nullptr)
+	if (CachedISMPoolComponent == nullptr)
 	{
 		if (UGeometryCollectionISMPoolSubSystem* ISMPoolSubSystem = UWorld::GetSubsystem<UGeometryCollectionISMPoolSubSystem>(GetWorld()))
 		{
 			check(OwningLevel);
-			ISMPoolActor = ISMPoolSubSystem->FindISMPoolActor(OwningLevel);
+			if (AGeometryCollectionISMPoolActor* ISMPoolActor = ISMPoolSubSystem->FindISMPoolActor(OwningLevel))
+			{
+				CachedISMPoolComponent = ISMPoolActor->GetISMPoolComp();
+			}
 		}
 	}
-	return ISMPoolActor != nullptr ? ISMPoolActor->GetISMPoolComp() : nullptr;
+
+	return CachedISMPoolComponent;
 }
 
 void UGeometryCollectionISMPoolRenderer::InitMergedMeshFromGeometryCollection(UGeometryCollection const& InGeometryCollection)
@@ -233,7 +252,14 @@ void UGeometryCollectionISMPoolRenderer::UpdateInstanceTransforms(UGeometryColle
 
 void UGeometryCollectionISMPoolRenderer::ReleaseGroup(FISMPoolGroup& InOutGroup)
 {
-	UGeometryCollectionISMPoolComponent* ISMPoolComponent = GetOrCreateISMPoolComponent();
+	if (InOutGroup.GroupIndex == INDEX_NONE)
+	{
+		return;
+	}
+
+	// Component and owning actor may already be released safely by a level unload.
+	// Don't want to create a new component here, so don't use GetOrCreateISMPoolComponent().
+	UGeometryCollectionISMPoolComponent* ISMPoolComponent = GetISMPoolComponent();
 	if (ISMPoolComponent != nullptr)
 	{
 		ISMPoolComponent->DestroyMeshGroup(InOutGroup.GroupIndex);
diff --git a/Engine/Source/Runtime/Experimental/GeometryCollectionEngine/Private/GeometryCollection/GeometryCollectionISMPoolRenderer.h b/Engine/Source/Runtime/Experimental/GeometryCollectionEngine/Private/GeometryCollection/GeometryCollectionISMPoolRenderer.h
--- a/Engine/Source/Runtime/Experimental/GeometryCollectionEngine/Private/GeometryCollection/GeometryCollectionISMPoolRenderer.h
+++ b/Engine/Source/Runtime/Experimental/GeometryCollectionEngine/Private/GeometryCollection/GeometryCollectionISMPoolRenderer.h
@@ -19,7 +19,7 @@ class UGeometryCollectionISMPoolRenderer : public UObject, public IGeometryColle
 
 public:
 	//~ Begin IGeometryCollectionExternalRenderInterface Interface.
-	virtual void OnRegisterGeometryCollection(UGeometryCollectionComponent const& InComponent) override;
+	virtual void OnRegisterGeometryCollection(UGeometryCollectionComponent& InComponent) override;
 	virtual void OnUnregisterGeometryCollection() override;
 	virtual void UpdateState(UGeometryCollection const& InGeometryCollection, FTransform const& InComponentTransform, uint32 InStateFlags) override;
 	virtual void UpdateRootTransform(UGeometryCollection const& InGeometryCollection, FTransform const& InRootTransform) override;
@@ -35,9 +35,9 @@ public:
 	};
 
 protected:
-	/** Instanced Static Mesh Pool actor that is used to render our meshes. */
+	/** Instanced Static Mesh Pool component that is used to render our meshes. */
 	UPROPERTY(Transient)
-	TObjectPtr<AGeometryCollectionISMPoolActor> ISMPoolActor;
+	TObjectPtr<UGeometryCollectionISMPoolComponent> CachedISMPoolComponent;
 
 	/** Cached component transform. */
 	FTransform ComponentTransform = FTransform::Identity;
@@ -50,6 +50,7 @@ protected:
 	ULevel* OwningLevel = nullptr;
 
 private:
+	UGeometryCollectionISMPoolComponent* GetISMPoolComponent() const;
 	UGeometryCollectionISMPoolComponent* GetOrCreateISMPoolComponent();
 	void InitMergedMeshFromGeometryCollection(UGeometryCollection const& InGeometryCollection);
 	void InitInstancesFromGeometryCollection(UGeometryCollection const& InGeometryCollection);
diff --git a/Engine/Source/Runtime/Experimental/GeometryCollectionEngine/Private/GeometryCollection/GeometryCollectionRootProxyRenderer.cpp b/Engine/Source/Runtime/Experimental/GeometryCollectionEngine/Private/GeometryCollection/GeometryCollectionRootProxyRenderer.cpp
new file mode 100644
--- /dev/null
+++ b/Engine/Source/Runtime/Experimental/GeometryCollectionEngine/Private/GeometryCollection/GeometryCollectionRootProxyRenderer.cpp
@@ -0,0 +1,116 @@
+// Copyright Epic Games, Inc. All Rights Reserved.
+#include "GeometryCollection/GeometryCollectionRootProxyRenderer.h"
+
+#include "Components/StaticMeshComponent.h"
+#include "Engine/StaticMesh.h"
+#include "Engine/World.h"
+#include "GeometryCollection/GeometryCollection.h"
+#include "GeometryCollection/GeometryCollectionComponent.h"
+#include "GeometryCollection/GeometryCollectionObject.h"
+
+#include UE_INLINE_GENERATED_CPP_BY_NAME(GeometryCollectionRootProxyRenderer)
+
+void UGeometryCollectionRootProxyRenderer::OnRegisterGeometryCollection(UGeometryCollectionComponent& InComponent)
+{
+	CreateRootProxyComponents(InComponent);
+}
+
+void UGeometryCollectionRootProxyRenderer::OnUnregisterGeometryCollection()
+{
+	ClearRootProxyComponents();
+}
+
+void UGeometryCollectionRootProxyRenderer::UpdateState(UGeometryCollection const& InGeometryCollection, FTransform const& InComponentTransform, uint32 InStateFlags)
+{
+	const bool bIsStateVisible = (InStateFlags & EState_Visible) != 0;
+	const bool bIsStateBroken = (InStateFlags & EState_Broken) != 0;
+	const bool bSetVisible = !bIsStateBroken && bIsStateVisible;
+	
+	if (bIsVisible != bSetVisible)
+	{
+		bIsVisible = bSetVisible;
+
+		for (UStaticMeshComponent* StaticMeshComponent : StaticMeshComponents)
+		{
+			if (StaticMeshComponent)
+			{
+				StaticMeshComponent->SetVisibility(bIsVisible);
+			}
+		}
+	}
+}
+
+void UGeometryCollectionRootProxyRenderer::UpdateRootTransform(UGeometryCollection const& InGeometryCollection, FTransform const& InRootTransform)
+{
+	UpdateRootProxyTransforms(InGeometryCollection, InRootTransform, {});
+}
+
+void UGeometryCollectionRootProxyRenderer::UpdateRootTransforms(UGeometryCollection const& InGeometryCollection, FTransform const& InRootTransform, TArrayView<const FTransform3f> InRootTransforms)
+{
+	UpdateRootProxyTransforms(InGeometryCollection, InRootTransform, InRootTransforms);
+}
+
+void UGeometryCollectionRootProxyRenderer::UpdateTransforms(UGeometryCollection const& InGeometryCollection, TArrayView<const FTransform3f> InTransforms)
+{
+	// Don't support non root proxy mesh.
+}
+
+void UGeometryCollectionRootProxyRenderer::CreateRootProxyComponents(UGeometryCollectionComponent& InComponent)
+{
+	UGeometryCollection const* GeometryCollection = InComponent.GetRestCollection();
+	if (GeometryCollection == nullptr)
+	{
+		return;
+	}
+
+	StaticMeshComponents.Reserve(GeometryCollection->RootProxyData.ProxyMeshes.Num());
+
+	for (UStaticMesh* ProxyMesh : GeometryCollection->RootProxyData.ProxyMeshes)
+	{
+		UStaticMeshComponent* MeshComponent = nullptr;
+		
+		if (ProxyMesh != nullptr)
+		{
+			MeshComponent = NewObject<UStaticMeshComponent>(InComponent.GetOwner(), NAME_None, RF_DuplicateTransient | RF_Transient);
+
+			MeshComponent->SetStaticMesh(ProxyMesh);
+			MeshComponent->SetCanEverAffectNavigation(false);
+			MeshComponent->SetCollisionProfileName(UCollisionProfile::NoCollision_ProfileName);
+			MeshComponent->SetMobility(InComponent.Mobility);
+			MeshComponent->SetupAttachment(&InComponent);
+			MeshComponent->RegisterComponent();
+		}
+
+		StaticMeshComponents.Add(MeshComponent);
+	}
+}
+
+void UGeometryCollectionRootProxyRenderer::UpdateRootProxyTransforms(UGeometryCollection const& InGeometryCollection, FTransform const& InRootTransform, TArrayView<const FTransform3f> InLocalRootTransforms)
+{
+	for (int32 MeshIndex = 0; MeshIndex < StaticMeshComponents.Num(); MeshIndex++)
+	{
+		if (UStaticMeshComponent* MeshComponent = StaticMeshComponents[MeshIndex])
+		{
+			if (InLocalRootTransforms.IsValidIndex(MeshIndex))
+			{
+				MeshComponent->SetRelativeTransform(FTransform(InLocalRootTransforms[MeshIndex]) * InRootTransform);
+			}
+			else
+			{
+				MeshComponent->SetRelativeTransform(InRootTransform);
+			}
+		}
+	}
+}
+
+void UGeometryCollectionRootProxyRenderer::ClearRootProxyComponents()
+{
+	for (UStaticMeshComponent* StaticMeshComponent : StaticMeshComponents)
+	{
+		if (StaticMeshComponent)
+		{
+			StaticMeshComponent->DestroyComponent();
+		}
+	}
+	StaticMeshComponents.Reset();
+}
diff --git a/Engine/Source/Runtime/Experimental/GeometryCollectionEngine/Private/GeometryCollection/GeometryCollectionRootProxyRenderer.h b/Engine/Source/Runtime/Experimental/GeometryCollectionEngine/Private/GeometryCollection/GeometryCollectionRootProxyRenderer.h
new file mode 100644
--- /dev/null
+++ b/Engine/Source/Runtime/Experimental/GeometryCollectionEngine/Private/GeometryCollection/GeometryCollectionRootProxyRenderer.h
@@ -0,0 +1,40 @@
+// Copyright Epic Games, Inc. All Rights Reserved.
+
+#pragma once
+
+#include "GeometryCollection/GeometryCollectionExternalRenderInterface.h"
+
+#include "GeometryCollectionRootProxyRenderer.generated.h"
+
+class UGeometryCollectionComponent;
+
+/** Implementation of a geometry collection custom renderer that creates static mesh components representing the root proxies. */
+UCLASS()
+class UGeometryCollectionRootProxyRenderer : public UObject, public IGeometryCollectionExternalRenderInterface
+{
+	GENERATED_BODY()
+
+public:
+	//~ Begin IGeometryCollectionExternalRenderInterface Interface.
+	virtual void OnRegisterGeometryCollection(UGeometryCollectionComponent& InComponent) override;
+	virtual void OnUnregisterGeometryCollection() override;
+	virtual void UpdateState(UGeometryCollection const& InGeometryCollection, FTransform const& InComponentTransform, uint32 InStateFlags) override;
+	virtual void UpdateRootTransform(UGeometryCollection const& InGeometryCollection, FTransform const& InRootTransform) override;
+	virtual void UpdateRootTransforms(UGeometryCollection const& InGeometryCollection, FTransform const& InRootTransform, TArrayView<const FTransform3f> InRootLocalTransforms) override;
+	virtual void UpdateTransforms(UGeometryCollection const& InGeometryCollection, TArrayView<const FTransform3f> InTransforms) override;
+	virtual bool ShouldUseNativeFallback(uint32 InStateFlags) const override { return (InStateFlags & EState_Broken) != 0; }
+	virtual bool CanEverUseNativeFallback() const override { return true; }
+	//~ End IGeometryCollectionExternalRenderInterface Interface.
+
+protected:
+	/** Current visibility state. */
+	bool bIsVisible = true;
+	/** Static mesh components for the root proxy meshes. */
+	UPROPERTY(Transient)
+	TArray<TObjectPtr<UStaticMeshComponent>> StaticMeshComponents;
+
+private:
+	void CreateRootProxyComponents(UGeometryCollectionComponent& InComponent);
+	void UpdateRootProxyTransforms(UGeometryCollection const& InGeometryCollection, FTransform const& InRootTransform, TArrayView<const FTransform3f> InLocalRootTransforms);
+	void ClearRootProxyComponents();
+};
diff --git a/Engine/Source/Runtime/Experimental/GeometryCollectionEngine/Public/GeometryCollection/GeometryCollectionComponent.h b/Engine/Source/Runtime/Experimental/GeometryCollectionEngine/Public/GeometryCollection/GeometryCollectionComponent.h
--- a/Engine/Source/Runtime/Experimental/GeometryCollectionEngine/Public/GeometryCollection/GeometryCollectionComponent.h
+++ b/Engine/Source/Runtime/Experimental/GeometryCollectionEngine/Public/GeometryCollection/GeometryCollectionComponent.h
@@ -1278,6 +1278,7 @@ public:
 
 #if WITH_EDITOR
 	GEOMETRYCOLLECTIONENGINE_API void SetEmbeddedGeometrySelectable(bool bSelectableIn);
+	GEOMETRYCOLLECTIONENGINE_API void ForceNativeRendering(bool bForce);
 	GEOMETRYCOLLECTIONENGINE_API int32 EmbeddedIndexToTransformIndex(const UInstancedStaticMeshComponent* ISMComponent, int32 InstanceIndex) const;
 	GEOMETRYCOLLECTIONENGINE_API void GetBoneColors(TArray<FColor>& OutColors) const;
 	GEOMETRYCOLLECTIONENGINE_API void GetHiddenTransforms(TArray<bool>& OutHiddenTransforms) const;
@@ -1331,13 +1332,9 @@ public:
 	/** Get any custom renderer. Returns nullptr if none is set. */
 	GEOMETRYCOLLECTIONENGINE_API IGeometryCollectionExternalRenderInterface* GetCustomRenderer() { return CustomRenderer.GetInterface(); }
 
-	
-	/** Enable or disable root proxy component creation when not using a custom renderer - this can be set at runtime */
-	GEOMETRYCOLLECTIONENGINE_API void EnableRootProxyStaticMeshComponents(bool bEnabled);
-
-	/** Enable or disable root proxy for custom rendering - this can be set at runtime */
+	/** Force any custom renderer to render using the broken/decayed path. This can be set at runtime */
 	UFUNCTION(BlueprintCallable, Category = "Physics")
-	GEOMETRYCOLLECTIONENGINE_API void EnableRootProxyForCustomRenderer(bool bEnable);
+	GEOMETRYCOLLECTIONENGINE_API void ForceBrokenForCustomRenderer(bool bForceBroken);
 
 	/** Set a specific root proxy local transform */
 	GEOMETRYCOLLECTIONENGINE_API void SetRootProxyLocalTransform(int32 Index, const FTransform3f& RootProxyTransform);
@@ -1348,13 +1345,16 @@ public:
 	/** Force all GC components to reregister their custom renderer objects. */
 	static GEOMETRYCOLLECTIONENGINE_API void ReregisterAllCustomRenderers();
 
-	/** allow update of the custom renderer ( valid if custom redner is being used ) - true by default */
-	GEOMETRYCOLLECTIONENGINE_API void SetUpdateCustomRenderer(bool bValue) { bUpdateCustomRenderer = bValue; }
-
 	/** update of the custom renderer when post physics sync callback is executing ( valid if custom redner is being used ) - true by default */
 	GEOMETRYCOLLECTIONENGINE_API void SetUpdateCustomRendererOnPostPhysicsSync(bool bValue) { bUpdateCustomRendererOnPostPhysicsSync = bValue; }
 	GEOMETRYCOLLECTIONENGINE_API bool GetUpdateCustomRendererOnPostPhysicsSync() const { return bUpdateCustomRendererOnPostPhysicsSync; }
 
+	UE_DEPRECATED(5.5, "SetUpdateCustomRenderer() shouldn't be called")
+	GEOMETRYCOLLECTIONENGINE_API void SetUpdateCustomRenderer(bool bValue) {}
+	
+	UE_DEPRECATED(5.5, "Please use ForceBrokenForCustomRenderer() instead")
+	UFUNCTION(BlueprintCallable, Category = "Physics", meta = (DeprecatedFunction, DeprecationMessage = "Please use ForceBrokenForCustomRenderer() instead"))
+	GEOMETRYCOLLECTIONENGINE_API void EnableRootProxyForCustomRenderer(bool bEnable) { ForceBrokenForCustomRenderer(!bEnable); }
 
 	GEOMETRYCOLLECTIONENGINE_API bool ShouldUpdateComponentTransformToRootBone() const { return bUpdateComponentTransformToRootBone; }
 
@@ -1471,8 +1471,16 @@ protected:
 	UPROPERTY(EditAnywhere, BlueprintReadWrite, Category = "ChaosPhysics|Rendering", meta = (editcondition = "bOverrideCustomRenderer", MustImplement = "/Script/GeometryCollectionEngine.GeometryCollectionExternalRenderInterface"))
 	TObjectPtr<UClass> CustomRendererType;
 
-	UPROPERTY()
-	bool bEnableRootProxyForCustomRenderer = true;
+	/** Force the broken state for custom renderer rendering. */
+	uint8 bForceBrokenForCustomRenderer : 1;
+	/** Whether to refresh the custom renderer on physics updates. */
+	uint8 bUpdateCustomRendererOnPostPhysicsSync : 1;
+	/** Custom renderer flag to specify if custom renderer can ever require native fallback. */
+	uint8 bCustomRendererCanUseNativeFallback : 1;
+	/** Custom renderer flag to specify if custom renderer should use native fallback. */
+	uint8 bCustomRendererShouldUseNativeFallback : 1;
+	/** Force native geometry collection rendering. This is used in the editor fracture mode. */
+	uint8 bForceNativeRenderer : 1;
 
 	/** A custom renderer object created from CustomRenderType. */
 	UPROPERTY(Transient)
@@ -1748,8 +1756,10 @@ private:
 	GEOMETRYCOLLECTIONENGINE_API bool IsEmbeddedGeometryValid() const;
 	GEOMETRYCOLLECTIONENGINE_API void ClearEmbeddedGeometry();
 
-	/** return true if a a custom renderer has been set and the feature is enabled */
-	GEOMETRYCOLLECTIONENGINE_API bool CanUseCustomRenderer() const;
+	/** return true if a custom renderer has been set and the feature is enabled */
+	GEOMETRYCOLLECTIONENGINE_API bool IsCustomRendererAvailable() const;
+	/** return true if a custom renderer is active */
+	GEOMETRYCOLLECTIONENGINE_API bool IsUsingCustomRenderer() const;
 
 	GEOMETRYCOLLECTIONENGINE_API void RegisterCustomRenderer();
 	GEOMETRYCOLLECTIONENGINE_API void UnregisterCustomRenderer();
@@ -1804,10 +1814,6 @@ private:
 	/** True if GeometryCollection transforms have changed from previous tick. */
 	bool bIsMoving;
 
-	bool bUpdateCustomRenderer;
-
-	bool bUpdateCustomRendererOnPostPhysicsSync;
-
 private:
 	struct FBrokenAndDecayedStates
 	{
@@ -1844,12 +1850,6 @@ private:
 	void UpdateBrokenAndDecayedStates();
 
 	bool ShouldCreateRootProxyComponents() const;
-	void CreateRootProxyComponentsIfNeeded();
-	void UpdateRootProxyComponentsIfNeeded();
-	void ClearRootProxyComponents();
-
-	TArray<TObjectPtr<UStaticMeshComponent>> RootProxyStaticMeshComponents;
-	bool bEnableRootProxyStaticMeshComponents = true;
 
 	TArray<FTransform3f> RootProxyLocalTransforms;
 
diff --git a/Engine/Source/Runtime/Experimental/GeometryCollectionEngine/Public/GeometryCollection/GeometryCollectionExternalRenderInterface.h b/Engine/Source/Runtime/Experimental/GeometryCollectionEngine/Public/GeometryCollection/GeometryCollectionExternalRenderInterface.h
--- a/Engine/Source/Runtime/Experimental/GeometryCollectionEngine/Public/GeometryCollection/GeometryCollectionExternalRenderInterface.h
+++ b/Engine/Source/Runtime/Experimental/GeometryCollectionEngine/Public/GeometryCollection/GeometryCollectionExternalRenderInterface.h
@@ -27,35 +27,46 @@ public:
 		EState_ForcedBroken	= 1 << 2,
 	};
 
-	virtual void OnRegisterGeometryCollection(UGeometryCollectionComponent const& InComponent) = 0;
+	/** Called on geometry collection component registration. */
+	virtual void OnRegisterGeometryCollection(UGeometryCollectionComponent& InComponent) = 0;
+	/** Called on geometry collection component unregistration. */
 	virtual void OnUnregisterGeometryCollection() = 0;
 
 	/** 
-	* Set the state of the geometry collection 
-	* this is used by the renderer to managed resources with regards to the state ( see EStateFlags )
-	*/
+	 * Set the state of the geometry collection.
+	 * This is used by the renderer to managed resources with regards to the state (see EStateFlags).
+	 */
 	virtual void UpdateState(UGeometryCollection const& InGeometryCollection, FTransform const& InComponentTransform, uint32 InStateFlags) = 0;
 
 	/**
-	* Update the root bone transform of the geometry collection 
-	* if the geometry collection is using of multiple proxy root meshes this transform applies to all of them 
-	*/
+	 * Update the root bone transform of the geometry collection.
+	 * if the geometry collection is using of multiple proxy root meshes this transform applies to all of them.
+	 */
 	virtual void UpdateRootTransform(UGeometryCollection const& InGeometryCollection, FTransform const& InRootTransform) = 0;
 
 	/**
-	* Update the root proxy transforms of the geometry collection
-	* if the geometry collection is using of multiple proxy root meshes, InRootTransforms is expected to contain an entry for each of them 
-	* @param InRootTransform		Component space root transform
-	* @param InRootLocalTransforms	Root space local transforms
-	*/
+	 * Update the root proxy transforms of the geometry collection.
+	 * if the geometry collection is using of multiple proxy root meshes, InRootTransforms is expected to contain an entry for each of them .
+	 * @param InRootTransform		Component space root transform
+	 * @param InRootLocalTransforms	Root space local transforms
+	 */
 	virtual void UpdateRootTransforms(UGeometryCollection const& InGeometryCollection, FTransform const& InRootTransform, TArrayView<const FTransform3f> InRootLocalTransforms) { check(false); };
 
 	/**
-	* Update all the bones transforms 
-	*/
+	 * Update all the bones transforms.
+	 */
 	virtual void UpdateTransforms(UGeometryCollection const& InGeometryCollection, TArrayView<const FTransform3f> InTransforms) = 0;
 
-	// todo: Maybe move this to a separate ISMPool renderer interface? But maybe that is overengineering?
+	/** Returns true if the custom renderer can't render for the given state, and requires native rendering instead. */
+	virtual bool ShouldUseNativeFallback(uint32 InStateFlags) const { return false; }
+	
+	/** Returns true if this renderer can ever return true for ShouldUseNativeFallback(). */
+	virtual bool CanEverUseNativeFallback() const { return false; }
+
+	/**
+	 * Set a custom instance data slot to be used by any instances the renderer is spawning.
+	 * todo: This is a bit of a leaky abstraction for ISMPool renderers. Move this to a separate ISMPool renderer interface? Or is that overengineering?
+	 */ 
 	virtual void SetCustomInstanceData(int32 CustomFloatIndex, float CustomFloatValue) {}
 
 	UE_DEPRECATED(5.4, "Use flags version of UpdateState instead")
