Adding ability to enable/disable Verbose Iris CSV stats via CVar rather than at compile time.
#tests Win64 Dedicated Server Cooked Build

diff --git a/Engine/Source/Runtime/Experimental/Iris/Core/Private/Iris/ReplicationSystem/ObjectReplicationBridge.cpp b/Engine/Source/Runtime/Experimental/Iris/Core/Private/Iris/ReplicationSystem/ObjectReplicationBridge.cpp
--- a/Engine/Source/Runtime/Experimental/Iris/Core/Private/Iris/ReplicationSystem/ObjectReplicationBridge.cpp
+++ b/Engine/Source/Runtime/Experimental/Iris/Core/Private/Iris/ReplicationSystem/ObjectReplicationBridge.cpp
@@ -97,6 +97,17 @@ static FAutoConsoleVariableRef CvarEnsureNetRefHandleError(
 	TEXT("Chooses if we should ensure when a NetRefHandleError was reported. -1=never ensure. 0=always ensure. 1..X=ensure only for specific error type")
 );
 
+#if UE_BUILD_SHIPPING
+static bool bUseVerboseIrisCsvStats = false;
+#else
+static bool bUseVerboseIrisCsvStats = true;
+#endif
+static FAutoConsoleVariableRef CVarUseVerboseCsvStats(
+	TEXT("net.Iris.UseVerboseIrisCsvStats"),
+	bUseVerboseIrisCsvStats,
+	TEXT("Whether to use verbose per-class csv stats. Default is false in Shipping, otherwise True.")
+);
+
 UObjectReplicationBridge::FCreateNetRefHandleParams UObjectReplicationBridge::DefaultCreateNetRefHandleParams =
 {
 	.bCanReceive=false, 
@@ -145,6 +156,11 @@ UObjectReplicationBridge::~UObjectReplicationBridge()
 	delete PollFrequencyLimiter;
 }
 
+bool UObjectReplicationBridge::ShouldUseVerboseCsvStats() const
+{
+	return bUseVerboseIrisCsvStats;
+}
+
 void UObjectReplicationBridge::Initialize(UReplicationSystem* InReplicationSystem)
 {
 	LLM_SCOPE_BYTAG(Iris);
@@ -1817,15 +1833,13 @@ void UObjectReplicationBridge::LoadConfig()
 	{
 		for (const FObjectReplicationBridgeTypeStatsConfig& TypeStatsConfig : BridgeConfig->GetTypeStatsConfigs())
 		{
-#if !UE_NET_IRIS_VERBOSE_CSV_STATS
-			// Skip all non shipping TypeStats
-			if (!TypeStatsConfig.bIncludeInMinimalCSVStats)
+			// Skip all non shipping TypeStats unless CVar is set
+			if (!bUseVerboseIrisCsvStats && !TypeStatsConfig.bIncludeInMinimalCSVStats)
 			{
 				continue;
 			}
-#endif			
+			
 			ClassesWithTypeStats.Add(TypeStatsConfig.ClassName, TypeStatsConfig.TypeStatsName);
-
 		}
 	}
 }
diff --git a/Engine/Source/Runtime/Experimental/Iris/Core/Public/Iris/IrisConfig.h b/Engine/Source/Runtime/Experimental/Iris/Core/Public/Iris/IrisConfig.h
--- a/Engine/Source/Runtime/Experimental/Iris/Core/Public/Iris/IrisConfig.h
+++ b/Engine/Source/Runtime/Experimental/Iris/Core/Public/Iris/IrisConfig.h
@@ -35,19 +35,6 @@ IRISCORE_API EReplicationSystem GetUseIrisReplicationCmdlineValue();
 #	define UE_NET_IRIS_CSV_STATS CSV_PROFILER
 #endif
 
-/** Verbose CSV stats is not recommended in shipping builds due to the expected CPU overhead. */
-#if UE_NET_IRIS_CSV_STATS
-#	ifndef UE_NET_IRIS_VERBOSE_CSV_STATS
-#		define UE_NET_IRIS_VERBOSE_CSV_STATS !UE_BUILD_SHIPPING
-#	endif
-#else
-// Force disable UE_NET_IRIS_VERBOSE_CSV_STATS if UE_NET_IRIS_CSV_STATS isn't enabled.
-#	ifdef UE_NET_IRIS_VERBOSE_CSV_STATS
-#		undef UE_NET_IRIS_VERBOSE_CSV_STATS
-#	endif
-#	define UE_NET_IRIS_VERBOSE_CSV_STATS 0
-#endif
-
 /** Enables code that detects non-thread safe access to network data */
 #ifndef UE_NET_THREAD_SAFETY_CHECK
 #	define UE_NET_THREAD_SAFETY_CHECK DO_CHECK
diff --git a/Engine/Source/Runtime/Experimental/Iris/Core/Public/Iris/ReplicationSystem/ObjectReplicationBridge.h b/Engine/Source/Runtime/Experimental/Iris/Core/Public/Iris/ReplicationSystem/ObjectReplicationBridge.h
--- a/Engine/Source/Runtime/Experimental/Iris/Core/Public/Iris/ReplicationSystem/ObjectReplicationBridge.h
+++ b/Engine/Source/Runtime/Experimental/Iris/Core/Public/Iris/ReplicationSystem/ObjectReplicationBridge.h
@@ -182,6 +182,9 @@ public:
 	IRISCORE_API void SetClassTypeStatsConfig(FName ClassPathName, FName TypeStatsName);
 	IRISCORE_API void SetClassTypeStatsConfig(const FString& ClassPathName, const FString& TypeStatsName);
 
+	/** Getter for CVar net.Iris.UseVerboseIrisCsvStats which is true if we want per-class iris stats output to CSV */
+	IRISCORE_API bool ShouldUseVerboseCsvStats() const;
+
 public:
 
 	// Debug functions exposed that are triggered via console commands
