Caching local state from UNetObjectGridFilter::UpdateCellInfoForObject() in order to debug a server crash.

diff --git a/Engine/Source/Runtime/Experimental/Iris/Core/Private/Iris/ReplicationSystem/Filtering/NetObjectGridFilter.cpp b/Engine/Source/Runtime/Experimental/Iris/Core/Private/Iris/ReplicationSystem/Filtering/NetObjectGridFilter.cpp
--- a/Engine/Source/Runtime/Experimental/Iris/Core/Private/Iris/ReplicationSystem/Filtering/NetObjectGridFilter.cpp
+++ b/Engine/Source/Runtime/Experimental/Iris/Core/Private/Iris/ReplicationSystem/Filtering/NetObjectGridFilter.cpp
@@ -372,6 +372,15 @@ void UNetObjectGridFilter::UpdateCellInfoForObject(const FObjectLocationInfo& Ob
 					for (int32 Y = NewCellBox.MinY, EndY = NewCellBox.MaxY + 1; Y < EndY; ++Y)
 					{
 						Coord.Y = Y;
+
+						// Temporarily cache the function state for debugging purposes.
+						DebugUpdateCellInfo.Coord = Coord;
+						DebugUpdateCellInfo.NewCellBox = NewCellBox;
+						DebugUpdateCellInfo.PrevCellBox = PrevCellBox;
+						DebugUpdateCellInfo.CellsSize = Cells.Num();
+						DebugUpdateCellInfo.ObjectPosition = PerObjectInfo.Position;
+						DebugUpdateCellInfo.ObjectCullDistance = PerObjectInfo.GetCullDistance();
+
 						FCellObjects& Cell = Cells.FindOrAdd(Coord);
 						Cell.ObjectIndices.Add(ObjectIndex);
 					}
@@ -431,6 +440,14 @@ void UNetObjectGridFilter::UpdateCellInfoForObject(const FObjectLocationInfo& Ob
 							continue;
 						}
 
+						// Temporarily cache the function state for debugging purposes.
+						DebugUpdateCellInfo.Coord = Coord;
+						DebugUpdateCellInfo.NewCellBox = NewCellBox;
+						DebugUpdateCellInfo.PrevCellBox = PrevCellBox;
+						DebugUpdateCellInfo.CellsSize = Cells.Num();
+						DebugUpdateCellInfo.ObjectPosition = PerObjectInfo.Position;
+						DebugUpdateCellInfo.ObjectCullDistance = PerObjectInfo.GetCullDistance();
+
 						FCellObjects& Cell = Cells.FindOrAdd(Coord);
 						Cell.ObjectIndices.Add(ObjectIndex);
 					}
@@ -459,6 +476,9 @@ void UNetObjectGridFilter::CalculateCellBox(const UNetObjectGridFilter::FPerObje
 {
 	const double CullDistance = PerObjectInfo.GetCullDistance();
 	const FVector Position = PerObjectInfo.Position;
+
+	ensureMsgf(!Position.ContainsNaN(), TEXT("Calculating grid cell box based off a position with a NaN for object: %s"), ToCStr(NetRefHandleManager->PrintObjectFromIndex(PerObjectInfo.ObjectIndex)));
+
 	FVector MinPosition = Position - CullDistance;
 	FVector MaxPosition = Position + CullDistance;
 
diff --git a/Engine/Source/Runtime/Experimental/Iris/Core/Public/Iris/ReplicationSystem/Filtering/NetObjectGridFilter.h b/Engine/Source/Runtime/Experimental/Iris/Core/Public/Iris/ReplicationSystem/Filtering/NetObjectGridFilter.h
--- a/Engine/Source/Runtime/Experimental/Iris/Core/Public/Iris/ReplicationSystem/Filtering/NetObjectGridFilter.h
+++ b/Engine/Source/Runtime/Experimental/Iris/Core/Public/Iris/ReplicationSystem/Filtering/NetObjectGridFilter.h
@@ -286,6 +286,19 @@ private:
 	TMap<FCellCoord, FCellObjects> Cells;
 	uint32 FrameIndex = 0;
 
+	struct FDebugUpdateCellInfo
+	{
+		FCellCoord Coord;
+		FCellBox NewCellBox;
+		FCellBox PrevCellBox;
+		SIZE_T CellsSize;
+		FVector ObjectPosition;
+		float ObjectCullDistance;
+	};
+
+	// Cached state from UpdateCellInfoForObject().
+	FDebugUpdateCellInfo DebugUpdateCellInfo;
+
 	const UE::Net::Private::FNetRefHandleManager*  NetRefHandleManager = nullptr;
 	const UE::Net::FNetCullDistanceOverrides* NetCullDistanceOverrides = nullptr;
 };
