IOStore on demand ChunkPlugin implementation.
Checking this in seperately from other work to help boostrap development.

diff --git a/Engine/Source/Runtime/Experimental/IoStoreOnDemand/Private/IoStoreOnDemand.cpp b/Engine/Source/Runtime/Experimental/IoStoreOnDemand/Private/IoStoreOnDemand.cpp
--- a/Engine/Source/Runtime/Experimental/IoStoreOnDemand/Private/IoStoreOnDemand.cpp
+++ b/Engine/Source/Runtime/Experimental/IoStoreOnDemand/Private/IoStoreOnDemand.cpp
@@ -24,6 +24,7 @@
 #include "Serialization/Archive.h"
 #include "Serialization/CompactBinarySerialization.h"
 #include "Serialization/CompactBinaryWriter.h"
+#include "Serialization/CustomVersion.h"
 #include "Serialization/LargeMemoryWriter.h"
 #include "Serialization/MemoryReader.h"
 #include "Statistics.h"
@@ -439,6 +440,22 @@ bool LoadFromCompactBinary(FCbFieldView Field, FOnDemandTocEntry& OutTocEntry)
 
 FArchive& operator<<(FArchive& Ar, FOnDemandTocContainerEntry& ContainerEntry)
 {
+	if (Ar.IsLoading())
+	{
+		const FCustomVersion* CustomVersion = Ar.GetCustomVersions().GetVersion(FOnDemandToc::VersionGuid);
+		check(CustomVersion);
+		const uint32 TocVersion = IntCastChecked<uint32>(CustomVersion->Version);
+
+		if (TocVersion >= uint32(EOnDemandTocVersion::ContainerId))
+		{
+			Ar << ContainerEntry.ContainerId;
+		}
+	}
+	else
+	{
+		Ar << ContainerEntry.ContainerId;
+	}
+
 	Ar << ContainerEntry.ContainerName;
 	Ar << ContainerEntry.EncryptionKeyGuid;
 	Ar << ContainerEntry.Entries;
@@ -555,6 +572,38 @@ FArchive& operator<<(FArchive& Ar, FOnDemandTocSentinel& Sentinel)
 	return Ar;
 }
 
+FArchive& operator<<(FArchive& Ar, FOnDemandTocAdditionalFile& AdditionalFile)
+{
+	Ar << AdditionalFile.Hash;
+	Ar << AdditionalFile.Filename;
+	Ar << AdditionalFile.FileSize;
+	return Ar;
+}
+
+FCbWriter& operator<<(FCbWriter& Writer, const FOnDemandTocAdditionalFile& AdditionalFile)
+{
+	Writer.BeginObject();
+	Writer.AddHash(UTF8TEXTVIEW("Hash"), AdditionalFile.Hash);
+	Writer.AddString(UTF8TEXTVIEW("Filename"), AdditionalFile.Filename);
+	Writer.AddInteger(UTF8TEXTVIEW("Filename"), AdditionalFile.FileSize);
+	Writer.EndObject();
+
+	return Writer;
+}
+
+bool LoadFromCompactBinary(FCbFieldView Field, FOnDemandTocAdditionalFile& AdditionalFile)
+{
+	if (FCbObjectView Obj = Field.AsObjectView())
+	{
+		AdditionalFile.Hash = Obj["Hash"].AsHash();
+		AdditionalFile.Filename = FString(Obj["Filename"].AsString());
+		AdditionalFile.FileSize = Obj["FileSize"].AsUInt64();
+		return true;
+	}
+
+	return false;
+}
+
 FArchive& operator<<(FArchive& Ar, FOnDemandToc& Toc)
 {
 	Ar << Toc.Header;
@@ -571,6 +620,11 @@ FArchive& operator<<(FArchive& Ar, FOnDemandToc& Toc)
 	}
 	Ar << Toc.Containers;
 
+	if (uint32(Toc.Header.Version) >= uint32(EOnDemandTocVersion::AdditionalFiles))
+	{
+		Ar << Toc.AdditionalFiles;
+	}
+
 	return Ar;
 }
 
@@ -585,6 +639,17 @@ FCbWriter& operator<<(FCbWriter& Writer, const FOnDemandToc& Toc)
 		Writer << Container;
 	}
 	Writer.EndArray();
+
+	if (Toc.AdditionalFiles.Num() > 0)
+	{
+		Writer.BeginArray(UTF8TEXTVIEW("Files"));
+		for (const FOnDemandTocAdditionalFile& File : Toc.AdditionalFiles)
+		{
+			Writer << File;
+		}
+		Writer.EndArray();
+	}
+
 	Writer.EndObject();
 	
 	return Writer;
@@ -619,6 +684,19 @@ bool LoadFromCompactBinary(FCbFieldView Field, FOnDemandToc& OutToc)
 			}
 		}
 
+		if (uint32(OutToc.Header.Version) >= uint32(EOnDemandTocVersion::AdditionalFiles))
+		{
+			FCbArrayView Files = Obj["Files"].AsArrayView();
+			OutToc.AdditionalFiles.Reserve(int32(Files.Num()));
+			for (FCbFieldView ArrayField : Files)
+			{
+				if (!LoadFromCompactBinary(ArrayField, OutToc.AdditionalFiles.AddDefaulted_GetRef()))
+				{
+					return false;
+				}
+			}
+		}
+
 		return true;
 	}
 
@@ -2379,9 +2457,9 @@ FIoStatus ListTocs(const FIoStoreListTocsParams& Params)
 ////////////////////////////////////////////////////////////////////////////////
 void FIoStoreOnDemandModule::SetBulkOptionalEnabled(bool bInEnabled)
 {
-	if (Backend.IsValid())
+	if (HttpIoDispatcherBackend.IsValid())
 	{
-		Backend->SetBulkOptionalEnabled(bInEnabled);
+		HttpIoDispatcherBackend->SetBulkOptionalEnabled(bInEnabled);
 	}
 	else
 	{
@@ -2392,9 +2470,9 @@ void FIoStoreOnDemandModule::SetBulkOptionalEnabled(bool bInEnabled)
 
 void FIoStoreOnDemandModule::SetEnabled(bool bInEnabled)
 {
-	if (Backend.IsValid())
+	if (HttpIoDispatcherBackend.IsValid())
 	{
-		Backend->SetEnabled(bInEnabled);
+		HttpIoDispatcherBackend->SetEnabled(bInEnabled);
 	}
 	else
 	{
@@ -2405,9 +2483,9 @@ void FIoStoreOnDemandModule::SetEnabled(bool bInEnabled)
 
 void FIoStoreOnDemandModule::AbandonCache()
 {
-	if (Backend.IsValid())
+	if (HttpIoDispatcherBackend.IsValid())
 	{
-		Backend->AbandonCache();
+		HttpIoDispatcherBackend->AbandonCache();
 	}
 	else
 	{
@@ -2418,14 +2496,14 @@ void FIoStoreOnDemandModule::AbandonCache()
 
 bool FIoStoreOnDemandModule::IsEnabled() const
 {
-	return Backend.IsValid()? Backend->IsEnabled():DeferredAbandonCache.IsSet();
+	return HttpIoDispatcherBackend.IsValid()? HttpIoDispatcherBackend->IsEnabled():DeferredAbandonCache.IsSet();
 }
 
 void FIoStoreOnDemandModule::ReportAnalytics(TArray<FAnalyticsEventAttribute>& OutAnalyticsArray) const
 {
-	if (Backend.IsValid())
+	if (HttpIoDispatcherBackend.IsValid())
 	{
-		Backend->ReportAnalytics(OutAnalyticsArray);
+		HttpIoDispatcherBackend->ReportAnalytics(OutAnalyticsArray);
 	}
 }
 
@@ -2453,7 +2531,7 @@ void FIoStoreOnDemandModule::InitializeInternal()
 #endif
 
 	// Make sure we haven't called initialize before
-	check(!Backend.IsValid());
+	check(!HttpIoDispatcherBackend.IsValid());
 
 	FOnDemandEndpoint Endpoint;
 	
@@ -2512,8 +2590,8 @@ void FIoStoreOnDemandModule::InitializeInternal()
 			(CacheConfig.DiskQuota > 0) ? TEXT("init-fail") : TEXT("zero-quota"));
 	}
 
-	Backend = MakeOnDemandIoDispatcherBackend(MoveTemp(Cache));
-	Backend->Mount(Endpoint);
+	HttpIoDispatcherBackend = MakeOnDemandIoDispatcherBackend(MoveTemp(Cache));
+	HttpIoDispatcherBackend->Mount(Endpoint);
 	int32 BackendPriority = -10;
 #if !UE_BUILD_SHIPPING
 	if (FParse::Param(CommandLine, TEXT("Ias")))
@@ -2526,14 +2604,14 @@ void FIoStoreOnDemandModule::InitializeInternal()
 	// Setup any states changes issued before initialization
 	if (DeferredEnabled.IsSet())
 	{
-		Backend->SetEnabled(*DeferredEnabled);
+		HttpIoDispatcherBackend->SetEnabled(*DeferredEnabled);
 	}
 	if (DeferredBulkOptionalEnabled.IsSet())
 	{
-		Backend->SetBulkOptionalEnabled(*DeferredBulkOptionalEnabled);
+		HttpIoDispatcherBackend->SetBulkOptionalEnabled(*DeferredBulkOptionalEnabled);
 	}
 	
-	FIoDispatcher::Get().Mount(Backend.ToSharedRef(), BackendPriority);
+	FIoDispatcher::Get().Mount(HttpIoDispatcherBackend.ToSharedRef(), BackendPriority);
 }
 	
 void FIoStoreOnDemandModule::StartupModule()
@@ -2568,7 +2646,7 @@ EOnDemandInitResult FIoStoreOnDemandModule::Initialize()
 
 	InitializeInternal();
 
-	return Backend.IsValid() ? EOnDemandInitResult::Success : EOnDemandInitResult::Disabled;
+	return HttpIoDispatcherBackend.IsValid() ? EOnDemandInitResult::Success : EOnDemandInitResult::Disabled;
 };
 
 #endif // UE_IAS_CUSTOM_INITIALIZATION
diff --git a/Engine/Source/Runtime/Experimental/IoStoreOnDemand/Private/Tool/ChunkPlugin.cpp b/Engine/Source/Runtime/Experimental/IoStoreOnDemand/Private/Tool/ChunkPlugin.cpp
new file mode 100644
--- /dev/null
+++ b/Engine/Source/Runtime/Experimental/IoStoreOnDemand/Private/Tool/ChunkPlugin.cpp
@@ -0,0 +1,319 @@
+// Copyright Epic Games, Inc. All Rights Reserved.
+
+#if defined(UE_WITH_IAS_TOOL)
+
+#include "Command.h"
+
+#include "Algo/Find.h"
+#include "Containers/StringView.h"
+#include "HAL/FileManager.h"
+#include "IO/IoChunkEncoding.h"
+#include "IO/IoStatus.h"
+#include "IO/IoStore.h"
+#include "IO/IoStoreOnDemand.h"
+#include "Misc/DateTime.h"
+#include "Misc/FileHelper.h"
+#include "Misc/Paths.h"
+#include "Misc/PathViews.h"
+#include "Misc/StringBuilder.h"
+#include "Serialization/LargeMemoryWriter.h"
+
+namespace UE::IO::IAS::Tool
+{
+
+////////////////////////////////////////////////////////////////////////////////
+FIoStatus WriteChunk(const FString& Directory, FMemoryView Chunk, const FIoHash& Hash)
+{
+	IFileManager& FileMgr = IFileManager::Get();
+	const FString HashString = LexToString(Hash);
+
+	TStringBuilder<256> Sb;
+	Sb << Directory << TEXT("/") << HashString.Left(2);
+
+	bool bTree = true;
+	if (FileMgr.MakeDirectory(Sb.ToString(), bTree) == false)
+	{
+		return FIoStatusBuilder(EIoErrorCode::WriteError)
+			<< TEXT("Failed to create directory '")
+			<< FString(Sb.ToString())
+			<< TEXT("'");
+	}
+
+	Sb << TEXT("/") << HashString << TEXT(".iochunk");
+
+	if (TUniquePtr<FArchive> Ar(IFileManager::Get().CreateFileWriter(Sb.ToString())); Ar.IsValid())
+	{
+		UE_LOG(LogIoStore, Display, TEXT("Writing file '%s' (%.2lf KiB)"),
+			Sb.ToString(), double(Chunk.GetSize()) / 1024.0);
+		Ar->Serialize((void*)Chunk.GetData(), Chunk.GetSize());
+
+		return EIoErrorCode::Ok;
+	}
+
+	return FIoStatusBuilder(EIoErrorCode::WriteError)
+		<< TEXT("Failed to write file '")
+		<< FString(Sb.ToString())
+		<< TEXT("'");
+}
+
+////////////////////////////////////////////////////////////////////////////////
+FIoStatus WriteChunk(const FString& Directory, FMemoryView Chunk)
+{
+	return WriteChunk(Directory, Chunk, FIoHash::HashBuffer(Chunk));
+}
+
+////////////////////////////////////////////////////////////////////////////////
+static int32 ChunkPluginCommandEntry(const FContext& Context)
+{
+	const FString Platform				= FString(Context.Get<FStringView>(TEXT("-Platform"), FString()));
+	const FString BuildVersion			= FString(Context.Get<FStringView>(TEXT("-BuildVersion"), FString()));
+	const FString OnDemandTocName		= FString(Context.Get<FStringView>(TEXT("-OnDemandTocName"), FString()));
+	const FString InputFolder			= FString(Context.Get<FStringView>(TEXT("-InputFolder"), FString()));
+	const FString OutputFolder			= FString(Context.Get<FStringView>(TEXT("-OutputFolder"), FString()));
+	const FString IntermediateFolder	= FString(Context.Get<FStringView>(TEXT("-IntermediateFolder"), FString()));
+	FString ContainerFolder				= IntermediateFolder / TEXT("Cooked");
+	FString IoStoreOutputFolder			= OutputFolder / TEXT("iostore");
+	FString ChunksOutputFolder			= IoStoreOutputFolder / TEXT("chunks");
+
+	FPaths::NormalizeDirectoryName(ContainerFolder);
+	FPaths::NormalizeDirectoryName(IoStoreOutputFolder);
+	FPaths::NormalizeDirectoryName(ChunksOutputFolder);
+
+	UE_LOG(LogIoStore, Display, TEXT("I/O store chunk plugin:"));
+	UE_LOG(LogIoStore, Display, TEXT("----------------------------------------"));
+	UE_LOG(LogIoStore, Display, TEXT("\tBuildVersion: %s"), *BuildVersion);
+	UE_LOG(LogIoStore, Display, TEXT("\tPlatform: %s"), *Platform);
+	UE_LOG(LogIoStore, Display, TEXT("\tOnDemandTocName: %s"), *OnDemandTocName);
+	UE_LOG(LogIoStore, Display, TEXT("\tInputFolder: %s"), *InputFolder);
+	UE_LOG(LogIoStore, Display, TEXT("\tOutputFolder: %s"), *OutputFolder);
+	UE_LOG(LogIoStore, Display, TEXT("\tIntermediateFolder: %s"), *IntermediateFolder);
+
+	IFileManager& FileMgr = IFileManager::Get();
+	if (FileMgr.MakeDirectory(*ChunksOutputFolder, true) == false)
+	{
+		UE_LOG(LogIoStore, Error, TEXT("Failed to create directory '%s'"), *ChunksOutputFolder);
+		return -1;
+	}
+
+	TMap<FGuid, FAES::FAESKey> EncryptionKeys;
+	if (IFileManager::Get().DirectoryExists(*ContainerFolder) == false)
+	{
+		UE_LOG(LogIoStore, Error, TEXT("Directory '%s' does not exist"), *ContainerFolder);
+		return -1;
+	}
+
+	TArray<FString> ContainerFilenames;
+	FileMgr.FindFiles(ContainerFilenames, *ContainerFolder, TEXT("*.utoc"));
+	UE_LOG(LogIoStore, Display, TEXT("Found %d container files(s)"), ContainerFilenames.Num());
+
+	FOnDemandToc OnDemandToc;
+	//OnDemandToc.Header.ChunksDirectory = TODO 
+	OnDemandToc.Containers.Reserve(ContainerFilenames.Num());
+
+	for (const FString& Filename : ContainerFilenames)
+	{
+		const FString FullPath = ContainerFolder / Filename;
+		FIoStoreReader ContainerFileReader;
+		{
+			FIoStatus Status = ContainerFileReader.Initialize(*FPaths::ChangeExtension(FullPath, TEXT("")), EncryptionKeys);
+			if (!Status.IsOk())
+			{
+				UE_LOG(LogIoStore, Error, TEXT("Failed to open container '%s' for reading"), *FullPath);
+				continue;
+			}
+		}
+
+		UE_LOG(LogIoStore, Display, TEXT("Processing container '%s'"), *FullPath);
+
+		const uint32 BlockSize = ContainerFileReader.GetCompressionBlockSize();
+		if (OnDemandToc.Header.BlockSize == 0)
+		{
+			OnDemandToc.Header.BlockSize = ContainerFileReader.GetCompressionBlockSize();
+		}
+		check(OnDemandToc.Header.BlockSize == ContainerFileReader.GetCompressionBlockSize());
+
+		TArray<FIoStoreTocChunkInfo> ChunkInfos;
+		ContainerFileReader.EnumerateChunks([&ChunkInfos](FIoStoreTocChunkInfo&& Info)
+		{ 
+			ChunkInfos.Emplace(MoveTemp(Info));
+			return true;
+		});
+
+		UE_LOG(LogIoStore, Display, TEXT("Serializing %d chunks"), ChunkInfos.Num());
+		FOnDemandTocContainerEntry& ContainerEntry = OnDemandToc.Containers.AddDefaulted_GetRef();
+		ContainerEntry.ContainerId		= ContainerFileReader.GetContainerId();
+		ContainerEntry.ContainerName	= FPaths::GetBaseFilename(FullPath);
+
+		if (EnumHasAnyFlags(ContainerFileReader.GetContainerFlags(), EIoContainerFlags::Encrypted))
+		{
+			ContainerEntry.EncryptionKeyGuid = LexToString(ContainerFileReader.GetEncryptionKeyGuid());
+		}
+
+		for (const FIoStoreTocChunkInfo& ChunkInfo : ChunkInfos)
+		{
+			const bool bDecrypt = false;
+			TIoStatusOr<FIoStoreCompressedReadResult> Status = ContainerFileReader.ReadCompressed(ChunkInfo.Id, FIoReadOptions(), bDecrypt);
+			if (!Status.IsOk())
+			{
+				UE_LOG(LogIoStore, Error, TEXT("Failed to read container chunk, Container='%s', Reason='%s'"),
+					*FullPath, *Status.Status().ToString());
+				return -1;
+			}
+
+			FIoStoreCompressedReadResult ReadResult = Status.ConsumeValueOrDie();
+
+			const uint32 BlockOffset	= ContainerEntry.BlockSizes.Num();
+			const uint32 BlockCount		= ReadResult.Blocks.Num();
+			const FIoHash ChunkHash		= FIoHash::HashBuffer(ReadResult.IoBuffer.GetView());
+			const FString HashString	= LexToString(ChunkHash);
+
+			FMemoryView EncodedBlocks = ReadResult.IoBuffer.GetView();
+			uint64 RawChunkSize = 0;
+			uint64 EncodedChunkSize = 0;
+			for (const FIoStoreCompressedBlockInfo& BlockInfo : ReadResult.Blocks)
+			{
+				check(Align(BlockInfo.CompressedSize, FAES::AESBlockSize) == BlockInfo.AlignedSize);
+				const uint64 EncodedBlockSize = BlockInfo.AlignedSize;
+				ContainerEntry.BlockSizes.Add(uint32(BlockInfo.CompressedSize));
+
+				FMemoryView EncodedBlock = EncodedBlocks.Left(EncodedBlockSize);
+				EncodedBlocks += EncodedBlock.GetSize();
+				ContainerEntry.BlockHashes.Add(FIoChunkEncoding::HashBlock(EncodedBlock));
+
+				EncodedChunkSize += EncodedBlockSize;
+				RawChunkSize += BlockInfo.UncompressedSize;
+
+				if (OnDemandToc.Header.CompressionFormat.IsEmpty() && BlockInfo.CompressionMethod != NAME_None)
+				{
+					OnDemandToc.Header.CompressionFormat = BlockInfo.CompressionMethod.ToString();
+				}
+			}
+
+			if (EncodedChunkSize != ReadResult.IoBuffer.GetSize())
+			{
+				UE_LOG(LogIoStore, Error, TEXT("Chunk size mismatch, Container='%s', ChunkId='%s'"),
+					*FullPath, *LexToString(ChunkInfo.Id));
+				return -1;
+			}
+
+			if (FIoStatus WriteStatus = WriteChunk(ChunksOutputFolder, ReadResult.IoBuffer.GetView(), ChunkHash);
+				WriteStatus.IsOk() == false)
+			{
+				UE_LOG(LogIoStore, Error, TEXT("%s"), *WriteStatus.ToString());
+				return -1;
+			}
+
+			FOnDemandTocEntry& TocEntry = ContainerEntry.Entries.AddDefaulted_GetRef();
+			TocEntry.ChunkId = ChunkInfo.Id;
+			TocEntry.Hash = ChunkHash;
+			TocEntry.RawSize = RawChunkSize;
+			TocEntry.EncodedSize = EncodedChunkSize;
+			TocEntry.BlockOffset = BlockOffset;
+			TocEntry.BlockCount = BlockCount;
+		}
+	}
+
+	// Write additional file(s)
+	{
+		const FStringView AllowedExt[]
+		{
+			TEXTVIEW("pak"),
+			TEXTVIEW("sig"),
+		};
+
+		TStringBuilder<256> Sb;
+		TArray<FString> AdditionalFiles;
+		FileMgr.FindFiles(AdditionalFiles, *ContainerFolder, TEXT("*.*"));
+
+		UE_LOG(LogIoStore, Display, TEXT("Serializing %d additional file(s)"), AdditionalFiles.Num());
+		for (const FString& Filename : AdditionalFiles)
+		{
+			const bool bIncludeDot = false;
+			const FStringView Ext = FPathViews::GetExtension(Filename, bIncludeDot);
+
+			if (const FStringView* Result =
+				Algo::FindByPredicate(AllowedExt, [&Ext](const FStringView& E) { return Ext == E; });
+				Result == nullptr)
+			{
+				continue;
+			}
+
+			const FString FullPath = ContainerFolder / Filename;
+			TArray<uint8> FileData;
+			if (!FFileHelper::LoadFileToArray(FileData, *FullPath))
+			{
+				UE_LOG(LogIoStore, Error, TEXT("Failed reading file '%s'"), *FullPath);
+				return -1;
+			}
+
+			FMemoryView Chunk = MakeMemoryView(FileData.GetData(), FileData.Num());
+			const FIoHash ChunkHash = FIoHash::HashBuffer(Chunk);
+			if (FIoStatus WriteStatus = WriteChunk(ChunksOutputFolder, Chunk, ChunkHash);
+				WriteStatus.IsOk() == false)
+			{
+				UE_LOG(LogIoStore, Error, TEXT("%s"), *WriteStatus.ToString());
+				return -1;
+			}
+
+			UE_LOG(LogIoStore, Display, TEXT("Adding additional file '%s'"), *Filename);
+			OnDemandToc.AdditionalFiles.Add(FOnDemandTocAdditionalFile
+			{
+				.Hash = ChunkHash,
+				.Filename = Filename,
+				.FileSize = IntCastChecked<uint64>(FileData.Num())
+			});
+		}
+	}
+
+	OnDemandToc.Meta.EpochTimestamp = FDateTime::UtcNow().ToUnixTimestamp();
+	OnDemandToc.Meta.TargetPlatform = Platform;
+	OnDemandToc.Meta.BuildVersion	= BuildVersion;
+	{
+		FString Filename = FPathViews::SetExtension(OnDemandTocName, TEXT(".uondemandtoc"));
+		Filename.ToLowerInline();
+
+		const FString TocPath = IoStoreOutputFolder / Filename;
+		if (TUniquePtr<FArchive> Ar(IFileManager::Get().CreateFileWriter(*TocPath)); Ar.IsValid())
+		{
+			*Ar << OnDemandToc;
+			if (Ar->IsError())
+			{
+				UE_LOG(LogIoStore, Error, TEXT("Failed to serialize TOC"));
+				return -1;
+			}
+			else
+			{
+				const int64 TocSize = Ar->Tell();
+				UE_LOG(LogIoStore, Display, TEXT("Writing file '%s' (%.2lf KiB)"),
+					*Filename, double(TocSize) / 1024);
+			}
+		}
+		else
+		{
+			UE_LOG(LogIoStore, Display, TEXT("Failed writing file '%s'"), *TocPath );
+			return -1;
+		}
+	}
+
+	return 0;
+}
+
+////////////////////////////////////////////////////////////////////////////////
+static FCommand ChunkPluginCommand(
+	ChunkPluginCommandEntry,
+	TEXT("ChunkPlugin"),
+	TEXT(""),
+	{
+		TArgument<FStringView>(TEXT("-Platform"),			TEXT("Platform name.")),
+		TArgument<FStringView>(TEXT("-BuildVersion"),		TEXT("Build version")),
+		TArgument<FStringView>(TEXT("-OnDemandTocName"),	TEXT("On Demand TOC Name")),
+		TArgument<FStringView>(TEXT("-InputFolder"),		TEXT("Input folder to plugin information.")),
+		TArgument<FStringView>(TEXT("-OutputFolder"),		TEXT("Ouptut folder.")),
+		TArgument<FStringView>(TEXT("-IntermediateFolder"),	TEXT("Intermediate folder.")),
+		TArgument<FStringView>(TEXT("-ErrorOutput"),		TEXT("Error output.")),
+	}
+);
+
+} // namespace UE::IO::IAS::Tool
+
+#endif // UE_WITH_IAS_TOOL
diff --git a/Engine/Source/Runtime/Experimental/IoStoreOnDemand/Public/IO/IoStoreOnDemand.h b/Engine/Source/Runtime/Experimental/IoStoreOnDemand/Public/IO/IoStoreOnDemand.h
--- a/Engine/Source/Runtime/Experimental/IoStoreOnDemand/Public/IO/IoStoreOnDemand.h
+++ b/Engine/Source/Runtime/Experimental/IoStoreOnDemand/Public/IO/IoStoreOnDemand.h
@@ -4,6 +4,7 @@
 
 #include "Containers/UnrealString.h"
 #include "IO/IoChunkId.h"
+#include "IO/IoContainerId.h"
 #include "IO/IoHash.h"
 #include "IO/IoStatus.h"
 #include "Misc/Guid.h"
@@ -53,6 +54,8 @@ enum class EOnDemandTocVersion : uint32
 	BlockHash32		= 3,
 	NoRawHash		= 4,
 	Meta			= 5,
+	ContainerId		= 6,
+	AdditionalFiles	= 7,
 
 	LatestPlusOne,
 	Latest			= (LatestPlusOne - 1)
@@ -113,6 +116,7 @@ UE_API bool LoadFromCompactBinary(FCbFieldView Field, FOnDemandTocEntry& OutTocE
 
 struct FOnDemandTocContainerEntry
 {
+	FIoContainerId ContainerId;
 	FString ContainerName;
 	FString EncryptionKeyGuid;
 	TArray<FOnDemandTocEntry> Entries;
@@ -142,6 +146,18 @@ private:
 	uint8 Data[SentinelSize] = { 0 };
 };
 
+struct FOnDemandTocAdditionalFile
+{
+	FIoHash Hash;
+	FString Filename;
+	uint64 FileSize = 0;
+
+	UE_API friend FArchive& operator<<(FArchive& Ar, FOnDemandTocAdditionalFile& AdditionalFile);
+	UE_API friend FCbWriter& operator<<(FCbWriter& Writer, const FOnDemandTocAdditionalFile& AdditionalFile);
+};
+
+UE_API bool LoadFromCompactBinary(FCbFieldView Field, FOnDemandTocAdditionalFile& AdditionalFile);
+
 struct FOnDemandToc
 {
 	FOnDemandToc() = default;
@@ -158,6 +174,7 @@ struct FOnDemandToc
 	FOnDemandTocHeader Header;
 	FTocMeta Meta;
 	TArray<FOnDemandTocContainerEntry> Containers;
+	TArray<FOnDemandTocAdditionalFile> AdditionalFiles;
 
 	UE_API friend FArchive& operator<<(FArchive& Ar, FOnDemandToc& Toc);
 	UE_API friend FCbWriter& operator<<(FCbWriter& Writer, const FOnDemandToc& Toc);
@@ -329,7 +346,7 @@ class FIoStoreOnDemandModule
 {
 private:
 	void InitializeInternal();
-	TSharedPtr<IOnDemandIoDispatcherBackend> Backend;
+	TSharedPtr<IOnDemandIoDispatcherBackend> HttpIoDispatcherBackend;
 	// Deferred state requests if called before backend
 	// is initialized
 	TOptional<bool> DeferredEnabled;
