Chaos Suspension Module:
Add impact normal to sim output.

diff --git a/Engine/Source/Runtime/Experimental/ChaosVehicles/ChaosVehiclesCore/Private/SimModule/SuspensionModule.cpp b/Engine/Source/Runtime/Experimental/ChaosVehicles/ChaosVehiclesCore/Private/SimModule/SuspensionModule.cpp
--- a/Engine/Source/Runtime/Experimental/ChaosVehicles/ChaosVehiclesCore/Private/SimModule/SuspensionModule.cpp
+++ b/Engine/Source/Runtime/Experimental/ChaosVehicles/ChaosVehiclesCore/Private/SimModule/SuspensionModule.cpp
@@ -187,6 +187,7 @@ namespace Chaos
 			SpringDisplacement = Sim->SpringDisplacement;
 			SpringDisplacementVector = -Sim->Setup().SuspensionAxis * Sim->SpringDisplacement + Sim->GetAnimationOffset();
 			SpringSpeed = Sim->SpringSpeed;
+			ImpactNormal = Sim->ImpactNormal;
 
 			AnimFlags = EAnimationFlags::AnimatePosition;
 			AnimationLocOffset = SpringDisplacementVector;
@@ -203,6 +204,7 @@ namespace Chaos
 
 		SpringDisplacement = FMath::Lerp(Current.SpringDisplacement, Next.SpringDisplacement, Alpha);
 		SpringSpeed = FMath::Lerp(Current.SpringSpeed, Next.SpringSpeed, Alpha);
+		ImpactNormal = FMath::Lerp(Current.ImpactNormal, Next.ImpactNormal, Alpha);
 	}
 
 #if !(UE_BUILD_SHIPPING || UE_BUILD_TEST)
diff --git a/Engine/Source/Runtime/Experimental/ChaosVehicles/ChaosVehiclesCore/Public/SimModule/SuspensionModule.h b/Engine/Source/Runtime/Experimental/ChaosVehicles/ChaosVehiclesCore/Public/SimModule/SuspensionModule.h
--- a/Engine/Source/Runtime/Experimental/ChaosVehicles/ChaosVehiclesCore/Public/SimModule/SuspensionModule.h
+++ b/Engine/Source/Runtime/Experimental/ChaosVehicles/ChaosVehiclesCore/Public/SimModule/SuspensionModule.h
@@ -57,6 +57,7 @@ namespace Chaos
 		float SpringDisplacement;
 		FVector SpringDisplacementVector;
 		float SpringSpeed;
+		FVector ImpactNormal;
 	};
 
 	struct CHAOSVEHICLESCORE_API FSuspensionSettings
