[Backout] - CL32679175 - CIS Compile Error
Original CL Desc
-----------------------------------------------------------------
IOStore on demand ChunkPlugin implementation.
Checking this in seperately from other work to help boostrap development.

diff --git a/Engine/Source/Runtime/Experimental/IoStoreOnDemand/Private/IoStoreOnDemand.cpp b/Engine/Source/Runtime/Experimental/IoStoreOnDemand/Private/IoStoreOnDemand.cpp
--- a/Engine/Source/Runtime/Experimental/IoStoreOnDemand/Private/IoStoreOnDemand.cpp
+++ b/Engine/Source/Runtime/Experimental/IoStoreOnDemand/Private/IoStoreOnDemand.cpp
@@ -24,7 +24,6 @@
 #include "Serialization/Archive.h"
 #include "Serialization/CompactBinarySerialization.h"
 #include "Serialization/CompactBinaryWriter.h"
-#include "Serialization/CustomVersion.h"
 #include "Serialization/LargeMemoryWriter.h"
 #include "Serialization/MemoryReader.h"
 #include "Statistics.h"
@@ -440,22 +439,6 @@ bool LoadFromCompactBinary(FCbFieldView Field, FOnDemandTocEntry& OutTocEntry)
 
 FArchive& operator<<(FArchive& Ar, FOnDemandTocContainerEntry& ContainerEntry)
 {
-	if (Ar.IsLoading())
-	{
-		const FCustomVersion* CustomVersion = Ar.GetCustomVersions().GetVersion(FOnDemandToc::VersionGuid);
-		check(CustomVersion);
-		const uint32 TocVersion = IntCastChecked<uint32>(CustomVersion->Version);
-
-		if (TocVersion >= uint32(EOnDemandTocVersion::ContainerId))
-		{
-			Ar << ContainerEntry.ContainerId;
-		}
-	}
-	else
-	{
-		Ar << ContainerEntry.ContainerId;
-	}
-
 	Ar << ContainerEntry.ContainerName;
 	Ar << ContainerEntry.EncryptionKeyGuid;
 	Ar << ContainerEntry.Entries;
@@ -572,38 +555,6 @@ FArchive& operator<<(FArchive& Ar, FOnDemandTocSentinel& Sentinel)
 	return Ar;
 }
 
-FArchive& operator<<(FArchive& Ar, FOnDemandTocAdditionalFile& AdditionalFile)
-{
-	Ar << AdditionalFile.Hash;
-	Ar << AdditionalFile.Filename;
-	Ar << AdditionalFile.FileSize;
-	return Ar;
-}
-
-FCbWriter& operator<<(FCbWriter& Writer, const FOnDemandTocAdditionalFile& AdditionalFile)
-{
-	Writer.BeginObject();
-	Writer.AddHash(UTF8TEXTVIEW("Hash"), AdditionalFile.Hash);
-	Writer.AddString(UTF8TEXTVIEW("Filename"), AdditionalFile.Filename);
-	Writer.AddInteger(UTF8TEXTVIEW("Filename"), AdditionalFile.FileSize);
-	Writer.EndObject();
-
-	return Writer;
-}
-
-bool LoadFromCompactBinary(FCbFieldView Field, FOnDemandTocAdditionalFile& AdditionalFile)
-{
-	if (FCbObjectView Obj = Field.AsObjectView())
-	{
-		AdditionalFile.Hash = Obj["Hash"].AsHash();
-		AdditionalFile.Filename = FString(Obj["Filename"].AsString());
-		AdditionalFile.FileSize = Obj["FileSize"].AsUInt64();
-		return true;
-	}
-
-	return false;
-}
-
 FArchive& operator<<(FArchive& Ar, FOnDemandToc& Toc)
 {
 	Ar << Toc.Header;
@@ -620,11 +571,6 @@ FArchive& operator<<(FArchive& Ar, FOnDemandToc& Toc)
 	}
 	Ar << Toc.Containers;
 
-	if (uint32(Toc.Header.Version) >= uint32(EOnDemandTocVersion::AdditionalFiles))
-	{
-		Ar << Toc.AdditionalFiles;
-	}
-
 	return Ar;
 }
 
@@ -639,17 +585,6 @@ FCbWriter& operator<<(FCbWriter& Writer, const FOnDemandToc& Toc)
 		Writer << Container;
 	}
 	Writer.EndArray();
-
-	if (Toc.AdditionalFiles.Num() > 0)
-	{
-		Writer.BeginArray(UTF8TEXTVIEW("Files"));
-		for (const FOnDemandTocAdditionalFile& File : Toc.AdditionalFiles)
-		{
-			Writer << File;
-		}
-		Writer.EndArray();
-	}
-
 	Writer.EndObject();
 	
 	return Writer;
@@ -684,19 +619,6 @@ bool LoadFromCompactBinary(FCbFieldView Field, FOnDemandToc& OutToc)
 			}
 		}
 
-		if (uint32(OutToc.Header.Version) >= uint32(EOnDemandTocVersion::AdditionalFiles))
-		{
-			FCbArrayView Files = Obj["Files"].AsArrayView();
-			OutToc.AdditionalFiles.Reserve(int32(Files.Num()));
-			for (FCbFieldView ArrayField : Files)
-			{
-				if (!LoadFromCompactBinary(ArrayField, OutToc.AdditionalFiles.AddDefaulted_GetRef()))
-				{
-					return false;
-				}
-			}
-		}
-
 		return true;
 	}
 
@@ -2457,9 +2379,9 @@ FIoStatus ListTocs(const FIoStoreListTocsParams& Params)
 ////////////////////////////////////////////////////////////////////////////////
 void FIoStoreOnDemandModule::SetBulkOptionalEnabled(bool bInEnabled)
 {
-	if (HttpIoDispatcherBackend.IsValid())
+	if (Backend.IsValid())
 	{
-		HttpIoDispatcherBackend->SetBulkOptionalEnabled(bInEnabled);
+		Backend->SetBulkOptionalEnabled(bInEnabled);
 	}
 	else
 	{
@@ -2470,9 +2392,9 @@ void FIoStoreOnDemandModule::SetBulkOptionalEnabled(bool bInEnabled)
 
 void FIoStoreOnDemandModule::SetEnabled(bool bInEnabled)
 {
-	if (HttpIoDispatcherBackend.IsValid())
+	if (Backend.IsValid())
 	{
-		HttpIoDispatcherBackend->SetEnabled(bInEnabled);
+		Backend->SetEnabled(bInEnabled);
 	}
 	else
 	{
@@ -2483,9 +2405,9 @@ void FIoStoreOnDemandModule::SetEnabled(bool bInEnabled)
 
 void FIoStoreOnDemandModule::AbandonCache()
 {
-	if (HttpIoDispatcherBackend.IsValid())
+	if (Backend.IsValid())
 	{
-		HttpIoDispatcherBackend->AbandonCache();
+		Backend->AbandonCache();
 	}
 	else
 	{
@@ -2496,14 +2418,14 @@ void FIoStoreOnDemandModule::AbandonCache()
 
 bool FIoStoreOnDemandModule::IsEnabled() const
 {
-	return HttpIoDispatcherBackend.IsValid()? HttpIoDispatcherBackend->IsEnabled():DeferredAbandonCache.IsSet();
+	return Backend.IsValid()? Backend->IsEnabled():DeferredAbandonCache.IsSet();
 }
 
 void FIoStoreOnDemandModule::ReportAnalytics(TArray<FAnalyticsEventAttribute>& OutAnalyticsArray) const
 {
-	if (HttpIoDispatcherBackend.IsValid())
+	if (Backend.IsValid())
 	{
-		HttpIoDispatcherBackend->ReportAnalytics(OutAnalyticsArray);
+		Backend->ReportAnalytics(OutAnalyticsArray);
 	}
 }
 
@@ -2531,7 +2453,7 @@ void FIoStoreOnDemandModule::InitializeInternal()
 #endif
 
 	// Make sure we haven't called initialize before
-	check(!HttpIoDispatcherBackend.IsValid());
+	check(!Backend.IsValid());
 
 	FOnDemandEndpoint Endpoint;
 	
@@ -2590,8 +2512,8 @@ void FIoStoreOnDemandModule::InitializeInternal()
 			(CacheConfig.DiskQuota > 0) ? TEXT("init-fail") : TEXT("zero-quota"));
 	}
 
-	HttpIoDispatcherBackend = MakeOnDemandIoDispatcherBackend(MoveTemp(Cache));
-	HttpIoDispatcherBackend->Mount(Endpoint);
+	Backend = MakeOnDemandIoDispatcherBackend(MoveTemp(Cache));
+	Backend->Mount(Endpoint);
 	int32 BackendPriority = -10;
 #if !UE_BUILD_SHIPPING
 	if (FParse::Param(CommandLine, TEXT("Ias")))
@@ -2604,14 +2526,14 @@ void FIoStoreOnDemandModule::InitializeInternal()
 	// Setup any states changes issued before initialization
 	if (DeferredEnabled.IsSet())
 	{
-		HttpIoDispatcherBackend->SetEnabled(*DeferredEnabled);
+		Backend->SetEnabled(*DeferredEnabled);
 	}
 	if (DeferredBulkOptionalEnabled.IsSet())
 	{
-		HttpIoDispatcherBackend->SetBulkOptionalEnabled(*DeferredBulkOptionalEnabled);
+		Backend->SetBulkOptionalEnabled(*DeferredBulkOptionalEnabled);
 	}
 	
-	FIoDispatcher::Get().Mount(HttpIoDispatcherBackend.ToSharedRef(), BackendPriority);
+	FIoDispatcher::Get().Mount(Backend.ToSharedRef(), BackendPriority);
 }
 	
 void FIoStoreOnDemandModule::StartupModule()
@@ -2646,7 +2568,7 @@ EOnDemandInitResult FIoStoreOnDemandModule::Initialize()
 
 	InitializeInternal();
 
-	return HttpIoDispatcherBackend.IsValid() ? EOnDemandInitResult::Success : EOnDemandInitResult::Disabled;
+	return Backend.IsValid() ? EOnDemandInitResult::Success : EOnDemandInitResult::Disabled;
 };
 
 #endif // UE_IAS_CUSTOM_INITIALIZATION
diff --git a/Engine/Source/Runtime/Experimental/IoStoreOnDemand/Private/Tool/ChunkPlugin.cpp b/Engine/Source/Runtime/Experimental/IoStoreOnDemand/Private/Tool/ChunkPlugin.cpp
deleted file mode 100644
diff --git a/Engine/Source/Runtime/Experimental/IoStoreOnDemand/Public/IO/IoStoreOnDemand.h b/Engine/Source/Runtime/Experimental/IoStoreOnDemand/Public/IO/IoStoreOnDemand.h
--- a/Engine/Source/Runtime/Experimental/IoStoreOnDemand/Public/IO/IoStoreOnDemand.h
+++ b/Engine/Source/Runtime/Experimental/IoStoreOnDemand/Public/IO/IoStoreOnDemand.h
@@ -4,7 +4,6 @@
 
 #include "Containers/UnrealString.h"
 #include "IO/IoChunkId.h"
-#include "IO/IoContainerId.h"
 #include "IO/IoHash.h"
 #include "IO/IoStatus.h"
 #include "Misc/Guid.h"
@@ -54,8 +53,6 @@ enum class EOnDemandTocVersion : uint32
 	BlockHash32		= 3,
 	NoRawHash		= 4,
 	Meta			= 5,
-	ContainerId		= 6,
-	AdditionalFiles	= 7,
 
 	LatestPlusOne,
 	Latest			= (LatestPlusOne - 1)
@@ -116,7 +113,6 @@ UE_API bool LoadFromCompactBinary(FCbFieldView Field, FOnDemandTocEntry& OutTocE
 
 struct FOnDemandTocContainerEntry
 {
-	FIoContainerId ContainerId;
 	FString ContainerName;
 	FString EncryptionKeyGuid;
 	TArray<FOnDemandTocEntry> Entries;
@@ -146,18 +142,6 @@ private:
 	uint8 Data[SentinelSize] = { 0 };
 };
 
-struct FOnDemandTocAdditionalFile
-{
-	FIoHash Hash;
-	FString Filename;
-	uint64 FileSize = 0;
-
-	UE_API friend FArchive& operator<<(FArchive& Ar, FOnDemandTocAdditionalFile& AdditionalFile);
-	UE_API friend FCbWriter& operator<<(FCbWriter& Writer, const FOnDemandTocAdditionalFile& AdditionalFile);
-};
-
-UE_API bool LoadFromCompactBinary(FCbFieldView Field, FOnDemandTocAdditionalFile& AdditionalFile);
-
 struct FOnDemandToc
 {
 	FOnDemandToc() = default;
@@ -174,7 +158,6 @@ struct FOnDemandToc
 	FOnDemandTocHeader Header;
 	FTocMeta Meta;
 	TArray<FOnDemandTocContainerEntry> Containers;
-	TArray<FOnDemandTocAdditionalFile> AdditionalFiles;
 
 	UE_API friend FArchive& operator<<(FArchive& Ar, FOnDemandToc& Toc);
 	UE_API friend FCbWriter& operator<<(FCbWriter& Writer, const FOnDemandToc& Toc);
@@ -346,7 +329,7 @@ class FIoStoreOnDemandModule
 {
 private:
 	void InitializeInternal();
-	TSharedPtr<IOnDemandIoDispatcherBackend> HttpIoDispatcherBackend;
+	TSharedPtr<IOnDemandIoDispatcherBackend> Backend;
 	// Deferred state requests if called before backend
 	// is initialized
 	TOptional<bool> DeferredEnabled;
