Chaos - Convex-Mesh Motion-Aware Collision Detection (MACD) improvements

For a Convex moving from X to P, MACD now chooses a point along X->P to detect collisions between the convex and each triangle. This is a cheap form of CCD. The chosen point is the closest point on X->P to the line through the triangle centroid along its normal. We are still selecting the convex and triangle features to project based on the nearest features at X, so it can generate contacts even if P is on the insde of the triangle (which is where hte non-MACD path fails).

diff --git a/Engine/Source/Runtime/Experimental/Chaos/Private/Chaos/Collision/ConvexTriangleContactPoint.cpp b/Engine/Source/Runtime/Experimental/Chaos/Private/Chaos/Collision/ConvexTriangleContactPoint.cpp
--- a/Engine/Source/Runtime/Experimental/Chaos/Private/Chaos/Collision/ConvexTriangleContactPoint.cpp
+++ b/Engine/Source/Runtime/Experimental/Chaos/Private/Chaos/Collision/ConvexTriangleContactPoint.cpp
@@ -2,49 +2,22 @@
 #include "Chaos/Collision/ConvexTriangleContactPoint.h"
 #include "Chaos/Collision/ContactPoint.h"
 #include "Chaos/Collision/ContactTriangles.h"
+#include "Chaos/Collision/ConvexContactPoint.h"
 #include "Chaos/Collision/ConvexContactPointUtilities.h"
-#include "Chaos/Collision/ConvexFeature.h"
+#include "Chaos/Collision/SATConvexTriangle.h"
 #include "Chaos/CollisionOneShotManifolds.h"
 #include "Chaos/Convex.h"
 #include "Chaos/DebugDrawQueue.h"
 #include "Chaos/SAT.h"
 #include "Chaos/Triangle.h"
+#include "Chaos/Utilities.h"
 #include "Misc/MemStack.h"
 
-//UE_DISABLE_OPTIMIZATION
-
-// Method to use when checking edge pairs in convex-triangle SAT
-// 0: Iterate over convex edges, use IsMinkowskiSumConvexTriangle to skip invalid edge pairs
-// 1: Iterate over convex faces, then their edges. This has the advantage we don't need the edge list on FConvex
-#define CHAOS_CONVEX_TRIANGLE_EDGEEDGE_METHOD 0
-
 namespace Chaos
 {
 	extern FRealSingle Chaos_Collision_GJKEpsilon;
 	extern FRealSingle Chaos_Collision_EPAEpsilon;
-
-	// Check whether the convex-triangle edge pair form part of the Minkowski Sum. Only edge pairs
-	// that contribute to the Minkowki Sum surface need to be checked for separation. The inputs
-	// are the convex normals for the two faces that share the convex edge, and the normal and
-	// edge vector of the triangle.
-	// 
-	// This is a custom version of IsMinkowskiSum for triangles where the two normals are directly
-	// opposing and therefore the regular edge vector calculation returns zero.
-	// 
-	// @param A ConvexNormalA
-	// @param B ConvexNormalB
-	// @param BA ConvexEdge
-	// @param C TriNormal (negated)
-	// @param DC TriEdge
-	bool IsMinkowskiSumConvexTriangle(const FVec3& A, const FVec3& B, const FVec3& BA, const FVec3& C, const FVec3& DC)
-	{
-		const FReal CBA = FVec3::DotProduct(C, BA);		// TriNormal | ConvexEdge
-		const FReal ADC = FVec3::DotProduct(A, DC);		// ConvexNormalA | TriEdge
-		const FReal BDC = FVec3::DotProduct(B, DC);		// ConvexNormalB | TriEdge
-
-		const FReal Tolerance = 1.e-2f;
-		return ((ADC * BDC) < -Tolerance) && ((CBA * BDC) > Tolerance);
-	}
+	extern bool bChaos_Collision_UseConvexTriangleGJKSAT;
 
 	// Clip the vertices of a triangle to a face of a convex, using some arbitrary vector as the clipping axis
 	// (the axis is assumed to not be parallel to the convex face surface).
@@ -277,7 +250,7 @@ namespace Chaos
 				// NOTE: This relies on the ordering of the edge planes from above. 
 				// I.e., we require Sign(ConvexEdgePlaneNormalA x ConvexEdgePlaneNormalB) == Sign(ConvexEdgeV1 - ConvexEdgeV0)
 				// Also note that we must pass the negated triangle normal in
-				if (!IsMinkowskiSumConvexTriangle(ConvexEdgePlaneNormalA, ConvexEdgePlaneNormalB, ConvexEdgeV1 - ConvexEdgeV0, -TriN, TriEdgeV1 - TriEdgeV0))
+				if (!Private::IsOnMinkowskiSumConvexTriangle(ConvexEdgePlaneNormalA, ConvexEdgePlaneNormalB, ConvexEdgeV1 - ConvexEdgeV0, -TriN, TriEdgeV1 - TriEdgeV0))
 				{
 					continue;
 				}
@@ -579,7 +552,7 @@ namespace Chaos::Private
 {
 	// Generate a contact manifold between a convex and a triangle, given the closest feature (i.e., single Contact point)
 	template <typename ConvexType>
-	void ConvexTriangleManifoldFromContact(const ConvexType& Convex, const FTriangle& Triangle, const FVec3& TriangleNormal, const FConvexContactPoint& Contact, const FReal CullDistance, FContactPointLargeManifold& OutManifold)
+	void ConvexTriangleManifoldFromContact(const ConvexType& Convex, const FTriangle& Triangle, const FVec3& TriangleNormal, const FConvexContactPoint& Contact, const FReal CullDistance, FContactPointManifold& OutManifold)
 	{
 		// Convex plane
 		const int32 ConvexPlaneIndex = Contact.Features[0].PlaneIndex;
@@ -650,7 +623,7 @@ namespace Chaos::Private
 		}
 		check(ClippedVertices.Num() <= OutManifold.Max());
 
-		const EContactPointType ContactType = EContactPointType::Unknown;
+		const EContactPointType ContactType = Contact.GetContactPointType();
 
 		// Add the clipped points to the contact list
 		const auto& AddContact = [&CullDistance, &ContactType, &SeparatingAxis, &OutManifold](const FVec3& ConvexX, const FVec3 TriX, const FReal Distance)
@@ -719,7 +692,7 @@ namespace Chaos::Private
 		const FVec3& TriangleNormal, 
 		const FConvexContactPoint& Contact,
 		const FReal CullDistance, 
-		FContactPointLargeManifold& OutManifold);
+		FContactPointManifold& OutManifold);
 
 	template void ConvexTriangleManifoldFromContact(
 		const TImplicitObjectInstanced<FImplicitConvex3>& Convex,
@@ -727,7 +700,7 @@ namespace Chaos::Private
 		const FVec3& TriangleNormal,
 		const FConvexContactPoint& Contact,
 		const FReal CullDistance, 
-		FContactPointLargeManifold& OutManifold);
+		FContactPointManifold& OutManifold);
 
 	template void ConvexTriangleManifoldFromContact(
 		const TImplicitObjectScaled<FImplicitConvex3>& Convex,
@@ -735,7 +708,7 @@ namespace Chaos::Private
 		const FVec3& TriangleNormal,
 		const FConvexContactPoint& Contact,
 		const FReal CullDistance,
-		FContactPointLargeManifold& OutManifold);
+		FContactPointManifold& OutManifold);
 
 	template void ConvexTriangleManifoldFromContact(
 		const FImplicitBox3& Convex,
@@ -743,7 +716,7 @@ namespace Chaos::Private
 		const FVec3& TriangleNormal,
 		const FConvexContactPoint& Contact,
 		const FReal CullDistance,
-		FContactPointLargeManifold& OutManifold);
+		FContactPointManifold& OutManifold);
 
 	template void ConvexTriangleManifoldFromContact(
 		const TImplicitObjectScaled<FImplicitBox3>& Convex,
@@ -751,7 +724,7 @@ namespace Chaos::Private
 		const FVec3& TriangleNormal,
 		const FConvexContactPoint& Contact,
 		const FReal CullDistance,
-		FContactPointLargeManifold& OutManifold);
+		FContactPointManifold& OutManifold);
 
 	template void ConvexTriangleManifoldFromContact(
 		const TImplicitObjectInstanced<FImplicitBox3>& Convex,
@@ -759,272 +732,151 @@ namespace Chaos::Private
 		const FVec3& TriangleNormal,
 		const FConvexContactPoint& Contact,
 		const FReal CullDistance,
-		FContactPointLargeManifold& OutManifold);
+		FContactPointManifold& OutManifold);
 
 
 	// Generate a single contact point between a convex and a triangle
+	// NOTE: Does not fill in the Feature properties of OutContactPoint (see GetConvexFeature for that)
 	template <typename ConvexType>
-	bool ConvexTriangleContactPoint(const ConvexType& Convex, const FTriangle& Triangle, const FReal CullDistance, FContactPoint& OutContactPoint)
+	bool ConvexTriangleContactPoint(const ConvexType& Convex, const FTriangle& Triangle, const FVec3& TriangleNormal, const FReal CullDistanceSq, FConvexContactPoint& OutContactPoint)
 	{
-		// The GJK version has issues with exactly touching shapes. 
-		// E.g., see failing unit test EPARealFailures_TouchingBoxTriangle
-#if 1
+		const FReal GJKEpsilon = Chaos_Collision_GJKEpsilon;
+		const FReal EPAEpsilon = Chaos_Collision_EPAEpsilon;
 		const TGJKCoreShape<ConvexType> GJKConvex(Convex, Convex.GetMargin());
 		const TGJKShape<FTriangle> GJKTriangle(Triangle);
 
-		FReal UnusedMaxMarginDelta = FReal(0);
-		int32 ConvexVertexIndex = INDEX_NONE;
-		int32 TriangleVertexIndex = INDEX_NONE;
-		FReal Penetration;
-		FVec3 ConvexClosest, TriangleClosest, ConvexNormal;
-		const FReal GJKEpsilon = Chaos_Collision_GJKEpsilon;
-		const FReal EPAEpsilon = Chaos_Collision_EPAEpsilon;
+		// The GJKPenetrationSameSpace version uses EPA which has issues with almost exactly touching shapes. 
+		// E.g., see failing unit test EPARealFailures_TouchingBoxTriangle
+		if (bChaos_Collision_UseConvexTriangleGJKSAT)
+		{
+			// Use GJK to find closest features. This will abort if the margin-reduced shapes are closer than GJKEpsilon or overlap
+			constexpr FReal InvalidDistance = std::numeric_limits<FReal>::max();
+			FReal Distance = InvalidDistance;
+			FVec3 ConvexClosest, TriangleClosest, ConvexNormal;
+
+			const FVec3 InitialV = GJKDistanceInitialVFromDirection(GJKConvex, GJKTriangle, Triangle.GetCentroid());
+
+			const EGJKDistanceResult GJKResult = GJKDistance(
+				GJKConvex, GJKTriangle, InitialV, 
+				Distance, ConvexClosest, TriangleClosest, ConvexNormal, 
+				GJKEpsilon);
+
+			// GJK did not provide a result if the margin-reduced shapes overlap, so run SAT to calculate the minimum separating axis
+			if (GJKResult == EGJKDistanceResult::DeepContact)
+			{
+				return Private::SATConvexTriangle(
+					Convex, Triangle, TriangleNormal, 
+					CullDistanceSq, OutContactPoint);
+			}
 
-		FVec3 InitialGJKDir = FVec3(1, 0, 0);
-
-		const bool bHaveContact = GJKPenetrationSameSpace(
-			GJKConvex,
-			GJKTriangle,
-			Penetration,
-			ConvexClosest,
-			TriangleClosest,
-			ConvexNormal,
-			ConvexVertexIndex,
-			TriangleVertexIndex,
-			UnusedMaxMarginDelta,
-			InitialGJKDir,
-			GJKEpsilon, EPAEpsilon);
-
-		if (bHaveContact && (-Penetration < CullDistance))
+			// We are separated, or ovelap within marging so fill in the output
+			check(Distance != InvalidDistance);
+			if (Utilities::SignedSquare(Distance) < CullDistanceSq)
+			{
+				OutContactPoint.ShapeContactPoints[0] = ConvexClosest;
+				OutContactPoint.ShapeContactPoints[1] = TriangleClosest;
+				OutContactPoint.ShapeContactNormal = -ConvexNormal;
+				OutContactPoint.Phi = Distance;
+				return true;
+			}
+		}
+		else
 		{
-			OutContactPoint.ShapeContactPoints[0] = ConvexClosest;
-			OutContactPoint.ShapeContactPoints[1] = TriangleClosest;
-			OutContactPoint.ShapeContactNormal = -ConvexNormal;
-			OutContactPoint.ContactType = EContactPointType::Unknown;
-			OutContactPoint.Phi = -Penetration;
-			OutContactPoint.FaceIndex = INDEX_NONE;
-			return true;
+			FReal UnusedMaxMarginDelta = FReal(0);
+			int32 ConvexVertexIndex = INDEX_NONE;
+			int32 TriangleVertexIndex = INDEX_NONE;
+			FReal Penetration;
+			FVec3 ConvexClosest, TriangleClosest, ConvexNormal;
+			FVec3 InitialGJKDir = FVec3(1, 0, 0);
+
+			const bool bHaveContact = GJKPenetrationSameSpace(
+				GJKConvex, GJKTriangle, 
+				Penetration, ConvexClosest, TriangleClosest, ConvexNormal, ConvexVertexIndex, TriangleVertexIndex, UnusedMaxMarginDelta, 
+				InitialGJKDir, GJKEpsilon, EPAEpsilon);
+
+			if (bHaveContact && (FMath::Square(Penetration) < CullDistanceSq))
+			{
+				OutContactPoint.ShapeContactPoints[0] = ConvexClosest;
+				OutContactPoint.ShapeContactPoints[1] = TriangleClosest;
+				OutContactPoint.ShapeContactNormal = -ConvexNormal;
+				OutContactPoint.Phi = -Penetration;
+				return true;
+			}
 		}
 
 		return false;
-#else
-
-#endif
 	}
 
 	template bool ConvexTriangleContactPoint(
 		const FImplicitConvex3& Convex,
 		const FTriangle& Triangle,
+		const FVec3& TriangleNormal,
 		const FReal CullDistance,
-		FContactPoint& OutContactPoint);
+		FConvexContactPoint& OutContactPoint);
 
 	template bool ConvexTriangleContactPoint(
 		const TImplicitObjectInstanced<FImplicitConvex3>& Convex,
 		const FTriangle& Triangle,
+		const FVec3& TriangleNormal,
 		const FReal CullDistance,
-		FContactPoint& OutContactPoint);
+		FConvexContactPoint& OutContactPoint);
 
 	template bool ConvexTriangleContactPoint(
 		const TImplicitObjectScaled<FImplicitConvex3>& Convex,
 		const FTriangle& Triangle,
+		const FVec3& TriangleNormal,
 		const FReal CullDistance,
-		FContactPoint& OutContactPoint);
+		FConvexContactPoint& OutContactPoint);
 
 	template bool ConvexTriangleContactPoint(
 		const FImplicitBox3& Convex,
 		const FTriangle& Triangle,
+		const FVec3& TriangleNormal,
 		const FReal CullDistance,
-		FContactPoint& OutContactPoint);
+		FConvexContactPoint& OutContactPoint);
 
 	template bool ConvexTriangleContactPoint(
 		const TImplicitObjectScaled<FImplicitBox3>& Convex,
 		const FTriangle& Triangle,
+		const FVec3& TriangleNormal,
 		const FReal CullDistance,
-		FContactPoint& OutContactPoint);
+		FConvexContactPoint& OutContactPoint);
 
 	template bool ConvexTriangleContactPoint(
 		const TImplicitObjectInstanced<FImplicitBox3>& Convex,
 		const FTriangle& Triangle,
+		const FVec3& TriangleNormal,
 		const FReal CullDistance,
-		FContactPoint& OutContactPoint);
-
-
-	// Get the convex feature at the specific position and normal
-	template<typename ConvexType>
-	bool GetConvexFeature(const ConvexType& Convex, const FVec3& Position, const FVec3& Normal, Private::FConvexFeature& OutFeature)
-	{
-		const FReal NormalTolerance = FReal(1.e-6);
-		const FReal PositionTolerance = FReal(1.e-4);
-		const FReal ToleranceSizeMultiplier = Convex.BoundingBox().Extents().GetAbsMax();
-		const FReal EdgeNormalTolerance = ToleranceSizeMultiplier * FReal(1.e-3);
-
-		int32 BestPlaneIndex = INDEX_NONE;
-		FReal BestPlaneDotNormal = FReal(-1);
-
-		// Get the support vertex along the normal (which must point away from the convex)
-		int SupportVertexIndex = INDEX_NONE;
-		Convex.SupportCore(Normal, 0, nullptr, SupportVertexIndex);
-
-		if (SupportVertexIndex != INDEX_NONE)
-		{
-			// See if the normal matches a face normal for any face using the vertex
-			int32 VertexPlanes[16];
-			int32 NumVertexPlanes = Convex.FindVertexPlanes(SupportVertexIndex, VertexPlanes, UE_ARRAY_COUNT(VertexPlanes));
-			for (int32 VertexPlaneIndex = 0; VertexPlaneIndex < NumVertexPlanes; ++VertexPlaneIndex)
-			{
-				const int32 PlaneIndex = VertexPlanes[VertexPlaneIndex];
-				FVec3 PlaneN, PlaneX;
-				Convex.GetPlaneNX(PlaneIndex, PlaneN, PlaneX);
-				const FReal PlaneDotNormal = FVec3::DotProduct(PlaneN, Normal);
-				if (FMath::IsNearlyEqual(PlaneDotNormal, FReal(1), NormalTolerance))
-				{
-					OutFeature.FeatureType = Private::EConvexFeatureType::Plane;
-					OutFeature.PlaneIndex = PlaneIndex;
-					OutFeature.PlaneFeatureIndex = 0;
-					return true;
-				}
-
-				if (PlaneDotNormal > BestPlaneDotNormal)
-				{
-					BestPlaneIndex = PlaneIndex;
-					BestPlaneDotNormal = PlaneDotNormal;
-				}
-			}
-
-			// See if any of the edges using the vertex are perpendicular to the normal
-			// @todo(chaos): we could visit the vertex edges here rather than use the plane edges
-			if (BestPlaneIndex != INDEX_NONE)
-			{
-				int32 BestPlaneVertexIndex = INDEX_NONE;
+		FConvexContactPoint& OutContactPoint);
 
-				const int32 NumPlaneVertices = Convex.NumPlaneVertices(BestPlaneIndex);
-				for (int32 PlaneVertexIndex = 0; PlaneVertexIndex < NumPlaneVertices; ++PlaneVertexIndex)
-				{
-					const int32 VertexIndex0 = Convex.GetPlaneVertex(BestPlaneIndex, PlaneVertexIndex);
-					const int32 VertexIndex1 = (PlaneVertexIndex == NumPlaneVertices - 1) ? Convex.GetPlaneVertex(BestPlaneIndex, 0) : Convex.GetPlaneVertex(BestPlaneIndex, PlaneVertexIndex + 1);
-
-					if (VertexIndex0 == SupportVertexIndex)
-					{
-						BestPlaneVertexIndex = PlaneVertexIndex;
-					}
-
-					if ((VertexIndex0 == SupportVertexIndex) || (VertexIndex1 == SupportVertexIndex))
-					{
-						const FVec3 Vertex0 = Convex.GetVertex(VertexIndex0);
-						const FVec3 Vertex1 = Convex.GetVertex(VertexIndex1);
-						const FVec3 EdgeDelta = Vertex1 - Vertex0;
-						const FReal EdgeDotNormal = FVec3::DotProduct(EdgeDelta, Normal);
-						if (FMath::Abs(EdgeDotNormal) < EdgeNormalTolerance)
-						{
-							// @todo(chaos): we need to be able to get an EdgeIndex (probably half edge index)
-							// Also, we probably want both the plane index and the edge index
-							OutFeature.FeatureType = Private::EConvexFeatureType::Edge;
-							OutFeature.PlaneIndex = BestPlaneIndex;
-							OutFeature.PlaneFeatureIndex = PlaneVertexIndex;
-							return true;
-						}
-					}
-				}
-
-				// Not a face or edge, so it should be the SupportVertex, but we need to specify the 
-				// plane and plane-index rather than the convex vertex index (which we found just above)
-				if (BestPlaneVertexIndex != INDEX_NONE)
-				{
-					OutFeature.FeatureType = Private::EConvexFeatureType::Vertex;
-					OutFeature.PlaneIndex = BestPlaneIndex;
-					OutFeature.PlaneFeatureIndex = BestPlaneVertexIndex;
-				}
-				return true;
-			}
-		}
-
-		return false;
-	}
-
-	// Get the triangle feature at the specific position and normal
-	template<>
-	bool GetConvexFeature(const FTriangle& Triangle, const FVec3& Position, const FVec3& Normal, Private::FConvexFeature& OutFeature)
-	{
-		// @todo(chaos): pass in the triangle normal - we almost certainly calculated it elsewhere
-		// NOTE: The normal epsilon needs to be less than the maximu error that GJK/EPA produces when it hits a degenerate
-		// case, which can happen when we have almost exact face-to-face contact. The max error is hard to know, since it 
-		// depends on the state of GJK on the iteration before it hits its tolerance, but seems to be typically ~0.01
-		const FReal NormalEpsilon = FReal(0.02);
-		const FVec3 TriangleNormal = Triangle.GetNormal();
-		const FReal NormalDot = FVec3::DotProduct(Normal, TriangleNormal);
-		if (FMath::IsNearlyEqual(NormalDot, FReal(1), NormalEpsilon))
-		{
-			OutFeature.FeatureType = Private::EConvexFeatureType::Plane;
-			OutFeature.PlaneIndex = 0;
-			OutFeature.PlaneFeatureIndex = 0;
-			return true;
-		}
-
-		const FReal BarycentricTolerance = FReal(1.e-6);
-		int32 VertexIndex0, VertexIndex1;
-		if (GetTriangleEdgeVerticesAtPosition(Position, &Triangle.GetVertex(0), VertexIndex0, VertexIndex1, BarycentricTolerance))
-		{
-			if ((VertexIndex0 != INDEX_NONE) && (VertexIndex1 != INDEX_NONE))
-			{
-				OutFeature.FeatureType = Private::EConvexFeatureType::Edge;
-				OutFeature.PlaneIndex = 0;
-				OutFeature.PlaneFeatureIndex = VertexIndex0;
-				return true;
-			}
-			else if (VertexIndex0 != INDEX_NONE)
-			{
-				OutFeature.FeatureType = Private::EConvexFeatureType::Vertex;
-				OutFeature.PlaneIndex = 0;
-				OutFeature.PlaneFeatureIndex = VertexIndex0;
-				return true;
-			}
-			else if (VertexIndex1 != INDEX_NONE)
-			{
-				OutFeature.FeatureType = Private::EConvexFeatureType::Vertex;
-				OutFeature.PlaneIndex = 0;
-				OutFeature.PlaneFeatureIndex = VertexIndex1;
-				return true;
-			}
-		}
-
-		return false;
-	}
 
 	// Generate the contact point and closest feature types between a convex and a triangle
 	template<typename ConvexType>
-	bool FindClosestFeatures(const ConvexType& Convex, const FRigidTransform3& ConvexTransform, const FTriangle& Triangle, const FVec3& ConvexRelativeMovement, const FReal CullDistance, FConvexContactPoint& OutContact)
+	bool FindClosestFeatures(const ConvexType& Convex, const FRigidTransform3& ConvexTransform, const FTriangle& Triangle, const FVec3& TriangleNormal, const FVec3& ConvexRelativeMovement, const FReal CullDistance, FConvexContactPoint& OutContactPoint)
 	{
 		// Find the closest point on the convex and triangle that we will use to generate the manifold
-		// NOTE: use an upper limit on cull distance here since the real cull distance depends on the motion against the contact normal
-		const FReal EarlyCullDistance = TNumericLimits<FReal>::Max();
-		FContactPoint ContactPoint;
-		if (!Private::ConvexTriangleContactPoint(Convex, Triangle, EarlyCullDistance, ContactPoint))
+		// NOTE: use an upper limit on cull distance here since the real cull distance depends on the motion against the contact normal which we don't know yet
+		const FReal EarlyCullDistanceSq = TNumericLimits<FReal>::Max();	// @todo(chaos): put something useful here
+		if (!Private::ConvexTriangleContactPoint(Convex, Triangle, TriangleNormal, EarlyCullDistanceSq, OutContactPoint))
 		{
 			return false;
 		}
 
 		// Now check the cull distance, taking movement into account
-		const FReal SeparationWithMotion = ContactPoint.Phi + FVec3::DotProduct(ConvexRelativeMovement, ContactPoint.ShapeContactNormal);
-		if ((ContactPoint.Phi > CullDistance) && (SeparationWithMotion > CullDistance))
+		const FReal SeparationWithMotion = OutContactPoint.Phi + FVec3::DotProduct(ConvexRelativeMovement, OutContactPoint.ShapeContactNormal);
+		if ((OutContactPoint.Phi > CullDistance) && (SeparationWithMotion > CullDistance))
 		{
 			return false;
 		}
 
-		// Initialize outputs
-		OutContact.Init();
-		OutContact.ShapeContactPoints[0] = ContactPoint.ShapeContactPoints[0];
-		OutContact.ShapeContactPoints[1] = ContactPoint.ShapeContactPoints[1];
-		OutContact.ShapeContactNormal = ContactPoint.ShapeContactNormal;
-		OutContact.Phi = ContactPoint.Phi;
-
 		// Find the triangle feature at the contact point
-		if (!Private::GetConvexFeature(Triangle, OutContact.ShapeContactPoints[1], OutContact.ShapeContactNormal, OutContact.Features[1]))
+		if (!Private::GetTriangleFeature(Triangle, TriangleNormal, OutContactPoint.ShapeContactPoints[1], OutContactPoint.ShapeContactNormal, OutContactPoint.Features[1]))
 		{
 			return false;
 		}
 
 		// Find the convex feature at the contact point
-		if (!Private::GetConvexFeature(Convex, OutContact.ShapeContactPoints[0], -OutContact.ShapeContactNormal, OutContact.Features[0]))
+		if (!Private::GetConvexFeature(Convex, OutContactPoint.ShapeContactPoints[0], -OutContactPoint.ShapeContactNormal, OutContactPoint.Features[0]))
 		{
 			return false;
 		}
@@ -1036,7 +888,8 @@ namespace Chaos::Private
 		const FImplicitConvex3& Convex,
 		const FRigidTransform3& ConvexTransform, 
 		const FTriangle& Triangle, 
-		const FVec3& ConvexRelativeMovement, 
+		const FVec3& TriangleNormal,
+		const FVec3& ConvexRelativeMovement,
 		const FReal CullDistance, 
 		FConvexContactPoint& OutContact);
 
@@ -1044,6 +897,7 @@ namespace Chaos::Private
 		const TImplicitObjectInstanced<FImplicitConvex3>& Convex,
 		const FRigidTransform3& ConvexTransform,
 		const FTriangle& Triangle,
+		const FVec3& TriangleNormal,
 		const FVec3& ConvexRelativeMovement,
 		const FReal CullDistance,
 		FConvexContactPoint& OutContact);
@@ -1052,6 +906,7 @@ namespace Chaos::Private
 		const TImplicitObjectScaled<FImplicitConvex3>& Convex,
 		const FRigidTransform3& ConvexTransform,
 		const FTriangle& Triangle,
+		const FVec3& TriangleNormal,
 		const FVec3& ConvexRelativeMovement,
 		const FReal CullDistance,
 		FConvexContactPoint& OutContact);
@@ -1060,6 +915,7 @@ namespace Chaos::Private
 		const FImplicitBox3& Convex,
 		const FRigidTransform3& ConvexTransform,
 		const FTriangle& Triangle,
+		const FVec3& TriangleNormal,
 		const FVec3& ConvexRelativeMovement,
 		const FReal CullDistance,
 		FConvexContactPoint& OutContact);
@@ -1068,6 +924,7 @@ namespace Chaos::Private
 		const TImplicitObjectScaled<FImplicitBox3>& Convex,
 		const FRigidTransform3& ConvexTransform,
 		const FTriangle& Triangle,
+		const FVec3& TriangleNormal,
 		const FVec3& ConvexRelativeMovement,
 		const FReal CullDistance,
 		FConvexContactPoint& OutContact);
@@ -1076,6 +933,7 @@ namespace Chaos::Private
 		const TImplicitObjectInstanced<FImplicitBox3>& Convex,
 		const FRigidTransform3& ConvexTransform,
 		const FTriangle& Triangle,
+		const FVec3& TriangleNormal,
 		const FVec3& ConvexRelativeMovement,
 		const FReal CullDistance,
 		FConvexContactPoint& OutContact);
diff --git a/Engine/Source/Runtime/Experimental/Chaos/Private/Chaos/Collision/MeshContactGenerator.cpp b/Engine/Source/Runtime/Experimental/Chaos/Private/Chaos/Collision/MeshContactGenerator.cpp
--- a/Engine/Source/Runtime/Experimental/Chaos/Private/Chaos/Collision/MeshContactGenerator.cpp
+++ b/Engine/Source/Runtime/Experimental/Chaos/Private/Chaos/Collision/MeshContactGenerator.cpp
@@ -25,6 +25,7 @@ namespace Chaos::Private
 		HashSize = FMath::RoundUpToPowerOfTwo(Chaos_Collision_MeshManifoldHashSize);
 		FaceNormalDotThreshold = Chaos_Collision_MeshContactNormalThreshold;
 		EdgeNormalDotRejectTolerance = Chaos_Collision_MeshContactNormalRejectionThreshold;
+		BackFaceCullTolerance = FReal(1.e-6);
 		BarycentricTolerance = FReal(1.e-3);
 		MaxContactsBufferSize = 1000;
 		bCullBackFaces = true;
@@ -180,7 +181,7 @@ namespace Chaos::Private
 			FTriangleContactPointData& ContactPointData = ContactDatas[ContactIndex];
 
 			// Reject back-faces
-			if (!!Settings.bCullBackFaces && (ContactPointData.GetContactNormalDotTriangleNormal() < 0))
+			if (!!Settings.bCullBackFaces && (ContactPointData.GetContactNormalDotTriangleNormal() < -Settings.BackFaceCullTolerance))
 			{
 				ContactPointData.SetDisabled();
 				continue;
diff --git a/Engine/Source/Runtime/Experimental/Chaos/Private/Chaos/CollisionOneShotManifolds.cpp b/Engine/Source/Runtime/Experimental/Chaos/Private/Chaos/CollisionOneShotManifolds.cpp
--- a/Engine/Source/Runtime/Experimental/Chaos/Private/Chaos/CollisionOneShotManifolds.cpp
+++ b/Engine/Source/Runtime/Experimental/Chaos/Private/Chaos/CollisionOneShotManifolds.cpp
@@ -110,6 +110,15 @@ namespace Chaos
 	bool bChaos_Collision_EnableMACDFallback = false;
 	FAutoConsoleVariableRef CVarChaos_Collision_EnableMACDFallback(TEXT("p.Chaos.Collision.EnableMACDFallback"), bChaos_Collision_EnableMACDFallback, TEXT(""));
 
+	// We have 2 MACD algorithms to test, selected by bChaos_Collision_EnableMACDPreManifoldFix
+	// false: Use the original algorithm which generates a manifold for every triangle and then uses the mesh information to correct the normal of each manifold point.
+	// true: Use the new algorithm which generates the closest feature and fixes its normal based on the mesh information, and then generates a manifold based on that.
+	// The second method does a much better job at handling deep collisions where a single convex-triangle contact might generate a contact that pushes the convex out
+	// of the triangle edge, even though the triangle is part of a mesh and pushing out of the edge is not a valid option. However this option is also more likely to 
+	// generate ghost collisions with nearby faces when moving very fast.
+	bool bChaos_Collision_EnableMACDPreManifoldFix = true;
+	FAutoConsoleVariableRef CVarChaos_Collision_EnableMACDPreManifoldFix(TEXT("p.Chaos.Collision.EnableMACDPreManifoldFix"), bChaos_Collision_EnableMACDPreManifoldFix, TEXT(""));
+
 	// Whether to use the new index-less GJK. 
 	// @todo(chaos): This should be removed once soaked for a bit (enabled 7 June 2022)
 	bool bChaos_Collision_UseGJK2 = false;
@@ -132,6 +141,11 @@ namespace Chaos
 	bool bChaos_Collision_UseConvexTriMesh2 = true;
 	FAutoConsoleVariableRef CVarChaos_Collision_UseConvexTriMesh2(TEXT("p.Chaos.Collision.UseConvexTriMesh2"), bChaos_Collision_UseConvexTriMesh2, TEXT(""));
 
+	// true: use GJK and then SAT if the shapes overlap
+	// false: use GJK and then EPA if the shapes overlap
+	bool bChaos_Collision_UseConvexTriangleGJKSAT = true;
+	FAutoConsoleVariableRef CVarChaos_Collision_UseConvexTriangleSAT(TEXT("p.Chaos.Collision.UseConvexTriangleGJKSAT"), bChaos_Collision_UseConvexTriangleGJKSAT, TEXT(""));
+
 	namespace Collisions
 	{
 		// Forward delarations we need from CollisionRestitution.cpp
diff --git a/Engine/Source/Runtime/Experimental/Chaos/Private/Chaos/CollisionOneShotManifoldsMeshShapes.cpp b/Engine/Source/Runtime/Experimental/Chaos/Private/Chaos/CollisionOneShotManifoldsMeshShapes.cpp
--- a/Engine/Source/Runtime/Experimental/Chaos/Private/Chaos/CollisionOneShotManifoldsMeshShapes.cpp
+++ b/Engine/Source/Runtime/Experimental/Chaos/Private/Chaos/CollisionOneShotManifoldsMeshShapes.cpp
@@ -7,8 +7,8 @@
 #include "Chaos/Collision/CapsuleTriangleContactPoint.h"
 #include "Chaos/Collision/ContactPointsMiscShapes.h"
 #include "Chaos/Collision/ContactTriangles.h"
+#include "Chaos/Collision/ConvexContactPoint.h"
 #include "Chaos/Collision/ConvexContactPointUtilities.h"
-#include "Chaos/Collision/ConvexFeature.h"
 #include "Chaos/Collision/ConvexTriangleContactPoint.h"
 #include "Chaos/Collision/MeshContactGenerator.h"
 #include "Chaos/Collision/PBDCollisionConstraint.h"
@@ -37,6 +37,7 @@ namespace Chaos
 	extern bool bChaos_Collision_EnableMeshManifoldOptimizedLoop;
 	extern bool bChaos_Collision_EnableMeshManifoldOptimizedLoop_TriMesh;
 	extern bool bChaos_Collision_EnableMACDFallback;
+	extern bool bChaos_Collision_EnableMACDPreManifoldFix;
 
 	extern bool bChaos_Collision_UseCapsuleTriMesh2;
 	extern bool bChaos_Collision_UseConvexTriMesh2;
@@ -46,6 +47,7 @@ namespace Chaos
 #if CHAOS_DEBUG_DRAW
 		extern DebugDraw::FChaosDebugDrawSettings ChaosSolverDebugDebugDrawSettings;
 		extern int32 ChaosSolverDebugDrawMeshContacts;
+		extern int32 ChaosSolverDebugDrawMeshContactDetails;
 #endif
 	}
 
@@ -157,9 +159,10 @@ namespace Chaos
 			ContactGenerator.ProcessGeneratedContacts(ConvexTransform, MeshToConvexTransform);
 		}
 
-		// MACD: Motion-Aware Collision Detection
+		// Original MACD algorithm that uses mesh information to fix manifold point normals after the manifold is built.
+		// @todo(chaos): remove this when the new version is well tested.
 		template<typename ConvexType>
-		void GenerateConvexTriangleOneShotManifoldMACD(const ConvexType& Convex, const FRigidTransform3& ConvexTransform, const FVec3& ConvexRelativeMovement, Private::FMeshContactGenerator& ContactGenerator, const int32 TriangleIndex, const FReal CullDistance, FContactPointLargeManifold& OutContactPoints)
+		void GenerateConvexTriangleOneShotManifoldMACD_PostManifoldFix(const ConvexType& Convex, const FRigidTransform3& ConvexTransform, const FVec3& ConvexRelativeMovement, Private::FMeshContactGenerator& ContactGenerator, const int32 TriangleIndex, const FReal CullDistance, FContactPointManifold& OutContactPoints)
 		{
 			const FTriangle TriangleAtP = ContactGenerator.GetTriangle(TriangleIndex);
 			const FVec3 TriangleNormal = ContactGenerator.GetTriangleNormal(TriangleIndex);
@@ -175,7 +178,7 @@ namespace Chaos
 
 			// Find the closest feature pair on the triangle and convex
 			Private::FConvexContactPoint ClosestContact;
-			if (Private::FindClosestFeatures(Convex, ConvexTransform, Triangle, ConvexRelativeMovement, CullDistance, ClosestContact))
+			if (Private::FindClosestFeatures(Convex, ConvexTransform, Triangle, TriangleNormal, ConvexRelativeMovement, CullDistance, ClosestContact))
 			{
 				ClosestContact.Features[0].ObjectIndex = 0;
 				ClosestContact.Features[1].ObjectIndex = TriangleIndex;
@@ -255,7 +258,7 @@ namespace Chaos
 					for (int32 ContactIndex = 0; ContactIndex < OutContactPoints.Num(); ++ContactIndex)
 					{
 						FContactPoint& ContactPoint = OutContactPoints[ContactIndex];
-						
+
 						const FReal ShiftDotNormal = FVec3::DotProduct(ConvexRelativeMovement, ContactPoint.ShapeContactNormal);
 						ContactPoint.ShapeContactPoints[1] += -ShiftDotNormal * ContactPoint.ShapeContactNormal;
 						ContactPoint.Phi += ShiftDotNormal;
@@ -264,9 +267,225 @@ namespace Chaos
 			}
 		}
 
+		// MACD: Motion-Aware Collision Detection
+		//
+		// We have a convex moving from position X to P in this tick.
+		// 
+		// Detect collisions between the convex and a triangle taking that motion into account.
+		// 
+		// Find the closest features between the Convex at X and the Triangle. Use those features
+		// to select the Convex and Triangle Faces that will be projected onto each other to form
+		// the manifold. 
+		// 
+		// Use the Mesh information to correct the Triangle feature so that the normal is within the
+		// valid range. Edge and Vertex collisions with normals outside their valid ranges (determined
+		// by the other triangles that share the edge/vertex) are converted to face collisions.
+		// 
+		// As long as the Convex starts off outside the Triangle, we will generate useful contacts, 
+		// even if the Convex is fully inside the triangle at P (this is where the non-MACD path 
+		// would fail since it detects collisions only at P).
+		// 
+		// To generate the best manifold we select a point along the X-P trajectory that is
+		// closest to the Triangle's Axis (line though its centroid along its normal). See 
+		// amazing ascii art below. The box should collide with the triangle but clipping
+		// the triangle to the bottom box face would lead to no contacts at both X and P.
+		// 
+		//        +--------+
+		//  X:    |        |
+		//        +--------+
+		//
+		//  Tri:             ---------
+		//
+		//                                +--------+
+		//  P:                            |        |
+		//                                +--------+
+		template<typename ConvexType>
+		void GenerateConvexTriangleOneShotManifoldMACD_PreManifoldFix(
+			const ConvexType& Convex, 
+			const FRigidTransform3& ConvexTransform, 
+			const FVec3& ConvexRelativeMovement, 
+			Private::FMeshContactGenerator& ContactGenerator, 
+			const int32 TriangleIndex, 
+			const FReal CullDistance, 
+			FContactPointManifold& OutContactPoints)
+		{
+			// NOTE: The triangles were generated in Convex space with the convex at its predicted position P. I.e., we are in the space where P = 0.
+			// The convex moved from its initial position X to its predicted position P, and P = X + ConvexRelativeMovement.
+
+			// Triangle relative to the convex at its predicted position P
+			const FTriangle TriangleP = ContactGenerator.GetTriangle(TriangleIndex);
+			const FVec3 TriangleNormal = ContactGenerator.GetTriangleNormal(TriangleIndex);
+			const FVec3 TriangleCentroidP = TriangleP.GetCentroid();
+
+			// If we started inside the triangle we ignore this triangle
+			const FReal ConvexRelativeMovementTriNormal = FVec3::DotProduct(ConvexRelativeMovement, TriangleNormal);
+			const FReal ConvexTriangleDistanceX = FVec3::DotProduct(Convex.GetCenterOfMass() - TriangleP.GetVertex(0), TriangleNormal) - ConvexRelativeMovementTriNormal;
+			if (ConvexTriangleDistanceX < 0)
+			{
+				return;
+			}
+
+			// Triangle relative to the convex at its initial position X
+			// NOTE: we do not move the convex, we move the triangle to the relative position as if the convex were moved by -ConvexRelativeMovement.
+			const FVec3 TriangleXShift = ConvexRelativeMovement;
+			FTriangle TriangleX = FTriangle(TriangleP.GetVertex(0) + TriangleXShift, TriangleP.GetVertex(1) + TriangleXShift, TriangleP.GetVertex(2) + TriangleXShift);
+
+			// Find the closest feature between the Convex at its initial position X and the Triangle
+			Private::FConvexContactPoint ClosestContact;
+			const bool bFoundClosestContact = Private::FindClosestFeatures(Convex, ConvexTransform, TriangleX, TriangleNormal, ConvexRelativeMovement, CullDistance, ClosestContact);
+
+			if (bFoundClosestContact)
+			{
+				ClosestContact.Features[0].ObjectIndex = 0;
+				ClosestContact.Features[1].ObjectIndex = TriangleIndex;
+
+#if CHAOS_DEBUG_DRAW
+				if (CVars::ChaosSolverDebugDrawMeshContacts && FDebugDrawQueue::GetInstance().IsDebugDrawingEnabled())
+				{
+					const FVec3 P = ConvexTransform.TransformPositionNoScale(ClosestContact.ShapeContactPoints[1] - ConvexRelativeMovement);
+					const FVec3 N = ConvexTransform.TransformVectorNoScale(ClosestContact.ShapeContactNormal);
+					FDebugDrawQueue::GetInstance().DrawDebugLine(P, P + 10.0f * N, FColor::Black, false, CVars::ChaosSolverDebugDebugDrawSettings.DrawDuration, (uint8)CVars::ChaosSolverDebugDebugDrawSettings.DrawPriority, 1.5f * CVars::ChaosSolverDebugDebugDrawSettings.LineThickness);
+				}
+#endif
+
+				// Use the mesh info to correct the normal - this corrects edge and vertex normals if they are
+				// outside the range allowed by the set of triangles sharing the feature
+				if (ContactGenerator.FixFeature(TriangleIndex, ClosestContact.Features[1].FeatureType, ClosestContact.Features[1].PlaneFeatureIndex, ClosestContact.ShapeContactNormal))
+				{
+					// The normal was remapped to the triangle plane
+					ClosestContact.Features[0].FeatureType = Private::EConvexFeatureType::Vertex;
+					ClosestContact.Features[0].PlaneIndex = Convex.GetMostOpposingPlane(ClosestContact.ShapeContactNormal);
+					ClosestContact.Features[0].PlaneFeatureIndex = INDEX_NONE;	// Not needed by ConvexTriangleManifoldFromContact so not worth calculating
+				}
+
+				// Back face culling based on the corrected feature
+				const FReal TriangleDotNormal = FVec3::DotProduct(TriangleNormal, ClosestContact.ShapeContactNormal);
+				if (TriangleDotNormal < 0)
+				{
+					return;
+				}
+
+				// We will detect collisions at some point as the convex moves from X to P. 
+				// The point we choose is the closest approach to the axis along the triangle normal through the triangle centroid.
+				FReal ConvexT, TriangleT;
+				FVec3 ConvexNearPos, TriangleNearPos;
+				const FVec3 TriangleCentroidX = TriangleCentroidP + ConvexRelativeMovement;
+				Utilities::NearestPointsOnLineSegmentToLine(
+					FVec3(0), ConvexRelativeMovement,
+					TriangleCentroidX, TriangleNormal,
+					ConvexT, TriangleT, ConvexNearPos, TriangleNearPos);
+
+				// Triangle relative to the convex at ConvexNearPos
+				const FVec3 TriangleNearestShift = ConvexRelativeMovement - ConvexNearPos;
+				FTriangle TriangleNearest = FTriangle(TriangleP.GetVertex(0) + TriangleNearestShift, TriangleP.GetVertex(1) + TriangleNearestShift, TriangleP.GetVertex(2) + TriangleNearestShift);
+
+#if CHAOS_DEBUG_DRAW
+				if (CVars::ChaosSolverDebugDrawMeshContacts && FDebugDrawQueue::GetInstance().IsDebugDrawingEnabled())
+				{
+					const FVec3 P = ConvexTransform.TransformPositionNoScale(ClosestContact.ShapeContactPoints[1] - ConvexRelativeMovement);
+					const FVec3 N = ConvexTransform.TransformVectorNoScale(ClosestContact.ShapeContactNormal);
+					FDebugDrawQueue::GetInstance().DrawDebugLine(P, P + 10.0f * N, FColor::Orange, false, CVars::ChaosSolverDebugDebugDrawSettings.DrawDuration, (uint8)CVars::ChaosSolverDebugDebugDrawSettings.DrawPriority, 1.25f * CVars::ChaosSolverDebugDebugDrawSettings.LineThickness);
+
+					const FVec3 X0 = ConvexTransform.TransformPositionNoScale(Convex.GetCenterOfMass() - ConvexRelativeMovement);
+					const FVec3 X1 = ConvexTransform.TransformPositionNoScale(Convex.GetCenterOfMass());
+					FDebugDrawQueue::GetInstance().DrawDebugLine(X0, X1, FColor::White, false, CVars::ChaosSolverDebugDebugDrawSettings.DrawDuration, (uint8)CVars::ChaosSolverDebugDebugDrawSettings.DrawPriority, 0.5f * CVars::ChaosSolverDebugDebugDrawSettings.LineThickness);
+				}
+				if (CVars::ChaosSolverDebugDrawMeshContactDetails && FDebugDrawQueue::GetInstance().IsDebugDrawingEnabled())
+				{
+					const FVec3 TriangleAxis = FMath::Abs(ConvexRelativeMovementTriNormal) * TriangleNormal;
+					const FVec3 TA0 = ConvexTransform.TransformPositionNoScale(TriangleCentroidX + TriangleAxis - ConvexRelativeMovement);
+					const FVec3 TA1 = ConvexTransform.TransformPositionNoScale(TriangleCentroidX - TriangleAxis - ConvexRelativeMovement);
+					const FVec3 XConvex = ConvexTransform.TransformPositionNoScale(ConvexNearPos - ConvexRelativeMovement);
+					const FVec3 XTriangle = ConvexTransform.TransformPositionNoScale(TriangleNearPos - ConvexRelativeMovement);
+					FDebugDrawQueue::GetInstance().DrawDebugLine(TA0, TA1, FColor::Black, false, CVars::ChaosSolverDebugDebugDrawSettings.DrawDuration, (uint8)CVars::ChaosSolverDebugDebugDrawSettings.DrawPriority, 0.5f * CVars::ChaosSolverDebugDebugDrawSettings.LineThickness);
+					FDebugDrawQueue::GetInstance().DrawDebugLine(XConvex, XTriangle, FColor::Black, false, CVars::ChaosSolverDebugDebugDrawSettings.DrawDuration, (uint8)CVars::ChaosSolverDebugDebugDrawSettings.DrawPriority, 0.5f * CVars::ChaosSolverDebugDebugDrawSettings.LineThickness);
+					FDebugDrawQueue::GetInstance().DrawDebugPoint(XConvex, FColor::White, false, CVars::ChaosSolverDebugDebugDrawSettings.DrawDuration, (uint8)CVars::ChaosSolverDebugDebugDrawSettings.DrawPriority, 20.0f * CVars::ChaosSolverDebugDebugDrawSettings.LineThickness);
+
+					const FVec3 XCentroid = ConvexTransform.TransformPositionNoScale(TriangleCentroidX - ConvexRelativeMovement);
+					const FMatrix TriangleMat = FRotationMatrix::MakeFromZ(ConvexTransform.TransformVectorNoScale(TriangleNormal));
+					FDebugDrawQueue::GetInstance().DrawDebugCircle(XCentroid, 2.0f, 8, FColor::Black, false, CVars::ChaosSolverDebugDebugDrawSettings.DrawDuration, (uint8)CVars::ChaosSolverDebugDebugDrawSettings.DrawPriority, 0.5f * CVars::ChaosSolverDebugDebugDrawSettings.LineThickness, TriangleMat.GetUnitAxis(EAxis::X), TriangleMat.GetUnitAxis(EAxis::Y), false);
+				}
+#endif
+
+				// Generate a manifold based on the closest features
+				// NOTE: normal points from triangle to convex
+				const FReal ConvexMotionAlongNormal = FVec3::DotProduct(ConvexRelativeMovement, ClosestContact.ShapeContactNormal);
+				const FReal CullDistancePadding = FMath::Max(0, -ConvexMotionAlongNormal);
+				const FReal NetCullDistance = CullDistance + CullDistancePadding;
+				Private::ConvexTriangleManifoldFromContact(Convex, TriangleNearest, TriangleNormal, ClosestContact, NetCullDistance, OutContactPoints);
+
+				if (OutContactPoints.Num() > 0)
+				{
+#if CHAOS_DEBUG_DRAW
+					if (CVars::ChaosSolverDebugDrawMeshContacts && FDebugDrawQueue::GetInstance().IsDebugDrawingEnabled())
+					{
+						for (int32 ContactIndex = 0; ContactIndex < OutContactPoints.Num(); ++ContactIndex)
+						{
+							FContactPoint& ContactPoint = OutContactPoints[ContactIndex];
+							const FVec3 P = ConvexTransform.TransformPositionNoScale(ContactPoint.ShapeContactPoints[1] - TriangleNearestShift);
+							const FVec3 N = ConvexTransform.TransformVectorNoScale(ContactPoint.ShapeContactNormal);
+							FColor Color = FColor::Black;
+							switch (ClosestContact.Features[1].FeatureType)
+							{
+							case Private::EConvexFeatureType::Plane:
+								Color = FColor::White;
+								break;
+							case Private::EConvexFeatureType::Edge:
+								Color = FColor::Cyan;
+								break;
+							case Private::EConvexFeatureType::Vertex:
+								Color = FColor::Magenta;
+								break;
+							}
+							FDebugDrawQueue::GetInstance().DrawDebugLine(P, P + 10.0f * N, Color, false, CVars::ChaosSolverDebugDebugDrawSettings.DrawDuration, (uint8)CVars::ChaosSolverDebugDebugDrawSettings.DrawPriority, 1.25f * CVars::ChaosSolverDebugDebugDrawSettings.LineThickness);
+						}
+					}
+#endif
+
+					// Correct the contact points based on convex movement
+					for (int32 ContactIndex = 0; ContactIndex < OutContactPoints.Num(); ++ContactIndex)
+					{
+						FContactPoint& ContactPoint = OutContactPoints[ContactIndex];
+						
+						const FReal ShiftDotNormal = FVec3::DotProduct(TriangleNearestShift, ContactPoint.ShapeContactNormal);
+						ContactPoint.ShapeContactPoints[1] += -ShiftDotNormal * ContactPoint.ShapeContactNormal;
+						ContactPoint.Phi += ShiftDotNormal;
+					}
+				}
+			}
+		}
+
+		template<typename ConvexType>
+		void GenerateConvexTriangleOneShotManifoldMACD(
+			const ConvexType& Convex,
+			const FRigidTransform3& ConvexTransform,
+			const FVec3& ConvexRelativeMovement,
+			Private::FMeshContactGenerator& ContactGenerator,
+			const int32 TriangleIndex,
+			const FReal CullDistance,
+			FContactPointManifold& OutContactPoints)
+		{
+			if (bChaos_Collision_EnableMACDPreManifoldFix)
+			{
+				GenerateConvexTriangleOneShotManifoldMACD_PreManifoldFix(Convex, ConvexTransform, ConvexRelativeMovement, ContactGenerator, TriangleIndex, CullDistance, OutContactPoints);
+			}
+			else
+			{
+				GenerateConvexTriangleOneShotManifoldMACD_PostManifoldFix(Convex, ConvexTransform, ConvexRelativeMovement, ContactGenerator, TriangleIndex, CullDistance, OutContactPoints);
+			}
+		}
+
 		// MACD: Motion-Aware Collision Detection
 		template<typename ConvexType, typename MeshType>
-		void ConstructConvexMeshOneShotManifoldMACD(const ConvexType& Convex, const FRigidTransform3& ConvexTransform, const MeshType& Mesh, const FRigidTransform3& MeshTransform, const FVec3& MeshScale, const FVec3& RelativeMovement, const FReal InCullDistance, Private::FMeshContactGenerator& ContactGenerator)
+		void ConstructConvexMeshOneShotManifoldMACD(
+			const ConvexType& Convex, 
+			const FRigidTransform3& ConvexTransform, 
+			const MeshType& Mesh, 
+			const FRigidTransform3& MeshTransform, 
+			const FVec3& MeshScale, 
+			const FVec3& RelativeMovement, 
+			const FReal InCullDistance, 
+			Private::FMeshContactGenerator& ContactGenerator)
 		{
 			if (RelativeMovement.IsZero())
 			{
@@ -287,11 +506,9 @@ namespace Chaos
 				Mesh.CollectTriangles(MeshQueryBounds, MeshToConvexTransform, ConvexBounds, ContactGenerator);
 
 				FContactPointManifold Manifold;
-				FContactPointLargeManifold LargeManifold;
-				LargeManifold.Reserve(6);
 
 				const auto& GenerateConvexTriangleContacts =
-					[&Convex, &ConvexTransform, &ConvexRelativeMovement, CullDistance, &Manifold, &LargeManifold](Private::FMeshContactGenerator& ContactGenerator, const int32 TriangleIndex)
+					[&Convex, &ConvexTransform, &ConvexRelativeMovement, CullDistance, &Manifold](Private::FMeshContactGenerator& ContactGenerator, const int32 TriangleIndex)
 				{
 					// If we are outside the plane of the triangle, just use non-MACD collision
 					const FTriangle& Triangle = ContactGenerator.GetTriangle(TriangleIndex);
@@ -303,24 +520,23 @@ namespace Chaos
 						bUseNonMACDPath = Distance > 0;
 					}
 
+					Manifold.Reset();
 					if (bUseNonMACDPath)
 					{
-						Manifold.Reset();
 						GenerateConvexTriangleOneShotManifold(Convex, Triangle, CullDistance, Manifold);
-						ContactGenerator.AddTriangleContacts(TriangleIndex, MakeArrayView(Manifold));
-						return;
 					}
 					else
 					{
-						LargeManifold.Reset();
-						GenerateConvexTriangleOneShotManifoldMACD(Convex, ConvexTransform, ConvexRelativeMovement, ContactGenerator, TriangleIndex, CullDistance, LargeManifold);
-						ContactGenerator.AddTriangleContacts(TriangleIndex, MakeArrayView(LargeManifold));
+						GenerateConvexTriangleOneShotManifoldMACD(Convex, ConvexTransform, ConvexRelativeMovement, ContactGenerator, TriangleIndex, CullDistance, Manifold);
 					}
+					ContactGenerator.AddTriangleContacts(TriangleIndex, MakeArrayView(Manifold));
 				};
 
 				ContactGenerator.GenerateMeshContacts(GenerateConvexTriangleContacts);
 
 				// Process the contacts to minimize manifold etc
+				// MACD does not require further normal fixup
+				ContactGenerator.SetFixNormalsEnabled(false);
 				ContactGenerator.ProcessGeneratedContacts(ConvexTransform, MeshToConvexTransform);
 			}
 		}
@@ -490,11 +706,6 @@ namespace Chaos
 			{
 				const FVec3 RelativeMovement = FVec3(Constraint.GetRelativeMovement());
 				Private::FMeshContactGeneratorSettings ContactGeneratorSettings;
-
-				// @todo(chaos): we don't need to enable normal fixup for MACD mode, except when 
-				// fall back to non-MACD (which is a per-triangle decision) so for now it must be enabled.
-				ContactGeneratorSettings.bFixNormals = true;//RelativeMovement.IsNearlyZero();
-
 				Private::FMeshContactGenerator ContactGenerator(ContactGeneratorSettings);
 
 				if (const FImplicitBox3* RawBox = Convex.template GetObject<FImplicitBox3>())
diff --git a/Engine/Source/Runtime/Experimental/Chaos/Private/PBDRigidsSolver.cpp b/Engine/Source/Runtime/Experimental/Chaos/Private/PBDRigidsSolver.cpp
--- a/Engine/Source/Runtime/Experimental/Chaos/Private/PBDRigidsSolver.cpp
+++ b/Engine/Source/Runtime/Experimental/Chaos/Private/PBDRigidsSolver.cpp
@@ -140,6 +140,7 @@ namespace Chaos
 		int32 ChaosSolverDebugDrawSuspensionConstraints = 0;
 		int32 ChaosSolverDrawClusterConstraints = 0;
 		int32 ChaosSolverDebugDrawMeshContacts = 0;
+		int32 ChaosSolverDebugDrawMeshContactDetails = 0;
 		int32 ChaosSolverDebugDrawMeshBVHOverlaps = 0;
 		int32 ChaosSolverDebugDrawColorShapeByClientServer = 0;
 		int32 ChaosSolverDebugDrawShowServer = 1;
@@ -175,6 +176,7 @@ namespace Chaos
 		FAutoConsoleVariableRef CVarChaosSolverDrawSuspensionConstraints(TEXT("p.Chaos.Solver.DebugDrawSuspension"), ChaosSolverDebugDrawSuspensionConstraints, TEXT("Draw Suspension (0 = never; 1 = end of frame)."));
 		FAutoConsoleVariableRef CVarChaosSolverDrawClusterConstraints(TEXT("p.Chaos.Solver.DebugDraw.Cluster.Constraints"), ChaosSolverDrawClusterConstraints, TEXT("Draw Active Cluster Constraints (0 = never; 1 = end of frame)."));
 		FAutoConsoleVariableRef CVarChaosSolverDrawMeshContacts(TEXT("p.Chaos.Solver.DebugDrawMeshContacts"), ChaosSolverDebugDrawMeshContacts, TEXT("Draw Mesh contacts"));
+		FAutoConsoleVariableRef CVarChaosSolverDrawMeshContactDetails(TEXT("p.Chaos.Solver.DebugDrawMeshContactDetails"), ChaosSolverDebugDrawMeshContactDetails, TEXT("Draw Mesh contacts"));
 		FAutoConsoleVariableRef CVarChaosSolverDrawMeshBVHOverlaps(TEXT("p.Chaos.Solver.DebugDrawMeshBVHOverlaps"), ChaosSolverDebugDrawMeshBVHOverlaps, TEXT("Draw BVH of objects overlapping meshes"));
 		FAutoConsoleVariableRef CVarChaosSolverDebugDrawColorShapeByClientServer(TEXT("p.Chaos.Solver.DebugDraw.ColorShapeByClientServer"), ChaosSolverDebugDrawColorShapeByClientServer, TEXT("Color shape according to client and server: red = server / blue = client "));
 		FAutoConsoleVariableRef CVarChaosSolverDebugDrawShowServer(TEXT("p.Chaos.Solver.DebugDraw.ShowServer"), ChaosSolverDebugDrawShowServer, TEXT("Draw server related debug data"));
diff --git a/Engine/Source/Runtime/Experimental/Chaos/Public/Chaos/Collision/ConvexContactPoint.h b/Engine/Source/Runtime/Experimental/Chaos/Public/Chaos/Collision/ConvexContactPoint.h
new file mode 100644
--- /dev/null
+++ b/Engine/Source/Runtime/Experimental/Chaos/Public/Chaos/Collision/ConvexContactPoint.h
@@ -0,0 +1,60 @@
+// Copyright Epic Games, Inc.All Rights Reserved.
+#pragma once
+
+#include "Chaos/Core.h"
+#include "Chaos/Collision/ContactPoint.h"
+#include "Chaos/Collision/ConvexFeature.h"
+
+namespace Chaos::Private
+{
+	// A contact point between two convex features
+	template<typename T>
+	class TConvexContactPoint
+	{
+	public:
+		using FRealType = T;
+
+		TConvexContactPoint()
+		{
+			Reset();
+		}
+
+		void Reset()
+		{
+			Phi = InvalidPhi<FRealType>();
+		}
+
+		bool IsSet() const
+		{
+			return Phi != InvalidPhi<FRealType>();
+		}
+
+		EContactPointType GetContactPointType() const
+		{
+			if (IsSet())
+			{
+				if ((Features[0].FeatureType == EConvexFeatureType::Plane) && (Features[1].FeatureType == EConvexFeatureType::Vertex))
+				{
+					return EContactPointType::PlaneVertex;
+				}
+				if ((Features[0].FeatureType == EConvexFeatureType::Vertex) && (Features[1].FeatureType == EConvexFeatureType::Plane))
+				{
+					return EContactPointType::VertexPlane;
+				}
+				if ((Features[0].FeatureType == EConvexFeatureType::Edge) && (Features[1].FeatureType == EConvexFeatureType::Edge))
+				{
+					return EContactPointType::EdgeEdge;
+				}
+			}
+			return EContactPointType::Unknown;
+		}
+
+		FConvexFeature Features[2];
+		TVec3<FRealType> ShapeContactPoints[2];
+		TVec3<FRealType> ShapeContactNormal;
+		FRealType Phi;
+	};
+
+	using FConvexContactPoint = TConvexContactPoint<FReal>;
+	using FConvexContactPointf = TConvexContactPoint<FRealSingle>;
+}
\ No newline at end of file
diff --git a/Engine/Source/Runtime/Experimental/Chaos/Public/Chaos/Collision/ConvexContactPointUtilities.h b/Engine/Source/Runtime/Experimental/Chaos/Public/Chaos/Collision/ConvexContactPointUtilities.h
--- a/Engine/Source/Runtime/Experimental/Chaos/Public/Chaos/Collision/ConvexContactPointUtilities.h
+++ b/Engine/Source/Runtime/Experimental/Chaos/Public/Chaos/Collision/ConvexContactPointUtilities.h
@@ -4,9 +4,10 @@
 #include "Chaos/Core.h"
 #include "Chaos/Capsule.h"
 #include "Chaos/Collision/ContactPoint.h"
+#include "Chaos/Collision/ConvexFeature.h"
 #include "Chaos/Triangle.h"
 
-namespace Chaos:: Private
+namespace Chaos::Private
 {
 	// Project a convex onto an axis and return the projected range as well as the vertex indices that bound the range
 	template <typename ConvexType>
@@ -75,6 +76,185 @@ namespace Chaos:: Private
 			PMax = D0;
 		}
 	}
+
+	// Get the convex feature at the specific position and normal
+	template<typename ConvexType>
+	bool GetConvexFeature(const ConvexType& Convex, const FVec3& Position, const FVec3& Normal, Private::FConvexFeature& OutFeature)
+	{
+		const FReal NormalTolerance = FReal(1.e-6);
+		const FReal PositionTolerance = FReal(1.e-4);
+		const FReal ToleranceSizeMultiplier = Convex.BoundingBox().Extents().GetAbsMax();
+		const FReal EdgeNormalTolerance = ToleranceSizeMultiplier * FReal(1.e-3);
+
+		int32 BestPlaneIndex = INDEX_NONE;
+		FReal BestPlaneDotNormal = FReal(-1);
+
+		// Get the support vertex along the normal (which must point away from the convex)
+		int SupportVertexIndex = INDEX_NONE;
+		Convex.SupportCore(Normal, 0, nullptr, SupportVertexIndex);
+
+		if (SupportVertexIndex != INDEX_NONE)
+		{
+			// See if the normal matches a face normal for any face using the vertex
+			int32 VertexPlanes[16];
+			int32 NumVertexPlanes = Convex.FindVertexPlanes(SupportVertexIndex, VertexPlanes, UE_ARRAY_COUNT(VertexPlanes));
+			for (int32 VertexPlaneIndex = 0; VertexPlaneIndex < NumVertexPlanes; ++VertexPlaneIndex)
+			{
+				const int32 PlaneIndex = VertexPlanes[VertexPlaneIndex];
+				FVec3 PlaneN, PlaneX;
+				Convex.GetPlaneNX(PlaneIndex, PlaneN, PlaneX);
+				const FReal PlaneDotNormal = FVec3::DotProduct(PlaneN, Normal);
+				if (FMath::IsNearlyEqual(PlaneDotNormal, FReal(1), NormalTolerance))
+				{
+					OutFeature.FeatureType = Private::EConvexFeatureType::Plane;
+					OutFeature.PlaneIndex = PlaneIndex;
+					OutFeature.PlaneFeatureIndex = 0;
+					return true;
+				}
+
+				if (PlaneDotNormal > BestPlaneDotNormal)
+				{
+					BestPlaneIndex = PlaneIndex;
+					BestPlaneDotNormal = PlaneDotNormal;
+				}
+			}
+
+			// See if any of the edges using the vertex are perpendicular to the normal
+			// @todo(chaos): we could visit the vertex edges here rather than use the plane edges
+			if (BestPlaneIndex != INDEX_NONE)
+			{
+				int32 BestPlaneVertexIndex = INDEX_NONE;
+
+				const int32 NumPlaneVertices = Convex.NumPlaneVertices(BestPlaneIndex);
+				for (int32 PlaneVertexIndex = 0; PlaneVertexIndex < NumPlaneVertices; ++PlaneVertexIndex)
+				{
+					const int32 VertexIndex0 = Convex.GetPlaneVertex(BestPlaneIndex, PlaneVertexIndex);
+					const int32 VertexIndex1 = (PlaneVertexIndex == NumPlaneVertices - 1) ? Convex.GetPlaneVertex(BestPlaneIndex, 0) : Convex.GetPlaneVertex(BestPlaneIndex, PlaneVertexIndex + 1);
+
+					if (VertexIndex0 == SupportVertexIndex)
+					{
+						BestPlaneVertexIndex = PlaneVertexIndex;
+					}
+
+					if ((VertexIndex0 == SupportVertexIndex) || (VertexIndex1 == SupportVertexIndex))
+					{
+						const FVec3 Vertex0 = Convex.GetVertex(VertexIndex0);
+						const FVec3 Vertex1 = Convex.GetVertex(VertexIndex1);
+						const FVec3 EdgeDelta = Vertex1 - Vertex0;
+						const FReal EdgeDotNormal = FVec3::DotProduct(EdgeDelta, Normal);
+						if (FMath::Abs(EdgeDotNormal) < EdgeNormalTolerance)
+						{
+							// @todo(chaos): we need to be able to get an EdgeIndex (probably half edge index)
+							// Also, we probably want both the plane index and the edge index
+							OutFeature.FeatureType = Private::EConvexFeatureType::Edge;
+							OutFeature.PlaneIndex = BestPlaneIndex;
+							OutFeature.PlaneFeatureIndex = PlaneVertexIndex;
+							return true;
+						}
+					}
+				}
+
+				// Not a face or edge, so it should be the SupportVertex, but we need to specify the 
+				// plane and plane-index rather than the convex vertex index (which we found just above)
+				if (BestPlaneVertexIndex != INDEX_NONE)
+				{
+					OutFeature.FeatureType = Private::EConvexFeatureType::Vertex;
+					OutFeature.PlaneIndex = BestPlaneIndex;
+					OutFeature.PlaneFeatureIndex = BestPlaneVertexIndex;
+				}
+				return true;
+			}
+		}
+
+		return false;
+	}
+
+	// Get the triangle feature at the specific position and normal
+	inline bool GetTriangleFeature(const FTriangle& Triangle, const FVec3& TriangleNormal, const FVec3& Position, const FVec3& Normal, Private::FConvexFeature& OutFeature)
+	{
+		// @todo(chaos): pass in the triangle normal - we almost certainly calculated it elsewhere
+		// NOTE: The normal epsilon needs to be less than the maximu error that GJK/EPA produces when it hits a degenerate
+		// case, which can happen when we have almost exact face-to-face contact. The max error is hard to know, since it 
+		// depends on the state of GJK on the iteration before it hits its tolerance, but seems to be typically ~0.01
+		const FReal NormalEpsilon = FReal(0.02);
+		const FReal NormalDot = FVec3::DotProduct(Normal, TriangleNormal);
+		if (FMath::IsNearlyEqual(NormalDot, FReal(1), NormalEpsilon))
+		{
+			OutFeature.FeatureType = Private::EConvexFeatureType::Plane;
+			OutFeature.PlaneIndex = 0;
+			OutFeature.PlaneFeatureIndex = 0;
+			return true;
+		}
+
+		const FReal BarycentricTolerance = FReal(1.e-6);
+		int32 VertexIndex0, VertexIndex1;
+		if (GetTriangleEdgeVerticesAtPosition(Position, &Triangle.GetVertex(0), VertexIndex0, VertexIndex1, BarycentricTolerance))
+		{
+			if ((VertexIndex0 != INDEX_NONE) && (VertexIndex1 != INDEX_NONE))
+			{
+				OutFeature.FeatureType = Private::EConvexFeatureType::Edge;
+				OutFeature.PlaneIndex = 0;
+				OutFeature.PlaneFeatureIndex = VertexIndex0;
+				return true;
+			}
+			else if (VertexIndex0 != INDEX_NONE)
+			{
+				OutFeature.FeatureType = Private::EConvexFeatureType::Vertex;
+				OutFeature.PlaneIndex = 0;
+				OutFeature.PlaneFeatureIndex = VertexIndex0;
+				return true;
+			}
+			else if (VertexIndex1 != INDEX_NONE)
+			{
+				OutFeature.FeatureType = Private::EConvexFeatureType::Vertex;
+				OutFeature.PlaneIndex = 0;
+				OutFeature.PlaneFeatureIndex = VertexIndex1;
+				return true;
+			}
+		}
+
+		return false;
+	}
+
+	// Check whether the two edges of two convex shapes contribute to the Minkowski sum.
+	// A and B are the face normals for the faces of the edge convex 1
+	// C and D are the negated face normals for the faces of the edge convex 2
+	template<typename RealType>
+	inline bool IsOnMinkowskiSumConvexConvex(const TVec3<RealType>& A, const TVec3<RealType>& B, const TVec3<RealType>& C, const TVec3<RealType>& D, const RealType Tolerance = 1.e-2f)
+	{
+		const TVec3<RealType> BA = TVec3<RealType>::CrossProduct(B, A);
+		const TVec3<RealType> DC = TVec3<RealType>::CrossProduct(D, C);
+		const RealType CBA = TVec3<RealType>::DotProduct(C, BA);
+		const RealType DBA = TVec3<RealType>::DotProduct(D, BA);
+		const RealType ADC = TVec3<RealType>::DotProduct(A, DC);
+		const RealType BDC = TVec3<RealType>::DotProduct(B, DC);
+
+		return ((CBA * DBA) < -Tolerance) && ((ADC * BDC) < -Tolerance) && ((CBA * BDC) > Tolerance);
+	}
+
+	// Check whether the convex-triangle edge pair form part of the Minkowski Sum. Only edge pairs
+	// that contribute to the Minkowki Sum surface need to be checked for separation. The inputs
+	// are the convex normals for the two faces that share the convex edge, and the normal and
+	// edge vector of the triangle.
+	// 
+	// This is a custom version of IsOnMinkowskiSumConvexConvex for triangles where the two normals 
+	// are directly opposing and therefore the regular edge vector calculation returns zero.
+	// 
+	// @param A ConvexNormalA
+	// @param B ConvexNormalB
+	// @param BA ConvexEdge
+	// @param C TriNormal (negated)
+	// @param DC TriEdge
+	inline bool IsOnMinkowskiSumConvexTriangle(const FVec3& A, const FVec3& B, const FVec3& BA, const FVec3& C, const FVec3& DC)
+	{
+		const FReal CBA = FVec3::DotProduct(C, BA);		// TriNormal | ConvexEdge
+		const FReal ADC = FVec3::DotProduct(A, DC);		// ConvexNormalA | TriEdge
+		const FReal BDC = FVec3::DotProduct(B, DC);		// ConvexNormalB | TriEdge
+
+		const FReal Tolerance = 1.e-2f;
+		return ((ADC * BDC) < -Tolerance) && ((CBA * BDC) > Tolerance);
+	}
+
 }
 
 namespace Chaos
diff --git a/Engine/Source/Runtime/Experimental/Chaos/Public/Chaos/Collision/ConvexFeature.h b/Engine/Source/Runtime/Experimental/Chaos/Public/Chaos/Collision/ConvexFeature.h
--- a/Engine/Source/Runtime/Experimental/Chaos/Public/Chaos/Collision/ConvexFeature.h
+++ b/Engine/Source/Runtime/Experimental/Chaos/Public/Chaos/Collision/ConvexFeature.h
@@ -2,6 +2,7 @@
 #pragma once
 
 #include "Chaos/Core.h"
+#include "Chaos/Collision/ContactPoint.h"
 
 namespace Chaos::Private
 {
@@ -14,64 +15,22 @@ namespace Chaos::Private
 		Plane
 	};
 
-	// Information used to identify a feasture on a convex shape
+	// Information used to identify a feature on a convex shape
 	class FConvexFeature
 	{
 	public:
-		FConvexFeature& Init()
-		{
-			ObjectIndex = INDEX_NONE;
-			PlaneIndex = INDEX_NONE;
-			PlaneFeatureIndex = INDEX_NONE;
-			FeatureType = EConvexFeatureType::Unknown;
-			return *this;
-		}
-
-		void Set(const int32 InObjectIndex, const EConvexFeatureType InFeatureType, const int32 InPlaneIndex, const int32 InPlaneFeatureIndex)
-		{
-			ObjectIndex = InObjectIndex;
-			PlaneIndex = InPlaneIndex;
-			PlaneFeatureIndex = InPlaneFeatureIndex;
-			FeatureType = InFeatureType;
-		}
-
 		// An index used to index into some collection (e.g., triangle index in a mesh)
 		int32 ObjectIndex;
 
 		// The plane index on the convex. E.g., See FConvex::GetPlane()
 		int32 PlaneIndex;
 
-		// For Vertex or Edge features, the vertex the feature on the plane. E.g., See FConvex::GetPlaneVertex()
+		// For Vertex or Edge features this is either:
+		// - the vertex index on the plane if PlaneIndex is set (see FConvex::GetPlaneVertex()), or
+		// - the absolute vertex index if PlaneIndex is INDEX_NONE (see FConvex::GetVertex())
 		int32 PlaneFeatureIndex;
 
 		// The feature type
 		EConvexFeatureType FeatureType;
 	};
-
-	// A contact point between two convex features
-	template<typename T>
-	class TConvexContactPoint
-	{
-	public:
-		using FRealType = T;
-
-		TConvexContactPoint<FRealType>& Init()
-		{
-			Features[0].Init();
-			Features[1].Init();
-			ShapeContactPoints[0] = TVec3<FRealType>(0);
-			ShapeContactPoints[1] = TVec3<FRealType>(0);
-			ShapeContactNormal = TVec3<FRealType>(0);
-			Phi = TNumericLimits<FRealType>::Max();
-			return *this;
-		}
-
-		FConvexFeature Features[2];
-		TVec3<FRealType> ShapeContactPoints[2];
-		TVec3<FRealType> ShapeContactNormal;
-		FRealType Phi;
-	};
-
-	using FConvexContactPoint = TConvexContactPoint<FReal>;
-	using FConvexContactPointf = TConvexContactPoint<FRealSingle>;
 }
\ No newline at end of file
diff --git a/Engine/Source/Runtime/Experimental/Chaos/Public/Chaos/Collision/ConvexTriangleContactPoint.h b/Engine/Source/Runtime/Experimental/Chaos/Public/Chaos/Collision/ConvexTriangleContactPoint.h
--- a/Engine/Source/Runtime/Experimental/Chaos/Public/Chaos/Collision/ConvexTriangleContactPoint.h
+++ b/Engine/Source/Runtime/Experimental/Chaos/Public/Chaos/Collision/ConvexTriangleContactPoint.h
@@ -11,10 +11,10 @@ namespace Chaos
 	namespace Private
 	{
 		template<typename ConvexType>
-		extern CHAOS_API bool FindClosestFeatures(const ConvexType& Convex, const FRigidTransform3& ConvexTransform, const FTriangle& Triangle, const FVec3& ConvexRelativeMovement, const FReal CullDistance, FConvexContactPoint& OutContact);
+		extern CHAOS_API bool FindClosestFeatures(const ConvexType& Convex, const FRigidTransform3& ConvexTransform, const FTriangle& Triangle, const FVec3& TriangleNormal, const FVec3& ConvexRelativeMovement, const FReal CullDistance, FConvexContactPoint& OutContact);
 
 		template <typename ConvexType>
-		extern CHAOS_API void ConvexTriangleManifoldFromContact(const ConvexType& Convex, const FTriangle& Triangle, const FVec3& TriangleNormal, const FConvexContactPoint& Contact, const FReal CullDistance, FContactPointLargeManifold& OutManifold);
+		extern CHAOS_API void ConvexTriangleManifoldFromContact(const ConvexType& Convex, const FTriangle& Triangle, const FVec3& TriangleNormal, const FConvexContactPoint& Contact, const FReal CullDistance, FContactPointManifold& OutManifold);
 	}
 
 	template <typename ConvexType>
diff --git a/Engine/Source/Runtime/Experimental/Chaos/Public/Chaos/Collision/MeshContactGenerator.h b/Engine/Source/Runtime/Experimental/Chaos/Public/Chaos/Collision/MeshContactGenerator.h
--- a/Engine/Source/Runtime/Experimental/Chaos/Public/Chaos/Collision/MeshContactGenerator.h
+++ b/Engine/Source/Runtime/Experimental/Chaos/Public/Chaos/Collision/MeshContactGenerator.h
@@ -22,6 +22,9 @@ namespace Chaos::Private
 		// Triangle edge/vertex contacts that are more than this far from a valid normal (dot product) will be rejected rather than corrected
 		FReal EdgeNormalDotRejectTolerance;
 
+		// When backface culling is enabled, the tolerance for the dot product of the contact normal against the face normal
+		FReal BackFaceCullTolerance;
+
 		// Used to determine whether a contact is on an edge or vertex
 		FReal BarycentricTolerance;
 
@@ -55,6 +58,12 @@ namespace Chaos::Private
 	public:
 		FMeshContactGenerator(const FMeshContactGeneratorSettings& InSettings);
 
+		// Enable or disable the normal fixup
+		void SetFixNormalsEnabled(const bool bInFixNormals)
+		{
+			Settings.bFixNormals = bInFixNormals;
+		}
+
 		// Clear and initialize buffers
 		void BeginCollect(const int32 InNumTriangles)
 		{
diff --git a/Engine/Source/Runtime/Experimental/Chaos/Public/Chaos/Collision/SATConvexTriangle.h b/Engine/Source/Runtime/Experimental/Chaos/Public/Chaos/Collision/SATConvexTriangle.h
new file mode 100644
--- /dev/null
+++ b/Engine/Source/Runtime/Experimental/Chaos/Public/Chaos/Collision/SATConvexTriangle.h
@@ -0,0 +1,281 @@
+// Copyright Epic Games, Inc. All Rights Reserved.
+#pragma once
+
+#include "Chaos/Core.h"
+#include "Chaos/Collision/ConvexContactPoint.h"
+#include "Chaos/Collision/ConvexContactPointUtilities.h"
+#include "Chaos/Convex.h"
+#include "Chaos/Triangle.h"
+#include "Chaos/Utilities.h"
+
+namespace Chaos::Private
+{
+	// Find the contact point between a convex and a triangle using SAT.
+	// NOTE: Does not fill in the features of OutContactPoint (see GetConvexFeature())
+	template<typename ConvexType>
+	bool SATConvexTriangle(const ConvexType& Convex, const FTriangle& Triangle, const FVec3& TriangleNormal, const FReal CullDistanceSq, Private::FConvexContactPoint& OutContactPoint)
+	{
+		// Constants. NOTE: We square and multiple InvalidPhi so not using lowest()
+		const FReal NormalTolerance = FReal(1.e-8);
+		const FReal NormalToleranceSq = NormalTolerance * NormalTolerance;
+		const FReal InvalidPhi = FReal(-1.e10);	// 1000 km penetration
+
+		// Triangle (same space as convex)
+		const FVec3& TriN = TriangleNormal;
+		const FVec3 TriC = Triangle.GetCentroid();
+
+		//
+		// SAT: Triangle face vs Convex verts
+		//
+
+		// We store the convex vertex distances to the triangle face for use in the edge-edge culling
+		FMemMark Mark(FMemStack::Get());
+		TArray<FReal, TMemStackAllocator<alignof(FReal)>> ConvexVertexDs;
+		ConvexVertexDs.SetNum(Convex.NumVertices());
+		TArrayView<FReal> ConvexVertexDsView = MakeArrayView(ConvexVertexDs);
+
+		// Project the convex onto the triangle normal, with distances relative to the triangle plane
+		FReal TriPlaneDMin, TriPlaneDMax;
+		int32 ConvexVertexIndexMin, ConvexVertexIndexMax;
+		Private::ProjectOntoAxis(Convex, TriN, TriC, TriPlaneDMin, TriPlaneDMax, ConvexVertexIndexMin, ConvexVertexIndexMax, &ConvexVertexDsView);
+
+		// Distance culling
+		if (Utilities::SignedSquare(TriPlaneDMin) > CullDistanceSq)
+		{
+			// Outside triangle face and separated by more than CullDistance
+			return false;
+		}
+		if (TriPlaneDMax < FReal(0))
+		{
+			// Inside triangle face (single-sided collision)
+			return false;
+		}
+
+		//
+		// SAT: Convex faces vs triangle verts
+		//
+
+		// For each convex plane, project the Triangle onto the convex plane normal
+		// and reject if the separation is more than cull distance.
+		FVec3 ConvexPlaneN = FVec3(0);
+		FVec3 ConvexPlaneX = FVec3(0);
+		FReal ConvexPlaneDMin = InvalidPhi;
+		int32 ConvexPlaneIndexMin = INDEX_NONE;
+		int32 ConvexPlaneTriangleIndexMin = INDEX_NONE;
+		for (int32 PlaneIndex = 0; PlaneIndex < Convex.NumPlanes(); ++PlaneIndex)
+		{
+			FVec3 ConN, ConX;
+			Convex.GetPlaneNX(PlaneIndex, ConN, ConX);
+
+			FReal DMin, DMax;
+			int32 IndexMin, IndexMax;
+			Private::ProjectOntoAxis(Triangle, ConN, ConX, DMin, DMax, IndexMin, IndexMax);
+
+			// Backface cull - we can't select a convex face that would us a contact normal pointing below the triangle
+			//const FReal ConvexDotTriangleNormal = FVec3::DotProduct(ConN, TriN);
+			//if (ConvexDotTriangleNormal > 0)
+			//{
+			//	continue;
+			//}
+
+			// Distance culling
+			// @todo(chaos): Cull against the projected convex hull, not just the outward face 
+			// (we can store the most-distance vertex for each face with the convex to avoid actually having to project)
+			if (Utilities::SignedSquare(DMin) > CullDistanceSq)
+			{
+				// Separated by more than CullDistance
+				return false;
+			}
+
+			if (DMin > ConvexPlaneDMin)
+			{
+				ConvexPlaneN = ConN;
+				ConvexPlaneX = ConX;
+				ConvexPlaneDMin = DMin;
+				ConvexPlaneIndexMin = PlaneIndex;
+				ConvexPlaneTriangleIndexMin = IndexMin;
+			}
+		}
+
+		//
+		// SAT: Convex edges vs triangle edges
+		//
+
+		// Calculate the distance of triangle each edge to the convex separating plane
+		FReal TriVertexConvexDMin0 = FVec3::DotProduct(Triangle.GetVertex(0) - ConvexPlaneX, ConvexPlaneN);
+		FReal TriVertexConvexDMin1 = FVec3::DotProduct(Triangle.GetVertex(1) - ConvexPlaneX, ConvexPlaneN);
+		FReal TriVertexConvexDMin2 = FVec3::DotProduct(Triangle.GetVertex(2) - ConvexPlaneX, ConvexPlaneN);
+		FReal TriEdgeConvexDMin[3] =
+		{
+			FMath::Min(TriVertexConvexDMin2, TriVertexConvexDMin0),
+			FMath::Min(TriVertexConvexDMin0, TriVertexConvexDMin1),
+			FMath::Min(TriVertexConvexDMin1, TriVertexConvexDMin2),
+		};
+
+		FReal ConvexWinding = Convex.GetWindingOrder();
+		FVec3 EdgeEdgeN = FVec3(0);
+		FReal EdgeEdgeDMin = InvalidPhi;
+		int32 ConvexEdgeIndexMin = INDEX_NONE;
+		int32 TriEdgeIndexMin = INDEX_NONE;
+		for (int32 ConvexEdgeLoopIndex = 0; ConvexEdgeLoopIndex < Convex.NumEdges(); ++ConvexEdgeLoopIndex)
+		{
+			// Handle reverse winding for negative scaled convexes. Loop over edges in reverse order, and reverse edge vertex order
+			const int32 ConvexEdgeIndex = (ConvexWinding >= 0) ? ConvexEdgeLoopIndex : Convex.NumEdges() - ConvexEdgeLoopIndex - 1;
+			const int32 ConvexEdgeVIndex0 = (ConvexWinding >= 0) ? 0 : 1;
+			const int32 ConvexEdgeVIndex1 = (ConvexWinding >= 0) ? 1 : 0;
+
+			// Skip convex edges beyond CullDistance of the triangle face
+			const int32 ConvexEdgeVertexIndex0 = Convex.GetEdgeVertex(ConvexEdgeIndex, ConvexEdgeVIndex0);
+			const int32 ConvexEdgeVertexIndex1 = Convex.GetEdgeVertex(ConvexEdgeIndex, ConvexEdgeVIndex1);
+			const FReal FaceConvexD0 = ConvexVertexDs[ConvexEdgeVertexIndex0];
+			const FReal FaceConvexD1 = ConvexVertexDs[ConvexEdgeVertexIndex1];
+			if ((Utilities::SignedSquare(FaceConvexD0) > CullDistanceSq) && (Utilities::SignedSquare(FaceConvexD1) > CullDistanceSq))
+			{
+				continue;
+			}
+
+			// Convex edge vertices
+			const FVec3 ConvexEdgeV0 = Convex.GetVertex(ConvexEdgeVertexIndex0);
+			const FVec3 ConvexEdgeV1 = Convex.GetVertex(ConvexEdgeVertexIndex1);
+
+			// Convex planes that form the edge
+			const int32 ConvexEdgePlaneIndexA = Convex.GetEdgePlane(ConvexEdgeIndex, 0);
+			const int32 ConvexEdgePlaneIndexB = Convex.GetEdgePlane(ConvexEdgeIndex, 1);
+			const FVec3 ConvexEdgePlaneNormalA = Convex.GetPlane(ConvexEdgePlaneIndexA).Normal();
+			const FVec3 ConvexEdgePlaneNormalB = Convex.GetPlane(ConvexEdgePlaneIndexB).Normal();
+
+			for (int32 TriEdgeIndex = 0; TriEdgeIndex < 3; ++TriEdgeIndex)
+			{
+				// Skip triangle edges beyond cull distance of the convex separating face
+				if (Utilities::SignedSquare(TriEdgeConvexDMin[TriEdgeIndex]) > CullDistanceSq)
+				{
+					continue;
+				}
+
+				// Triangle edge vertices
+				const FVec3& TriEdgeV0 = Triangle.GetVertex(TriEdgeIndex);
+				const FVec3& TriEdgeV1 = (TriEdgeIndex == 2) ? Triangle.GetVertex(0) : Triangle.GetVertex(TriEdgeIndex + 1);
+
+				// Skip edge pairs that do not contribute to the Minkowski Sum surface
+				// NOTE: This relies on the ordering of the edge planes from above. 
+				// I.e., we require Sign(ConvexEdgePlaneNormalA x ConvexEdgePlaneNormalB) == Sign(ConvexEdgeV1 - ConvexEdgeV0)
+				// Also note that we must pass the negated triangle normal in
+				if (!IsOnMinkowskiSumConvexTriangle(ConvexEdgePlaneNormalA, ConvexEdgePlaneNormalB, ConvexEdgeV1 - ConvexEdgeV0, -TriN, TriEdgeV1 - TriEdgeV0))
+				{
+					continue;
+				}
+
+				// Separating axis
+				// NOTE: Not normalized at this stage. We perform the projection against the
+				// non-normalized axis and defer the square root until we know we need it
+				FVec3 Axis = FVec3::CrossProduct(ConvexEdgeV1 - ConvexEdgeV0, TriEdgeV1 - TriEdgeV0);
+				const FReal AxisLenSq = Axis.SizeSquared();
+
+				// Pick consistent axis direction: away from triangle (we want a signed distance)
+				const FReal Sign = FVec3::DotProduct(TriEdgeV0 - TriC, Axis);
+				if (Sign < FReal(0))
+				{
+					Axis = -Axis;
+				}
+
+				// Ignore backface edges
+				//const FReal AxisDotNormal = FVec3::DotProduct(TriN, Axis);
+				//if (AxisDotNormal < -NormalTolerance)
+				//{
+				//	continue;
+				//}
+
+				const FReal ScaledSeparation = FVec3::DotProduct(ConvexEdgeV0 - TriEdgeV0, Axis);
+
+				// Check cull distance on projected segments
+				// Comparing square distances scaled by axis length to defer square root (keep the sign)
+				const FReal ScaledSeparationSq = ScaledSeparation * FMath::Abs(ScaledSeparation);
+				const FReal ScaledCullDistanceSq = CullDistanceSq * AxisLenSq;
+				if (ScaledSeparationSq > ScaledCullDistanceSq)
+				{
+					return false;
+				}
+
+				// Comparing square distances scaled by axis length to defer square root (keep the sign)
+				const FReal ScaledEdgeEdgeDMinSq = EdgeEdgeDMin * FMath::Abs(EdgeEdgeDMin) * AxisLenSq;
+				if (ScaledSeparationSq > ScaledEdgeEdgeDMinSq)
+				{
+					// Now we need to know the actual separation and axis
+					const FReal AxisInvLen = FMath::InvSqrt(AxisLenSq);
+					EdgeEdgeDMin = ScaledSeparation * AxisInvLen;
+					EdgeEdgeN = Axis * AxisInvLen;
+					ConvexEdgeIndexMin = ConvexEdgeIndex;
+					TriEdgeIndexMin = TriEdgeIndex;
+				}
+			}
+		}
+
+		// Determine which of the features we want to use ad fill in the output
+		// NOTE: we rely on the fact that all valid Phi values are greater than InvalidPhi here
+
+		const FReal TriFaceBias = FReal(1.e-2);	// Prevent flip-flop on near-parallel cases
+		if ((TriPlaneDMin != InvalidPhi) && (TriPlaneDMin + TriFaceBias > ConvexPlaneDMin) && (TriPlaneDMin + TriFaceBias > EdgeEdgeDMin))
+		{
+			// Triangle face contact. The triangle normal is the separating axis
+			const FReal SeparatingDistance = TriPlaneDMin;
+			FVec3 SeparatingAxis = TriN;
+
+			const FVec3 ConvexContactPoint = Convex.GetVertex(ConvexVertexIndexMin);
+			const FVec3 TriangleContactPoint = ConvexContactPoint - SeparatingAxis * SeparatingDistance;
+			OutContactPoint.ShapeContactPoints[0] = ConvexContactPoint;
+			OutContactPoint.ShapeContactPoints[1] = TriangleContactPoint;
+			OutContactPoint.ShapeContactNormal = SeparatingAxis;
+			OutContactPoint.Phi = SeparatingDistance;
+			return true;
+		}
+
+		if ((ConvexPlaneDMin != InvalidPhi) && (ConvexPlaneDMin > EdgeEdgeDMin))
+		{
+			// Convex face contact. The convex face is the separating axis, but it must point from the triangle to the convex
+			const FReal SeparatingDistance = ConvexPlaneDMin;
+			FVec3 SeparatingAxis = -ConvexPlaneN;
+
+			const FVec3 TriangleContactPoint = Triangle.GetVertex(ConvexPlaneTriangleIndexMin);
+			const FVec3 ConvexContactPoint = TriangleContactPoint + SeparatingAxis * SeparatingDistance;
+			OutContactPoint.ShapeContactPoints[0] = ConvexContactPoint;
+			OutContactPoint.ShapeContactPoints[1] = TriangleContactPoint;
+			OutContactPoint.ShapeContactNormal = SeparatingAxis;
+			OutContactPoint.Phi = SeparatingDistance;
+			return true;
+		}
+
+		if (EdgeEdgeDMin != InvalidPhi)
+		{
+			// Edge-edge contact
+			// The separating axis must point from the triangle to the convex
+			const FReal SeparatingDistance = EdgeEdgeDMin;
+			FVec3 SeparatingAxis = EdgeEdgeN;
+
+			// Convex edge vertices
+			const int32 ConvexEdgeVIndex0 = (ConvexWinding >= 0) ? 0 : 1;
+			const int32 ConvexEdgeVIndex1 = (ConvexWinding >= 0) ? 1 : 0;
+			const int32 ConvexEdgeVertexIndex0 = Convex.GetEdgeVertex(ConvexEdgeIndexMin, ConvexEdgeVIndex0);
+			const int32 ConvexEdgeVertexIndex1 = Convex.GetEdgeVertex(ConvexEdgeIndexMin, ConvexEdgeVIndex1);
+			const FVec3 ConvexEdgeV0 = Convex.GetVertex(ConvexEdgeVertexIndex0);
+			const FVec3 ConvexEdgeV1 = Convex.GetVertex(ConvexEdgeVertexIndex1);
+
+			// Triangle edge vertices
+			const FVec3& TriEdgeV0 = Triangle.GetVertex(TriEdgeIndexMin);
+			const FVec3& TriEdgeV1 = (TriEdgeIndexMin == 2) ? Triangle.GetVertex(0) : Triangle.GetVertex(TriEdgeIndexMin + 1);
+
+			FReal ConvexEdgeT, TriEdgeT;
+			FVec3 ConvexContactPoint, TriangleContactPoint;
+			Utilities::NearestPointsOnLineSegments(ConvexEdgeV0, ConvexEdgeV1, TriEdgeV0, TriEdgeV1, ConvexEdgeT, TriEdgeT, ConvexContactPoint, TriangleContactPoint);
+
+			OutContactPoint.ShapeContactPoints[0] = ConvexContactPoint;
+			OutContactPoint.ShapeContactPoints[1] = TriangleContactPoint;
+			OutContactPoint.ShapeContactNormal = SeparatingAxis;
+			OutContactPoint.Phi = SeparatingDistance;
+			return true;
+		}
+
+		// No valid features (should not happen - at least TriPlaneDMin should always be valid)
+		ensureMsgf(false, TEXT("SATConvexTriangle failed to select a feature"));
+		return false;
+	}
+}
\ No newline at end of file
diff --git a/Engine/Source/Runtime/Experimental/Chaos/Public/Chaos/Utilities.h b/Engine/Source/Runtime/Experimental/Chaos/Public/Chaos/Utilities.h
--- a/Engine/Source/Runtime/Experimental/Chaos/Public/Chaos/Utilities.h
+++ b/Engine/Source/Runtime/Experimental/Chaos/Public/Chaos/Utilities.h
@@ -598,6 +598,8 @@ namespace Chaos
 			return V;
 		}
 
+		// Get the closest point on a line segment between P1 and Q1 and a 
+		// second line segment between P2 and Q2
 		// For implementation notes, see "Realtime Collision Detection", Christer Ericson, 2005
 		inline void NearestPointsOnLineSegments(
 			const FVec3& P1, const FVec3& Q1,
@@ -624,12 +626,12 @@ namespace Chaos
 			{
 				// Both segments are points
 			}
-			else if (A <= Epsilon)
+			else if (A <= EpsilonSq)
 			{
 				// First segment (only) is a point
 				T = FMath::Clamp<FReal>(F / E, Min, Max);
 			}
-			else if (E <= Epsilon)
+			else if (E <= EpsilonSq)
 			{
 				// Second segment (only) is a point
 				S = FMath::Clamp<FReal>(-C / A, Min, Max);
@@ -660,6 +662,62 @@ namespace Chaos
 			C2 = P2 + T * D2;
 		}
 
+		// Get the closest point on a line segment between SegmentBegin and SegmentEnd to an infinite
+		// line through LinePos along LineVector in both direction.
+		// For implementation notes, see "Realtime Collision Detection", Christer Ericson, 2005
+		inline void NearestPointsOnLineSegmentToLine(
+			const FVec3& SegmentBegin, const FVec3& SegmentEnd,
+			const FVec3& LinePos, const FVec3& LineVector,
+			FReal& S, FReal& T,
+			FVec3& C1, FVec3& C2,
+			const FReal Epsilon = 1.e-4f)
+		{
+			const FReal EpsilonSq = Epsilon * Epsilon;
+			const FVec3& P1 = SegmentBegin;
+			const FVec3& Q1 = SegmentEnd;
+			const FVec3& P2 = LinePos;
+			const FVec3 D1 = Q1 - P1;
+			const FVec3& D2 = LineVector;
+			const FVec3 R = P1 - P2;
+			const FReal A = FVec3::DotProduct(D1, D1);
+			const FReal B = FVec3::DotProduct(D1, D2);
+			const FReal C = FVec3::DotProduct(D1, R);
+			const FReal E = FVec3::DotProduct(D2, D2);
+			const FReal F = FVec3::DotProduct(D2, R);
+			constexpr FReal Min = 0, Max = 1;
+
+			S = 0.0f;
+			T = 0.0f;
+
+			if ((A <= EpsilonSq) && (B <= EpsilonSq))
+			{
+				// Both segments are points
+			}
+			else if (A <= EpsilonSq)
+			{
+				// First segment (only) is a point
+				T = F / E;
+			}
+			else if (E <= EpsilonSq)
+			{
+				// Second line (only) is a point (i.e., LineVector is zero)
+				S = FMath::Clamp<FReal>(-C / A, Min, Max);
+			}
+			else
+			{
+				// Non-degenrate case - we have two lines
+				const FReal Denom = A * E - B * B;
+				if (Denom != 0.0f)
+				{
+					S = FMath::Clamp<FReal>((B * F - C * E) / Denom, Min, Max);
+				}
+				T = (B * S + F) / E;
+			}
+
+			C1 = P1 + S * D1;
+			C2 = P2 + T * D2;
+		}
+
 		template<typename T>
 		inline T ClosestTimeOnLineSegment(const TVec3<T>& Point, const TVec3<T>& StartPoint, const TVec3<T>& EndPoint)
 		{
