GeometryCollection : SkeletalMesh conversion with geometry split by component.

diff --git a/Engine/Source/Runtime/Experimental/GeometryCollectionEngine/GeometryCollectionEngine.Build.cs b/Engine/Source/Runtime/Experimental/GeometryCollectionEngine/GeometryCollectionEngine.Build.cs
--- a/Engine/Source/Runtime/Experimental/GeometryCollectionEngine/GeometryCollectionEngine.Build.cs
+++ b/Engine/Source/Runtime/Experimental/GeometryCollectionEngine/GeometryCollectionEngine.Build.cs
@@ -25,6 +25,7 @@ namespace UnrealBuildTool.Rules
 					"DataflowEngine",
 					"MeshDescription",
 					"StaticMeshDescription",
+					"SkeletalMeshDescription",
 				}
 				);
 			PrivateDependencyModuleNames.AddRange(
diff --git a/Engine/Source/Runtime/Experimental/GeometryCollectionEngine/Private/GeometryCollection/GeometryCollectionEngineConversion.cpp b/Engine/Source/Runtime/Experimental/GeometryCollectionEngine/Private/GeometryCollection/GeometryCollectionEngineConversion.cpp
--- a/Engine/Source/Runtime/Experimental/GeometryCollectionEngine/Private/GeometryCollection/GeometryCollectionEngineConversion.cpp
+++ b/Engine/Source/Runtime/Experimental/GeometryCollectionEngine/Private/GeometryCollection/GeometryCollectionEngineConversion.cpp
@@ -7,32 +7,40 @@
 #include "Async/ParallelFor.h"
 #include "Components/SkeletalMeshComponent.h"
 #include "Components/StaticMeshComponent.h"
+#include "DynamicMesh/DynamicMesh3.h"
+#include "DynamicMesh/DynamicMeshAttributeSet.h"
 #include "Engine/Selection.h"
 #include "Engine/SkeletalMesh.h"
 #include "Engine/SkinnedAssetCommon.h"
 #include "Engine/StaticMesh.h"
 #include "GameFramework/Actor.h"
+#include "GeometryCollection/Facades/CollectionInstancedMeshFacade.h"
+#include "GeometryCollection/Facades/CollectionTransformFacade.h"
 #include "GeometryCollection/Facades/CollectionTransformSourceFacade.h"
 #include "GeometryCollection/GeometryCollection.h"
 #include "GeometryCollection/GeometryCollectionActor.h"
 #include "GeometryCollection/GeometryCollectionAlgo.h"
 #include "GeometryCollection/GeometryCollectionComponent.h"
 #include "GeometryCollection/GeometryCollectionClusteringUtility.h"
-#include "GeometryCollectionProxyData.h"
+#include "GeometryCollection/GeometryCollectionEngineUtility.h"
 #include "GeometryCollection/GeometryCollectionUtility.h"
+#include "GeometryCollectionProxyData.h"
+#include "IndexTypes.h"
 #include "Logging/LogMacros.h"
 #include "MaterialDomain.h"
 #include "Materials/Material.h"
-#include "Rendering/SkeletalMeshRenderData.h"
 #include "MeshDescription.h"
+#include "MeshDescriptionBuilder.h"
+#include "MeshDescriptionToDynamicMesh.h"
 #include "Misc/ScopedSlowTask.h"
-#include "StaticMeshAttributes.h"
-#include "StaticMeshOperations.h"
 #include "Physics/Experimental/ChaosInterfaceUtils.h"
 #include "PhysicsEngine/BodySetup.h"
-#include "MeshDescriptionBuilder.h"
+#include "Rendering/SkeletalMeshRenderData.h"
+#include "SkeletalMeshAttributes.h"
+#include "StaticMeshAttributes.h"
+#include "StaticMeshOperations.h"
 #include "VertexConnectedComponents.h"
-#include "GeometryCollection/Facades/CollectionInstancedMeshFacade.h"
+#include "Util/ColorConstants.h"
 
 
 DEFINE_LOG_CATEGORY_STATIC(UGeometryCollectionConversionLogging, Log, All);
@@ -1128,201 +1136,237 @@ void FGeometryCollectionEngineConversion::AppendGeometryCollection(const UGeomet
 }
 
 
-bool FGeometryCollectionEngineConversion::AppendSkeletalMesh(const USkeletalMesh* SkeletalMesh, int32 MaterialStartIndex, const FTransform& SkeletalMeshTransform, FManagedArrayCollection* InCollection, bool bReindexMaterials)
+
+
+bool FGeometryCollectionEngineConversion::AppendSkeletalMesh(const USkeletalMesh* InSkeletalMesh, int32 MaterialStartIndex, const FTransform& SkeletalMeshTransform, 
+	FManagedArrayCollection* InManagedArrayCollection, bool bReindexMaterials)
 {
 	//UE_LOG(UGeometryCollectionConversionLogging, Log, TEXT("FGeometryCollectionEngineConversion::AppendSkeletalMesh()"));
-	if (!InCollection)
+#if WITH_EDITOR
+	int LODIndex = 0;
+	if (!InManagedArrayCollection)
 	{
 		return false;
 	}
-	FGeometryCollection::DefineGeometrySchema(*InCollection);
 
-	const FSkeletalMeshLODRenderData* MeshLODData = GetSkeletalMeshLOD(SkeletalMesh, 0);
-	if (!MeshLODData)
+
+	FMeshDescription MeshDescription;
+	if (!InSkeletalMesh->CloneMeshDescription(LODIndex, MeshDescription))
 	{
 		return false;
 	}
-		
-	const FSkeletalMeshLODRenderData& SkeletalMeshLODRenderData = *MeshLODData;
-	const FSkinWeightVertexBuffer& SkinWeightVertexBuffer = *SkeletalMeshLODRenderData.GetSkinWeightVertexBuffer();
-	const FStaticMeshVertexBuffers& VertexBuffers = SkeletalMeshLODRenderData.StaticVertexBuffers;
-	const FPositionVertexBuffer& PositionVertexBuffer = VertexBuffers.PositionVertexBuffer;
-	const int32 VertexCount = PositionVertexBuffer.GetNumVertices();
-	// Check that all vertex weightings are rigid. 
-	for (int32 VertexIndex = 0; VertexIndex < VertexCount; VertexIndex++)
-	{
-		int32 SkeletalBoneIndex = -1;
-		if (!SkinWeightVertexBuffer.GetRigidWeightBone(VertexIndex, SkeletalBoneIndex))
-		{
-			UE_LOG(UGeometryCollectionConversionLogging, Error, TEXT("Non-rigid weighting found on vertex %d: Cannot convert to GeometryCollection."), VertexIndex);
-			return false;
-		}
-	}
-	const FSkelMeshRenderSection& RenderSection = SkeletalMeshLODRenderData.RenderSections[0];
-	//const TArray<FBoneIndexType>& SkeletalBoneMap = RenderSection.BoneMap;
+
+	FGeometryCollection GeometryCollection;
+
+	// Transform Attributes
+	TManagedArray<FTransform3f>& LocalSpaceTransform = GeometryCollection.ModifyAttribute<FTransform3f>(FTransformCollection::TransformAttribute, FTransformCollection::TransformGroup);
+	TManagedArray<int32>& Parent = GeometryCollection.ModifyAttribute<int32>(FTransformCollection::ParentAttribute, FTransformCollection::TransformGroup);
+	TManagedArray<FLinearColor>& BoneColor = GeometryCollection.ModifyAttribute<FLinearColor>("BoneColor", FTransformCollection::TransformGroup);
+	TManagedArray<FString>& BoneName = GeometryCollection.ModifyAttribute<FString>("BoneName", FTransformCollection::TransformGroup);
+	TManagedArray<int32>& SimulationType = GeometryCollection.ModifyAttribute<int32>("SimulationType", FTransformCollection::TransformGroup);
+	// Vertices Attributes
+	TManagedArray<FVector3f>& Vertex = GeometryCollection.ModifyAttribute<FVector3f>("Vertex", FGeometryCollection::VerticesGroup);
+	TManagedArray<FVector3f>& Normal = GeometryCollection.ModifyAttribute<FVector3f>("Normal", FGeometryCollection::VerticesGroup);
+	TManagedArray<FLinearColor>& Color = GeometryCollection.ModifyAttribute<FLinearColor>("Color", FGeometryCollection::VerticesGroup);
+	TManagedArray<FVector3f>& TangentU = GeometryCollection.ModifyAttribute<FVector3f>("TangentU", FGeometryCollection::VerticesGroup);
+	TManagedArray<FVector3f>& TangentV = GeometryCollection.ModifyAttribute<FVector3f>("TangentV", FGeometryCollection::VerticesGroup);
+	TManagedArray<int32>& BoneMap = GeometryCollection.ModifyAttribute<int32>("BoneMap", FGeometryCollection::VerticesGroup);
+	// Index Attributes
+	TManagedArray<FIntVector>& Indices = GeometryCollection.ModifyAttribute<FIntVector>("Indices", FGeometryCollection::FacesGroup);
+	TManagedArray<bool>& Visible = GeometryCollection.ModifyAttribute<bool>("Visible", FGeometryCollection::FacesGroup);
+	TManagedArray<int32>& MaterialIndex = GeometryCollection.ModifyAttribute<int32>("MaterialIndex", FGeometryCollection::FacesGroup);
+	TManagedArray<int32>& MaterialID = GeometryCollection.ModifyAttribute<int32>("MaterialID", FGeometryCollection::FacesGroup);
+
+	FMeshDescriptionToDynamicMesh Converter;
+	UE::Geometry::FDynamicMesh3 DynamicMesh;
+	Converter.Convert(&MeshDescription, DynamicMesh);
 
 	//
-	// The Component transform for each Mesh will become the FTransform that drives
-	// its associated VerticesGroup. The Skeleton will contain a nested transform hierarchy
-	// that is evaluated using the GetComponentSpaceTransformRefPose. The resulting
-	// Transforms array stored in the GeometryCollection will be the same size as
-	// the SkeletalBoneMap. Note the @todo: the SkeletalBoneMap is pulled from only
-	// the first render section, this will need to be expanded to include all render
-	// sections.
+	// Convert the transform hierarchy
 	//
 	int32 RootIndex = INDEX_NONE;
-	const USkeleton* Skeleton = SkeletalMesh->GetSkeleton();
+	int32 TransformBaseIndex = INDEX_NONE;
+	const USkeleton* Skeleton = InSkeletalMesh->GetSkeleton();
 	const TArray<FTransform>& RestArray = Skeleton->GetRefLocalPoses();
 	const FReferenceSkeleton& ReferenceSkeletion = Skeleton->GetReferenceSkeleton();
 
-	TManagedArray<FTransform3f>& LocalSpaceTransform = InCollection->ModifyAttribute<FTransform3f>(FTransformCollection::TransformAttribute, FTransformCollection::TransformGroup);
-	TManagedArray<int32>& Parent = InCollection->ModifyAttribute<int32>(FTransformCollection::ParentAttribute, FTransformCollection::TransformGroup);
-	TManagedArray<FLinearColor>& BoneColor = InCollection->ModifyAttribute<FLinearColor>("BoneColor", FTransformCollection::TransformGroup);
-	TManagedArray<FString>& BoneName = InCollection->ModifyAttribute<FString>("BoneName", FTransformCollection::TransformGroup);	
-	int32 TransformBaseIndex = InCollection->AddElements(ReferenceSkeletion.GetNum(), FGeometryCollection::TransformGroup);
-	TManagedArray<int32>& SimulationType = InCollection->ModifyAttribute<int32>("SimulationType", FTransformCollection::TransformGroup);
-
-	for (int32 BoneIndex = 0; BoneIndex < ReferenceSkeletion.GetNum(); BoneIndex++)
+	if (ReferenceSkeletion.GetNum())
 	{
-		LocalSpaceTransform[TransformBaseIndex + BoneIndex] = FTransform3f(RestArray[BoneIndex]);
-		// For validation against the component space position use
-		// FTransform ComponentSpaceTransform =FAnimationRuntime::GetComponentSpaceTransformRefPose(ReferenceSkeletion, SkeletalBoneMap[BoneIndex]);
+		TransformBaseIndex = GeometryCollection.AddElements(ReferenceSkeletion.GetNum(), FGeometryCollection::TransformGroup);
+		RootIndex = TransformBaseIndex;
 
-		BoneName[TransformBaseIndex + BoneIndex] = ReferenceSkeletion.GetRefBoneInfo()[BoneIndex].Name.ToString();
-		const FColor RandBoneColor(FMath::Rand() % 100 + 5, FMath::Rand() % 100 + 5, FMath::Rand() % 100 + 5, 255);
-		BoneColor[TransformBaseIndex + BoneIndex] = FLinearColor(RandBoneColor);
+		for (int32 BoneIndex = 0; BoneIndex < ReferenceSkeletion.GetNum(); BoneIndex++)
+		{
+			// For validation against the component space position use
+			// FTransform ComponentSpaceTransform =FAnimationRuntime::GetComponentSpaceTransformRefPose(ReferenceSkeletion, SkeletalBoneMap[BoneIndex]);
 
-		Parent[TransformBaseIndex + BoneIndex] = ReferenceSkeletion.GetRefBoneInfo()[BoneIndex].ParentIndex;
-		SimulationType[TransformBaseIndex + BoneIndex] = FGeometryCollection::ESimulationTypes::FST_None;
+			LocalSpaceTransform[TransformBaseIndex + BoneIndex] = FTransform3f(RestArray[BoneIndex]);
+			BoneName[TransformBaseIndex + BoneIndex] = ReferenceSkeletion.GetRefBoneInfo()[BoneIndex].Name.ToString();
+			Parent[TransformBaseIndex + BoneIndex] = ReferenceSkeletion.GetRefBoneInfo()[BoneIndex].ParentIndex;
+			SimulationType[TransformBaseIndex + BoneIndex] = FGeometryCollection::ESimulationTypes::FST_None;
+			BoneColor[TransformBaseIndex + BoneIndex] = FLinearColor::MakeRandomColor();
 
-		if (Parent[TransformBaseIndex + BoneIndex] == INDEX_NONE)
-		{
-			RootIndex = TransformBaseIndex + BoneIndex;
+			if (Parent[TransformBaseIndex + BoneIndex] == INDEX_NONE)
+			{
+				RootIndex = TransformBaseIndex + BoneIndex;
+			}
 		}
 	}
 
 
 	//
-	// The Triangle Indices
+	// Identify disconnected geoemtry
+	//
+	int32 VertexCount = 0, TriangleCount = 0;
+	TArray<int32> SourceVertexToComponentMap;				// Map from mesh vertex index to target vertex index
+	TArray<TArray<FIntVector>> ComponentsSourceIndices;		// Mesh triangle indices of each component.
+	TArray<TArray<FIntVector2>> SourceToTargetTriangleMap;	// Mesh triangle index of each triangle in the component. 
+	GeometryCollectionEngineUtility::GenerateConnectedComponents(InSkeletalMesh, ComponentsSourceIndices, 
+		SourceToTargetTriangleMap, SourceVertexToComponentMap,TriangleCount, VertexCount);
+
+
+	//
+	// Add the Triangles to the Geometry Collection
+	//.. ensure all component vertices are contigious in the array 
 	//
-	TManagedArray<FIntVector>& Indices = InCollection->ModifyAttribute<FIntVector>("Indices", FGeometryCollection::FacesGroup);
-	TManagedArray<bool>& Visible = InCollection->ModifyAttribute<bool>("Visible", FGeometryCollection::FacesGroup);
-	TManagedArray<int32>& MaterialIndex = InCollection->ModifyAttribute<int32>("MaterialIndex", FGeometryCollection::FacesGroup); 
-	TManagedArray<int32>& MaterialID = InCollection->ModifyAttribute<int32>("MaterialID", FGeometryCollection::FacesGroup);
-
-	TArray<uint32> IndexBuffer;
-	SkeletalMeshLODRenderData.MultiSizeIndexContainer.GetIndexBuffer(IndexBuffer);
-
-	const int32 IndicesCount = IndexBuffer.Num() / 3;
-	int NumVertices = InCollection->NumElements(FGeometryCollection::VerticesGroup);
-	int InitialNumIndices = InCollection->NumElements(FGeometryCollection::FacesGroup);
-	int IndicesBaseIndex = InCollection->AddElements(IndicesCount, FGeometryCollection::FacesGroup);
-	for (int32 IndicesIndex = 0, StaticIndex = 0; IndicesIndex < IndicesCount; IndicesIndex++, StaticIndex += 3)
+	int NumVertices = GeometryCollection.NumElements(FGeometryCollection::VerticesGroup);
+	int IndicesBaseIndex = GeometryCollection.AddElements(TriangleCount, FGeometryCollection::FacesGroup);
+
+	int CurrentIndex = IndicesBaseIndex;
+	for (int ComponentIndex = 0; ComponentIndex < ComponentsSourceIndices.Num(); ComponentIndex++)
 	{
-		int32 IndicesOffset = IndicesBaseIndex + IndicesIndex;
-		Indices[IndicesOffset] = FIntVector(
-			IndexBuffer[StaticIndex] + NumVertices,
-			IndexBuffer[StaticIndex + 1] + NumVertices,
-			IndexBuffer[StaticIndex + 2] + NumVertices);
-		Visible[IndicesOffset] = true;
-		MaterialID[IndicesOffset] = 0;
-		MaterialIndex[IndicesOffset] = IndicesOffset;
+		for(int32 TriangleIndex = 0; TriangleIndex< ComponentsSourceIndices[ComponentIndex].Num(); TriangleIndex++)
+		{
+			SourceToTargetTriangleMap[ComponentIndex][TriangleIndex][1] = CurrentIndex;
+
+			FIntVector& Triangle = ComponentsSourceIndices[ComponentIndex][TriangleIndex];
+			for (int k = 0; k < 3; k++)
+			{
+				Indices[CurrentIndex][k] = SourceVertexToComponentMap[Triangle[k]] + IndicesBaseIndex;
+			}
+
+			Visible[CurrentIndex] = true;
+			MaterialID[CurrentIndex] = 0;
+			MaterialIndex[CurrentIndex] = CurrentIndex;
+			CurrentIndex++;
+		}
 	}
 
 	//
 	// Vertex Attributes
 	//
-	TManagedArray<FVector3f>& Vertex = InCollection->ModifyAttribute<FVector3f>("Vertex", FGeometryCollection::VerticesGroup);
-	TManagedArray<FVector3f>& Normal = InCollection->ModifyAttribute<FVector3f>("Normal", FGeometryCollection::VerticesGroup);
-	TManagedArray<FLinearColor>& Color = InCollection->ModifyAttribute<FLinearColor>("Color", FGeometryCollection::VerticesGroup);
-	TManagedArray<FVector3f>& TangentU = InCollection->ModifyAttribute<FVector3f>("TangentU", FGeometryCollection::VerticesGroup);
-	TManagedArray<FVector3f>& TangentV = InCollection->ModifyAttribute<FVector3f>("TangentV", FGeometryCollection::VerticesGroup);
-	TManagedArray<int32>& BoneMap = InCollection->ModifyAttribute<int32>("BoneMap", FGeometryCollection::VerticesGroup);
+	int VertexBaseIndex = GeometryCollection.AddElements(VertexCount, FGeometryCollection::VerticesGroup);
+	int NumTargetVertices = GeometryCollection.NumElements(FGeometryCollection::VerticesGroup);
+
 
 
 	// 
 	// Transform Attributes 
 	// 
 
+	// add transforms for the separated geometry components. 
+	TArray<int32> ComponentToTransformGroupIndex;
+	ComponentToTransformGroupIndex.Init(INDEX_NONE, ComponentsSourceIndices.Num());
+	int32 ComponentTransformBaseIndex = GeometryCollection.AddElements(ComponentsSourceIndices.Num(), FTransformCollection::TransformGroup);
+	for (int ComponentIndex = 0; ComponentIndex < ComponentsSourceIndices.Num(); ComponentIndex++)
+	{
+		int32 ComponentTransformIndex = ComponentTransformBaseIndex + ComponentIndex;
+		Parent[ComponentTransformIndex] = RootIndex;
+		BoneName[ComponentTransformIndex] = FString::Printf(TEXT("%s_Mesh%d"), *BoneName[RootIndex], ComponentTransformIndex);
+		LocalSpaceTransform[ComponentTransformIndex] = FTransform3f::Identity;
+		SimulationType[ComponentTransformIndex] = FGeometryCollection::ESimulationTypes::FST_None;
+		BoneColor[ComponentTransformIndex] = FLinearColor::MakeRandomColor();
+		ComponentToTransformGroupIndex[ComponentIndex] = ComponentTransformIndex;
+	}
+
 	TArray<FTransform> ComponentTransform;
 	GeometryCollectionAlgo::GlobalMatrices(LocalSpaceTransform, Parent, ComponentTransform);
 
-	int InitialNumVertices = InCollection->NumElements(FGeometryCollection::VerticesGroup);
-	int VertexBaseIndex = InCollection->AddElements(VertexCount, FGeometryCollection::VerticesGroup);
-	const int32 NumUVLayers = VertexBuffers.StaticMeshVertexBuffer.GetNumTexCoords();
-	GeometryCollection::UV::SetNumUVLayers(*InCollection, NumUVLayers);
-	GeometryCollection::UV::FUVLayers UVLayers = GeometryCollection::UV::FindActiveUVLayers(*InCollection);
-	for (int32 VertexIndex = 0; VertexIndex < VertexCount; VertexIndex++)
-	{
-		int VertexOffset = VertexBaseIndex + VertexIndex;
-		BoneMap[VertexOffset] = -1;
-		int32 SkeletalBoneIndex = -1;
-		SkinWeightVertexBuffer.GetRigidWeightBone(VertexIndex, SkeletalBoneIndex);
-		if (SkeletalBoneIndex > -1)
-		{
-			BoneMap[VertexOffset] = RootIndex;
-			Vertex[VertexOffset] = (FVector4f)ComponentTransform[BoneMap[VertexOffset]].ToInverseMatrixWithScale().TransformPosition((FVector)PositionVertexBuffer.VertexPosition(VertexIndex));
-			//@todo(GeometryCollectionConversion) : Support nesting rigid geometry within a transform space
-			//BoneMap[VertexOffset] = SkeletalBoneIndex + TransformBaseIndex;
-			//Vertex[VertexOffset] = (FVector3f)PositionVertexBuffer.VertexPosition(VertexIndex);
-		}
-		check(BoneMap[VertexOffset] != -1);
-		TangentU[VertexOffset] = VertexBuffers.StaticMeshVertexBuffer.VertexTangentX(VertexIndex);
-		TangentV[VertexOffset] = VertexBuffers.StaticMeshVertexBuffer.VertexTangentY(VertexIndex);
-		Normal[VertexOffset] = VertexBuffers.StaticMeshVertexBuffer.VertexTangentZ(VertexIndex);
 
-		for (int32 UVLayerIdx = 0; UVLayerIdx < NumUVLayers; ++UVLayerIdx)
+	FLinearColor DefaultColor = FLinearColor::White;
+	TVertexInstanceAttributesConstRef<FVector4f> InstanceColors;
+	UE::Geometry::FDynamicMeshColorOverlay* ColorOverlay = nullptr;
+	if (UE::Geometry::FDynamicMeshAttributeSet* Attributes = DynamicMesh.Attributes())
+	{
+		FSkeletalMeshConstAttributes SkeletalMeshConstAttributes(MeshDescription);
+		InstanceColors = SkeletalMeshConstAttributes.GetVertexInstanceColors();
+		if (InstanceColors.IsValid())
 		{
-			UVLayers[UVLayerIdx][VertexOffset] = VertexBuffers.StaticMeshVertexBuffer.GetVertexUV(VertexIndex, UVLayerIdx);
+			ColorOverlay = Attributes->PrimaryColors();
+			DefaultColor = InstanceColors.GetDefaultValue();
 		}
-
-		if (VertexBuffers.ColorVertexBuffer.GetNumVertices() == VertexCount)
-			Color[VertexOffset] = VertexBuffers.ColorVertexBuffer.VertexColor(VertexIndex);
-		else
-			Color[VertexOffset] = FLinearColor::White;
 	}
 
-
-	// Geometry Group
-	TArray<int32> GeometryIndices;
-	GeometryCollectionAlgo::ContiguousArray(GeometryIndices, InCollection->NumElements(FGeometryCollection::GeometryGroup));
-	InCollection->RemoveDependencyFor(FGeometryCollection::GeometryGroup);
-	InCollection->RemoveElements(FGeometryCollection::GeometryGroup, GeometryIndices);
-	::GeometryCollection::AddGeometryProperties(InCollection);
-
-	const TArray<FSkelMeshRenderSection> &StaticMeshSections = SkeletalMesh->GetResourceForRendering()->LODRenderData[0].RenderSections;
-
-	TManagedArray<FGeometryCollectionSection> & Sections = InCollection->ModifyAttribute<FGeometryCollectionSection>("Sections", FGeometryCollection::MaterialGroup);
-
-	for (const FSkelMeshRenderSection &CurrSection : StaticMeshSections)
+	// @todo(GeometryCollectionConversion) : Add support for UV's, Normals  
+	//const int32 NumUVLayers = VertexBuffers.StaticMeshVertexBuffer.GetNumTexCoords();
+	//GeometryCollection::UV::SetNumUVLayers(*GeometryCollection, NumUVLayers);
+	//GeometryCollection::UV::FUVLayers UVLayers = GeometryCollection::UV::FindActiveUVLayers(*GeometryCollection);
+	TArray<bool> TargetVertexVisited; TargetVertexVisited.Init(false, NumTargetVertices);
+	for (int ComponentIndex = 0;ComponentIndex<ComponentsSourceIndices.Num();ComponentIndex++)
 	{
-		// create new section
-		int32 SectionIndex = InCollection->AddElements(1, FGeometryCollection::MaterialGroup);
-						
-		Sections[SectionIndex].MaterialID = MaterialStartIndex + CurrSection.MaterialIndex;
+		for (int32 TriangleIndex = 0; TriangleIndex < ComponentsSourceIndices[ComponentIndex].Num(); TriangleIndex++)
+		{
+			int32 SourceTriangleIndex = SourceToTargetTriangleMap[ComponentIndex][TriangleIndex][0];
+			FIntVector& Triangle = ComponentsSourceIndices[ComponentIndex][TriangleIndex];
+			for (int k = 0; k < 3; k++)
+			{
+				int SourceVertexIndex = Triangle[k];
+				if (ensure(0 <= SourceVertexIndex && SourceVertexIndex < SourceVertexToComponentMap.Num()))
+				{
+					int TargetVertexIndex = SourceVertexToComponentMap[SourceVertexIndex] + VertexBaseIndex;
+					if (ensure(VertexBaseIndex <= TargetVertexIndex && TargetVertexIndex < TargetVertexVisited.Num()))
+					{
+						if (!TargetVertexVisited[TargetVertexIndex])
+						{
+							// Vertex position
+							//  @todo(Skin Wieghts) : Currently its a static pose in component space. 
+							BoneMap[TargetVertexIndex] = ComponentToTransformGroupIndex[ComponentIndex];
 
-		Sections[SectionIndex].FirstIndex = IndicesBaseIndex * 3 + CurrSection.BaseIndex;
-		Sections[SectionIndex].MinVertexIndex = VertexBaseIndex + CurrSection.BaseVertexIndex;
+							FVector SourceVertex = (FVector)DynamicMesh.GetVertex(SourceVertexIndex);
+							FMatrix M = ComponentTransform[BoneMap[TargetVertexIndex]].ToInverseMatrixWithScale();
+							Vertex[TargetVertexIndex] = (FVector4f)M.TransformPosition(SourceVertex);
 
-		Sections[SectionIndex].NumTriangles = CurrSection.NumTriangles;
+							TargetVertexVisited[TargetVertexIndex] = true;
 
-		// #todo(dmp): what should we set this to?  SkeletalMesh sections are different
-		// but we are resetting this when the re indexing happens
-		Sections[SectionIndex].MaxVertexIndex = VertexBaseIndex + CurrSection.NumVertices;
+							//colors
+							Color[TargetVertexIndex] = FLinearColor::White;
+							if (ColorOverlay)
+							{
+								FVertexInstanceID SourceInstanceID = MeshDescription.GetTriangleVertexInstance(SourceTriangleIndex, k);
+								FVector4f InstColor = InstanceColors.Get(SourceInstanceID);
+								UE::Geometry::LinearColors::SRGBToLinear(InstColor);
+								//Color[TargetVertexIndex] = InstColor;
+							}
 
-		// set the materialid for all of the faces
-		for (int32 i = Sections[SectionIndex].FirstIndex / 3; i < Sections[SectionIndex].FirstIndex / 3 + Sections[SectionIndex].NumTriangles; ++i)
-		{
-			MaterialID[i] = SectionIndex;
+							// @todo(GeometryCollectionConversion) : Add support for UV's, Normals  
+							//TangentU[VertexOffset] = VertexBuffers.StaticMeshVertexBuffer.VertexTangentX(VertexIndex);
+							//TangentV[VertexOffset] = VertexBuffers.StaticMeshVertexBuffer.VertexTangentY(VertexIndex);
+							//Normal[VertexOffset] = VertexBuffers.StaticMeshVertexBuffer.VertexTangentZ(VertexIndex);
+							//for (int32 UVLayerIdx = 0; UVLayerIdx < NumUVLayers; ++UVLayerIdx)
+							//{
+							//	UVLayers[UVLayerIdx][VertexOffset] = VertexBuffers.StaticMeshVertexBuffer.GetVertexUV(VertexIndex, UVLayerIdx);
+							//}
+						}
+					}
+				}
+			}
 		}
 	}
+	
 
-	if (bReindexMaterials)
-	{
-		FGeometryCollection::ReindexMaterials(*InCollection);
-	}
+	// Geometry Group
+	TArray<int32> GeometryIndices;
+	FGeometryCollection::DefineGeometrySchema(GeometryCollection);
+	GeometryCollectionAlgo::ContiguousArray(GeometryIndices, GeometryCollection.NumElements(FGeometryCollection::GeometryGroup));
+	GeometryCollection.RemoveDependencyFor(FGeometryCollection::GeometryGroup);
+	GeometryCollection.RemoveElements(FGeometryCollection::GeometryGroup, GeometryIndices);
+	::GeometryCollection::AddGeometryProperties(&GeometryCollection);
 
+	GeometryCollection.CopyTo(InManagedArrayCollection);
 	return true;
+#else
+	return false;
+#endif
 }
 
 void FGeometryCollectionEngineConversion::AppendSkeleton(const USkeleton* InSkeleton, const FTransform& SkeletalMeshTransform, FManagedArrayCollection* InCollection)
diff --git a/Engine/Source/Runtime/Experimental/GeometryCollectionEngine/Private/GeometryCollection/GeometryCollectionEngineUtility.cpp b/Engine/Source/Runtime/Experimental/GeometryCollectionEngine/Private/GeometryCollection/GeometryCollectionEngineUtility.cpp
--- a/Engine/Source/Runtime/Experimental/GeometryCollectionEngine/Private/GeometryCollection/GeometryCollectionEngineUtility.cpp
+++ b/Engine/Source/Runtime/Experimental/GeometryCollectionEngine/Private/GeometryCollection/GeometryCollectionEngineUtility.cpp
@@ -1,18 +1,23 @@
 // Copyright Epic Games, Inc. All Rights Reserved.
 
 #include "GeometryCollection/GeometryCollectionEngineUtility.h"
+
+#include "Chaos/Utilities.h"
+#include "DynamicMesh/DynamicMesh3.h"
+#include "Engine/SkeletalMesh.h"
 #include "GeometryCollection/GeometryCollection.h"
 #include "GeometryCollection/GeometryCollectionCache.h"
 #include "GeometryCollection/GeometryCollectionAlgo.h"
+#include "Rendering/SkeletalMeshRenderData.h"
+#include "MeshDescriptionToDynamicMesh.h"
+#include "MeshDescription.h"
+#include "Modules/ModuleInterface.h"
+#include "Modules/ModuleManager.h"
 
 #if WITH_EDITOR
 #include "MeshUtilities.h"
 #endif
 
-#include "Modules/ModuleInterface.h"
-#include "Modules/ModuleManager.h"
-
-
 DEFINE_LOG_CATEGORY_STATIC(LogGeoemtryCollectionClean, Verbose, All);
 
 
@@ -443,3 +448,109 @@ void GeometryCollectionEngineUtility::ComputeTangents(FGeometryCollection* Geome
 #endif
 }
 
+class FSortableIntVector2 : public UE::Math::TIntVector2<int32>
+{
+public:
+	FSortableIntVector2()
+	{
+		X = INT_MAX;
+		Y = -INT_MAX;
+	}
+
+	bool operator < (const FIntVector2& Other) const
+	{
+		return X < Other.X;
+	}
+};
+
+void GeometryCollectionEngineUtility::GenerateConnectedComponents(const USkeletalMesh* InSkeletalMesh,
+	TArray<TArray<FIntVector>>& SourceTriangleVertices,
+	TArray<TArray<FIntVector2>>& SourceTriangleIndex,
+	TArray<int32>& VertexComponentMap,
+	int32& TriangleCount, int32& VertexCount)
+{
+#if WITH_EDITOR
+	constexpr int32 LODIndex = 0;
+	if (InSkeletalMesh->HasMeshDescription(LODIndex))
+	{
+		FMeshDescription SourceMesh;
+		UE::Geometry::FDynamicMesh3 DynamicMesh;
+		FMeshDescriptionToDynamicMesh Converter;
+		InSkeletalMesh->CloneMeshDescription(LODIndex, SourceMesh);
+		Converter.Convert(&SourceMesh, DynamicMesh);
+
+		TArray<int32> VertexComponentID;
+		VertexComponentID.Init(INDEX_NONE, DynamicMesh.MaxVertexID());
+
+		int32 MaxComponentID = INDEX_NONE;
+		for (int VisitedVertexIdx = 0; VisitedVertexIdx < VertexComponentID.Num(); VisitedVertexIdx++)
+		{
+			if (VertexComponentID[VisitedVertexIdx] == INDEX_NONE)
+			{
+				MaxComponentID++;
+				TArray<int32> Neighbors;
+				Neighbors.Push(VisitedVertexIdx);
+
+				while (!Neighbors.IsEmpty())
+				{
+					int32 CurrentVertex = Neighbors.Pop();
+					for (int32 EdgeID : DynamicMesh.VtxEdgesItr(CurrentVertex))
+					{
+						const FDynamicMesh3::FEdge& Edge = DynamicMesh.GetEdge(EdgeID);
+						int NewVertex = Edge.Vert[0] == CurrentVertex ? Edge.Vert[1] : Edge.Vert[0];
+						if (VertexComponentID[NewVertex] == INDEX_NONE)
+						{
+							Neighbors.Push(NewVertex);
+						}
+					}
+					VertexComponentID[CurrentVertex] = MaxComponentID;
+				}
+			}
+		}
+
+		// Build Component Triangles
+		TriangleCount = 0;
+		SourceTriangleIndex.Init(TArray< FIntVector2 >(), MaxComponentID + 1);
+		SourceTriangleVertices.Init(TArray<FIntVector>(), MaxComponentID + 1);
+		for (int32 TriangleID : DynamicMesh.TriangleIndicesItr())
+		{
+			const FIntVector3& Triangle = DynamicMesh.GetTriangle(TriangleID);
+			int32 ComponentID = VertexComponentID[Triangle[0]];
+			for (int k = 1; k < 3; k++)
+			{
+				if (!ensure(VertexComponentID[Triangle[0]] == ComponentID))
+				{
+					return;
+				}
+			}
+			SourceTriangleVertices[ComponentID].Add(Triangle);
+			SourceTriangleIndex[ComponentID].Add(FIntVector2(TriangleID, INDEX_NONE));
+			TriangleCount++;
+		}
+
+		// Build Remapping
+		int32 CurrentRemapIndex = 0;
+		VertexComponentMap.Init(INDEX_NONE, DynamicMesh.MaxVertexID());
+		for (TArray<FIntVector>& Component : SourceTriangleVertices)
+		{
+			for (FIntVector& Triangle : Component)
+			{
+				for (int k = 0; k < 3; k++)
+				{
+					int32 Index = Triangle[k];
+					if (VertexComponentMap[Index] == INDEX_NONE)
+					{
+						VertexComponentMap[Index] = CurrentRemapIndex++;
+					}
+				}
+			}
+		}
+		VertexCount = CurrentRemapIndex;
+		if (!ensure(CurrentRemapIndex <= VertexComponentMap.Num()))
+		{
+			return;
+		}
+		return;
+	}
+#endif
+}
diff --git a/Engine/Source/Runtime/Experimental/GeometryCollectionEngine/Public/GeometryCollection/GeometryCollectionEngineUtility.h b/Engine/Source/Runtime/Experimental/GeometryCollectionEngine/Public/GeometryCollection/GeometryCollectionEngineUtility.h
--- a/Engine/Source/Runtime/Experimental/GeometryCollectionEngine/Public/GeometryCollection/GeometryCollectionEngineUtility.h
+++ b/Engine/Source/Runtime/Experimental/GeometryCollectionEngine/Public/GeometryCollection/GeometryCollectionEngineUtility.h
@@ -6,6 +6,7 @@
 
 class FGeometryCollection;
 class UGeometryCollectionCache;
+class USkeletalMesh;
 
 namespace GeometryCollectionEngineUtility
 {
@@ -17,4 +18,21 @@ namespace GeometryCollectionEngineUtility
 
 	void GEOMETRYCOLLECTIONENGINE_API ComputeTangents(FGeometryCollection* GeometryCollection);
 
+	/***
+	* Generate a vertex to component map that defines the disjoint geometries.
+	* @param IndexBuffer			Input flat index of triangles
+	* @param ComponentIndices		Output list of triangle indices by component
+	* @param TriangleComponentMap	Output Component remapping for triangles
+	* @param VertexComponentMap		Output Component remapping for vertices
+	* @param TriangleCount		    Output Sum of the triangles in the Components
+	* @param VertexCount		    Output Num of vertices in the Components
+	* @return If the remapping was valid.
+	*/
+	void
+	GEOMETRYCOLLECTIONENGINE_API
+	GenerateConnectedComponents(const USkeletalMesh* InSkeletalMesh, 
+		TArray<TArray<FIntVector>>& ComponentIndices,
+		TArray<TArray<FIntVector2>>& TrangleComponentMap, 
+		TArray<int32>& VertexComponentMap, 
+		int32& TriangleCount, int32& VertexCount);
 }
\ No newline at end of file
