Separate out the interface for setting custom data on geometry collection custom renderers.
Add support for setting custom data by named channels for which the custom renderer can handle any packing logic.

diff --git a/Engine/Source/Runtime/Experimental/GeometryCollectionEngine/Private/GeometryCollection/GeometryCollectionBlueprintLibrary.cpp b/Engine/Source/Runtime/Experimental/GeometryCollectionEngine/Private/GeometryCollection/GeometryCollectionBlueprintLibrary.cpp
--- a/Engine/Source/Runtime/Experimental/GeometryCollectionEngine/Private/GeometryCollection/GeometryCollectionBlueprintLibrary.cpp
+++ b/Engine/Source/Runtime/Experimental/GeometryCollectionEngine/Private/GeometryCollection/GeometryCollectionBlueprintLibrary.cpp
@@ -5,16 +5,39 @@
 #include "GeometryCollection/GeometryCollectionComponent.h"
 #include "GeometryCollection/GeometryCollectionExternalRenderInterface.h"
 
-void UGeometryCollectionBlueprintLibrary::SetISMPoolCustomInstanceData(UGeometryCollectionComponent* GeometryCollectionComponent, int32 CustomFloatIndex, float CustomFloatValue)
+void UGeometryCollectionBlueprintLibrary::SetCustomInstanceDataByIndex(UGeometryCollectionComponent* GeometryCollectionComponent, int32 CustomDataIndex, float CustomDataValue)
 {
 	if (GeometryCollectionComponent == nullptr)
 	{
 		return;
 	}
-	IGeometryCollectionExternalRenderInterface* CustomRenderer = GeometryCollectionComponent->GetCustomRenderer();
+
+	IGeometryCollectionCustomDataInterface* CustomRenderer = Cast<IGeometryCollectionCustomDataInterface>(GeometryCollectionComponent->GetCustomRenderer());
+	if (CustomRenderer == nullptr)
+	{
+		return;
+	}
+
+	CustomRenderer->SetCustomInstanceData(CustomDataIndex, CustomDataValue);
+}
+
+void UGeometryCollectionBlueprintLibrary::SetCustomInstanceDataByName(UGeometryCollectionComponent* GeometryCollectionComponent, FName CustomDataName, float CustomDataValue)
+{
+	if (GeometryCollectionComponent == nullptr)
+	{
+		return;
+	}
+
+	IGeometryCollectionCustomDataInterface* CustomRenderer = Cast<IGeometryCollectionCustomDataInterface>(GeometryCollectionComponent->GetCustomRenderer());
 	if (CustomRenderer == nullptr)
 	{
 		return;
 	}
-	CustomRenderer->SetCustomInstanceData(CustomFloatIndex, CustomFloatValue);
+
+	CustomRenderer->SetCustomInstanceData(CustomDataName, CustomDataValue);
+}
+
+void UGeometryCollectionBlueprintLibrary::SetISMPoolCustomInstanceData(UGeometryCollectionComponent* GeometryCollectionComponent, int32 CustomDataIndex, float CustomDataValue)
+{
+	SetCustomInstanceDataByIndex(GeometryCollectionComponent, CustomDataIndex, CustomDataValue);
 }
diff --git a/Engine/Source/Runtime/Experimental/GeometryCollectionEngine/Public/GeometryCollection/GeometryCollectionBlueprintLibrary.h b/Engine/Source/Runtime/Experimental/GeometryCollectionEngine/Public/GeometryCollection/GeometryCollectionBlueprintLibrary.h
--- a/Engine/Source/Runtime/Experimental/GeometryCollectionEngine/Public/GeometryCollection/GeometryCollectionBlueprintLibrary.h
+++ b/Engine/Source/Runtime/Experimental/GeometryCollectionEngine/Public/GeometryCollection/GeometryCollectionBlueprintLibrary.h
@@ -15,11 +15,25 @@ class GEOMETRYCOLLECTIONENGINE_API UGeometryCollectionBlueprintLibrary : public
 public:
 	/** 
 	 * Set a custom instance data value for all instances associated with a geometry collection. 
-	 * This assumes that the geometry collection is using a custom instanced renderer.
+	 * This assumes that the geometry collection is using a custom renderer that supports IGeometryCollectionCustomDataInterface.
 	 * @param GeometryCollectionComponent	The Geometry Collection Component that we want to set custom instance data on.
-	 * @param CustomDataIndex	The custom instance data slot that we want to set.
+	 * @param CustomDataIndex	The index of the custom instance data slot that we want to set.
 	 * @param CustomDataValue	The value to set to the custom instance data slot.
 	 */
-	UFUNCTION(BlueprintCallable, Category = AtomMaterial)
+	UFUNCTION(BlueprintCallable, Category = "ChaosPhysics")
+	static void SetCustomInstanceDataByIndex(UGeometryCollectionComponent* GeometryCollectionComponent, int32 CustomDataIndex, float CustomDataValue);
+
+	/** 
+	 * Set a custom instance data value for all instances associated with a geometry collection. 
+	 * This assumes that the geometry collection is using a custom renderer that supports IGeometryCollectionCustomDataInterface.
+	 * @param GeometryCollectionComponent	The Geometry Collection Component that we want to set custom instance data on.
+	 * @param CustomDataName	The name of the custom instance data slot that we want to set.
+	 * @param CustomDataValue	The value to set to the custom instance data slot.
+	 */
+	UFUNCTION(BlueprintCallable, Category = "ChaosPhysics")
+	static void SetCustomInstanceDataByName(UGeometryCollectionComponent* GeometryCollectionComponent, FName CustomDataName, float CustomDataValue);
+
+	UE_DEPRECATED(5.5, "Use SetCustomInstanceDataByIndex() instead")
+	UFUNCTION(BlueprintCallable, Category = "ChaosPhysics", meta = (DeprecatedFunction, DeprecationMessage = "Use SetCustomInstanceDataByIndex() instead"))
 	static void SetISMPoolCustomInstanceData(UGeometryCollectionComponent* GeometryCollectionComponent, int32 CustomDataIndex, float CustomDataValue);
 };
diff --git a/Engine/Source/Runtime/Experimental/GeometryCollectionEngine/Public/GeometryCollection/GeometryCollectionExternalRenderInterface.h b/Engine/Source/Runtime/Experimental/GeometryCollectionEngine/Public/GeometryCollection/GeometryCollectionExternalRenderInterface.h
--- a/Engine/Source/Runtime/Experimental/GeometryCollectionEngine/Public/GeometryCollection/GeometryCollectionExternalRenderInterface.h
+++ b/Engine/Source/Runtime/Experimental/GeometryCollectionEngine/Public/GeometryCollection/GeometryCollectionExternalRenderInterface.h
@@ -27,6 +27,11 @@ public:
 		EState_ForcedBroken	= 1 << 2,
 	};
 
+	/** Returns true if the custom renderer can't render for the given state, and requires native rendering instead. */
+	virtual bool ShouldUseNativeFallback(uint32 InStateFlags) const { return false; }
+	/** Returns true if this renderer can ever return true for ShouldUseNativeFallback(). */
+	virtual bool CanEverUseNativeFallback() const { return false; }
+
 	/** Called on geometry collection component registration. */
 	virtual void OnRegisterGeometryCollection(UGeometryCollectionComponent& InComponent) = 0;
 	/** Called on geometry collection component unregistration. */
@@ -40,13 +45,13 @@ public:
 
 	/**
 	 * Update the root bone transform of the geometry collection.
-	 * if the geometry collection is using of multiple proxy root meshes this transform applies to all of them.
+	 * If the geometry collection is using of multiple proxy root meshes this transform applies to all of them.
 	 */
 	virtual void UpdateRootTransform(UGeometryCollection const& InGeometryCollection, FTransform const& InRootTransform) = 0;
 
 	/**
 	 * Update the root proxy transforms of the geometry collection.
-	 * if the geometry collection is using of multiple proxy root meshes, InRootTransforms is expected to contain an entry for each of them .
+	 * If the geometry collection is using of multiple proxy root meshes, InRootTransforms is expected to contain an entry for each of them .
 	 * @param InRootTransform		Component space root transform
 	 * @param InRootLocalTransforms	Root space local transforms
 	 */
@@ -57,18 +62,6 @@ public:
 	 */
 	virtual void UpdateTransforms(UGeometryCollection const& InGeometryCollection, TArrayView<const FTransform3f> InTransforms) = 0;
 
-	/** Returns true if the custom renderer can't render for the given state, and requires native rendering instead. */
-	virtual bool ShouldUseNativeFallback(uint32 InStateFlags) const { return false; }
-	
-	/** Returns true if this renderer can ever return true for ShouldUseNativeFallback(). */
-	virtual bool CanEverUseNativeFallback() const { return false; }
-
-	/**
-	 * Set a custom instance data slot to be used by any instances the renderer is spawning.
-	 * todo: This is a bit of a leaky abstraction for ISMPool renderers. Move this to a separate ISMPool renderer interface? Or is that overengineering?
-	 */ 
-	virtual void SetCustomInstanceData(int32 CustomFloatIndex, float CustomFloatValue) {}
-
 	UE_DEPRECATED(5.4, "Use flags version of UpdateState instead")
 	virtual void UpdateState(UGeometryCollection const& InGeometryCollection, FTransform const& InComponentTransform, bool bInIsBroken, bool bInIsVisible) 
 	{
@@ -80,3 +73,22 @@ public:
 	UE_DEPRECATED(5.3, "Use FTransform version of UpdateTransforms instead")
 	virtual void UpdateTransforms(UGeometryCollection const& InGeometryCollection, TArrayView<const FMatrix> InMatrices) {};
 };
+
+
+UINTERFACE(MinimalAPI)
+class UGeometryCollectionCustomDataInterface : public UInterface
+{
+	GENERATED_BODY()
+};
+
+class IGeometryCollectionCustomDataInterface
+{
+	GENERATED_BODY()
+
+public:
+	/** Set value to a custom instance data slot by index. The value will be set for all instances the renderer is spawning. */
+	virtual void SetCustomInstanceData(int32 CustomDataIndex, float CustomDataValue) = 0;
+
+	/** Set value a custom instance data slot by name. The value will be set for all instances the renderer is spawning. */
+	virtual void SetCustomInstanceData(FName CustomDataName, float CustomDataValue) = 0;
+};
