Changed namespace for individual asset streaming.

diff --git a/Engine/Plugins/Experimental/StudioTelemetry/Source/StudioTelemetry/Private/StudioTelemetryEditor.cpp b/Engine/Plugins/Experimental/StudioTelemetry/Source/StudioTelemetry/Private/StudioTelemetryEditor.cpp
--- a/Engine/Plugins/Experimental/StudioTelemetry/Source/StudioTelemetry/Private/StudioTelemetryEditor.cpp
+++ b/Engine/Plugins/Experimental/StudioTelemetry/Source/StudioTelemetry/Private/StudioTelemetryEditor.cpp
@@ -284,7 +284,7 @@ void FStudioTelemetryEditor::RecordEvent_IAS(const FString& Context, TArray<FAna
 	// Gather the summary stats
 	FDerivedDataCacheSummaryStats SummaryStats;
 
-	using namespace UE::IO::IAS;
+	using namespace UE::IoStore;
 	if (FIoStoreOnDemandModule* OnDemandModule = FModuleManager::Get().GetModulePtr<FIoStoreOnDemandModule>("IoStoreOnDemand"))
 	{
 		if (OnDemandModule->IsEnabled())
@@ -1041,4 +1041,4 @@ void FStudioTelemetryEditor::Shutdown()
 
 UE_ENABLE_OPTIMIZATION_SHIP
 
-#endif // WITH_EDITOR
\ No newline at end of file
+#endif // WITH_EDITOR
diff --git a/Engine/Source/Programs/IoStoreOnDemand/Private/IasTool.cpp b/Engine/Source/Programs/IoStoreOnDemand/Private/IasTool.cpp
--- a/Engine/Source/Programs/IoStoreOnDemand/Private/IasTool.cpp
+++ b/Engine/Source/Programs/IoStoreOnDemand/Private/IasTool.cpp
@@ -12,11 +12,11 @@
 IMPLEMENT_APPLICATION(IasTool, "IasTool");
 
 ////////////////////////////////////////////////////////////////////////////////
-namespace UE::IO::IAS::Tool {
+namespace UE::IoStore::Tool {
 
 int32 Main(int32, TCHAR*[]);
 
-} // namespace UE::IO::IAS
+} // namespace UE::IoStore
 
 ////////////////////////////////////////////////////////////////////////////////
 INT32_MAIN_INT32_ARGC_TCHAR_ARGV()
@@ -36,5 +36,5 @@ INT32_MAIN_INT32_ARGC_TCHAR_ARGV()
 		FEngineLoop::AppExit();
 	};
 
-	return UE::IO::IAS::Tool::Main(ArgC, ArgV);
+	return UE::IoStore::Tool::Main(ArgC, ArgV);
 }
diff --git a/Engine/Source/Runtime/Experimental/IoStoreOnDemand/Private/CoreHttp/Client.h b/Engine/Source/Runtime/Experimental/IoStoreOnDemand/Private/CoreHttp/Client.h
--- a/Engine/Source/Runtime/Experimental/IoStoreOnDemand/Private/CoreHttp/Client.h
+++ b/Engine/Source/Runtime/Experimental/IoStoreOnDemand/Private/CoreHttp/Client.h
@@ -15,7 +15,7 @@
 ////////////////////////////////////////////////////////////////////////////////
 class FIoBuffer;
 
-namespace UE::IO::IAS::HTTP
+namespace UE::IoStore::HTTP
 {
 
 ////////////////////////////////////////////////////////////////////////////////
@@ -215,7 +215,7 @@ private:
 	FEventLoop&				operator = (FEventLoop&&)		= delete;
 };
 
-} // namespace UE::IO::IAS::HTTP
+} // namespace UE::IoStore::HTTP
 
 #undef UE_API
 
diff --git a/Engine/Source/Runtime/Experimental/IoStoreOnDemand/Private/CoreHttp/Http.cpp b/Engine/Source/Runtime/Experimental/IoStoreOnDemand/Private/CoreHttp/Http.cpp
--- a/Engine/Source/Runtime/Experimental/IoStoreOnDemand/Private/CoreHttp/Http.cpp
+++ b/Engine/Source/Runtime/Experimental/IoStoreOnDemand/Private/CoreHttp/Http.cpp
@@ -81,7 +81,7 @@ static const SocketType InvalidSocket = ~SocketType(0);
 
 // }}}
 
-namespace UE::IO::IAS::HTTP
+namespace UE::IoStore::HTTP
 {
 
 // {{{1 trace ..................................................................
@@ -4119,4 +4119,4 @@ IOSTOREONDEMAND_API void IasHttpTest(const ANSICHAR* TestHost="localhost")
 
 // }}}
 
-} // namespace UE::IO::IAS::HTTP
+} // namespace UE::IoStore::HTTP
diff --git a/Engine/Source/Runtime/Experimental/IoStoreOnDemand/Private/CoreHttp/LatencyTesting.cpp b/Engine/Source/Runtime/Experimental/IoStoreOnDemand/Private/CoreHttp/LatencyTesting.cpp
--- a/Engine/Source/Runtime/Experimental/IoStoreOnDemand/Private/CoreHttp/LatencyTesting.cpp
+++ b/Engine/Source/Runtime/Experimental/IoStoreOnDemand/Private/CoreHttp/LatencyTesting.cpp
@@ -10,7 +10,7 @@
 #include "Misc/StringBuilder.h"
 #include "Templates/Function.h"
 
-namespace UE::IO::IAS::HTTP
+namespace UE::IoStore::HTTP
 {
 
 void LatencyTest(FStringView InUrl, FStringView InPath, uint32 InTimeOutMs, TArrayView<int32> OutResults)
@@ -66,4 +66,4 @@ void LatencyTest(FStringView InUrl, FStringView InPath, uint32 InTimeOutMs, TArr
 	UE_LOG(LogIas, Log, TEXT("Testing endpoint %s"), ANSI_TO_TCHAR(ConnectionDesc.ToString()));
 }
 
-} // namespace UE::IO::IAS::HTTP
+} // namespace UE::IoStore::HTTP
diff --git a/Engine/Source/Runtime/Experimental/IoStoreOnDemand/Private/CoreHttp/LatencyTesting.h b/Engine/Source/Runtime/Experimental/IoStoreOnDemand/Private/CoreHttp/LatencyTesting.h
--- a/Engine/Source/Runtime/Experimental/IoStoreOnDemand/Private/CoreHttp/LatencyTesting.h
+++ b/Engine/Source/Runtime/Experimental/IoStoreOnDemand/Private/CoreHttp/LatencyTesting.h
@@ -5,9 +5,9 @@
 #include "Containers/ArrayView.h"
 #include "Containers/StringFwd.h"
 
-namespace UE::IO::IAS::HTTP
+namespace UE::IoStore::HTTP
 {
 
 void LatencyTest(FStringView InUrl, FStringView InPath, uint32 InTimeOutMs, TArrayView<int32> OutResults);
 
-} // namespace UE::IO::IAS::HTTP
+} // namespace UE::IoStore::HTTP
diff --git a/Engine/Source/Runtime/Experimental/IoStoreOnDemand/Private/DistributionEndpoints.cpp b/Engine/Source/Runtime/Experimental/IoStoreOnDemand/Private/DistributionEndpoints.cpp
--- a/Engine/Source/Runtime/Experimental/IoStoreOnDemand/Private/DistributionEndpoints.cpp
+++ b/Engine/Source/Runtime/Experimental/IoStoreOnDemand/Private/DistributionEndpoints.cpp
@@ -14,7 +14,7 @@
 #include "Serialization/JsonSerializer.h"
 #include "Statistics.h"
 
-namespace UE::IO::IAS
+namespace UE::IoStore
 {
 
 static int32 GDistributedEndpointTimeout = 30;
@@ -119,4 +119,4 @@ FDistributionEndpoints::EResult FDistributionEndpoints::ParseResponse(FHttpRespo
 	return EResult::Failure;
 }
 
-} // namespace UE::IO::IAS
+} // namespace UE::IoStore
diff --git a/Engine/Source/Runtime/Experimental/IoStoreOnDemand/Private/DistributionEndpoints.h b/Engine/Source/Runtime/Experimental/IoStoreOnDemand/Private/DistributionEndpoints.h
--- a/Engine/Source/Runtime/Experimental/IoStoreOnDemand/Private/DistributionEndpoints.h
+++ b/Engine/Source/Runtime/Experimental/IoStoreOnDemand/Private/DistributionEndpoints.h
@@ -8,7 +8,7 @@
 #include "HAL/CriticalSection.h"
 #include "HttpFwd.h"
 
-namespace UE::IO::IAS
+namespace UE::IoStore
 {
 
 class FDistributionEndpoints
@@ -45,4 +45,4 @@ private:
 	FDistributionEndpoints::EResult ParseResponse(FHttpResponsePtr HttpResponse, TArray<FString>& OutUrls);
 };
 
-} // namespace UE::IO::IAS
+} // namespace UE::IoStore
diff --git a/Engine/Source/Runtime/Experimental/IoStoreOnDemand/Private/IasCache.h b/Engine/Source/Runtime/Experimental/IoStoreOnDemand/Private/IasCache.h
--- a/Engine/Source/Runtime/Experimental/IoStoreOnDemand/Private/IasCache.h
+++ b/Engine/Source/Runtime/Experimental/IoStoreOnDemand/Private/IasCache.h
@@ -12,7 +12,8 @@ class FIoBuffer;
 class FIoReadOptions;
 struct FIoHash;
 
-namespace UE::IO::IAS {
+namespace UE::IoStore
+{
 
 /** Cache for binary blobs with a 20 byte cache key. */
 class IIasCache
@@ -76,4 +77,4 @@ struct FIasCacheConfig
 
 TUniquePtr<IIasCache> MakeIasCache(const TCHAR* RootDir, const FIasCacheConfig& Config);
 
-} // namespace UE::IO::IAS
+} // namespace UE::IoStore
diff --git a/Engine/Source/Runtime/Experimental/IoStoreOnDemand/Private/IoStoreOnDemand.cpp b/Engine/Source/Runtime/Experimental/IoStoreOnDemand/Private/IoStoreOnDemand.cpp
--- a/Engine/Source/Runtime/Experimental/IoStoreOnDemand/Private/IoStoreOnDemand.cpp
+++ b/Engine/Source/Runtime/Experimental/IoStoreOnDemand/Private/IoStoreOnDemand.cpp
@@ -47,7 +47,7 @@
 
 DEFINE_LOG_CATEGORY(LogIas);
 
-namespace UE::IO::IAS
+namespace UE::IoStore
 {
 
 FString GIasOnDemandTocExt = TEXT(".uondemandtoc");
@@ -2651,8 +2651,8 @@ EOnDemandInitResult FIoStoreOnDemandModule::Initialize()
 
 #endif // UE_IAS_CUSTOM_INITIALIZATION
 
-} // namespace UE::IO::IAS
+} // namespace UE::IoStore
 
 ////////////////////////////////////////////////////////////////////////////////
 
-IMPLEMENT_MODULE(UE::IO::IAS::FIoStoreOnDemandModule, IoStoreOnDemand);
+IMPLEMENT_MODULE(UE::IoStore::FIoStoreOnDemandModule, IoStoreOnDemand);
diff --git a/Engine/Source/Runtime/Experimental/IoStoreOnDemand/Private/JournaledCache.cpp b/Engine/Source/Runtime/Experimental/IoStoreOnDemand/Private/JournaledCache.cpp
--- a/Engine/Source/Runtime/Experimental/IoStoreOnDemand/Private/JournaledCache.cpp
+++ b/Engine/Source/Runtime/Experimental/IoStoreOnDemand/Private/JournaledCache.cpp
@@ -43,7 +43,7 @@ writing of data (and journal) to disk depending on how full the memory cache
 is. They are expressed as percentages.
 */
 
-namespace UE::IO::IAS::JournaledCache
+namespace UE::IoStore::JournaledCache
 {
 
 // {{{1 misc ...................................................................
@@ -2073,11 +2073,11 @@ uint32 FServiceThread::UpdateCache(FCache* Cache)
 
 // }}}
 
-} // namespace UE::IO::IAS::JournaledCache
+} // namespace UE::IoStore::JournaledCache
 
 
 
-namespace UE::IO::IAS {
+namespace UE::IoStore {
 
 // {{{1 journaled-cache ........................................................
 
@@ -2998,6 +2998,6 @@ IOSTOREONDEMAND_API void Tests(const TCHAR* CacheDir=nullptr)
 
 // }}}
 
-} // namespace UE::IO::IAS
+} // namespace UE::IoStore
 
 /* vim: set noet foldlevel=1 : */
diff --git a/Engine/Source/Runtime/Experimental/IoStoreOnDemand/Private/LatencyInjector.cpp b/Engine/Source/Runtime/Experimental/IoStoreOnDemand/Private/LatencyInjector.cpp
--- a/Engine/Source/Runtime/Experimental/IoStoreOnDemand/Private/LatencyInjector.cpp
+++ b/Engine/Source/Runtime/Experimental/IoStoreOnDemand/Private/LatencyInjector.cpp
@@ -9,7 +9,7 @@
 
 #if IAS_WITH_LATENCY_INJECTOR
 
-namespace UE::IO::IAS
+namespace UE::IoStore
 {
 
 ////////////////////////////////////////////////////////////////////////////////
@@ -162,6 +162,6 @@ bool FLatencyInjector::HasExpired(uint32 Param)
 	return Cycles >= (int64(Param) + Injector::CycleBase);
 }
 
-} // namespace UE::IO::IAS
+} // namespace UE::IoStore
 
 #endif // IAS_WITH_LATENCY_INJECTOR
diff --git a/Engine/Source/Runtime/Experimental/IoStoreOnDemand/Private/LatencyInjector.h b/Engine/Source/Runtime/Experimental/IoStoreOnDemand/Private/LatencyInjector.h
--- a/Engine/Source/Runtime/Experimental/IoStoreOnDemand/Private/LatencyInjector.h
+++ b/Engine/Source/Runtime/Experimental/IoStoreOnDemand/Private/LatencyInjector.h
@@ -5,7 +5,7 @@
 #include "HAL/Platform.h"
 #include "Misc/Build.h"
 
-namespace UE::IO::IAS
+namespace UE::IoStore
 {
 
 #if !UE_BUILD_SHIPPING
@@ -26,4 +26,4 @@ struct FLatencyInjector
 	static bool HasExpired(uint32 Param)			IAS_DISABLED_IMPL(return true);
 };
 
-} // namespace UE::IO::IAS
+} // namespace UE::IoStore
diff --git a/Engine/Source/Runtime/Experimental/IoStoreOnDemand/Private/OnDemandHttpClient.cpp b/Engine/Source/Runtime/Experimental/IoStoreOnDemand/Private/OnDemandHttpClient.cpp
--- a/Engine/Source/Runtime/Experimental/IoStoreOnDemand/Private/OnDemandHttpClient.cpp
+++ b/Engine/Source/Runtime/Experimental/IoStoreOnDemand/Private/OnDemandHttpClient.cpp
@@ -8,7 +8,7 @@
 #include "HAL/PlatformTime.h"
 #include "IO/IoStoreOnDemand.h"
 
-namespace UE::IO::IAS
+namespace UE::IoStore
 {
 
 static int32 GIasHttpFailTimeOutMs = 4 * 1000;
@@ -107,7 +107,7 @@ TUniquePtr<FHttpClient> FHttpClient::Create(FHttpClientConfig&& ClientConfig)
 		}
 
 		//TODO: Get rid of all string conversions
-		if (UE::IO::IAS::HTTP::FConnectionPool::IsValidHostUrl(StringCast<ANSICHAR>(*Ep, Ep.Len())) == false)
+		if (UE::IoStore::HTTP::FConnectionPool::IsValidHostUrl(StringCast<ANSICHAR>(*Ep, Ep.Len())) == false)
 		{
 			return TUniquePtr<FHttpClient>();
 		}
@@ -125,7 +125,7 @@ TUniquePtr<FHttpClient> FHttpClient::Create(const FString& Endpoint)
 
 void FHttpClient::IssueRequest(FRequestParams&& Params)
 {
-	using namespace UE::IO::IAS::HTTP;
+	using namespace UE::IoStore::HTTP;
 
 	//TODO: Remove string conversion
 	const auto Url = StringCast<ANSICHAR>(*Params.Url, Params.Url.Len());
@@ -263,4 +263,4 @@ const FString& FHttpClient::GetEndpointUrl(int32 Endpoint)
 	return Endpoint == INDEX_NONE ? None : Config.Endpoints[Endpoint];
 }
 
-} //namespace UE::IO::IAS
+} //namespace UE::IoStore
diff --git a/Engine/Source/Runtime/Experimental/IoStoreOnDemand/Private/OnDemandHttpClient.h b/Engine/Source/Runtime/Experimental/IoStoreOnDemand/Private/OnDemandHttpClient.h
--- a/Engine/Source/Runtime/Experimental/IoStoreOnDemand/Private/OnDemandHttpClient.h
+++ b/Engine/Source/Runtime/Experimental/IoStoreOnDemand/Private/OnDemandHttpClient.h
@@ -11,7 +11,7 @@
 #include "Templates/Function.h"
 #include "Templates/UniquePtr.h"
 
-namespace UE::IO::IAS
+namespace UE::IoStore
 {
 
 struct FHttpClientConfig
@@ -67,4 +67,4 @@ private:
 	int32 CurrentEndpoint = INDEX_NONE;
 };
 
-} // namespace UE::IO::IAS
+} // namespace UE::IoStore
diff --git a/Engine/Source/Runtime/Experimental/IoStoreOnDemand/Private/OnDemandIoDispatcherBackend.cpp b/Engine/Source/Runtime/Experimental/IoStoreOnDemand/Private/OnDemandIoDispatcherBackend.cpp
--- a/Engine/Source/Runtime/Experimental/IoStoreOnDemand/Private/OnDemandIoDispatcherBackend.cpp
+++ b/Engine/Source/Runtime/Experimental/IoStoreOnDemand/Private/OnDemandIoDispatcherBackend.cpp
@@ -55,7 +55,7 @@
 /** When enabled the IAS system can add additional debug console commands for development use */
 #define UE_IAS_DEBUG_CONSOLE_CMDS (1 && !NO_CVARS && !UE_BUILD_SHIPPING)
 
-namespace UE::IO::IAS
+namespace UE::IoStore
 {
 
 ///////////////////////////////////////////////////////////////////////////////
@@ -230,7 +230,7 @@ static void LatencyTest(FStringView Url, FStringView Path)
 	TRACE_CPUPROFILER_EVENT_SCOPE(IasBackend::LatencyTest);
 
 	int32 Results[4] = {};
-	UE::IO::IAS::HTTP::LatencyTest(Url, Path, GIasHttpTimeOutMs, MakeArrayView(Results));
+	UE::IoStore::HTTP::LatencyTest(Url, Path, GIasHttpTimeOutMs, MakeArrayView(Results));
 	UE_LOG(LogIas, Log, TEXT("Endpoint '%s' latency test (ms): %d %d %d %d"),
 		Url.GetData(), Results[0], Results[1], Results[2], Results[3]);
 }
@@ -267,7 +267,7 @@ static int32 LatencyTest(TConstArrayView<FString> Urls, FStringView Path, std::a
 	for (int32 Idx = 0; Idx < Urls.Num() && !bCancel.load(std::memory_order_relaxed); ++Idx)
 	{
 		int32 LatencyMs = -1;
-		UE::IO::IAS::HTTP::LatencyTest(Urls[Idx], Path, GIasHttpTimeOutMs, MakeArrayView(&LatencyMs, 1));
+		UE::IoStore::HTTP::LatencyTest(Urls[Idx], Path, GIasHttpTimeOutMs, MakeArrayView(&LatencyMs, 1));
 		if (LatencyMs > 0)
 		{
 			return Idx;
@@ -2597,6 +2597,6 @@ TSharedPtr<IOnDemandIoDispatcherBackend> MakeOnDemandIoDispatcherBackend(TUnique
 	return MakeShareable<IOnDemandIoDispatcherBackend>(new FOnDemandIoBackend(MoveTemp(Cache)));
 }
 
-} // namespace UE::IO::IAS
+} // namespace UE::IoStore
 
 #undef UE_IAS_DEBUG_CONSOLE_CMDS
diff --git a/Engine/Source/Runtime/Experimental/IoStoreOnDemand/Private/OnDemandIoDispatcherBackend.h b/Engine/Source/Runtime/Experimental/IoStoreOnDemand/Private/OnDemandIoDispatcherBackend.h
--- a/Engine/Source/Runtime/Experimental/IoStoreOnDemand/Private/OnDemandIoDispatcherBackend.h
+++ b/Engine/Source/Runtime/Experimental/IoStoreOnDemand/Private/OnDemandIoDispatcherBackend.h
@@ -8,7 +8,7 @@
 class IIasCache;
 struct FAnalyticsEventAttribute;
 
-namespace UE::IO::IAS
+namespace UE::IoStore
 {
 
 struct FDistributedEndpointUrl
@@ -65,4 +65,4 @@ public:
 
 TSharedPtr<IOnDemandIoDispatcherBackend> MakeOnDemandIoDispatcherBackend(TUniquePtr<IIasCache>&& Cache);
 
-} // namespace UE::IO::IAS
+} // namespace UE::IoStore
diff --git a/Engine/Source/Runtime/Experimental/IoStoreOnDemand/Private/Statistics.cpp b/Engine/Source/Runtime/Experimental/IoStoreOnDemand/Private/Statistics.cpp
--- a/Engine/Source/Runtime/Experimental/IoStoreOnDemand/Private/Statistics.cpp
+++ b/Engine/Source/Runtime/Experimental/IoStoreOnDemand/Private/Statistics.cpp
@@ -15,7 +15,7 @@ LLM_DEFINE_TAG(Ias);
 
 #if IAS_WITH_STATISTICS
 
-namespace UE::IO::IAS
+namespace UE::IoStore
 {
 
 float GIasStatisticsLogInterval = 30.f;
@@ -661,6 +661,6 @@ void FOnDemandIoBackendStats::OnHttpError()
 	GHttpErrorCount.Add(1);
 }
 
-} // namespace UE::IO::IAS
+} // namespace UE::IoStore
 
 #endif // IAS_WITH_STATISTICS
diff --git a/Engine/Source/Runtime/Experimental/IoStoreOnDemand/Private/Statistics.h b/Engine/Source/Runtime/Experimental/IoStoreOnDemand/Private/Statistics.h
--- a/Engine/Source/Runtime/Experimental/IoStoreOnDemand/Private/Statistics.h
+++ b/Engine/Source/Runtime/Experimental/IoStoreOnDemand/Private/Statistics.h
@@ -16,7 +16,7 @@ struct FAnalyticsEventAttribute;
 
 LLM_DECLARE_TAG(Ias);
 
-namespace UE::IO::IAS
+namespace UE::IoStore
 {
 
 #if IAS_WITH_STATISTICS
@@ -68,4 +68,4 @@ public:
 
 #undef IAS_STATISTICS_IMPL
 
-} // namespace UE::IO::IAS
+} // namespace UE::IoStore
diff --git a/Engine/Source/Runtime/Experimental/IoStoreOnDemand/Private/Tool/ChunkPlugin.cpp b/Engine/Source/Runtime/Experimental/IoStoreOnDemand/Private/Tool/ChunkPlugin.cpp
--- a/Engine/Source/Runtime/Experimental/IoStoreOnDemand/Private/Tool/ChunkPlugin.cpp
+++ b/Engine/Source/Runtime/Experimental/IoStoreOnDemand/Private/Tool/ChunkPlugin.cpp
@@ -18,7 +18,7 @@
 #include "Misc/StringBuilder.h"
 #include "Serialization/LargeMemoryWriter.h"
 
-namespace UE::IO::IAS::Tool
+namespace UE::IoStore::Tool
 {
 
 ////////////////////////////////////////////////////////////////////////////////
@@ -314,6 +314,6 @@ static FCommand ChunkPluginCommand(
 	}
 );
 
-} // namespace UE::IO::IAS::Tool
+} // namespace UE::IoStore::Tool
 
 #endif // UE_WITH_IAS_TOOL
diff --git a/Engine/Source/Runtime/Experimental/IoStoreOnDemand/Private/Tool/Command.cpp b/Engine/Source/Runtime/Experimental/IoStoreOnDemand/Private/Tool/Command.cpp
--- a/Engine/Source/Runtime/Experimental/IoStoreOnDemand/Private/Tool/Command.cpp
+++ b/Engine/Source/Runtime/Experimental/IoStoreOnDemand/Private/Tool/Command.cpp
@@ -15,7 +15,7 @@
 #	include "Windows/HideWindowsPlatformTypes.h"
 #endif
 
-namespace UE::IO::IAS::Tool {
+namespace UE::IoStore::Tool {
 
 ////////////////////////////////////////////////////////////////////////////////
 namespace Error
@@ -578,6 +578,6 @@ void CommandTest()
 	}
 }
 
-} // namespace UE::IO::IAS::Tool
+} // namespace UE::IoStore::Tool
 
 #endif // UE_WITH_IAS_TOOL
diff --git a/Engine/Source/Runtime/Experimental/IoStoreOnDemand/Private/Tool/Command.h b/Engine/Source/Runtime/Experimental/IoStoreOnDemand/Private/Tool/Command.h
--- a/Engine/Source/Runtime/Experimental/IoStoreOnDemand/Private/Tool/Command.h
+++ b/Engine/Source/Runtime/Experimental/IoStoreOnDemand/Private/Tool/Command.h
@@ -11,7 +11,7 @@
 
 #include <initializer_list>
 
-namespace UE::IO::IAS::Tool {
+namespace UE::IoStore::Tool {
 
 ////////////////////////////////////////////////////////////////////////////////
 struct FArgument
@@ -125,6 +125,6 @@ private:
 	EntryFunc*			Entry;
 };
 
-} // namespace UE::IO::IAS::Tool
+} // namespace UE::IoStore::Tool
 
 #endif // UE_WITH_IAS_TOOL
diff --git a/Engine/Source/Runtime/Experimental/IoStoreOnDemand/Private/Tool/Download.cpp b/Engine/Source/Runtime/Experimental/IoStoreOnDemand/Private/Tool/Download.cpp
--- a/Engine/Source/Runtime/Experimental/IoStoreOnDemand/Private/Tool/Download.cpp
+++ b/Engine/Source/Runtime/Experimental/IoStoreOnDemand/Private/Tool/Download.cpp
@@ -6,7 +6,8 @@
 
 #include <IO/IoStoreOnDemand.h>
 
-namespace UE::IO::IAS::Tool {
+namespace UE::IoStore::Tool
+{
 
 ////////////////////////////////////////////////////////////////////////////////
 static FIoStoreDownloadParams BuildDownloadParams(const FContext& Context)
@@ -65,6 +66,6 @@ static FCommand DownloadCommand(
 	}
 );
 
-} // namespace UE::IO::IAS::Tool
+} // namespace UE::IoStore::Tool
 
 #endif // UE_WITH_IAS_TOOL
diff --git a/Engine/Source/Runtime/Experimental/IoStoreOnDemand/Private/Tool/ListTocs.cpp b/Engine/Source/Runtime/Experimental/IoStoreOnDemand/Private/Tool/ListTocs.cpp
--- a/Engine/Source/Runtime/Experimental/IoStoreOnDemand/Private/Tool/ListTocs.cpp
+++ b/Engine/Source/Runtime/Experimental/IoStoreOnDemand/Private/Tool/ListTocs.cpp
@@ -6,7 +6,7 @@
 
 #include <IO/IoStoreOnDemand.h>
 
-namespace UE::IO::IAS::Tool {
+namespace UE::IoStore::Tool {
 
 ////////////////////////////////////////////////////////////////////////////////
 static FIoStoreListTocsParams BuildListTocsParams(const FContext& Context)
@@ -113,6 +113,6 @@ static FCommand ListTocsCommand(
 	}
 );
 
-} // namespace UE::IO::IAS::Tool
+} // namespace UE::IoStore::Tool
 
 #endif // UE_WITH_IAS_TOOL
diff --git a/Engine/Source/Runtime/Experimental/IoStoreOnDemand/Private/Tool/Main.cpp b/Engine/Source/Runtime/Experimental/IoStoreOnDemand/Private/Tool/Main.cpp
--- a/Engine/Source/Runtime/Experimental/IoStoreOnDemand/Private/Tool/Main.cpp
+++ b/Engine/Source/Runtime/Experimental/IoStoreOnDemand/Private/Tool/Main.cpp
@@ -4,7 +4,7 @@
 
 #include "Command.h"
 
-namespace UE::IO::IAS::Tool {
+namespace UE::IoStore::Tool {
 
 ////////////////////////////////////////////////////////////////////////////////
 FArgumentSet S3Arguments = {
@@ -24,6 +24,6 @@ int32 Main(int32 ArgC, TCHAR* ArgV[])
 	return FCommand::Main(ArgC, ArgV);
 }
 
-} // namespace UE::IO::IAS::Tool
+} // namespace UE::IoStore::Tool
 
 #endif // UE_WITH_IAS_TOOL
diff --git a/Engine/Source/Runtime/Experimental/IoStoreOnDemand/Private/Tool/Test.cpp b/Engine/Source/Runtime/Experimental/IoStoreOnDemand/Private/Tool/Test.cpp
--- a/Engine/Source/Runtime/Experimental/IoStoreOnDemand/Private/Tool/Test.cpp
+++ b/Engine/Source/Runtime/Experimental/IoStoreOnDemand/Private/Tool/Test.cpp
@@ -9,7 +9,7 @@
 #include <Containers/UnrealString.h>
 #include <Misc/CommandLine.h>
 
-namespace UE::IO::IAS {
+namespace UE::IoStore {
 
 ////////////////////////////////////////////////////////////////////////////////
 namespace IasJournaledFileCacheTest { void Tests(const TCHAR*);				}
@@ -74,6 +74,6 @@ static FCommand TestCommand(
 );
 
 } // namespace Tool
-} // namespace UE::IO::IAS
+} // namespace UE::IoStore
 
 #endif // UE_WITH_IAS_TOOL
diff --git a/Engine/Source/Runtime/Experimental/IoStoreOnDemand/Private/Tool/Upload.cpp b/Engine/Source/Runtime/Experimental/IoStoreOnDemand/Private/Tool/Upload.cpp
--- a/Engine/Source/Runtime/Experimental/IoStoreOnDemand/Private/Tool/Upload.cpp
+++ b/Engine/Source/Runtime/Experimental/IoStoreOnDemand/Private/Tool/Upload.cpp
@@ -10,7 +10,8 @@
 #include <Misc/Paths.h>
 #include <Misc/PathViews.h>
 
-namespace UE::IO::IAS::Tool {
+namespace UE::IoStore::Tool
+{
 
 ////////////////////////////////////////////////////////////////////////////////
 static FKeyChain LoadKeys(const FContext& Context)
@@ -150,6 +151,6 @@ static FCommand UploadCommand(
 	}
 );
 
-} // namespace UE::IO::IAS::Tool
+} // namespace UE::IoStore::Tool
 
 #endif // UE_WITH_IAS_TOOL
diff --git a/Engine/Source/Runtime/Experimental/IoStoreOnDemand/Public/IO/IoStoreOnDemand.h b/Engine/Source/Runtime/Experimental/IoStoreOnDemand/Public/IO/IoStoreOnDemand.h
--- a/Engine/Source/Runtime/Experimental/IoStoreOnDemand/Public/IO/IoStoreOnDemand.h
+++ b/Engine/Source/Runtime/Experimental/IoStoreOnDemand/Public/IO/IoStoreOnDemand.h
@@ -27,7 +27,7 @@ class IIoStoreWriter;
 struct FAnalyticsEventAttribute;
 struct FIoContainerSettings;
 struct FIoStoreWriterSettings;
-namespace UE::IO::IAS { struct FOnDemandEndpoint; }
+namespace UE::IoStore { struct FOnDemandEndpoint; }
 using FIoBlockHash = uint32;
 
 // Custom initialization allows users to control when
@@ -38,7 +38,7 @@ using FIoBlockHash = uint32;
 
 UE_API DECLARE_LOG_CATEGORY_EXTERN(LogIas, Log, All);
 
-namespace UE::IO::IAS
+namespace UE::IoStore
 {
 
 ////////////////////////////////////////////////////////////////////////////////
@@ -369,6 +369,6 @@ public:
 #endif //UE_IAS_CUSTOM_INITIALIZATION
 };
 
-} // namespace UE::IO::IAS
+} // namespace UE::IoStore
 
 #undef UE_API
