Reduce memory allocations during FReplicationReader::UpdateObjectReferenceTracking().

diff --git a/Engine/Source/Runtime/Experimental/Iris/Core/Private/Iris/ReplicationSystem/ReplicationReader.cpp b/Engine/Source/Runtime/Experimental/Iris/Core/Private/Iris/ReplicationSystem/ReplicationReader.cpp
--- a/Engine/Source/Runtime/Experimental/Iris/Core/Private/Iris/ReplicationSystem/ReplicationReader.cpp
+++ b/Engine/Source/Runtime/Experimental/Iris/Core/Private/Iris/ReplicationSystem/ReplicationReader.cpp
@@ -60,6 +60,15 @@ static FAutoConsoleVariableRef CVarColdResolvingRetryTimeMS(
 	ColdResolvingRetryTimeMS,
 	TEXT("Resolve unresolved cold references after this many milliseconds."));
 
+// When true, the function UpdateObjectReferenceTracking_Fast() is used and UpdateObjectReferenceTracking()
+// otherwise. Once UpdateObjectReferenceTracking_Fast() has been tested sufficiently it will become the only
+// version of this function and the cvar removed.
+static bool bUseOptObjectRefTracking = false;
+static FAutoConsoleVariableRef CVarUseOptObjectRefTracking(
+	TEXT("net.Iris.UseOptObjectRefTracking"),
+	bUseOptObjectRefTracking,
+	TEXT("Use a more optimized version of FReplicationReader::UpdateObjectReferenceTracking()."));
+
 static bool bExecuteReliableRPCsBeforeApplyState = true;
 static FAutoConsoleVariableRef CVarExecuteReliableRPCsBeforeApplyState(
 		TEXT("net.Iris.ExecuteReliableRPCsBeforeApplyState"),
@@ -204,6 +213,44 @@ FReplicationReader::FReplicatedObjectInfo::FReplicatedObjectInfo()
 	PrevStoredBaselineIndex = FDeltaCompressionBaselineManager::InvalidBaselineIndex;
 }
 
+bool FReplicationReader::FReplicatedObjectInfo::RemoveUnresolvedHandleCount(FNetRefHandle RefHandle)
+{
+	int16* HandleCount = UnresolvedHandleCount.Find(RefHandle);
+	if (ensureMsgf(HandleCount != nullptr, TEXT("Unresolved handle counter could not be found for %s"), ToCStr(RefHandle.ToString())))
+	{
+		ensure(*HandleCount > 0);
+
+		(*HandleCount)--;
+
+		if (*HandleCount <= 0)
+		{
+			UnresolvedHandleCount.Remove(RefHandle);
+			return true;
+		}
+	}
+
+	return false;
+}
+
+bool FReplicationReader::FReplicatedObjectInfo::RemoveResolvedDynamicHandleCount(FNetRefHandle RefHandle)
+{
+	int16* HandleCount = ResolvedDynamicHandleCount.Find(RefHandle);
+	if (ensureMsgf(HandleCount != nullptr, TEXT("Resolved dynamic handle counter could not be found for% s"), ToCStr(RefHandle.ToString())))
+	{
+		ensure(*HandleCount > 0);
+
+		(*HandleCount)--;
+
+		if (*HandleCount <= 0)
+		{
+			ResolvedDynamicHandleCount.Remove(RefHandle);
+			return true;
+		}
+	}
+
+	return false;
+}
+
 FReplicationReader::FReplicationReader()
 : TempLinearAllocator()
 , TempChangeMaskAllocator(&TempLinearAllocator)
@@ -1200,6 +1247,156 @@ void FReplicationReader::UpdateObjectReferenceTracking(FReplicatedObjectInfo* Re
 	}
 }
 
+void FReplicationReader::UpdateObjectReferenceTracking_Fast(FReplicatedObjectInfo* ReplicationInfo, FNetBitArrayView ChangeMask, bool bIncludeInitState, FResolvedNetRefHandlesArray& OutNewResolvedRefHandles, const FObjectReferenceTracker& NewUnresolvedReferences, const FObjectReferenceTracker& NewMappedDynamicReferences)
+{
+	IRIS_PROFILER_SCOPE(FReplicationReader_UpdateObjectReferenceTracking);
+
+	/*
+	 * As we store references per changemask we need to construct a set of all unresolved references and
+	 * compare with the new set of unresolved references. The new set is found by first updating the 
+	 * references that were found in the changemask.
+	 */
+	{
+		// Try to avoid dynamic allocations during the update of the UnresolvedObjectReferences.
+		ReplicationInfo->UnresolvedObjectReferences.Reserve(ReplicationInfo->UnresolvedObjectReferences.Num() + NewUnresolvedReferences.Num());
+		ReplicationInfo->UnresolvedHandleCount.Reserve(ReplicationInfo->UnresolvedObjectReferences.Num() + NewUnresolvedReferences.Num());
+
+		// Replace each entry in UnresolvedObjectReferences for the given changemask
+		auto UpdateUnresolvedReferencesForChange = [ReplicationInfo, &NewUnresolvedReferences, &OutNewResolvedRefHandles, this](uint32 ChangeBit)
+		{
+			FObjectReferenceTracker& UnresolvedObjectReferences = ReplicationInfo->UnresolvedObjectReferences;
+
+			bool bUnresolvedHandleCountShrink = false;
+
+			for (FObjectReferenceTracker::TKeyIterator It = UnresolvedObjectReferences.CreateKeyIterator(ChangeBit); It; ++It)
+			{
+				const FNetRefHandle RefHandle = It.Value();
+
+				const bool bInNewUnresolved = (NewUnresolvedReferences.FindPair(ChangeBit, RefHandle) != nullptr);
+				if (!bInNewUnresolved)
+				{
+					It.RemoveCurrent();
+
+					if (ReplicationInfo->RemoveUnresolvedHandleCount(RefHandle))
+					{
+						bUnresolvedHandleCountShrink = true;
+							
+						// Store new resolved handles so we can update partially resolved references properly
+						OutNewResolvedRefHandles.Add(RefHandle);
+
+						// Remove from tracking
+						const uint32 OwnerInternalIndex = ReplicationInfo->InternalIndex;
+						UnresolvedHandleToDependents.Remove(RefHandle, OwnerInternalIndex);
+						RemoveFromUnresolvedCache(RefHandle);
+						UE_LOG(LogIris, Verbose, TEXT("FReplicationReader::UpdateObjectReferenceTracking Removing unresolved reference %s for %s (OwnerInternalIndex=%d)"), ToCStr(RefHandle.ToString()), ToCStr(NetRefHandleManager->GetNetRefHandleFromInternalIndex(OwnerInternalIndex).ToString()), OwnerInternalIndex);
+					}
+				}
+			}
+
+			for (FObjectReferenceTracker::TConstKeyIterator It = NewUnresolvedReferences.CreateConstKeyIterator(ChangeBit); It; ++It)
+			{
+				const FNetRefHandle RefHandle = It.Value();
+				const bool bInCurrUnresolved = (UnresolvedObjectReferences.FindPair(ChangeBit, RefHandle) != nullptr);
+
+				if (!bInCurrUnresolved)
+				{
+					UnresolvedObjectReferences.Add(ChangeBit, RefHandle);
+
+					int16& HandleCount = ReplicationInfo->UnresolvedHandleCount.FindOrAdd(RefHandle, 0);
+
+					HandleCount++;
+
+					// Add to tracking
+					const uint32 OwnerInternalIndex = ReplicationInfo->InternalIndex;
+					UnresolvedHandleToDependents.Add(RefHandle, OwnerInternalIndex);
+					UE_LOG(LogIris, Verbose, TEXT("FReplicationReader::UpdateObjectReferenceTracking Adding unresolved reference %s for %s (OwnerInternalIndex=%d)"), ToCStr(RefHandle.ToString()), ToCStr(NetRefHandleManager->GetNetRefHandleFromInternalIndex(OwnerInternalIndex).ToString()), OwnerInternalIndex);
+				}
+			}
+
+			if (bUnresolvedHandleCountShrink)
+			{
+				ReplicationInfo->UnresolvedHandleCount.Shrink();
+			}
+		};
+
+		ChangeMask.ForAllSetBits(UpdateUnresolvedReferencesForChange);
+		if (bIncludeInitState)
+		{
+			UpdateUnresolvedReferencesForChange(FakeInitChangeMaskOffset);
+		}
+
+		// Update ReplicationInfo with the status of unresolved references
+		ReplicationInfo->bHasUnresolvedReferences = ReplicationInfo->UnresolvedHandleCount.Num() > 0;
+		ReplicationInfo->bHasUnresolvedInitialReferences = ReplicationInfo->UnresolvedObjectReferences.Find(FakeInitChangeMaskOffset) != nullptr;
+	}
+
+	// Update tracking for resolved dynamic references
+	if (bRemapDynamicObjects)
+	{
+		// Try to avoid dynamic allocations during the update of the ResolvedDynamicObjectReferences.
+		ReplicationInfo->ResolvedDynamicObjectReferences.Reserve(ReplicationInfo->ResolvedDynamicObjectReferences.Num() + NewMappedDynamicReferences.Num());
+		ReplicationInfo->ResolvedDynamicHandleCount.Reserve(ReplicationInfo->ResolvedDynamicObjectReferences.Num() + NewMappedDynamicReferences.Num());
+
+		// Replace each entry in ResolvedDynamicObjectReferences for the given changemask
+		auto UpdateResolvedReferencesForChange = [ReplicationInfo, &NewMappedDynamicReferences, this](uint32 ChangeBit)
+		{
+			FObjectReferenceTracker& ResolvedDynamicObjectReferences = ReplicationInfo->ResolvedDynamicObjectReferences;
+
+			bool bResolvedDynamicObjectReferenceShrink = false;
+
+			for (FObjectReferenceTracker::TKeyIterator It = ResolvedDynamicObjectReferences.CreateKeyIterator(ChangeBit); It; ++It)
+			{
+				const FNetRefHandle RefHandle = It.Value();
+
+				const bool bInNewResolved = (NewMappedDynamicReferences.FindPair(ChangeBit, RefHandle) != nullptr);
+				if (!bInNewResolved)
+				{
+					It.RemoveCurrent();
+
+					if (ReplicationInfo->RemoveResolvedDynamicHandleCount(RefHandle))
+					{
+						bResolvedDynamicObjectReferenceShrink = true;
+
+						// Remove from tracking
+						const uint32 OwnerInternalIndex = ReplicationInfo->InternalIndex;
+						ResolvedDynamicHandleToDependents.Remove(RefHandle, OwnerInternalIndex);
+						UE_LOG(LogIris, Verbose, TEXT("FReplicationReader::UpdateObjectReferenceTracking Removing resolved dynamic reference %s for %s"), ToCStr(RefHandle.ToString()), ToCStr(NetRefHandleManager->GetNetRefHandleFromInternalIndex(OwnerInternalIndex).ToString()));
+					}
+				}
+			}
+
+			for (FObjectReferenceTracker::TConstKeyIterator It = NewMappedDynamicReferences.CreateConstKeyIterator(ChangeBit); It; ++It)
+			{
+				const FNetRefHandle RefHandle = It.Value();
+				const bool bInCurrResolved = (ResolvedDynamicObjectReferences.FindPair(ChangeBit, RefHandle) != nullptr);
+
+				if (!bInCurrResolved)
+				{
+					ResolvedDynamicObjectReferences.Add(ChangeBit, RefHandle);
+
+					int16& HandleCount = ReplicationInfo->ResolvedDynamicHandleCount.FindOrAdd(RefHandle, 0);
+
+					HandleCount++;
+
+					// Add to tracking
+					const uint32 OwnerInternalIndex = ReplicationInfo->InternalIndex;
+					ResolvedDynamicHandleToDependents.Add(RefHandle, OwnerInternalIndex);
+					UE_LOG(LogIris, Verbose, TEXT("FReplicationReader::UpdateObjectReferenceTracking Adding resolved dynamic reference %s for %s"), ToCStr(RefHandle.ToString()), ToCStr(NetRefHandleManager->GetNetRefHandleFromInternalIndex(OwnerInternalIndex).ToString()));
+				}
+			}
+
+			if (bResolvedDynamicObjectReferenceShrink)
+			{
+				ReplicationInfo->ResolvedDynamicHandleCount.Shrink();
+			}
+		};
+
+		ChangeMask.ForAllSetBits(UpdateResolvedReferencesForChange);
+		// Intentionally leaving out init state. It seems weird to call rep notifies and update init only properties after
+		// the initial state has already been applied.
+	}
+}
+
 void FReplicationReader::RemoveUnresolvedObjectReferenceInReplicationInfo(FReplicatedObjectInfo* ReplicationInfo, FNetRefHandle Handle)
 {
 	for (FObjectReferenceTracker::TIterator It = ReplicationInfo->UnresolvedObjectReferences.CreateIterator(); It; ++It)
@@ -1231,6 +1428,8 @@ bool FReplicationReader::MoveResolvedObjectReferenceToUnresolvedInReplicationInf
 	bool bHasUnresolvedInitialReferences = ReplicationInfo->bHasUnresolvedInitialReferences;
 	FNetBitArrayView UnresolvedChangeMask = FChangeMaskUtil::MakeChangeMask(ReplicationInfo->UnresolvedChangeMaskOrPointer, ReplicationInfo->ChangeMaskBitCount);
 	FObjectReferenceTracker& UnresolvedObjectReferences = ReplicationInfo->UnresolvedObjectReferences;
+	TMap<FNetRefHandle, int16>& UnresolvedHandleCount = ReplicationInfo->UnresolvedHandleCount;
+	TMap<FNetRefHandle, int16>& ResolvedDynamicHandleCount = ReplicationInfo->ResolvedDynamicHandleCount;
 	for (FObjectReferenceTracker::TIterator It = ReplicationInfo->ResolvedDynamicObjectReferences.CreateIterator(); It; ++It)
 	{
 		const FNetRefHandle RefHandle = It->Value;
@@ -1252,9 +1451,20 @@ bool FReplicationReader::MoveResolvedObjectReferenceToUnresolvedInReplicationInf
 			// At this point we'd like to skip iteration to the next key as a handle can only be found once per changemask.
 			It.RemoveCurrent();
 
+			if (bUseOptObjectRefTracking)
+			{
+				ReplicationInfo->RemoveResolvedDynamicHandleCount(UnresolvableHandle);
+			}
+
 			// This handle should only have existed once in the ResolvedDynamicObjectReferences map and should not be able to
 			// already exist in the UnresolvedObjectReferences map, so no need to call AddUnique.
 			UnresolvedObjectReferences.Add(ChangemaskOffset, UnresolvableHandle);
+			
+			if (bUseOptObjectRefTracking)
+			{
+				int16& UnresolvableHandleCount = UnresolvedHandleCount.FindOrAdd(UnresolvableHandle, 0);
+				UnresolvableHandleCount++;
+			}
 
 			UE_LOG(LogIris, Verbose, TEXT("FReplicationReader::MoveResolvedObjectReferenceToUnresolvedInReplicationInfo Moving from resolved to unresolved reference %s for %s"), ToCStr(UnresolvableHandle.ToString()), ToCStr(NetRefHandleManager->GetNetRefHandleFromInternalIndex(ReplicationInfo->InternalIndex).ToString()));
 		}
@@ -1301,7 +1511,9 @@ void FReplicationReader::BuildUnresolvedChangeMaskAndUpdateObjectReferenceTracki
 	OutUnresolvedChangeMask.Reset();
 	bool bHasUnresolvedInitReferences = false;
 
-	FObjectReferenceTracker UnresolvedReferences;
+	UnresolvedReferencesCache.Reset();
+	MappedDynamicReferencesCache.Reset();
+	
 	for (const auto& RefInfo : Collector.GetUnresolvedReferences())
 	{
 		const FNetSerializerChangeMaskParam& ChangeMaskInfo = RefInfo.ChangeMaskInfo;
@@ -1315,21 +1527,27 @@ void FReplicationReader::BuildUnresolvedChangeMaskAndUpdateObjectReferenceTracki
 		}
 
 		const uint32 BitOffset = (ChangeMaskInfo.BitCount > 0U ? ChangeMaskInfo.BitOffset : FakeInitChangeMaskOffset);
-		UnresolvedReferences.AddUnique(BitOffset, RefInfo.Reference.GetRefHandle());
+		UnresolvedReferencesCache.AddUnique(BitOffset, RefInfo.Reference.GetRefHandle());
 	}
 
-	FObjectReferenceTracker MappedDynamicReferences;
 	for (const auto& RefInfo : Collector.GetResolvedReferences())
 	{
 		if (RefInfo.Reference.GetRefHandle().IsDynamic())
 		{
 			const uint32 BitOffset = (RefInfo.ChangeMaskInfo.BitCount > 0U ? RefInfo.ChangeMaskInfo.BitOffset : FakeInitChangeMaskOffset);
-			MappedDynamicReferences.AddUnique(BitOffset, RefInfo.Reference.GetRefHandle());
+			MappedDynamicReferencesCache.AddUnique(BitOffset, RefInfo.Reference.GetRefHandle());
 		}
 	}
 
 	// Update object specific
-	UpdateObjectReferenceTracking(ReplicationInfo, CollectorChangeMask, Collector.IsInitStateIncluded(), OutNewResolvedRefHandles, UnresolvedReferences, MappedDynamicReferences);
+	if (bUseOptObjectRefTracking)
+	{
+		UpdateObjectReferenceTracking_Fast(ReplicationInfo, CollectorChangeMask, Collector.IsInitStateIncluded(), OutNewResolvedRefHandles, UnresolvedReferencesCache, MappedDynamicReferencesCache);
+	}
+	else
+	{
+		UpdateObjectReferenceTracking(ReplicationInfo, CollectorChangeMask, Collector.IsInitStateIncluded(), OutNewResolvedRefHandles, UnresolvedReferencesCache, MappedDynamicReferencesCache);
+	}
 }
 
 void FReplicationReader::ResolveAndDispatchUnresolvedReferencesForObject(FNetSerializationContext& Context, uint32 InternalIndex)
diff --git a/Engine/Source/Runtime/Experimental/Iris/Core/Private/Iris/ReplicationSystem/ReplicationReader.h b/Engine/Source/Runtime/Experimental/Iris/Core/Private/Iris/ReplicationSystem/ReplicationReader.h
--- a/Engine/Source/Runtime/Experimental/Iris/Core/Private/Iris/ReplicationSystem/ReplicationReader.h
+++ b/Engine/Source/Runtime/Experimental/Iris/Core/Private/Iris/ReplicationSystem/ReplicationReader.h
@@ -82,6 +82,11 @@ private:
 		FObjectReferenceTracker UnresolvedObjectReferences;
 		FObjectReferenceTracker ResolvedDynamicObjectReferences;
 
+		// These maps provide a O(1) lookup for the number of handles referenced in 
+		// UnresolvedObjectReferences and ResolvedDynamicObjectReferences.
+		TMap<FNetRefHandle, int16> UnresolvedHandleCount;
+		TMap<FNetRefHandle, int16> ResolvedDynamicHandleCount;
+
 		// Baselines
 		uint8* StoredBaselines[2];
 
@@ -103,6 +108,9 @@ private:
 				uint32 Padding : 7;
 			};
 		};
+
+		bool RemoveUnresolvedHandleCount(FNetRefHandle RefHandle);
+		bool RemoveResolvedDynamicHandleCount(FNetRefHandle RefHandle);
 	};
 
 	// Temporary Data to dispatch
@@ -181,6 +189,9 @@ private:
 
 	// Update reference tracking maps for the current object
 	void UpdateObjectReferenceTracking(FReplicatedObjectInfo* ReplicationInfo, FNetBitArrayView ChangeMask, bool bIncludeInitState, FResolvedNetRefHandlesArray& OutNewResolvedRefHandles, const FObjectReferenceTracker& NewUnresolvedReferences, const FObjectReferenceTracker& NewMappedDynamicReferences);
+	
+	// An optimized version of UpdateObjectReferenceTracking().
+	void UpdateObjectReferenceTracking_Fast(FReplicatedObjectInfo* ReplicationInfo, FNetBitArrayView ChangeMask, bool bIncludeInitState, FResolvedNetRefHandlesArray& OutNewResolvedRefHandles, const FObjectReferenceTracker& NewUnresolvedReferences, const FObjectReferenceTracker& NewMappedDynamicReferences);
 
 	// Remove all references for object
 	void CleanupReferenceTracking(FReplicatedObjectInfo* ObjectInfo);
@@ -260,6 +271,11 @@ private:
 	// Used during receive and processing of pending batches
 	TArray<FNetRefHandle> TempMustBeMappedReferences;
 
+	// Preallocate the arrays used by BuildUnresolvedChangeMaskAndUpdateObjectReferenceTracking() to 
+	// avoid memory allocations during the frame.
+	FObjectReferenceTracker UnresolvedReferencesCache;
+	FObjectReferenceTracker MappedDynamicReferencesCache;
+
 	FNetBlobHandlerManager* NetBlobHandlerManager;
 	FNetBlobType NetObjectBlobType;
 	FNetObjectAttachmentsReader Attachments;
