Chaos - fix FGeometryCollectionPhysicsProxy::ScaleClusterGeometry_Internal crash

We were changing the geometry on a particle without notifying the evolution. We need to invalidate collisions before changing geometry.

diff --git a/Engine/Source/Runtime/Experimental/Chaos/Private/PhysicsProxy/GeometryCollectionPhysicsProxy.cpp b/Engine/Source/Runtime/Experimental/Chaos/Private/PhysicsProxy/GeometryCollectionPhysicsProxy.cpp
--- a/Engine/Source/Runtime/Experimental/Chaos/Private/PhysicsProxy/GeometryCollectionPhysicsProxy.cpp
+++ b/Engine/Source/Runtime/Experimental/Chaos/Private/PhysicsProxy/GeometryCollectionPhysicsProxy.cpp
@@ -3247,10 +3247,18 @@ void FGeometryCollectionPhysicsProxy::SetWorldTransform_External(const FTransfor
 
 void FGeometryCollectionPhysicsProxy::ScaleClusterGeometry_Internal(const FVector& WorldScale)
 {
+	Chaos::FPhysicsSolver* RigidSolver = GetSolver<Chaos::FPhysicsSolver>();
+
 	for (int32 ParticleIndex = 0; ParticleIndex < NumEffectiveParticles; ++ParticleIndex)
 	{
 		if (Chaos::FPBDRigidClusteredParticleHandle* ParticleHandle = SolverParticleHandles[ParticleIndex])
 		{
+			if (RigidSolver != nullptr)
+			{
+				// We are about to change the geometry, so must reset any systems that depend on it (e.g., collisions)
+				RigidSolver->GetEvolution()->InvalidateParticle(ParticleHandle);
+			}
+
 			// Scale the geometry if necessary
 			BuildScaledGeometry(ParticleHandle, ParticleHandle->GetGeometry(), WorldScale);
 
