Fix regression from 32594507
- Fixed the headless chaos failure
- Also still properly handle UE-197548 case

diff --git a/Engine/Source/Runtime/Experimental/Chaos/Private/Chaos/PBDRigidClustering.cpp b/Engine/Source/Runtime/Experimental/Chaos/Private/Chaos/PBDRigidClustering.cpp
--- a/Engine/Source/Runtime/Experimental/Chaos/Private/Chaos/PBDRigidClustering.cpp
+++ b/Engine/Source/Runtime/Experimental/Chaos/Private/Chaos/PBDRigidClustering.cpp
@@ -1510,6 +1510,11 @@ namespace Chaos
 							// just got re-enabled in RemoveParticlesFromCluster.
 							for (FPBDRigidParticleHandle* Particle : Island)
 							{
+								if (FPBDRigidClusteredParticleHandle* ClusterParticle = Particle->CastToClustered())
+								{
+									TopLevelClusterParentsStrained.Remove(Particle->CastToClustered());
+									TopLevelClusterParents.Remove(Particle->CastToClustered());
+								}
 								MEvolution.DisableParticle(Particle);
 							}
 						}
