Soft constraints don't preview correctly in PhAT and in PIE

diff --git a/Engine/Source/Runtime/Experimental/Chaos/Private/Chaos/Joint/PBDJointCachedSolverGaussSeidel.cpp b/Engine/Source/Runtime/Experimental/Chaos/Private/Chaos/Joint/PBDJointCachedSolverGaussSeidel.cpp
--- a/Engine/Source/Runtime/Experimental/Chaos/Private/Chaos/Joint/PBDJointCachedSolverGaussSeidel.cpp
+++ b/Engine/Source/Runtime/Experimental/Chaos/Private/Chaos/Joint/PBDJointCachedSolverGaussSeidel.cpp
@@ -1415,7 +1415,7 @@ void FPBDJointCachedSolver::ApplyRotationProjection(
 
 	for(int32 ConstraintIndex = 0; ConstraintIndex < 3; ++ConstraintIndex)
 	{
-		if(RotationConstraints.bValidDatas[ConstraintIndex])
+		if(RotationConstraints.bValidDatas[ConstraintIndex] && !RotationConstraints.bSoftLimit[ConstraintIndex])
 		{
 			FReal DeltaAngle = RotationConstraints.ConstraintCX[ConstraintIndex] +
 				FVec3::DotProduct(Body(1).DQ() - Body(0).DQ(), RotationConstraints.ConstraintAxis[ConstraintIndex]);
@@ -1467,7 +1467,7 @@ void FPBDJointCachedSolver::ApplyPositionProjection(
 
 	for(int32 ConstraintIndex = 0; ConstraintIndex < 3; ++ConstraintIndex)
 	{
-		if(PositionConstraints.bValidDatas[ConstraintIndex])
+		if(PositionConstraints.bValidDatas[ConstraintIndex] && !PositionConstraints.bSoftLimit[ConstraintIndex])
 		{
 			const FVec3 CX = Body(1).DP()  - Body(0).DP() +
 				FVec3::CrossProduct(Body(1).DQ(), PositionConstraints.ConstraintArms[ConstraintIndex][1]) -
@@ -1529,7 +1529,7 @@ void FPBDJointCachedSolver::ApplyPositionTeleport(
 
 	for (int32 ConstraintIndex = 0; ConstraintIndex < 3; ++ConstraintIndex)
 	{
-		if (PositionConstraints.bValidDatas[ConstraintIndex])
+		if (PositionConstraints.bValidDatas[ConstraintIndex] && !PositionConstraints.bSoftLimit[ConstraintIndex])
 		{
 			FReal DeltaPosition = PositionConstraints.ConstraintCX[ConstraintIndex];
 
