Fix for static analysis error.

diff --git a/Engine/Source/Runtime/Experimental/Iris/Core/Private/Iris/ReplicationSystem/ReplicationReader.cpp b/Engine/Source/Runtime/Experimental/Iris/Core/Private/Iris/ReplicationSystem/ReplicationReader.cpp
--- a/Engine/Source/Runtime/Experimental/Iris/Core/Private/Iris/ReplicationSystem/ReplicationReader.cpp
+++ b/Engine/Source/Runtime/Experimental/Iris/Core/Private/Iris/ReplicationSystem/ReplicationReader.cpp
@@ -218,11 +218,13 @@ bool FReplicationReader::FReplicatedObjectInfo::RemoveUnresolvedHandleCount(FNet
 	int16* HandleCount = UnresolvedHandleCount.Find(RefHandle);
 	if (ensureMsgf(HandleCount != nullptr, TEXT("Unresolved handle counter could not be found for %s"), ToCStr(RefHandle.ToString())))
 	{
-		ensure(*HandleCount > 0);
+		int16& HandleCountRef = (*HandleCount);
 
-		(*HandleCount)--;
+		ensure(HandleCountRef > 0);
 
-		if (*HandleCount <= 0)
+		HandleCountRef--;
+
+		if (HandleCountRef <= 0)
 		{
 			UnresolvedHandleCount.Remove(RefHandle);
 			return true;
@@ -237,11 +239,13 @@ bool FReplicationReader::FReplicatedObjectInfo::RemoveResolvedDynamicHandleCount
 	int16* HandleCount = ResolvedDynamicHandleCount.Find(RefHandle);
 	if (ensureMsgf(HandleCount != nullptr, TEXT("Resolved dynamic handle counter could not be found for% s"), ToCStr(RefHandle.ToString())))
 	{
-		ensure(*HandleCount > 0);
+		int16& HandleCountRef = (*HandleCount);
+
+		ensure(HandleCountRef > 0);
 
-		(*HandleCount)--;
+		HandleCountRef--;
 
-		if (*HandleCount <= 0)
+		if (HandleCountRef <= 0)
 		{
 			ResolvedDynamicHandleCount.Remove(RefHandle);
 			return true;
