Removed c++ 20 usage from FastArrayNetSerializer headers

diff --git a/Engine/Source/Runtime/Experimental/Iris/Core/Public/Iris/ReplicationSystem/Private/FastArrayReplicationFragmentInternal.h b/Engine/Source/Runtime/Experimental/Iris/Core/Public/Iris/ReplicationSystem/Private/FastArrayReplicationFragmentInternal.h
--- a/Engine/Source/Runtime/Experimental/Iris/Core/Public/Iris/ReplicationSystem/Private/FastArrayReplicationFragmentInternal.h
+++ b/Engine/Source/Runtime/Experimental/Iris/Core/Public/Iris/ReplicationSystem/Private/FastArrayReplicationFragmentInternal.h
@@ -87,11 +87,7 @@ struct FFastArrayReplicationFragmentHelper
 	static inline typename TEnableIf<TModels_V<FFastArraySerializer::CPostReplicatedReceiveFuncable, FastArrayType, const FFastArraySerializer::FPostReplicatedReceiveParameters&>, void>::Type CallPostReplicatedReceiveOrNot(FastArrayType& ArraySerializer, int32 OldArraySize, bool bHasUnresolvedReferences)
 	{
 PRAGMA_DISABLE_DEPRECATION_WARNINGS
-		FFastArraySerializer::FPostReplicatedReceiveParameters PostReceivedParameters =
-		{
-			.OldArraySize = OldArraySize,
-			.bHasMoreUnmappedReferences = bHasUnresolvedReferences,
-		};
+		FFastArraySerializer::FPostReplicatedReceiveParameters PostReceivedParameters = { OldArraySize, bHasUnresolvedReferences };
 PRAGMA_ENABLE_DEPRECATION_WARNINGS
 		
 		ArraySerializer.PostReplicatedReceive(PostReceivedParameters);
diff --git a/Engine/Source/Runtime/Net/Core/Classes/Net/Serialization/FastArraySerializer.h b/Engine/Source/Runtime/Net/Core/Classes/Net/Serialization/FastArraySerializer.h
--- a/Engine/Source/Runtime/Net/Core/Classes/Net/Serialization/FastArraySerializer.h
+++ b/Engine/Source/Runtime/Net/Core/Classes/Net/Serialization/FastArraySerializer.h
@@ -703,11 +703,7 @@ private:
 		inline typename TEnableIf<TModels_V<CPostReplicatedReceiveFuncable, FastArrayType, const FFastArraySerializer::FPostReplicatedReceiveParameters>, void>::Type CallPostReplicatedReceiveOrNot(int32 OldArraySize)
 		{
 PRAGMA_DISABLE_DEPRECATION_WARNINGS
-			FFastArraySerializer::FPostReplicatedReceiveParameters PostReceivedParameters =
-			{
-				.OldArraySize = OldArraySize,
-				.bHasMoreUnmappedReferences = Parms.bOutHasMoreUnmapped,
-			};
+			FFastArraySerializer::FPostReplicatedReceiveParameters PostReceivedParameters = { OldArraySize, Parms.bOutHasMoreUnmapped };
 			ArraySerializer.PostReplicatedReceive(PostReceivedParameters);
 PRAGMA_ENABLE_DEPRECATION_WARNINGS
 		}
