Don't recreate geometry collection custom renderer object unnecessarily at registration.
This allows the custom renderer to cache state across register/unregister teardown.

diff --git a/Engine/Source/Runtime/Experimental/GeometryCollectionEngine/Private/GeometryCollection/GeometryCollectionComponent.cpp b/Engine/Source/Runtime/Experimental/GeometryCollectionEngine/Private/GeometryCollection/GeometryCollectionComponent.cpp
--- a/Engine/Source/Runtime/Experimental/GeometryCollectionEngine/Private/GeometryCollection/GeometryCollectionComponent.cpp
+++ b/Engine/Source/Runtime/Experimental/GeometryCollectionEngine/Private/GeometryCollection/GeometryCollectionComponent.cpp
@@ -3560,9 +3560,16 @@ void UGeometryCollectionComponent::OnRegister()
 	UClass* Type = bOverrideCustomRenderer ? CustomRendererType : (RestCollection ? RestCollection->CustomRendererType : nullptr);
 	if (Type && Type->ImplementsInterface(UGeometryCollectionExternalRenderInterface::StaticClass()))
 	{
-		CustomRenderer = NewObject<UObject>(this, Type);
+		if (!CustomRenderer || CustomRenderer.GetObject()->GetClass() != Type)
+		{
+			CustomRenderer = NewObject<UObject>(this, Type);
+		}
 		RegisterCustomRenderer();
 	}
+	else
+	{
+		CustomRenderer = nullptr;
+	}
 
 	Super::OnRegister();
 }
@@ -3571,9 +3578,7 @@ void UGeometryCollectionComponent::OnUnregister()
 {
 	Super::OnUnregister();
 
-	// Remove any custom renderer.
 	UnregisterCustomRenderer();
-	CustomRenderer = nullptr;
 }
 
 void UGeometryCollectionComponent::RegisterCustomRenderer()
