cluster union intercluster edges are now *only* generated lazily upon release due to strain

- move intercluster edge generation into FClusterUnionManager::GenerateInterclusterEdgesForParticle from GenerateEdges in PBDRigidClustering and only do this on strained release. I'm fairly sure this removes the need for the FClusterUnionParticleProperties::bEdgesAreGenerated boolean but leaving that in place for now.

diff --git a/Engine/Source/Runtime/Experimental/Chaos/Private/Chaos/ClusterUnionManager.cpp b/Engine/Source/Runtime/Experimental/Chaos/Private/Chaos/ClusterUnionManager.cpp
--- a/Engine/Source/Runtime/Experimental/Chaos/Private/Chaos/ClusterUnionManager.cpp
+++ b/Engine/Source/Runtime/Experimental/Chaos/Private/Chaos/ClusterUnionManager.cpp
@@ -156,47 +156,13 @@ namespace Chaos
 						{
 							FPBDRigidParticleHandle* OtherParticle = ClusterUnion.ChildParticles[RootObjectIndex];
 							Clustering.CreateNodeConnection(Particle, OtherParticle);
-							
-							if (FClusterUnionParticleProperties* Properties = ClusterUnion.ChildProperties.Find(Particle))
-							{
-								Properties->bEdgesAreGenerated = false;
-								if (FClusterUnionParticleProperties* OtherProperties = ClusterUnion.ChildProperties.Find(OtherParticle))
-								{
-									OtherProperties->bEdgesAreGenerated = false;
-								}
-
-								if (bChaosClusterUnionGenerateInterclusterEdges && Particle && OtherParticle)
-								{
-									// Only generate intercluster edges for main particles. Auxiliary particles that are just bits and pieces of geometry collections
-									// shouldn't also generate intercluster edges.
-									if (!Properties || !Properties->bIsAuxiliaryParticle)
-									{
-										const TArray<FPBDRigidParticleHandle*>& ParticleChildren = Clustering.GetChildrenMap().FindRef(Particle->CastToClustered());
-										const TArray<FPBDRigidParticleHandle*>& OtherChildren = Clustering.GetChildrenMap().FindRef(OtherParticle->CastToClustered());
-
-										TSet<FPBDRigidParticleHandle*> FromSet{ ParticleChildren };
-										TSet<FPBDRigidParticleHandle*> ToSet{ OtherChildren };
-
-										TArray<FPBDRigidParticleHandle*> AllParticles;
-										AllParticles.Reserve(FromSet.Num() + ToSet.Num());
-										AllParticles.Append(ParticleChildren);
-										AllParticles.Append(OtherChildren);
-
-										FClusterCreationParameters Parameters{ 0.3f, 100, false, false };
-										Parameters.ConnectionMethod = FClusterCreationParameters::EConnectionMethod::BoundsOverlapFilteredDelaunayTriangulation;
-										Parameters.ConnectionGraphBoundsFilteringMargin = 1.0;
-
-										Clustering.GenerateConnectionGraph(AllParticles, Parameters, &FromSet, &ToSet);
-									}
-								}
-							}
 						}
 					}
 				);
 			}
 		}
 
-		void RemoveClusterUnionEdges(FRigidClustering& Clustering, FPBDRigidParticleHandle* ParticleHandle)
+		void RemoveClusterUnionEdges(FRigidClustering& Clustering, FClusterUnion& ClusterUnion, FPBDRigidParticleHandle* ParticleHandle)
 		{
 			if (!ParticleHandle)
 			{
@@ -205,6 +171,11 @@ namespace Chaos
 
 			Clustering.RemoveNodeConnections(ParticleHandle);
 
+			if (FClusterUnionParticleProperties* Properties = ClusterUnion.ChildProperties.Find(ParticleHandle))
+			{
+				Properties->bEdgesAreGenerated = false;
+			}
+
 			if (bChaosClusterUnionGenerateInterclusterEdges)
 			{
 				if (FPBDRigidClusteredParticleHandle* ClusterParticle = ParticleHandle->CastToClustered())
@@ -238,7 +209,7 @@ namespace Chaos
 		{
 			for (FPBDRigidParticleHandle* ChildParticle : ClusterUnion.ChildParticles)
 			{
-				RemoveClusterUnionEdges(Clustering, ChildParticle);
+				RemoveClusterUnionEdges(Clustering, ClusterUnion, ChildParticle);
 				AddParticleToConnectionGraph(Clustering, ClusterUnion, ChildParticle);
 			}
 		}
@@ -1042,6 +1013,100 @@ namespace Chaos
 		AddParticleToConnectionGraph(MClustering, ClusterUnion, Particle);
 	}
 
+	DECLARE_CYCLE_STAT(TEXT("FClusterUnionManager::GenerateInterclusterEdgesForParticle"), STAT_GenerateInterclusterEdgesForParticle, STATGROUP_Chaos);
+	void FClusterUnionManager::GenerateInterclusterEdgesForParticle(FClusterUnion& ClusterUnion, FPBDRigidParticleHandle* Particle)
+	{
+		SCOPE_CYCLE_COUNTER(STAT_GenerateInterclusterEdgesForParticle);
+		if (!Particle)
+		{
+			return;
+		}
+
+		// Just a safety check really to make sure the particle is actually in the cluster union.
+		FClusterUnionParticleProperties* Properties = ClusterUnion.ChildProperties.Find(Particle);
+		if (!Properties || Properties->bEdgesAreGenerated)
+		{
+			return;
+		}
+
+		auto PrepareParticleForInterclusterEdges = [](FPBDRigidParticleHandle& Particle)
+		{
+			IPhysicsProxyBase* Proxy = Particle.PhysicsProxy();
+			if (!Proxy)
+			{
+				return;
+			}
+
+			if (Proxy->GetType() == FGeometryCollectionPhysicsProxy::ConcreteType())
+			{
+				FGeometryCollectionPhysicsProxy* GCProxy = static_cast<FGeometryCollectionPhysicsProxy*>(Proxy);
+				check(GCProxy != nullptr);
+				GCProxy->CreateChildrenGeometry_Internal();
+			}
+		};
+
+		PrepareParticleForInterclusterEdges(*Particle);
+
+		FPBDRigidClusteredParticleHandle* ClusteredParticle = Particle->CastToClustered();
+		if (!ClusteredParticle)
+		{
+			return;
+		}
+
+		const TArray<Chaos::FConnectivityEdge>& Edges = ClusteredParticle->ConnectivityEdges();
+		for (const Chaos::FConnectivityEdge& Edge : Edges)
+		{
+			if (Edge.Sibling != nullptr)
+			{
+				PrepareParticleForInterclusterEdges(*Edge.Sibling);
+				GenerateInterclusterEdgesBetweenParticles(ClusterUnion, Particle, Edge.Sibling);
+			}
+		}
+
+		Properties->bEdgesAreGenerated = true;
+	}
+
+	DECLARE_CYCLE_STAT(TEXT("FClusterUnionManager::GenerateInterclusterEdgesBetweenParticles"), STAT_GenerateInterclusterEdgesBetweenParticles, STATGROUP_Chaos);
+	void FClusterUnionManager::GenerateInterclusterEdgesBetweenParticles(FClusterUnion& ClusterUnion, FPBDRigidParticleHandle* Particle, FPBDRigidParticleHandle* OtherParticle)
+	{
+		SCOPE_CYCLE_COUNTER(STAT_GenerateInterclusterEdgesBetweenParticles);
+		if (!Particle || !OtherParticle)
+		{
+			return;
+		}
+
+		FClusterUnionParticleProperties* Properties = ClusterUnion.ChildProperties.Find(Particle);
+		if (!Properties)
+		{
+			return;
+		}
+
+		if (bChaosClusterUnionGenerateInterclusterEdges && Particle && OtherParticle)
+		{
+			// Only generate intercluster edges for main particles. Auxiliary particles that are just bits and pieces of geometry collections
+			// shouldn't also generate intercluster edges.
+			if (!Properties || !Properties->bIsAuxiliaryParticle)
+			{
+				const TArray<FPBDRigidParticleHandle*>& ParticleChildren = MClustering.GetChildrenMap().FindRef(Particle->CastToClustered());
+				const TArray<FPBDRigidParticleHandle*>& OtherChildren = MClustering.GetChildrenMap().FindRef(OtherParticle->CastToClustered());
+
+				TSet<FPBDRigidParticleHandle*> FromSet{ ParticleChildren };
+				TSet<FPBDRigidParticleHandle*> ToSet{ OtherChildren };
+
+				TArray<FPBDRigidParticleHandle*> AllParticles;
+				AllParticles.Reserve(FromSet.Num() + ToSet.Num());
+				AllParticles.Append(ParticleChildren);
+				AllParticles.Append(OtherChildren);
+
+				FClusterCreationParameters Parameters{ 0.3f, 100, false, false };
+				Parameters.ConnectionMethod = FClusterCreationParameters::EConnectionMethod::BoundsOverlapFilteredDelaunayTriangulation;
+				Parameters.ConnectionGraphBoundsFilteringMargin = 1.0;
+
+				MClustering.GenerateConnectionGraph(AllParticles, Parameters, &FromSet, &ToSet);
+			}
+		}
+	}
+
 	DECLARE_CYCLE_STAT(TEXT("FClusterUnionManager::FlushIncrementalConnectivityGraphOperations"), STAT_FlushIncrementalConnectivityGraphOperations, STATGROUP_Chaos);
 	void FClusterUnionManager::FlushIncrementalConnectivityGraphOperations(FClusterUnion& ClusterUnion)
 	{
@@ -1054,7 +1119,7 @@ namespace Chaos
 			}
 			else if (Op.Value == EClusterUnionConnectivityOperation::Remove)
 			{
-				RemoveClusterUnionEdges(MClustering, Op.Key);
+				RemoveClusterUnionEdges(MClustering, ClusterUnion, Op.Key);
 			}
 		}
 
diff --git a/Engine/Source/Runtime/Experimental/Chaos/Private/Chaos/PBDRigidClustering.cpp b/Engine/Source/Runtime/Experimental/Chaos/Private/Chaos/PBDRigidClustering.cpp
--- a/Engine/Source/Runtime/Experimental/Chaos/Private/Chaos/PBDRigidClustering.cpp
+++ b/Engine/Source/Runtime/Experimental/Chaos/Private/Chaos/PBDRigidClustering.cpp
@@ -1076,48 +1076,6 @@ namespace Chaos
 
 		return ReleaseClusterParticlesImpl(ClusteredParticle, bForceRelease, false /*bCreateNewClusters*/);
 	}
-	
-	void GenerateEdges(FGeometryCollectionPhysicsProxy& ConcreteGCProxy, FPBDRigidClusteredParticleHandle& ClusteredParticle, FClusterUnionManager& ClusterUnionManager)
-	{
-		ConcreteGCProxy.CreateChildrenGeometry_Internal();
-		if (Chaos::FClusterUnion* ClusterUnion = ClusterUnionManager.FindClusterUnionFromParticle(&ClusteredParticle))
-		{
-			bool bHasBuiltAllEdges = false;
-			FClusterUnionParticleProperties* Properties = ClusterUnion->ChildProperties.Find(&ClusteredParticle);
-			if (Properties)
-			{
-				bHasBuiltAllEdges = Properties->bEdgesAreGenerated;
-			}
-
-			bool bAllNeighborsHasBuiltEdges = true;
-			const TArray<Chaos::TConnectivityEdge<Chaos::FReal>> Edges = ClusteredParticle.ConnectivityEdges();
-			for (const Chaos::TConnectivityEdge<Chaos::FReal>& Edge : Edges)
-			{
-				if (Edge.Sibling != nullptr && Edge.Sibling->GetParticleType() == Chaos::EParticleType::Clustered)
-				{
-					Chaos::FPBDRigidClusteredParticleHandle* Sibling = Edge.Sibling->CastToClustered();
-					if (Sibling->PhysicsProxy()->GetType() == FGeometryCollectionPhysicsProxy::ConcreteType())
-					{
-						FGeometryCollectionPhysicsProxy* GCProxy = GetConcreteProxy<FGeometryCollectionPhysicsProxy>(Sibling);
-						GCProxy->CreateChildrenGeometry_Internal();
-						if (FClusterUnionParticleProperties* SiblingProperties = ClusterUnion->ChildProperties.Find(Sibling))
-						{
-							bAllNeighborsHasBuiltEdges &= SiblingProperties->bEdgesAreGenerated;
-						}
-					}
-				}
-			}
-			// If has current GC has built all edges or if all neighbors have build all edges don't need to compute neighbors edges. 
-			if (!(bHasBuiltAllEdges || bAllNeighborsHasBuiltEdges))
-			{
-				ClusterUnionManager.AddParticleToConnectionGraphInCluster(*ClusterUnion, &ClusteredParticle);
-				if (Properties)
-				{
-					Properties->bEdgesAreGenerated = true;
-				}
-			}
-		}
-	}
 
 	TSet<FPBDRigidParticleHandle*> FRigidClustering::ReleaseClusterParticlesImpl(
 		FPBDRigidClusteredParticleHandle* ClusteredParticle,
@@ -1150,7 +1108,7 @@ namespace Chaos
 			
 			if (bBuildGeometryForChildrenOnPT == false)
 			{
-				GenerateEdges(*ConcreteGCProxy, *ClusteredParticle, ClusterUnionManager);
+				ConcreteGCProxy->CreateChildrenGeometry_Internal();
 			}
 		}
 
@@ -1205,9 +1163,17 @@ namespace Chaos
 					const FClusterUnionIndex ClusterUnionIndex = ClusterUnionManager.FindClusterUnionIndexFromParticle(ClusteredParticle);
 					if (ClusterUnionIndex != INDEX_NONE)
 					{
+						// Generating intercluster edges depends on the parent particle's node connections existing. Thus this must go before RemoveNodeConnections.
+						// Our parent particle is being removed from the cluster union - initialize intercluster edges if necessary.
+						if (FClusterUnion* ClusterUnion = ClusterUnionManager.FindClusterUnion(ClusterUnionIndex))
+						{
+							ClusterUnionManager.GenerateInterclusterEdgesForParticle(*ClusterUnion, ClusteredParticle);
+						}
+
 						// Remove node connections here immediately just in case we need to manage connectivity on the cluster union.
 						RemoveNodeConnections(ClusteredParticle);
 						ClusterUnionsToConsiderForConnectivity.Add(ClusterUnionIndex);
+
 						ClusterUnionManager.HandleRemoveOperationWithClusterLookup({ ClusteredParticle }, EClusterUnionOperationTiming::Defer);
 					}
 					bFoundFirstRelease = true;
@@ -2129,7 +2095,7 @@ namespace Chaos
 	DECLARE_CYCLE_STAT(TEXT("TPBDRigidClustering<>::GenerateConnectionGraph"), STAT_GenerateConnectionGraph, STATGROUP_Chaos);
 	void
 	FRigidClustering::GenerateConnectionGraph(
-		TArray<FPBDRigidParticleHandle*> Particles,
+		const TArray<FPBDRigidParticleHandle*>& Particles,
 		const FClusterCreationParameters& Parameters,
 		const TSet<FPBDRigidParticleHandle*>* FromParticles,
 		const TSet<FPBDRigidParticleHandle*>* ToParticles)
diff --git a/Engine/Source/Runtime/Experimental/Chaos/Public/Chaos/ClusterUnionManager.h b/Engine/Source/Runtime/Experimental/Chaos/Public/Chaos/ClusterUnionManager.h
--- a/Engine/Source/Runtime/Experimental/Chaos/Public/Chaos/ClusterUnionManager.h
+++ b/Engine/Source/Runtime/Experimental/Chaos/Public/Chaos/ClusterUnionManager.h
@@ -78,6 +78,7 @@ namespace Chaos
 		{}
 		// An auxiliary particle will be removed from the cluster union if FRigidClustering::HandleConnectivityOnReleaseClusterParticle detects an island made up of only auxiliary particles.
 		uint8 bIsAuxiliaryParticle: 1;
+		// Really just a protection for external users accidentally incurring a perf cost if they try to manually generate edges. We will never do this more than once per particle per cluster union.
 		uint8 bEdgesAreGenerated: 1;
 	};
 
@@ -213,7 +214,10 @@ namespace Chaos
 		// Update the cluster union's properties after its set of particle changes.
 		CHAOS_API void UpdateAllClusterUnionProperties(FClusterUnion& ClusterUnion, EUpdateClusterUnionPropertiesFlags Flags = EUpdateClusterUnionPropertiesFlags::All);
 
-		CHAOS_API  void AddParticleToConnectionGraphInCluster(FClusterUnion& ClusterUnion, FPBDRigidParticleHandle* Particle);
+		CHAOS_API void AddParticleToConnectionGraphInCluster(FClusterUnion& ClusterUnion, FPBDRigidParticleHandle* Particle);
+
+		CHAOS_API void GenerateInterclusterEdgesForParticle(FClusterUnion& ClusterUnion, FPBDRigidParticleHandle* Particle);
+		CHAOS_API void GenerateInterclusterEdgesBetweenParticles(FClusterUnion& ClusterUnion, FPBDRigidParticleHandle* Particle, FPBDRigidParticleHandle* OtherParticle);
 
 		// Returns all cluster unions. Really meant only to be used for debugging.
 		const TMap<FClusterUnionIndex, FClusterUnion>& GetAllClusterUnions() const { return ClusterUnions; }
diff --git a/Engine/Source/Runtime/Experimental/Chaos/Public/Chaos/PBDRigidClustering.h b/Engine/Source/Runtime/Experimental/Chaos/Public/Chaos/PBDRigidClustering.h
--- a/Engine/Source/Runtime/Experimental/Chaos/Public/Chaos/PBDRigidClustering.h
+++ b/Engine/Source/Runtime/Experimental/Chaos/Public/Chaos/PBDRigidClustering.h
@@ -376,7 +376,7 @@ public:
 	FClusterCreationParameters::EConnectionMethod GetClusterUnionConnectionType() const { return MClusterUnionConnectionType; }
 
 	CHAOS_API void GenerateConnectionGraph(
-		TArray<FPBDRigidParticleHandle*> Particles,
+		const TArray<FPBDRigidParticleHandle*>& Particles,
 		const FClusterCreationParameters& Parameters = FClusterCreationParameters(),
 		const TSet<FPBDRigidParticleHandle*>* FromParticles = nullptr,
 		const TSet<FPBDRigidParticleHandle*>* ToParticles = nullptr);
