Phrase data size and write size are the same now

#rnx
#ushell-unshelve of 32465373

diff --git a/Engine/Source/Runtime/Experimental/IoStoreOnDemand/Private/JournaledCache.cpp b/Engine/Source/Runtime/Experimental/IoStoreOnDemand/Private/JournaledCache.cpp
--- a/Engine/Source/Runtime/Experimental/IoStoreOnDemand/Private/JournaledCache.cpp
+++ b/Engine/Source/Runtime/Experimental/IoStoreOnDemand/Private/JournaledCache.cpp
@@ -305,16 +305,13 @@ struct FDiskPhrase
 	const FDataEntry*	GetEntries() const		{ return Entries.GetData() + Index; }
 	int32				GetEntryCount() const	{ return Entries.Num() - Index; }
 	uint8*				GetPhraseData() const	{ return Buffer.Get(); }
-	uint32				GetDataSize() const		{ return EntriesSize; }
+	uint32				GetDataSize() const		{ return Cursor; }
 	int32				GetRemainingEntries() const { return MaxEntries; }
-	uint32				GetWriteSize() const	{ return Cursor; }
 
 private:
 	TUniquePtr<uint8[]>	Buffer;
 	TArray<FDataEntry>&	Entries;
 	uint32				Cursor = 0;
-	uint32				CurrentOffset = 0;
-	uint32				EntriesSize = 0;
 	uint32				Index;
 	int32				MaxEntries;
 
@@ -341,16 +338,15 @@ bool FDiskPhrase::Add(uint64 Key, FIoBuffer&& Data, uint32 PartialBias)
 	check(DataSize < (1 << SIZE_BITS));
 	Entries.Add(FDataEntry{
 		.Key = Key,
-		.Offset = CurrentOffset + PartialBias,
+		.Offset = Cursor + PartialBias,
 		.Size = DataSize + PartialBias,
 		.EntryCount = 0,
 	});
-	EntriesSize += DataSize;
-	CurrentOffset += DataSize;
 	--MaxEntries;
 
 	std::memcpy(Buffer.Get() + Cursor, Data.GetData(), DataSize);
 	Cursor += DataSize;
+
 	return MaxEntries > 0;
 }
 
@@ -643,7 +639,7 @@ void FDiskCache::Wrap()
 void FDiskCache::ClosePhrase(FDiskPhrase&& Phrase)
 {
 	int32 EntryCount = Phrase.GetEntryCount();
-	if (Phrase.GetWriteSize() == 0)
+	if (Phrase.GetDataSize() == 0)
 	{
 		Journal.ClosePhrase(MoveTemp(Phrase), 0);
 		return;
@@ -658,7 +654,7 @@ void FDiskCache::ClosePhrase(FDiskPhrase&& Phrase)
 		return;
 	}
 
-	uint32 WriteSize = Phrase.GetWriteSize();
+	uint32 WriteSize = Phrase.GetDataSize();
 	check(DataCursor + WriteSize <= MaxDataSize);
 
 	bool bWriteOk;
