Configurable cache magic to allow for hotfix invalidation

#rnx
#ushell-unshelve of 32465373

diff --git a/Engine/Source/Runtime/Experimental/IoStoreOnDemand/Private/IasCache.h b/Engine/Source/Runtime/Experimental/IoStoreOnDemand/Private/IasCache.h
--- a/Engine/Source/Runtime/Experimental/IoStoreOnDemand/Private/IasCache.h
+++ b/Engine/Source/Runtime/Experimental/IoStoreOnDemand/Private/IasCache.h
@@ -70,6 +70,7 @@ struct FIasCacheConfig
 	uint64		DiskQuota = 512ull << 20;
 	uint32		MemoryQuota = 2 << 20;
 	uint32		JournalQuota = 4 << 20; // description in JournalCache.cpp
+	uint32		JournalMagic = 0; // can be used to invalidate cache
 	FRate		WriteRate;
 	FDemand		Demand;
 	bool		DropCache = false;
diff --git a/Engine/Source/Runtime/Experimental/IoStoreOnDemand/Private/IoStoreOnDemand.cpp b/Engine/Source/Runtime/Experimental/IoStoreOnDemand/Private/IoStoreOnDemand.cpp
--- a/Engine/Source/Runtime/Experimental/IoStoreOnDemand/Private/IoStoreOnDemand.cpp
+++ b/Engine/Source/Runtime/Experimental/IoStoreOnDemand/Private/IoStoreOnDemand.cpp
@@ -257,6 +257,7 @@ static FIasCacheConfig GetIasCacheConfig(const TCHAR* CommandLine)
 	GetConfigInt(TEXT("DiskQuota"),				Ret.DiskQuota);
 	GetConfigInt(TEXT("MemoryQuota"),			Ret.MemoryQuota);
 	GetConfigInt(TEXT("JournalQuota"),			Ret.JournalQuota);
+	GetConfigInt(TEXT("JournalMagic"),			Ret.JournalMagic);
 	GetConfigInt(TEXT("DemandThreshold"),		Ret.Demand.Threshold);
 	GetConfigInt(TEXT("DemandBoost"),			Ret.Demand.Boost);
 	GetConfigInt(TEXT("DemandSuperBoost"),		Ret.Demand.SuperBoost);
diff --git a/Engine/Source/Runtime/Experimental/IoStoreOnDemand/Private/JournaledCache.cpp b/Engine/Source/Runtime/Experimental/IoStoreOnDemand/Private/JournaledCache.cpp
--- a/Engine/Source/Runtime/Experimental/IoStoreOnDemand/Private/JournaledCache.cpp
+++ b/Engine/Source/Runtime/Experimental/IoStoreOnDemand/Private/JournaledCache.cpp
@@ -268,7 +268,6 @@ int32 FMemCache::Drop(uint32 Size)
 // {{{1 phrase .................................................................
 
 ////////////////////////////////////////////////////////////////////////////////
-static const uint32 MAGIC = 0x04930003;
 static const uint32 SIZE_BITS = 25;
 static const uint32 MARKER_MAX = 0x3fffffff;
 static const uint32	HASH_CHECKSUM_SIZE = 64;
@@ -358,7 +357,7 @@ bool FDiskPhrase::Add(uint64 Key, FIoBuffer&& Data, uint32 PartialBias)
 class FDiskJournal
 {
 public:
-							FDiskJournal(FStringView InRootPath, uint32 InMaxSize);
+							FDiskJournal(FStringView InRootPath, uint32 InMaxSize, uint32 MagicSeed=0);
 	uint32					GetAilments() const;
 	void					Drop();
 	int32					Flush();
@@ -370,6 +369,7 @@ public:
 
 private:
 	friend int32			LoadCache(FDiskCache&);
+	uint32					GetMagic() const { return Magic; }
 	void					GetPath(TStringBuilder<64>& Out);
 	void					OpenJrnFile();
 	static uint32			HashBytes(const uint8* Data, uint32 Size, uint32 Seed);
@@ -379,16 +379,23 @@ private:
 	uint32					Marker = 0;
 	uint32					Cursor = 0;
 	uint32					MaxSize;
+	uint32					Magic;
 };
 
 ////////////////////////////////////////////////////////////////////////////////
-FDiskJournal::FDiskJournal(FStringView InRootPath, uint32 InMaxSize)
+FDiskJournal::FDiskJournal(FStringView InRootPath, uint32 InMaxSize, uint32 MagicSeed)
 : RootPath(InRootPath)
 , MaxSize(InMaxSize)
 {
 	// Align down to keep to some assumptions
 	MaxSize &= ~(sizeof(FDataEntry) - 1);
 
+	// High word is constant to identify the magic value, lower word mixes in the
+	// seed to afford some control over validation. The value "3" is for backwards
+	// compatibilty; it comes from the hardcoded magic value already in use.
+	Magic = 0x0493'0003;
+	Magic += ((Magic + MagicSeed) & 0xffff);
+
 	OpenJrnFile();
 }
 
@@ -487,7 +494,7 @@ void FDiskJournal::ClosePhrase(FDiskPhrase&& Phrase, uint64 DataCursor)
 	LastEntry.EntryCount = uint16(EntryCount);
 
 	auto& Desc = (FPhraseDesc&)(Phrase.GetEntries()[-1]);
-	Desc.Magic = MAGIC;
+	Desc.Magic = GetMagic();
 	Desc.Marker = Marker;
 	Desc.DataCursor = DataCursor;
 
@@ -939,7 +946,8 @@ static int32 LoadCache(FDiskCache& DiskCache)
 		return (UPTRINT(Address) - UPTRINT(Data.Get())) > DataSize;
 	};
 
-	auto ReadPhrases = [&IsOob] (const uint8* Cursor, FParagraph& Out) -> const uint8*
+	uint32 Magic = Journal.GetMagic();
+	auto ReadPhrases = [&IsOob, Magic] (const uint8* Cursor, FParagraph& Out) -> const uint8*
 	{
 		// Only proceed if we can read at least three integers
 		if (IsOob(Cursor + sizeof(FPhraseDesc)))
@@ -948,7 +956,7 @@ static int32 LoadCache(FDiskCache& DiskCache)
 		}
 
 		const auto* Header = (FPhraseDesc*)Cursor;
-		if (Header->Magic != MAGIC)
+		if (Header->Magic != Magic)
 		{
 			return nullptr;
 		}
