Partial.Key can serve as a "set or not" condition

#rnx
#ushell-unshelve of 32465373

diff --git a/Engine/Source/Runtime/Experimental/IoStoreOnDemand/Private/JournaledCache.cpp b/Engine/Source/Runtime/Experimental/IoStoreOnDemand/Private/JournaledCache.cpp
--- a/Engine/Source/Runtime/Experimental/IoStoreOnDemand/Private/JournaledCache.cpp
+++ b/Engine/Source/Runtime/Experimental/IoStoreOnDemand/Private/JournaledCache.cpp
@@ -1197,13 +1197,12 @@ public:
 	uint32			DebugVisit(void* Param, FDebugCacheEntry::Callback* Callback);
 
 private:
-	struct FPartialItem
+	struct FPartial
 	{
-		uint64		Key;
+		uint64		Key = 0;
 		FIoBuffer	Data;
 		uint32		Cursor;
 	};
-	using FPartial = TOptional<FPartialItem>;
 	
 	mutable FRWLock	MemLock;
 	FMemCache		MemCache;
@@ -1266,7 +1265,7 @@ bool FCache::Has(uint64 Key) const
 	}
 
 	FReadScopeLock _(MemLock);
-	return (MemCache.Get(Key) != nullptr) || (Partial && Partial->Key == Key);
+	return (MemCache.Get(Key) != nullptr) || (Partial.Key == Key);
 }
 
 ////////////////////////////////////////////////////////////////////////////////
@@ -1287,9 +1286,9 @@ FCache::FGetToken FCache::Get(uint64 Key, FIoBuffer& OutData) const
 		OutData = *Data;
 	}
 
-	if (Partial && Partial->Key == Key)
+	if (Partial.Key == Key)
 	{
-		OutData = Partial->Data;
+		OutData = Partial.Data;
 	}
 
 	return 0;
@@ -1348,28 +1347,28 @@ uint32 FCache::WriteMemToDisk(int32 Allowance)
 
 		// If we have any partials that was previously written process that first
 		// and peel off as much of that buffer as possible.
-		if (Partial)
+		if (Partial.Key)
 		{
-			FMemoryView View = Partial->Data.GetView();
-			View = View.Mid(Partial->Cursor, Allowance);
+			FMemoryView View = Partial.Data.GetView();
+			View = View.Mid(Partial.Cursor, Allowance);
 
 			uint32 ViewSize = uint32(View.GetSize());
-			Partial->Cursor += ViewSize;
+			Partial.Cursor += ViewSize;
 			WriteSize += ViewSize;
 
 			FMemCache::FItem PeelItem = {
 				.Key	= 0,
-				.Data	= FIoBuffer(View, Partial->Data),
+				.Data	= FIoBuffer(View, Partial.Data),
 			};
-			if (Partial->Cursor >= uint32(Partial->Data.GetSize()))
+			if (Partial.Cursor >= uint32(Partial.Data.GetSize()))
 			{
-				PeelItem.Key = Partial->Key;
-				PartialBias = uint32(Partial->Data.GetSize()) - ViewSize;
-				Partial.Reset();
+				PeelItem.Key = Partial.Key;
+				PartialBias = uint32(Partial.Data.GetSize()) - ViewSize;
+				Partial = FPartial();
 			}
 			else if (bEof)
 			{
-				Partial->Cursor = 0;
+				Partial.Cursor = 0;
 			}
 			PeelItems.Push(PeelItem);
 		}
@@ -1398,7 +1397,7 @@ uint32 FCache::WriteMemToDisk(int32 Allowance)
 			WriteSize -= Overshoot;
 
 			uint32 Cursor = bEof ? 0 : PartialSize;
-			Partial.Emplace(FPartialItem{Key, MoveTemp(Data), uint32(Cursor)});
+			Partial = FPartial{Key, MoveTemp(Data), uint32(Cursor)};
 		}
 	}
 
@@ -1436,10 +1435,10 @@ uint32 FCache::DebugVisit(void* Param, FDebugCacheEntry::Callback* Callback)
 	uint32 Count = 0;
 	Count += MemCache.DebugVisit(Param, Callback);
 	Count += DiskCache.DebugVisit(Param, Callback);
-	if (Partial)
+	if (Partial.Key)
 	{
 		Count++;
-		Callback(Param, {Partial->Key, uint32(Partial->Data.GetSize()), 1});
+		Callback(Param, {Partial.Key, uint32(Partial.Data.GetSize()), 1});
 	}
 	return Count;
 }
