Swapped over some internal APIs to use FOutcome object

#rnx

diff --git a/Engine/Source/Runtime/Experimental/IoStoreOnDemand/Private/CoreHttp/Http.cpp b/Engine/Source/Runtime/Experimental/IoStoreOnDemand/Private/CoreHttp/Http.cpp
--- a/Engine/Source/Runtime/Experimental/IoStoreOnDemand/Private/CoreHttp/Http.cpp
+++ b/Engine/Source/Runtime/Experimental/IoStoreOnDemand/Private/CoreHttp/Http.cpp
@@ -806,10 +806,10 @@ public:
 	bool		IsValid() const				{ return Socket != InvalidSocket; }
 	bool		Create();
 	void		Destroy();
-	bool		Connect(uint32 Ip, uint32 Port);
+	FOutcome	Connect(uint32 Ip, uint32 Port);
 	void		Disconnect();
-	int32		Send(const char* Data, uint32 Size);
-	int32		Recv(char* Dest, uint32 Size);
+	FOutcome	Send(const char* Data, uint32 Size);
+	FOutcome	Recv(char* Dest, uint32 Size);
 	bool		SetBlocking(bool bBlocking);
 	bool		SetSendBufSize(int32 Size);
 	bool		SetRecvBufSize(int32 Size);
@@ -869,7 +869,7 @@ void FSocket::Destroy()
 }
 
 ////////////////////////////////////////////////////////////////////////////////
-bool FSocket::Connect(uint32 IpAddress, uint32 Port)
+FOutcome FSocket::Connect(uint32 IpAddress, uint32 Port)
 {
 	check(IsValid());
 
@@ -886,10 +886,15 @@ bool FSocket::Connect(uint32 IpAddress, uint32 Port)
 
 	if (IsSocketResult(EWOULDBLOCK) | IsSocketResult(EINPROGRESS))
 	{
-		return true;
+		return FOutcome::Waiting();
+	}
+
+	if (Result < 0)
+	{
+		return FOutcome::Error("Socket connect failed", Result);
 	}
 
-	return (Result >= 0);
+	return FOutcome::Ok();
 }
 
 ////////////////////////////////////////////////////////////////////////////////
@@ -900,26 +905,15 @@ void FSocket::Disconnect()
 }
 
 ////////////////////////////////////////////////////////////////////////////////
-int32 FSocket::Send(const char* Data, uint32 Size)
+FOutcome FSocket::Send(const char* Data, uint32 Size)
 {
 	Trace(Socket, ETrace::Send, -1);
 	int32 Result = send(Socket, Data, Size, MsgFlagType(0));
 	Trace(Socket, ETrace::Send, FMath::Max(Result, 0));
 
-	if (Result > 0)
-	{
-		return Result;
-	}
-
-	if (Result == 0)
-	{
-		return int32(EResult::HangUp);
-	}
-
-	if (IsSocketResult(EWOULDBLOCK))
-	{
-		return int32(EResult::Wait);
-	}
+	if (Result > 0)						return FOutcome::Ok(Result);
+	if (Result == 0)					return FOutcome::Error("Send ATH0");
+	if (IsSocketResult(EWOULDBLOCK))	return FOutcome::Waiting();
 
 	if (IsSocketResult(ENOTCONN))
 	{
@@ -928,38 +922,29 @@ int32 FSocket::Send(const char* Data, uint32 Size)
 		Result = getsockopt(Socket, SOL_SOCKET, SO_ERROR, (char*)&Error, &ErrorSize);
 		if (Result < 0 || Error != 0)
 		{
-			return int32(EResult::ConnectError);
+			return FOutcome::Error("Error while connecting");
 		}
 
-		return int32(EResult::Wait);
+		return FOutcome::Waiting();
 	}
 
-	return int32(EResult::Error);
+	Result = LastSocketResult();
+	return FOutcome::Error("Send", Result);
 }
 
 ////////////////////////////////////////////////////////////////////////////////
-int32 FSocket::Recv(char* Dest, uint32 Size)
+FOutcome FSocket::Recv(char* Dest, uint32 Size)
 {
 	Trace(Socket, ETrace::Recv, -1);
 	int32 Result = recv(Socket, Dest, Size, MsgFlagType(0));
 	Trace(Socket, ETrace::Recv, FMath::Max(0, Result));
 
-	if (Result > 0)
-	{
-		return Result;
-	}
-
-	if (Result == 0)
-	{
-		return int32(EResult::HangUp);
-	}
+	if (Result > 0)						return FOutcome::Ok(Result);
+	if (Result == 0)					return FOutcome::Error("Recv ATH0");
+	if (IsSocketResult(EWOULDBLOCK))	return FOutcome::Waiting();
 
-	if (IsSocketResult(EWOULDBLOCK))
-	{
-		return int32(EResult::Wait);
-	}
-
-	return int32(EResult::Error);
+	Result = LastSocketResult();
+	return FOutcome::Error("Recv", Result);
 }
 
 ////////////////////////////////////////////////////////////////////////////////
@@ -1165,26 +1150,29 @@ static int32 ConnectSocks4(FSocket& Socket, uint32 IpAddress, uint32 Port)
 	};
 
 	uint32 SocksIpAddress = GetSocksIpAddress();
-	if (!SocksIpAddress || !Socket.Connect(SocksIpAddress, GSocksPort))
+	if (!SocksIpAddress)
 	{
 		return -1;
 	}
 
-	int32 Result;
+	FOutcome Outcome = FOutcome::None();
+
+	if (Outcome = Socket.Connect(SocksIpAddress, GSocksPort); Outcome.IsError())
+	{
+		return -1;
+	}
 
 	FSocks4Request Request = {
 		.Port		= htons(uint16(Port)),
 		.IpAddress	= htonl(IpAddress),
 	};
-	Result = Socket.Send((const char*)&Request, sizeof(Request));
-	if (Result <= 0)
+	if (Outcome = Socket.Send((const char*)&Request, sizeof(Request)); Outcome.IsError())
 	{
 		return -1;
 	}
 
 	FSocks4Reply Reply;
-	Result = Socket.Recv((char*)&Reply, sizeof(Reply));
-	if (Result <= 0)
+	if (Outcome = Socket.Recv((char*)&Reply, sizeof(Reply)); Outcome.IsError())
 	{
 		return -1;
 	}
@@ -1203,25 +1191,30 @@ static int32 ConnectSocks5(FSocket& Socket, uint32 IpAddress, uint32 Port)
 #endif
 
 	uint32 SocksIpAddress = GetSocksIpAddress();
-	if (!SocksIpAddress || !Socket.Connect(SocksIpAddress, GSocksPort))
+	if (!SocksIpAddress)
 	{
 		return -1;
 	}
 
-	int32 Result;
+	FOutcome Outcome = FOutcome::None();
+
+	if (Outcome = Socket.Connect(SocksIpAddress, GSocksPort); Outcome.IsError())
+	{
+		return -1;
+	}
 
 	// Greeting
 	const char Greeting[] = { 5, 1, 0 };
-	Result = Socket.Send(Greeting, sizeof(Greeting));
-	if (Result != sizeof(Greeting))
+	Outcome = Socket.Send(Greeting, sizeof(Greeting));
+	if (!Outcome.IsOk() || Outcome.GetResult() != sizeof(Greeting))
 	{
 		return -1;
 	}
 
 	// Server auth-choice
 	char ServerChoice[1 + 1];
-	Result = Socket.Recv(ServerChoice, sizeof(ServerChoice));
-	if (Result != sizeof(ServerChoice))
+	Outcome = Socket.Recv(ServerChoice, sizeof(ServerChoice));
+	if (!Outcome.IsOk() || Outcome.GetResult() != sizeof(ServerChoice))
 	{
 		return -1;
 	}
@@ -1237,16 +1230,16 @@ static int32 ConnectSocks5(FSocket& Socket, uint32 IpAddress, uint32 Port)
 	char Request[] = { 5, 1, 0, 1, 0x11,0x11,0x11,0x11, 0x22,0x22 };
 	std::memcpy(Request + 4, &IpAddress, sizeof(IpAddress));
 	std::memcpy(Request + 8, &NsPort, sizeof(NsPort));
-	Result = Socket.Send(Request, sizeof(Request));
-	if (Result != sizeof(Request))
+	Outcome = Socket.Send(Request, sizeof(Request));
+	if (!Outcome.IsOk() || Outcome.GetResult() != sizeof(Request))
 	{
 		return -1;
 	}
 
 	// Connect reply
 	char Reply[3 + (1 + 4) + 2];
-	Result = Socket.Recv(Reply, sizeof(Reply));
-	if (Result != sizeof(Reply))
+	Outcome = Socket.Recv(Reply, sizeof(Reply));
+	if (!Outcome.IsOk() || Outcome.GetResult() != sizeof(Reply))
 	{
 		return -1;
 	}
@@ -1300,9 +1293,9 @@ public:
 					FHost(const ANSICHAR* InHostName, uint32 InPort, uint32 InMaxConn);
 	void			SetBufferSize(EDirection Dir, int32 Size);
 	int32			GetBufferSize(EDirection Dir) const;
-	FResult			Connect(FSocket& Socket);
+	FOutcome		Connect(FSocket& Socket);
 	int32			IsResolved() const;
-	FResult			ResolveHostName();
+	FOutcome		ResolveHostName();
 	uint32			GetMaxConnections() const	{ return MaxConnections; }
 	uint32			GetIpAddress() const		{ return IpAddresses[0]; }
 	FAnsiStringView	GetHostName() const			{ return HostName; }
@@ -1339,7 +1332,7 @@ int32 FHost::GetBufferSize(EDirection Dir) const
 }
 
 ////////////////////////////////////////////////////////////////////////////////
-FResult FHost::ResolveHostName()
+FOutcome FHost::ResolveHostName()
 {
 	// todo: GetAddrInfoW() for async resolve on Windows
 
@@ -1357,12 +1350,12 @@ FResult FHost::ResolveHostName()
 	auto Result = getaddrinfo(HostName, nullptr, &Hints, &Info);
 	if (uint32(Result) || Info == nullptr)
 	{
-		return FResult(-1, "Error encountered resolving");
+		return FOutcome::Error("Error encountered resolving");
 	}
 
 	if (Info->ai_family != AF_INET)
 	{
-		return FResult(-2, "Unexpected address family during resolve");
+		return FOutcome::Error("Unexpected address family during resolve");
 	}
 
 	uint32 AddressCount = 0;
@@ -1391,10 +1384,10 @@ FResult FHost::ResolveHostName()
 
 	if (AddressCount > 0)
 	{
-		return FResult(AddressCount);
+		return FOutcome::Ok(AddressCount);
 	}
 
-	return FResult(0, "Unable to resolve host");
+	return FOutcome::Error("Unable to resolve host");
 }
 
 ////////////////////////////////////////////////////////////////////////////////
@@ -1409,14 +1402,13 @@ int32 FHost::IsResolved() const
 }
 
 ////////////////////////////////////////////////////////////////////////////////
-FResult FHost::Connect(FSocket& Socket)
+FOutcome FHost::Connect(FSocket& Socket)
 {
 	if (IsResolved() <= 0)
 	{
-		FResult Result = ResolveHostName();
-		if (Result.GetValue() <= 0)
+		if (FOutcome Outcome = ResolveHostName(); Outcome.IsError())
 		{
-			return Result;
+			return Outcome;
 		}
 	}
 
@@ -1428,7 +1420,7 @@ FResult FHost::Connect(FSocket& Socket)
 	FSocket Candidate;
 	if (!Candidate.Create())
 	{
-		return FResult(-1, "Failed to create socket");
+		return FOutcome::Error("Failed to create socket");
 	}
 
 	// Attempt a SOCKS connect
@@ -1437,7 +1429,7 @@ FResult FHost::Connect(FSocket& Socket)
 	{
 		if (Result < 0)
 		{
-			return FResult("Failed establishing SOCKS connection");
+			return FOutcome::Error("Failed establishing SOCKS connection");
 		}
 
 		bSocksConnected = true;
@@ -1446,7 +1438,7 @@ FResult FHost::Connect(FSocket& Socket)
 	// Condition the socket
 	if (!Candidate.SetBlocking(false))
 	{
-		return FResult("Unable to set socket non-blocking");
+		return FOutcome::Error("Unable to set socket non-blocking");
 	}
 
 	if (int32 OptValue = GetBufferSize(FHost::EDirection::Send); OptValue >= 0)
@@ -1463,17 +1455,17 @@ FResult FHost::Connect(FSocket& Socket)
 	if (bSocksConnected)
 	{
 		Socket = MoveTemp(Candidate);
-		return FResult(1);
+		return FOutcome::Ok();
 	}
 
 	// Issue the connect - this is done non-blocking so we need to wait (ret=0)
-	if (!Candidate.Connect(IpAddress, Port))
+	if (FOutcome Outcome = Candidate.Connect(IpAddress, Port); Outcome.IsError())
 	{
-		return FResult("Socket connect failed");
+		return Outcome;
 	}
 
 	Socket = MoveTemp(Candidate);
-	return FResult(0);
+	return FOutcome::Waiting();
 }
 
 
@@ -1538,7 +1530,7 @@ FConnectionPool::~FConnectionPool()
 ////////////////////////////////////////////////////////////////////////////////
 bool FConnectionPool::Resolve()
 {
-	return (Ptr->ResolveHostName().GetValue() > 0);
+	return Ptr->ResolveHostName().IsOk();
 }
 
 ////////////////////////////////////////////////////////////////////////////////
@@ -2053,17 +2045,22 @@ static int32 DoSend(FActivity* Activity, FSocket& Socket)
 	SendSize -= AlreadySent;
 	check(SendSize > 0);
 
-	int32 Result = Socket.Send(SendData, SendSize);
+	FOutcome Outcome = Socket.Send(SendData, SendSize);
+
+	if (Outcome.IsError())
+	{
+		Activity_SetError(Activity, Outcome.GetMessage().GetData());
+		return -1;
+	}
 
-	switch (FSocket::EResult(Result))
+	if (Outcome.IsWaiting())
 	{
-	case FSocket::EResult::HangUp:		Activity_SetError(Activity, "ATH0.Send"); return Result;
-	case FSocket::EResult::Error:		Activity_SetError(Activity, "Error returned from socket send"); return Result;
-	case FSocket::EResult::ConnectError:Activity_SetError(Activity, "Connection error"); return Result;
-	case FSocket::EResult::Wait:		return Result;
+		return 0;
 	}
 
-	checkf(Result > 0, TEXT("Result wasn't caught by switch statement so it is expected to be a positive amount of bytes sent"));
+	check(Outcome.IsOk());
+
+	int32 Result = Outcome.GetResult();
 	Activity->StateParam += Result;
 	if (Activity->StateParam < Buffer.GetSize())
 	{
@@ -2097,19 +2094,22 @@ static int32 DoRecvMessage(FActivity* Activity, FSocket& Socket)
 #endif
 
 		auto [Dest, DestSize] = Buffer.GetMutableFree(0, PageSize);
-		int32 Result = Socket.Recv(Dest, DestSize);
+		FOutcome Outcome = Socket.Recv(Dest, DestSize);
 
-		if (Result == int32(FSocket::EResult::Wait))
+		if (Outcome.IsError())
 		{
-			return 1;
+			Activity_SetError(Activity, Outcome.GetMessage().GetData());
+			return -1;
 		}
 
-		if (Result < 0)
+		if (Outcome.IsWaiting())
 		{
-			Activity_SetError(Activity, "Error returned from socket recv");
-			return -1;
+			return 1;
 		}
 
+		check(Outcome.IsOk());
+
+		int32 Result = Outcome.GetResult();
 		Buffer.AdvanceUsed(Result);
 
 		// Rewind a little to cover cases where the terminal is fragmented across
@@ -2129,12 +2129,6 @@ static int32 DoRecvMessage(FActivity* Activity, FSocket& Socket)
 				return -1;
 			}
 
-			if (Result == 0)
-			{
-				Activity_SetError(Activity, "ATH0.RecvMessage");
-				return -1;
-			}
-
 			continue;
 		}
 
@@ -2311,26 +2305,24 @@ static int32 DoRecvContent(FActivity* Activity, FSocket& Socket, int32& MaxRecvS
 
 		char* Cursor = (char*)(DestView.GetData()) + Activity->StateParam;
 
-		int32 Result = Socket.Recv(Cursor, Size);
+		FOutcome Outcome = Socket.Recv(Cursor, Size);
 
-		if (Result == int32(FSocket::EResult::Wait))
+		if (Outcome.IsWaiting())
 		{
 			return 1;
 		}
 
-		if (Result < 0)
+		if (Outcome.IsError())
 		{
-			Activity_SetError(Activity, "Socket error while receiving content");
+			Activity_SetError(Activity, Outcome.GetMessage().GetData());
 			return -1;
 		}
 
-		if (Result == 0 && (Activity->StateParam != Response.ContentLength))
-		{
-			Activity_SetError(Activity, "ATH0.RecvContent");
-			return -1;
-		}
+		check(Outcome.IsOk());
 
+		int32 Result = Outcome.GetResult();
 		check(Result <= MaxRecvSize);
+
 		Activity->StateParam += Result;
 		MaxRecvSize -= Result;
 	}
@@ -2834,7 +2826,7 @@ void FSocketGroup::SendInternal(FTickState& State)
 
 	int32 Result = DoSend(Activity, Socket);
 
-	if (Result == int32(FSocket::EResult::Wait))
+	if (Result == 0)
 	{
 		// For now we'll not add the socket as a waiter. It is unlikely that we
 		// send enough to need to wait currently.
@@ -2902,18 +2894,18 @@ void FSocketGroup::TickSend(FTickState& State, FHost& Host)
 	if (!Socket.IsValid())
 	{
 		IsKeepAlive = 1;
-		FResult Result = Host.Connect(Socket);
+		FOutcome Outcome = Host.Connect(Socket);
 
 		// We failed to connect, let's bail.
-		if (Result.GetValue() < 0)
+		if (Outcome.IsError())
 		{
 			Pending->Next = Recv;
 			Recv = Pending;
-			Fail(State, Result.GetMessage());
+			Fail(State, Outcome.GetMessage().GetData());
 			return;
 		}
 
-		bWillBlock = (Result.GetValue() == 0);
+		bWillBlock = Outcome.IsWaiting();
 	}
 
 	Send = Pending;
