Chaos - better initial overlap depentration

Initial overlap is now per shape pair rather than per-contact. This better handles the case where we initial generate a contact with an under-estimate of the penetration depth (because we corrected an edge collision), but later get a face collision with the correct deeper depth. Prior to this we wouild pop out by the delta between the two leading to a large velocity.

diff --git a/Engine/Source/Runtime/Experimental/Chaos/Private/Chaos/Collision/PBDCollisionConstraint.cpp b/Engine/Source/Runtime/Experimental/Chaos/Private/Chaos/Collision/PBDCollisionConstraint.cpp
--- a/Engine/Source/Runtime/Experimental/Chaos/Private/Chaos/Collision/PBDCollisionConstraint.cpp
+++ b/Engine/Source/Runtime/Experimental/Chaos/Private/Chaos/Collision/PBDCollisionConstraint.cpp
@@ -1211,7 +1211,6 @@ namespace Chaos
 					ManifoldPoints[PointIndex].Flags.bHasStaticFrictionAnchor = true;
 					ManifoldPoints[PointIndex].ShapeAnchorPoints[0] = SavedManifoldPoints[SavedManifoldPointIndex].ShapeContactPoints[0];
 					ManifoldPoints[PointIndex].ShapeAnchorPoints[1] = SavedManifoldPoints[SavedManifoldPointIndex].ShapeContactPoints[1];
-					ManifoldPoints[PointIndex].InitialPhi = SavedManifoldPoints[SavedManifoldPointIndex].InitialPhi;
 					ManifoldPoints[PointIndex].Flags.bInitialContact = false;
 				}
 				// Nothing to do if no saved friction point because we already set the achor to the most recently detected contact point
diff --git a/Engine/Source/Runtime/Experimental/Chaos/Private/Chaos/Collision/PBDCollisionContainerSolver.cpp b/Engine/Source/Runtime/Experimental/Chaos/Private/Chaos/Collision/PBDCollisionContainerSolver.cpp
--- a/Engine/Source/Runtime/Experimental/Chaos/Private/Chaos/Collision/PBDCollisionContainerSolver.cpp
+++ b/Engine/Source/Runtime/Experimental/Chaos/Private/Chaos/Collision/PBDCollisionContainerSolver.cpp
@@ -179,12 +179,12 @@ namespace Chaos
 				}
 
 			}
-			else if (ManifoldPoint.InitialPhi < 0)
+			else if (Constraint->GetMinInitialPhi() < 0)
 			{
-				// This is a pre-existing manifold point with some initial penetration to resolve.
-				// If we are currently penetrating less than the inital overlap, reduce the initial overlap
-				// Also resolve initial overlap over time by reducing allowed penetration by MaxDepenetrationVelocity
-				WorldContactInitialPhi = FMath::Max(ManifoldPoint.InitialPhi + MaxDepenetrationVelocity * Dt, WorldContactDeltaNormal);
+				// This is a pre-existing manifold point, but maybe we are still resolving initial penetrations
+				// Don't allow this contact to penetrate any deeper than it currently is. MinInitialPhi will
+				// decrease over time if we have a non-zero depenetration velocity
+				WorldContactInitialPhi = FMath::Max(WorldContactDeltaNormal, Constraint->GetMinInitialPhi());
 			}
 
 			// InitialPhi is only for tracking penetration - cannot be positive
@@ -238,9 +238,7 @@ namespace Chaos
 		const FConstraintSolverBody& Body1 = Solver.SolverBody1();
 
 		// MaxDepenetrationVelocity controls the rate at which initial-overlaps are resolved
-		// If the constraint has a non-negative MaxDepenetrationVelocity we use it, otherwise use the solver setting.
-		// If resultant MaxDepenetrationVelocity is negative, it means depenetrate immediately
-		const FSolverReal MaxDepenetrationVelocity = (Constraint->GetInitialOverlapDepentrationVelocity() >= 0) ? Constraint->GetInitialOverlapDepentrationVelocity() : SolverSettings.DepenetrationVelocity;
+		const FSolverReal MaxDepenetrationVelocity = Constraint->GetInitialOverlapDepenetrationVelocity();
 
 		// The maximum correction we can apply in one frame
 		// @todo(chaos): consider removing this functionality?
@@ -377,6 +375,8 @@ namespace Chaos
 				StaticFrictionRatio,
 				Dt);
 		}
+
+		Constraint->FinalizeSolverResults(Dt);
 	}
 
 	//////////////////////////////////////////////////////////////////////////////////////////////////
diff --git a/Engine/Source/Runtime/Experimental/Chaos/Private/Chaos/Collision/PBDCollisionContainerSolverJacobi.cpp b/Engine/Source/Runtime/Experimental/Chaos/Private/Chaos/Collision/PBDCollisionContainerSolverJacobi.cpp
--- a/Engine/Source/Runtime/Experimental/Chaos/Private/Chaos/Collision/PBDCollisionContainerSolverJacobi.cpp
+++ b/Engine/Source/Runtime/Experimental/Chaos/Private/Chaos/Collision/PBDCollisionContainerSolverJacobi.cpp
@@ -248,6 +248,8 @@ namespace Chaos
 					StaticFrictionRatio,
 					Dt);
 			}
+
+			Constraint->FinalizeSolverResults(Dt);
 		}
 
 		//////////////////////////////////////////////////////////////////////////////////////////////////
diff --git a/Engine/Source/Runtime/Experimental/Chaos/Private/Chaos/Collision/PBDCollisionContainerSolverSimd.cpp b/Engine/Source/Runtime/Experimental/Chaos/Private/Chaos/Collision/PBDCollisionContainerSolverSimd.cpp
--- a/Engine/Source/Runtime/Experimental/Chaos/Private/Chaos/Collision/PBDCollisionContainerSolverSimd.cpp
+++ b/Engine/Source/Runtime/Experimental/Chaos/Private/Chaos/Collision/PBDCollisionContainerSolverSimd.cpp
@@ -292,6 +292,8 @@ namespace Chaos
 						StaticFrictionRatio,
 						Dt);
 				}
+
+				Constraint->FinalizeSolverResults(Dt);
 			}
 		}
 
diff --git a/Engine/Source/Runtime/Experimental/Chaos/Public/Chaos/Collision/ContactPoint.h b/Engine/Source/Runtime/Experimental/Chaos/Public/Chaos/Collision/ContactPoint.h
--- a/Engine/Source/Runtime/Experimental/Chaos/Public/Chaos/Collision/ContactPoint.h
+++ b/Engine/Source/Runtime/Experimental/Chaos/Public/Chaos/Collision/ContactPoint.h
@@ -170,7 +170,6 @@ namespace Chaos
 	{
 	public:
 		FVec3f ShapeContactPoints[2];			// Contact anchor points for friction
-		FRealSingle InitialPhi;					// Non-resolved initial penetration
 	};
 
 	class FManifoldPointResult
diff --git a/Engine/Source/Runtime/Experimental/Chaos/Public/Chaos/Collision/PBDCollisionConstraint.h b/Engine/Source/Runtime/Experimental/Chaos/Public/Chaos/Collision/PBDCollisionConstraint.h
--- a/Engine/Source/Runtime/Experimental/Chaos/Public/Chaos/Collision/PBDCollisionConstraint.h
+++ b/Engine/Source/Runtime/Experimental/Chaos/Public/Chaos/Collision/PBDCollisionConstraint.h
@@ -505,7 +505,7 @@ namespace Chaos
 		bool GetUseIncrementalCollisionDetection() const { return !Flags.bUseManifold || Flags.bUseIncrementalManifold; }
 
 		// Initial overlap depenetration velocity from the maximum of the two bodies
-		FRealSingle GetInitialOverlapDepentrationVelocity() const { return InitialOverlapDepenetrationVelocity; }
+		FRealSingle GetInitialOverlapDepenetrationVelocity() const { return InitialOverlapDepenetrationVelocity; }
 
 		/**
 		* Reset the material properties to those from the shape materials. Called each frame to reset contact modifications to the material.
@@ -817,12 +817,20 @@ namespace Chaos
 			{
 				SavedManifoldPoint->ShapeContactPoints[0] = Anchor0;
 				SavedManifoldPoint->ShapeContactPoints[1] = Anchor1;
-				SavedManifoldPoint->InitialPhi = ManifoldPoint.InitialPhi;
 			}
 
 			MinInitialPhi = FMath::Min(MinInitialPhi, ManifoldPoint.InitialPhi);
 		}
 
+		void FinalizeSolverResults(const FRealSingle Dt)
+		{
+			if (MinInitialPhi < 0)
+			{
+				// Apply depentration velocity to the initial overlap
+				MinInitialPhi = FMath::Min(MinInitialPhi + InitialOverlapDepenetrationVelocity * Dt, 0.0f);
+			}
+		}
+
 		/**
 		 *	A key used to uniquely identify the constraint (it is based on the two particle IDs)
 		 */
