Chaos - improved convex-mesh collision detection

Use the mesh data to fix cloest-contact normal prior to generating the manifold. Previously we would fix the normals on the manifold points. The new method better handles some problem cases like a box embedded in a mesh overlapping an edge where both triangle would prefer to push the box out of the edge rather than the face.

Enabled by default, can be reverted with a cvar (set p.Chaos.Collision.ConvexTriMeshMode 1).

diff --git a/Engine/Source/Runtime/Experimental/Chaos/Private/Chaos/Collision/ConvexTriangleContactPoint.cpp b/Engine/Source/Runtime/Experimental/Chaos/Private/Chaos/Collision/ConvexTriangleContactPoint.cpp
--- a/Engine/Source/Runtime/Experimental/Chaos/Private/Chaos/Collision/ConvexTriangleContactPoint.cpp
+++ b/Engine/Source/Runtime/Experimental/Chaos/Private/Chaos/Collision/ConvexTriangleContactPoint.cpp
@@ -852,7 +852,7 @@ namespace Chaos::Private
 
 	// Generate the contact point and closest feature types between a convex and a triangle
 	template<typename ConvexType>
-	bool FindClosestFeatures(const ConvexType& Convex, const FRigidTransform3& ConvexTransform, const FTriangle& Triangle, const FVec3& TriangleNormal, const FVec3& ConvexRelativeMovement, const FReal CullDistance, FConvexContactPoint& OutContactPoint)
+	bool FindClosestFeatures(const ConvexType& Convex, const FTriangle& Triangle, const FVec3& TriangleNormal, const FVec3& ConvexRelativeMovement, const FReal CullDistance, FConvexContactPoint& OutContactPoint)
 	{
 		// Find the closest point on the convex and triangle that we will use to generate the manifold
 		// NOTE: use an upper limit on cull distance here since the real cull distance depends on the motion against the contact normal which we don't know yet
@@ -886,7 +886,6 @@ namespace Chaos::Private
 
 	template bool FindClosestFeatures(
 		const FImplicitConvex3& Convex,
-		const FRigidTransform3& ConvexTransform, 
 		const FTriangle& Triangle, 
 		const FVec3& TriangleNormal,
 		const FVec3& ConvexRelativeMovement,
@@ -895,7 +894,6 @@ namespace Chaos::Private
 
 	template bool FindClosestFeatures(
 		const TImplicitObjectInstanced<FImplicitConvex3>& Convex,
-		const FRigidTransform3& ConvexTransform,
 		const FTriangle& Triangle,
 		const FVec3& TriangleNormal,
 		const FVec3& ConvexRelativeMovement,
@@ -904,7 +902,6 @@ namespace Chaos::Private
 
 	template bool FindClosestFeatures(
 		const TImplicitObjectScaled<FImplicitConvex3>& Convex,
-		const FRigidTransform3& ConvexTransform,
 		const FTriangle& Triangle,
 		const FVec3& TriangleNormal,
 		const FVec3& ConvexRelativeMovement,
@@ -913,7 +910,6 @@ namespace Chaos::Private
 
 	template bool FindClosestFeatures(
 		const FImplicitBox3& Convex,
-		const FRigidTransform3& ConvexTransform,
 		const FTriangle& Triangle,
 		const FVec3& TriangleNormal,
 		const FVec3& ConvexRelativeMovement,
@@ -922,7 +918,6 @@ namespace Chaos::Private
 
 	template bool FindClosestFeatures(
 		const TImplicitObjectScaled<FImplicitBox3>& Convex,
-		const FRigidTransform3& ConvexTransform,
 		const FTriangle& Triangle,
 		const FVec3& TriangleNormal,
 		const FVec3& ConvexRelativeMovement,
@@ -931,7 +926,6 @@ namespace Chaos::Private
 
 	template bool FindClosestFeatures(
 		const TImplicitObjectInstanced<FImplicitBox3>& Convex,
-		const FRigidTransform3& ConvexTransform,
 		const FTriangle& Triangle,
 		const FVec3& TriangleNormal,
 		const FVec3& ConvexRelativeMovement,
diff --git a/Engine/Source/Runtime/Experimental/Chaos/Private/Chaos/Collision/MeshContactGenerator.cpp b/Engine/Source/Runtime/Experimental/Chaos/Private/Chaos/Collision/MeshContactGenerator.cpp
--- a/Engine/Source/Runtime/Experimental/Chaos/Private/Chaos/Collision/MeshContactGenerator.cpp
+++ b/Engine/Source/Runtime/Experimental/Chaos/Private/Chaos/Collision/MeshContactGenerator.cpp
@@ -276,6 +276,22 @@ namespace Chaos::Private
 				for (int32 OtherLocalTriangleIndex = 0; OtherLocalTriangleIndex < Triangles.Num(); ++OtherLocalTriangleIndex)
 				{
 					const FTriangleExt& OtherTriangle = Triangles[OtherLocalTriangleIndex];
+
+					// We don't collide with boundary vertices. The vertex is a boundary vertex if any of the edges including
+					// the vertex are boundary edges (i.e., not shared between two triangles)
+					int32 OtherVertexIndexB, OtherVertexIndexC;
+					if (OtherTriangle.GetOtherVertexIDs(VertexIndexA, OtherVertexIndexB, OtherVertexIndexC))
+					{
+						const FContactEdgeID EdgeAB = FContactEdgeID(VertexIndexA, OtherVertexIndexB);
+						const FContactEdgeID EdgeCA = FContactEdgeID(OtherVertexIndexC, VertexIndexA);
+						if (!IsSharedEdge(EdgeAB) || !IsSharedEdge(EdgeCA))
+						{
+							ContactPoint.ShapeContactNormal = TriangleNormal;
+							ContactPoint.ShapeContactPoints[0] = ContactPoint.ShapeContactPoints[1] + ContactPoint.Phi * TriangleNormal;
+							break;
+						}
+					}
+
 					FVec3 OtherVertexB, OtherVertexC;
 					if (OtherTriangle.GetOtherVerticesFromID(VertexIndexA, OtherVertexB, OtherVertexC))
 					{
@@ -442,56 +458,103 @@ namespace Chaos::Private
 		if (InOutFeatureType == Private::EConvexFeatureType::Vertex)
 		{
 			check(InOutFeatureIndex != INDEX_NONE);
-
 			const int32 LocalVertexIndex0 = InOutFeatureIndex;
 			const int32 VertexIndexA = Triangle.GetVertexIndex(LocalVertexIndex0);
 			const FVec3& VertexA = Triangle.GetVertex(LocalVertexIndex0);
 
-			// @todo(chaos): the map of Vertex->TriangleIndices would help here but may not be a net win
-			for (int32 OtherLocalTriangleIndex = 0; OtherLocalTriangleIndex < Triangles.Num(); ++OtherLocalTriangleIndex)
+			// Loop over the fan of triangles that share VertexA
+			// If the fan is not complete, we assume we have a boundary vertex and convert to a face collision
+			int32 CurrentLocalTriangleIndex = LocalTriangleIndex;
+			int32 NumTrianglesVisited = 0;
+			while (true)
 			{
-				const FTriangleExt& OtherTriangle = Triangles[OtherLocalTriangleIndex];
+				const FTriangleExt& CurrentTriangle = Triangles[CurrentLocalTriangleIndex];
 
-				// We don't collide with boundary vertices. The vertex is a boundary vertex if any of the edges including
-				// the vertex are boundary edges (i.e., not shared between two triangles)
-				int32 OtherVertexIndexB, OtherVertexIndexC;
-				if (OtherTriangle.GetOtherVertexIDs(VertexIndexA, OtherVertexIndexB, OtherVertexIndexC))
+				// Does the contact normal point into the infinite prism formed by extruding the triangle along the face normal?
+				FVec3 CurrentVertexB, CurrentVertexC;
+				if (CurrentTriangle.GetOtherVerticesFromID(VertexIndexA, CurrentVertexB, CurrentVertexC))
 				{
-					const FContactEdgeID EdgeAB = FContactEdgeID(VertexIndexA, OtherVertexIndexB);
-					const FContactEdgeID EdgeCA = FContactEdgeID(OtherVertexIndexC, VertexIndexA);
-					if (!IsSharedEdge(EdgeAB) || !IsSharedEdge(EdgeCA))
+					const FVec3& CurrentTriangleNormal = CurrentTriangle.GetNormal();
+					const FVec3 CurrentEdge0 = CurrentVertexB - VertexA;
+					const FVec3 CurrentEdge1 = CurrentVertexC - VertexA;
+					const FVec3 CurrentEdgeNormal0 = FVec3::CrossProduct(CurrentTriangleNormal, CurrentEdge0);	// Not normalized
+					const FVec3 CurrentEdgeNormal1 = FVec3::CrossProduct(CurrentEdge1, CurrentTriangleNormal);	// Not normalized
+					const FReal CurrentEdgeSign0 = FVec3::DotProduct(InOutPlaneNormal, CurrentEdgeNormal0);		// Not normalized
+					const FReal CurrentEdgeSign1 = FVec3::DotProduct(InOutPlaneNormal, CurrentEdgeNormal1);		// Not normalized
+					const FReal CurrentEdgeSign0Sq = Utilities::SignedSquare(CurrentEdgeSign0);					// Not normalized
+					const FReal CurrentEdgeSign1Sq = Utilities::SignedSquare(CurrentEdgeSign1);					// Not normalized
+
+					// Scale tolerance to account for non-normalized sign
+					const FReal SignToleranceSq = FReal(1.e-8);
+					const FReal SignTolerance0Sq = SignToleranceSq * CurrentEdge0.SizeSquared();
+					const FReal SignTolerance1Sq = SignToleranceSq * CurrentEdge1.SizeSquared();
+					
+					// If the normal is in the voronoi region of the triangle face, the only valid value for the normal is the face normal
+					if ((CurrentEdgeSign0Sq >= FReal(-SignTolerance0Sq)) && (CurrentEdgeSign1Sq >= FReal(-SignTolerance1Sq)))
 					{
 						InOutFeatureType = Private::EConvexFeatureType::Plane;
 						InOutFeatureIndex = 0;
-						InOutPlaneNormal = TriangleNormal;
+						InOutPlaneNormal = CurrentTriangleNormal;
 						return true;
 					}
-				}
 
-				// Does the contact normal point into the infinite prism formed by extruding the triangle along the face normal?
-				FVec3 OtherVertexB, OtherVertexC;
-				if (OtherTriangle.GetOtherVerticesFromID(VertexIndexA, OtherVertexB, OtherVertexC))
-				{
-					const FVec3& OtherTriangleNormal = OtherTriangle.GetNormal();
-					const FVec3 OtherEdge0 = OtherVertexB - VertexA;
-					const FVec3 OtherEdge1 = VertexA - OtherVertexC;
-					const FVec3 OtherEdgeNormal0 = FVec3::CrossProduct(OtherTriangleNormal, OtherEdge0);	// Not normlized
-					const FVec3 OtherEdgeNormal1 = FVec3::CrossProduct(OtherTriangleNormal, OtherEdge1);	// Not normlized
-					const FReal OtherEdgeSign0 = FVec3::DotProduct(InOutPlaneNormal, OtherEdgeNormal0);
-					const FReal OtherEdgeSign1 = FVec3::DotProduct(InOutPlaneNormal, OtherEdgeNormal1);
-					const FReal OtherEdgeSign0Sq = OtherEdgeSign0 * FMath::Abs(OtherEdgeSign0);
-					const FReal OtherEdgeSign1Sq = OtherEdgeSign1 * FMath::Abs(OtherEdgeSign1);
-					const FReal NormalToleranceSq = FReal(1.e-8);
-					const FReal NormalTolerance0Sq = NormalToleranceSq * OtherEdge0.SizeSquared();
-					const FReal NormalTolerance1Sq = NormalToleranceSq * OtherEdge1.SizeSquared();
-					if ((OtherEdgeSign0Sq >= FReal(-NormalTolerance0Sq)) && (OtherEdgeSign1Sq >= FReal(-NormalTolerance1Sq)))
+					// If the normal is in the voronoi region of the shared edge, use the triangle face normal
+					// @todo(chaos): we should really use the two triangles sharing the edge to map the normal to a valid edge normal.
+					if ((CurrentEdgeSign0Sq <= FReal(SignTolerance0Sq)) && (CurrentEdgeSign1Sq <= FReal(SignTolerance1Sq)))
 					{
-						InOutFeatureType = Private::EConvexFeatureType::Plane;
-						InOutFeatureIndex = 0;
-						InOutPlaneNormal = OtherTriangleNormal;
-						return true;
+						const FReal NormalDotEdge0 = FVec3::DotProduct(InOutPlaneNormal, CurrentEdge0);
+						const FReal NormalDotEdge1 = FVec3::DotProduct(InOutPlaneNormal, CurrentEdge1);
+						if ((NormalDotEdge0 > 0) || (NormalDotEdge1 > 0))
+						{
+							InOutFeatureType = Private::EConvexFeatureType::Plane;
+							InOutFeatureIndex = 0;
+							InOutPlaneNormal = CurrentTriangleNormal;
+							return true;
+						}
 					}
 				}
+
+				// Move to next triangle that shares edge A-B
+				// NOTE: This relies on constent winding and that GetOtherVertexIDs returns indices in winding order
+				int32 CurrentVertexIndexB, CurrentVertexIndexC;
+				if (!ensure(CurrentTriangle.GetOtherVertexIDs(VertexIndexA, CurrentVertexIndexB, CurrentVertexIndexC)))
+				{
+					// Unexpected error - the triangle does not contain VertexA but it always should
+					return false;
+				}
+				const int32 NextLocalTriangleIndex = GetOtherTriangleIndexForEdge(CurrentLocalTriangleIndex, FContactEdgeID(VertexIndexA, CurrentVertexIndexB));
+
+				// Fan complete? We are done
+				if (NextLocalTriangleIndex == LocalTriangleIndex)
+				{
+					break;
+				}
+
+				// Incomplete fan means boundary vertex. We treat boundary vertex contacts as a face contacts
+				// NOTE: we may get false-positives here. Our triangle cache may not contain all the triangles that use
+				// VertexA because the convex bounds does not overlap them all. However, if that happens we cannot
+				// have a real vertex collision (if we overlap the vertex we should overlap all triangles that use it),
+				// so the fallback should be reasonable even for false positives.
+				if (NextLocalTriangleIndex == INDEX_NONE)
+				{
+					InOutFeatureType = Private::EConvexFeatureType::Plane;
+					InOutFeatureIndex = 0;
+					InOutPlaneNormal = TriangleNormal;
+					return true;
+				}
+
+
+				CurrentLocalTriangleIndex = NextLocalTriangleIndex;
+				++NumTrianglesVisited;
+
+				// Malformed mesh?
+				if (!ensure(NumTrianglesVisited < GetNumTriangles()))
+				{
+					InOutFeatureType = Private::EConvexFeatureType::Plane;
+					InOutFeatureIndex = 0;
+					InOutPlaneNormal = TriangleNormal;
+					return true;
+				}
 			}
 		}
 
diff --git a/Engine/Source/Runtime/Experimental/Chaos/Private/Chaos/CollisionOneShotManifolds.cpp b/Engine/Source/Runtime/Experimental/Chaos/Private/Chaos/CollisionOneShotManifolds.cpp
--- a/Engine/Source/Runtime/Experimental/Chaos/Private/Chaos/CollisionOneShotManifolds.cpp
+++ b/Engine/Source/Runtime/Experimental/Chaos/Private/Chaos/CollisionOneShotManifolds.cpp
@@ -137,9 +137,12 @@ namespace Chaos
 
 	bool bChaos_Collision_UseCapsuleTriMesh2 = true;
 	FAutoConsoleVariableRef CVarChaos_Collision_UseCapsuleTriMesh2(TEXT("p.Chaos.Collision.UseCapsuleTriMesh2"), bChaos_Collision_UseCapsuleTriMesh2, TEXT(""));
-
-	bool bChaos_Collision_UseConvexTriMesh2 = true;
-	FAutoConsoleVariableRef CVarChaos_Collision_UseConvexTriMesh2(TEXT("p.Chaos.Collision.UseConvexTriMesh2"), bChaos_Collision_UseConvexTriMesh2, TEXT(""));
+	// @todo(chaos): (2) is the intended solution. The others are for rollback if blocking bugs are found. Remove this once (2) is tested
+	// 0: Use GJK + EPA to generate closest feature, generate multi-point manifold from its normal, and then fix manifold based on mesh topology
+	// 1: Use SAT to generate closest feature, generate multi-point manifold from its normal, and then fix manifold based on mesh topology
+	// 2: Use GJK + SAT to generate closest feature, then fix its normal based on mesh topology, then generate the multi-point manifold
+	int32 Chaos_Collision_ConvexTriMeshMode = 2;
+	FAutoConsoleVariableRef CVarChaos_Collision_UseConvexTriMesh2(TEXT("p.Chaos.Collision.ConvexTriMeshMode"), Chaos_Collision_ConvexTriMeshMode, TEXT(""));
 
 	// true: use GJK and then SAT if the shapes overlap
 	// false: use GJK and then EPA if the shapes overlap
diff --git a/Engine/Source/Runtime/Experimental/Chaos/Private/Chaos/CollisionOneShotManifoldsMeshShapes.cpp b/Engine/Source/Runtime/Experimental/Chaos/Private/Chaos/CollisionOneShotManifoldsMeshShapes.cpp
--- a/Engine/Source/Runtime/Experimental/Chaos/Private/Chaos/CollisionOneShotManifoldsMeshShapes.cpp
+++ b/Engine/Source/Runtime/Experimental/Chaos/Private/Chaos/CollisionOneShotManifoldsMeshShapes.cpp
@@ -25,8 +25,6 @@
 #include "Chaos/Utilities.h"
 #include "ChaosStats.h"
 
-//UE_DISABLE_OPTIMIZATION
-
 namespace Chaos
 {
 	extern bool bChaos_Collision_OneSidedTriangleMesh;
@@ -40,7 +38,7 @@ namespace Chaos
 	extern bool bChaos_Collision_EnableMACDPreManifoldFix;
 
 	extern bool bChaos_Collision_UseCapsuleTriMesh2;
-	extern bool bChaos_Collision_UseConvexTriMesh2;
+	extern int32 Chaos_Collision_ConvexTriMeshMode;
 	
 	namespace CVars
 	{
@@ -91,6 +89,94 @@ namespace Chaos
 			}
 		}
 
+		template<typename ConvexType>
+		void ConstructConvexTriangleOneShotManifold3(
+			const ConvexType& Convex,
+			const FRigidTransform3& ConvexTransform,
+			Private::FMeshContactGenerator& ContactGenerator,
+			const int32 TriangleIndex,
+			const FReal CullDistance,
+			FContactPointManifold& OutContactPoints)
+		{
+			// Triangle relative to the convex at its predicted position P
+			const FTriangle& Triangle = ContactGenerator.GetTriangle(TriangleIndex);
+			const FVec3 TriangleNormal = ContactGenerator.GetTriangleNormal(TriangleIndex);
+
+			// Find the closest feature between the Convex at its initial position X and the Triangle
+			Private::FConvexContactPoint ClosestContact;
+			const bool bFoundClosestContact = Private::FindClosestFeatures(Convex, Triangle, TriangleNormal, FVec3(0), CullDistance, ClosestContact);
+
+			if (bFoundClosestContact)
+			{
+				ClosestContact.Features[0].ObjectIndex = 0;
+				ClosestContact.Features[1].ObjectIndex = TriangleIndex;
+
+				#if CHAOS_DEBUG_DRAW
+				if (CVars::ChaosSolverDebugDrawMeshContacts && FDebugDrawQueue::GetInstance().IsDebugDrawingEnabled())
+				{
+					const FVec3 P = ConvexTransform.TransformPositionNoScale(ClosestContact.ShapeContactPoints[1]);
+					const FVec3 N = ConvexTransform.TransformVectorNoScale(ClosestContact.ShapeContactNormal);
+					FDebugDrawQueue::GetInstance().DrawDebugLine(P, P + 10.0f * N, FColor::Black, false, CVars::ChaosSolverDebugDebugDrawSettings.DrawDuration, (uint8)CVars::ChaosSolverDebugDebugDrawSettings.DrawPriority, 1.5f * CVars::ChaosSolverDebugDebugDrawSettings.LineThickness);
+				}
+				#endif
+
+				// Use the mesh info to correct the normal - this corrects edge and vertex normals if they are
+				// outside the range allowed by the set of triangles sharing the feature
+				if (ContactGenerator.FixFeature(TriangleIndex, ClosestContact.Features[1].FeatureType, ClosestContact.Features[1].PlaneFeatureIndex, ClosestContact.ShapeContactNormal))
+				{
+					// The normal was remapped to the triangle plane
+					ClosestContact.Features[0].FeatureType = Private::EConvexFeatureType::Vertex;
+					ClosestContact.Features[0].PlaneIndex = Convex.GetMostOpposingPlane(ClosestContact.ShapeContactNormal);
+					ClosestContact.Features[0].PlaneFeatureIndex = INDEX_NONE;	// Not needed by ConvexTriangleManifoldFromContact so not worth calculating
+				}
+
+				#if CHAOS_DEBUG_DRAW
+				if (CVars::ChaosSolverDebugDrawMeshContacts && FDebugDrawQueue::GetInstance().IsDebugDrawingEnabled())
+				{
+					const FVec3 P = ConvexTransform.TransformPositionNoScale(ClosestContact.ShapeContactPoints[1]);
+					const FVec3 N = ConvexTransform.TransformVectorNoScale(ClosestContact.ShapeContactNormal);
+					FDebugDrawQueue::GetInstance().DrawDebugLine(P, P + 10.0f * N, FColor::Orange, false, CVars::ChaosSolverDebugDebugDrawSettings.DrawDuration, (uint8)CVars::ChaosSolverDebugDebugDrawSettings.DrawPriority, 1.25f * CVars::ChaosSolverDebugDebugDrawSettings.LineThickness);
+				}
+				#endif
+
+				// Back face culling based on the corrected feature
+				const FReal TriangleDotNormal = FVec3::DotProduct(TriangleNormal, ClosestContact.ShapeContactNormal);
+				if (TriangleDotNormal < 0)
+				{
+					return;
+				}
+
+				// Generate a manifold from on the closest features by projecting the triangle and most opposing convex face onto each other
+				Private::ConvexTriangleManifoldFromContact(Convex, Triangle, TriangleNormal, ClosestContact, CullDistance, OutContactPoints);
+
+				#if CHAOS_DEBUG_DRAW
+				if (CVars::ChaosSolverDebugDrawMeshContacts && FDebugDrawQueue::GetInstance().IsDebugDrawingEnabled())
+				{
+					for (int32 ContactIndex = 0; ContactIndex < OutContactPoints.Num(); ++ContactIndex)
+					{
+						FContactPoint& ContactPoint = OutContactPoints[ContactIndex];
+						const FVec3 P = ConvexTransform.TransformPositionNoScale(ContactPoint.ShapeContactPoints[1]);
+						const FVec3 N = ConvexTransform.TransformVectorNoScale(ContactPoint.ShapeContactNormal);
+						FColor Color = FColor::Black;
+						if (ContactPoint.ContactType == EContactPointType::VertexPlane)
+						{
+							Color = FColor::White;
+						}
+						else if (ContactPoint.ContactType == EContactPointType::PlaneVertex)
+						{
+							Color = FColor::Magenta;
+						}
+						else if (ContactPoint.ContactType == EContactPointType::EdgeEdge)
+						{
+							Color = FColor::Cyan;
+						}
+						FDebugDrawQueue::GetInstance().DrawDebugLine(P, P + 10.0f * N, Color, false, CVars::ChaosSolverDebugDebugDrawSettings.DrawDuration, (uint8)CVars::ChaosSolverDebugDebugDrawSettings.DrawPriority, 1.25f * CVars::ChaosSolverDebugDebugDrawSettings.LineThickness);
+					}
+				}
+				#endif
+			}
+		}
+
 		/**
 		 * @brief Generate a manifold between a convex shape and a single triangle
 		 * Templated so we can specialize for some shape types
@@ -98,7 +184,7 @@ namespace Chaos
 		template<typename ConvexType>
 		void GenerateConvexTriangleOneShotManifold(const ConvexType& Convex, const FTriangle& Triangle, const FReal CullDistance, FContactPointManifold& OutContactPoints)
 		{
-			if (bChaos_Collision_UseConvexTriMesh2)
+			if (Chaos_Collision_ConvexTriMeshMode != 0)
 			{
 				ConstructConvexTriangleOneShotManifold2(Convex, Triangle, CullDistance, OutContactPoints);
 			}
@@ -127,6 +213,47 @@ namespace Chaos
 			ConstructSphereTriangleOneShotManifold(Sphere, Triangle, CullDistance, OutContactPoints);
 		}
 
+		template<typename ConvexType>
+		void GenerateConvexTriangleOneShotManifold(const ConvexType& Convex, const FRigidTransform3& ConvexTransform, Private::FMeshContactGenerator& ContactGenerator, const int32 TriangleIndex, const FReal CullDistance, FContactPointManifold& OutContactPoints)
+		{
+			if (Chaos_Collision_ConvexTriMeshMode == 2)
+			{
+				ConstructConvexTriangleOneShotManifold3(Convex, ConvexTransform, ContactGenerator, TriangleIndex, CullDistance, OutContactPoints);
+				ContactGenerator.SetFixNormalsEnabled(false);
+			}
+			else if (Chaos_Collision_ConvexTriMeshMode == 1)
+			{
+				const FTriangle& Triangle = ContactGenerator.GetTriangle(TriangleIndex);
+				ConstructConvexTriangleOneShotManifold2(Convex, Triangle, CullDistance, OutContactPoints);
+			}
+			else
+			{
+				const FTriangle& Triangle = ContactGenerator.GetTriangle(TriangleIndex);
+				ConstructPlanarConvexTriangleOneShotManifold(Convex, Triangle, CullDistance, OutContactPoints);
+			}
+		}
+
+		template<>
+		void GenerateConvexTriangleOneShotManifold<FImplicitCapsule3>(const FImplicitCapsule3& Capsule, const FRigidTransform3& ConvexTransform, Private::FMeshContactGenerator& ContactGenerator, const int32 TriangleIndex, const FReal CullDistance, FContactPointManifold& OutContactPoints)
+		{
+			const FTriangle& Triangle = ContactGenerator.GetTriangle(TriangleIndex);
+			if (bChaos_Collision_UseCapsuleTriMesh2)
+			{
+				ConstructCapsuleTriangleOneShotManifold2(Capsule, Triangle, CullDistance, OutContactPoints);
+			}
+			else
+			{
+				ConstructCapsuleTriangleOneShotManifold(Capsule, Triangle, CullDistance, OutContactPoints);
+			}
+		}
+
+		template<>
+		void GenerateConvexTriangleOneShotManifold<FImplicitSphere3>(const FImplicitSphere3& Sphere, const FRigidTransform3& ConvexTransform, Private::FMeshContactGenerator& ContactGenerator, const int32 TriangleIndex, const FReal CullDistance, FContactPointManifold& OutContactPoints)
+		{
+			const FTriangle& Triangle = ContactGenerator.GetTriangle(TriangleIndex);
+			ConstructSphereTriangleOneShotManifold(Sphere, Triangle, CullDistance, OutContactPoints);
+		}
+
 		template<typename ConvexType, typename MeshType>
 		void ConstructConvexMeshOneShotManifold2(const ConvexType& Convex, const FRigidTransform3& ConvexTransform, const MeshType& Mesh, const FRigidTransform3& MeshTransform, const FVec3& MeshScale, const FReal CullDistance, Private::FMeshContactGenerator& ContactGenerator)
 		{
@@ -139,12 +266,10 @@ namespace Chaos
 
 			// Generate the contact manifold between Convex and a Triangle
 			const auto& GenerateConvexTriangleContacts =
-				[&Convex, CullDistance](Private::FMeshContactGenerator& ContactGenerator, const int32 TriangleIndex)
+				[&Convex, &ConvexTransform, CullDistance](Private::FMeshContactGenerator& ContactGenerator, const int32 TriangleIndex)
 			{
-				const FTriangle& Triangle = ContactGenerator.GetTriangle(TriangleIndex);
-
 				FContactPointManifold Contacts;
-				GenerateConvexTriangleOneShotManifold(Convex, Triangle, CullDistance, Contacts);
+				GenerateConvexTriangleOneShotManifold(Convex, ConvexTransform, ContactGenerator, TriangleIndex, CullDistance, Contacts);
 
 				ContactGenerator.AddTriangleContacts(TriangleIndex, MakeArrayView(Contacts));
 			};
@@ -178,7 +303,7 @@ namespace Chaos
 
 			// Find the closest feature pair on the triangle and convex
 			Private::FConvexContactPoint ClosestContact;
-			if (Private::FindClosestFeatures(Convex, ConvexTransform, Triangle, TriangleNormal, ConvexRelativeMovement, CullDistance, ClosestContact))
+			if (Private::FindClosestFeatures(Convex, Triangle, TriangleNormal, ConvexRelativeMovement, CullDistance, ClosestContact))
 			{
 				ClosestContact.Features[0].ObjectIndex = 0;
 				ClosestContact.Features[1].ObjectIndex = TriangleIndex;
@@ -332,7 +457,7 @@ namespace Chaos
 
 			// Find the closest feature between the Convex at its initial position X and the Triangle
 			Private::FConvexContactPoint ClosestContact;
-			const bool bFoundClosestContact = Private::FindClosestFeatures(Convex, ConvexTransform, TriangleX, TriangleNormal, ConvexRelativeMovement, CullDistance, ClosestContact);
+			const bool bFoundClosestContact = Private::FindClosestFeatures(Convex, TriangleX, TriangleNormal, ConvexRelativeMovement, CullDistance, ClosestContact);
 
 			if (bFoundClosestContact)
 			{
diff --git a/Engine/Source/Runtime/Experimental/Chaos/Public/Chaos/Collision/ConvexContactPoint.h b/Engine/Source/Runtime/Experimental/Chaos/Public/Chaos/Collision/ConvexContactPoint.h
--- a/Engine/Source/Runtime/Experimental/Chaos/Public/Chaos/Collision/ConvexContactPoint.h
+++ b/Engine/Source/Runtime/Experimental/Chaos/Public/Chaos/Collision/ConvexContactPoint.h
@@ -19,34 +19,44 @@ namespace Chaos::Private
 			Reset();
 		}
 
-		void Reset()
+		inline void Reset()
 		{
 			Phi = InvalidPhi<FRealType>();
 		}
 
-		bool IsSet() const
+		inline bool IsSet() const
 		{
 			return Phi != InvalidPhi<FRealType>();
 		}
 
-		EContactPointType GetContactPointType() const
+		// Convert a feature pair into a contact type (used by callbacks, events, etc)
+		inline EContactPointType GetContactPointType() const
 		{
-			if (IsSet())
+			if (!IsSet())
 			{
-				if ((Features[0].FeatureType == EConvexFeatureType::Plane) && (Features[1].FeatureType == EConvexFeatureType::Vertex))
-				{
-					return EContactPointType::PlaneVertex;
-				}
-				if ((Features[0].FeatureType == EConvexFeatureType::Vertex) && (Features[1].FeatureType == EConvexFeatureType::Plane))
-				{
-					return EContactPointType::VertexPlane;
-				}
-				if ((Features[0].FeatureType == EConvexFeatureType::Edge) && (Features[1].FeatureType == EConvexFeatureType::Edge))
-				{
-					return EContactPointType::EdgeEdge;
-				}
+				return EContactPointType::Unknown;
 			}
-			return EContactPointType::Unknown;
+
+			// Both features should be set to something
+			if ((Features[0].FeatureType == EConvexFeatureType::Unknown) || (Features[1].FeatureType == EConvexFeatureType::Unknown))
+			{
+				return EContactPointType::Unknown;
+			}
+
+			// Plane-Plane, Edge-Plane and Vertex-Plane are treated as Vertex-Plane
+			if (Features[1].FeatureType == EConvexFeatureType::Plane)
+			{
+				return EContactPointType::VertexPlane;
+			}
+
+			// Plane-Vertex and Plane-Edge are treated as Plane-Vertex
+			if (Features[0].FeatureType == EConvexFeatureType::Plane)
+			{
+				return EContactPointType::PlaneVertex;
+			}
+
+			// Vertex-Vertex, Edge-Vertex and Edge-Edge are treated as Edge-Edge
+			return EContactPointType::EdgeEdge;
 		}
 
 		FConvexFeature Features[2];
diff --git a/Engine/Source/Runtime/Experimental/Chaos/Public/Chaos/Collision/ConvexTriangleContactPoint.h b/Engine/Source/Runtime/Experimental/Chaos/Public/Chaos/Collision/ConvexTriangleContactPoint.h
--- a/Engine/Source/Runtime/Experimental/Chaos/Public/Chaos/Collision/ConvexTriangleContactPoint.h
+++ b/Engine/Source/Runtime/Experimental/Chaos/Public/Chaos/Collision/ConvexTriangleContactPoint.h
@@ -11,7 +11,7 @@ namespace Chaos
 	namespace Private
 	{
 		template<typename ConvexType>
-		extern CHAOS_API bool FindClosestFeatures(const ConvexType& Convex, const FRigidTransform3& ConvexTransform, const FTriangle& Triangle, const FVec3& TriangleNormal, const FVec3& ConvexRelativeMovement, const FReal CullDistance, FConvexContactPoint& OutContact);
+		extern CHAOS_API bool FindClosestFeatures(const ConvexType& Convex, const FTriangle& Triangle, const FVec3& TriangleNormal, const FVec3& ConvexRelativeMovement, const FReal CullDistance, FConvexContactPoint& OutContact);
 
 		template <typename ConvexType>
 		extern CHAOS_API void ConvexTriangleManifoldFromContact(const ConvexType& Convex, const FTriangle& Triangle, const FVec3& TriangleNormal, const FConvexContactPoint& Contact, const FReal CullDistance, FContactPointManifold& OutManifold);
diff --git a/Engine/Source/Runtime/Experimental/Chaos/Public/Chaos/Collision/MeshContactGenerator.h b/Engine/Source/Runtime/Experimental/Chaos/Public/Chaos/Collision/MeshContactGenerator.h
--- a/Engine/Source/Runtime/Experimental/Chaos/Public/Chaos/Collision/MeshContactGenerator.h
+++ b/Engine/Source/Runtime/Experimental/Chaos/Public/Chaos/Collision/MeshContactGenerator.h
@@ -50,6 +50,195 @@ namespace Chaos::Private
 		uint32 bUseTwoPassLoop : 1;
 	};
 
+	// A triangle plus some extended data and state
+	// Should be a member of FMeshContactGenerator but that causes natvis issues.
+	class FMeshContactGeneratorTriangle
+	{
+		static constexpr FReal InvalidNormalMarker = std::numeric_limits<FReal>::max();
+
+	public:
+		FMeshContactGeneratorTriangle(const FTriangle& InTriangle, const int32 InTriangleIndex, const int32 InVertexIndex0, const int32 InVertexIndex1, const int32 InVertexIndex2)
+			: Triangle(InTriangle)
+			, Normal(InvalidNormalMarker)
+			, TriangleIndex(InTriangleIndex)
+			, VertexIndices{ InVertexIndex0, InVertexIndex1, InVertexIndex2 }
+			, NumFaceEdgeCollisions(0)
+			, VisitIndex(INDEX_NONE)
+			, bEnabled(true)
+		{
+		}
+
+		// Does this triangle contains the specified vertex? (VertexIndex is an index into the owning mesh's vertices)
+		inline bool HasVertexID(const int32 VertexIndex) const
+		{
+			return (VertexIndices[0] == VertexIndex) || (VertexIndices[1] == VertexIndex) || (VertexIndices[2] == VertexIndex);
+		}
+
+		// Get the vertex position from the vertex ID (not the triangle-local vertex index)
+		inline bool GetVertexWithID(const int32 VertexID, FVec3& OutVertex) const
+		{
+			if (VertexID == VertexIndices[0])
+			{
+				OutVertex = Triangle.GetVertex(0);
+				return true;
+			}
+			else if (VertexID == VertexIndices[1])
+			{
+				OutVertex = Triangle.GetVertex(1);
+				return true;
+			}
+			else if (VertexID == VertexIndices[2])
+			{
+				OutVertex = Triangle.GetVertex(2);
+				return true;
+			}
+			return false;
+		}
+
+		inline bool GetOtherVertexIDs(const int32 VertexID, int32& OutVertexID0, int32& OutVertexID1) const
+		{
+			if (VertexID == VertexIndices[0])
+			{
+				OutVertexID0 = VertexIndices[1];
+				OutVertexID1 = VertexIndices[2];
+				return true;
+			}
+			else if (VertexID == VertexIndices[1])
+			{
+				OutVertexID0 = VertexIndices[2];
+				OutVertexID1 = VertexIndices[0];
+				return true;
+			}
+			else if (VertexID == VertexIndices[2])
+			{
+				OutVertexID0 = VertexIndices[0];
+				OutVertexID1 = VertexIndices[1];
+				return true;
+			}
+			return false;
+		}
+
+		// Get the positions of the other two vertices in the triangle. (VertexIndex is an index into the owning mesh's vertices)
+		inline bool GetOtherVerticesFromID(const int32 VertexID, FVec3& OutVertex0, FVec3& OutVertex1) const
+		{
+			if (VertexID == VertexIndices[0])
+			{
+				OutVertex0 = Triangle.GetVertex(1);
+				OutVertex1 = Triangle.GetVertex(2);
+				return true;
+			}
+			else if (VertexID == VertexIndices[1])
+			{
+				OutVertex0 = Triangle.GetVertex(2);
+				OutVertex1 = Triangle.GetVertex(0);
+				return true;
+			}
+			else if (VertexID == VertexIndices[2])
+			{
+				OutVertex0 = Triangle.GetVertex(0);
+				OutVertex1 = Triangle.GetVertex(1);
+				return true;
+			}
+			return false;
+		}
+
+		int32 GetLocalVertexIndexAt(const FVec3& InPos, const FReal InTolerance) const
+		{
+			for (int32 LocalVertexIndex = 0; LocalVertexIndex < 3; ++LocalVertexIndex)
+			{
+				if (FVec3::IsNearlyEqual(GetVertex(LocalVertexIndex), InPos, InTolerance))
+				{
+					return LocalVertexIndex;
+				}
+			}
+			return INDEX_NONE;
+		};
+
+		int32 GetVertexIDAt(const FVec3& InPos, const FReal InTolerance) const
+		{
+			const int32 LocalVertexIndex = GetLocalVertexIndexAt(InPos, InTolerance);
+			if (LocalVertexIndex != INDEX_NONE)
+			{
+				return VertexIndices[LocalVertexIndex];
+			}
+			return INDEX_NONE;
+		};
+
+		const FTriangle& GetTriangle() const
+		{
+			return Triangle;
+		}
+
+		// Get the vertex for the triangle-local vertex index [0,2]
+		const FVec3& GetVertex(const int32 LocalVertexIndex) const
+		{
+			return Triangle.GetVertex(LocalVertexIndex);
+		}
+
+		int32 GetTriangleIndex() const
+		{
+			return TriangleIndex;
+		}
+
+		int32 GetVertexIndex(const int32 LocalIndex) const
+		{
+			return VertexIndices[LocalIndex];
+		}
+
+		const FVec3& GetNormal() const
+		{
+			if (Normal.X == InvalidNormalMarker)
+			{
+				Normal = Triangle.GetNormal();
+			}
+			return Normal;
+		}
+
+		FVec3 GetCentroid() const
+		{
+			return Triangle.GetCentroid();
+		}
+
+		void SetVisitIndex(const int8 InVisitIndex)
+		{
+			VisitIndex = InVisitIndex;
+		}
+
+		int8 GetVisitIndex() const
+		{
+			return VisitIndex;
+		}
+
+		void SetEnabled(const bool bInEnabled)
+		{
+			bEnabled = bInEnabled;
+		}
+
+		bool GetIsEnabled() const
+		{
+			return bEnabled;
+		}
+
+		void AddFaceEdgeCollision()
+		{
+			++NumFaceEdgeCollisions;
+		}
+
+		int32 GetNumFaceEdgeCollisions() const
+		{
+			return NumFaceEdgeCollisions;
+		}
+
+	private:
+		FTriangle Triangle;
+		mutable FVec3 Normal;
+		int32 TriangleIndex;
+		int32 VertexIndices[3];
+		int8 NumFaceEdgeCollisions;
+		int8 VisitIndex;
+		bool bEnabled;
+	};
+
 	/**
 	* Generate contacts between a collision shape and the triangles from a mesh.
 	*/
@@ -115,193 +304,7 @@ namespace Chaos::Private
 		}
 
 	private:
-		// A triangle plus some extended data and state
-		class FTriangleExt
-		{
-			static constexpr FReal InvalidNormalMarker = std::numeric_limits<FReal>::max();
-
-		public:
-			FTriangleExt(const FTriangle& InTriangle, const int32 InTriangleIndex, const int32 InVertexIndex0, const int32 InVertexIndex1, const int32 InVertexIndex2)
-				: Triangle(InTriangle)
-				, Normal(InvalidNormalMarker)
-				, TriangleIndex(InTriangleIndex)
-				, VertexIndices{ InVertexIndex0, InVertexIndex1, InVertexIndex2 }
-				, NumFaceEdgeCollisions(0)
-				, VisitIndex(INDEX_NONE)
-				, bEnabled(true)
-			{
-			}
-
-			// Does this triangle contains the specified vertex? (VertexIndex is an index into the owning mesh's vertices)
-			inline bool HasVertexID(const int32 VertexIndex) const
-			{
-				return (VertexIndices[0] == VertexIndex) || (VertexIndices[1] == VertexIndex) || (VertexIndices[2] == VertexIndex);
-			}
-
-			// Get the vertex position from the vertex ID (not the triangle-local vertex index)
-			inline bool GetVertexWithID(const int32 VertexID, FVec3& OutVertex) const
-			{
-				if (VertexID == VertexIndices[0])
-				{
-					OutVertex = Triangle.GetVertex(0);
-					return true;
-				}
-				else if (VertexID == VertexIndices[1])
-				{
-					OutVertex = Triangle.GetVertex(1);
-					return true;
-				}
-				else if (VertexID == VertexIndices[2])
-				{
-					OutVertex = Triangle.GetVertex(2);
-					return true;
-				}
-				return false;
-			}
-
-			inline bool GetOtherVertexIDs(const int32 VertexID, int32& OutVertexID0, int32& OutVertexID1) const
-			{
-				if (VertexID == VertexIndices[0])
-				{
-					OutVertexID0 = VertexIndices[1];
-					OutVertexID1 = VertexIndices[2];
-					return true;
-				}
-				else if (VertexID == VertexIndices[1])
-				{
-					OutVertexID0 = VertexIndices[2];
-					OutVertexID1 = VertexIndices[0];
-					return true;
-				}
-				else if (VertexID == VertexIndices[2])
-				{
-					OutVertexID0 = VertexIndices[0];
-					OutVertexID1 = VertexIndices[1];
-					return true;
-				}
-				return false;
-			}
-
-			// Get the positions of the other two vertices in the triangle. (VertexIndex is an index into the owning mesh's vertices)
-			inline bool GetOtherVerticesFromID(const int32 VertexID, FVec3& OutVertex0, FVec3& OutVertex1) const
-			{
-				if (VertexID == VertexIndices[0])
-				{
-					OutVertex0 = Triangle.GetVertex(1);
-					OutVertex1 = Triangle.GetVertex(2);
-					return true;
-				}
-				else if (VertexID == VertexIndices[1])
-				{
-					OutVertex0 = Triangle.GetVertex(2);
-					OutVertex1 = Triangle.GetVertex(0);
-					return true;
-				}
-				else if (VertexID == VertexIndices[2])
-				{
-					OutVertex0 = Triangle.GetVertex(0);
-					OutVertex1 = Triangle.GetVertex(1);
-					return true;
-				}
-				return false;
-			}
-
-			int32 GetLocalVertexIndexAt(const FVec3& InPos, const FReal InTolerance) const
-			{
-				for (int32 LocalVertexIndex = 0; LocalVertexIndex < 3; ++LocalVertexIndex)
-				{
-					if (FVec3::IsNearlyEqual(GetVertex(LocalVertexIndex), InPos, InTolerance))
-					{
-						return LocalVertexIndex;
-					}
-				}
-				return INDEX_NONE;
-			};
-
-			int32 GetVertexIDAt(const FVec3& InPos, const FReal InTolerance) const
-			{
-				const int32 LocalVertexIndex = GetLocalVertexIndexAt(InPos, InTolerance);
-				if (LocalVertexIndex != INDEX_NONE)
-				{
-					return VertexIndices[LocalVertexIndex];
-				}
-				return INDEX_NONE;
-			};
-
-			const FTriangle& GetTriangle() const
-			{
-				return Triangle;
-			}
-
-			// Get the vertex for the triangle-local vertex index [0,2]
-			const FVec3& GetVertex(const int32 LocalVertexIndex) const
-			{
-				return Triangle.GetVertex(LocalVertexIndex);
-			}
-
-			int32 GetTriangleIndex() const
-			{
-				return TriangleIndex;
-			}
-
-			int32 GetVertexIndex(const int32 LocalIndex) const
-			{
-				return VertexIndices[LocalIndex];
-			}
-
-			const FVec3& GetNormal() const
-			{
-				if (Normal.X == InvalidNormalMarker)
-				{
-					Normal = Triangle.GetNormal();
-				}
-				return Normal;
-			}
-
-			FVec3 GetCentroid() const
-			{
-				return Triangle.GetCentroid();
-			}
-
-			void SetVisitIndex(const int8 InVisitIndex)
-			{
-				VisitIndex = InVisitIndex;
-			}
-
-			int8 GetVisitIndex() const
-			{
-				return VisitIndex;
-			}
-
-			void SetEnabled(const bool bInEnabled)
-			{
-				bEnabled = bInEnabled;
-			}
-
-			bool GetIsEnabled() const
-			{
-				return bEnabled;
-			}
-
-			void AddFaceEdgeCollision()
-			{
-				++NumFaceEdgeCollisions;
-			}
-
-			int32 GetNumFaceEdgeCollisions() const
-			{
-				return NumFaceEdgeCollisions;
-			}
-
-		private:
-			FTriangle Triangle;
-			mutable FVec3 Normal;
-			int32 TriangleIndex;
-			int32 VertexIndices[3];
-			int8 NumFaceEdgeCollisions;
-			int8 VisitIndex;
-			bool bEnabled;
-		};
+		using FTriangleExt = FMeshContactGeneratorTriangle;
 
 		// A contact index combined with a flag to indicate if the normal is roughly along the triangle face
 		struct FVertexContactIndex
