Dataflow:
- Implemented BlueprintToCollection node

diff --git a/Engine/Plugins/Experimental/GeometryCollectionPlugin/Source/GeometryCollectionNodes/Private/Dataflow/Nodes/GeometryCollectionAssetNodes.cpp b/Engine/Plugins/Experimental/GeometryCollectionPlugin/Source/GeometryCollectionNodes/Private/Dataflow/Nodes/GeometryCollectionAssetNodes.cpp
--- a/Engine/Plugins/Experimental/GeometryCollectionPlugin/Source/GeometryCollectionNodes/Private/Dataflow/Nodes/GeometryCollectionAssetNodes.cpp
+++ b/Engine/Plugins/Experimental/GeometryCollectionPlugin/Source/GeometryCollectionNodes/Private/Dataflow/Nodes/GeometryCollectionAssetNodes.cpp
@@ -10,6 +10,7 @@
 #include "GeometryCollection/GeometryCollectionClusteringUtility.h"
 #include "GeometryCollection/Facades/CollectionHierarchyFacade.h"
 #include "GeometryCollection/Facades/CollectionInstancedMeshFacade.h"
+#include "PreviewScene.h"
 
 namespace Dataflow
 {
@@ -23,6 +24,7 @@ namespace Dataflow
 		DATAFLOW_NODE_REGISTER_CREATION_FACTORY(FCreateGeometryCollectionFromSourcesDataflowNode);
 		DATAFLOW_NODE_REGISTER_CREATION_FACTORY(FStaticMeshToCollectionDataflowNode);
 		DATAFLOW_NODE_REGISTER_CREATION_FACTORY(FGeometryCollectionToCollectionDataflowNode);
+		DATAFLOW_NODE_REGISTER_CREATION_FACTORY(FBlueprintToCollectionDataflowNode);
 
 		// Terminal
 		DATAFLOW_NODE_REGISTER_CREATION_FACTORY_NODE_COLORS_BY_CATEGORY("Terminal", FLinearColor(0.f, 0.f, 0.f), CDefaultNodeBodyTintColor);
@@ -283,3 +285,47 @@ void FGeometryCollectionToCollectionDataflowNode::Evaluate(Dataflow::FContext& C
 	SetValue(Context, MoveTemp(OutMaterials), &Materials);
 	SetValue(Context, MoveTemp(OutInstancedMeshes), &InstancedMeshes);
 }
+
+// ===========================================================================================================================
+
+FBlueprintToCollectionDataflowNode::FBlueprintToCollectionDataflowNode(const Dataflow::FNodeParameters& InParam, FGuid InGuid)
+	: FDataflowNode(InParam, InGuid)
+{
+	RegisterOutputConnection(&Collection);
+	RegisterOutputConnection(&Materials);
+	RegisterOutputConnection(&InstancedMeshes);
+}
+
+void FBlueprintToCollectionDataflowNode::Evaluate(Dataflow::FContext& Context, const FDataflowOutput* Out) const
+{
+	ensure(Out->IsA(&Collection) || Out->IsA(&Materials) || Out->IsA(&InstancedMeshes));
+
+	FManagedArrayCollection OutCollection;
+	TArray<TObjectPtr<UMaterial>> OutMaterials;
+	TArray<FGeometryCollectionAutoInstanceMesh> OutInstancedMeshes;
+
+	if (Blueprint)
+	{
+		if (TUniquePtr<FPreviewScene> PreviewScene = MakeUnique<FPreviewScene>(FPreviewScene::ConstructionValues()))
+		{
+			if (UWorld* PreviewWorld = PreviewScene->GetWorld())
+			{
+				FActorSpawnParameters SpawnInfo;
+				SpawnInfo.SpawnCollisionHandlingOverride = ESpawnActorCollisionHandlingMethod::AlwaysSpawn;
+				SpawnInfo.bNoFail = true;
+				SpawnInfo.ObjectFlags = RF_Transient;
+
+				if (AActor* PreviewActor = PreviewWorld->SpawnActor(Blueprint->GeneratedClass, nullptr, SpawnInfo))
+				{
+					FGeometryCollectionEngineConversion::ConvertActorToGeometryCollection(PreviewActor, OutCollection, OutMaterials, OutInstancedMeshes, bSplitComponents);
+				}
+			}
+		}
+	}
+
+	// Set Outputs
+	SetValue(Context, MoveTemp(OutCollection), &Collection);
+	SetValue(Context, MoveTemp(OutMaterials), &Materials);
+	SetValue(Context, MoveTemp(OutInstancedMeshes), &InstancedMeshes);
+}
+
diff --git a/Engine/Plugins/Experimental/GeometryCollectionPlugin/Source/GeometryCollectionNodes/Public/Dataflow/Nodes/GeometryCollectionAssetNodes.h b/Engine/Plugins/Experimental/GeometryCollectionPlugin/Source/GeometryCollectionNodes/Public/Dataflow/Nodes/GeometryCollectionAssetNodes.h
--- a/Engine/Plugins/Experimental/GeometryCollectionPlugin/Source/GeometryCollectionNodes/Public/Dataflow/Nodes/GeometryCollectionAssetNodes.h
+++ b/Engine/Plugins/Experimental/GeometryCollectionPlugin/Source/GeometryCollectionNodes/Public/Dataflow/Nodes/GeometryCollectionAssetNodes.h
@@ -6,11 +6,13 @@
 #include "Dataflow/DataflowEngine.h"
 #include "GeometryCollection/ManagedArrayCollection.h"
 #include "GeometryCollection/GeometryCollectionObject.h"
+#include "Engine/Blueprint.h"
 
 #include "GeometryCollectionAssetNodes.generated.h"
 
 class UGeometryCollection;
 class UStaticMesh;
+class UBlueprint;
 
 USTRUCT(meta = (DataflowGeometryCollection, DataflowTerminal))
 struct FGeometryCollectionTerminalDataflowNode : public FDataflowTerminalNode
@@ -185,6 +187,43 @@ public:
 	virtual void Evaluate(Dataflow::FContext& Context, const FDataflowOutput* Out) const override;
 };
 
+
+/**
+ * Create a geometry collection from an asset
+ */
+ USTRUCT(meta = (DataflowContext = "GeometryCollection", DataflowGeometryCollection, DataflowTerminal))
+ struct FBlueprintToCollectionDataflowNode : public FDataflowNode
+ {
+ 	GENERATED_USTRUCT_BODY()
+ 	DATAFLOW_NODE_DEFINE_INTERNAL(FBlueprintToCollectionDataflowNode, "BlueprintToCollection", "GeometryCollection|Asset", "")
+	DATAFLOW_NODE_RENDER_TYPE(FGeometryCollection::StaticType(), "Collection")
+
+ public:
+ 	/** Asset input */
+ 	UPROPERTY(EditAnywhere, Category = "Asset");
+	TObjectPtr<UBlueprint> Blueprint;
+ 
+ 	/** Split components */
+ 	UPROPERTY(EditAnywhere, Category = "Asset");
+ 	bool bSplitComponents = false;
+ 
+ 	/** Geometry collection newly created */
+ 	UPROPERTY(meta = (DataflowOutput, DisplayName = "Collection"))
+ 	FManagedArrayCollection Collection;
+ 
+ 	/** Materials array to use for this asset */
+ 	UPROPERTY(meta = (DataflowOutput, DisplayName = "Materials"))
+ 	TArray<TObjectPtr<UMaterial>> Materials;
+ 
+ 	/** Array of instanced meshes*/
+ 	UPROPERTY(meta = (DataflowOutput, DisplayName = "InstancedMeshes"))
+ 	TArray<FGeometryCollectionAutoInstanceMesh> InstancedMeshes;
+ 
+	FBlueprintToCollectionDataflowNode(const Dataflow::FNodeParameters& InParam, FGuid InGuid = FGuid::NewGuid());
+ 	virtual void Evaluate(Dataflow::FContext& Context, const FDataflowOutput* Out) const override;
+ };
+
+ 
 namespace Dataflow
 {
 	void GeometryCollectionEngineAssetNodes();
diff --git a/Engine/Source/Runtime/Experimental/GeometryCollectionEngine/Private/GeometryCollection/GeometryCollectionEngineConversion.cpp b/Engine/Source/Runtime/Experimental/GeometryCollectionEngine/Private/GeometryCollection/GeometryCollectionEngineConversion.cpp
--- a/Engine/Source/Runtime/Experimental/GeometryCollectionEngine/Private/GeometryCollection/GeometryCollectionEngineConversion.cpp
+++ b/Engine/Source/Runtime/Experimental/GeometryCollectionEngine/Private/GeometryCollection/GeometryCollectionEngineConversion.cpp
@@ -1547,7 +1547,7 @@ void FGeometryCollectionEngineConversion::ConvertStaticMeshToGeometryCollection(
 		NewInstanceMesh.Materials = SourceMaterials;
 		OutInstancedMeshes.Emplace(NewInstanceMesh);
 
-		bool bAddInternalMaterials = true;
+		bool bAddInternalMaterials = false;
 
 		NewGeometryCollection->GeometrySource.Emplace(SourceSoftObjectPath, ComponentTransform, SourceMaterials, bSplitComponents, bSetInternalFromMaterialIndex);
 		FGeometryCollectionEngineConversion::AppendStaticMesh(StaticMesh, SourceMaterials, ComponentTransform, NewGeometryCollection, false, bAddInternalMaterials, bSplitComponents, bSetInternalFromMaterialIndex);
@@ -1588,4 +1588,80 @@ void FGeometryCollectionEngineConversion::ConvertGeometryCollectionToGeometryCol
 	}
 }
 
+void FGeometryCollectionEngineConversion::ConvertActorToGeometryCollection(const AActor* Actor, FManagedArrayCollection& OutCollection, TArray<TObjectPtr<UMaterial>>& OutMaterials, TArray<FGeometryCollectionAutoInstanceMesh>& OutInstancedMeshes, bool bSplitComponents)
+{
+#if WITH_EDITORONLY_DATA
+	const FTransform ActorTransform(Actor->GetTransform());
+
+	if (UGeometryCollection* NewGeometryCollection = NewObject<UGeometryCollection>())
+	{
+		TInlineComponentArray<UStaticMeshComponent*> StaticMeshComponents(Actor);
+		for (UStaticMeshComponent* StaticMeshComponent : StaticMeshComponents)
+		{
+			if (StaticMeshComponent)
+			{
+				if (const UStaticMesh* ComponentStaticMesh = StaticMeshComponent->GetStaticMesh())
+				{
+					NewGeometryCollection->EnableNanite |= ComponentStaticMesh->IsNaniteEnabled();
+
+					FTransform ComponentTransform(StaticMeshComponent->GetComponentTransform());
+					ComponentTransform.SetTranslation((ComponentTransform.GetTranslation() - ActorTransform.GetTranslation()));
+
+					// Record the contributing source on the asset.
+					FSoftObjectPath SourceSoftObjectPath(ComponentStaticMesh);
+					TArray<TObjectPtr<UMaterialInterface>> SourceMaterials(StaticMeshComponent->GetMaterials());
+
+					NewGeometryCollection->GeometrySource.Emplace(SourceSoftObjectPath, ComponentTransform, SourceMaterials, bSplitComponents, true/*bSetInternalFromMaterialIndex*/);
+
+					FGeometryCollectionEngineConversion::AppendStaticMesh(ComponentStaticMesh, SourceMaterials, ComponentTransform, NewGeometryCollection, false/*bReindexMaterials*/, false/*bAddInternalMaterials*/, bSplitComponents, true/*bSetInternalFromMaterialIndex*/);
+				}
+			}
+		}
+
+		TInlineComponentArray<UGeometryCollectionComponent*> GeometryCollectionComponents(Actor);
+		for (UGeometryCollectionComponent* GeometryCollectionComponent : GeometryCollectionComponents)
+		{
+			if (GeometryCollectionComponent)
+			{
+				if (const UGeometryCollection* RestCollection = GeometryCollectionComponent->GetRestCollection())
+				{
+					NewGeometryCollection->EnableNanite |= RestCollection->EnableNanite;
+
+					FTransform ComponentTransform(GeometryCollectionComponent->GetComponentTransform());
+					ComponentTransform.SetTranslation((ComponentTransform.GetTranslation() - ActorTransform.GetTranslation()));
+
+					// Record the contributing source on the asset.
+					FSoftObjectPath SourceSoftObjectPath(RestCollection);
+
+					int32 NumMaterials = GeometryCollectionComponent->GetNumMaterials();
+					TArray<TObjectPtr<UMaterialInterface>> SourceMaterials;
+					SourceMaterials.SetNum(NumMaterials);
+					for (int32 MaterialIndex = 0; MaterialIndex < NumMaterials; ++MaterialIndex)
+					{
+						SourceMaterials[MaterialIndex] = GeometryCollectionComponent->GetMaterial(MaterialIndex);
+					}
+					NewGeometryCollection->GeometrySource.Emplace(SourceSoftObjectPath, ComponentTransform, SourceMaterials, bSplitComponents, true/*bSetInternalFromMaterialIndex*/);
+
+					FGeometryCollectionEngineConversion::AppendGeometryCollection(RestCollection, GeometryCollectionComponent, ComponentTransform, NewGeometryCollection, false /*bReindexMaterials*/);
+				}
+			}
+		}
+
+		NewGeometryCollection->InitializeMaterials();
+
+		// InstanceMeshes
+		OutInstancedMeshes.Append(NewGeometryCollection->AutoInstanceMeshes);
+
+		// Materials
+		for (auto& Material : NewGeometryCollection->Materials)
+		{
+			OutMaterials.Emplace(Material->GetMaterial());
+		}
+
+		TSharedPtr<FGeometryCollection> OutCollectionPtr = NewGeometryCollection->GetGeometryCollection();
+		OutCollectionPtr->CopyTo(&OutCollection);
+	}
+#endif //WITH_EDITORONLY_DATA
+}
+
 #undef LOCTEXT_NAMESPACE 
diff --git a/Engine/Source/Runtime/Experimental/GeometryCollectionEngine/Public/GeometryCollection/GeometryCollectionEngineConversion.h b/Engine/Source/Runtime/Experimental/GeometryCollectionEngine/Public/GeometryCollection/GeometryCollectionEngineConversion.h
--- a/Engine/Source/Runtime/Experimental/GeometryCollectionEngine/Public/GeometryCollection/GeometryCollectionEngineConversion.h
+++ b/Engine/Source/Runtime/Experimental/GeometryCollectionEngine/Public/GeometryCollection/GeometryCollectionEngineConversion.h
@@ -195,4 +195,14 @@ public:
 	*/	
 	static GEOMETRYCOLLECTIONENGINE_API void ConvertGeometryCollectionToGeometryCollection(const TObjectPtr<UGeometryCollection> InGeometryCollectionAssetPtr, FManagedArrayCollection& OutCollection, TArray<TObjectPtr<UMaterial>>& OutMaterials, TArray<FGeometryCollectionAutoInstanceMesh>& OutInstancedMeshes);
 
+	/**
+	*  Converts an actor to an FManagedArrayCollection, and arrays of materials and instanced meshes
+	*  @param Actor : actor input to convert
+	*  @param OutCollection : FGeometryCollection output
+	*  @param OutMaterials : Materials from the UGeometryCollection
+	*  @param OutInstancedMeshes : InstancedMeshes
+	*  @param bSplitComponents : Split the components
+	*/
+	static GEOMETRYCOLLECTIONENGINE_API void ConvertActorToGeometryCollection(const AActor* Actor, FManagedArrayCollection& OutCollection, TArray<TObjectPtr<UMaterial>>& OutMaterials, TArray<FGeometryCollectionAutoInstanceMesh>& OutInstancedMeshes, bool bSplitComponents = false);
+
 };
