Add dataflow reroute node

diff --git a/Engine/Plugins/Experimental/Dataflow/Source/DataflowEditor/Private/Dataflow/DataflowCollectionSpreadSheet.cpp b/Engine/Plugins/Experimental/Dataflow/Source/DataflowEditor/Private/Dataflow/DataflowCollectionSpreadSheet.cpp
--- a/Engine/Plugins/Experimental/Dataflow/Source/DataflowEditor/Private/Dataflow/DataflowCollectionSpreadSheet.cpp
+++ b/Engine/Plugins/Experimental/Dataflow/Source/DataflowEditor/Private/Dataflow/DataflowCollectionSpreadSheet.cpp
@@ -31,7 +31,7 @@ void FDataflowCollectionSpreadSheet::UpdateViewData()
 						FName Name = Output->GetName();
 						FName Type = Output->GetType();
 
-						if (Output->GetType() == "FManagedArrayCollection")
+						if (Output->GetType() == "FManagedArrayCollection" && !Output->IsAnyType())
 						{
 							const FManagedArrayCollection& Value = Output->GetValue<FManagedArrayCollection>(*GetContext(), FManagedArrayCollection());
 
diff --git a/Engine/Plugins/Experimental/Dataflow/Source/DataflowEditor/Private/Dataflow/DataflowEditorCommands.cpp b/Engine/Plugins/Experimental/Dataflow/Source/DataflowEditor/Private/Dataflow/DataflowEditorCommands.cpp
--- a/Engine/Plugins/Experimental/Dataflow/Source/DataflowEditor/Private/Dataflow/DataflowEditorCommands.cpp
+++ b/Engine/Plugins/Experimental/Dataflow/Source/DataflowEditor/Private/Dataflow/DataflowEditorCommands.cpp
@@ -935,18 +935,22 @@ void FDataflowEditorCommands::PasteNodes(UDataflow* Graph, const TSharedPtr<SDat
 					{
 						if (TSharedPtr<FDataflowNode> DataflowNodeTo = DataflowGraph->FindBaseNode(GuidIn))
 						{
-							FDataflowInput* InputConnection = DataflowNodeTo->FindInput(InputputName);
-							FDataflowOutput* OutputConnection = DataflowNodeFrom->FindOutput(OutputputName);
-
-							DataflowGraph->Connect(OutputConnection, InputConnection);
-
-							if (UEdGraphPin* OutputPin = GetPin(EdNodeMap[NodeOut], EEdGraphPinDirection::EGPD_Output, OutputputName))
+							// first connect the edgraph as this may affect the dataflow inputs ( for AnyType )
+							UEdGraphPin* OutputPin = GetPin(EdNodeMap[NodeOut], EEdGraphPinDirection::EGPD_Output, OutputputName);
+							UEdGraphPin* InputPin = GetPin(EdNodeMap[NodeIn], EEdGraphPinDirection::EGPD_Input, InputputName);
+							if (OutputPin && InputPin)
 							{
-								if (UEdGraphPin* InputPin = GetPin(EdNodeMap[NodeIn], EEdGraphPinDirection::EGPD_Input, InputputName))
+								if (UEdGraph* EdGraph = EdNodeMap[NodeOut]->GetGraph())
 								{
-									OutputPin->MakeLinkTo(InputPin);
+									EdGraph->GetSchema()->TryCreateConnection(OutputPin, InputPin);
 								}
 							}
+
+							// now connect the dataflow
+							FDataflowInput* InputConnection = DataflowNodeTo->FindInput(InputputName);
+							FDataflowOutput* OutputConnection = DataflowNodeFrom->FindOutput(OutputputName);
+							DataflowGraph->Connect(OutputConnection, InputConnection);
+
 						}
 					}
 				}
diff --git a/Engine/Plugins/Experimental/Dataflow/Source/DataflowEditor/Private/Dataflow/DataflowSNodeFactories.cpp b/Engine/Plugins/Experimental/Dataflow/Source/DataflowEditor/Private/Dataflow/DataflowSNodeFactories.cpp
--- a/Engine/Plugins/Experimental/Dataflow/Source/DataflowEditor/Private/Dataflow/DataflowSNodeFactories.cpp
+++ b/Engine/Plugins/Experimental/Dataflow/Source/DataflowEditor/Private/Dataflow/DataflowSNodeFactories.cpp
@@ -6,11 +6,19 @@
 #include "EdGraphNode_Comment.h"
 #include "Dataflow/DataflowSCommentNode.h"
 #include "Dataflow/DataflowSchema.h"
+#include "SGraphNodeKnot.h"
 
 TSharedPtr<class SGraphNode> FDataflowSNodeFactory::CreateNode(UEdGraphNode* InNode) const
 {
 	if (UDataflowEdNode* Node = Cast<UDataflowEdNode>(InNode))
 	{
+		if (const TSharedPtr<const FDataflowNode> DataflowNode = Node->GetDataflowNode())
+		{
+			if (DataflowNode->GetType() == FDataflowReRouteNode::StaticType())
+			{
+				return SNew(SGraphNodeKnot, Node);
+			}
+		}
 		return SNew(SDataflowEdNode, Node);
 	}	
 	else if (UEdGraphNode_Comment* CommentNode = Cast<UEdGraphNode_Comment>(InNode))
diff --git a/Engine/Plugins/Experimental/Dataflow/Source/DataflowEditor/Private/Dataflow/DataflowSchema.cpp b/Engine/Plugins/Experimental/Dataflow/Source/DataflowEditor/Private/Dataflow/DataflowSchema.cpp
--- a/Engine/Plugins/Experimental/Dataflow/Source/DataflowEditor/Private/Dataflow/DataflowSchema.cpp
+++ b/Engine/Plugins/Experimental/Dataflow/Source/DataflowEditor/Private/Dataflow/DataflowSchema.cpp
@@ -175,7 +175,9 @@ const FPinConnectionResponse UDataflowSchema::CanCreateConnection(const UEdGraph
 			if (PinA->GetOwningNode() != PinB->GetOwningNode())
 			{
 				// Make sure types match. 
-				if (PinA->PinType == PinB->PinType)
+				const bool bIsAnyTypeA = PinA->PinType.PinCategory == FDataflowAnyType::TypeName;
+				const bool bIsAnyTypeB = PinB->PinType.PinCategory == FDataflowAnyType::TypeName;
+				if (PinA->PinType == PinB->PinType || (bIsAnyTypeA != bIsAnyTypeB))
 				{
 					// cycle checking on connect
 					if (!HasLoopIfConnected(PinA->GetOwningNode(), PinB->GetOwningNode()))
@@ -264,9 +266,36 @@ FLinearColor UDataflowSchema::GetTypeColor(const FName& Type)
 	return Settings->DefaultPinTypeColor;
 }
 
-//void UDataflowSchema::OnPinConnectionDoubleCicked(UEdGraphPin* PinA, UEdGraphPin* PinB, const FVector2D& GraphPosition) const
-//{
-//}
+static void CreateAndConnectNewReRouteNode(UEdGraphPin* FromPin, UEdGraphPin* ToPin, const FVector2D& GraphPosition)
+{
+	const UEdGraphNode* FromNode = FromPin->GetOwningNode();
+	UEdGraph* EdGraph = FromNode->GetGraph();
+
+	// Add the new reroute node and connect it 
+	TSharedPtr<FAssetSchemaAction_Dataflow_CreateNode_DataflowEdNode> NewNodeAction 
+		= FAssetSchemaAction_Dataflow_CreateNode_DataflowEdNode::CreateAction(EdGraph, FDataflowReRouteNode::StaticType());
+	if (NewNodeAction)
+	{
+		UEdGraphNode* NewEdNode = NewNodeAction->PerformAction(EdGraph, nullptr, GraphPosition, false);
+		if (NewEdNode)
+		{
+			const FName PinName = "Value";
+			UEdGraphPin* InputPin = NewEdNode->FindPin(PinName, EGPD_Input);
+			UEdGraphPin* OutputPin = NewEdNode->FindPin(PinName, EGPD_Output);
+			if (InputPin && OutputPin)
+			{
+				EdGraph->GetSchema()->TryCreateConnection(FromPin, InputPin);
+				EdGraph->GetSchema()->TryCreateConnection(OutputPin, ToPin);
+			}
+		}
+	}
+}
+
+
+void UDataflowSchema::OnPinConnectionDoubleCicked(UEdGraphPin* PinA, UEdGraphPin * PinB, const FVector2D & GraphPosition) const
+{
+	CreateAndConnectNewReRouteNode(PinA, PinB, GraphPosition);
+}
 
 FConnectionDrawingPolicy* UDataflowSchema::CreateConnectionDrawingPolicy(int32 InBackLayerID, int32 InFrontLayerID, float InZoomFactor, const FSlateRect& InClippingRect, class FSlateWindowElementList& InDrawElements, class UEdGraph* InGraphObj) const
 {
diff --git a/Engine/Plugins/Experimental/Dataflow/Source/DataflowEditor/Public/Dataflow/DataflowSchema.h b/Engine/Plugins/Experimental/Dataflow/Source/DataflowEditor/Public/Dataflow/DataflowSchema.h
--- a/Engine/Plugins/Experimental/Dataflow/Source/DataflowEditor/Public/Dataflow/DataflowSchema.h
+++ b/Engine/Plugins/Experimental/Dataflow/Source/DataflowEditor/Public/Dataflow/DataflowSchema.h
@@ -24,7 +24,7 @@ public:
 
 	virtual FLinearColor GetPinTypeColor(const FEdGraphPinType& PinType) const override;
 	virtual FConnectionDrawingPolicy* CreateConnectionDrawingPolicy(int32 InBackLayerID, int32 InFrontLayerID, float InZoomFactor, const FSlateRect& InClippingRect, class FSlateWindowElementList& InDrawElements, class UEdGraph* InGraphObj) const override;
-	//virtual void OnPinConnectionDoubleCicked(UEdGraphPin* PinA, UEdGraphPin* PinB, const FVector2D& GraphPosition) const override;
+	virtual void OnPinConnectionDoubleCicked(UEdGraphPin* PinA, UEdGraphPin* PinB, const FVector2D& GraphPosition) const override;
 
 	//~ End EdGraphSchema Interface
 
diff --git a/Engine/Source/Runtime/Core/Public/UObject/FortniteMainBranchObjectVersion.h b/Engine/Source/Runtime/Core/Public/UObject/FortniteMainBranchObjectVersion.h
--- a/Engine/Source/Runtime/Core/Public/UObject/FortniteMainBranchObjectVersion.h
+++ b/Engine/Source/Runtime/Core/Public/UObject/FortniteMainBranchObjectVersion.h
@@ -438,6 +438,9 @@ struct FFortniteMainBranchObjectVersion
 		// Change workspaces to store asset references as external objects
 		AnimNextWorkspaceEntryConversion,
 
+		// Add support for anytype in dataflow
+		DataflowAnyTypeSupport,
+
 		// -----<new versions can be added above this line>-------------------------------------------------
 		VersionPlusOne,
 		LatestVersion = VersionPlusOne - 1
diff --git a/Engine/Source/Runtime/Experimental/Dataflow/Core/Private/Dataflow/DataflowConnection.cpp b/Engine/Source/Runtime/Experimental/Dataflow/Core/Private/Dataflow/DataflowConnection.cpp
--- a/Engine/Source/Runtime/Experimental/Dataflow/Core/Private/Dataflow/DataflowConnection.cpp
+++ b/Engine/Source/Runtime/Experimental/Dataflow/Core/Private/Dataflow/DataflowConnection.cpp
@@ -14,7 +14,9 @@ FDataflowConnection::FDataflowConnection(Dataflow::FPin::EDirection InDirection,
 	, OwningNode(InOwningNode)
 	, Property(InProperty)
 	, Guid(InGuid)
-{}
+{
+	bIsAnyType = FDataflowConnection::IsAnyType(Type);
+}
 
 uint32 FDataflowConnection::GetOffset() const
 {
@@ -41,3 +43,25 @@ uint32 FDataflowConnection::GetOwningNodeValueHash() const
 	return OwningNode ? OwningNode->GetValueHash() : 0;
 }
 
+bool FDataflowConnection::IsAnyType(const FName& InType)
+{
+	return (InType == FDataflowAnyType::TypeName);
+}
+
+void FDataflowConnection::SetAsAnyType(bool bAnyType, const FName& ConcreteType)
+{
+	bIsAnyType = bAnyType;
+	if (bIsAnyType)
+	{
+		Type = ConcreteType;
+	}
+}
+
+void FDataflowConnection::SetConcreteType(const FName& InType)
+{
+	// Can only change from AnyType to a concrete type
+	if (ensure(IsAnyType() && IsAnyType(Type) && !IsAnyType(InType)))
+	{
+		Type = InType;
+	}
+}
\ No newline at end of file
diff --git a/Engine/Source/Runtime/Experimental/Dataflow/Core/Private/Dataflow/DataflowGraph.cpp b/Engine/Source/Runtime/Experimental/Dataflow/Core/Private/Dataflow/DataflowGraph.cpp
--- a/Engine/Source/Runtime/Experimental/Dataflow/Core/Private/Dataflow/DataflowGraph.cpp
+++ b/Engine/Source/Runtime/Experimental/Dataflow/Core/Private/Dataflow/DataflowGraph.cpp
@@ -8,6 +8,7 @@
 #include "Logging/LogMacros.h"
 #include "Dataflow/DataflowArchive.h"
 #include "UObject/UE5MainStreamObjectVersion.h"
+#include "UObject/FortniteMainBranchObjectVersion.h"
 
 DEFINE_LOG_CATEGORY_STATIC(DATAFLOW_LOG, Error, All);
 
@@ -130,6 +131,7 @@ namespace Dataflow
 	void FGraph::Serialize(FArchive& Ar, UObject* OwningObject)
 	{
 		Ar.UsingCustomVersion(FUE5MainStreamObjectVersion::GUID);
+		Ar.UsingCustomVersion(FFortniteMainBranchObjectVersion::GUID);
 
 		Ar << Guid;
 		if (Ar.IsSaving())
@@ -171,6 +173,9 @@ namespace Dataflow
 					ArType = Output->GetType();
 					ArName = Output->GetName();
 					Ar << ArGuid << ArType << ArName;
+
+					bool bIsAnytype = Output->IsAnyType();
+					Ar << bIsAnytype;
 				}
 
 				int32 ArNumInputs = Node->GetInputs().Num();
@@ -181,6 +186,9 @@ namespace Dataflow
 					ArType = Input->GetType();
 					ArName = Input->GetName();
 					Ar << ArGuid << ArType << ArName;
+
+					bool bIsAnytype = Input->IsAnyType();
+					Ar << bIsAnytype;
 				}
 			}
 			DATAFLOW_OPTIONAL_BLOCK_WRITE_END();
@@ -210,7 +218,10 @@ namespace Dataflow
 				ensure(!NodeGuidMap.Contains(ArGuid));
 				NodeGuidMap.Add(ArGuid, Node);
 
-				if ((Ar.CustomVer(FUE5MainStreamObjectVersion::GUID) < FUE5MainStreamObjectVersion::DataflowSeparateInputOutputSerialization))
+				const bool bDataflowSeparateInputOutputSerialization = (Ar.CustomVer(FUE5MainStreamObjectVersion::GUID) >= FUE5MainStreamObjectVersion::DataflowSeparateInputOutputSerialization);
+				const bool bDataflowAnyTypeSupport = (Ar.CustomVer(FFortniteMainBranchObjectVersion::GUID) >= FFortniteMainBranchObjectVersion::DataflowAnyTypeSupport);
+
+				if (!bDataflowSeparateInputOutputSerialization)
 				{
 
 					// former input / output serialization method where we only store aggregate number of inputs and outputs
@@ -268,6 +279,7 @@ namespace Dataflow
 					// be referenced when deserializing them below ( see Dataflow Node AddPin method )
 					Node->SerializeInternal(Ar);
 
+					bool bIsAnyType = false;
 					// Outputs deserialization
 					{
 						int32 ArNumOutputs;
@@ -276,9 +288,17 @@ namespace Dataflow
 						for (int32 OutputIndex = 0; OutputIndex < ArNumOutputs; OutputIndex++)
 						{
 							Ar << ArGuid << ArType << ArName;
+							if (bDataflowAnyTypeSupport)
+							{
+								Ar << bIsAnyType;
+							}
 
 							if (FDataflowOutput* Output = Node->FindOutput(ArName))
 							{
+								if (bIsAnyType)
+								{
+									Output->SetAsAnyType(bIsAnyType, ArType);
+								}
 								check(Output->GetType() == ArType);
 								Output->SetGuid(ArGuid);
 								ensure(!ConnectionGuidMap.Contains(ArGuid));
@@ -302,10 +322,18 @@ namespace Dataflow
 						for (int32 InputIndex = 0; InputIndex < ArNumInputs; InputIndex++)
 						{
 							Ar << ArGuid << ArType << ArName;
+							if (bDataflowAnyTypeSupport)
+							{
+								Ar << bIsAnyType;
+							}
 
 							if (FDataflowInput* Input = Node->FindInput(ArName))
 							{
-								check(Input->GetType() == ArType);
+								if (bIsAnyType)
+								{
+									Input->SetAsAnyType(bIsAnyType, ArType);
+								}
+								check(Input->GetType() == ArType || bIsAnyType);
 								Input->SetGuid(ArGuid);
 								ensure(!ConnectionGuidMap.Contains(ArGuid));
 								ConnectionGuidMap.Add(ArGuid, Input);
diff --git a/Engine/Source/Runtime/Experimental/Dataflow/Core/Private/Dataflow/DataflowInputOutput.cpp b/Engine/Source/Runtime/Experimental/Dataflow/Core/Private/Dataflow/DataflowInputOutput.cpp
--- a/Engine/Source/Runtime/Experimental/Dataflow/Core/Private/Dataflow/DataflowInputOutput.cpp
+++ b/Engine/Source/Runtime/Experimental/Dataflow/Core/Private/Dataflow/DataflowInputOutput.cpp
@@ -19,11 +19,9 @@ FDataflowInput::FDataflowInput(const Dataflow::FInputParameters& Param, FGuid In
 {
 }
 
-
-
 bool FDataflowInput::AddConnection(FDataflowConnection* InOutput)
 {
-	if (ensure(InOutput->GetType() == this->GetType()))
+	if (Type == InOutput->GetType())
 	{
 		Connection = (FDataflowOutput*)InOutput;
 		GetOwningNode()->Invalidate();
@@ -68,6 +66,18 @@ void FDataflowInput::Invalidate(const Dataflow::FTimestamp& ModifiedTimestamp)
 	OwningNode->Invalidate(ModifiedTimestamp);
 }
 
+void FDataflowInput::PullValue(Dataflow::FContext& Context) const
+{
+	if (GetConnectedOutputs().Num())
+	{
+		ensure(GetConnectedOutputs().Num() == 1);
+		if (const FDataflowOutput* ConnectionOut = GetConnection())
+		{
+			ConnectionOut->Evaluate(Context);
+		}
+	}
+}
+
 //
 //
 //  Output
@@ -109,7 +119,7 @@ TArray< FDataflowInput*> FDataflowOutput::GetConnectedInputs()
 
 bool FDataflowOutput::AddConnection(FDataflowConnection* InOutput)
 {
-	if (ensure(InOutput->GetType() == this->GetType()))
+	if (Type == InOutput->GetType())
 	{
 		Connections.Add((FDataflowInput*)InOutput);
 		return true;
@@ -130,6 +140,23 @@ void FDataflowOutput::Invalidate(const Dataflow::FTimestamp& ModifiedTimestamp)
 	}
 }
 
+bool FDataflowOutput::Evaluate(Dataflow::FContext& Context) const
+{
+	check(OwningNode);
+
+	if (IsOwningNodeEnabled())
+	{
+		return Context.Evaluate(*this);
+	}
+	else if (const FDataflowInput* PassthroughInput = GetPassthroughInput())
+	{
+		ForwardInput(PassthroughInput->RealAddress(), Context);
+		return true;
+	}
+
+	return false;
+}
+
 bool FDataflowOutput::EvaluateImpl(Dataflow::FContext& Context) const
 {
 	Dataflow::FContextScopedCallstack Callstack(Context, this);
@@ -156,6 +183,11 @@ bool FDataflowOutput::EvaluateImpl(Dataflow::FContext& Context) const
 	return true;
 }
 
+TFuture<bool> FDataflowOutput::EvaluateParallel(Dataflow::FContext& Context) const
+{
+	return Async(EAsyncExecution::TaskGraph, [&]() -> bool { return this->Evaluate(Context); });
+}
+
 const FDataflowInput* FDataflowOutput::GetPassthroughInput() const
 {
 	return OwningNode ? OwningNode->FindInput(GetPassthroughRealAddress()) : nullptr;
@@ -172,6 +204,7 @@ void FDataflowOutput::ForwardInput(const void* InputReference, Dataflow::FContex
 			ensure(InputToForward->GetConnectedOutputs().Num() == 1);
 			if (const FDataflowOutput* ConnectionOut = InputToForward->GetConnection())
 			{
+				InputToForward->PullValue(Context);
 				Context.SetDataReference(CacheKey(), Property, ConnectionOut->CacheKey());
 			}
 		}
diff --git a/Engine/Source/Runtime/Experimental/Dataflow/Core/Private/Dataflow/DataflowNode.cpp b/Engine/Source/Runtime/Experimental/Dataflow/Core/Private/Dataflow/DataflowNode.cpp
--- a/Engine/Source/Runtime/Experimental/Dataflow/Core/Private/Dataflow/DataflowNode.cpp
+++ b/Engine/Source/Runtime/Experimental/Dataflow/Core/Private/Dataflow/DataflowNode.cpp
@@ -19,6 +19,8 @@ const FName FDataflowNode::DataflowIntrinsic = TEXT("DataflowIntrinsic");
 const FLinearColor FDataflowNode::DefaultNodeTitleColor = FLinearColor(1.f, 1.f, 0.8f);
 const FLinearColor FDataflowNode::DefaultNodeBodyTintColor = FLinearColor(0.f, 0.f, 0.f, 0.5f);
 
+const FName FDataflowAnyType::TypeName = TEXT("FDataflowAnyType");
+
 //
 // Inputs
 //
@@ -810,6 +812,26 @@ void FDataflowNode::CopyNodeProperties(const TSharedPtr<FDataflowNode> CopyFromD
 }
 
 
+void FDataflowNode::ForwardInput(Dataflow::FContext& Context, const void* InputReference, const void* Reference) const
+{
+	if (const FDataflowOutput* Output = FindOutput(Reference))
+	{
+		if (const FDataflowInput* Input = FindInput(Reference))
+		{
+			// we need to pull the value first so the upstream of the graph evaluate 
+			Input->PullValue(Context);
+			Output->ForwardInput(InputReference, Context);
+		}
+		else
+		{
+			checkfSlow(false, TEXT("This input could not be found within this node, check this has been properly registered in the node constructor"));
+		}
+	}
+	else
+	{
+		checkfSlow(false, TEXT("This output could not be found within this node, check this has been properly registered in the node constructor"));
+	}
+}
 
 
 
diff --git a/Engine/Source/Runtime/Experimental/Dataflow/Core/Private/Dataflow/DataflowNodeFactory.cpp b/Engine/Source/Runtime/Experimental/Dataflow/Core/Private/Dataflow/DataflowNodeFactory.cpp
--- a/Engine/Source/Runtime/Experimental/Dataflow/Core/Private/Dataflow/DataflowNodeFactory.cpp
+++ b/Engine/Source/Runtime/Experimental/Dataflow/Core/Private/Dataflow/DataflowNodeFactory.cpp
@@ -12,6 +12,13 @@ namespace Dataflow
 {
 	FNodeFactory* FNodeFactory::Instance = nullptr;
 
+	void FNodeFactory::RegisterDefaultNodes()
+	{
+		// by default always register a reroute node
+		DATAFLOW_NODE_REGISTER_CREATION_FACTORY(FDataflowReRouteNode);
+	}
+
+
 	TSharedPtr<FDataflowNode> FNodeFactory::NewNodeFromRegisteredType(FGraph& Graph, const FNewNodeParameters& Param)
 	{ 
 		if (ClassMap.Contains(Param.Type))
diff --git a/Engine/Source/Runtime/Experimental/Dataflow/Core/Public/Dataflow/DataflowConnection.h b/Engine/Source/Runtime/Experimental/Dataflow/Core/Public/Dataflow/DataflowConnection.h
--- a/Engine/Source/Runtime/Experimental/Dataflow/Core/Public/Dataflow/DataflowConnection.h
+++ b/Engine/Source/Runtime/Experimental/Dataflow/Core/Public/Dataflow/DataflowConnection.h
@@ -33,6 +33,8 @@ namespace Dataflow
 		FName Type;
 		FName Name;
 	};
+
+	class FGraph;
 }
 
 //
@@ -50,14 +52,19 @@ protected:
 	FDataflowNode* OwningNode = nullptr;
 	const FProperty* Property = nullptr;
 	FGuid  Guid;
+	bool bIsAnyType = false;
 
 	friend struct FDataflowNode;
+	friend class Dataflow::FGraph;
 
 protected:
 	DATAFLOWCORE_API bool IsOwningNodeEnabled() const;
 	DATAFLOWCORE_API FGuid GetOwningNodeGuid() const;
 	DATAFLOWCORE_API uint32 GetOwningNodeValueHash() const;
 
+	/** this should only be used for serialization */
+	DATAFLOWCORE_API void SetAsAnyType(bool bAnyType, const FName& ConcreteType);
+
 public:
 	FDataflowConnection() {};
 	DATAFLOWCORE_API FDataflowConnection(Dataflow::FPin::EDirection Direction, FName InType, FName InName, FDataflowNode* OwningNode = nullptr, const FProperty* InProperty = nullptr, FGuid InGuid = FGuid::NewGuid());
@@ -85,6 +92,11 @@ public:
 	virtual bool AddConnection(FDataflowConnection* In) { return false; };
 	virtual bool RemoveConnection(FDataflowConnection* In) { return false; }
 
+	DATAFLOWCORE_API bool IsAnyType() const { return bIsAnyType; }
+	static bool IsAnyType(const FName& InType);
+
+	DATAFLOWCORE_API void SetConcreteType(const FName& InType);
+
 	template<class T>
 	bool IsA(const T* InVar) const
 	{
diff --git a/Engine/Source/Runtime/Experimental/Dataflow/Core/Public/Dataflow/DataflowInputOutput.h b/Engine/Source/Runtime/Experimental/Dataflow/Core/Public/Dataflow/DataflowInputOutput.h
--- a/Engine/Source/Runtime/Experimental/Dataflow/Core/Public/Dataflow/DataflowInputOutput.h
+++ b/Engine/Source/Runtime/Experimental/Dataflow/Core/Public/Dataflow/DataflowInputOutput.h
@@ -54,9 +54,19 @@ public:
 	virtual TArray< FDataflowOutput* > GetConnectedOutputs();
 	virtual const TArray< const FDataflowOutput* > GetConnectedOutputs() const;
 
+	/** 
+	* Get the value of this input by evaluating the value of the connected output 
+	* @return the typed value of the input 
+	*/
 	template<class T>
 	const T& GetValue(Dataflow::FContext& Context, const T& Default) const;
 
+	/**
+	* pull the value from the upstream connections
+	* the upstream graph is evaluated if necessary and values are cached along the way 
+	*/
+	void PullValue(Dataflow::FContext& Context) const;
+
 	template<class T>
 	TFuture<const T&> GetValueParallel(Dataflow::FContext& Context, const T& Default) const;
 
@@ -133,9 +143,10 @@ public:
 		}
 	}
 
-	template<class T> const T& GetValue(Dataflow::FContext& Context, const T& Default) const
+	template<class T>
+	const T& GetValue(Dataflow::FContext& Context, const T& Default) const
 	{
-		if (!this->Evaluate<T>(Context))
+		if (!this->Evaluate(Context))
 		{
 			Context.SetData(CacheKey(), Property, Default, GetOwningNodeGuid(), GetOwningNodeValueHash(), Dataflow::FTimestamp::Current());
 		}
@@ -154,11 +165,9 @@ public:
 
 	DATAFLOWCORE_API bool EvaluateImpl(Dataflow::FContext& Context) const;
 	
-	template<class T>
-	bool Evaluate(Dataflow::FContext& Context) const;
+	DATAFLOWCORE_API bool Evaluate(Dataflow::FContext& Context) const;
 
-	template<class T>
-	TFuture<bool> EvaluateParallel(Dataflow::FContext& Context) const;
+	DATAFLOWCORE_API TFuture<bool> EvaluateParallel(Dataflow::FContext& Context) const;
 
 	DATAFLOWCORE_API virtual void Invalidate(const Dataflow::FTimestamp& ModifiedTimestamp = Dataflow::FTimestamp::Current()) override;
 
@@ -174,7 +183,7 @@ const T& FDataflowInput::GetValue(Dataflow::FContext& Context, const T& Default)
 		ensure(GetConnectedOutputs().Num() == 1);
 		if (const FDataflowOutput* ConnectionOut = GetConnection())
 		{
-			if (!ConnectionOut->Evaluate<T>(Context))
+			if (!ConnectionOut->Evaluate(Context))
 			{
 				Context.SetData(ConnectionOut->CacheKey(), Property, Default, GetOwningNodeGuid(), GetOwningNodeValueHash(), Dataflow::FTimestamp::Current());
 			}
@@ -194,28 +203,3 @@ TFuture<const T&> FDataflowInput::GetValueParallel(Dataflow::FContext& Context,
 	return Async(EAsyncExecution::TaskGraph, [&]() -> const T& { return this->GetValue<T>(Context, Default); });
 }
 
-template<class T>
-bool FDataflowOutput::Evaluate(Dataflow::FContext& Context) const
-{
-	check(OwningNode);
- 
-	if (IsOwningNodeEnabled())
-	{
-		return Context.Evaluate(*this);
-	}
-	else if(const FDataflowInput* PassthroughInput = GetPassthroughInput())
-	{
-		// @todo(dataflow) would be nice if the passthrough does not overwrite the existing cache value.
-		const T& PassthroughData = PassthroughInput->GetValue<T>(Context, *reinterpret_cast<const T*>(PassthroughInput->RealAddress()));
-		SetValue(PassthroughData, Context);
-		return true;
-	}
- 
-	return false;
-}
-
-template<class T>
-TFuture<bool> FDataflowOutput::EvaluateParallel(Dataflow::FContext& Context) const
-{
-	return Async(EAsyncExecution::TaskGraph, [&]() -> bool { return this->Evaluate<T>(Context); });
-}
diff --git a/Engine/Source/Runtime/Experimental/Dataflow/Core/Public/Dataflow/DataflowNode.h b/Engine/Source/Runtime/Experimental/Dataflow/Core/Public/Dataflow/DataflowNode.h
--- a/Engine/Source/Runtime/Experimental/Dataflow/Core/Public/Dataflow/DataflowNode.h
+++ b/Engine/Source/Runtime/Experimental/Dataflow/Core/Public/Dataflow/DataflowNode.h
@@ -254,26 +254,7 @@ struct FDataflowNode
 	*   @param InputReference : Pointer to a input member of this node that needs to be forwarded
 	*   @param Reference : Pointer to a member of this node that corresponds with the output to set.
 	*/
-	template<class T> void ForwardInput(Dataflow::FContext& Context, const typename TDecay<T>::Type* InputReference, const typename TDecay<T>::Type* Reference) const
-	{
-		if (const FDataflowOutput* Output = FindOutput(Reference))
-		{
-			if (const FDataflowInput* Input = FindInput(Reference))
-			{
-				// we need to pull the value first so the upstream of the graph evaluate 
-				Input->template GetValue<T>(Context, *Reference);
-				Output->ForwardInput(InputReference, Context);
-			}
-			else
-			{
-				checkfSlow(false, TEXT("This input could not be found within this node, check this has been properly registered in the node constructor"));
-			}
-		}
-		else
-		{
-			checkfSlow(false, TEXT("This output could not be found within this node, check this has been properly registered in the node constructor"));
-		}
-	}
+	DATAFLOWCORE_API void ForwardInput(Dataflow::FContext& Context, const void* InputReference, const void* Reference) const;
 
 	/**
 	*   IsConnected(...)
@@ -352,6 +333,7 @@ protected:
 
 };
 
+
 namespace Dataflow
 {
 	//
@@ -405,3 +387,38 @@ private:
 
 }
 
+USTRUCT()
+struct FDataflowAnyType
+{
+	GENERATED_USTRUCT_BODY()
+	void* Value = nullptr;
+
+	DATAFLOWCORE_API static const FName TypeName;
+};
+
+USTRUCT()
+struct FDataflowReRouteNode : public FDataflowNode
+{
+	GENERATED_USTRUCT_BODY()
+
+	DATAFLOW_NODE_DEFINE_INTERNAL(FDataflowReRouteNode, "DataflowReRouteNode", "Core", "")
+
+public:
+	FDataflowReRouteNode(const Dataflow::FNodeParameters& Param, FGuid InGuid = FGuid::NewGuid())
+		: Super(Param, InGuid)
+	{
+		RegisterInputConnection(&Value);
+		RegisterOutputConnection(&Value, &Value);
+	}
+
+	virtual void Evaluate(Dataflow::FContext& Context, const FDataflowOutput* Out) const
+	{
+		ForwardInput(Context, &Value, &Value);
+	}
+
+public:
+	UPROPERTY(meta = (DataflowInput, DataflowOutput, DataflowPassthrough = "Value", DisplayName = "Value"))
+	FDataflowAnyType Value;
+};
+
+
diff --git a/Engine/Source/Runtime/Experimental/Dataflow/Core/Public/Dataflow/DataflowNodeFactory.h b/Engine/Source/Runtime/Experimental/Dataflow/Core/Public/Dataflow/DataflowNodeFactory.h
--- a/Engine/Source/Runtime/Experimental/Dataflow/Core/Public/Dataflow/DataflowNodeFactory.h
+++ b/Engine/Source/Runtime/Experimental/Dataflow/Core/Public/Dataflow/DataflowNodeFactory.h
@@ -47,7 +47,7 @@ namespace Dataflow
 		TMap<FName, FName > DisplayMap;					// [DisplayName] -> TypeName
 
 		DATAFLOWCORE_API static FNodeFactory* Instance;
-		FNodeFactory() {}
+		FNodeFactory() {};
 
 	public:
 		~FNodeFactory() { delete Instance; }
@@ -57,6 +57,7 @@ namespace Dataflow
 			if (!Instance)
 			{
 				Instance = new FNodeFactory();
+				Instance->RegisterDefaultNodes();
 			}
 			return Instance;
 		}
@@ -94,8 +95,9 @@ namespace Dataflow
 			for (auto Elem : ParametersMap) RetVal.Add(Elem.Value);
 			return RetVal;
 		}
+	private:
+		DATAFLOWCORE_API void RegisterDefaultNodes();
 
 	};
-
 }
 
diff --git a/Engine/Source/Runtime/Experimental/Dataflow/Engine/Private/Dataflow/DataflowEdNode.cpp b/Engine/Source/Runtime/Experimental/Dataflow/Engine/Private/Dataflow/DataflowEdNode.cpp
--- a/Engine/Source/Runtime/Experimental/Dataflow/Engine/Private/Dataflow/DataflowEdNode.cpp
+++ b/Engine/Source/Runtime/Experimental/Dataflow/Engine/Private/Dataflow/DataflowEdNode.cpp
@@ -265,6 +265,11 @@ void UDataflowEdNode::PinConnectionListChanged(UEdGraphPin* Pin)
 										DataflowGraph->Connect(LinkedConOutput, ConnectionInput);
 									}
 								}
+								OnPinConnection(*Pin, LinkedCon->PinType);
+							}
+							else
+							{
+								OnPinDisconnection(*Pin);
 							}
 						}
 					}
@@ -288,6 +293,11 @@ void UDataflowEdNode::PinConnectionListChanged(UEdGraphPin* Pin)
 										DataflowGraph->Connect(ConnectionOutput, LinkedConInput);
 									}
 								}
+								OnPinConnection(*Pin, LinkedCon->PinType);
+							}
+							else
+							{
+								OnPinDisconnection(*Pin);
 							}
 						}
 					}
@@ -301,6 +311,62 @@ void UDataflowEdNode::PinConnectionListChanged(UEdGraphPin* Pin)
 
 #endif // WITH_EDITOR && !UE_BUILD_SHIPPING
 
+void UDataflowEdNode::OnPinConnection(const UEdGraphPin& Pin, const FEdGraphPinType& Type)
+{
+#if WITH_EDITOR
+	if (Pin.PinType.PinCategory == FDataflowAnyType::TypeName)
+	{
+		SetAnyTypePinsToType(Type);
+	}
+#endif
+}
+
+void UDataflowEdNode::OnPinDisconnection(const UEdGraphPin& Pin)
+{
+}
+
+void UDataflowEdNode::SetAnyTypePinsToType(const FEdGraphPinType& Type)
+{
+#if WITH_EDITOR
+	if (const TSharedPtr<FDataflowNode> DataflowNode = this->GetDataflowNode())
+	{
+		for (UEdGraphPin* Pin : GetAllPins())
+		{
+			if (Pin)
+			{
+				switch (Pin->Direction)
+				{
+					case EEdGraphPinDirection::EGPD_Input:
+					{
+						if (FDataflowInput* Input = DataflowNode->FindInput(Pin->GetFName()))
+						{
+							if (Input->IsAnyType())
+							{
+								Pin->PinType = Type;
+								Input->SetConcreteType(Type.PinCategory);
+							}
+						}
+						break;
+					}
+					case EEdGraphPinDirection::EGPD_Output:
+					{
+						if (FDataflowOutput* Output = DataflowNode->FindOutput(Pin->GetFName()))
+						{
+							if (Output->IsAnyType())
+							{
+								Pin->PinType = Type;
+								Output->SetConcreteType(Type.PinCategory);
+							}
+						}
+						break;
+					}
+				}
+			}
+		}
+	}
+#endif
+}
+
 void UDataflowEdNode::Serialize(FArchive& Ar)
 {
 	Super::Serialize(Ar);
@@ -477,7 +543,22 @@ void UDataflowEdNode::OnPinRemoved(UEdGraphPin* InRemovedPin)
 	}
 }
 
-#endif
+bool UDataflowEdNode::ShouldDrawNodeAsControlPointOnly(int32& OutInputPinIndex, int32& OutOutputPinIndex) const
+{
+	UEdGraphNode::ShouldDrawNodeAsControlPointOnly(OutInputPinIndex, OutOutputPinIndex);
+	if (const TSharedPtr<const FDataflowNode> DataflowNode = GetDataflowNode())
+	{
+		if (DataflowNode->GetType() == FDataflowReRouteNode::StaticType())
+		{
+			OutInputPinIndex = 0;
+			OutOutputPinIndex = 1;
+			return true;
+		}
+	}
+	return false;
+}
+
+#endif //WITH_EDITOR
 
 
 TArray<Dataflow::FRenderingParameter> UDataflowEdNode::GetRenderParameters() const
diff --git a/Engine/Source/Runtime/Experimental/Dataflow/Engine/Public/Dataflow/DataflowEdNode.h b/Engine/Source/Runtime/Experimental/Dataflow/Engine/Public/Dataflow/DataflowEdNode.h
--- a/Engine/Source/Runtime/Experimental/Dataflow/Engine/Public/Dataflow/DataflowEdNode.h
+++ b/Engine/Source/Runtime/Experimental/Dataflow/Engine/Public/Dataflow/DataflowEdNode.h
@@ -47,6 +47,7 @@ public:
 	DATAFLOWENGINE_API virtual FText GetPinDisplayName(const UEdGraphPin* Pin) const override;
 	DATAFLOWENGINE_API virtual void AutowireNewNode(UEdGraphPin* FromPin) override;
 	DATAFLOWENGINE_API virtual void OnPinRemoved(UEdGraphPin* InRemovedPin) override;
+	DATAFLOWENGINE_API virtual bool ShouldDrawNodeAsControlPointOnly(int32& OutInputPinIndex, int32& OutOutputPinIndex) const override;
 #endif // WITH_EDITOR
 	// End of UEdGraphNode interface
 
@@ -84,5 +85,10 @@ public:
 	UPROPERTY()
 	bool bRenderInAssetEditor = false;
 
+private:
+	void OnPinConnection(const UEdGraphPin& Pin, const FEdGraphPinType& Type);
+	void OnPinDisconnection(const UEdGraphPin& Pin);
+	void SetAnyTypePinsToType(const FEdGraphPinType& Type);
+
 };
 
