Constraints: safety guards for possible null World (follow-up for 32761364)

#rnx

diff --git a/Engine/Source/Runtime/Experimental/Animation/Constraints/Private/TransformConstraint.cpp b/Engine/Source/Runtime/Experimental/Animation/Constraints/Private/TransformConstraint.cpp
--- a/Engine/Source/Runtime/Experimental/Animation/Constraints/Private/TransformConstraint.cpp
+++ b/Engine/Source/Runtime/Experimental/Animation/Constraints/Private/TransformConstraint.cpp
@@ -436,7 +436,7 @@ void UTickableTransformConstraint::Setup()
 
 }
 
-void UTickableTransformConstraint::SetupDependencies(UWorld* InWorld)
+void UTickableTransformConstraint::SetupDependencies(const UWorld* InWorld)
 {
 	//we may not be the outer for old files so move it over
 	if (ParentTRSHandle)
@@ -455,6 +455,11 @@ void UTickableTransformConstraint::SetupDependencies(UWorld* InWorld)
 			ChildTRSHandle->Rename(nullptr, this, REN_ForceNoResetLoaders | REN_DontCreateRedirectors);
 		}
 	}
+
+	if (!ensure(InWorld))
+	{
+		return;
+	}
 	
 	FTickFunction* ParentTickFunction = GetParentHandleTickFunction();
 	FTickFunction* ChildTickFunction = GetChildHandleTickFunction();
@@ -483,8 +488,13 @@ void UTickableTransformConstraint::SetupDependencies(UWorld* InWorld)
 	}
 }
 
-void UTickableTransformConstraint::EnsurePrimaryDependency(UWorld* InWorld)
+void UTickableTransformConstraint::EnsurePrimaryDependency(const UWorld* InWorld) const
 {
+	if (!ensure(InWorld))
+	{
+		return;
+	}
+	
 	const FTickFunction* ParentTickFunction = GetParentHandleTickFunction();
 	const FTickFunction* ChildTickFunction = GetChildHandleTickFunction();
 	if (ParentTickFunction && (ChildTickFunction != ParentTickFunction))
@@ -780,6 +790,11 @@ FTickFunction* UTickableTransformConstraint::GetHandleTickFunction(const TObject
 
 void UTickableTransformConstraint::InitConstraint(UWorld *InWorld)
 {
+	if (!ensure(InWorld))
+	{
+		return;
+	}
+	
 	FConstraintTickFunction& ConstraintTick = ConstraintTicks.FindOrAdd(InWorld->GetCurrentLevel());
 
 	if (ConstraintTick.ConstraintFunctions.IsEmpty())
@@ -795,6 +810,11 @@ void UTickableTransformConstraint::InitConstraint(UWorld *InWorld)
 
 void UTickableTransformConstraint::TeardownConstraint(UWorld* InWorld)
 {
+	if (!ensure(InWorld))
+	{
+		return;
+	}
+	
 	FConstraintTickFunction& ConstraintTick = ConstraintTicks.FindOrAdd(InWorld->GetCurrentLevel());
 	ConstraintTick.SetTickFunctionEnable(false);
 	ConstraintTick.UnRegisterTickFunction();
diff --git a/Engine/Source/Runtime/Experimental/Animation/Constraints/Public/TransformConstraint.h b/Engine/Source/Runtime/Experimental/Animation/Constraints/Public/TransformConstraint.h
--- a/Engine/Source/Runtime/Experimental/Animation/Constraints/Public/TransformConstraint.h
+++ b/Engine/Source/Runtime/Experimental/Animation/Constraints/Public/TransformConstraint.h
@@ -117,7 +117,7 @@ public:
 	/**
 	* Sets up dependencies with the first primary prerequisite available if the parent does not tick.   
 	*/
-	CONSTRAINTS_API void EnsurePrimaryDependency(UWorld* InWorld);
+	CONSTRAINTS_API void EnsurePrimaryDependency(const UWorld* InWorld) const;
 	
 protected:
 
@@ -136,7 +136,7 @@ protected:
 	 * Sets up dependencies between the parent, the constraint and the child using their respective tick functions.
 	 * It creates a dependency graph between them so that they tick in the right order when evaluated.   
 	*/
-	CONSTRAINTS_API void SetupDependencies(UWorld* InWorld);
+	CONSTRAINTS_API void SetupDependencies(const UWorld* InWorld);
 
 	/** Set the current child's global transform. */
 	CONSTRAINTS_API void SetChildGlobalTransform(const FTransform& InGlobal) const;
