Fixing single header compilation.

diff --git a/Engine/Plugins/Online/OnlineBase/Source/Public/Online/OnlineSessionNames.h b/Engine/Plugins/Online/OnlineBase/Source/Public/Online/OnlineSessionNames.h
--- a/Engine/Plugins/Online/OnlineBase/Source/Public/Online/OnlineSessionNames.h
+++ b/Engine/Plugins/Online/OnlineBase/Source/Public/Online/OnlineSessionNames.h
@@ -3,6 +3,8 @@
 #pragma once
 
 #include "HAL/Platform.h"
+#include "Misc/CoreMiscDefines.h"
+#include "UObject/NameTypes.h"
 
 /** default beacon port, if not specified by other means */
 #define DEFAULT_BEACON_PORT 15000
diff --git a/Engine/Source/Editor/UniversalObjectLocatorEditor/Public/IUniversalObjectLocatorCustomization.h b/Engine/Source/Editor/UniversalObjectLocatorEditor/Public/IUniversalObjectLocatorCustomization.h
--- a/Engine/Source/Editor/UniversalObjectLocatorEditor/Public/IUniversalObjectLocatorCustomization.h
+++ b/Engine/Source/Editor/UniversalObjectLocatorEditor/Public/IUniversalObjectLocatorCustomization.h
@@ -5,6 +5,7 @@
 #include "Templates/SharedPointer.h"
 
 struct FUniversalObjectLocator;
+struct FUniversalObjectLocatorFragment;
 
 class FString;
 class UObject;
diff --git a/Engine/Source/Runtime/Core/Public/Templates/NonNullPointer.h b/Engine/Source/Runtime/Core/Public/Templates/NonNullPointer.h
--- a/Engine/Source/Runtime/Core/Public/Templates/NonNullPointer.h
+++ b/Engine/Source/Runtime/Core/Public/Templates/NonNullPointer.h
@@ -5,6 +5,7 @@
 #include "CoreTypes.h"
 #include "Misc/AssertionMacros.h"
 #include "Misc/OptionalFwd.h"
+#include "Templates/Requires.h"
 #include "Templates/UnrealTypeTraits.h"
 
 #include <type_traits>
diff --git a/Engine/Source/Runtime/Experimental/Chaos/Public/Chaos/Collision/ConvexContactPointUtilities.h b/Engine/Source/Runtime/Experimental/Chaos/Public/Chaos/Collision/ConvexContactPointUtilities.h
--- a/Engine/Source/Runtime/Experimental/Chaos/Public/Chaos/Collision/ConvexContactPointUtilities.h
+++ b/Engine/Source/Runtime/Experimental/Chaos/Public/Chaos/Collision/ConvexContactPointUtilities.h
@@ -4,6 +4,7 @@
 #include "Chaos/Core.h"
 #include "Chaos/Capsule.h"
 #include "Chaos/Collision/ContactPoint.h"
+#include "Chaos/Collision/ContactTriangles.h"
 #include "Chaos/Collision/ConvexFeature.h"
 #include "Chaos/Triangle.h"
 
