Geometry collection : make sure RefreshCustomRenderer check the broken state using the BrokenAndDecayedStates API instead of checking the Active attribute directly
- This ensure that the geometry collection will not pretend it is broken for a split time when spawning

diff --git a/Engine/Source/Runtime/Experimental/GeometryCollectionEngine/Private/GeometryCollection/GeometryCollectionComponent.cpp b/Engine/Source/Runtime/Experimental/GeometryCollectionEngine/Private/GeometryCollection/GeometryCollectionComponent.cpp
--- a/Engine/Source/Runtime/Experimental/GeometryCollectionEngine/Private/GeometryCollection/GeometryCollectionComponent.cpp
+++ b/Engine/Source/Runtime/Experimental/GeometryCollectionEngine/Private/GeometryCollection/GeometryCollectionComponent.cpp
@@ -6276,8 +6276,7 @@ void UGeometryCollectionComponent::RefreshCustomRenderer()
 	
 	if (IGeometryCollectionExternalRenderInterface* RendererInterface = CustomRenderer.GetInterface())
 	{
-		const int32 RootIndex = GetRootIndex();
-		bool bIsBroken = DynamicCollection ? !DynamicCollection->Active[RootIndex] : false;
+		bool bIsBroken = BrokenAndDecayedStates.GetIsRootBroken();
 
 #if !(UE_BUILD_SHIPPING)
 		if (CVarNumToForceBreak->GetInt() >= 0)
