Chaos - better initial overlap depentration revisited

Partial rollback of CL32811662. Initial overlap tracking can now be per-manifold-point or across all manifold points. The latter works best for collision against meshes. Still needs some work, but for now it will do...

diff --git a/Engine/Source/Runtime/Experimental/Chaos/Private/Chaos/Collision/PBDCollisionConstraint.cpp b/Engine/Source/Runtime/Experimental/Chaos/Private/Chaos/Collision/PBDCollisionConstraint.cpp
--- a/Engine/Source/Runtime/Experimental/Chaos/Private/Chaos/Collision/PBDCollisionConstraint.cpp
+++ b/Engine/Source/Runtime/Experimental/Chaos/Private/Chaos/Collision/PBDCollisionConstraint.cpp
@@ -71,6 +71,9 @@ namespace Chaos
 	int32 Chaos_Collision_MaxManifoldPoints = -1;
 	FAutoConsoleVariableRef CVarChaos_Collision_MaxManifoldPoints(TEXT("p.Chaos.Collision.MaxManifoldPoints"), Chaos_Collision_MaxManifoldPoints, TEXT(""));
 
+	bool bChaos_Collision_AllowGlobalInitialPhi = true;
+	FAutoConsoleVariableRef CVarChaos_Collision_AllowGlobalInitialPhi(TEXT("p.Chaos.Collision.AllowGlobalInitialPhi"), bChaos_Collision_AllowGlobalInitialPhi, TEXT(""));
+
 
 	struct FCollisionTolerances
 	{
@@ -344,6 +347,15 @@ namespace Chaos
 		const bool bOneWay1 = FConstGenericParticleHandle(GetParticle1())->OneWayInteraction();
 		Flags.bIsOneWayInteraction = (bDynamic0 || bDynamic1) && (bOneWay0 || bOneWay1);
 
+		// For dynamic against kinematic (as long as no spheres or capsules are involved) we share initial overlap between all manifold points.
+		// This is to make convex-mesh collisions work better for large overlaps on complicated meshes, but it is not ideal.
+		// @todo(chaos): try to remove this
+		Flags.bUsePerContactInitialPhi = true;
+		if (bChaos_Collision_AllowGlobalInitialPhi && !Flags.bIsQuadratic0 && !Flags.bIsQuadratic1 && (!bDynamic0 || !bDynamic1))
+		{
+			Flags.bUsePerContactInitialPhi = false;
+		}
+
 		// Only levelsets use incremental collision manifolds
 		if (bInUseManifold && ((ImplicitType0 == ImplicitObjectType::LevelSet) || (ImplicitType1 == ImplicitObjectType::LevelSet)))
 		{
@@ -1211,6 +1223,7 @@ namespace Chaos
 					ManifoldPoints[PointIndex].Flags.bHasStaticFrictionAnchor = true;
 					ManifoldPoints[PointIndex].ShapeAnchorPoints[0] = SavedManifoldPoints[SavedManifoldPointIndex].ShapeContactPoints[0];
 					ManifoldPoints[PointIndex].ShapeAnchorPoints[1] = SavedManifoldPoints[SavedManifoldPointIndex].ShapeContactPoints[1];
+					ManifoldPoints[PointIndex].InitialPhi = SavedManifoldPoints[SavedManifoldPointIndex].InitialPhi;
 					ManifoldPoints[PointIndex].Flags.bInitialContact = false;
 				}
 				// Nothing to do if no saved friction point because we already set the achor to the most recently detected contact point
diff --git a/Engine/Source/Runtime/Experimental/Chaos/Private/Chaos/Collision/PBDCollisionContainerSolver.cpp b/Engine/Source/Runtime/Experimental/Chaos/Private/Chaos/Collision/PBDCollisionContainerSolver.cpp
--- a/Engine/Source/Runtime/Experimental/Chaos/Private/Chaos/Collision/PBDCollisionContainerSolver.cpp
+++ b/Engine/Source/Runtime/Experimental/Chaos/Private/Chaos/Collision/PBDCollisionContainerSolver.cpp
@@ -154,14 +154,18 @@ namespace Chaos
 			}
 		}
 
-		// Overlap remaining from the previous frame, estimated from current contact phi and velocity
-		FRealSingle WorldContactResidualPhi = FMath::Min(WorldContactDeltaNormal, FRealSingle(0)) - FMath::Min(ContactVelocityNormal * Dt, FRealSingle(0));
-
-		// Initial Phi for initial-overlap depenetration.
-		// If we have an initial contact, calculate the initial overlap. This will get saved in SetSolverResults
-		FRealSingle WorldContactInitialPhi = 0;
+		// Handle initial overlaps. Initially overlapping objects will de-penetrate at MaxDepenetrationVelocity (which may be zero).
 		if ((MaxDepenetrationVelocity >= 0) && CVars::bChaos_Collision_EnableInitialDepenetration)
 		{
+			// Overlap remaining from the previous frame, estimated from current contact phi and velocity
+			FRealSingle WorldContactResidualPhi = FMath::Min(WorldContactDeltaNormal, FRealSingle(0)) - FMath::Min(ContactVelocityNormal * Dt, FRealSingle(0));
+
+			// We can handle initial overlap depenentration individually per manifold point or
+			// globally over all maifold points. See FPBDCollisionConstraint::Setup
+			const bool bUsePerContactInitialPhi = Constraint->UsePerContactInitialPhi();
+			
+			// If we have an initial contact, calculate the initial overlap. This will get saved in SetSolverResults
+			FRealSingle WorldContactInitialPhi = 0;
 			if (ManifoldPoint.Flags.bInitialContact)
 			{
 				// This is a new manifold point, capture current Phi as the initial Phi
@@ -173,22 +177,26 @@ namespace Chaos
 				// new deeper manifold points as full initial overlaps.
 				// NOTE: here we are checking IsInitialContact on the constraint, which is only set when we first make contact, as opposed 
 				// to the bInitialContact on the manifold point which is true for any new manifold point, regardless of the constraint age.
-				if (!Constraint->IsInitialContact())
+				if (bUsePerContactInitialPhi && !Constraint->IsInitialContact())
 				{
 					WorldContactInitialPhi = FMath::Max(WorldContactInitialPhi, Constraint->GetMinInitialPhi());
 				}
 
 			}
-			else if (Constraint->GetMinInitialPhi() < 0)
+			else
 			{
 				// This is a pre-existing manifold point, but maybe we are still resolving initial penetrations
-				// Don't allow this contact to penetrate any deeper than it currently is. MinInitialPhi will
+				// Don't allow this contact to penetrate any deeper than it currently is. InitialPhi will
 				// decrease over time if we have a non-zero depenetration velocity
-				WorldContactInitialPhi = FMath::Max(WorldContactDeltaNormal, Constraint->GetMinInitialPhi());
+				const FRealSingle InitialPhi = bUsePerContactInitialPhi ? ManifoldPoint.InitialPhi : Constraint->GetMinInitialPhi();
+				if (InitialPhi < 0)
+				{
+					WorldContactInitialPhi = FMath::Max(WorldContactDeltaNormal, InitialPhi);
+				}
 			}
 
-			// InitialPhi is only for tracking penetration - cannot be positive
-			WorldContactInitialPhi = FMath::Min(WorldContactInitialPhi, FRealSingle(0));
+			// Update the initial overlap based on depenetration velocity
+			WorldContactInitialPhi = FMath::Min(WorldContactInitialPhi + MaxDepenetrationVelocity * Dt, 0.0f);
 
 			// Apply initial penetration allowance to depth correction
 			WorldContactDeltaNormal -= WorldContactInitialPhi;
@@ -209,8 +217,7 @@ namespace Chaos
 		}
 
 		// Adjust depth to account for target penetration from user
-		const FRealSingle TargetPhi = ManifoldPoint.TargetPhi;
-		WorldContactDeltaNormal -= TargetPhi;
+		WorldContactDeltaNormal -= ManifoldPoint.TargetPhi;
 
 		Solver.InitManifoldPoint(
 			SolverPointIndex,
@@ -375,8 +382,6 @@ namespace Chaos
 				StaticFrictionRatio,
 				Dt);
 		}
-
-		Constraint->FinalizeSolverResults(Dt);
 	}
 
 	//////////////////////////////////////////////////////////////////////////////////////////////////
diff --git a/Engine/Source/Runtime/Experimental/Chaos/Private/Chaos/Collision/PBDCollisionContainerSolverJacobi.cpp b/Engine/Source/Runtime/Experimental/Chaos/Private/Chaos/Collision/PBDCollisionContainerSolverJacobi.cpp
--- a/Engine/Source/Runtime/Experimental/Chaos/Private/Chaos/Collision/PBDCollisionContainerSolverJacobi.cpp
+++ b/Engine/Source/Runtime/Experimental/Chaos/Private/Chaos/Collision/PBDCollisionContainerSolverJacobi.cpp
@@ -248,8 +248,6 @@ namespace Chaos
 					StaticFrictionRatio,
 					Dt);
 			}
-
-			Constraint->FinalizeSolverResults(Dt);
 		}
 
 		//////////////////////////////////////////////////////////////////////////////////////////////////
diff --git a/Engine/Source/Runtime/Experimental/Chaos/Private/Chaos/Collision/PBDCollisionContainerSolverSimd.cpp b/Engine/Source/Runtime/Experimental/Chaos/Private/Chaos/Collision/PBDCollisionContainerSolverSimd.cpp
--- a/Engine/Source/Runtime/Experimental/Chaos/Private/Chaos/Collision/PBDCollisionContainerSolverSimd.cpp
+++ b/Engine/Source/Runtime/Experimental/Chaos/Private/Chaos/Collision/PBDCollisionContainerSolverSimd.cpp
@@ -292,8 +292,6 @@ namespace Chaos
 						StaticFrictionRatio,
 						Dt);
 				}
-
-				Constraint->FinalizeSolverResults(Dt);
 			}
 		}
 
diff --git a/Engine/Source/Runtime/Experimental/Chaos/Public/Chaos/Collision/ContactPoint.h b/Engine/Source/Runtime/Experimental/Chaos/Public/Chaos/Collision/ContactPoint.h
--- a/Engine/Source/Runtime/Experimental/Chaos/Public/Chaos/Collision/ContactPoint.h
+++ b/Engine/Source/Runtime/Experimental/Chaos/Public/Chaos/Collision/ContactPoint.h
@@ -170,6 +170,7 @@ namespace Chaos
 	{
 	public:
 		FVec3f ShapeContactPoints[2];			// Contact anchor points for friction
+		FRealSingle InitialPhi;					// Non-resolved initial penetration
 	};
 
 	class FManifoldPointResult
diff --git a/Engine/Source/Runtime/Experimental/Chaos/Public/Chaos/Collision/PBDCollisionConstraint.h b/Engine/Source/Runtime/Experimental/Chaos/Public/Chaos/Collision/PBDCollisionConstraint.h
--- a/Engine/Source/Runtime/Experimental/Chaos/Public/Chaos/Collision/PBDCollisionConstraint.h
+++ b/Engine/Source/Runtime/Experimental/Chaos/Public/Chaos/Collision/PBDCollisionConstraint.h
@@ -414,6 +414,7 @@ namespace Chaos
 		void SetIsInitialContact(const bool bInIsInitialContact) { Flags.bInitialContact = bInIsInitialContact; }
 		bool IsInitialContact() const { return Flags.bInitialContact; }
 		FRealSingle GetMinInitialPhi() const { return MinInitialPhi; }
+		bool UsePerContactInitialPhi() const { return Flags.bUsePerContactInitialPhi; }
 
 		virtual bool SupportsSleeping() const override final { return true; }
 		CHAOS_API virtual bool IsSleeping() const override final;
@@ -817,20 +818,12 @@ namespace Chaos
 			{
 				SavedManifoldPoint->ShapeContactPoints[0] = Anchor0;
 				SavedManifoldPoint->ShapeContactPoints[1] = Anchor1;
+				SavedManifoldPoint->InitialPhi = ManifoldPoint.InitialPhi;
 			}
 
 			MinInitialPhi = FMath::Min(MinInitialPhi, ManifoldPoint.InitialPhi);
 		}
 
-		void FinalizeSolverResults(const FRealSingle Dt)
-		{
-			if (MinInitialPhi < 0)
-			{
-				// Apply depentration velocity to the initial overlap
-				MinInitialPhi = FMath::Min(MinInitialPhi + InitialOverlapDepenetrationVelocity * Dt, 0.0f);
-			}
-		}
-
 		/**
 		 *	A key used to uniquely identify the constraint (it is based on the two particle IDs)
 		 */
@@ -953,6 +946,7 @@ namespace Chaos
 				uint16 bMaterialSet : 1;				// Has the material been set (or does it need to be reset)
 				uint16 bInitialContact : 1;				// Is this contact considered an initial contact
 				uint16 bIsOneWayInteraction : 1;		// Does one of the bodies have the one-way interaction bit set?
+				uint16 bUsePerContactInitialPhi : 1;	// Whether we track initial overlap per-contact or shared across all contacts in the manifold
 			};
 			uint16 Bits;
 		};
