[NetPhysics] Refactor CVars for render interpolation.

diff --git a/Engine/Source/Runtime/Experimental/Chaos/Private/Framework/PhysicsProxyBase.cpp b/Engine/Source/Runtime/Experimental/Chaos/Private/Framework/PhysicsProxyBase.cpp
--- a/Engine/Source/Runtime/Experimental/Chaos/Private/Framework/PhysicsProxyBase.cpp
+++ b/Engine/Source/Runtime/Experimental/Chaos/Private/Framework/PhysicsProxyBase.cpp
@@ -3,6 +3,24 @@
 #include "Chaos/Framework/PhysicsProxyBase.h"
 #include "Chaos/Framework/PhysicsSolverBase.h"
 
+namespace RenderInterpolationCVars
+{
+	float RenderInterpErrorCorrectionDuration = 0.5f;
+	static FAutoConsoleVariableRef CVarRenderInterpErrorCorrectionDuration(TEXT("p.RenderInterp.ErrorCorrectionDuration"), RenderInterpErrorCorrectionDuration, TEXT("How long in seconds to apply error correction over."));
+
+	float RenderInterpMaximumErrorCorrectionBeforeSnapping = 250.0f;
+	static FAutoConsoleVariableRef CVarRenderInterpErrorCorrectionMaximumError(TEXT("p.RenderInterp.MaximumErrorCorrectionBeforeSnapping"), RenderInterpMaximumErrorCorrectionBeforeSnapping, TEXT("Maximum error correction in cm before we stop interpolating and snap to target."));
+
+	float RenderInterpErrorVelocitySmoothingDuration = 0.5f;
+	static FAutoConsoleVariableRef CVarRenderInterpErrorVelocitySmoothingDuration(TEXT("p.RenderInterp.ErrorVelocitySmoothingDuration"), RenderInterpErrorVelocitySmoothingDuration, TEXT("How long in seconds to apply error velocity smoothing correction over, should be smaller than or equal to p.RenderInterp.ErrorCorrectionDuration. RENDERINTERPOLATION_VELOCITYSMOOTHING needs to be defined."));
+
+	float RenderInterpErrorDirectionalDecayMultiplier = 0.0f;
+	static FAutoConsoleVariableRef CVarRenderInterpErrorDirectionalDecayMultiplier(TEXT("p.RenderInterp.DirectionalDecayMultiplier"), RenderInterpErrorDirectionalDecayMultiplier, TEXT("Decay error offset in the direction that the physics object is moving, value is multiplier of projected offset direction, 0.25 means a 25% decay of the magnitude in the direction of physics travel. Deactivate by setting to 0."));
+
+	bool bRenderInterpDebugDraw = false;
+	static FAutoConsoleVariableRef CVarRenderInterpDebugDraw(TEXT("p.RenderInterp.DebugDraw"), bRenderInterpDebugDraw, TEXT("Draw debug lines for physics render interpolation, also needs p.Chaos.DebugDraw.Enabled set"));
+}
+
 IPhysicsProxyBase::~IPhysicsProxyBase()
 {
 	if(GetSolver<Chaos::FPhysicsSolverBase>())
diff --git a/Engine/Source/Runtime/Experimental/Chaos/Private/PhysicsProxy/ClusterUnionPhysicsProxy.cpp b/Engine/Source/Runtime/Experimental/Chaos/Private/PhysicsProxy/ClusterUnionPhysicsProxy.cpp
--- a/Engine/Source/Runtime/Experimental/Chaos/Private/PhysicsProxy/ClusterUnionPhysicsProxy.cpp
+++ b/Engine/Source/Runtime/Experimental/Chaos/Private/PhysicsProxy/ClusterUnionPhysicsProxy.cpp
@@ -481,11 +481,11 @@ namespace Chaos
 		if (Error)
 		{
 			const FReal ErrorMagSq = Error->ErrorX.SizeSquared();
-			const FReal MaxErrorCorrection = GetRenderInterpMaximumErrorCorrectionBeforeSnapping();
+			const FReal MaxErrorCorrection = RenderInterpolationCVars::RenderInterpMaximumErrorCorrectionBeforeSnapping;
 			int32 RenderInterpErrorCorrectionDurationTicks = 0;
 			if (ErrorMagSq < MaxErrorCorrection * MaxErrorCorrection)
 			{
-				RenderInterpErrorCorrectionDurationTicks = FMath::FloorToInt32(GetRenderInterpErrorCorrectionDuration() / AsyncFixedTimeStep); // Convert duration from seconds to simulation ticks
+				RenderInterpErrorCorrectionDurationTicks = FMath::FloorToInt32(RenderInterpolationCVars::RenderInterpErrorCorrectionDuration / AsyncFixedTimeStep); // Convert duration from seconds to simulation ticks
 			}
 			InterpolationData.AccumlateErrorXR(Error->ErrorX, Error->ErrorR, SolverSyncTimestamp, RenderInterpErrorCorrectionDurationTicks);
 		}
@@ -565,7 +565,7 @@ namespace Chaos
 
 			if (const FVec3* Prev = LerpHelper(PullData.X, ProxyTimestamp->OverWriteX))
 			{
-				if (GetRenderInterpErrorDirectionalDecayMultiplier() > 0.0f)
+				if (RenderInterpolationCVars::RenderInterpErrorDirectionalDecayMultiplier > 0.0f)
 				{
 					DirectionalDecayPerformed = InterpolationData.DirectionalDecay(NextPullData->X - *Prev);
 				}
@@ -599,7 +599,7 @@ namespace Chaos
 			}
 
 #if CHAOS_DEBUG_DRAW
-			if (GetRenderInterpDebugDraw())
+			if (RenderInterpolationCVars::bRenderInterpDebugDraw)
 			{
 				Chaos::FDebugDrawQueue::GetInstance().DrawDebugBox(NextPullData->X, FVector(2, 1, 1), NextPullData->R, FColor::Yellow, false, 5.f, 0, 0.5f);
 				Chaos::FDebugDrawQueue::GetInstance().DrawDebugDirectionalArrow(PullData.X, NextPullData->X, 0.5f, FColor::Yellow, false, 5.0f, 0, 0.5f);
diff --git a/Engine/Source/Runtime/Experimental/Chaos/Private/PhysicsProxy/GeometryCollectionPhysicsProxy.cpp b/Engine/Source/Runtime/Experimental/Chaos/Private/PhysicsProxy/GeometryCollectionPhysicsProxy.cpp
--- a/Engine/Source/Runtime/Experimental/Chaos/Private/PhysicsProxy/GeometryCollectionPhysicsProxy.cpp
+++ b/Engine/Source/Runtime/Experimental/Chaos/Private/PhysicsProxy/GeometryCollectionPhysicsProxy.cpp
@@ -4298,11 +4298,11 @@ bool FGeometryCollectionPhysicsProxy::PullFromPhysicsState(const Chaos::FDirtyGe
 	if (Error)
 	{
 		const Chaos::FReal ErrorMagSq = Error->ErrorX.SizeSquared();
-		const Chaos::FReal MaxErrorCorrection = GetRenderInterpMaximumErrorCorrectionBeforeSnapping();
+		const Chaos::FReal MaxErrorCorrection = RenderInterpolationCVars::RenderInterpMaximumErrorCorrectionBeforeSnapping;
 		int32 RenderInterpErrorCorrectionDurationTicks = 0;
 		if (ErrorMagSq < MaxErrorCorrection * MaxErrorCorrection)
 		{
-			RenderInterpErrorCorrectionDurationTicks = FMath::FloorToInt32(GetRenderInterpErrorCorrectionDuration() / AsyncFixedTimeStep); // Convert duration from seconds to simulation ticks
+			RenderInterpErrorCorrectionDurationTicks = FMath::FloorToInt32(RenderInterpolationCVars::RenderInterpErrorCorrectionDuration / AsyncFixedTimeStep); // Convert duration from seconds to simulation ticks
 		}
 		InterpolationData.AccumlateErrorXR(Error->ErrorX, Error->ErrorR, SolverSyncTimestamp, RenderInterpErrorCorrectionDurationTicks);
 	}
@@ -4334,7 +4334,7 @@ bool FGeometryCollectionPhysicsProxy::PullFromPhysicsState(const Chaos::FDirtyGe
 			const FGeometryCollectionResults& NextResults = NextPullData->Results();
 
 			InterpolationData.UpdateError(SolverSyncTimestamp, AsyncFixedTimeStep);
-			if (IPhysicsProxyBase::GetRenderInterpErrorDirectionalDecayMultiplier() > 0.0f && PrevResults.GetNumEntries() && NextResults.GetNumEntries())
+			if (RenderInterpolationCVars::RenderInterpErrorDirectionalDecayMultiplier > 0.0f && PrevResults.GetNumEntries() && NextResults.GetNumEntries())
 			{
 				InterpolationData.DirectionalDecay(NextResults.GetPositions(0).ParticleX - PrevResults.GetPositions(0).ParticleX);
 			}
diff --git a/Engine/Source/Runtime/Experimental/Chaos/Private/PhysicsProxy/SingleParticlePhysicsProxy.cpp b/Engine/Source/Runtime/Experimental/Chaos/Private/PhysicsProxy/SingleParticlePhysicsProxy.cpp
--- a/Engine/Source/Runtime/Experimental/Chaos/Private/PhysicsProxy/SingleParticlePhysicsProxy.cpp
+++ b/Engine/Source/Runtime/Experimental/Chaos/Private/PhysicsProxy/SingleParticlePhysicsProxy.cpp
@@ -292,11 +292,11 @@ bool FSingleParticlePhysicsProxy::PullFromPhysicsState(const Chaos::FDirtyRigidP
 		if (Error)
 		{
 			const FReal ErrorMagSq = Error->ErrorX.SizeSquared();
-			const FReal MaxErrorCorrection = GetRenderInterpMaximumErrorCorrectionBeforeSnapping();
+			const FReal MaxErrorCorrection = RenderInterpolationCVars::RenderInterpMaximumErrorCorrectionBeforeSnapping;
 			int32 RenderInterpErrorCorrectionDurationTicks = 0;
 			if (ErrorMagSq < MaxErrorCorrection * MaxErrorCorrection)
 			{
-				RenderInterpErrorCorrectionDurationTicks = FMath::FloorToInt32(GetRenderInterpErrorCorrectionDuration() / AsyncFixedTimeStep); // Convert duration from seconds to simulation ticks
+				RenderInterpErrorCorrectionDurationTicks = FMath::FloorToInt32(RenderInterpolationCVars::RenderInterpErrorCorrectionDuration / AsyncFixedTimeStep); // Convert duration from seconds to simulation ticks
 			}
 			InterpolationData.AccumlateErrorXR(Error->ErrorX, Error->ErrorR, SolverSyncTimestamp, RenderInterpErrorCorrectionDurationTicks);
 #if RENDERINTERP_ERRORVELOCITYSMOOTHING
@@ -330,7 +330,7 @@ bool FSingleParticlePhysicsProxy::PullFromPhysicsState(const Chaos::FDirtyRigidP
 					FVec3 Target = FMath::Lerp(*Prev, NextPullData->X, *Alpha);
 					if (bIsReplicationErrorSmoothing)
 					{
-						if (GetRenderInterpErrorDirectionalDecayMultiplier() > 0.0f)
+						if (RenderInterpolationCVars::RenderInterpErrorDirectionalDecayMultiplier > 0.0f)
 						{
 							DirectionalDecayPerformed = InterpolationData.DirectionalDecay(NextPullData->X - *Prev);
 						}
@@ -367,7 +367,7 @@ bool FSingleParticlePhysicsProxy::PullFromPhysicsState(const Chaos::FDirtyRigidP
 				}
 				
 #if CHAOS_DEBUG_DRAW
-				if (GetRenderInterpDebugDraw())
+				if (RenderInterpolationCVars::bRenderInterpDebugDraw)
 				{
 					Chaos::FDebugDrawQueue::GetInstance().DrawDebugBox(NextPullData->X, FVector(2, 1, 1), NextPullData->R, FColor::Yellow, false, 5.f, 0, 0.5f);
 					Chaos::FDebugDrawQueue::GetInstance().DrawDebugDirectionalArrow(PullData.X, NextPullData->X, 0.5f, FColor::Yellow, false, 5.0f, 0, 0.5f);
diff --git a/Engine/Source/Runtime/Experimental/Chaos/Public/Chaos/Framework/PhysicsProxyBase.h b/Engine/Source/Runtime/Experimental/Chaos/Public/Chaos/Framework/PhysicsProxyBase.h
--- a/Engine/Source/Runtime/Experimental/Chaos/Public/Chaos/Framework/PhysicsProxyBase.h
+++ b/Engine/Source/Runtime/Experimental/Chaos/Public/Chaos/Framework/PhysicsProxyBase.h
@@ -27,6 +27,15 @@ namespace Chaos
 	class FPhysicsSolverBase;
 }
 
+namespace RenderInterpolationCVars
+{
+	extern float RenderInterpErrorCorrectionDuration;
+	extern float RenderInterpMaximumErrorCorrectionBeforeSnapping;
+	extern float RenderInterpErrorVelocitySmoothingDuration;
+	extern float RenderInterpErrorDirectionalDecayMultiplier;
+	extern bool bRenderInterpDebugDraw;
+}
+
 struct FProxyTimestampBase
 {
 	bool bDeleted = false;
@@ -123,35 +132,30 @@ public:
 	void SetParentProxy(IPhysicsProxyBase* InProxy) { ParentProxy = InProxy; }
 
 	// Render Interpolation CVars
+	UE_DEPRECATED(5.5, "Deprecated, use RenderInterpolationCVars::RenderInterpErrorCorrectionDuration")
 	static float GetRenderInterpErrorCorrectionDuration()
 	{
-		static float RenderInterpErrorCorrectionDuration = 0.5f;
-		static FAutoConsoleVariableRef CVarRenderInterpErrorCorrectionDuration(TEXT("p.RenderInterp.ErrorCorrectionDuration"), RenderInterpErrorCorrectionDuration, TEXT("How long in seconds to apply error correction over."));
-		return RenderInterpErrorCorrectionDuration;
+		return RenderInterpolationCVars::RenderInterpErrorCorrectionDuration;
 	}
+	UE_DEPRECATED(5.5, "Deprecated, use RenderInterpolationCVars::RenderInterpMaximumErrorCorrectionBeforeSnapping")
 	static float GetRenderInterpMaximumErrorCorrectionBeforeSnapping()
 	{
-		static float RenderInterpMaximumErrorCorrectionBeforeSnapping = 250.0f;
-		static FAutoConsoleVariableRef CVarRenderInterpErrorCorrectionMaximumError(TEXT("p.RenderInterp.MaximumErrorCorrectionBeforeSnapping"), RenderInterpMaximumErrorCorrectionBeforeSnapping, TEXT("Maximum error correction in cm before we stop interpolating and snap to target."));
-		return RenderInterpMaximumErrorCorrectionBeforeSnapping;
+		return RenderInterpolationCVars::RenderInterpMaximumErrorCorrectionBeforeSnapping;
 	}
+	UE_DEPRECATED(5.5, "Deprecated, use RenderInterpolationCVars::RenderInterpErrorVelocitySmoothingDuration")
 	static float GetRenderInterpErrorVelocitySmoothingDuration()
 	{
-		static float RenderInterpErrorVelocitySmoothingDuration = 0.5f;
-		static FAutoConsoleVariableRef CVarRenderInterpErrorVelocitySmoothingDuration(TEXT("p.RenderInterp.ErrorVelocitySmoothingDuration"), RenderInterpErrorVelocitySmoothingDuration, TEXT("How long in seconds to apply error velocity smoothing correction over, should be smaller than or equal to p.RenderInterp.ErrorCorrectionDuration. RENDERINTERPOLATION_VELOCITYSMOOTHING needs to be defined."));
-		return RenderInterpErrorVelocitySmoothingDuration;
+		return RenderInterpolationCVars::RenderInterpErrorVelocitySmoothingDuration;
 	}
+	UE_DEPRECATED(5.5, "Deprecated, use RenderInterpolationCVars::bRenderInterpDebugDraw")
 	static bool GetRenderInterpDebugDraw()
 	{
-		static bool RenderInterpDebugDraw = false;
-		static FAutoConsoleVariableRef CVarRenderInterpDebugDraw(TEXT("p.RenderInterp.DebugDraw"), RenderInterpDebugDraw, TEXT("Draw debug lines for physics render interpolation, also needs p.Chaos.DebugDraw.Enabled set"));
-		return RenderInterpDebugDraw;
+		return RenderInterpolationCVars::bRenderInterpDebugDraw;
 	}
+	UE_DEPRECATED(5.5, "Deprecated, use RenderInterpolationCVars::RenderInterpErrorDirectionalDecayMultiplier")
 	static float GetRenderInterpErrorDirectionalDecayMultiplier()
 	{
-		static float RenderInterpErrorDirectionalDecayMultiplier = 0.0f;
-		static FAutoConsoleVariableRef CVarRenderInterpErrorDirectionalDecayMultiplier(TEXT("p.RenderInterp.DirectionalDecayMultiplier"), RenderInterpErrorDirectionalDecayMultiplier, TEXT("Decay error offset in the direction that the physics object is moving, value is multiplier of projected offset direction, 0.25 means a 25% decay of the magnitude in the direction of physics travel. Deactivate by setting to 0."));
-		return RenderInterpErrorDirectionalDecayMultiplier;
+		return RenderInterpolationCVars::RenderInterpErrorDirectionalDecayMultiplier;
 	}
 
 protected:
@@ -262,7 +266,7 @@ struct FProxyInterpolationError : FProxyInterpolationBase
 			Chaos::FRealDouble DotProd = Chaos::FVec3::DotProduct(DirectionNormal, ErrorX);
 			if (DotProd > 0.0f)
 			{
-				Chaos::FVec3 DirProjection = ErrorX.ProjectOnToNormal(DirectionNormal) * IPhysicsProxyBase::GetRenderInterpErrorDirectionalDecayMultiplier();
+				Chaos::FVec3 DirProjection = ErrorX.ProjectOnToNormal(DirectionNormal) * RenderInterpolationCVars::RenderInterpErrorDirectionalDecayMultiplier;
 				ErrorX -= DirProjection;
 				return true;
 			}
