Fix failed check on undo of geometry collection changes in fracture mode

diff --git a/Engine/Source/Runtime/Experimental/GeometryCollectionEngine/Private/GeometryCollection/GeometryCollectionObject.cpp b/Engine/Source/Runtime/Experimental/GeometryCollectionEngine/Private/GeometryCollection/GeometryCollectionObject.cpp
--- a/Engine/Source/Runtime/Experimental/GeometryCollectionEngine/Private/GeometryCollection/GeometryCollectionObject.cpp
+++ b/Engine/Source/Runtime/Experimental/GeometryCollectionEngine/Private/GeometryCollection/GeometryCollectionObject.cpp
@@ -1752,6 +1752,12 @@ FGuid UGeometryCollection::GetStateGuid() const
 
 #if WITH_EDITOR
 
+void UGeometryCollection::PostEditUndo()
+{
+	PropagateTransformUpdateToComponents();
+	Super::PostEditUndo();
+}
+
 void UGeometryCollection::PostEditChangeProperty(struct FPropertyChangedEvent& PropertyChangedEvent)
 {
 	bool bDoInvalidateCollection = false;
diff --git a/Engine/Source/Runtime/Experimental/GeometryCollectionEngine/Public/GeometryCollection/GeometryCollectionObject.h b/Engine/Source/Runtime/Experimental/GeometryCollectionEngine/Public/GeometryCollection/GeometryCollectionObject.h
--- a/Engine/Source/Runtime/Experimental/GeometryCollectionEngine/Public/GeometryCollection/GeometryCollectionObject.h
+++ b/Engine/Source/Runtime/Experimental/GeometryCollectionEngine/Public/GeometryCollection/GeometryCollectionObject.h
@@ -367,6 +367,7 @@ public:
 	/** UObject Interface */
 #if WITH_EDITOR
 	GEOMETRYCOLLECTIONENGINE_API virtual void PostEditChangeProperty(struct FPropertyChangedEvent& PropertyChangedEvent) override;
+	GEOMETRYCOLLECTIONENGINE_API virtual void PostEditUndo() override;
 	GEOMETRYCOLLECTIONENGINE_API virtual bool Modify(bool bAlwaysMarkDirty = true) override;
 #endif
 	GEOMETRYCOLLECTIONENGINE_API virtual void PostInitProperties() override;
