Stub mounting interface for FIoStoreOnDemandModule

diff --git a/Engine/Source/Runtime/Experimental/IoStoreOnDemand/Private/IoStoreOnDemand.cpp b/Engine/Source/Runtime/Experimental/IoStoreOnDemand/Private/IoStoreOnDemand.cpp
--- a/Engine/Source/Runtime/Experimental/IoStoreOnDemand/Private/IoStoreOnDemand.cpp
+++ b/Engine/Source/Runtime/Experimental/IoStoreOnDemand/Private/IoStoreOnDemand.cpp
@@ -45,6 +45,7 @@
 #include "S3/S3Client.h"
 #endif // (PLATFORM_DESKTOP && (IS_PROGRAM || WITH_EDITOR))
 
+DEFINE_LOG_CATEGORY(LogIoStoreOnDemand);
 DEFINE_LOG_CATEGORY(LogIas);
 
 namespace UE::IoStore
@@ -2514,6 +2515,16 @@ void FIoStoreOnDemandModule::ReportAnalytics(TArray<FAnalyticsEventAttribute>& O
 	}
 }
 
+void FIoStoreOnDemandModule::Mount(FOnDemandMountArgs&& Args, FOnDemandMountCompleted&& OnCompleted)
+{
+	OnCompleted(TIoStatusOr<FOnDemandMountResult>(FIoStatus(EIoErrorCode::InvalidCode, TEXT("Mount not implemented"))));
+}
+
+FIoStatus FIoStoreOnDemandModule::Unmount(FStringView MountId)
+{
+	return FIoStatus(EIoErrorCode::InvalidCode, TEXT("Unmount not implemented"));
+}
+
 void FIoStoreOnDemandModule::InitializeInternal()
 {
 	LLM_SCOPE_BYTAG(Ias);
diff --git a/Engine/Source/Runtime/Experimental/IoStoreOnDemand/Public/IO/IoStoreOnDemand.h b/Engine/Source/Runtime/Experimental/IoStoreOnDemand/Public/IO/IoStoreOnDemand.h
--- a/Engine/Source/Runtime/Experimental/IoStoreOnDemand/Public/IO/IoStoreOnDemand.h
+++ b/Engine/Source/Runtime/Experimental/IoStoreOnDemand/Public/IO/IoStoreOnDemand.h
@@ -36,6 +36,7 @@ using FIoBlockHash = uint32;
 	#define UE_IAS_CUSTOM_INITIALIZATION 0
 #endif
 
+UE_API DECLARE_LOG_CATEGORY_EXTERN(LogIoStoreOnDemand, Log, All);
 UE_API DECLARE_LOG_CATEGORY_EXTERN(LogIas, Log, All);
 
 namespace UE::IoStore
@@ -341,6 +342,21 @@ enum class EOnDemandInitResult
 
 #endif // UE_IAS_CUSTOM_INITIALIZATION
 
+struct FOnDemandMountArgs
+{
+	FString MountId;
+	FString Url;
+	FString InstallDirectory;
+	bool bInstall = false;
+};
+
+struct FOnDemandMountResult
+{
+	FString MountId;
+};
+
+using FOnDemandMountCompleted = TFunction<void(TIoStatusOr<FOnDemandMountResult>)>;
+
 class FIoStoreOnDemandModule
 	: public IModuleInterface
 {
@@ -361,6 +377,9 @@ public:
 
 	UE_API void ReportAnalytics(TArray<FAnalyticsEventAttribute>& OutAnalyticsArray) const;
 
+	UE_API void Mount(FOnDemandMountArgs&& Args, FOnDemandMountCompleted&& OnCompleted);
+	UE_API FIoStatus Unmount(FStringView MountId);
+
 	virtual void StartupModule() override;
 	virtual void ShutdownModule() override;
 	
