Dataflow : Fix reroute node autowire
Also fix the ReRoute node display name to not include Dataflow

diff --git a/Engine/Plugins/Experimental/Dataflow/Source/DataflowEditor/Private/Dataflow/DataflowSchema.cpp b/Engine/Plugins/Experimental/Dataflow/Source/DataflowEditor/Private/Dataflow/DataflowSchema.cpp
--- a/Engine/Plugins/Experimental/Dataflow/Source/DataflowEditor/Private/Dataflow/DataflowSchema.cpp
+++ b/Engine/Plugins/Experimental/Dataflow/Source/DataflowEditor/Private/Dataflow/DataflowSchema.cpp
@@ -177,7 +177,9 @@ const FPinConnectionResponse UDataflowSchema::CanCreateConnection(const UEdGraph
 				// Make sure types match. 
 				const bool bIsAnyTypeA = PinA->PinType.PinCategory == FDataflowAnyType::TypeName;
 				const bool bIsAnyTypeB = PinB->PinType.PinCategory == FDataflowAnyType::TypeName;
-				if (PinA->PinType == PinB->PinType || (bIsAnyTypeA != bIsAnyTypeB))
+				const bool bAreSameType = (PinA->PinType == PinB->PinType);
+				const bool bOnlyOneIsAnyType = (!bAreSameType && (bIsAnyTypeA != bIsAnyTypeB));
+				if (bAreSameType || bOnlyOneIsAnyType)
 				{
 					// cycle checking on connect
 					if (!HasLoopIfConnected(PinA->GetOwningNode(), PinB->GetOwningNode()))
diff --git a/Engine/Source/Runtime/Experimental/Dataflow/Core/Public/Dataflow/DataflowNode.h b/Engine/Source/Runtime/Experimental/Dataflow/Core/Public/Dataflow/DataflowNode.h
--- a/Engine/Source/Runtime/Experimental/Dataflow/Core/Public/Dataflow/DataflowNode.h
+++ b/Engine/Source/Runtime/Experimental/Dataflow/Core/Public/Dataflow/DataflowNode.h
@@ -401,7 +401,7 @@ struct FDataflowReRouteNode : public FDataflowNode
 {
 	GENERATED_USTRUCT_BODY()
 
-	DATAFLOW_NODE_DEFINE_INTERNAL(FDataflowReRouteNode, "DataflowReRouteNode", "Core", "")
+	DATAFLOW_NODE_DEFINE_INTERNAL(FDataflowReRouteNode, "ReRouteNode", "Core", "")
 
 public:
 	FDataflowReRouteNode(const Dataflow::FNodeParameters& Param, FGuid InGuid = FGuid::NewGuid())
diff --git a/Engine/Source/Runtime/Experimental/Dataflow/Engine/Private/Dataflow/DataflowEdNode.cpp b/Engine/Source/Runtime/Experimental/Dataflow/Engine/Private/Dataflow/DataflowEdNode.cpp
--- a/Engine/Source/Runtime/Experimental/Dataflow/Engine/Private/Dataflow/DataflowEdNode.cpp
+++ b/Engine/Source/Runtime/Experimental/Dataflow/Engine/Private/Dataflow/DataflowEdNode.cpp
@@ -487,14 +487,38 @@ void UDataflowEdNode::AutowireNewNode(UEdGraphPin* FromPin)
 				{
 					if (FDataflowOutput* FromOutput = FromDataFlowNode->FindOutput(FromPin->PinName))
 					{
-						const FName OutputType = FromOutput->GetType();
-
 						const TSharedPtr<FDataflowNode> ToDataFlowNode = this->GetDataflowNode();
 						for (UEdGraphPin* InputPin : this->GetAllPins())
 						{
 							if (FDataflowInput* ToInput = ToDataFlowNode->FindInput(InputPin->PinName))
 							{
-								if (ToInput->GetType() == OutputType)
+								const bool bAreSameType = (ToInput->GetType() == FromOutput->GetType());
+								const bool bOnlyOneIsAnyType = (!bAreSameType && (ToInput->GetType() == FDataflowAnyType::TypeName) != (FromOutput->GetType() == FDataflowAnyType::TypeName));
+								if (bAreSameType || bOnlyOneIsAnyType)
+								{
+									if (const UEdGraph* EdGraph = this->GetGraph())
+									{
+										if (EdGraph->GetSchema()->TryCreateConnection(FromPin, InputPin))
+										{
+											FromGraphNode->NodeConnectionListChanged();
+											this->NodeConnectionListChanged();
+											break;
+										}
+									}
+								}
+							}
+						}
+					}
+					else if(FDataflowInput* FromInput = FromDataFlowNode->FindInput(FromPin->PinName))
+					{
+						const TSharedPtr<FDataflowNode> ToDataFlowNode = this->GetDataflowNode();
+						for (UEdGraphPin* InputPin : this->GetAllPins())
+						{
+							if (FDataflowOutput* ToOutput = ToDataFlowNode->FindOutput(InputPin->PinName))
+							{
+								const bool bAreSameType = (ToOutput->GetType() == FromInput->GetType());
+								const bool bOnlyOneIsAnyType = (!bAreSameType && (ToOutput->GetType() == FDataflowAnyType::TypeName) != (FromInput->GetType() == FDataflowAnyType::TypeName));
+								if (bAreSameType || bOnlyOneIsAnyType)
 								{
 									if (const UEdGraph* EdGraph = this->GetGraph())
 									{
@@ -596,4 +620,3 @@ bool UDataflowEdNode::Render(GeometryCollection::Facades::FRenderingFacade& Rend
 
 #undef LOCTEXT_NAMESPACE
 
-
