Chaos - fix capsule-triangle collision corner case

diff --git a/Engine/Source/Runtime/Experimental/Chaos/Private/Chaos/Collision/CapsuleTriangleContactPoint.cpp b/Engine/Source/Runtime/Experimental/Chaos/Private/Chaos/Collision/CapsuleTriangleContactPoint.cpp
--- a/Engine/Source/Runtime/Experimental/Chaos/Private/Chaos/Collision/CapsuleTriangleContactPoint.cpp
+++ b/Engine/Source/Runtime/Experimental/Chaos/Private/Chaos/Collision/CapsuleTriangleContactPoint.cpp
@@ -167,7 +167,7 @@ namespace Chaos
 
 			// Separating axis always points away from the triangle
 			const FReal DotEdge = FVec3::DotProduct(SegmentEdgeN, EdgeNs[EdgeIndex]);
-			if (DotEdge < FReal(0))
+			if (DotEdge < FReal(-NormalTolerance))
 			{
 				SegmentEdgeN = -SegmentEdgeN;
 				SegmentEdgeDistSign = FReal(-1);
@@ -352,7 +352,7 @@ namespace Chaos
 			if (((EdgeD0s[EdgeIndex] > -DistanceTolerance) | (EdgeD1s[EdgeIndex] > -DistanceTolerance)) != 0)
 			{
 				// Don't collide with inside face
-				if (DotFace < FReal(0))
+				if (DotFace < FReal(-NormalTolerance))
 				{
 					continue;
 				}
@@ -386,35 +386,45 @@ namespace Chaos
 					{
 						SegmentEdgeN = -SegmentEdgeN;
 					}
+					const FReal DotEdge = FVec3::DotProduct(EdgeNs[EdgeIndex], SegmentEdgeN);
+					if (DotEdge < -NormalTolerance)
+					{
+						continue;
+					}
 				}
 
-				// We cannot collide with the "inside" of the edge (normal must be facing away from triangle center).
-				// However, we can convert it to a face contact if we are within the Face Contact angle threshold.
 				if ((bCrossedEdgeSegment & bPreferFaceContact) == 0)
 				{
-					// We use a tolerance here so we don't reject very nearly face collisions
-					// @todo(chaos): size dependence issue?
-					const FReal DotCentroid = FVec3::DotProduct(EdgeP - Centroid, SegmentEdgeN);
-					if (DotCentroid < -NormalTolerance)
+					// We cannot collide with the inside of the edge
+					const FReal DotEdge = FVec3::DotProduct(EdgeNs[EdgeIndex], SegmentEdgeN);
+					if (DotEdge < -NormalTolerance)
 					{
 						continue;
 					}
-				
-					// For Vertex contacts, the normal needs to be outside both planes
-					if ((EdgeT == FReal(0)) && ((SegmentT == FReal(0)) || (SegmentT == FReal(1))))
+
+					// For Vertex contacts, check that the normal is in the valid range
+					if (EdgeT == FReal(0))
 					{
-						const int32 PrevEdgeIndex = (EdgeIndex > 0) ? EdgeIndex - 1 : 2;
-						const FReal PrevDotEdge = FVec3::DotProduct(EdgeNs[PrevEdgeIndex], SegmentEdgeN);
-						if (PrevDotEdge < -NormalTolerance)
+						const int32 PrevEdgeVertexIndex0 = (EdgeIndex >= 2) ? (EdgeIndex - 2) : (EdgeIndex - 2 + 3);
+						const FVec3& PrevEdgeP0 = Triangle.GetVertex(PrevEdgeVertexIndex0);
+						const FVec3& PrevEdgeP1 = EdgeP0;
+
+						const FReal PrevEdgeDotNormal = FVec3::DotProduct(PrevEdgeP1 - PrevEdgeP0, SegmentEdgeN);
+						const FReal EdgeDotNormal = FVec3::DotProduct(EdgeP1 - EdgeP0, SegmentEdgeN);
+						if ((PrevEdgeDotNormal < -NormalTolerance) || (EdgeDotNormal <- NormalTolerance))
 						{
 							continue;
 						}
 					}
-					if ((EdgeT == FReal(1)) && ((SegmentT == FReal(0)) || (SegmentT == FReal(1))))
+					if (EdgeT == FReal(1))
 					{
-						const int32 NextEdgeIndex = (EdgeIndex < 2) ? EdgeIndex + 1 : 0;
-						const FReal NextDotEdge = FVec3::DotProduct(EdgeNs[NextEdgeIndex], SegmentEdgeN);
-						if (NextDotEdge < -NormalTolerance)
+						const int32 NextEdgeVertexIndex1 = (EdgeIndex < 2) ? EdgeIndex + 1 : 0;
+						const FVec3& NextEdgeP0 = EdgeP1;
+						const FVec3& NextEdgeP1 = Triangle.GetVertex(NextEdgeVertexIndex1);
+
+						const FReal EdgeDotNormal = FVec3::DotProduct(EdgeP1 - EdgeP0, SegmentEdgeN);
+						const FReal NextEdgeDotNormal = FVec3::DotProduct(NextEdgeP0 - NextEdgeP1, SegmentEdgeN);
+						if ((EdgeDotNormal < -NormalTolerance) || (NextEdgeDotNormal <- NormalTolerance))
 						{
 							continue;
 						}
