Moved default port determination into FHost. Simplified pool host name parameter.

#rnx

diff --git a/Engine/Source/Runtime/Experimental/IoStoreOnDemand/Private/CoreHttp/Client.h b/Engine/Source/Runtime/Experimental/IoStoreOnDemand/Private/CoreHttp/Client.h
--- a/Engine/Source/Runtime/Experimental/IoStoreOnDemand/Private/CoreHttp/Client.h
+++ b/Engine/Source/Runtime/Experimental/IoStoreOnDemand/Private/CoreHttp/Client.h
@@ -56,14 +56,11 @@ public:
 	struct FParams
 	{
 		int32				SetHostFromUrl(FAnsiStringView Url);
-		uint16				ConnectionCount = 1;
+		FAnsiStringView		HostName;
 		int32				SendBufSize = -1;
 		int32				RecvBufSize = -1;
-		struct {
-			FAnsiStringView	Name;
-			uint32			Port = 80;
-			/* uint8		Protocol */;
-		}					Host;
+		uint32				Port = 0;
+		uint16				ConnectionCount = 1;
 		/*
 		enum class ProxyType { Http, Socks4 };
 		Proxy = { ip, port, type }
diff --git a/Engine/Source/Runtime/Experimental/IoStoreOnDemand/Private/CoreHttp/Http.cpp b/Engine/Source/Runtime/Experimental/IoStoreOnDemand/Private/CoreHttp/Http.cpp
--- a/Engine/Source/Runtime/Experimental/IoStoreOnDemand/Private/CoreHttp/Http.cpp
+++ b/Engine/Source/Runtime/Experimental/IoStoreOnDemand/Private/CoreHttp/Http.cpp
@@ -1290,7 +1290,14 @@ public:
 	enum class EDirection : uint8 { Send, Recv };
 	static const uint32 InvalidIp = 0x00ff'ffff;
 
-					FHost(const ANSICHAR* InHostName, uint32 InPort, uint32 InMaxConn);
+	struct FParams
+	{
+		const ANSICHAR* HostName;
+		uint32			Port = 0;
+		uint32			MaxConnections = 1;
+	};
+
+					FHost(const FParams& Params);
 	void			SetBufferSize(EDirection Dir, int32 Size);
 	int32			GetBufferSize(EDirection Dir) const;
 	FOutcome		Connect(FSocket& Socket);
@@ -1311,12 +1318,17 @@ private:
 };
 
 ////////////////////////////////////////////////////////////////////////////////
-FHost::FHost(const ANSICHAR* InHostName, uint32 InPort, uint32 InMaxConn)
-: HostName(InHostName)
-, Port(uint16(InPort))
-, MaxConnections(uint8(InMaxConn))
+FHost::FHost(const FParams& Params)
+: HostName(Params.HostName)
+, Port(uint16(Params.Port))
+, MaxConnections(uint8(Params.MaxConnections))
 {
-	check(MaxConnections && MaxConnections == InMaxConn);
+	check(MaxConnections && MaxConnections == Params.MaxConnections);
+
+	if (Port == 0)
+	{
+		Port = 80;
+	}
 }
 
 ////////////////////////////////////////////////////////////////////////////////
@@ -1479,12 +1491,12 @@ int32 FConnectionPool::FParams::SetHostFromUrl(FAnsiStringView Url)
 		return -1;
 	}
 
-	Host.Name = Offsets.HostName.Get(Url);
+	HostName = Offsets.HostName.Get(Url);
 
 	if (Offsets.Port)
 	{
 		FAnsiStringView PortView = Offsets.Port.Get(Url);
-		Host.Port = uint16(CrudeToInt(PortView));
+		Port = uint16(CrudeToInt(PortView));
 	}
 
 	return Offsets.Path;
@@ -1494,24 +1506,24 @@ int32 FConnectionPool::FParams::SetHostFromUrl(FAnsiStringView Url)
 FConnectionPool::FConnectionPool(const FParams& Params)
 {
 	check(Params.ConnectionCount - 1 <= 63u);
-	check(Params.Host.Port - 1 <= 0xfffeu);
+	check(Params.Port <= 0xffffu);
 
 	// Alloc a new internal object
-	uint32 HostNameLen = Params.Host.Name.Len();
+	uint32 HostNameLen = Params.HostName.Len();
 	uint32 AllocSize = sizeof(FHost) + (HostNameLen + 1);
 	auto* Internal = (FHost*)FMemory::Malloc(AllocSize, alignof(FHost));
 
 	// Copy host
 	char* HostDest = (char*)(Internal + 1);
-	memcpy(HostDest, Params.Host.Name.GetData(), HostNameLen);
+	memcpy(HostDest, Params.HostName.GetData(), HostNameLen);
 	HostDest[HostNameLen] = '\0';
 
 	// Init internal object
-	new (Internal) FHost(
-		HostDest,
-		Params.Host.Port,
-		Params.ConnectionCount
-	);
+	new (Internal) FHost({
+		.HostName		= HostDest,
+		.Port			= Params.Port,
+		.MaxConnections	= Params.ConnectionCount,
+	});
 	Internal->SetBufferSize(FHost::EDirection::Send, Params.SendBufSize);
 	Internal->SetBufferSize(FHost::EDirection::Recv, Params.RecvBufSize);
 
@@ -3375,14 +3387,7 @@ FRequest FEventLoop::Request(
 
 	FAnsiStringView HostName = UrlOffsets.HostName.Get(Url);
 
-	uint32 Port = 80;
-	if (UrlOffsets.SchemeLength == 5)
-	{
-		//Port = 443;
-		//Protocol = Protocol::Tls;
-		return FRequest();
-	}
-
+	uint32 Port = 0;
 	if (UrlOffsets.Port)
 	{
 		FAnsiStringView PortView = UrlOffsets.Port.Get(Url);
@@ -3415,7 +3420,10 @@ FRequest FEventLoop::Request(
 
 	memcpy(HostNamePtr, HostName.GetData(), HostNameLength);
 	HostNamePtr[HostNameLength] = '\0';
-	new (Host) FHost(HostNamePtr, Port, 1);
+	new (Host) FHost({
+		.HostName	= HostNamePtr,
+		.Port		= Port,
+	});
 
 	return Impl->Request(Method, Path, Activity);
 }
