Response.ReadHeaders() to fully expose the repsonse

#rnx

diff --git a/Engine/Source/Runtime/Experimental/IoStoreOnDemand/Private/CoreHttp/Client.h b/Engine/Source/Runtime/Experimental/IoStoreOnDemand/Private/CoreHttp/Client.h
--- a/Engine/Source/Runtime/Experimental/IoStoreOnDemand/Private/CoreHttp/Client.h
+++ b/Engine/Source/Runtime/Experimental/IoStoreOnDemand/Private/CoreHttp/Client.h
@@ -114,6 +114,8 @@ private:
 class UE_API FResponse
 {
 public:
+	using FHeaderSink = TFunction<bool (FAnsiStringView, FAnsiStringView)>;
+
 	EStatusCodeClass	GetStatus() const;
 	uint32				GetStatusCode() const;
 	FAnsiStringView 	GetStatusMessage() const;
@@ -121,6 +123,7 @@ public:
 	EMimeType			GetContentType() const;
 	void				GetContentType(FAnsiStringView& Out) const;
 	FAnsiStringView 	GetHeader(FAnsiStringView Name) const;
+	void				ReadHeaders(FHeaderSink Sink) const;
 	void				SetDestination(FIoBuffer* Buffer);
 
 private:
@@ -185,7 +188,6 @@ public:
 	struct FRequestParams
 	{
 		uint32	BufferSize	= 256;
-	 // uint32	PageSize	= 2 << 10;
 	};
 
 	template <typename... T> [[nodiscard]] FRequest Get(T&&... t)  { return Request("GET",  Forward<T&&>(t)...); }
diff --git a/Engine/Source/Runtime/Experimental/IoStoreOnDemand/Private/CoreHttp/Http.cpp b/Engine/Source/Runtime/Experimental/IoStoreOnDemand/Private/CoreHttp/Http.cpp
--- a/Engine/Source/Runtime/Experimental/IoStoreOnDemand/Private/CoreHttp/Http.cpp
+++ b/Engine/Source/Runtime/Experimental/IoStoreOnDemand/Private/CoreHttp/Http.cpp
@@ -1895,16 +1895,8 @@ void FResponse::GetContentType(FAnsiStringView& Out) const
 ////////////////////////////////////////////////////////////////////////////////
 FAnsiStringView FResponse::GetHeader(FAnsiStringView Name) const
 {
-	const auto* Activity = (const FActivity*)this;
-	const FResponseInternal& Internal = Activity->Response;
-
-	const char* MessageData = Activity->Buffer.GetData() + Activity->StateParam;
-	FAnsiStringView Result, Headers(
-		MessageData + Internal.Offsets.Headers,
-		Internal.MessageLength - Internal.Offsets.Headers
-	);
-
-	EnumerateHeaders(Headers, [&Result, Name] (FAnsiStringView Candidate, FAnsiStringView Value)
+	FAnsiStringView Result;
+	ReadHeaders([&Result, Name] (FAnsiStringView Candidate, FAnsiStringView Value)
 	{
 		if (Candidate != Name)
 		{
@@ -1917,6 +1909,21 @@ FAnsiStringView FResponse::GetHeader(FAnsiStringView Name) const
 	return Result;
 }
 
+////////////////////////////////////////////////////////////////////////////////
+void FResponse::ReadHeaders(FHeaderSink Sink) const
+{
+	const auto* Activity = (const FActivity*)this;
+	const FResponseInternal& Internal = Activity->Response;
+
+	const char* MessageData = Activity->Buffer.GetData() + Activity->StateParam;
+	FAnsiStringView Result, Headers(
+		MessageData + Internal.Offsets.Headers,
+		Internal.MessageLength - Internal.Offsets.Headers
+	);
+
+	EnumerateHeaders(Headers, Sink);
+}
+
 ////////////////////////////////////////////////////////////////////////////////
 void FResponse::SetDestination(FIoBuffer* Buffer)
 {
