Removed unused result struct and enum

#rnx

diff --git a/Engine/Source/Runtime/Experimental/IoStoreOnDemand/Private/CoreHttp/Http.cpp b/Engine/Source/Runtime/Experimental/IoStoreOnDemand/Private/CoreHttp/Http.cpp
--- a/Engine/Source/Runtime/Experimental/IoStoreOnDemand/Private/CoreHttp/Http.cpp
+++ b/Engine/Source/Runtime/Experimental/IoStoreOnDemand/Private/CoreHttp/Http.cpp
@@ -192,22 +192,6 @@ static void Trace(T Id, ETrace Action, U Param=0)
 static IAS_CVAR(int32, RecvWorkThresholdKiB,80,		"Threshold of data remaining at which next request is sent (in KiB)");
 static IAS_CVAR(int32, IdleMs,				50'000,	"Time in seconds to close idle connections or fail waits");
 
-////////////////////////////////////////////////////////////////////////////////
-class FResult
-{
-public:
-				FResult() : FResult(0, "") {}
-	explicit	FResult(const char* Msg) : FResult(-1, Msg) {}
-	explicit	FResult(int32 Val, const char* Msg="") : Message(UPTRINT(Msg)), Value(Val) {}
-	const char*	GetMessage() const		{ return (const char*)Message; }
-	int32		GetValue() const		{ return int16(Value); }
-
-private:
-	UPTRINT		Message : 48;
-	PTRINT		Value : 16;
-};
-static_assert(sizeof(FResult) == sizeof(void*));
-
 ////////////////////////////////////////////////////////////////////////////////
 class FOutcome
 {
@@ -779,14 +763,6 @@ FMessageBuilder& FMessageBuilder::operator << (FAnsiStringView Lhs)
 class FSocket
 {
 public:
-	enum class EResult
-	{
-		HangUp			=  0,
-		Wait			= -1,
-		Error			= -2,
-		ConnectError	= -3,
-	};
-
 	struct FWaiter
 	{
 		enum class EWhat { Send = 0b01, Recv = 0b10, Both = Send|Recv };
@@ -3600,16 +3576,6 @@ static void MiscTest()
 	check(FOutcome::Waiting().IsOk() == false);
 	check(FOutcome::Waiting().IsWaiting());
 	check(FOutcome::Waiting().IsError() == false);
-
-	check(FResult(-5).GetValue()		== -5);
-	check(FResult(-1).GetValue()		== -1);
-	check(FResult( 0).GetValue()		==  0);
-	check(FResult( 1).GetValue()		==  1);
-	check(FResult(19).GetValue()		== 19);
-	check(FResult(0xffff).GetValue()	== -1);
-	check(FResult(1, "yes").GetValue()	==  1);
-	check(FResult(0, "?").GetValue()	==  0);
-	check(FResult(-1, "no").GetValue()	== -1);
 }
 
 ////////////////////////////////////////////////////////////////////////////////
