Retain activity's host pointer until at least a dest is set

#rnx

diff --git a/Engine/Source/Runtime/Experimental/IoStoreOnDemand/Private/CoreHttp/Http.cpp b/Engine/Source/Runtime/Experimental/IoStoreOnDemand/Private/CoreHttp/Http.cpp
--- a/Engine/Source/Runtime/Experimental/IoStoreOnDemand/Private/CoreHttp/Http.cpp
+++ b/Engine/Source/Runtime/Experimental/IoStoreOnDemand/Private/CoreHttp/Http.cpp
@@ -2191,7 +2191,7 @@ static int32 DoRecvMessage(FActivity* Activity, FSocket& Socket)
 	}
 
 	// Call out to the sink to get a content destination
-	Activity->Dest = nullptr;
+	FIoBuffer* PriorDest = Activity->Dest; // to retain unioned Host ptr
 	Internal.Code = -1;
 	Internal.ContentLength = ContentLength;
 	{
@@ -2201,7 +2201,7 @@ static int32 DoRecvMessage(FActivity* Activity, FSocket& Socket)
 
 	if (Activity->NoContent == 0)
 	{
-		if (Activity->Dest == nullptr)
+		if (Activity->Dest == PriorDest)
 		{
 			Activity_SetError(Activity, "User did not provide a destination buffer");
 			return -1;
