Loop.Tick() could falsely reports it is not busy even if there are requests pending

#rnx

diff --git a/Engine/Source/Runtime/Experimental/IoStoreOnDemand/Private/CoreHttp/Http.cpp b/Engine/Source/Runtime/Experimental/IoStoreOnDemand/Private/CoreHttp/Http.cpp
--- a/Engine/Source/Runtime/Experimental/IoStoreOnDemand/Private/CoreHttp/Http.cpp
+++ b/Engine/Source/Runtime/Experimental/IoStoreOnDemand/Private/CoreHttp/Http.cpp
@@ -3330,10 +3330,12 @@ uint32 FEventLoop::FImpl::Tick(int32 PollTimeoutMs)
 		Activity = Next;
 	}
 
+	uint32 BusyBias = 0;
 	if (ReturnedSlots)
 	{
+		uint64 LatestFree = FreeSlots.fetch_add(ReturnedSlots, std::memory_order_relaxed);
+		BusyBias += (LatestFree != PrevFreeSlots);
 		PrevFreeSlots += ReturnedSlots;
-		FreeSlots.fetch_add(ReturnedSlots, std::memory_order_relaxed);
 	}
 
 	if (CancelsLoad)
@@ -3341,7 +3343,7 @@ uint32 FEventLoop::FImpl::Tick(int32 PollTimeoutMs)
 		Cancels.fetch_and(~CancelsLoad, std::memory_order_relaxed);
 	}
 
-	return BusyCount;
+	return BusyCount + BusyBias;
 }
 
 
