Iris
- Added IRIS_CSV_PROFILER_SCOPE macro to allow reporting to both CSV and cpu trace if enabled.

diff --git a/Engine/Source/Runtime/Experimental/Iris/Core/Private/Iris/ReplicationSystem/ObjectReplicationBridge.cpp b/Engine/Source/Runtime/Experimental/Iris/Core/Private/Iris/ReplicationSystem/ObjectReplicationBridge.cpp
--- a/Engine/Source/Runtime/Experimental/Iris/Core/Private/Iris/ReplicationSystem/ObjectReplicationBridge.cpp
+++ b/Engine/Source/Runtime/Experimental/Iris/Core/Private/Iris/ReplicationSystem/ObjectReplicationBridge.cpp
@@ -1059,10 +1059,7 @@ void UObjectReplicationBridge::PreUpdate(const UE::Net::FNetBitArrayView Objects
 {
 	using namespace UE::Net::Private;
 
-#if UE_NET_IRIS_CSV_STATS
-	CSV_SCOPED_TIMING_STAT(Iris, ReplicationBridge_PreUpdate);
-#endif
-	IRIS_PROFILER_SCOPE(UObjectReplicationBridge_PreUpdate);
+	IRIS_CSV_PROFILER_SCOPE(Iris, ReplicationBridge_PreUpdate);
 
 	// TODO: Get rid of Poller
 	FObjectPoller::FInitParams PollerInitParams;
@@ -1080,10 +1077,7 @@ void UObjectReplicationBridge::PollAndCopy(const UE::Net::FNetBitArrayView Objec
 {
 	using namespace UE::Net::Private;
 
-#if UE_NET_IRIS_CSV_STATS
-	CSV_SCOPED_TIMING_STAT(Iris, ReplicationBridge_PollAndCopy);
-#endif
-	IRIS_PROFILER_SCOPE(UObjectReplicationBridge_PollAndCopy);
+	IRIS_CSV_PROFILER_SCOPE(Iris, ReplicationBridge_PollAndCopy);
 
 	FObjectPoller::FInitParams PollerInitParams;
 	PollerInitParams.ObjectReplicationBridge = this;
diff --git a/Engine/Source/Runtime/Experimental/Iris/Core/Private/Iris/ReplicationSystem/ReplicationSystem.cpp b/Engine/Source/Runtime/Experimental/Iris/Core/Private/Iris/ReplicationSystem/ReplicationSystem.cpp
--- a/Engine/Source/Runtime/Experimental/Iris/Core/Private/Iris/ReplicationSystem/ReplicationSystem.cpp
+++ b/Engine/Source/Runtime/Experimental/Iris/Core/Private/Iris/ReplicationSystem/ReplicationSystem.cpp
@@ -314,7 +314,7 @@ public:
 
 	void UpdateWorldLocations()
 	{
-		IRIS_PROFILER_SCOPE(FReplicationSystem_UpdateWorldLocations);
+		IRIS_CSV_PROFILER_SCOPE(Iris, ReplicationSystem_UpdateWorldLocations);
 
 		// Reset dirty object info before updating.
 		FWorldLocations& WorldLocations = ReplicationSystemInternal.GetWorldLocations();
@@ -325,7 +325,7 @@ public:
 
 	void UpdateFilterPrePoll()
 	{
-		IRIS_PROFILER_SCOPE(FReplicationSystem_UpdateFilterPrePoll);
+		IRIS_CSV_PROFILER_SCOPE(Iris, ReplicationSystem_UpdateFilterPrePoll);
 		LLM_SCOPE_BYTAG(Iris);
 
 		FReplicationFiltering& Filtering = ReplicationSystemInternal.GetFiltering();
@@ -339,12 +339,12 @@ public:
 		FReplicationFiltering& Filtering = ReplicationSystemInternal.GetFiltering();
 
 		{
-			IRIS_PROFILER_SCOPE(FReplicationSystem_UpdateFilterPostPoll);
+			IRIS_CSV_PROFILER_SCOPE(Iris, ReplicationSystem_UpdateFilterPostPoll);
 			Filtering.FilterPostPoll();
 		}
 
 		{
-			IRIS_PROFILER_SCOPE(FReplicationSystem_UpdateConnectionsScope);
+			IRIS_CSV_PROFILER_SCOPE(Iris, ReplicationSystem_UpdateConnectionsScope);
 		
 			// Iterate over all valid connections and propagate updated scopes
 			FReplicationConnections& Connections = ReplicationSystemInternal.GetConnections();
@@ -364,7 +364,7 @@ public:
 	// Can run at any time between scoping and replication.
 	void UpdateConditionals()
 	{
-		IRIS_PROFILER_SCOPE(FReplicationSystem_UpdateConditionals);
+		IRIS_CSV_PROFILER_SCOPE(Iris, ReplicationSystem_UpdateConditionals);
 
 		FReplicationConditionals& Conditionals = ReplicationSystemInternal.GetConditionals();
 		Conditionals.Update();
@@ -373,10 +373,7 @@ public:
 	// Runs after filtering
 	void UpdatePrioritization(const FNetBitArrayView& ReplicatingConnections)
 	{
-#if UE_NET_IRIS_CSV_STATS
-		CSV_SCOPED_TIMING_STAT(Iris, ReplicationSystem_UpdatePrioritization);
-#endif
-		IRIS_PROFILER_SCOPE(FReplicationSystem::FImpl::UpdatePrioritization);
+		IRIS_CSV_PROFILER_SCOPE(Iris, ReplicationSystem_UpdatePrioritization);
 		LLM_SCOPE_BYTAG(Iris);
 
 		const FNetBitArrayView RelevantObjects = ReplicationSystemInternal.GetNetRefHandleManager().GetRelevantObjectsInternalIndices();
@@ -394,7 +391,7 @@ public:
 
 	void PropagateDirtyChanges()
 	{
-		IRIS_PROFILER_SCOPE(FReplicationSystem_PropagateDirtyChanges);
+		IRIS_CSV_PROFILER_SCOPE(Iris, ReplicationSystem_PropagateDirtyChanges);
 
 		FReplicationConnections& Connections = ReplicationSystemInternal.GetConnections();
 		const FChangeMaskCache& UpdatedChangeMasks = ReplicationSystemInternal.GetChangeMaskCache();
@@ -411,7 +408,7 @@ public:
 
 	void QuantizeDirtyStateData()
 	{
-		IRIS_PROFILER_SCOPE(FReplicationSystem_QuantizeDirtyStateData);
+		IRIS_CSV_PROFILER_SCOPE(Iris, ReplicationSystem_QuantizeDirtyStateData);
 		LLM_SCOPE_BYTAG(IrisState);
 
 		FNetRefHandleManager& NetRefHandleManager = ReplicationSystemInternal.GetNetRefHandleManager();
@@ -451,7 +448,7 @@ public:
 
 	void ResetObjectStateDirtiness()
 	{
-		IRIS_PROFILER_SCOPE(FReplicationSystem_ResetObjectStateDirtiness);
+		IRIS_CSV_PROFILER_SCOPE(Iris, ReplicationSystem_ResetObjectStateDirtiness);
 
 		FNetRefHandleManager& NetRefHandleManager = ReplicationSystemInternal.GetNetRefHandleManager();
 
@@ -469,7 +466,7 @@ public:
 
 	void ProcessNetObjectAttachmentSendQueue(FNetBlobManager::EProcessMode ProcessMode)
 	{
-		IRIS_PROFILER_SCOPE(FReplicationSystem_ProcessNetObjectAttachmentSendQueue);
+		IRIS_CSV_PROFILER_SCOPE(Iris, ReplicationSystem_ProcessNetObjectAttachmentSendQueue);
 
 		FNetBlobManager& NetBlobManager = ReplicationSystemInternal.GetNetBlobManager();
 		NetBlobManager.ProcessNetObjectAttachmentSendQueue(ProcessMode);
@@ -477,7 +474,7 @@ public:
 
 	void ProcessOOBNetObjectAttachmentSendQueue()
 	{
-		IRIS_PROFILER_SCOPE(FReplicationSystem_ProcessOOBNetObjectAttachmentSendQueue);
+		IRIS_CSV_PROFILER_SCOPE(Iris, ReplicationSystem_ProcessOOBNetObjectAttachmentSendQueue);
 
 		FNetBlobManager& NetBlobManager = ReplicationSystemInternal.GetNetBlobManager();
 		NetBlobManager.ProcessOOBNetObjectAttachmentSendQueue(ConnectionsPendingPostTickDispatchSend);
@@ -576,6 +573,8 @@ public:
 
 	void UpdateUnresolvableReferenceTracking()
 	{
+		IRIS_CSV_PROFILER_SCOPE(Iris, ReplicationSystem_UpdateUnresolvableReferenceTracking);
+
 		FReplicationConnections& Connections = ReplicationSystemInternal.GetConnections();
 		auto UpdateUnresolvableReferenceTracking = [&Connections](uint32 ConnectionId)
 		{
@@ -649,7 +648,7 @@ void UReplicationSystem::PreSendUpdate(const FSendUpdateParams& Params)
 	using namespace UE::Net;
 	using namespace UE::Net::Private;
 
-	IRIS_PROFILER_SCOPE(FReplicationSystem_PreSendUpdate);
+	IRIS_CSV_PROFILER_SCOPE(Iris, ReplicationSystem_PreSendUpdate);
 
 	ensureAlways(Impl->CurrentSendPass == EReplicationSystemSendPass::Invalid);
 	Impl->CurrentSendPass = Params.SendPass;
@@ -789,7 +788,7 @@ void UReplicationSystem::PostSendUpdate()
 	using namespace UE::Net;
 	using namespace UE::Net::Private;
 
-	IRIS_PROFILER_SCOPE(FReplicationSystem_PostSendUpdate);
+	IRIS_CSV_PROFILER_SCOPE(Iris, ReplicationSystem_PostSendUpdate);
 
 	if (!ensure(Impl->CurrentSendPass != EReplicationSystemSendPass::Invalid))
 	{
@@ -834,7 +833,7 @@ void UReplicationSystem::PostGarbageCollection()
 
 void UReplicationSystem::CollectGarbage()
 {
-	IRIS_PROFILER_SCOPE(ReplicationSystem_CollectGarbage);
+	IRIS_CSV_PROFILER_SCOPE(Iris, ReplicationSystem_CollectGarbage);
 
 	// Prune stale object instances before descriptors and protocols are pruned
 	Impl->ReplicationSystemInternal.GetReplicationBridge()->CallPruneStaleObjects();
diff --git a/Engine/Source/Runtime/Experimental/Iris/Core/Public/Iris/Core/IrisCsv.h b/Engine/Source/Runtime/Experimental/Iris/Core/Public/Iris/Core/IrisCsv.h
--- a/Engine/Source/Runtime/Experimental/Iris/Core/Public/Iris/Core/IrisCsv.h
+++ b/Engine/Source/Runtime/Experimental/Iris/Core/Public/Iris/Core/IrisCsv.h
@@ -3,5 +3,16 @@
 #pragma once
 
 #include "ProfilingDebugging/CsvProfiler.h"
+#include "Iris/IrisConfig.h"
+#include "Iris/Core/IrisProfiler.h"
 
 CSV_DECLARE_CATEGORY_EXTERN(Iris);
+
+#if UE_NET_IRIS_CSV_STATS
+	#define IRIS_CSV_PROFILER_SCOPE(CsvCategory, x) \
+	CSV_SCOPED_TIMING_STAT(CsvCategory, x); \
+	IRIS_PROFILER_SCOPE(x)
+#else
+	#define IRIS_CSV_PROFILER_SCOPE(CsvCategory, x) IRIS_PROFILER_SCOPE(x)
+#endif
+
