Constraints: use the constraint index when sorting to maintain the initial order if there are no dependencies.
    - NOTE: bubble-sorting is only used if the array size is less than 9, so the order could be wrong after sorting without index comparison.

#rnx

diff --git a/Engine/Source/Runtime/Experimental/Animation/Constraints/Private/ConstraintsEvaluationGraph.cpp b/Engine/Source/Runtime/Experimental/Animation/Constraints/Private/ConstraintsEvaluationGraph.cpp
--- a/Engine/Source/Runtime/Experimental/Animation/Constraints/Private/ConstraintsEvaluationGraph.cpp
+++ b/Engine/Source/Runtime/Experimental/Animation/Constraints/Private/ConstraintsEvaluationGraph.cpp
@@ -128,7 +128,8 @@ void FConstraintsEvaluationGraph::Rebuild()
 		{
 			return Prerex.PrerequisiteTickFunction == LHSTickFunction;
 		});
-		return bIsLHSAPrerexOfRHS;
+		// if not a prerequisite then compare constraints indices
+		return bIsLHSAPrerexOfRHS ? bIsLHSAPrerexOfRHS :  LHS.ConstraintIndex < RHS.ConstraintIndex;
 	};
 	Algo::Sort(Nodes, EvaluationOrderPredicate);
 
diff --git a/Engine/Source/Runtime/Experimental/Animation/Constraints/Private/ConstraintsManager.cpp b/Engine/Source/Runtime/Experimental/Animation/Constraints/Private/ConstraintsManager.cpp
--- a/Engine/Source/Runtime/Experimental/Animation/Constraints/Private/ConstraintsManager.cpp
+++ b/Engine/Source/Runtime/Experimental/Animation/Constraints/Private/ConstraintsManager.cpp
@@ -769,47 +769,95 @@ bool FConstraintsManagerController::RemoveAllConstraints(bool bDoNotCompensate)
 static void SortConstraints(UWorld* InWorld, TArray< TWeakObjectPtr<UTickableConstraint> >& InOutSortedConstraints)
 {
 	using ConstraintPtr = TWeakObjectPtr<UTickableConstraint>;
+
+	// store original indices
+	TArray<int32> ConstraintIndices;
+	ConstraintIndices.Reserve(InOutSortedConstraints.Num());
+	for (int32 Index = 0; Index < InOutSortedConstraints.Num(); ++Index)
+	{
+		ConstraintIndices.Add(Index);
+	}
+	
 	// LHS ticks before RHS = LHS is a prerex of RHS 
-	auto TicksBefore = [InWorld](const UTickableConstraint& LHS, const UTickableConstraint& RHS)
+	auto TicksBefore = [InWorld, &InOutSortedConstraints](const int32 LHSIndex, const int32 RHSIndex)
 	{
+		const ConstraintPtr& LHS = InOutSortedConstraints[LHSIndex];
+		const ConstraintPtr& RHS = InOutSortedConstraints[RHSIndex];
 		if (RHS.IsValid() == false || LHS.IsValid() == false)
 		{
 			return false;
 		}
-		const TArray<FTickPrerequisite>& RHSPrerex = RHS.GetTickFunction(InWorld).GetPrerequisites();
-		const FConstraintTickFunction& LHSTickFunction = LHS.GetTickFunction(InWorld);
+	
+		const TArray<FTickPrerequisite>& RHSPrerex = RHS->GetTickFunction(InWorld).GetPrerequisites();
+		const FConstraintTickFunction& LHSTickFunction = LHS->GetTickFunction(InWorld);
 		const bool bIsLHSAPrerexOfRHS = RHSPrerex.ContainsByPredicate([&LHSTickFunction](const FTickPrerequisite& Prerex)
-			{
-				return Prerex.PrerequisiteTickFunction == &LHSTickFunction;
-			});
-		return bIsLHSAPrerexOfRHS;
+		{
+			return Prerex.PrerequisiteTickFunction == &LHSTickFunction;
+		});
+
+		// if not a prerequisite then compare constraints indices
+		return bIsLHSAPrerexOfRHS ? bIsLHSAPrerexOfRHS : LHSIndex < RHSIndex;
 	};
 
-	Algo::Sort(InOutSortedConstraints, [TicksBefore](const ConstraintPtr& LHS, const ConstraintPtr& RHS)
+	// sort
+	Algo::Sort(ConstraintIndices, [TicksBefore](const int32 LHSIndex, const int32 RHSIndex)
 	{
-		return TicksBefore(*LHS, *RHS);
+		return TicksBefore(LHSIndex, RHSIndex);
 	});
+
+	// re-order using sorted indices
+	TArray<TWeakObjectPtr<UTickableConstraint> > SortedConstraints;
+	SortedConstraints.Reserve(InOutSortedConstraints.Num());
+	for (int32 Index = 0; Index < ConstraintIndices.Num(); ++Index)
+	{
+		SortedConstraints.Add(InOutSortedConstraints[ConstraintIndices[Index]]);
+	}
+	InOutSortedConstraints = SortedConstraints;
 }
 
 static void SortConstraints(UWorld* InWorld, TArray< TObjectPtr<UTickableConstraint> >& InOutSortedConstraints)
 {
 	using ConstraintPtr = TObjectPtr<UTickableConstraint>;
+
+	// store original indices
+	TArray<int32> ConstraintIndices;
+	ConstraintIndices.Reserve(InOutSortedConstraints.Num());
+	for (int32 Index = 0; Index < InOutSortedConstraints.Num(); ++Index)
+	{
+		ConstraintIndices[Index] = Index;
+	}
+	
 	// LHS ticks before RHS = LHS is a prerex of RHS 
-	auto TicksBefore = [InWorld](const UTickableConstraint& LHS, const UTickableConstraint& RHS)
+	auto TicksBefore = [InWorld, &InOutSortedConstraints](const int32 LHSIndex, const int32 RHSIndex)
 	{
-		const TArray<FTickPrerequisite>& RHSPrerex = RHS.GetTickFunction(InWorld).GetPrerequisites();
-		const FConstraintTickFunction& LHSTickFunction = LHS.GetTickFunction(InWorld);
+		const ConstraintPtr& LHS = InOutSortedConstraints[LHSIndex];
+		const ConstraintPtr& RHS = InOutSortedConstraints[RHSIndex];
+		
+		const TArray<FTickPrerequisite>& RHSPrerex = RHS->GetTickFunction(InWorld).GetPrerequisites();
+		const FConstraintTickFunction& LHSTickFunction = LHS->GetTickFunction(InWorld);
 		const bool bIsLHSAPrerexOfRHS = RHSPrerex.ContainsByPredicate([&LHSTickFunction](const FTickPrerequisite& Prerex)
-			{
-				return Prerex.PrerequisiteTickFunction == &LHSTickFunction;
-			});
-		return bIsLHSAPrerexOfRHS;
+		{
+			return Prerex.PrerequisiteTickFunction == &LHSTickFunction;
+		});
+
+		// if not a prerequisite then compare constraints indices
+		return bIsLHSAPrerexOfRHS ? bIsLHSAPrerexOfRHS : LHSIndex < RHSIndex;
 	};
 
-	Algo::Sort(InOutSortedConstraints, [TicksBefore](const ConstraintPtr& LHS, const ConstraintPtr& RHS)
+	// sort
+	Algo::Sort(ConstraintIndices, [TicksBefore](const int32 LHSIndex, const int32 RHSIndex)
 	{
-		return TicksBefore(*LHS, *RHS);
+		return TicksBefore(LHSIndex, RHSIndex);
 	});
+
+	// re-order using sorted indices
+	TArray<TObjectPtr<UTickableConstraint> > SortedConstraints;
+	SortedConstraints.Reserve(InOutSortedConstraints.Num());
+	for (int32 Index = 0; Index < ConstraintIndices.Num(); ++Index)
+	{
+		SortedConstraints[Index] = InOutSortedConstraints[ConstraintIndices[Index]];
+	}
+	InOutSortedConstraints = SortedConstraints;
 }
 
 TArray< TObjectPtr<UTickableConstraint> > FConstraintsManagerController::GetStaticConstraints(const bool bSorted) const
