Opt-in 30x redirect support

#rnx

diff --git a/Engine/Source/Runtime/Experimental/IoStoreOnDemand/Private/CoreHttp/Client.h b/Engine/Source/Runtime/Experimental/IoStoreOnDemand/Private/CoreHttp/Client.h
--- a/Engine/Source/Runtime/Experimental/IoStoreOnDemand/Private/CoreHttp/Client.h
+++ b/Engine/Source/Runtime/Experimental/IoStoreOnDemand/Private/CoreHttp/Client.h
@@ -188,6 +188,7 @@ public:
 	struct FRequestParams
 	{
 		uint32	BufferSize	= 256;
+		bool	bAutoRedirect = false;
 	};
 
 	template <typename... T> [[nodiscard]] FRequest Get(T&&... t)  { return Request("GET",  Forward<T&&>(t)...); }
@@ -205,6 +206,7 @@ public:
 	FTicket					Send(FRequest&& Request, FTicketSink Sink, UPTRINT SinkParam=0);
 
 private:
+	bool					Redirect(const FTicketStatus& Status, FTicketSink& OuterSink);
 	FImpl*					Impl;
 
 private:
diff --git a/Engine/Source/Runtime/Experimental/IoStoreOnDemand/Private/CoreHttp/Http.cpp b/Engine/Source/Runtime/Experimental/IoStoreOnDemand/Private/CoreHttp/Http.cpp
--- a/Engine/Source/Runtime/Experimental/IoStoreOnDemand/Private/CoreHttp/Http.cpp
+++ b/Engine/Source/Runtime/Experimental/IoStoreOnDemand/Private/CoreHttp/Http.cpp
@@ -1623,7 +1623,8 @@ struct alignas(16) FActivity
 	EState				State = EState::None;
 	uint8				IsKeepAlive : 1;
 	uint8				NoContent : 1;
-	uint8				_Unused0 : 6;
+	uint8				bFollow30x : 1;
+	uint8				_Unused0 : 5;
 	uint32				StateParam = 0;
 #if IAS_HTTP_WITH_PERF
 	FStopwatch			Stopwatch;
@@ -3397,6 +3398,7 @@ FRequest FEventLoop::Request(
 
 	FHost* Host = Activity->Host = Buffer.Alloc<FHost>();
 	Activity->IsKeepAlive = 0;
+	Activity->bFollow30x = (Params->bAutoRedirect == true);
 
 	uint32 HostNameLength = HostName.Len();
 	char* HostNamePtr = Buffer.Alloc<char>(HostNameLength + 1);
@@ -3430,10 +3432,102 @@ FRequest FEventLoop::Request(
 
 	Activity->Host = Pool.Ptr;
 	Activity->IsKeepAlive = 1;
+	Activity->bFollow30x = (Params->bAutoRedirect == true);
 
 	return Impl->Request(Method, Path, Activity);
 }
 
+////////////////////////////////////////////////////////////////////////////////
+bool FEventLoop::Redirect(const FTicketStatus& Status, FTicketSink& OuterSink)
+{
+	const FResponse& Response = Status.GetResponse();
+
+	switch (Response.GetStatusCode())
+	{
+		case 301:				// RedirectMoved
+		case 302:				// RedirectFound
+		case 307:				// RedirectTemp
+		case 308: break;		// RedirectPerm
+		default: return false;
+	}
+
+	FAnsiStringView Location = Response.GetHeader("Location");
+	if (Location.IsEmpty())
+	{
+		return false;
+	}
+
+	check(Response.GetContentLength() == 0); // should we ever hit this, we'll fix it
+
+	const auto& Activity = (FActivity&)Response; // todo: yuk
+
+	// Original method should remain unchanged
+	const char* Data = Activity.Buffer.GetData();
+	FAnsiStringView Method;
+	for (uint32 i = 0; i < 5; ++i)
+	{
+		if (Data[i] <= ' ')
+		{
+			Method = FAnsiStringView(Data, i);
+			break;
+		}
+	}
+	check(!Method.IsEmpty());
+
+	FRequest ForwardRequest;
+	if (!Location.StartsWith("http://") && !Location.StartsWith("https://"))
+	{
+		if (Location[0] != '/')
+		{
+			return false;
+		}
+
+		FHost& Host = *(Activity.Host);
+
+		TAnsiStringBuilder<256> Url;
+		Url << "http://";
+		Url << Host.GetHostName();
+		Url << ":" << Host.GetPort();
+		Url << Location;
+
+		new (&ForwardRequest) FRequest(Request(Method, Url));
+	}
+	else
+	{
+		new (&ForwardRequest) FRequest(Request(Method, Location));
+	}
+
+	// Transfer original request headers
+	check(Activity.State == FActivity::EState::RecvMessage);
+	uint32 Length = 0;
+	for (const char* End = Data + Activity.StateParam; Data + Length < End; ++Length)
+	{
+		if (Data[Length] != '\n')
+		{
+			continue;
+		}
+
+		Data = Data + Length + 1;
+		Length = uint32(ptrdiff_t(End - Data));
+		break;
+	}
+
+	FAnsiStringView OriginalHeaders(Data, Length);
+	EnumerateHeaders(OriginalHeaders, [&ForwardRequest] (FAnsiStringView Name, FAnsiStringView Value)
+	{
+		if (Name != "Host" && Name != "Connection")
+		{
+			ForwardRequest.Header(Name, Value);
+		}
+		return true;
+	});
+
+	// Send the request
+	Send(MoveTemp(ForwardRequest), MoveTemp(OuterSink), Status.GetParam());
+
+	return true;
+}
+
 ////////////////////////////////////////////////////////////////////////////////
 FTicket FEventLoop::Send(FRequest&& Request, FTicketSink Sink, UPTRINT SinkParam)
 {
@@ -3441,6 +3535,31 @@ FTicket FEventLoop::Send(FRequest&& Request, FTicketSink Sink, UPTRINT SinkParam
 	Swap(Activity, Request.Ptr);
 	Activity->SinkParam = SinkParam;
 	Activity->Sink = Sink;
+
+	// Intercept sink calls to catch 30x status codes and follow them
+	if (Activity->bFollow30x)
+	{
+		auto RedirectSink = [
+				this,
+				OuterSink=MoveTemp(Activity->Sink)
+			] (const FTicketStatus& Status) mutable
+		{
+			if (Status.GetId() == FTicketStatus::EId::Response)
+			{
+				if (Redirect(Status, OuterSink))
+				{
+					return;
+				}
+			}
+
+			if (OuterSink)
+			{
+				return OuterSink(Status);
+			}
+		};
+		Activity->Sink = RedirectSink;
+	}
+
 	return Impl->Send(Activity);
 }
 
@@ -3694,15 +3813,87 @@ static void ThrottleTest(FAnsiStringView TestUrl)
 }
 
 ////////////////////////////////////////////////////////////////////////////////
-IOSTOREONDEMAND_API void IasHttpTest(const ANSICHAR* TestHost="localhost")
+static void RedirectTest(const ANSICHAR* TestHost)
 {
-#if PLATFORM_WINDOWS
-	WSADATA WsaData;
-	if (WSAStartup(MAKEWORD(2, 2), &WsaData) == 0x0a9e0493)
-		return;
-	ON_SCOPE_EXIT { WSACleanup(); };
-#endif
+	FEventLoop Loop;
+
+	auto WaitForLoopIdle = [&] {
+		for (; Loop.Tick(-1); FPlatformProcess::SleepNoStats(0.02f));
+	};
+
+	FEventLoop::FRequestParams RequestParams = {
+		.bAutoRedirect = true,
+	};
+
+	enum ReTyp { ReAbs, ReRel };
+	enum { RecvDataSize = 48 };
+
+	TAnsiStringBuilder<64> Builder;
+	auto BuildUrl = [&] (ReTyp Typ, uint32 Code) -> const auto&
+	{
+		Builder.Reset();
+		Builder << "https://";
+		Builder << TestHost;
+		Builder << ":9493";
+		Builder << "/redirect";
+		Builder << ((Typ < ReRel) ? "/abs/" : "/rel/");
+		Builder << Code;
+		Builder << "/data/" << uint32(RecvDataSize);
+		return Builder;
+	};
+
+	UPTRINT SinkParam = 0xaa'493'493'493'493'bbull;
+
+	uint32 RecvCount;
+	auto OkSink = [Dest=FIoBuffer(), SinkParam, &RecvCount] (const FTicketStatus& Status) mutable {
+		check(Status.GetParam() == SinkParam);
+		check(Status.GetId() != FTicketStatus::EId::Error);
+		if (Status.GetId() == FTicketStatus::EId::Response)
+		{
+			FResponse& Response = Status.GetResponse();
+			check(Response.GetStatusCode() == 200);
+			Response.SetDestination(&Dest);
+			return;
+		}
+		check(Status.GetId() == FTicketStatus::EId::Content);
+		RecvCount += uint32(Dest.GetSize());
+	};
+
+	uint32 TestCodes[] = { 301, 302, 307, 308 };
 
+	for (auto ReTest : { ReAbs, ReRel })
+	{
+		RecvCount = 0;
+		for (uint32 Code : TestCodes)
+		{
+			FRequest Request = Loop.Get(BuildUrl(ReTest, Code), &RequestParams);
+			if (Code > TestCodes[1])
+			{
+				Request.Header("TestCodeHeader", "Header-Of-Test-Codes");
+			}
+			Loop.Send(MoveTemp(Request), OkSink, SinkParam);
+		}
+		WaitForLoopIdle();
+		check(RecvCount == RecvDataSize * UE_ARRAY_COUNT(TestCodes));
+	}
+
+	RequestParams = FEventLoop::FRequestParams();
+	RequestParams.bAutoRedirect = true;
+
+	FConnectionPool::FParams Params;
+	Params.SetHostFromUrl(BuildUrl(ReAbs, 0));
+	Params.ConnectionCount = 4;
+	FConnectionPool Pool(Params);
+
+	RecvCount = 0;
+	Loop.Send(Loop.Get("/redirect/abs/307/data/55", Pool, &RequestParams), OkSink, SinkParam);
+	WaitForLoopIdle();
+	check(RecvCount == 55);
+}
+
+////////////////////////////////////////////////////////////////////////////////
+static void HttpTest(const ANSICHAR* TestHost)
+{
 	TAnsiStringBuilder<64> Ret;
 	auto BuildUrl = [&] (const ANSICHAR* Suffix=nullptr, uint32 Port=9493) -> const auto&
 	{
@@ -3713,8 +3904,6 @@ IOSTOREONDEMAND_API void IasHttpTest(const ANSICHAR* TestHost="localhost")
 		return (Suffix != nullptr) ? (Ret << Suffix) : Ret;
 	};
 
-	MiscTest();
-
 	struct
 	{
 		FIoBuffer Dest;
@@ -4090,6 +4279,21 @@ IOSTOREONDEMAND_API void IasHttpTest(const ANSICHAR* TestHost="localhost")
 	// (utf-8 host names)
 }
 
+////////////////////////////////////////////////////////////////////////////////
+IOSTOREONDEMAND_API void IasHttpTest(const ANSICHAR* TestHost="localhost")
+{
+#if PLATFORM_WINDOWS
+	WSADATA WsaData;
+	if (WSAStartup(MAKEWORD(2, 2), &WsaData) == 0x0a9e0493)
+		return;
+	ON_SCOPE_EXIT { WSACleanup(); };
+#endif
+
+	MiscTest();
+	HttpTest(TestHost);
+	RedirectTest(TestHost);
+}
+
 #endif // !SHIP|TEST
 
 // }}}
diff --git a/Engine/Source/Runtime/Experimental/IoStoreOnDemand/Private/CoreHttp/TestServer.py b/Engine/Source/Runtime/Experimental/IoStoreOnDemand/Private/CoreHttp/TestServer.py
--- a/Engine/Source/Runtime/Experimental/IoStoreOnDemand/Private/CoreHttp/TestServer.py
+++ b/Engine/Source/Runtime/Experimental/IoStoreOnDemand/Private/CoreHttp/TestServer.py
@@ -164,7 +164,19 @@ def http_data(handler, payload_size):
     handler.end_headers()
     handler.wfile.write(payload)
 
-
+def http_redirect(handler, style="abs", code=302, *dest):
+    loc = "/" + "/".join(dest)
+    if style.startswith("abs"):
+        proto = "https://" if style == "abs_s" else "http://"
+        port = 4939 if style == "abs_s" else 9493
+        host = handler.headers.get("Host")
+        if host.startswith("localhost"):
+            host = "127.0.49.3"
+        loc = f"{proto}{host}:{port}{loc}"
+    handler.send_response(int(code))
+    handler.send_header("Content-Length", "0")
+    handler.send_header("Location", loc)
+    handler.end_headers()
 
 class Handler(http.server.BaseHTTPRequestHandler):
     def _preroll(self):
@@ -208,6 +220,9 @@ class Handler(http.server.BaseHTTPRequestHandler):
             self.wfile.write(b"hello")
             return
 
+        if parts[0] == "redirect":
+            return http_redirect(self, *parts[1:])
+
         if parts[0] == "data":
             size = -1
             if len(parts) > 1:
