Vectorized 3 constraints in the Linear Solver and only for RBAN :
- Position 3 lock
- Rotation 3 Soft Limits
- Rotation 3 Motorized
This optimizations can improve up to 30% gain on several joint constraint types.
Optimization are under a cVar and disabled for now.
To prevent any cache issue each constraint is either using Simd vector or standard data. An union is used to store either Simd or not.
For all phases Position, Velocity and Projection if the constraint use Simd or not all those phases will be solved accordingly. Consequently this improves Velocity solve and projection solve. Gather has been slightly improved too.

When the cVar is disabled no regression has been noticed on the solver comparing to the original code.

diff --git a/Engine/Source/Runtime/Engine/Private/PhysicsEngine/ImmediatePhysics/ImmediatePhysicsChaos/ImmediatePhysicsSimulation_Chaos.cpp b/Engine/Source/Runtime/Engine/Private/PhysicsEngine/ImmediatePhysics/ImmediatePhysicsChaos/ImmediatePhysicsSimulation_Chaos.cpp
--- a/Engine/Source/Runtime/Engine/Private/PhysicsEngine/ImmediatePhysics/ImmediatePhysicsChaos/ImmediatePhysicsSimulation_Chaos.cpp
+++ b/Engine/Source/Runtime/Engine/Private/PhysicsEngine/ImmediatePhysics/ImmediatePhysicsChaos/ImmediatePhysicsSimulation_Chaos.cpp
@@ -156,6 +156,9 @@ FAutoConsoleVariableRef  CVarChaosImmPhysInertiaConditioningMaxInvInertiaCompone
 int32 ChaosImmediate_Joint_UseLinearSolver = -1;
 FAutoConsoleVariableRef CVarChaosImmPhysJointUseCachedSolver(TEXT("p.Chaos.ImmPhys.Joint.UseLinearSolver"), ChaosImmediate_Joint_UseLinearSolver, TEXT("Force use of linear or non-linear joint solver. (-1 to use PhysicsAsset setting)"));
 
+bool bChaosImmediate_UseSimdForLinearSolver = false;
+FAutoConsoleVariableRef  CVarChaosImmPhysUseSimdForLinearSolver(TEXT("p.Chaos.ImmPhys.UseSimd"), bChaosImmediate_UseSimdForLinearSolver, TEXT("Enable/Disable SIMD on the linear joint solver"));
+
 //
 // end remove when finished
 //
@@ -814,6 +817,7 @@ namespace ImmediatePhysics_Chaos
 		if (UseLinearJointSolver >= 0)
 		{
 			Implementation->Joints.SetUseLinearJointSolver(UseLinearJointSolver != 0);
+			Implementation->Joints.SetUseSimd(bChaosImmediate_UseSimdForLinearSolver); // RBAN solver can use safely the SIMD version
 		}
 
 		Implementation->Collisions.SetAllowManifolds(bUseManifoilds || (ChaosImmediate_Collision_UseManifolds != 0));
diff --git a/Engine/Source/Runtime/Experimental/Chaos/Private/Chaos/Joint/PBDJointCachedSolverGaussSeidel.cpp b/Engine/Source/Runtime/Experimental/Chaos/Private/Chaos/Joint/PBDJointCachedSolverGaussSeidel.cpp
--- a/Engine/Source/Runtime/Experimental/Chaos/Private/Chaos/Joint/PBDJointCachedSolverGaussSeidel.cpp
+++ b/Engine/Source/Runtime/Experimental/Chaos/Private/Chaos/Joint/PBDJointCachedSolverGaussSeidel.cpp
@@ -6,11 +6,12 @@
 #include "Chaos/ParticleHandle.h"
 #include "Chaos/PBDJointConstraintUtilities.h"
 #include "Chaos/Utilities.h"
+#include "Chaos/VectorUtility.h"
 
 
 namespace Chaos
 {
-	
+
 /** Derived states management */
 
 void FPBDJointCachedSolver::InitDerivedState()
@@ -20,7 +21,7 @@ void FPBDJointCachedSolver::InitDerivedState()
 	InitConnectorRs[0] = R(0) * LocalConnectorXs[0].GetRotation();
 	InitConnectorRs[1] = R(1) * LocalConnectorXs[1].GetRotation();
 	InitConnectorRs[1].EnforceShortestArcWith(InitConnectorRs[0]);
-	
+
 	ComputeBodyState(0);
 	ComputeBodyState(1);
 
@@ -28,6 +29,9 @@ void FPBDJointCachedSolver::InitDerivedState()
 
 	ConnectorWDts[0] = FRotation3::CalculateAngularVelocity(InitConnectorRs[0], ConnectorRs[0], 1.0f);
 	ConnectorWDts[1] = FRotation3::CalculateAngularVelocity(InitConnectorRs[1], ConnectorRs[1], 1.0f);
+
+	ConnectorWDtsSimd[0] = MakeVectorRegisterFloatFromDouble(MakeVectorRegister(ConnectorWDts[0][0], ConnectorWDts[0][1], ConnectorWDts[0][2], 0.0f));
+	ConnectorWDtsSimd[1] = MakeVectorRegisterFloatFromDouble(MakeVectorRegister(ConnectorWDts[1][0], ConnectorWDts[1][1], ConnectorWDts[1][2], 0.0f));
 }
 
 void FPBDJointCachedSolver::ComputeBodyState(const int32 BodyIndex)
@@ -81,9 +85,9 @@ bool FPBDJointCachedSolver::UpdateIsActive()
 }
 
 void FPBDJointCachedSolver::Update(
-	   const FReal Dt,
-	   const FPBDJointSolverSettings& SolverSettings,
-	   const FPBDJointSettings& JointSettings)
+	const FReal Dt,
+	const FPBDJointSolverSettings& SolverSettings,
+	const FPBDJointSettings& JointSettings)
 {
 	//UpdateIsActive();
 }
@@ -135,23 +139,23 @@ void FPBDJointCachedSolver::SetShockPropagationScales(const FReal InvMScale0, co
 		Body1().SetShockPropagationScale(InvMScale1);
 		bNeedsUpdate = true;
 	}
-	if(bNeedsUpdate)
+	if (bNeedsUpdate)
 	{
-		for(int32 ConstraintIndex = 0; ConstraintIndex < 3; ++ConstraintIndex)
+		for (int32 ConstraintIndex = 0; ConstraintIndex < 3; ++ConstraintIndex)
 		{
-			if(PositionConstraints.bValidDatas[ConstraintIndex])
+			if (PositionConstraints.bValidDatas[ConstraintIndex])
 			{
 				InitPositionDatasMass(PositionConstraints, ConstraintIndex, Dt);
 			}
-			if(RotationConstraints.bValidDatas[ConstraintIndex])
+			if (RotationConstraints.bValidDatas[ConstraintIndex])
 			{
 				InitRotationDatasMass(RotationConstraints, ConstraintIndex, Dt);
 			}
-			if(PositionDrives.bValidDatas[ConstraintIndex])
+			if (PositionDrives.bValidDatas[ConstraintIndex])
 			{
 				InitPositionDatasMass(PositionDrives, ConstraintIndex, Dt);
 			}
-			if(RotationDrives.bValidDatas[ConstraintIndex])
+			if (RotationDrives.bValidDatas[ConstraintIndex])
 			{
 				InitRotationDatasMass(RotationDrives, ConstraintIndex, Dt);
 			}
@@ -193,6 +197,7 @@ void FPBDJointCachedSolver::Init(
 	NumActiveConstraints = -1;
 	bIsActive = true;
 	bIsBroken = false;
+	bUseSimd = SolverSettings.bUseSimd;
 	bUsePositionBasedDrives = SolverSettings.bUsePositionBasedDrives;
 
 	SolverStiffness = 1.0f;
@@ -201,14 +206,14 @@ void FPBDJointCachedSolver::Init(
 
 	// Set the mass and inertia.
 	// If enabled, adjust the mass so that we limit the maximum mass and inertia ratios
-	FReal ConditionedInvMs[2] = 
+	FReal ConditionedInvMs[2] =
 	{
 		Body0().InvM(),
 		Body1().InvM()
 	};
-	FVec3 ConditionedInvILs[2] = 
+	FVec3 ConditionedInvILs[2] =
 	{
-		Body0().InvILocal(), 
+		Body0().InvILocal(),
 		Body1().InvILocal()
 	};
 	if (JointSettings.bMassConditioningEnabled)
@@ -231,7 +236,7 @@ void FPBDJointCachedSolver::Init(
 	LastDQs[0] = FVec3(0.f);
 	LastDQs[1] = FVec3(0.f);
 }
-	
+
 void FPBDJointCachedSolver::InitProjection(
 	const FReal Dt,
 	const FPBDJointSolverSettings& SolverSettings,
@@ -240,9 +245,10 @@ void FPBDJointCachedSolver::InitProjection(
 	const FReal LinearProjection = FPBDJointUtilities::GetLinearProjection(SolverSettings, JointSettings);
 	const FReal AngularProjection = FPBDJointUtilities::GetAngularProjection(SolverSettings, JointSettings);
 	const bool bHasLinearProjection = JointSettings.bProjectionEnabled && ((LinearProjection > 0) || (JointSettings.TeleportDistance > 0));
-	const bool bHasAngularProjection = JointSettings.bProjectionEnabled && ((AngularProjection > 0) || (JointSettings.TeleportAngle > 0));
+	// Teleport angle is not implemented in this linear solver, not need to initialize rotation for teleport. 
+	const bool bHasAngularProjection = JointSettings.bProjectionEnabled && ((AngularProjection > 0) /*|| (JointSettings.TeleportAngle > 0)*/);
 
-	if(bHasLinearProjection || bHasAngularProjection)
+	if (bHasLinearProjection || bHasAngularProjection)
 	{
 		ComputeBodyState(0);
 		ComputeBodyState(1);
@@ -258,12 +264,12 @@ void FPBDJointCachedSolver::InitProjection(
 		// We are reusing the constraints for projection but we don't want to reset the accumulated corrections
 		const bool bResetLambdas = false;
 
-		if(bHasLinearProjection)
+		if (bHasLinearProjection)
 		{
 			InitPositionConstraints(Dt, SolverSettings, JointSettings, bResetLambdas);
 		}
 
-		if(bHasAngularProjection)
+		if (bHasAngularProjection)
 		{
 			InitRotationConstraints(Dt, SolverSettings, JointSettings, bResetLambdas);
 		}
@@ -279,10 +285,10 @@ void FPBDJointCachedSolver::Deinit()
 /** Main Apply function to solve all the constraint*/
 
 void FPBDJointCachedSolver::ApplyConstraints(
-		const FReal Dt,
-		const FReal InSolverStiffness,
-		const FPBDJointSolverSettings& SolverSettings,
-		const FPBDJointSettings& JointSettings)
+	const FReal Dt,
+	const FReal InSolverStiffness,
+	const FPBDJointSolverSettings& SolverSettings,
+	const FPBDJointSettings& JointSettings)
 {
 	NumActiveConstraints = 0;
 	SolverStiffness = InSolverStiffness;
@@ -339,7 +345,7 @@ void FPBDJointCachedSolver::ApplyVelocityConstraints(
 
 /** UTILS FOR POSITION CONSTRAINT **************************************************************************************/
 
-FORCEINLINE bool ExtractLinearMotion( const FPBDJointSettings& JointSettings,
+FORCEINLINE bool ExtractLinearMotion(const FPBDJointSettings& JointSettings,
 	TVec3<bool>& bLinearLocked, TVec3<bool>& bLinearLimited)
 {
 	bool bHasPositionConstraints =
@@ -378,32 +384,70 @@ void FPBDJointCachedSolver::InitPositionConstraints(
 	PositionConstraints.bValidDatas[0] = false;
 	PositionConstraints.bValidDatas[1] = false;
 	PositionConstraints.bValidDatas[2] = false;
+	PositionConstraints.bUseSimd = false;
 
 	TVec3<bool> bLinearLocked, bLinearLimited;
-	if(!ExtractLinearMotion(JointSettings, bLinearLocked, bLinearLimited))
+	if (!ExtractLinearMotion(JointSettings, bLinearLocked, bLinearLimited))
 		return;
 
+	PositionConstraints.bUseSimd = bUseSimd && bLinearLocked[0] && bLinearLocked[1] && bLinearLocked[2];
 	PositionConstraints.bAccelerationMode = FPBDJointUtilities::GetLinearSoftAccelerationMode(SolverSettings, JointSettings);
 
-	for(int32 ConstraintIndex = 0; ConstraintIndex < 3; ++ ConstraintIndex)
+	if (PositionConstraints.bUseSimd)
+	{
+		PositionConstraints.bSoftLimit[0] = false;
+		PositionConstraints.bSoftLimit[1] = false;
+		PositionConstraints.bSoftLimit[2] = false;
+		FRealSingle HardStiffness = FRealSingle(FPBDJointUtilities::GetLinearStiffness(SolverSettings, JointSettings));
+		PositionConstraints.Simd.ConstraintHardStiffnessSimd = VectorLoadFloat1(&HardStiffness);
+		PositionConstraints.ConstraintMaxLambda = FVec3(0);
+		PositionConstraints.bValidDatas[0] = false;
+		PositionConstraints.bValidDatas[1] = false;
+		PositionConstraints.bValidDatas[2] = false;
+		PositionConstraints.bLimitsCheck[0] = true;
+		PositionConstraints.bLimitsCheck[1] = true;
+		PositionConstraints.bLimitsCheck[2] = true;
+		PositionConstraints.MotionType[0] = EJointMotionType::Locked;
+		PositionConstraints.MotionType[1] = EJointMotionType::Locked;
+		PositionConstraints.MotionType[2] = EJointMotionType::Locked;
+		if (bResetLambdas)
+		{
+			PositionConstraints.Simd.ConstraintLambda = VectorZeroFloat();
+			PositionConstraints.ConstraintLambdaVelocity = FVec3::Zero();
+			PositionConstraints.Simd.ConstraintLimitsSimd = VectorZeroFloat();
+		}
+
+	}
+	else
 	{
-		PositionConstraints.InitDatas(ConstraintIndex,bLinearLimited[ConstraintIndex] &&
-			FPBDJointUtilities::GetSoftLinearLimitEnabled(SolverSettings, JointSettings),
-			FPBDJointUtilities::GetSoftLinearStiffness(SolverSettings, JointSettings),
-			FPBDJointUtilities::GetSoftLinearDamping(SolverSettings, JointSettings),
-			FPBDJointUtilities::GetLinearStiffness(SolverSettings, JointSettings),
-			bResetLambdas);
+		for (int32 ConstraintIndex = 0; ConstraintIndex < 3; ++ConstraintIndex)
+		{
+			PositionConstraints.InitDatas(ConstraintIndex, bLinearLimited[ConstraintIndex] &&
+				FPBDJointUtilities::GetSoftLinearLimitEnabled(SolverSettings, JointSettings),
+				FPBDJointUtilities::GetSoftLinearStiffness(SolverSettings, JointSettings),
+				FPBDJointUtilities::GetSoftLinearDamping(SolverSettings, JointSettings),
+				FPBDJointUtilities::GetLinearStiffness(SolverSettings, JointSettings),
+				bResetLambdas);
+		}
 	}
 
 	const TVec3<EJointMotionType>& LinearMotion = JointSettings.LinearMotionTypes;
 
-	if (bLinearLocked[0] || bLinearLocked[1] || bLinearLocked[2]) 
+	if (bLinearLocked[0] || bLinearLocked[1] || bLinearLocked[2])
 	{
-		// Process locked constraints
-		InitLockedPositionConstraint(JointSettings, Dt, LinearMotion);
+		if (PositionConstraints.bUseSimd)
+		{
+			InitLockedPositionConstraintSimd(JointSettings, Dt, LinearMotion);
+		}
+		else
+		{
+			// Process locked constraints
+			InitLockedPositionConstraint(JointSettings, Dt, LinearMotion);
+		}
 	}
 	if (bLinearLimited[0] || bLinearLimited[1] || bLinearLimited[2])
 	{
+		check(PositionConstraints.bUseSimd == false);
 		// Process limited constraints
 		if (bLinearLimited[0] && bLinearLimited[1] && bLinearLimited[2])
 		{
@@ -444,12 +488,12 @@ void FPBDJointCachedSolver::InitPositionConstraints(
 }
 
 void FPBDJointCachedSolver::InitPositionDatasMass(
-	FAxisConstraintDatas& PositionDatas, 
+	FAxisConstraintDatas& PositionDatas,
 	const int32 ConstraintIndex,
 	const FReal Dt)
 {
-	const FVec3 AngularAxis0 = FVec3::CrossProduct(PositionDatas.ConstraintArms[ConstraintIndex][0], PositionDatas.ConstraintAxis[ConstraintIndex]);
-	const FVec3 AngularAxis1 = FVec3::CrossProduct(PositionDatas.ConstraintArms[ConstraintIndex][1], PositionDatas.ConstraintAxis[ConstraintIndex]);
+	const FVec3 AngularAxis0 = FVec3::CrossProduct(PositionDatas.D.ConstraintArms[ConstraintIndex][0], PositionDatas.D.ConstraintAxis[ConstraintIndex]);
+	const FVec3 AngularAxis1 = FVec3::CrossProduct(PositionDatas.D.ConstraintArms[ConstraintIndex][1], PositionDatas.D.ConstraintAxis[ConstraintIndex]);
 	const FVec3 IA0 = Utilities::Multiply(InvI(0), AngularAxis0);
 	const FVec3 IA1 = Utilities::Multiply(InvI(1), AngularAxis1);
 	const FReal II0 = FVec3::DotProduct(AngularAxis0, IA0);
@@ -457,7 +501,7 @@ void FPBDJointCachedSolver::InitPositionDatasMass(
 
 	PositionDatas.UpdateMass(ConstraintIndex, IA0, IA1, InvM(0) + II0 + InvM(1) + II1, Dt, bUsePositionBasedDrives);
 }
-	
+
 void FPBDJointCachedSolver::SetInitConstraintVelocity(
 	const FVec3& ConstraintArm0,
 	const FVec3& ConstraintArm1)
@@ -485,15 +529,15 @@ void FPBDJointCachedSolver::InitPositionConstraintDatas(
 
 	PositionConstraints.MotionType[ConstraintIndex] = JointType;
 
-	if(JointType == EJointMotionType::Locked)
+	if (JointType == EJointMotionType::Locked)
 	{
-		PositionConstraints.ConstraintLimits[ConstraintIndex] = 0.0f;
+		PositionConstraints.D.ConstraintLimits[ConstraintIndex] = 0.0f;
 		PositionConstraints.UpdateDatas(ConstraintIndex, LocalAxis, LocalDelta,
 			0.0, false, ConstraintArm0, ConstraintArm1);
 	}
-	else if(JointType == EJointMotionType::Limited)
+	else if (JointType == EJointMotionType::Limited)
 	{
-		PositionConstraints.ConstraintLimits[ConstraintIndex] = ConstraintLimit;
+		PositionConstraints.D.ConstraintLimits[ConstraintIndex] = ConstraintLimit;
 		PositionConstraints.UpdateDatas(ConstraintIndex, LocalAxis, LocalDelta,
 			ConstraintRestitution, true, ConstraintArm0, ConstraintArm1);
 	}
@@ -503,46 +547,108 @@ void FPBDJointCachedSolver::InitPositionConstraintDatas(
 	InitPositionDatasMass(PositionConstraints, ConstraintIndex, Dt);
 }
 
-	
+
 void FPBDJointCachedSolver::InitLockedPositionConstraint(
-		const FPBDJointSettings& JointSettings,
-		const FReal Dt,
-		const TVec3<EJointMotionType>& LinearMotion)
+	const FPBDJointSettings& JointSettings,
+	const FReal Dt,
+	const TVec3<EJointMotionType>& LinearMotion)
 {
 	FVec3 ConstraintArm0 = ConnectorXs[1] - CurrentPs[0];
 	const FVec3 ConstraintArm1 = ConnectorXs[1] - CurrentPs[1];
-	
+
 	const FVec3 DX = ConnectorXs[1] - ConnectorXs[0];
-	const FMatrix33 R0M = ConnectorRs[0].ToMatrix(); 
+	const FMatrix33 R0M = ConnectorRs[0].ToMatrix();
 	FVec3 CX = FVec3::ZeroVector;
-	
+
 	// For a locked constraint we try to match an exact constraint, 
 	// it is why we are adding back the constraint projection along each axis.
 	// If the 3 axis are locked the constraintarm0 is then ConnectorXs[0] - CurrentPs[0];
-	for(int32 ConstraintIndex = 0; ConstraintIndex < 3; ++ConstraintIndex)
+	for (int32 ConstraintIndex = 0; ConstraintIndex < 3; ++ConstraintIndex)
 	{
-		if(LinearMotion[ConstraintIndex] == EJointMotionType::Locked)
+		if (LinearMotion[ConstraintIndex] == EJointMotionType::Locked)
 		{
 			const FVec3 ConstraintAxis = R0M.GetAxis(ConstraintIndex);
 			CX[ConstraintIndex] = FVec3::DotProduct(DX, ConstraintAxis);
 			ConstraintArm0 -= ConstraintAxis * CX[ConstraintIndex];
-			
+
 		}
 	}
 
 	SetInitConstraintVelocity(ConstraintArm0, ConstraintArm1);
 
-	for(int32 ConstraintIndex = 0; ConstraintIndex < 3; ++ConstraintIndex)
+	for (int32 ConstraintIndex = 0; ConstraintIndex < 3; ++ConstraintIndex)
 	{
-		if(LinearMotion[ConstraintIndex] == EJointMotionType::Locked)
+		if (LinearMotion[ConstraintIndex] == EJointMotionType::Locked)
 		{
 			const FVec3 ConstraintAxis = R0M.GetAxis(ConstraintIndex);
-			InitPositionConstraintDatas(ConstraintIndex, ConstraintAxis,  CX[ConstraintIndex], 0.0, Dt,
+			InitPositionConstraintDatas(ConstraintIndex, ConstraintAxis, CX[ConstraintIndex], 0.0, Dt,
 				0.0, EJointMotionType::Locked, ConstraintArm0, ConstraintArm1);
 		}
 	}
 }
 
+void FPBDJointCachedSolver::InitLockedPositionConstraintSimd(
+	const FPBDJointSettings& JointSettings,
+	const FReal Dt,
+	const TVec3<EJointMotionType>& LinearMotion)
+{
+	FVec3 ConstraintArm0 = ConnectorXs[1] - CurrentPs[0];
+	const FVec3 ConstraintArm1 = ConnectorXs[1] - CurrentPs[1];
+
+	const FVec3 DX = ConnectorXs[1] - ConnectorXs[0];
+	const FMatrix33 R0M = ConnectorRs[0].ToMatrix();
+	FVec3 CX = FVec3::ZeroVector;
+
+	FRealSingle AddedIMass = FRealSingle(InvM(0) + InvM(1));
+	FVec3f HardIM(AddedIMass);
+	FVec3f LocalDeltas;
+
+	// For a locked constraint we try to match an exact constraint, 
+	// it is why we are adding back the constraint projection along each axis.
+	// If the 3 axis are locked the constraintarm0 is then ConnectorXs[0] - CurrentPs[0];
+	for (int32 ConstraintIndex = 0; ConstraintIndex < 3; ++ConstraintIndex)
+	{
+		check(LinearMotion[ConstraintIndex] == EJointMotionType::Locked);
+		const FVec3 ConstraintAxis = R0M.GetAxis(ConstraintIndex);
+		CX[ConstraintIndex] = FVec3::DotProduct(DX, ConstraintAxis);
+		ConstraintArm0 -= ConstraintAxis * CX[ConstraintIndex];
+
+
+		const FVec3 LocalAxis = (CX[ConstraintIndex] < 0.0f) ? -ConstraintAxis : ConstraintAxis;
+		LocalDeltas[ConstraintIndex] = FRealSingle((CX[ConstraintIndex] < 0.0) ? -CX[ConstraintIndex] : CX[ConstraintIndex]);
+
+		PositionConstraints.MotionType[ConstraintIndex] = EJointMotionType::Locked;
+		PositionConstraints.bValidDatas[ConstraintIndex] = true;
+		PositionConstraints.bLimitsCheck[ConstraintIndex] = false;
+		PositionConstraints.ConstraintVX[ConstraintIndex] = 0.0;
+		PositionConstraints.ConstraintRestitution[ConstraintIndex] = 0.0;
+
+		PositionConstraints.Simd.ConstraintArmsSimd[ConstraintIndex][0] = MakeVectorRegisterFloatFromDouble(MakeVectorRegister(ConstraintArm0[0], ConstraintArm0[1], ConstraintArm0[2], 0.0f));
+		PositionConstraints.Simd.ConstraintArmsSimd[ConstraintIndex][1] = MakeVectorRegisterFloatFromDouble(MakeVectorRegister(ConstraintArm1[0], ConstraintArm1[1], ConstraintArm1[2], 0.0f));
+		PositionConstraints.Simd.ConstraintAxisSimd[ConstraintIndex] = MakeVectorRegisterFloatFromDouble(MakeVectorRegister(LocalAxis[0], LocalAxis[1], LocalAxis[2], 0.0f));
+
+		InitConstraintAxisLinearVelocities[ConstraintIndex] = FVec3::DotProduct(InitConstraintVelocity, LocalAxis);
+
+		const VectorRegister4Float AngularAxis0 = VectorCross(PositionConstraints.Simd.ConstraintArmsSimd[ConstraintIndex][0], PositionConstraints.Simd.ConstraintAxisSimd[ConstraintIndex]);
+		const VectorRegister4Float IA0 = Private::VectorMatrixMultiply(AngularAxis0, InvI(0));
+
+		const VectorRegister4Float AngularAxis1 = VectorCross(PositionConstraints.Simd.ConstraintArmsSimd[ConstraintIndex][1], PositionConstraints.Simd.ConstraintAxisSimd[ConstraintIndex]);
+		const VectorRegister4Float IA1 = Private::VectorMatrixMultiply(AngularAxis1, InvI(1));
+
+		const FRealSingle II0 = VectorDot3Scalar(AngularAxis0, IA0);
+		const FRealSingle II1 = VectorDot3Scalar(AngularAxis1, IA1);
+		HardIM[ConstraintIndex] += II0 + II1;
+		PositionConstraints.Simd.ConstraintDRAxisSimd[ConstraintIndex][0] = IA0;
+		PositionConstraints.Simd.ConstraintDRAxisSimd[ConstraintIndex][1] = VectorNegate(IA1);
+
+
+	}
+	PositionConstraints.Simd.ConstraintCX = MakeVectorRegisterFloat(LocalDeltas[0], LocalDeltas[1], LocalDeltas[2], 0.0f);
+	PositionConstraints.Simd.ConstraintHardIMSimd = MakeVectorRegisterFloatFromDouble(MakeVectorRegister(HardIM[0], HardIM[1], HardIM[2], 0.0f));
+	SetInitConstraintVelocity(ConstraintArm0, ConstraintArm1);
+
+}
+
 void FPBDJointCachedSolver::InitSphericalPositionConstraint(
 	const FPBDJointSettings& JointSettings,
 	const FReal Dt)
@@ -550,29 +656,29 @@ void FPBDJointCachedSolver::InitSphericalPositionConstraint(
 	FVec3 SphereAxis0;
 	FReal SphereDelta0;
 	FPBDJointUtilities::GetSphericalAxisDelta(ConnectorXs[0], ConnectorXs[1], SphereAxis0, SphereDelta0);
-	
-	const FVec3 SphereAxis1 = (FMath::Abs(FMath::Abs(SphereAxis0.Dot(FVec3(1,0,0))-1.0)) > UE_SMALL_NUMBER) ? SphereAxis0.Cross(FVec3(1,0,0)) :
-	                          (FMath::Abs(FMath::Abs(SphereAxis0.Dot(FVec3(0,1,0))-1.0)) > UE_SMALL_NUMBER) ? SphereAxis0.Cross(FVec3(0,1,0)) : SphereAxis0.Cross(FVec3(0,0,1)) ;
+
+	const FVec3 SphereAxis1 = (FMath::Abs(FMath::Abs(SphereAxis0.Dot(FVec3(1, 0, 0)) - 1.0)) > UE_SMALL_NUMBER) ? SphereAxis0.Cross(FVec3(1, 0, 0)) :
+		(FMath::Abs(FMath::Abs(SphereAxis0.Dot(FVec3(0, 1, 0)) - 1.0)) > UE_SMALL_NUMBER) ? SphereAxis0.Cross(FVec3(0, 1, 0)) : SphereAxis0.Cross(FVec3(0, 0, 1));
 	const FVec3 SphereAxis2 = SphereAxis0.Cross(SphereAxis1);
 
 	// Using Connector1 for both conserves angular momentum and avoid having 
 	// too much torque applied onto the COM. But can only be used for limited constraints
 	const FVec3 ConstraintArm0 = ConnectorXs[1] - CurrentPs[0];
 	const FVec3 ConstraintArm1 = ConnectorXs[1] - CurrentPs[1];
-	
+
 	SetInitConstraintVelocity(ConstraintArm0, ConstraintArm1);
 
-	InitPositionConstraintDatas(0, SphereAxis0,  SphereDelta0, JointSettings.LinearRestitution, Dt,
+	InitPositionConstraintDatas(0, SphereAxis0, SphereDelta0, JointSettings.LinearRestitution, Dt,
 		JointSettings.LinearLimit, EJointMotionType::Limited, ConstraintArm0, ConstraintArm1);
 
 	// SphereAxis0 being the direction axis, the geometric error for the 2 other axis are 0
 	// We need to add these 2 constraints for a linear solver to avoid drifting away in
 	//  the other directions while solving. For a non linear solver since we are recomputing 
 	// the main direction at each time we don't need that
-	InitPositionConstraintDatas(1, SphereAxis1,  0.0, JointSettings.LinearRestitution, Dt,
+	InitPositionConstraintDatas(1, SphereAxis1, 0.0, JointSettings.LinearRestitution, Dt,
 		JointSettings.LinearLimit, EJointMotionType::Limited, ConstraintArm0, ConstraintArm1);
 
-	InitPositionConstraintDatas(2, SphereAxis2,  0.0, JointSettings.LinearRestitution, Dt,
+	InitPositionConstraintDatas(2, SphereAxis2, 0.0, JointSettings.LinearRestitution, Dt,
 		JointSettings.LinearLimit, EJointMotionType::Limited, ConstraintArm0, ConstraintArm1);
 }
 
@@ -587,18 +693,18 @@ void FPBDJointCachedSolver::InitCylindricalPositionConstraint(
 		AxisIndex, PlaneAxis, PlaneDelta, RadialAxis0, RadialDelta0);
 
 	const FVec3 RadialAxis1 = PlaneAxis.Cross(RadialAxis0);
-	const FReal RadialDelta1 =  (ConnectorXs[1] - ConnectorXs[0]).Dot(RadialAxis1);
+	const FReal RadialDelta1 = (ConnectorXs[1] - ConnectorXs[0]).Dot(RadialAxis1);
 
 	const FVec3 ConstraintArm0 = ConnectorXs[1] - CurrentPs[0];
 	const FVec3 ConstraintArm1 = ConnectorXs[1] - CurrentPs[1];
-	
+
 	SetInitConstraintVelocity(ConstraintArm0, ConstraintArm1);
 
-	InitPositionConstraintDatas((AxisIndex+1)%3, RadialAxis0, RadialDelta0, JointSettings.LinearRestitution, Dt,
+	InitPositionConstraintDatas((AxisIndex + 1) % 3, RadialAxis0, RadialDelta0, JointSettings.LinearRestitution, Dt,
 		JointSettings.LinearLimit, EJointMotionType::Limited, ConstraintArm0, ConstraintArm1);
 
-	InitPositionConstraintDatas((AxisIndex+2)%3, RadialAxis1, RadialDelta1, JointSettings.LinearRestitution, Dt,
-	    JointSettings.LinearLimit, EJointMotionType::Limited, ConstraintArm0, ConstraintArm1);   
+	InitPositionConstraintDatas((AxisIndex + 2) % 3, RadialAxis1, RadialDelta1, JointSettings.LinearRestitution, Dt,
+		JointSettings.LinearLimit, EJointMotionType::Limited, ConstraintArm0, ConstraintArm1);
 }
 
 void FPBDJointCachedSolver::InitPlanarPositionConstraint(
@@ -609,10 +715,10 @@ void FPBDJointCachedSolver::InitPlanarPositionConstraint(
 	FVec3 PlaneAxis;
 	FReal PlaneDelta;
 	FPBDJointUtilities::GetPlanarAxisDelta(ConnectorRs[0], ConnectorXs[0], ConnectorXs[1], AxisIndex, PlaneAxis, PlaneDelta);
-	
+
 	const FVec3 ConstraintArm0 = ConnectorXs[1] - CurrentPs[0];
 	const FVec3 ConstraintArm1 = ConnectorXs[1] - CurrentPs[1];
-	
+
 	SetInitConstraintVelocity(ConstraintArm0, ConstraintArm1);
 
 	InitPositionConstraintDatas(AxisIndex, PlaneAxis, PlaneDelta, JointSettings.LinearRestitution, Dt,
@@ -624,35 +730,42 @@ void FPBDJointCachedSolver::InitPlanarPositionConstraint(
 void FPBDJointCachedSolver::ApplyPositionConstraints(
 	const FReal Dt)
 {
-	for(int32 ConstraintIndex = 0; ConstraintIndex < 3; ++ConstraintIndex)
+	if (PositionConstraints.bUseSimd)
+	{
+		ApplyPositionConstraintsSimd(Dt);
+	}
+	else
 	{
-		if(PositionConstraints.bValidDatas[ConstraintIndex])
+		for (int32 ConstraintIndex = 0; ConstraintIndex < 3; ++ConstraintIndex)
 		{
-			ApplyAxisPositionConstraint(ConstraintIndex, Dt);
+			if (PositionConstraints.bValidDatas[ConstraintIndex])
+			{
+				ApplyAxisPositionConstraint(ConstraintIndex, Dt);
+			}
 		}
 	}
 }
-	
+
 void FPBDJointCachedSolver::SolvePositionConstraintDelta(
-	const int32 ConstraintIndex, 
+	const int32 ConstraintIndex,
 	const FReal DeltaLambda,
 	const FAxisConstraintDatas& ConstraintDatas)
 {
-	const FVec3 DX = ConstraintDatas.ConstraintAxis[ConstraintIndex] * DeltaLambda;
+	const FVec3 DX = ConstraintDatas.D.ConstraintAxis[ConstraintIndex] * DeltaLambda;
 
-	if(Body(0).IsDynamic())
+	if (Body(0).IsDynamic())
 	{
 		const FVec3 DP0 = InvM(0) * DX;
-		const FVec3 DR0 = ConstraintDatas.ConstraintDRAxis[ConstraintIndex][0] * DeltaLambda;
-		ApplyPositionDelta(0,DP0);
-		ApplyRotationDelta(0,DR0);
+		const FVec3 DR0 = ConstraintDatas.D.ConstraintDRAxis[ConstraintIndex][0] * DeltaLambda;
+		ApplyPositionDelta(0, DP0);
+		ApplyRotationDelta(0, DR0);
 	}
-	if(Body(1).IsDynamic())
+	if (Body(1).IsDynamic())
 	{
 		const FVec3 DP1 = -InvM(1) * DX;
-		const FVec3 DR1 = ConstraintDatas.ConstraintDRAxis[ConstraintIndex][1] * DeltaLambda;
-		ApplyPositionDelta(1,DP1);
-		ApplyRotationDelta(1,DR1);
+		const FVec3 DR1 = ConstraintDatas.D.ConstraintDRAxis[ConstraintIndex][1] * DeltaLambda;
+		ApplyPositionDelta(1, DP1);
+		ApplyRotationDelta(1, DR1);
 	}
 
 	++NumActiveConstraints;
@@ -662,10 +775,10 @@ void FPBDJointCachedSolver::SolvePositionConstraintHard(
 	const int32 ConstraintIndex,
 	const FReal DeltaConstraint)
 {
-	const FReal DeltaLambda = SolverStiffness * PositionConstraints.ConstraintHardStiffness[ConstraintIndex] * DeltaConstraint /
-		PositionConstraints.ConstraintHardIM[ConstraintIndex];
+	const FReal DeltaLambda = SolverStiffness * PositionConstraints.D.ConstraintHardStiffness[ConstraintIndex] * DeltaConstraint /
+		PositionConstraints.D.ConstraintHardIM[ConstraintIndex];
 
-	PositionConstraints.ConstraintLambda[ConstraintIndex] += DeltaLambda;
+	PositionConstraints.D.ConstraintLambda[ConstraintIndex] += DeltaLambda;
 	SolvePositionConstraintDelta(ConstraintIndex, DeltaLambda, PositionConstraints);
 }
 
@@ -675,17 +788,18 @@ void FPBDJointCachedSolver::SolvePositionConstraintSoft(
 	const FReal Dt,
 	const FReal TargetVel)
 {
+	check(PositionConstraints.bUseSimd == false);
 	FReal VelDt = 0;
-	if (PositionConstraints.ConstraintSoftDamping[ConstraintIndex] > UE_KINDA_SMALL_NUMBER)
+	if (PositionConstraints.D.ConstraintSoftDamping[ConstraintIndex] > UE_KINDA_SMALL_NUMBER)
 	{
-		const FVec3 V0Dt = FVec3::CalculateVelocity(InitConnectorXs[0], ConnectorXs[0]+Body(0).DP() + FVec3::CrossProduct(Body(0).DQ(), PositionConstraints.ConstraintArms[ConstraintIndex][0]), 1.0f);
-		const FVec3 V1Dt = FVec3::CalculateVelocity(InitConnectorXs[1], ConnectorXs[1]+Body(1).DP() + FVec3::CrossProduct(Body(1).DQ(), PositionConstraints.ConstraintArms[ConstraintIndex][1]), 1.0f);
-		VelDt = TargetVel * Dt + FVec3::DotProduct(V0Dt - V1Dt, PositionConstraints.ConstraintAxis[ConstraintIndex] );
+		const FVec3 V0Dt = FVec3::CalculateVelocity(InitConnectorXs[0], ConnectorXs[0] + Body(0).DP() + FVec3::CrossProduct(Body(0).DQ(), PositionConstraints.D.ConstraintArms[ConstraintIndex][0]), 1.0f);
+		const FVec3 V1Dt = FVec3::CalculateVelocity(InitConnectorXs[1], ConnectorXs[1] + Body(1).DP() + FVec3::CrossProduct(Body(1).DQ(), PositionConstraints.D.ConstraintArms[ConstraintIndex][1]), 1.0f);
+		VelDt = TargetVel * Dt + FVec3::DotProduct(V0Dt - V1Dt, PositionConstraints.D.ConstraintAxis[ConstraintIndex]);
 	}
 
-	const FReal DeltaLambda = SolverStiffness * (PositionConstraints.ConstraintSoftStiffness[ConstraintIndex] * DeltaConstraint - PositionConstraints.ConstraintSoftDamping[ConstraintIndex] * VelDt - PositionConstraints.ConstraintLambda[ConstraintIndex]) /
-		PositionConstraints.ConstraintSoftIM[ConstraintIndex];
-	PositionConstraints.ConstraintLambda[ConstraintIndex] += DeltaLambda;
+	const FReal DeltaLambda = SolverStiffness * (PositionConstraints.D.ConstraintSoftStiffness[ConstraintIndex] * DeltaConstraint - PositionConstraints.D.ConstraintSoftDamping[ConstraintIndex] * VelDt - PositionConstraints.D.ConstraintLambda[ConstraintIndex]) /
+		PositionConstraints.D.ConstraintSoftIM[ConstraintIndex];
+	PositionConstraints.D.ConstraintLambda[ConstraintIndex] += DeltaLambda;
 
 	SolvePositionConstraintDelta(ConstraintIndex, DeltaLambda, PositionConstraints);
 }
@@ -693,27 +807,29 @@ void FPBDJointCachedSolver::SolvePositionConstraintSoft(
 void FPBDJointCachedSolver::ApplyAxisPositionConstraint(
 	const int32 ConstraintIndex, const FReal Dt)
 {
-	const FVec3 CX = Body(1).DP()  - Body(0).DP() +
-		FVec3::CrossProduct(Body(1).DQ(), PositionConstraints.ConstraintArms[ConstraintIndex][1]) -
-			FVec3::CrossProduct(Body(0).DQ(), PositionConstraints.ConstraintArms[ConstraintIndex][0]) ;
-	
-	FReal DeltaPosition = PositionConstraints.ConstraintCX[ConstraintIndex] + FVec3::DotProduct(CX, PositionConstraints.ConstraintAxis[ConstraintIndex]);
+
+	check(PositionConstraints.bUseSimd == false);
+	const FVec3 CX = Body(1).DP() - Body(0).DP() +
+		FVec3::CrossProduct(Body(1).DQ(), PositionConstraints.D.ConstraintArms[ConstraintIndex][1]) -
+		FVec3::CrossProduct(Body(0).DQ(), PositionConstraints.D.ConstraintArms[ConstraintIndex][0]);
+
+	FReal DeltaPosition = PositionConstraints.D.ConstraintCX[ConstraintIndex] + FVec3::DotProduct(CX, PositionConstraints.D.ConstraintAxis[ConstraintIndex]);
 
 	bool NeedsSolve = false;
-	if(PositionConstraints.bLimitsCheck[ConstraintIndex])
+	if (PositionConstraints.bLimitsCheck[ConstraintIndex])
 	{
-		if(DeltaPosition > PositionConstraints.ConstraintLimits[ConstraintIndex] )
+		if (DeltaPosition > PositionConstraints.D.ConstraintLimits[ConstraintIndex])
 		{
-			DeltaPosition -= PositionConstraints.ConstraintLimits[ConstraintIndex];
+			DeltaPosition -= PositionConstraints.D.ConstraintLimits[ConstraintIndex];
 			NeedsSolve = true;
 		}
-		else if(DeltaPosition < -PositionConstraints.ConstraintLimits[ConstraintIndex])
+		else if (DeltaPosition < -PositionConstraints.D.ConstraintLimits[ConstraintIndex])
 		{
-			DeltaPosition += PositionConstraints.ConstraintLimits[ConstraintIndex];
+			DeltaPosition += PositionConstraints.D.ConstraintLimits[ConstraintIndex];
 			NeedsSolve = true;
 		}
-	}  
-	if (!PositionConstraints.bLimitsCheck[ConstraintIndex] || (PositionConstraints.bLimitsCheck[ConstraintIndex] && NeedsSolve && FMath::Abs(DeltaPosition ) > PositionTolerance))
+	}
+	if (!PositionConstraints.bLimitsCheck[ConstraintIndex] || (PositionConstraints.bLimitsCheck[ConstraintIndex] && NeedsSolve && FMath::Abs(DeltaPosition) > PositionTolerance))
 	{
 		if ((PositionConstraints.MotionType[ConstraintIndex] == EJointMotionType::Limited) && PositionConstraints.bSoftLimit[ConstraintIndex])
 		{
@@ -726,15 +842,110 @@ void FPBDJointCachedSolver::ApplyAxisPositionConstraint(
 	}
 }
 
+void FPBDJointCachedSolver::ApplyPositionConstraintsSimd(
+	const FReal Dt)
+{
+	const VectorRegister4Float Body0DP = MakeVectorRegisterFloatFromDouble(MakeVectorRegister(Body(0).DP()[0], Body(0).DP()[1], Body(0).DP()[2], 0.0f));
+	const VectorRegister4Float Body1DP = MakeVectorRegisterFloatFromDouble(MakeVectorRegister(Body(1).DP()[0], Body(1).DP()[1], Body(1).DP()[2], 0.0f));
+
+	const VectorRegister4Float Body0DQ = MakeVectorRegisterFloatFromDouble(MakeVectorRegister(Body(0).DQ()[0], Body(0).DQ()[1], Body(0).DQ()[2], 0.0f));
+	const VectorRegister4Float Body1DQ = MakeVectorRegisterFloatFromDouble(MakeVectorRegister(Body(1).DQ()[0], Body(1).DQ()[1], Body(1).DQ()[2], 0.0f));
+
+	const VectorRegister4Float DPDiff = VectorSubtract(Body1DP, Body0DP);
+
+
+	VectorRegister4Float DeltaPositions[3];
+	for (int32 CIndex = 0; CIndex < 3; ++CIndex)
+	{
+		VectorRegister4Float  Cross1 = VectorCross(Body1DQ, PositionConstraints.Simd.ConstraintArmsSimd[CIndex][1]);
+		VectorRegister4Float  Cross0 = VectorCross(Body0DQ, PositionConstraints.Simd.ConstraintArmsSimd[CIndex][0]);
+		VectorRegister4Float  CrossDiff = VectorSubtract(Cross1, Cross0);
+		VectorRegister4Float CX = VectorAdd(DPDiff, CrossDiff);
+
+		DeltaPositions[CIndex] = Private::VectorDot3FastX(CX, PositionConstraints.Simd.ConstraintAxisSimd[CIndex]);
+	}
+
+	VectorRegister4Float DeltaPosition = VectorUnpackLo(DeltaPositions[0], DeltaPositions[1]);
+	DeltaPosition = VectorMoveLh(DeltaPosition, DeltaPositions[2]);
+
+	DeltaPosition = VectorAdd(DeltaPosition, PositionConstraints.Simd.ConstraintCX);
+
+	const FRealSingle SolverStiffnessf = FRealSingle(SolverStiffness);
+	VectorRegister4Float Stiffness = VectorLoadFloat1(&SolverStiffnessf);
+
+	Stiffness = VectorMultiply(Stiffness, PositionConstraints.Simd.ConstraintHardStiffnessSimd);
+	DeltaPosition = VectorMultiply(Stiffness, DeltaPosition);
+	const VectorRegister4Float DeltaLambda = VectorDivide(DeltaPosition, PositionConstraints.Simd.ConstraintHardIMSimd);
+	PositionConstraints.Simd.ConstraintLambda = VectorAdd(PositionConstraints.Simd.ConstraintLambda, DeltaLambda);
+
+	VectorRegister4Float DLambda[3];
+	DLambda[0] = VectorReplicate(DeltaLambda, 0);
+	DLambda[1] = VectorReplicate(DeltaLambda, 1);
+	DLambda[2] = VectorReplicate(DeltaLambda, 2);
+
+	VectorRegister4Float DX[3];
+	for (int32 CIndex = 0; CIndex < 3; CIndex++)
+	{
+		DX[CIndex] = VectorMultiply(PositionConstraints.Simd.ConstraintAxisSimd[CIndex], DLambda[CIndex]);
+	}
+
+	if (Body(0).IsDynamic())
+	{
+		const FRealSingle Inv0f = FRealSingle(InvM(0));
+		const VectorRegister4Float Inv0 = VectorLoadFloat1(&Inv0f);
+
+		VectorRegister4Float DP0 = VectorZero();
+		VectorRegister4Float DR0 = VectorZero();
+		for (int32 CIndex = 0; CIndex < 3; CIndex++)
+		{
+			DP0 = VectorMultiplyAdd(Inv0, DX[CIndex], DP0);
+			DR0 = VectorMultiplyAdd(PositionConstraints.Simd.ConstraintDRAxisSimd[CIndex][0], DLambda[CIndex], DR0);
+		}
+		FVec3f DP0f;
+		VectorStoreFloat3(DP0, &DP0f[0]);
+		ApplyPositionDelta(0, DP0f);
+		FVec3f DR0f;
+		VectorStoreFloat3(DR0, &DR0f[0]);
+		ApplyRotationDelta(0, DR0f);
+	}
+	if (Body(1).IsDynamic())
+	{
+		const FRealSingle Inv1f = FRealSingle(InvM(1));
+		const VectorRegister4Float Inv1 = VectorLoadFloat1(&Inv1f);
+
+		VectorRegister4Float DP1 = VectorZero();
+		VectorRegister4Float DR1 = VectorZero();
+		for (int32 CIndex = 0; CIndex < 3; CIndex++)
+		{
+			DP1 = VectorSubtract(DP1, VectorMultiply(Inv1, DX[CIndex]));
+			DR1 = VectorMultiplyAdd(PositionConstraints.Simd.ConstraintDRAxisSimd[CIndex][1], DLambda[CIndex], DR1);
+		}
+		FVec3f DP1f;
+		VectorStoreFloat3(DP1, &DP1f[0]);
+		ApplyPositionDelta(1, DP1f);
+		FVec3f DR1f;
+		VectorStoreFloat3(DR1, &DR1f[0]);
+		ApplyRotationDelta(1, DR1f);
+	}
+	NumActiveConstraints += 3;
+}
+
 /** APPLY LINEAR VELOCITY *********************************************************************************************/
 
 void FPBDJointCachedSolver::ApplyLinearVelocityConstraints()
 {
-	for(int32 ConstraintIndex = 0; ConstraintIndex < 3; ++ConstraintIndex)
+	if (PositionConstraints.bUseSimd)
+	{
+		ApplyVelocityConstraintSimd();
+	}
+	else
 	{
-		if(PositionConstraints.bValidDatas[ConstraintIndex])
+		for (int32 ConstraintIndex = 0; ConstraintIndex < 3; ++ConstraintIndex)
 		{
-			ApplyAxisVelocityConstraint(ConstraintIndex);
+			if (PositionConstraints.bValidDatas[ConstraintIndex])
+			{
+				ApplyAxisVelocityConstraint(ConstraintIndex);
+			}
 		}
 	}
 }
@@ -743,29 +954,30 @@ void FPBDJointCachedSolver::SolveLinearVelocityConstraint(
 	const int32 ConstraintIndex,
 	const FReal TargetVel)
 {
-	const FVec3 CV0 = V(0) + FVec3::CrossProduct(W(0), PositionConstraints.ConstraintArms[ConstraintIndex][0]);
-	const FVec3 CV1 = V(1) + FVec3::CrossProduct(W(1), PositionConstraints.ConstraintArms[ConstraintIndex][1]);
+	check(PositionConstraints.bUseSimd == false);
+	const FVec3 CV0 = V(0) + FVec3::CrossProduct(W(0), PositionConstraints.D.ConstraintArms[ConstraintIndex][0]);
+	const FVec3 CV1 = V(1) + FVec3::CrossProduct(W(1), PositionConstraints.D.ConstraintArms[ConstraintIndex][1]);
 	const FVec3 CV = CV1 - CV0;
 
-	const FReal DeltaLambda = SolverStiffness * PositionConstraints.ConstraintHardStiffness[ConstraintIndex] *
-	 (FVec3::DotProduct(CV, PositionConstraints.ConstraintAxis[ConstraintIndex]) - TargetVel) / PositionConstraints.ConstraintHardIM[ConstraintIndex];
+	const FReal DeltaLambda = SolverStiffness * PositionConstraints.D.ConstraintHardStiffness[ConstraintIndex] *
+		(FVec3::DotProduct(CV, PositionConstraints.D.ConstraintAxis[ConstraintIndex]) - TargetVel) / PositionConstraints.D.ConstraintHardIM[ConstraintIndex];
 
 	// @todo(chaos): We should be adding to the net positional impulse here
-	//PositionConstraints.ConstraintLambda[ConstraintIndex] += DeltaLambda * Dt;
-	
-	const FVec3 MDV = DeltaLambda * PositionConstraints.ConstraintAxis[ConstraintIndex];
-	
-	if(Body(0).IsDynamic())
+	//PositionConstraints.D.ConstraintLambda[ConstraintIndex] += DeltaLambda * Dt;
+
+	const FVec3 MDV = DeltaLambda * PositionConstraints.D.ConstraintAxis[ConstraintIndex];
+
+	if (Body(0).IsDynamic())
 	{
 		const FVec3 DV0 = InvM(0) * MDV;
-		const FVec3 DW0 = PositionConstraints.ConstraintDRAxis[ConstraintIndex][0] * DeltaLambda;
+		const FVec3 DW0 = PositionConstraints.D.ConstraintDRAxis[ConstraintIndex][0] * DeltaLambda;
 
 		Body(0).ApplyVelocityDelta(DV0, DW0);
 	}
-	if(Body(1).IsDynamic())
+	if (Body(1).IsDynamic())
 	{
 		const FVec3 DV1 = -InvM(1) * MDV;
-		const FVec3 DW1 = PositionConstraints.ConstraintDRAxis[ConstraintIndex][1] * DeltaLambda;
+		const FVec3 DW1 = PositionConstraints.D.ConstraintDRAxis[ConstraintIndex][1] * DeltaLambda;
 
 		Body(1).ApplyVelocityDelta(DV1, DW1);
 	}
@@ -774,9 +986,10 @@ void FPBDJointCachedSolver::SolveLinearVelocityConstraint(
 void FPBDJointCachedSolver::ApplyAxisVelocityConstraint(
 	const int32 ConstraintIndex)
 {
+	check(PositionConstraints.bUseSimd == false);
 	// Apply restitution for limited joints when we have exceeded the limits
 	// We also drive the velocity to zero for locked constraints (ignoring restitution)
-	if (FMath::Abs(PositionConstraints.ConstraintLambda[ConstraintIndex]) > UE_SMALL_NUMBER)
+	if (FMath::Abs(PositionConstraints.D.ConstraintLambda[ConstraintIndex]) > UE_SMALL_NUMBER)
 	{
 		FReal TargetVel = 0.0f;
 		const FReal Restitution = PositionConstraints.ConstraintRestitution[ConstraintIndex];
@@ -791,15 +1004,90 @@ void FPBDJointCachedSolver::ApplyAxisVelocityConstraint(
 	}
 }
 
+
+void FPBDJointCachedSolver::ApplyVelocityConstraintSimd()
+{
+	check(PositionConstraints.MotionType[0] != EJointMotionType::Limited);
+	check(PositionConstraints.MotionType[1] != EJointMotionType::Limited);
+	check(PositionConstraints.MotionType[2] != EJointMotionType::Limited);
+
+	const VectorRegister4Float IsGTEps = VectorCompareGT(VectorAbs(PositionConstraints.Simd.ConstraintLambda), GlobalVectorConstants::SmallNumber);
+
+	if (VectorMaskBits(IsGTEps))
+	{
+		const FVec3& V0d = V(0);
+		const VectorRegister4Float V0 = MakeVectorRegisterFloatFromDouble(MakeVectorRegister(V0d[0], V0d[1], V0d[2], 0.0f));
+		const FVec3& V1d = V(1);
+		const VectorRegister4Float V1 = MakeVectorRegisterFloatFromDouble(MakeVectorRegister(V1d[0], V1d[1], V1d[2], 0.0f));
+		const FVec3& W0d = W(0);
+		const VectorRegister4Float W0 = MakeVectorRegisterFloatFromDouble(MakeVectorRegister(W0d[0], W0d[1], W0d[2], 0.0f));
+		const FVec3& W1d = W(1);
+		const VectorRegister4Float W1 = MakeVectorRegisterFloatFromDouble(MakeVectorRegister(W1d[0], W1d[1], W1d[2], 0.0f));
+
+		const FRealSingle SolverStiffnessf = FRealSingle(SolverStiffness);
+		VectorRegister4Float Stiffness = VectorLoadFloat1(&SolverStiffnessf);
+		Stiffness = VectorMultiply(Stiffness, PositionConstraints.Simd.ConstraintHardStiffnessSimd);
+
+		VectorRegister4Float HardIM[3];
+		HardIM[0] = VectorReplicate(PositionConstraints.Simd.ConstraintHardIMSimd, 0);
+		HardIM[1] = VectorReplicate(PositionConstraints.Simd.ConstraintHardIMSimd, 1);
+		HardIM[2] = VectorReplicate(PositionConstraints.Simd.ConstraintHardIMSimd, 2);
+		VectorRegister4Float DeltaLambdas[3];
+		for (int32 CIndex = 0; CIndex < 3; CIndex++)
+		{
+			const VectorRegister4Float CV0 = VectorAdd(V0, VectorCross(W0, PositionConstraints.Simd.ConstraintArmsSimd[CIndex][0]));
+			const VectorRegister4Float CV1 = VectorAdd(V1, VectorCross(W1, PositionConstraints.Simd.ConstraintArmsSimd[CIndex][1]));
+			const VectorRegister4Float CV = VectorSubtract(CV1, CV0);
+			const VectorRegister4Float Dot3 = VectorDot3(CV, PositionConstraints.Simd.ConstraintAxisSimd[CIndex]);
+
+			DeltaLambdas[CIndex] = VectorDivide(VectorMultiply(Stiffness, Dot3), HardIM[CIndex]);
+		}
+		if (Body(0).IsDynamic())
+		{
+			const FRealSingle Inv0f = FRealSingle(InvM(0));
+			const VectorRegister4Float InvM0 = VectorLoadFloat1(&Inv0f);
+			VectorRegister4Float DV0 = VectorZeroFloat();
+			VectorRegister4Float DW0 = VectorZeroFloat();
+			for (int32 CIndex = 0; CIndex < 3; CIndex++)
+			{
+				DV0 = VectorMultiplyAdd(InvM0, VectorMultiply(DeltaLambdas[CIndex], PositionConstraints.Simd.ConstraintAxisSimd[CIndex]), DV0);
+				DW0 = VectorMultiplyAdd(PositionConstraints.Simd.ConstraintDRAxisSimd[CIndex][0], DeltaLambdas[CIndex], DW0);
+			}
+			FVec3f DV0f;
+			VectorStoreFloat3(DV0, &DV0f[0]);
+			FVec3f DW0f;
+			VectorStoreFloat3(DW0, &DW0f[0]);
+			Body(0).ApplyVelocityDelta(FVec3(DW0f), FVec3(DW0f));
+		}
+		if (Body(1).IsDynamic())
+		{
+			const FRealSingle Inv1f = -FRealSingle(InvM(1));
+			const VectorRegister4Float InvM1 = VectorLoadFloat1(&Inv1f);
+			VectorRegister4Float DV1 = VectorZeroFloat();
+			VectorRegister4Float DW1 = VectorZeroFloat();
+			for (int32 CIndex = 0; CIndex < 3; CIndex++)
+			{
+				DV1 = VectorMultiplyAdd(InvM1, VectorMultiply(DeltaLambdas[CIndex], PositionConstraints.Simd.ConstraintAxisSimd[CIndex]), DV1);
+				DW1 = VectorMultiplyAdd(PositionConstraints.Simd.ConstraintDRAxisSimd[CIndex][1], DeltaLambdas[CIndex], DW1);
+			}
+			FVec3f DV1f;
+			VectorStoreFloat3(DV1, &DV1f[0]);
+			FVec3f DW1f;
+			VectorStoreFloat3(DW1, &DW1f[0]);
+			Body(1).ApplyVelocityDelta(FVec3(DW1f), FVec3(DW1f));
+		}
+	}
+}
+
 /** UTILS FOR ROTATION CONSTRAINT **************************************************************************************/
 
-FORCEINLINE bool ExtractAngularMotion( const FPBDJointSettings& JointSettings,
-		TVec3<bool>& bAngularLocked, TVec3<bool>& bAngularLimited, TVec3<bool>& bAngularFree)
+bool ExtractAngularMotion(const FPBDJointSettings& JointSettings,
+	TVec3<bool>& bAngularLocked, TVec3<bool>& bAngularLimited, TVec3<bool>& bAngularFree)
 {
 	bool bHasRotationConstraints =
-			  (JointSettings.AngularMotionTypes[0] != EJointMotionType::Free)
-		   || (JointSettings.AngularMotionTypes[1] != EJointMotionType::Free)
-		   || (JointSettings.AngularMotionTypes[2] != EJointMotionType::Free);
+		(JointSettings.AngularMotionTypes[0] != EJointMotionType::Free)
+		|| (JointSettings.AngularMotionTypes[1] != EJointMotionType::Free)
+		|| (JointSettings.AngularMotionTypes[2] != EJointMotionType::Free);
 	if (!bHasRotationConstraints)
 	{
 		return false;
@@ -818,7 +1106,7 @@ FORCEINLINE bool ExtractAngularMotion( const FPBDJointSettings& JointSettings,
 		(AngularMotion[1] == EJointMotionType::Limited),
 		(AngularMotion[2] == EJointMotionType::Limited),
 	};
-	bAngularFree=
+	bAngularFree =
 	{
 		(AngularMotion[0] == EJointMotionType::Free),
 		(AngularMotion[1] == EJointMotionType::Free),
@@ -838,129 +1126,245 @@ void FPBDJointCachedSolver::InitRotationConstraints(
 	RotationConstraints.bValidDatas[0] = false;
 	RotationConstraints.bValidDatas[1] = false;
 	RotationConstraints.bValidDatas[2] = false;
+	RotationConstraints.bUseSimd = false;
+
 
 	TVec3<bool> bAngularLocked, bAngularLimited, bAngularFree;
-	if(!ExtractAngularMotion(JointSettings, bAngularLocked, bAngularLimited, bAngularFree))
+	if (!ExtractAngularMotion(JointSettings, bAngularLocked, bAngularLimited, bAngularFree))
 		return;
 
 	RotationConstraints.bAccelerationMode = FPBDJointUtilities::GetAngularSoftAccelerationMode(SolverSettings, JointSettings);
 
-	const int32 TW = (int32)EJointAngularConstraintIndex::Twist;
-	const int32 S1 = (int32)EJointAngularConstraintIndex::Swing1;
-	const int32 S2 = (int32)EJointAngularConstraintIndex::Swing2;
-
-	RotationConstraints.InitDatas(TW,FPBDJointUtilities::GetSoftTwistLimitEnabled(SolverSettings, JointSettings) && !bAngularLocked[TW],
-					FPBDJointUtilities::GetSoftTwistStiffness(SolverSettings, JointSettings),
-					FPBDJointUtilities::GetSoftTwistDamping(SolverSettings, JointSettings),
-					FPBDJointUtilities::GetTwistStiffness(SolverSettings, JointSettings),
-					bResetLambdas);
-
-	RotationConstraints.InitDatas(S1,FPBDJointUtilities::GetSoftSwingLimitEnabled(SolverSettings, JointSettings) && !bAngularLocked[S1],
-					FPBDJointUtilities::GetSoftSwingStiffness(SolverSettings, JointSettings),
-					FPBDJointUtilities::GetSoftSwingDamping(SolverSettings, JointSettings),
-					FPBDJointUtilities::GetSwingStiffness(SolverSettings, JointSettings),
-					bResetLambdas);
-
-	RotationConstraints.InitDatas(S2, FPBDJointUtilities::GetSoftSwingLimitEnabled(SolverSettings, JointSettings) && !bAngularLocked[S2],
-					FPBDJointUtilities::GetSoftSwingStiffness(SolverSettings, JointSettings),
-					FPBDJointUtilities::GetSoftSwingDamping(SolverSettings, JointSettings),
-					FPBDJointUtilities::GetSwingStiffness(SolverSettings, JointSettings),
-					bResetLambdas);
+	RotationConstraints.bUseSimd = bUseSimd && bAngularLimited[0] && bAngularLimited[1] && bAngularLimited[2] &&
+		FPBDJointUtilities::GetSoftTwistLimitEnabled(SolverSettings, JointSettings) && FPBDJointUtilities::GetSoftSwingLimitEnabled(SolverSettings, JointSettings);
 
-	const FVec3 Twist0 = ConnectorRs[0] * FJointConstants::TwistAxis();
-	const FVec3 Twist1 = ConnectorRs[1] * FJointConstants::TwistAxis();
-	const bool bDegenerate = (FVec3::DotProduct(Twist0, Twist1) < Chaos_Joint_DegenerateRotationLimit);
-
-	
-	// Apply twist constraint
-	// NOTE: Cannot calculate twist angle at 180degree swing
-	if (SolverSettings.bEnableTwistLimits)
+	if (RotationConstraints.bUseSimd)
 	{
-		if (bAngularLimited[TW] && !bDegenerate)
+		RotationConstraints.SettingsSoftStiffness = FVec3(FPBDJointUtilities::GetSoftTwistStiffness(SolverSettings, JointSettings), FPBDJointUtilities::GetSoftSwingStiffness(SolverSettings, JointSettings), FPBDJointUtilities::GetSoftSwingStiffness(SolverSettings, JointSettings));
+		RotationConstraints.SettingsSoftDamping = FVec3(FPBDJointUtilities::GetSoftTwistDamping(SolverSettings, JointSettings), FPBDJointUtilities::GetSoftSwingDamping(SolverSettings, JointSettings), FPBDJointUtilities::GetSoftSwingDamping(SolverSettings, JointSettings));
+		RotationConstraints.ConstraintMaxLambda = FVec3::Zero();
+
+		if (bResetLambdas)
 		{
-			InitTwistConstraint(JointSettings, Dt);
+			RotationConstraints.Simd.ConstraintLambda = VectorZeroFloat();
+			RotationConstraints.ConstraintLambdaVelocity = FVec3::Zero();
 		}
-	}
 
-	// Apply swing constraints
-	// NOTE: Cannot separate swing angles at 180degree swing (but we can still apply locks)
-	if (SolverSettings.bEnableSwingLimits)
+		RotationConstraints.Simd.ConstraintHardStiffnessSimd = MakeVectorRegisterFloatFromDouble(MakeVectorRegister(FPBDJointUtilities::GetTwistStiffness(SolverSettings, JointSettings),
+			FPBDJointUtilities::GetSwingStiffness(SolverSettings, JointSettings), FPBDJointUtilities::GetSwingStiffness(SolverSettings, JointSettings), 0.0f));
+		InitRotationConstraintsSimd(JointSettings, Dt);
+	}
+	else
 	{
-		if (bAngularLimited[S1] && bAngularLimited[S2])
-		{
-			// When using non linear solver, the cone swing direction could change at each iteration
-			// stabilizing the solver. In the linear case we need to constraint along the 2 directions
-			// for better stability
-			InitPyramidSwingConstraint(JointSettings, Dt, true, true);
-		}
-		else if (bAngularLimited[S1] && bAngularLocked[S2])
+		const int32 TW = (int32)EJointAngularConstraintIndex::Twist;
+		const int32 S1 = (int32)EJointAngularConstraintIndex::Swing1;
+		const int32 S2 = (int32)EJointAngularConstraintIndex::Swing2;
+
+		RotationConstraints.InitDatas(TW, FPBDJointUtilities::GetSoftTwistLimitEnabled(SolverSettings, JointSettings) && !bAngularLocked[TW],
+			FPBDJointUtilities::GetSoftTwistStiffness(SolverSettings, JointSettings),
+			FPBDJointUtilities::GetSoftTwistDamping(SolverSettings, JointSettings),
+			FPBDJointUtilities::GetTwistStiffness(SolverSettings, JointSettings),
+			bResetLambdas);
+
+		RotationConstraints.InitDatas(S1, FPBDJointUtilities::GetSoftSwingLimitEnabled(SolverSettings, JointSettings) && !bAngularLocked[S1],
+			FPBDJointUtilities::GetSoftSwingStiffness(SolverSettings, JointSettings),
+			FPBDJointUtilities::GetSoftSwingDamping(SolverSettings, JointSettings),
+			FPBDJointUtilities::GetSwingStiffness(SolverSettings, JointSettings),
+			bResetLambdas);
+
+		RotationConstraints.InitDatas(S2, FPBDJointUtilities::GetSoftSwingLimitEnabled(SolverSettings, JointSettings) && !bAngularLocked[S2],
+			FPBDJointUtilities::GetSoftSwingStiffness(SolverSettings, JointSettings),
+			FPBDJointUtilities::GetSoftSwingDamping(SolverSettings, JointSettings),
+			FPBDJointUtilities::GetSwingStiffness(SolverSettings, JointSettings),
+			bResetLambdas);
+
+
+
+		const FVec3 Twist0 = ConnectorRs[0] * FJointConstants::TwistAxis();
+		const FVec3 Twist1 = ConnectorRs[1] * FJointConstants::TwistAxis();
+		const bool bDegenerate = (FVec3::DotProduct(Twist0, Twist1) < Chaos_Joint_DegenerateRotationLimit);
+
+		// Apply twist constraint
+		// NOTE: Cannot calculate twist angle at 180degree swing
+		if (SolverSettings.bEnableTwistLimits)
 		{
-			if (!bDegenerate)
+			if (bAngularLimited[TW] && !bDegenerate)
 			{
-				InitPyramidSwingConstraint(JointSettings, Dt, true, false);
+				InitTwistConstraint(JointSettings, Dt);
 			}
 		}
-		else if (bAngularLimited[S1] && bAngularFree[S2])
+
+		// Apply swing constraints
+		// NOTE: Cannot separate swing angles at 180degree swing (but we can still apply locks)
+		if (SolverSettings.bEnableSwingLimits)
 		{
-			if (!bDegenerate)
+			if (bAngularLimited[S1] && bAngularLimited[S2])
 			{
-				InitDualConeSwingConstraint(JointSettings, Dt, EJointAngularConstraintIndex::Swing1);
+				// When using non linear solver, the cone swing direction could change at each iteration
+				// stabilizing the solver. In the linear case we need to constraint along the 2 directions
+				// for better stability
+				InitPyramidSwingConstraint(JointSettings, Dt, true, true);
 			}
-		}
-		else if (bAngularLocked[S1] && bAngularLimited[S2])
-		{
-			if (!bDegenerate)
+			else if (bAngularLimited[S1] && bAngularLocked[S2])
+			{
+				if (!bDegenerate)
+				{
+					InitPyramidSwingConstraint(JointSettings, Dt, true, false);
+				}
+			}
+			else if (bAngularLimited[S1] && bAngularFree[S2])
 			{
-				InitPyramidSwingConstraint(JointSettings, Dt, false, true);
+				if (!bDegenerate)
+				{
+					InitDualConeSwingConstraint(JointSettings, Dt, EJointAngularConstraintIndex::Swing1);
+				}
 			}
-		}
-		else if (bAngularFree[S1] && bAngularLimited[S2])
-		{
-			if (!bDegenerate)
+			else if (bAngularLocked[S1] && bAngularLimited[S2])
+			{
+				if (!bDegenerate)
+				{
+					InitPyramidSwingConstraint(JointSettings, Dt, false, true);
+				}
+			}
+			else if (bAngularFree[S1] && bAngularLimited[S2])
 			{
-				InitDualConeSwingConstraint(JointSettings, Dt, EJointAngularConstraintIndex::Swing2);
+				if (!bDegenerate)
+				{
+					InitDualConeSwingConstraint(JointSettings, Dt, EJointAngularConstraintIndex::Swing2);
+				}
 			}
 		}
+
+
+		// Note: single-swing locks are already handled above so we only need to do something here if both are locked
+		const bool bLockedTwist = SolverSettings.bEnableTwistLimits && bAngularLocked[TW];
+		const bool bLockedSwing1 = SolverSettings.bEnableSwingLimits && bAngularLocked[S1];
+		const bool bLockedSwing2 = SolverSettings.bEnableSwingLimits && bAngularLocked[S2];
+		if (bLockedTwist || bLockedSwing1 || bLockedSwing2)
+		{
+			InitLockedRotationConstraints(JointSettings, Dt, bLockedTwist, bLockedSwing1, bLockedSwing2);
+		}
 	}
+	// Todo at this point Motion Type has never been initialized so it will be always Free here. 
+	// This cause restitution to be always disabled and maybe more... 
+}
+
+void FPBDJointCachedSolver::InitRotationConstraintsSimd(
+	const FPBDJointSettings& JointSettings,
+	const FReal Dt)
+{
+	const FVec3 Twist0 = ConnectorRs[0] * FJointConstants::TwistAxis();
+	const FVec3 Twist1 = ConnectorRs[1] * FJointConstants::TwistAxis();
+	const bool bDegenerate = (FVec3::DotProduct(Twist0, Twist1) < Chaos_Joint_DegenerateRotationLimit);
+
+	FVec3 Axes[3];
+	FReal Angles[3];
+
+	FPBDJointUtilities::GetTwistAxisAngle(ConnectorRs[0], ConnectorRs[1], Axes[0], Angles[0]);
+
+	// Project the angle directly to avoid checking the limits during the solve.
+
+	// Decompose rotation of body 1 relative to body 0 into swing and twist rotations, assuming twist is X axis
+	FRotation3 R01Twist, R01Swing;
+	FPBDJointUtilities::DecomposeSwingTwistLocal(ConnectorRs[0], ConnectorRs[1], R01Swing, R01Twist);
+	const FRotation3 R0Swing = ConnectorRs[0] * R01Swing;
+	Axes[2] = R0Swing * FJointConstants::Swing1Axis();
+	Angles[2] = 4.0 * FMath::Atan2(R01Swing.Z, (FReal)(1. + R01Swing.W));
+	Axes[1] = R0Swing * FJointConstants::Swing2Axis();
+	Angles[1] = 4.0 * FMath::Atan2(R01Swing.Y, (FReal)(1. + R01Swing.W));
+
+	RotationConstraints.ConstraintRestitution = FVec3(JointSettings.TwistRestitution, JointSettings.SwingRestitution, JointSettings.SwingRestitution);
+
+	FVec3f ConstraintHardIM;
+	FVec3f ConstraintSoftStiffness;
+	FVec3f ConstraintSoftDamping;
+	FVec3f ConstraintSoftIM;
+	FVec3f LocalAngles;
 
-	// Note: single-swing locks are already handled above so we only need to do something here if both are locked
-	const bool bLockedTwist = SolverSettings.bEnableTwistLimits && bAngularLocked[TW];
-	const bool bLockedSwing1 = SolverSettings.bEnableSwingLimits && bAngularLocked[S1];
-	const bool bLockedSwing2 = SolverSettings.bEnableSwingLimits && bAngularLocked[S2];
-	if (bLockedTwist || bLockedSwing1 || bLockedSwing2)
+	for (int32 ConstraintIndex = 0; ConstraintIndex < 3; ConstraintIndex++)
 	{
-		InitLockedRotationConstraints(JointSettings, Dt, bLockedTwist, bLockedSwing1, bLockedSwing2);
+		const FVec3 LocalAxis = (Angles[ConstraintIndex] < 0.0f) ? -Axes[ConstraintIndex] : Axes[ConstraintIndex];
+		LocalAngles[ConstraintIndex] = FRealSingle((Angles[ConstraintIndex] < 0.0f) ? -Angles[ConstraintIndex] : Angles[ConstraintIndex]);
+
+		RotationConstraints.bValidDatas[ConstraintIndex] = true;
+		RotationConstraints.bLimitsCheck[ConstraintIndex] = true;
+		RotationConstraints.ConstraintVX[ConstraintIndex] = 0;
+
+		RotationConstraints.Simd.ConstraintArmsSimd[ConstraintIndex][0] = VectorZeroFloat();
+		RotationConstraints.Simd.ConstraintArmsSimd[ConstraintIndex][1] = VectorZeroFloat();
+
+		InitConstraintAxisAngularVelocities[ConstraintIndex] = FVec3::DotProduct(W(1) - W(0), LocalAxis);
+
+		// InitRotationDatasMass(RotationConstraints, ConstraintIndex, Dt);
+
+		RotationConstraints.Simd.ConstraintAxisSimd[ConstraintIndex] = MakeVectorRegisterFloatFromDouble(MakeVectorRegister(LocalAxis[0], LocalAxis[1], LocalAxis[2], 0.0f));
+		const VectorRegister4Float& Axis = RotationConstraints.Simd.ConstraintAxisSimd[ConstraintIndex];
+
+		const VectorRegister4Float AxisX = VectorReplicate(Axis, 0);
+		const VectorRegister4Float AxisY = VectorReplicate(Axis, 1);
+		const VectorRegister4Float AxisZ = VectorReplicate(Axis, 2);
+		const FMatrix33& InvI0 = InvI(0);
+		const VectorRegister4Float InvI00 = MakeVectorRegisterFloatFromDouble(MakeVectorRegister(InvI0.M[0][0], InvI0.M[0][1], InvI0.M[0][2], 0.0f));
+		const VectorRegister4Float InvI01 = MakeVectorRegisterFloatFromDouble(MakeVectorRegister(InvI0.M[1][0], InvI0.M[1][1], InvI0.M[1][2], 0.0f));
+		const VectorRegister4Float InvI02 = MakeVectorRegisterFloatFromDouble(MakeVectorRegister(InvI0.M[2][0], InvI0.M[2][1], InvI0.M[2][2], 0.0f));
+		const VectorRegister4Float IA0 = VectorMultiplyAdd(InvI00, AxisX, VectorMultiplyAdd(InvI01, AxisY, VectorMultiply(InvI02, AxisZ)));
+
+		const FMatrix33& InvI1 = InvI(1);
+		const VectorRegister4Float InvI10 = MakeVectorRegisterFloatFromDouble(MakeVectorRegister(InvI1.M[0][0], InvI1.M[0][1], InvI1.M[0][2], 0.0f));
+		const VectorRegister4Float InvI11 = MakeVectorRegisterFloatFromDouble(MakeVectorRegister(InvI1.M[1][0], InvI1.M[1][1], InvI1.M[1][2], 0.0f));
+		const VectorRegister4Float InvI12 = MakeVectorRegisterFloatFromDouble(MakeVectorRegister(InvI1.M[2][0], InvI1.M[2][1], InvI1.M[2][2], 0.0f));
+		const VectorRegister4Float IA1 = VectorMultiplyAdd(InvI10, AxisX, VectorMultiplyAdd(InvI11, AxisY, VectorMultiply(InvI12, AxisZ)));
+
+		const FRealSingle II0 = VectorDot3Scalar(Axis, IA0);
+		const FRealSingle II1 = VectorDot3Scalar(Axis, IA1);
+		RotationConstraints.Simd.ConstraintDRAxisSimd[ConstraintIndex][0] = IA0;
+		RotationConstraints.Simd.ConstraintDRAxisSimd[ConstraintIndex][1] = VectorNegate(IA1);
+
+		ConstraintHardIM[ConstraintIndex] = II0 + II1;
+		const FReal SpringMassScale = (bUsePositionBasedDrives) ? (FReal)1 / (ConstraintHardIM[ConstraintIndex]) : (FReal)1;
+		ConstraintSoftStiffness[ConstraintIndex] = FRealSingle(SpringMassScale * RotationConstraints.SettingsSoftStiffness[ConstraintIndex] * Dt * Dt);
+		ConstraintSoftDamping[ConstraintIndex] = (bUsePositionBasedDrives) ? FRealSingle(SpringMassScale * RotationConstraints.SettingsSoftDamping[ConstraintIndex] * Dt) : 0.0f;
+		ConstraintSoftIM[ConstraintIndex] = FRealSingle((ConstraintSoftStiffness[ConstraintIndex] + ConstraintSoftDamping[ConstraintIndex]) * ConstraintHardIM[ConstraintIndex] + (FReal)1);
+
 	}
+
+	RotationConstraints.Simd.ConstraintCX = MakeVectorRegisterFloat(LocalAngles[0], LocalAngles[1], LocalAngles[2], 0.0f);
+	RotationConstraints.Simd.ConstraintHardIMSimd = MakeVectorRegisterFloat(ConstraintHardIM[0], ConstraintHardIM[1], ConstraintHardIM[2], 0.0f);
+	RotationConstraints.Simd.ConstraintSoftStiffnessSimd = MakeVectorRegisterFloat(ConstraintSoftStiffness[0], ConstraintSoftStiffness[1], ConstraintSoftStiffness[2], 0.0f);
+	RotationConstraints.Simd.ConstraintSoftDampingSimd = MakeVectorRegisterFloat(ConstraintSoftDamping[0], ConstraintSoftDamping[1], ConstraintSoftDamping[2], 0.0f);
+	RotationConstraints.Simd.ConstraintSoftIMSimd = MakeVectorRegisterFloat(ConstraintSoftIM[0], ConstraintSoftIM[1], ConstraintSoftIM[2], 0.0f);
+	RotationConstraints.Simd.ConstraintLimitsSimd = MakeVectorRegisterFloatFromDouble(MakeVectorRegister(JointSettings.AngularLimits[0], JointSettings.AngularLimits[1], JointSettings.AngularLimits[2], UE_BIG_NUMBER));
 }
 
 void FPBDJointCachedSolver::InitRotationDatasMass(
-		FAxisConstraintDatas& RotationDatas,
-		const int32 ConstraintIndex,
-		const FReal Dt)
+	FAxisConstraintDatas& RotationDatas,
+	const int32 ConstraintIndex,
+	const FReal Dt)
 {
-	const FVec3 IA0 = Utilities::Multiply(InvI(0), RotationDatas.ConstraintAxis[ConstraintIndex]);
-	const FVec3 IA1 = Utilities::Multiply(InvI(1), RotationDatas.ConstraintAxis[ConstraintIndex]);
-	const FReal II0 = FVec3::DotProduct(RotationDatas.ConstraintAxis[ConstraintIndex], IA0);
-	const FReal II1 = FVec3::DotProduct(RotationDatas.ConstraintAxis[ConstraintIndex], IA1);
+	const FVec3 IA0 = Utilities::Multiply(InvI(0), RotationDatas.D.ConstraintAxis[ConstraintIndex]);
+	const FVec3 IA1 = Utilities::Multiply(InvI(1), RotationDatas.D.ConstraintAxis[ConstraintIndex]);
+	const FReal II0 = FVec3::DotProduct(RotationDatas.D.ConstraintAxis[ConstraintIndex], IA0);
+	const FReal II1 = FVec3::DotProduct(RotationDatas.D.ConstraintAxis[ConstraintIndex], IA1);
 
-	RotationDatas.UpdateMass(ConstraintIndex,  IA0, IA1, II0 + II1, Dt, bUsePositionBasedDrives);
+	RotationDatas.UpdateMass(ConstraintIndex, IA0, IA1, II0 + II1, Dt, bUsePositionBasedDrives);
 }
 
 void FPBDJointCachedSolver::InitRotationConstraintDatas(
-		const FPBDJointSettings& JointSettings,
-		const int32 ConstraintIndex,
-		const FVec3& ConstraintAxis,
-		const FReal ConstraintAngle,
-		const FReal ConstraintRestitution,
-		const FReal Dt,
-		const bool bCheckLimit)
+	const FPBDJointSettings& JointSettings,
+	const int32 ConstraintIndex,
+	const FVec3& ConstraintAxis,
+	const FReal ConstraintAngle,
+	const FReal ConstraintRestitution,
+	const FReal Dt,
+	const bool bCheckLimit)
 {
 	const FVec3 LocalAxis = (ConstraintAngle < 0.0f) ? -ConstraintAxis : ConstraintAxis;
 	const FReal LocalAngle = (ConstraintAngle < 0.0f) ? -ConstraintAngle : ConstraintAngle;
 
 	RotationConstraints.UpdateDatas(ConstraintIndex, LocalAxis, LocalAngle, ConstraintRestitution, bCheckLimit);
 
-	RotationConstraints.ConstraintLimits[ConstraintIndex] = JointSettings.AngularLimits[ConstraintIndex];
+
+	RotationConstraints.D.ConstraintLimits[ConstraintIndex] = JointSettings.AngularLimits[ConstraintIndex];
 
 	InitConstraintAxisAngularVelocities[ConstraintIndex] = FVec3::DotProduct(W(1) - W(0), LocalAxis);
 
@@ -968,10 +1372,10 @@ void FPBDJointCachedSolver::InitRotationConstraintDatas(
 }
 
 void FPBDJointCachedSolver::CorrectAxisAngleConstraint(
-		const FPBDJointSettings& JointSettings,
-		const int32 ConstraintIndex,
-		FVec3& ConstraintAxis,
-		FReal& ConstraintAngle) const
+	const FPBDJointSettings& JointSettings,
+	const int32 ConstraintIndex,
+	FVec3& ConstraintAxis,
+	FReal& ConstraintAngle) const
 {
 	const FReal AngleMax = JointSettings.AngularLimits[ConstraintIndex];
 
@@ -992,22 +1396,22 @@ void FPBDJointCachedSolver::CorrectAxisAngleConstraint(
 }
 
 void FPBDJointCachedSolver::InitTwistConstraint(
-		const FPBDJointSettings& JointSettings,
-		const FReal Dt)
+	const FPBDJointSettings& JointSettings,
+	const FReal Dt)
 {
-	 FVec3 TwistAxis;
-	 FReal TwistAngle;
-	 FPBDJointUtilities::GetTwistAxisAngle(ConnectorRs[0], ConnectorRs[1], TwistAxis, TwistAngle);
-	
-	 // Project the angle directly to avoid checking the limits during the solve.
-	 InitRotationConstraintDatas( JointSettings, (int32)EJointAngularConstraintIndex::Twist, TwistAxis, TwistAngle, JointSettings.TwistRestitution, Dt, true);
+	FVec3 TwistAxis;
+	FReal TwistAngle;
+	FPBDJointUtilities::GetTwistAxisAngle(ConnectorRs[0], ConnectorRs[1], TwistAxis, TwistAngle);
+
+	// Project the angle directly to avoid checking the limits during the solve.
+	InitRotationConstraintDatas(JointSettings, (int32)EJointAngularConstraintIndex::Twist, TwistAxis, TwistAngle, JointSettings.TwistRestitution, Dt, true);
 }
 
 void FPBDJointCachedSolver::InitPyramidSwingConstraint(
-   const FPBDJointSettings& JointSettings,
-   const FReal Dt,
-   const bool bApplySwing1,
-   const bool bApplySwing2)
+	const FPBDJointSettings& JointSettings,
+	const FReal Dt,
+	const bool bApplySwing1,
+	const bool bApplySwing2)
 {
 	// Decompose rotation of body 1 relative to body 0 into swing and twist rotations, assuming twist is X axis
 	FRotation3 R01Twist, R01Swing;
@@ -1015,32 +1419,32 @@ void FPBDJointCachedSolver::InitPyramidSwingConstraint(
 
 	const FRotation3 R0Swing = ConnectorRs[0] * R01Swing;
 
-	if(bApplySwing1)
+	if (bApplySwing1)
 	{
 		const FVec3 SwingAxis = R0Swing * FJointConstants::Swing1Axis();
 		const FReal SwingAngle = 4.0 * FMath::Atan2(R01Swing.Z, (FReal)(1. + R01Swing.W));
-		InitRotationConstraintDatas( JointSettings, (int32)EJointAngularConstraintIndex::Swing1, SwingAxis, SwingAngle, JointSettings.SwingRestitution, Dt, true);
+		InitRotationConstraintDatas(JointSettings, (int32)EJointAngularConstraintIndex::Swing1, SwingAxis, SwingAngle, JointSettings.SwingRestitution, Dt, true);
 	}
-	if(bApplySwing2)
+	if (bApplySwing2)
 	{
 		const FVec3 SwingAxis = R0Swing * FJointConstants::Swing2Axis();
 		const FReal SwingAngle = 4.0 * FMath::Atan2(R01Swing.Y, (FReal)(1. + R01Swing.W));
-		InitRotationConstraintDatas( JointSettings, (int32)EJointAngularConstraintIndex::Swing2, SwingAxis, SwingAngle, JointSettings.SwingRestitution, Dt, true);
+		InitRotationConstraintDatas(JointSettings, (int32)EJointAngularConstraintIndex::Swing2, SwingAxis, SwingAngle, JointSettings.SwingRestitution, Dt, true);
 	}
 }
 
 void FPBDJointCachedSolver::InitConeConstraint(
-   const FPBDJointSettings& JointSettings,
-   const FReal Dt)
+	const FPBDJointSettings& JointSettings,
+	const FReal Dt)
 {
 	FVec3 SwingAxisLocal;
 	FReal SwingAngle = 0.0f;
 
 	FPBDJointUtilities::GetEllipticalConeAxisErrorLocal(ConnectorRs[0], ConnectorRs[1], 0.0, 0.0, SwingAxisLocal, SwingAngle);
 	SwingAxisLocal.SafeNormalize();
-	
+
 	const FVec3 SwingAxis = ConnectorRs[0] * SwingAxisLocal;
-	InitRotationConstraintDatas( JointSettings, (int32)EJointAngularConstraintIndex::Swing2, SwingAxis, SwingAngle, JointSettings.SwingRestitution, Dt, true);
+	InitRotationConstraintDatas(JointSettings, (int32)EJointAngularConstraintIndex::Swing2, SwingAxis, SwingAngle, JointSettings.SwingRestitution, Dt, true);
 }
 
 void FPBDJointCachedSolver::InitSingleLockedSwingConstraint(
@@ -1050,12 +1454,12 @@ void FPBDJointCachedSolver::InitSingleLockedSwingConstraint(
 {
 	//NOTE: SwingAxis is not normalized in this mode. It has length Sin(SwingAngle).
 	//Likewise, the SwingAngle is actually Sin(SwingAngle)
-    // FVec3 SwingAxis;
-    // FReal SwingAngle;
-    // FPBDJointUtilities::GetLockedSwingAxisAngle(ConnectorRs[0], ConnectorRs[1], SwingConstraintIndex, SwingAxis, SwingAngle);
-    //SwingAxis.SafeNormalize();
+	// FVec3 SwingAxis;
+	// FReal SwingAngle;
+	// FPBDJointUtilities::GetLockedSwingAxisAngle(ConnectorRs[0], ConnectorRs[1], SwingConstraintIndex, SwingAxis, SwingAngle);
+	//SwingAxis.SafeNormalize();
 
-    // Using the locked swing axis angle results in potential axis switching since this axis is the result of OtherSwing x TwistAxis
+	// Using the locked swing axis angle results in potential axis switching since this axis is the result of OtherSwing x TwistAxis
 	FVec3 SwingAxis;
 	FReal SwingAngle;
 	FPBDJointUtilities::GetSwingAxisAngle(ConnectorRs[0], ConnectorRs[1], 0.0, SwingConstraintIndex, SwingAxis, SwingAngle);
@@ -1068,7 +1472,7 @@ void FPBDJointCachedSolver::InitDualConeSwingConstraint(
 	const FPBDJointSettings& JointSettings,
 	const FReal Dt,
 	const EJointAngularConstraintIndex SwingConstraintIndex)
-	
+
 {
 	FVec3 SwingAxis;
 	FReal SwingAngle;
@@ -1088,15 +1492,15 @@ void FPBDJointCachedSolver::InitSwingConstraint(
 	FReal SwingAngle;
 	FPBDJointUtilities::GetSwingAxisAngle(ConnectorRs[0], ConnectorRs[1], SolverSettings.SwingTwistAngleTolerance, SwingConstraintIndex, SwingAxis, SwingAngle);
 
-	InitRotationConstraintDatas( JointSettings, (int32)SwingConstraintIndex, SwingAxis, SwingAngle, JointSettings.SwingRestitution, Dt, true);
+	InitRotationConstraintDatas(JointSettings, (int32)SwingConstraintIndex, SwingAxis, SwingAngle, JointSettings.SwingRestitution, Dt, true);
 }
 
 void FPBDJointCachedSolver::InitLockedRotationConstraints(
-	  const FPBDJointSettings& JointSettings,
-	  const FReal Dt,
-	  const bool bApplyTwist,
-	  const bool bApplySwing1,
-	  const bool bApplySwing2)
+	const FPBDJointSettings& JointSettings,
+	const FReal Dt,
+	const bool bApplyTwist,
+	const bool bApplySwing1,
+	const bool bApplySwing2)
 {
 	FVec3 Axis0, Axis1, Axis2;
 	FPBDJointUtilities::GetLockedRotationAxes(ConnectorRs[0], ConnectorRs[1], Axis0, Axis1, Axis2);
@@ -1124,95 +1528,102 @@ void FPBDJointCachedSolver::InitLockedRotationConstraints(
 void FPBDJointCachedSolver::ApplyRotationConstraints(
 	const FReal Dt)
 {
-	for(int32 ConstraintIndex = 0; ConstraintIndex < 3; ++ConstraintIndex)
+	if (RotationConstraints.bUseSimd)
+	{
+		ApplyRotationSoftConstraintsSimd(Dt);
+	}
+	else
 	{
-		if(RotationConstraints.bValidDatas[ConstraintIndex])
+		for (int32 ConstraintIndex = 0; ConstraintIndex < 3; ++ConstraintIndex)
 		{
-			ApplyRotationConstraint(ConstraintIndex, Dt);
+			if (RotationConstraints.bValidDatas[ConstraintIndex])
+			{
+				ApplyRotationConstraint(ConstraintIndex, Dt);
+			}
 		}
 	}
 }
 
 void FPBDJointCachedSolver::SolveRotationConstraintDelta(
-		const int32 ConstraintIndex, 
-		const FReal DeltaLambda,
-		const bool bIsSoftConstraint,
-		const FAxisConstraintDatas& ConstraintDatas)
+	const int32 ConstraintIndex,
+	const FReal DeltaLambda,
+	const bool bIsSoftConstraint,
+	const FAxisConstraintDatas& ConstraintDatas)
 {
-	const FVec3 DeltaImpulse =  ConstraintDatas.ConstraintAxis[ConstraintIndex]  * DeltaLambda;
-	if(Body(0).IsDynamic())
+	const FVec3 DeltaImpulse = ConstraintDatas.D.ConstraintAxis[ConstraintIndex] * DeltaLambda;
+	if (Body(0).IsDynamic())
 	{
-		const FVec3 DR0 = !bIsSoftConstraint ? ConstraintDatas.ConstraintDRAxis[ConstraintIndex][0] * DeltaLambda :
-			DeltaImpulse *  (FVec3::DotProduct(ConstraintDatas.ConstraintAxis[ConstraintIndex], ConstraintDatas.ConstraintDRAxis[ConstraintIndex][0]));
+		const FVec3 DR0 = !bIsSoftConstraint ? ConstraintDatas.D.ConstraintDRAxis[ConstraintIndex][0] * DeltaLambda :
+			DeltaImpulse * (FVec3::DotProduct(ConstraintDatas.D.ConstraintAxis[ConstraintIndex], ConstraintDatas.D.ConstraintDRAxis[ConstraintIndex][0]));
 		ApplyRotationDelta(0, DR0);
 	}
-	if(Body(1).IsDynamic())
+	if (Body(1).IsDynamic())
 	{
-		const FVec3 DR1 = !bIsSoftConstraint ? ConstraintDatas.ConstraintDRAxis[ConstraintIndex][1] * DeltaLambda :
-			DeltaImpulse * ( FVec3::DotProduct(ConstraintDatas.ConstraintAxis[ConstraintIndex], ConstraintDatas.ConstraintDRAxis[ConstraintIndex][1]));
+		const FVec3 DR1 = !bIsSoftConstraint ? ConstraintDatas.D.ConstraintDRAxis[ConstraintIndex][1] * DeltaLambda :
+			DeltaImpulse * (FVec3::DotProduct(ConstraintDatas.D.ConstraintAxis[ConstraintIndex], ConstraintDatas.D.ConstraintDRAxis[ConstraintIndex][1]));
 		ApplyRotationDelta(1, DR1);
 	}
 	++NumActiveConstraints;
 }
 
 void FPBDJointCachedSolver::SolveRotationConstraintHard(
-			const int32 ConstraintIndex,
-			const FReal DeltaConstraint)
+	const int32 ConstraintIndex,
+	const FReal DeltaConstraint)
 {
-	const FReal DeltaLambda = SolverStiffness * RotationConstraints.ConstraintHardStiffness[ConstraintIndex] * DeltaConstraint /
-		RotationConstraints.ConstraintHardIM[ConstraintIndex];
+	const FReal DeltaLambda = SolverStiffness * RotationConstraints.D.ConstraintHardStiffness[ConstraintIndex] * DeltaConstraint /
+		RotationConstraints.D.ConstraintHardIM[ConstraintIndex];
 
-	RotationConstraints.ConstraintLambda[ConstraintIndex] += DeltaLambda;
+	RotationConstraints.D.ConstraintLambda[ConstraintIndex] += DeltaLambda;
 	SolveRotationConstraintDelta(ConstraintIndex, DeltaLambda, false, RotationConstraints);
 }
 
 void FPBDJointCachedSolver::SolveRotationConstraintSoft(
-			const int32 ConstraintIndex,
-			const FReal DeltaConstraint,
-			const FReal Dt,
-			const FReal TargetVel)
+	const int32 ConstraintIndex,
+	const FReal DeltaConstraint,
+	const FReal Dt,
+	const FReal TargetVel)
 {
 	// Damping angular velocity
 	FReal AngVelDt = 0;
-	if (RotationConstraints.ConstraintSoftDamping[ConstraintIndex] > UE_KINDA_SMALL_NUMBER)
+	if (RotationConstraints.D.ConstraintSoftDamping[ConstraintIndex] > UE_KINDA_SMALL_NUMBER)
 	{
 		const FVec3 W0Dt = FVec3(Body(0).DQ()) + ConnectorWDts[0];
 		const FVec3 W1Dt = FVec3(Body(1).DQ()) + ConnectorWDts[1];
-		AngVelDt = TargetVel * Dt + FVec3::DotProduct(RotationConstraints.ConstraintAxis[ConstraintIndex] , W0Dt - W1Dt);
+		AngVelDt = TargetVel * Dt + FVec3::DotProduct(RotationConstraints.D.ConstraintAxis[ConstraintIndex], W0Dt - W1Dt);
 	}
 
-	const FReal DeltaLambda = SolverStiffness * (RotationConstraints.ConstraintSoftStiffness[ConstraintIndex] * DeltaConstraint -
-		RotationConstraints.ConstraintSoftDamping[ConstraintIndex] * AngVelDt - RotationConstraints.ConstraintLambda[ConstraintIndex]) /
-		RotationConstraints.ConstraintSoftIM[ConstraintIndex];
-	RotationConstraints.ConstraintLambda[ConstraintIndex] += DeltaLambda;
+	const FReal DeltaLambda = SolverStiffness * (RotationConstraints.D.ConstraintSoftStiffness[ConstraintIndex] * DeltaConstraint -
+		RotationConstraints.D.ConstraintSoftDamping[ConstraintIndex] * AngVelDt - RotationConstraints.D.ConstraintLambda[ConstraintIndex]) /
+		RotationConstraints.D.ConstraintSoftIM[ConstraintIndex];
+	RotationConstraints.D.ConstraintLambda[ConstraintIndex] += DeltaLambda;
 
-	SolveRotationConstraintDelta(ConstraintIndex, DeltaLambda,false, RotationConstraints);
+	SolveRotationConstraintDelta(ConstraintIndex, DeltaLambda, false, RotationConstraints);
 }
 
 void FPBDJointCachedSolver::ApplyRotationConstraint(
 	const int32 ConstraintIndex,
 	const FReal Dt)
 {
-	FReal DeltaAngle = RotationConstraints.ConstraintCX[ConstraintIndex] +
-		FVec3::DotProduct(Body(1).DQ() - Body(0).DQ(), RotationConstraints.ConstraintAxis[ConstraintIndex]);
-	
+	FReal DeltaAngle = RotationConstraints.D.ConstraintCX[ConstraintIndex] +
+		FVec3::DotProduct(Body(1).DQ() - Body(0).DQ(), RotationConstraints.D.ConstraintAxis[ConstraintIndex]);
+
 	bool NeedsSolve = false;
-	if(RotationConstraints.bLimitsCheck[ConstraintIndex])
+	if (RotationConstraints.bLimitsCheck[ConstraintIndex])
 	{
-		if(DeltaAngle > RotationConstraints.ConstraintLimits[ConstraintIndex] )
+		if (DeltaAngle > RotationConstraints.D.ConstraintLimits[ConstraintIndex])
 		{
-			DeltaAngle -= RotationConstraints.ConstraintLimits[ConstraintIndex];
+			DeltaAngle -= RotationConstraints.D.ConstraintLimits[ConstraintIndex];
 			NeedsSolve = true;
 		}
-		else if(DeltaAngle < -RotationConstraints.ConstraintLimits[ConstraintIndex])
+		else if (DeltaAngle < -RotationConstraints.D.ConstraintLimits[ConstraintIndex])
 		{
-			DeltaAngle += RotationConstraints.ConstraintLimits[ConstraintIndex];
+			DeltaAngle += RotationConstraints.D.ConstraintLimits[ConstraintIndex];
 			NeedsSolve = true;
-			
+
 		}
 	}
 
-	if (!RotationConstraints.bLimitsCheck[ConstraintIndex]|| (RotationConstraints.bLimitsCheck[ConstraintIndex] && NeedsSolve && FMath::Abs(DeltaAngle) > AngleTolerance))
+	if (!RotationConstraints.bLimitsCheck[ConstraintIndex] || (RotationConstraints.bLimitsCheck[ConstraintIndex] && NeedsSolve && FMath::Abs(DeltaAngle) > AngleTolerance))
 	{
 		if (RotationConstraints.bSoftLimit[ConstraintIndex])
 		{
@@ -1225,15 +1636,101 @@ void FPBDJointCachedSolver::ApplyRotationConstraint(
 	}
 }
 
+void FPBDJointCachedSolver::ApplyRotationSoftConstraintsSimd(
+	const FReal Dt)
+{
+	const VectorRegister4Float Body0DQ = MakeVectorRegisterFloatFromDouble(MakeVectorRegister(Body(0).DQ()[0], Body(0).DQ()[1], Body(0).DQ()[2], 0.0f));
+	const VectorRegister4Float Body1DQ = MakeVectorRegisterFloatFromDouble(MakeVectorRegister(Body(1).DQ()[0], Body(1).DQ()[1], Body(1).DQ()[2], 0.0f));
+
+	const VectorRegister4Float DQDiff = VectorSubtract(Body1DQ, Body0DQ);
+	VectorRegister4Float ProjAxes[3];
+	for (int32 CIndex = 0; CIndex < 3; ++CIndex)
+	{
+		ProjAxes[CIndex] = Private::VectorDot3FastX(DQDiff, RotationConstraints.Simd.ConstraintAxisSimd[CIndex]);
+	}
+	VectorRegister4Float ProjAxis = VectorUnpackLo(ProjAxes[0], ProjAxes[1]);
+	ProjAxis = VectorMoveLh(ProjAxis, ProjAxes[2]);
+	VectorRegister4Float DeltaAngle = VectorAdd(RotationConstraints.Simd.ConstraintCX, ProjAxis);
+
+
+	const VectorRegister4Float AngleGT = VectorCompareGT(DeltaAngle, RotationConstraints.Simd.ConstraintLimitsSimd);
+	const VectorRegister4Float AngleLT = VectorCompareLT(DeltaAngle, VectorNegate(RotationConstraints.Simd.ConstraintLimitsSimd));
+
+	DeltaAngle = VectorSelect(AngleGT, VectorSubtract(DeltaAngle, RotationConstraints.Simd.ConstraintLimitsSimd),
+		VectorSelect(AngleLT, VectorAdd(DeltaAngle, RotationConstraints.Simd.ConstraintLimitsSimd), DeltaAngle));
+
+	FRealSingle AngleTolerancef = FRealSingle(AngleTolerance);
+	const VectorRegister4Float AngleToleranceSimd = VectorLoadFloat1(&AngleTolerancef);
+	const VectorRegister4Float InTolerance = VectorCompareGT(VectorAbs(DeltaAngle), AngleToleranceSimd);
+	const VectorRegister4Float NeedsSolve = VectorBitwiseAnd(VectorBitwiseOr(AngleGT, AngleLT), InTolerance);
+
+	if (VectorMaskBits(NeedsSolve))
+	{
+		VectorRegister4Float AngVelDts[3];
+		VectorRegister4Float WDiff = VectorSubtract(ConnectorWDtsSimd[0], VectorAdd(DQDiff, ConnectorWDtsSimd[1]));
+		for (int32 CIndex = 0; CIndex < 3; ++CIndex)
+		{
+			AngVelDts[CIndex] = Private::VectorDot3FastX(RotationConstraints.Simd.ConstraintAxisSimd[CIndex], WDiff);
+		}
+		VectorRegister4Float AngVelDt = VectorUnpackLo(AngVelDts[0], AngVelDts[1]);
+		AngVelDt = VectorMoveLh(AngVelDt, AngVelDts[2]);
+
+		const FRealSingle SolverStiffnessf = FRealSingle(SolverStiffness);
+		VectorRegister4Float Stiffness = VectorLoadFloat1(&SolverStiffnessf);
+		VectorRegister4Float DeltaLambda = VectorMultiply(Stiffness, VectorDivide(VectorSubtract(VectorMultiply(RotationConstraints.Simd.ConstraintSoftStiffnessSimd, DeltaAngle),
+			VectorSubtract(VectorMultiply(RotationConstraints.Simd.ConstraintSoftDampingSimd, AngVelDt), RotationConstraints.Simd.ConstraintLambda)), RotationConstraints.Simd.ConstraintSoftIMSimd));
+
+		DeltaLambda = VectorSelect(NeedsSolve, DeltaLambda, VectorZeroFloat());
+		RotationConstraints.Simd.ConstraintLambda = VectorAdd(RotationConstraints.Simd.ConstraintLambda, DeltaLambda);
+
+		VectorRegister4Float DeltaLambdas[3];
+		DeltaLambdas[0] = VectorReplicate(DeltaLambda, 0);
+		DeltaLambdas[1] = VectorReplicate(DeltaLambda, 1);
+		DeltaLambdas[2] = VectorReplicate(DeltaLambda, 2);
+
+		if (Body(0).IsDynamic())
+		{
+			VectorRegister4Float DR0 = VectorZeroFloat();
+			for (int32 CIndex = 0; CIndex < 3; ++CIndex)
+			{
+				DR0 = VectorMultiplyAdd(RotationConstraints.Simd.ConstraintDRAxisSimd[CIndex][0], DeltaLambdas[CIndex], DR0);
+			}
+			FVec3f DR0f;
+			VectorStoreFloat3(DR0, &DR0f[0]);
+			ApplyRotationDelta(0, FVec3(DR0f));
+		}
+		if (Body(1).IsDynamic())
+		{
+			VectorRegister4Float DR1 = VectorZeroFloat();
+			for (int32 CIndex = 0; CIndex < 3; ++CIndex)
+			{
+				DR1 = VectorMultiplyAdd(RotationConstraints.Simd.ConstraintDRAxisSimd[CIndex][1], DeltaLambdas[CIndex], DR1);
+			}
+			FVec3f DR1f;
+			VectorStoreFloat3(DR1, &DR1f[0]);
+			ApplyRotationDelta(1, FVec3(DR1f));
+		}
+		NumActiveConstraints += 3;
+	}
+}
+
+
 /** APPLY ANGULAR VELOCITY CONSTRAINT *********************************************************************************/
 
 void FPBDJointCachedSolver::ApplyAngularVelocityConstraints()
 {
-	for(int32 ConstraintIndex = 0; ConstraintIndex < 3; ++ConstraintIndex)
+	if (RotationConstraints.bUseSimd)
 	{
-		if(RotationConstraints.bValidDatas[ConstraintIndex])
+		ApplyAngularVelocityConstraintSimd();
+	}
+	else
+	{
+		for (int32 ConstraintIndex = 0; ConstraintIndex < 3; ++ConstraintIndex)
 		{
-			ApplyAngularVelocityConstraint(ConstraintIndex);
+			if (RotationConstraints.bValidDatas[ConstraintIndex])
+			{
+				ApplyAngularVelocityConstraint(ConstraintIndex);
+			}
 		}
 	}
 }
@@ -1244,22 +1741,22 @@ void FPBDJointCachedSolver::SolveAngularVelocityConstraint(
 {
 	const FVec3 CW = W(1) - W(0);
 
-	const FReal DeltaLambda = SolverStiffness * RotationConstraints.ConstraintHardStiffness[ConstraintIndex] *
-	 (FVec3::DotProduct(CW, RotationConstraints.ConstraintAxis[ConstraintIndex]) - TargetVel) / RotationConstraints.ConstraintHardIM[ConstraintIndex];
+	const FReal DeltaLambda = SolverStiffness * RotationConstraints.D.ConstraintHardStiffness[ConstraintIndex] *
+		(FVec3::DotProduct(CW, RotationConstraints.D.ConstraintAxis[ConstraintIndex]) - TargetVel) / RotationConstraints.D.ConstraintHardIM[ConstraintIndex];
 
-	 // @todo(chaos): we should be adding to the net positional impulse here
-	 // RotationConstraints.ConstraintLambda[ConsraintIndex] += DeltaLambda * Dt;
+	// @todo(chaos): we should be adding to the net positional impulse here
+	// RotationConstraints.D.ConstraintLambda[ConsraintIndex] += DeltaLambda * Dt;
 
-	if(Body(0).IsDynamic())
+	if (Body(0).IsDynamic())
 	{
-		const FVec3 DW0 = RotationConstraints.ConstraintDRAxis[ConstraintIndex][0] * DeltaLambda;
-	
+		const FVec3 DW0 = RotationConstraints.D.ConstraintDRAxis[ConstraintIndex][0] * DeltaLambda;
+
 		Body(0).ApplyAngularVelocityDelta(DW0);
 	}
-	if(Body(1).IsDynamic())
+	if (Body(1).IsDynamic())
 	{
-		const FVec3 DW1 = RotationConstraints.ConstraintDRAxis[ConstraintIndex][1] * DeltaLambda;
-	
+		const FVec3 DW1 = RotationConstraints.D.ConstraintDRAxis[ConstraintIndex][1] * DeltaLambda;
+
 		Body(1).ApplyAngularVelocityDelta(DW1);
 	}
 }
@@ -1269,11 +1766,11 @@ void FPBDJointCachedSolver::ApplyAngularVelocityConstraint(
 {
 	// Apply restitution for limited joints when we have exceeded the limits
 	// We also drive the velocity to zero for locked constraints (ignoring restitution)
-	if (FMath::Abs(RotationConstraints.ConstraintLambda[ConstraintIndex]) > UE_SMALL_NUMBER)
+	if (FMath::Abs(RotationConstraints.D.ConstraintLambda[ConstraintIndex]) > UE_SMALL_NUMBER)
 	{
 		FReal TargetVel = 0.0f;
 		if ((RotationConstraints.MotionType[ConstraintIndex] == EJointMotionType::Limited) && (RotationConstraints.ConstraintRestitution[ConstraintIndex] != 0.0f))
-		{ 
+		{
 			const FReal InitVel = InitConstraintAxisAngularVelocities[ConstraintIndex];
 			TargetVel = InitVel > Chaos_Joint_AngularVelocityThresholdToApplyRestitution ?
 				-RotationConstraints.ConstraintRestitution[ConstraintIndex] * InitVel : 0.0f;
@@ -1282,6 +1779,76 @@ void FPBDJointCachedSolver::ApplyAngularVelocityConstraint(
 	}
 }
 
+void FPBDJointCachedSolver::ApplyAngularVelocityConstraintSimd()
+{
+	/*check(RotationConstraints.MotionType[0] == EJointMotionType::Limited);
+	check(RotationConstraints.MotionType[1] == EJointMotionType::Limited);
+	check(RotationConstraints.MotionType[2] == EJointMotionType::Limited);*/
+
+	const VectorRegister4Float IsGTEps = VectorCompareGT(VectorAbs(RotationConstraints.Simd.ConstraintLambda), GlobalVectorConstants::SmallNumber);
+
+	if (VectorMaskBits(IsGTEps))
+	{
+		const VectorRegister4Float Restitution = MakeVectorRegisterFloatFromDouble(MakeVectorRegister(RotationConstraints.ConstraintRestitution[0], RotationConstraints.ConstraintRestitution[1], RotationConstraints.ConstraintRestitution[2], 0.0f));
+
+		const VectorRegister4Float HasRestitution = VectorCompareNE(Restitution, VectorZeroFloat());
+		VectorRegister4Float TargetVel = VectorZeroFloat();
+		if (VectorMaskBits(HasRestitution))
+		{
+			const VectorRegister4Float InitVel = MakeVectorRegisterFloatFromDouble(MakeVectorRegister(InitConstraintAxisAngularVelocities[0], InitConstraintAxisAngularVelocities[1], InitConstraintAxisAngularVelocities[2], 0.0f));
+			VectorRegister4Float VeclocityThreshold = VectorLoadFloat1(&Chaos_Joint_AngularVelocityThresholdToApplyRestitution);
+			TargetVel = VectorSelect(VectorCompareGT(InitVel, VeclocityThreshold), VectorMultiply(VectorNegate(Restitution), InitVel), VectorZeroFloat());
+		}
+
+		const FVec3& W0d = W(0);
+		const VectorRegister4Float W0 = MakeVectorRegisterFloatFromDouble(MakeVectorRegister(W0d[0], W0d[1], W0d[2], 0.0f));
+		const FVec3& W1d = W(1);
+		const VectorRegister4Float W1 = MakeVectorRegisterFloatFromDouble(MakeVectorRegister(W1d[0], W1d[1], W1d[2], 0.0f));
+		const VectorRegister4Float CW = VectorSubtract(W1, W0);
+		const FRealSingle SolverStiffnessf = FRealSingle(SolverStiffness);
+		VectorRegister4Float Stiffness = VectorLoadFloat1(&SolverStiffnessf);
+		Stiffness = VectorMultiply(Stiffness, RotationConstraints.Simd.ConstraintHardStiffnessSimd);
+
+		VectorRegister4Float DeltaLambdas[3];
+		for (int32 CIndex = 0; CIndex < 3; CIndex++)
+		{
+			DeltaLambdas[CIndex] = Private::VectorDot3FastX(CW, RotationConstraints.Simd.ConstraintAxisSimd[CIndex]);
+		}
+		VectorRegister4Float DeltaLambda = VectorUnpackLo(DeltaLambdas[0], DeltaLambdas[1]);
+		DeltaLambda = VectorMoveLh(DeltaLambda, DeltaLambdas[2]);
+		DeltaLambda = VectorDivide(VectorMultiply(Stiffness, VectorSubtract(DeltaLambda, TargetVel)), RotationConstraints.Simd.ConstraintHardIMSimd);
+
+		DeltaLambda = VectorSelect(IsGTEps, DeltaLambda, VectorZeroFloat());
+		DeltaLambdas[0] = VectorReplicate(DeltaLambda, 0);
+		DeltaLambdas[1] = VectorReplicate(DeltaLambda, 1);
+		DeltaLambdas[2] = VectorReplicate(DeltaLambda, 2);
+
+		if (Body(0).IsDynamic())
+		{
+			VectorRegister4Float DW0 = VectorZeroFloat();
+			for (int32 CIndex = 0; CIndex < 3; CIndex++)
+			{
+				DW0 = VectorMultiplyAdd(RotationConstraints.Simd.ConstraintDRAxisSimd[CIndex][0], DeltaLambdas[CIndex], DW0);
+			}
+
+			FVec3f DW0f;
+			VectorStoreFloat3(DW0, &DW0f[0]);
+			Body(0).ApplyAngularVelocityDelta(FVec3(DW0f));
+		}
+		if (Body(1).IsDynamic())
+		{
+			VectorRegister4Float DW1 = VectorZeroFloat();
+			for (int32 CIndex = 0; CIndex < 3; CIndex++)
+			{
+				DW1 = VectorMultiplyAdd(RotationConstraints.Simd.ConstraintDRAxisSimd[CIndex][1], DeltaLambdas[CIndex], DW1);
+			}
+			FVec3f DW1f;
+			VectorStoreFloat3(DW1, &DW1f[0]);
+			Body(0).ApplyAngularVelocityDelta(FVec3(DW1f));
+		}
+	}
+}
+
 /** INIT POSITION DRIVES *********************************************************************************/
 
 void FPBDJointCachedSolver::InitPositionDrives(
@@ -1292,6 +1859,7 @@ void FPBDJointCachedSolver::InitPositionDrives(
 	PositionDrives.bValidDatas[0] = false;
 	PositionDrives.bValidDatas[1] = false;
 	PositionDrives.bValidDatas[2] = false;
+	PositionDrives.bUseSimd = false;
 
 	if (SolverSettings.bEnableDrives)
 	{
@@ -1301,7 +1869,7 @@ void FPBDJointCachedSolver::InitPositionDrives(
 			(JointSettings.bLinearPositionDriveEnabled[1] || JointSettings.bLinearVelocityDriveEnabled[1]) && (JointSettings.LinearMotionTypes[1] != EJointMotionType::Locked),
 			(JointSettings.bLinearPositionDriveEnabled[2] || JointSettings.bLinearVelocityDriveEnabled[2]) && (JointSettings.LinearMotionTypes[2] != EJointMotionType::Locked),
 		};
-	
+
 		PositionDrives.bAccelerationMode = FPBDJointUtilities::GetLinearDriveAccelerationMode(SolverSettings, JointSettings);
 
 		// Rectangular position drives
@@ -1323,8 +1891,8 @@ void FPBDJointCachedSolver::InitPositionDrives(
 					PositionDrives.InitDatas(AxisIndex, true, FPBDJointUtilities::GetLinearDriveStiffness(SolverSettings, JointSettings, AxisIndex),
 						FPBDJointUtilities::GetLinearDriveDamping(SolverSettings, JointSettings, AxisIndex), 0);
 					const FVec3 Axis = R0M.GetAxis(AxisIndex);
-				
-					if ((FMath::Abs(FVec3::DotProduct(CX,Axis)) > PositionTolerance) || (PositionDrives.ConstraintSoftDamping[AxisIndex] > 0.0f))
+
+					if ((FMath::Abs(FVec3::DotProduct(CX, Axis)) > PositionTolerance) || (PositionDrives.D.ConstraintSoftDamping[AxisIndex] > 0.0f))
 					{
 						InitAxisPositionDrive(AxisIndex, Axis, CX, VTarget, Dt);
 					}
@@ -1337,28 +1905,28 @@ void FPBDJointCachedSolver::InitPositionDrives(
 }
 
 void FPBDJointCachedSolver::InitAxisPositionDrive(
-		const int32 ConstraintIndex,
-		const FVec3& ConstraintAxis,
-		const FVec3& DeltaPosition,
-		const FVec3& DeltaVelocity,
-		const FReal Dt)
-{	
+	const int32 ConstraintIndex,
+	const FVec3& ConstraintAxis,
+	const FVec3& DeltaPosition,
+	const FVec3& DeltaVelocity,
+	const FReal Dt)
+{
 	const FVec3 ConstraintArm0 = ConnectorXs[0] - CurrentPs[0];
 	const FVec3 ConstraintArm1 = ConnectorXs[1] - CurrentPs[1];
 
 	PositionDrives.UpdateDatas(ConstraintIndex, ConstraintAxis, FVec3::DotProduct(DeltaPosition, ConstraintAxis),
-		0.0f,  true, ConstraintArm0, ConstraintArm1,
+		0.0f, true, ConstraintArm0, ConstraintArm1,
 		FVec3::DotProduct(DeltaVelocity, ConstraintAxis));
 
 	InitPositionDatasMass(PositionDrives, ConstraintIndex, Dt);
 }
 /** APPLY POSITION PROJECTIONS *********************************************************************************/
-	
+
 void FPBDJointCachedSolver::ApplyProjections(
-		const FReal Dt,
-		const FPBDJointSolverSettings& SolverSettings,
-		const FPBDJointSettings& JointSettings,
-		const bool bLastIteration)
+	const FReal Dt,
+	const FPBDJointSolverSettings& SolverSettings,
+	const FPBDJointSettings& JointSettings,
+	const bool bLastIteration)
 {
 	if (!JointSettings.bProjectionEnabled)
 	{
@@ -1384,80 +1952,195 @@ void FPBDJointCachedSolver::ApplyProjections(
 		ApplyRotationProjection(Dt, SolverSettings, JointSettings);
 	}
 
-	if(bLastIteration)
+	if (bLastIteration)
 	{
 		// Add velocity correction from the net projection motion
 		// @todo(chaos): this should be a joint setting?
 		if (Chaos_Joint_VelProjectionAlpha > 0.0f)
 		{
-			const FSolverReal VelocityScale = Chaos_Joint_VelProjectionAlpha  / static_cast<FSolverReal>(Dt);
+			const FSolverReal VelocityScale = Chaos_Joint_VelProjectionAlpha / static_cast<FSolverReal>(Dt);
 			const FSolverVec3 DV1 = Body1().DP() * VelocityScale;
-			const FSolverVec3 DW1 = Body1().DQ()* VelocityScale;
-		
+			const FSolverVec3 DW1 = Body1().DQ() * VelocityScale;
+
 			Body(1).ApplyVelocityDelta(DV1, DW1);
 		}
 	}
 }
 
 void FPBDJointCachedSolver::ApplyRotationProjection(
-			const FReal Dt,
-			const FPBDJointSolverSettings& SolverSettings,
-			const FPBDJointSettings& JointSettings)
+	const FReal Dt,
+	const FPBDJointSolverSettings& SolverSettings,
+	const FPBDJointSettings& JointSettings)
 {
 	const FReal AngularProjection = FPBDJointUtilities::GetAngularProjection(SolverSettings, JointSettings);
 	if (AngularProjection == 0)
 	{
 		return;
 	}
-	
 	const TVec3<EJointMotionType>& LinearMotion = JointSettings.LinearMotionTypes;
 	const bool bLinearLocked = (LinearMotion[0] == EJointMotionType::Locked) && (LinearMotion[1] == EJointMotionType::Locked) && (LinearMotion[2] == EJointMotionType::Locked);
+	if (RotationConstraints.bUseSimd)
+	{
+		// TODO Here there is a paradox it can be vectorized only if soft but cannot be projected if soft 
+		ApplyRotationProjectionSimd(FRealSingle(AngularProjection), bLinearLocked);
+	}
+	else
+	{
+		for (int32 ConstraintIndex = 0; ConstraintIndex < 3; ++ConstraintIndex)
+		{
+			ApplyAxisRotationProjection(AngularProjection, bLinearLocked, ConstraintIndex);
+		}
+	}
+}
 
-	for(int32 ConstraintIndex = 0; ConstraintIndex < 3; ++ConstraintIndex)
+void FPBDJointCachedSolver::ApplyAxisRotationProjection(
+	const FReal AngularProjection,
+	const bool bLinearLocked,
+	const int32 ConstraintIndex)
+{
+	check(RotationConstraints.bUseSimd == false);
+	if (RotationConstraints.bValidDatas[ConstraintIndex] && !RotationConstraints.bSoftLimit[ConstraintIndex])
 	{
-		if(RotationConstraints.bValidDatas[ConstraintIndex] && !RotationConstraints.bSoftLimit[ConstraintIndex])
+		FReal DeltaAngle = RotationConstraints.D.ConstraintCX[ConstraintIndex] +
+			FVec3::DotProduct(Body(1).DQ() - Body(0).DQ(), RotationConstraints.D.ConstraintAxis[ConstraintIndex]);
+
+		bool NeedsSolve = false;
+		if (RotationConstraints.bLimitsCheck[ConstraintIndex])
 		{
-			FReal DeltaAngle = RotationConstraints.ConstraintCX[ConstraintIndex] +
-				FVec3::DotProduct(Body(1).DQ() - Body(0).DQ(), RotationConstraints.ConstraintAxis[ConstraintIndex]);
-	
-			bool NeedsSolve = false;
-			if(RotationConstraints.bLimitsCheck[ConstraintIndex])
+			if (DeltaAngle > RotationConstraints.D.ConstraintLimits[ConstraintIndex])
 			{
-				if(DeltaAngle > RotationConstraints.ConstraintLimits[ConstraintIndex] )
+				DeltaAngle -= RotationConstraints.D.ConstraintLimits[ConstraintIndex];
+				NeedsSolve = true;
+			}
+			else if (DeltaAngle < -RotationConstraints.D.ConstraintLimits[ConstraintIndex])
+			{
+				DeltaAngle += RotationConstraints.D.ConstraintLimits[ConstraintIndex];
+				NeedsSolve = true;
+
+			}
+		}
+
+		if (!RotationConstraints.bLimitsCheck[ConstraintIndex] || (RotationConstraints.bLimitsCheck[ConstraintIndex] && NeedsSolve && FMath::Abs(DeltaAngle) > AngleTolerance))
+		{
+			const FReal IM = -FVec3::DotProduct(RotationConstraints.D.ConstraintAxis[ConstraintIndex], RotationConstraints.D.ConstraintDRAxis[ConstraintIndex][1]);
+			const FReal DeltaLambda = SolverStiffness * RotationConstraints.D.ConstraintHardStiffness[ConstraintIndex] * DeltaAngle / IM;
+
+			const FVec3 DR1 = AngularProjection * RotationConstraints.D.ConstraintDRAxis[ConstraintIndex][1] * DeltaLambda;
+			ApplyRotationDelta(1, DR1);
+
+			if (bLinearLocked)
+			{
+				FVec3 PositionConstraintsArms1;
+				if (!PositionConstraints.bUseSimd)
 				{
-					DeltaAngle -= RotationConstraints.ConstraintLimits[ConstraintIndex];
-					NeedsSolve = true;
+					PositionConstraintsArms1 = PositionConstraints.D.ConstraintArms[ConstraintIndex][1];
 				}
-				else if(DeltaAngle < -RotationConstraints.ConstraintLimits[ConstraintIndex])
+				else
 				{
-					DeltaAngle += RotationConstraints.ConstraintLimits[ConstraintIndex];
-					NeedsSolve = true;
-			
+					FVec3f PositionConstraintsArmsf;
+					VectorStoreFloat3(PositionConstraints.Simd.ConstraintArmsSimd[ConstraintIndex][1], &PositionConstraintsArmsf[0]);
+					PositionConstraintsArms1 = FVec3(PositionConstraintsArmsf);
 				}
+				const FVec3 DP1 = -AngularProjection * FVec3::CrossProduct(DR1, PositionConstraintsArms1);
+				ApplyPositionDelta(1, DP1);
 			}
+		}
+	}
+}
 
-			if (!RotationConstraints.bLimitsCheck[ConstraintIndex] || (RotationConstraints.bLimitsCheck[ConstraintIndex] && NeedsSolve && FMath::Abs(DeltaAngle) > AngleTolerance))
-			{
-				const FReal IM = -FVec3::DotProduct(RotationConstraints.ConstraintAxis[ConstraintIndex], RotationConstraints.ConstraintDRAxis[ConstraintIndex][1]);
-				const FReal DeltaLambda = SolverStiffness * RotationConstraints.ConstraintHardStiffness[ConstraintIndex] * DeltaAngle / IM;
-					
-				const FVec3 DR1 = AngularProjection * RotationConstraints.ConstraintDRAxis[ConstraintIndex][1] * DeltaLambda;
-				ApplyRotationDelta(1, DR1);
+void FPBDJointCachedSolver::ApplyRotationProjectionSimd(
+	const FRealSingle AngularProjectionf,
+	const bool bLinearLocked)
+{
+	const VectorRegister4Float Body0DQ = MakeVectorRegisterFloatFromDouble(MakeVectorRegister(Body(0).DQ()[0], Body(0).DQ()[1], Body(0).DQ()[2], 0.0f));
+	const VectorRegister4Float Body1DQ = MakeVectorRegisterFloatFromDouble(MakeVectorRegister(Body(1).DQ()[0], Body(1).DQ()[1], Body(1).DQ()[2], 0.0f));
+
+	const VectorRegister4Float DQDiff = VectorSubtract(Body1DQ, Body0DQ);
+	VectorRegister4Float ProjAxes[3];
+	for (int32 CIndex = 0; CIndex < 3; ++CIndex)
+	{
+		ProjAxes[CIndex] = Private::VectorDot3FastX(DQDiff, RotationConstraints.Simd.ConstraintAxisSimd[CIndex]);
+	}
+	VectorRegister4Float ProjAxis = VectorUnpackLo(ProjAxes[0], ProjAxes[1]);
+	ProjAxis = VectorMoveLh(ProjAxis, ProjAxes[2]);
+	VectorRegister4Float DeltaAngle = VectorAdd(RotationConstraints.Simd.ConstraintCX, ProjAxis);
+
+
+	const VectorRegister4Float AngleGT = VectorCompareGT(DeltaAngle, RotationConstraints.Simd.ConstraintLimitsSimd);
+	const VectorRegister4Float AngleLT = VectorCompareLT(DeltaAngle, VectorNegate(RotationConstraints.Simd.ConstraintLimitsSimd));
+
+	DeltaAngle = VectorSelect(AngleGT, VectorSubtract(DeltaAngle, RotationConstraints.Simd.ConstraintLimitsSimd),
+		VectorSelect(AngleLT, VectorAdd(DeltaAngle, RotationConstraints.Simd.ConstraintLimitsSimd), DeltaAngle));
 
-				if(bLinearLocked)
+	FRealSingle AngleTolerancef = FRealSingle(AngleTolerance);
+	const VectorRegister4Float AngleToleranceSimd = VectorLoadFloat1(&AngleTolerancef);
+	const VectorRegister4Float InTolerance = VectorCompareGT(VectorAbs(DeltaAngle), AngleToleranceSimd);
+	const VectorRegister4Float NeedsSolve = VectorBitwiseAnd(VectorBitwiseOr(AngleGT, AngleLT), InTolerance);
+
+	if (VectorMaskBits(NeedsSolve))
+	{
+		VectorRegister4Float IMs[3];
+		for (int32 CIndex = 0; CIndex < 3; ++CIndex)
+		{
+			IMs[CIndex] = Private::VectorDot3FastX(RotationConstraints.Simd.ConstraintAxisSimd[CIndex], RotationConstraints.Simd.ConstraintDRAxisSimd[CIndex][1]);
+		}
+		VectorRegister4Float IM = VectorUnpackLo(IMs[0], IMs[1]);
+		IM = VectorMoveLh(IM, IMs[2]);
+		IM = VectorNegate(IM);
+
+		const FRealSingle SolverStiffnessf = FRealSingle(SolverStiffness);
+		VectorRegister4Float Stiffness = VectorLoadFloat1(&SolverStiffnessf);
+		VectorRegister4Float DeltaLambda = VectorDivide(VectorMultiply(VectorMultiply(Stiffness, RotationConstraints.Simd.ConstraintHardStiffnessSimd), DeltaAngle), IM);
+		DeltaLambda = VectorSelect(NeedsSolve, DeltaLambda, VectorZeroFloat());
+
+		VectorRegister4Float AngularProjection = VectorLoadFloat1(&AngularProjectionf);
+		DeltaLambda = VectorMultiply(DeltaLambda, AngularProjection);
+
+		VectorRegister4Float DeltaLambdas[3];
+		DeltaLambdas[0] = VectorReplicate(DeltaLambda, 0);
+		DeltaLambdas[1] = VectorReplicate(DeltaLambda, 1);
+		DeltaLambdas[2] = VectorReplicate(DeltaLambda, 2);
+
+		VectorRegister4Float DR1 = VectorZeroFloat();
+		VectorRegister4Float DR1s[3];
+		for (int32 CIndex = 0; CIndex < 3; ++CIndex)
+		{
+			DR1s[CIndex] = VectorMultiply(RotationConstraints.Simd.ConstraintDRAxisSimd[CIndex][1], DeltaLambdas[CIndex]);
+			DR1 = VectorAdd(DR1s[CIndex], DR1);
+		}
+		FVec3f DR1f;
+		VectorStoreFloat3(DR1, &DR1f[0]);
+		ApplyRotationDelta(1, FVec3(DR1f));
+
+		if (bLinearLocked)
+		{
+			VectorRegister4Float DP1 = VectorZeroFloat();
+			AngularProjection = VectorNegate(AngularProjection);
+			for (int32 CIndex = 0; CIndex < 3; ++CIndex)
+			{
+				VectorRegister4Float PositionConstraintsArms1;
+				if (PositionConstraints.bUseSimd)
 				{
-					const FVec3 DP1 = -AngularProjection * FVec3::CrossProduct(DR1, PositionConstraints.ConstraintArms[ConstraintIndex][1]);
-					ApplyPositionDelta(1,DP1);
+					PositionConstraintsArms1 = PositionConstraints.Simd.ConstraintArmsSimd[CIndex][1];
 				}
+				else
+				{
+					PositionConstraintsArms1 = MakeVectorRegisterFloatFromDouble(MakeVectorRegister(PositionConstraints.D.ConstraintArms[CIndex][1][0], PositionConstraints.D.ConstraintArms[CIndex][1][1], PositionConstraints.D.ConstraintArms[CIndex][1][2], 0.0));
+				}
+				DP1 = VectorMultiplyAdd(AngularProjection, VectorCross(DR1, PositionConstraintsArms1), DP1);
 			}
+
+			FVec3f DP1f;
+			VectorStoreFloat3(DP1, &DP1f[0]);
+			ApplyPositionDelta(1, FVec3(DP1f));
 		}
 	}
 }
 
 void FPBDJointCachedSolver::ApplyPositionProjection(
-		const FReal Dt,
-		const FPBDJointSolverSettings& SolverSettings,
-		const FPBDJointSettings& JointSettings)
+	const FReal Dt,
+	const FPBDJointSolverSettings& SolverSettings,
+	const FPBDJointSettings& JointSettings)
 {
 	const FReal LinearProjection = FPBDJointUtilities::GetLinearProjection(SolverSettings, JointSettings);
 	if (LinearProjection == 0)
@@ -1465,46 +2148,137 @@ void FPBDJointCachedSolver::ApplyPositionProjection(
 		return;
 	}
 
-	for(int32 ConstraintIndex = 0; ConstraintIndex < 3; ++ConstraintIndex)
+	if (PositionConstraints.bUseSimd)
 	{
-		if(PositionConstraints.bValidDatas[ConstraintIndex] && !PositionConstraints.bSoftLimit[ConstraintIndex])
+		check(PositionConstraints.bSoftLimit[0] == false);
+		check(PositionConstraints.bSoftLimit[1] == false);
+		check(PositionConstraints.bSoftLimit[2] == false);
+		ApplyPositionProjectionSimd(LinearProjection);
+	}
+	else
+	{
+		for (int32 ConstraintIndex = 0; ConstraintIndex < 3; ++ConstraintIndex)
 		{
-			const FVec3 CX = Body(1).DP()  - Body(0).DP() +
-				FVec3::CrossProduct(Body(1).DQ(), PositionConstraints.ConstraintArms[ConstraintIndex][1]) -
-				FVec3::CrossProduct(Body(0).DQ(), PositionConstraints.ConstraintArms[ConstraintIndex][0]);
-				
-			FReal DeltaPosition = PositionConstraints.ConstraintCX[ConstraintIndex] + FVec3::DotProduct(CX, PositionConstraints.ConstraintAxis[ConstraintIndex]);
-				
-			bool NeedsSolve = false;
-			if(PositionConstraints.bLimitsCheck[ConstraintIndex])
+			ApplyAxisPositionProjection(LinearProjection, ConstraintIndex);
+		}
+	}
+}
+
+void FPBDJointCachedSolver::ApplyAxisPositionProjection(
+	const FReal LinearProjection,
+	const int32 ConstraintIndex)
+{
+	if (PositionConstraints.bValidDatas[ConstraintIndex] && !PositionConstraints.bSoftLimit[ConstraintIndex])
+	{
+		const FVec3 CX = Body(1).DP() - Body(0).DP() +
+			FVec3::CrossProduct(Body(1).DQ(), PositionConstraints.D.ConstraintArms[ConstraintIndex][1]) -
+			FVec3::CrossProduct(Body(0).DQ(), PositionConstraints.D.ConstraintArms[ConstraintIndex][0]);
+
+		FReal DeltaPosition = PositionConstraints.D.ConstraintCX[ConstraintIndex] + FVec3::DotProduct(CX, PositionConstraints.D.ConstraintAxis[ConstraintIndex]);
+
+		bool NeedsSolve = false;
+		if (PositionConstraints.bLimitsCheck[ConstraintIndex])
+		{
+			if (DeltaPosition > PositionConstraints.D.ConstraintLimits[ConstraintIndex])
 			{
-				if(DeltaPosition > PositionConstraints.ConstraintLimits[ConstraintIndex] )
-				{
-					DeltaPosition -= PositionConstraints.ConstraintLimits[ConstraintIndex];
-					NeedsSolve = true;
-				}
-				else if(DeltaPosition < -PositionConstraints.ConstraintLimits[ConstraintIndex])
-				{
-					DeltaPosition += PositionConstraints.ConstraintLimits[ConstraintIndex];
-					NeedsSolve = true;
-				}
+				DeltaPosition -= PositionConstraints.D.ConstraintLimits[ConstraintIndex];
+				NeedsSolve = true;
 			}
-			if (!PositionConstraints.bLimitsCheck[ConstraintIndex] || (PositionConstraints.bLimitsCheck[ConstraintIndex] && NeedsSolve && FMath::Abs(DeltaPosition ) > PositionTolerance))
+			else if (DeltaPosition < -PositionConstraints.D.ConstraintLimits[ConstraintIndex])
 			{
-				const FVec3 AngularAxis1 = FVec3::CrossProduct(PositionConstraints.ConstraintArms[ConstraintIndex][1], PositionConstraints.ConstraintAxis[ConstraintIndex]);
-				const FReal IM = InvM(1) - FVec3::DotProduct(AngularAxis1, PositionConstraints.ConstraintDRAxis[ConstraintIndex][1]);
-				const FReal DeltaLambda = SolverStiffness * PositionConstraints.ConstraintHardStiffness[ConstraintIndex] * DeltaPosition / IM;
-	
-				const FVec3 DX = PositionConstraints.ConstraintAxis[ConstraintIndex] * DeltaLambda;
-					
-				const FVec3 DP1 = -LinearProjection * InvM(1) * DX;
-				const FVec3 DR1 = LinearProjection * PositionConstraints.ConstraintDRAxis[ConstraintIndex][1] * DeltaLambda;
-
-				ApplyPositionDelta(1, DP1);
-				ApplyRotationDelta(1, DR1);
+				DeltaPosition += PositionConstraints.D.ConstraintLimits[ConstraintIndex];
+				NeedsSolve = true;
 			}
 		}
+		if (!PositionConstraints.bLimitsCheck[ConstraintIndex] || (PositionConstraints.bLimitsCheck[ConstraintIndex] && NeedsSolve && FMath::Abs(DeltaPosition) > PositionTolerance))
+		{
+			const FVec3 AngularAxis1 = FVec3::CrossProduct(PositionConstraints.D.ConstraintArms[ConstraintIndex][1], PositionConstraints.D.ConstraintAxis[ConstraintIndex]);
+			const FReal IM = InvM(1) - FVec3::DotProduct(AngularAxis1, PositionConstraints.D.ConstraintDRAxis[ConstraintIndex][1]);
+			const FReal DeltaLambda = SolverStiffness * PositionConstraints.D.ConstraintHardStiffness[ConstraintIndex] * DeltaPosition / IM;
+
+			const FVec3 DX = PositionConstraints.D.ConstraintAxis[ConstraintIndex] * DeltaLambda;
+
+			const FVec3 DP1 = -LinearProjection * InvM(1) * DX;
+			const FVec3 DR1 = LinearProjection * PositionConstraints.D.ConstraintDRAxis[ConstraintIndex][1] * DeltaLambda;
+
+			ApplyPositionDelta(1, DP1);
+			ApplyRotationDelta(1, DR1);
+		}
+	}
+}
+
+void FPBDJointCachedSolver::ApplyPositionProjectionSimd(
+	const FReal LinearProjection)
+{
+	const VectorRegister4Float Body0DP = MakeVectorRegisterFloatFromDouble(MakeVectorRegister(Body(0).DP()[0], Body(0).DP()[1], Body(0).DP()[2], 0.0f));
+	const VectorRegister4Float Body1DP = MakeVectorRegisterFloatFromDouble(MakeVectorRegister(Body(1).DP()[0], Body(1).DP()[1], Body(1).DP()[2], 0.0f));
+
+	const VectorRegister4Float Body0DQ = MakeVectorRegisterFloatFromDouble(MakeVectorRegister(Body(0).DQ()[0], Body(0).DQ()[1], Body(0).DQ()[2], 0.0f));
+	const VectorRegister4Float Body1DQ = MakeVectorRegisterFloatFromDouble(MakeVectorRegister(Body(1).DQ()[0], Body(1).DQ()[1], Body(1).DQ()[2], 0.0f));
+
+	const VectorRegister4Float DPDiff = VectorSubtract(Body1DP, Body0DP);
+
+	VectorRegister4Float DeltaPositions[3];
+	for (int32 CIndex = 0; CIndex < 3; ++CIndex)
+	{
+		VectorRegister4Float  Cross1 = VectorCross(Body1DQ, PositionConstraints.Simd.ConstraintArmsSimd[CIndex][1]);
+		VectorRegister4Float  Cross0 = VectorCross(Body0DQ, PositionConstraints.Simd.ConstraintArmsSimd[CIndex][0]);
+		VectorRegister4Float  CrossDiff = VectorSubtract(Cross1, Cross0);
+		VectorRegister4Float CX = VectorAdd(DPDiff, CrossDiff);
+
+		DeltaPositions[CIndex] = Private::VectorDot3FastX(CX, PositionConstraints.Simd.ConstraintAxisSimd[CIndex]);
+	}
+
+	VectorRegister4Float DeltaPosition = VectorUnpackLo(DeltaPositions[0], DeltaPositions[1]);
+	DeltaPosition = VectorMoveLh(DeltaPosition, DeltaPositions[2]);
+
+	DeltaPosition = VectorAdd(DeltaPosition, PositionConstraints.Simd.ConstraintCX);
+
+	const FRealSingle SolverStiffnessf = FRealSingle(SolverStiffness);
+	VectorRegister4Float Stiffness = VectorLoadFloat1(&SolverStiffnessf);
+
+	Stiffness = VectorMultiply(Stiffness, PositionConstraints.Simd.ConstraintHardStiffnessSimd);
+	DeltaPosition = VectorMultiply(Stiffness, DeltaPosition);
+
+	FRealSingle InvM1f = FRealSingle(InvM(1));
+	const VectorRegister4Float InvM1 = VectorLoadFloat1(&InvM1f);
+
+	VectorRegister4Float IMs[3];
+	for (int32 CIndex = 0; CIndex < 3; CIndex++)
+	{
+		const VectorRegister4Float AngularAxis1 = VectorCross(PositionConstraints.Simd.ConstraintArmsSimd[CIndex][1], PositionConstraints.Simd.ConstraintAxisSimd[CIndex]);
+		IMs[CIndex] = Private::VectorDot3FastX(AngularAxis1, PositionConstraints.Simd.ConstraintDRAxisSimd[CIndex][1]);
+	}
+	VectorRegister4Float IM = VectorUnpackLo(IMs[0], IMs[1]);
+	IM = VectorMoveLh(IM, IMs[2]);
+	IM = VectorSubtract(InvM1, IM);
+
+	const VectorRegister4Float DeltaLambda = VectorDivide(DeltaPosition, IM);
+	VectorRegister4Float DeltaLambdas[3];
+	DeltaLambdas[0] = VectorReplicate(DeltaLambda, 0);
+	DeltaLambdas[1] = VectorReplicate(DeltaLambda, 1);
+	DeltaLambdas[2] = VectorReplicate(DeltaLambda, 2);
+
+	FRealSingle NegLinPorjInvMf = -FRealSingle(LinearProjection * InvM(1));
+	VectorRegister4Float NegLinPorjInvM = VectorLoadFloat1(&NegLinPorjInvMf);
+	FRealSingle LinearProjectionf = FRealSingle(LinearProjection);
+	VectorRegister4Float LinearProjectionSimd = VectorLoadFloat1(&LinearProjectionf);
+
+	VectorRegister4Float DP1 = VectorZeroFloat();
+	VectorRegister4Float DR1 = VectorZeroFloat();
+	for (int32 CIndex = 0; CIndex < 3; CIndex++)
+	{
+		VectorRegister4Float DX = VectorMultiply(PositionConstraints.Simd.ConstraintAxisSimd[CIndex], DeltaLambdas[CIndex]);
+		DP1 = VectorMultiplyAdd(NegLinPorjInvM, DX, DP1);
+		DR1 = VectorMultiplyAdd(VectorMultiply(LinearProjectionSimd, PositionConstraints.Simd.ConstraintDRAxisSimd[CIndex][1]), DeltaLambdas[CIndex], DR1);
 	}
+	FVec3f DP1f;
+	VectorStoreFloat3(DP1, &DP1f[0]);
+	ApplyPositionDelta(1, FVec3(DP1f));
+	FVec3f DR1f;
+	VectorStoreFloat3(DR1, &DR1f[0]);
+	ApplyRotationDelta(1, FVec3(DR1f));
+
 }
 
 void FPBDJointCachedSolver::ApplyTeleports(
@@ -1526,36 +2300,76 @@ void FPBDJointCachedSolver::ApplyPositionTeleport(
 	{
 		return;
 	}
-
-	for (int32 ConstraintIndex = 0; ConstraintIndex < 3; ++ConstraintIndex)
+	if (PositionConstraints.bUseSimd)
+	{
+		ApplyPositionTeleportSimd(FRealSingle(JointSettings.TeleportDistance));
+	}
+	else
 	{
-		if (PositionConstraints.bValidDatas[ConstraintIndex] && !PositionConstraints.bSoftLimit[ConstraintIndex])
+		for (int32 ConstraintIndex = 0; ConstraintIndex < 3; ++ConstraintIndex)
 		{
-			FReal DeltaPosition = PositionConstraints.ConstraintCX[ConstraintIndex];
+			ApplyAxisPositionTeleport(JointSettings.TeleportDistance, ConstraintIndex);
+		}
+	}
+}
+
+void FPBDJointCachedSolver::ApplyAxisPositionTeleport(
+	const FReal TeleportDistance,
+	const int32 ConstraintIndex)
+{
+	check(PositionConstraints.bUseSimd == false);
 
-			bool NeedsSolve = false;
-			if (PositionConstraints.bLimitsCheck[ConstraintIndex])
+	if (PositionConstraints.bValidDatas[ConstraintIndex] && !PositionConstraints.bSoftLimit[ConstraintIndex])
+	{
+		FReal DeltaPosition = PositionConstraints.D.ConstraintCX[ConstraintIndex];
+
+		bool NeedsSolve = false;
+		if (PositionConstraints.bLimitsCheck[ConstraintIndex])
+		{
+			if (DeltaPosition > PositionConstraints.D.ConstraintLimits[ConstraintIndex])
 			{
-				if (DeltaPosition > PositionConstraints.ConstraintLimits[ConstraintIndex])
-				{
-					DeltaPosition -= PositionConstraints.ConstraintLimits[ConstraintIndex];
-					NeedsSolve = true;
-				}
-				else if (DeltaPosition < -PositionConstraints.ConstraintLimits[ConstraintIndex])
-				{
-					DeltaPosition += PositionConstraints.ConstraintLimits[ConstraintIndex];
-					NeedsSolve = true;
-				}
+				DeltaPosition -= PositionConstraints.D.ConstraintLimits[ConstraintIndex];
+				NeedsSolve = true;
 			}
-			if (!PositionConstraints.bLimitsCheck[ConstraintIndex] || (PositionConstraints.bLimitsCheck[ConstraintIndex] && NeedsSolve))
+			else if (DeltaPosition < -PositionConstraints.D.ConstraintLimits[ConstraintIndex])
 			{
-				if (FMath::Abs(DeltaPosition) > JointSettings.TeleportDistance)
-				{
-					const FVec3 DP1 = -DeltaPosition * PositionConstraints.ConstraintAxis[ConstraintIndex];
-					ApplyPositionDelta(1, DP1);
-				}
+				DeltaPosition += PositionConstraints.D.ConstraintLimits[ConstraintIndex];
+				NeedsSolve = true;
 			}
 		}
+		if (!PositionConstraints.bLimitsCheck[ConstraintIndex] || (PositionConstraints.bLimitsCheck[ConstraintIndex] && NeedsSolve))
+		{
+			if (FMath::Abs(DeltaPosition) > TeleportDistance)
+			{
+				const FVec3 DP1 = -DeltaPosition * PositionConstraints.D.ConstraintAxis[ConstraintIndex];
+				ApplyPositionDelta(1, DP1);
+			}
+		}
+	}
+}
+
+void FPBDJointCachedSolver::ApplyPositionTeleportSimd(
+	const FRealSingle TeleportDistance)
+{
+	const VectorRegister4Float TeleportDistanceSimd = VectorLoadFloat1(&TeleportDistance);
+	const VectorRegister4Float IsGT = VectorCompareGT(VectorAbs(PositionConstraints.Simd.ConstraintCX), TeleportDistanceSimd);
+
+	if (VectorMaskBits(IsGT))
+	{
+		const VectorRegister4Float ConstraintCX = VectorSelect(IsGT, VectorNegate(PositionConstraints.Simd.ConstraintCX), VectorZeroFloat());
+
+		VectorRegister4Float ConstraintCXs[3];
+		ConstraintCXs[0] = VectorReplicate(ConstraintCX, 0);
+		ConstraintCXs[1] = VectorReplicate(ConstraintCX, 1);
+		ConstraintCXs[2] = VectorReplicate(ConstraintCX, 2);
+		VectorRegister4Float DP1 = VectorMultiply(PositionConstraints.Simd.ConstraintAxisSimd[0], ConstraintCXs[0]);
+		for (int32 ConstraintIndex = 1; ConstraintIndex < 3; ++ConstraintIndex)
+		{
+			DP1 = VectorMultiplyAdd(PositionConstraints.Simd.ConstraintAxisSimd[ConstraintIndex], ConstraintCXs[ConstraintIndex], DP1);
+		}
+		FVec3f DP1f;
+		VectorStoreFloat3(DP1, &DP1f[0]);
+		ApplyPositionDelta(1, FVec3(DP1f));
 	}
 }
 
@@ -1576,27 +2390,28 @@ void FPBDJointCachedSolver::ApplyRotationTeleport(
 void FPBDJointCachedSolver::ApplyPositionDrives(
 	const FReal Dt)
 {
-	for(int32 ConstraintIndex = 0; ConstraintIndex < 3; ++ConstraintIndex)
+	for (int32 ConstraintIndex = 0; ConstraintIndex < 3; ++ConstraintIndex)
 	{
-		if(PositionDrives.bValidDatas[ConstraintIndex])
+		if (PositionDrives.bValidDatas[ConstraintIndex])
 		{
 			ApplyAxisPositionDrive(ConstraintIndex, Dt);
 		}
 	}
 }
-	
+
 
 void FPBDJointCachedSolver::ApplyAxisPositionDrive(
 	const int32 ConstraintIndex,
 	const FReal Dt)
 {
-	const FReal Stiffness = PositionDrives.ConstraintSoftStiffness[ConstraintIndex];
-	const FReal Damping = PositionDrives.ConstraintSoftDamping[ConstraintIndex];
-	const FReal IM = PositionDrives.ConstraintSoftIM[ConstraintIndex];
+	check(PositionDrives.bUseSimd == false);
+	const FReal Stiffness = PositionDrives.D.ConstraintSoftStiffness[ConstraintIndex];
+	const FReal Damping = PositionDrives.D.ConstraintSoftDamping[ConstraintIndex];
+	const FReal IM = PositionDrives.D.ConstraintSoftIM[ConstraintIndex];
 
-	const FVec3 Delta0 = Body(0).DP() + FVec3::CrossProduct(Body(0).DQ(), PositionDrives.ConstraintArms[ConstraintIndex][0]);
-	const FVec3 Delta1 = Body(1).DP() + FVec3::CrossProduct(Body(1).DQ(), PositionDrives.ConstraintArms[ConstraintIndex][1]);
-	const FReal CX = PositionDrives.ConstraintCX[ConstraintIndex] + FVec3::DotProduct(Delta1 - Delta0, PositionDrives.ConstraintAxis[ConstraintIndex]);
+	const FVec3 Delta0 = Body(0).DP() + FVec3::CrossProduct(Body(0).DQ(), PositionDrives.D.ConstraintArms[ConstraintIndex][0]);
+	const FVec3 Delta1 = Body(1).DP() + FVec3::CrossProduct(Body(1).DQ(), PositionDrives.D.ConstraintArms[ConstraintIndex][1]);
+	const FReal CX = PositionDrives.D.ConstraintCX[ConstraintIndex] + FVec3::DotProduct(Delta1 - Delta0, PositionDrives.D.ConstraintAxis[ConstraintIndex]);
 
 	FReal CVDt = 0;
 	if (Damping > UE_KINDA_SMALL_NUMBER)
@@ -1604,15 +2419,15 @@ void FPBDJointCachedSolver::ApplyAxisPositionDrive(
 		const FVec3 V0Dt = FVec3::CalculateVelocity(InitConnectorXs[0], ConnectorXs[0] + Delta0, 1.0f);
 		const FVec3 V1Dt = FVec3::CalculateVelocity(InitConnectorXs[1], ConnectorXs[1] + Delta1, 1.0f);
 		const FReal TargetVDt = PositionDrives.ConstraintVX[ConstraintIndex] * Dt;
-		CVDt = TargetVDt + FVec3::DotProduct(V0Dt - V1Dt, PositionDrives.ConstraintAxis[ConstraintIndex]);
+		CVDt = TargetVDt + FVec3::DotProduct(V0Dt - V1Dt, PositionDrives.D.ConstraintAxis[ConstraintIndex]);
 	}
 
-	FReal Lambda = PositionDrives.ConstraintLambda[ConstraintIndex];
+	FReal Lambda = PositionDrives.D.ConstraintLambda[ConstraintIndex];
 	FReal DeltaLambda = SolverStiffness * (Stiffness * CX - Damping * CVDt - Lambda) / IM;
 	Lambda += DeltaLambda;
 
 	PositionDrives.ApplyMaxLambda(ConstraintIndex, DeltaLambda, Lambda);
-	PositionDrives.ConstraintLambda[ConstraintIndex] = Lambda;
+	PositionDrives.D.ConstraintLambda[ConstraintIndex] = Lambda;
 
 	SolvePositionConstraintDelta(ConstraintIndex, DeltaLambda, PositionDrives);
 }
@@ -1638,6 +2453,7 @@ void FPBDJointCachedSolver::ApplyAxisPositionVelocityDrive(
 	const int32 ConstraintIndex,
 	const FReal Dt)
 {
+	check(PositionDrives.bUseSimd == false);
 	// NOTE: Using the actual damping, not the PBD modified value
 	const FReal Damping = PositionDrives.SettingsSoftDamping[ConstraintIndex] * Dt;
 	if (Damping < UE_SMALL_NUMBER)
@@ -1645,13 +2461,13 @@ void FPBDJointCachedSolver::ApplyAxisPositionVelocityDrive(
 		return;
 	}
 
-	const FReal MassScale = PositionDrives.bAccelerationMode ? (FReal(1) / PositionDrives.ConstraintHardIM[ConstraintIndex]) : FReal(1);
-	const FReal IM = MassScale * Damping * PositionDrives.ConstraintHardIM[ConstraintIndex] + (FReal)1;
+	const FReal MassScale = PositionDrives.bAccelerationMode ? (FReal(1) / PositionDrives.D.ConstraintHardIM[ConstraintIndex]) : FReal(1);
+	const FReal IM = MassScale * Damping * PositionDrives.D.ConstraintHardIM[ConstraintIndex] + (FReal)1;
 
 	// Velocity error to correct
-	const FVec3 V0 = V(0) + FVec3::CrossProduct(W(0), PositionDrives.ConstraintArms[ConstraintIndex][0]);
-	const FVec3 V1 = V(1) + FVec3::CrossProduct(W(1), PositionDrives.ConstraintArms[ConstraintIndex][1]);
-	const FReal VRel = FVec3::DotProduct(V1 - V0, PositionDrives.ConstraintAxis[ConstraintIndex]);
+	const FVec3 V0 = V(0) + FVec3::CrossProduct(W(0), PositionDrives.D.ConstraintArms[ConstraintIndex][0]);
+	const FVec3 V1 = V(1) + FVec3::CrossProduct(W(1), PositionDrives.D.ConstraintArms[ConstraintIndex][1]);
+	const FReal VRel = FVec3::DotProduct(V1 - V0, PositionDrives.D.ConstraintAxis[ConstraintIndex]);
 	const FReal TargetV = PositionDrives.ConstraintVX[ConstraintIndex];
 	const FReal CV = (VRel - TargetV);
 
@@ -1662,25 +2478,25 @@ void FPBDJointCachedSolver::ApplyAxisPositionVelocityDrive(
 	// Apply limits and accumulate total impulse 
 	// (NOTE: Limits and net impulses are position based, not velocity based)
 	FReal DeltaLambda = DeltaLambdaVel * Dt;
-	FReal Lambda = PositionDrives.ConstraintLambda[ConstraintIndex] + DeltaLambda;
+	FReal Lambda = PositionDrives.D.ConstraintLambda[ConstraintIndex] + DeltaLambda;
 	PositionDrives.ApplyMaxLambda(ConstraintIndex, DeltaLambda, Lambda);
-	PositionDrives.ConstraintLambda[ConstraintIndex] = Lambda;	
+	PositionDrives.D.ConstraintLambda[ConstraintIndex] = Lambda;
 	DeltaLambdaVel = DeltaLambda / Dt;
 
 	LambdaVel += DeltaLambdaVel;
-	const FVec3 Impulse = DeltaLambdaVel * PositionDrives.ConstraintAxis[ConstraintIndex];
+	const FVec3 Impulse = DeltaLambdaVel * PositionDrives.D.ConstraintAxis[ConstraintIndex];
 
 	if (Body(0).IsDynamic())
 	{
 		const FVec3 DV0 = InvM(0) * Impulse;
-		const FVec3 DW0 = PositionDrives.ConstraintDRAxis[ConstraintIndex][0] * DeltaLambdaVel;
+		const FVec3 DW0 = PositionDrives.D.ConstraintDRAxis[ConstraintIndex][0] * DeltaLambdaVel;
 
 		Body(0).ApplyVelocityDelta(DV0, DW0);
 	}
 	if (Body(1).IsDynamic())
 	{
 		const FVec3 DV1 = -InvM(1) * Impulse;
-		const FVec3 DW1 = PositionDrives.ConstraintDRAxis[ConstraintIndex][1] * DeltaLambdaVel;
+		const FVec3 DW1 = PositionDrives.D.ConstraintDRAxis[ConstraintIndex][1] * DeltaLambdaVel;
 
 		Body(1).ApplyVelocityDelta(DV1, DW1);
 	}
@@ -1697,25 +2513,25 @@ void FPBDJointCachedSolver::InitRotationDrives(
 	RotationDrives.bValidDatas[0] = false;
 	RotationDrives.bValidDatas[1] = false;
 	RotationDrives.bValidDatas[2] = false;
+	RotationDrives.bUseSimd = false;
 
 	bool bHasRotationDrives =
 		JointSettings.bAngularTwistPositionDriveEnabled
-			|| JointSettings.bAngularTwistVelocityDriveEnabled
-			|| JointSettings.bAngularSwingPositionDriveEnabled
-			|| JointSettings.bAngularSwingVelocityDriveEnabled
-			|| JointSettings.bAngularSLerpPositionDriveEnabled
-			|| JointSettings.bAngularSLerpVelocityDriveEnabled;
+		|| JointSettings.bAngularTwistVelocityDriveEnabled
+		|| JointSettings.bAngularSwingPositionDriveEnabled
+		|| JointSettings.bAngularSwingVelocityDriveEnabled
+		|| JointSettings.bAngularSLerpPositionDriveEnabled
+		|| JointSettings.bAngularSLerpVelocityDriveEnabled;
 	if (!bHasRotationDrives)
 	{
 		return;
 	}
 
-	EJointMotionType TwistMotion = JointSettings.AngularMotionTypes[(int32)EJointAngularConstraintIndex::Twist];
-	EJointMotionType Swing1Motion = JointSettings.AngularMotionTypes[(int32)EJointAngularConstraintIndex::Swing1];
-	EJointMotionType Swing2Motion = JointSettings.AngularMotionTypes[(int32)EJointAngularConstraintIndex::Swing2];
-
 	if (SolverSettings.bEnableDrives)
 	{
+		EJointMotionType TwistMotion = JointSettings.AngularMotionTypes[(int32)EJointAngularConstraintIndex::Twist];
+		EJointMotionType Swing1Motion = JointSettings.AngularMotionTypes[(int32)EJointAngularConstraintIndex::Swing1];
+		EJointMotionType Swing2Motion = JointSettings.AngularMotionTypes[(int32)EJointAngularConstraintIndex::Swing2];
 		bool bTwistLocked = TwistMotion == EJointMotionType::Locked;
 		bool bSwing1Locked = Swing1Motion == EJointMotionType::Locked;
 		bool bSwing2Locked = Swing2Motion == EJointMotionType::Locked;
@@ -1724,6 +2540,7 @@ void FPBDJointCachedSolver::InitRotationDrives(
 		// @todo(ccaulfield): setting should be cleaned up before being passed to the solver
 		if ((JointSettings.bAngularSLerpPositionDriveEnabled || JointSettings.bAngularSLerpVelocityDriveEnabled) && !bTwistLocked && !bSwing1Locked && !bSwing2Locked)
 		{
+			RotationDrives.bUseSimd = bUseSimd;
 			InitSLerpDrive(Dt, SolverSettings, JointSettings);
 		}
 		else
@@ -1741,16 +2558,80 @@ void FPBDJointCachedSolver::InitRotationDrives(
 }
 
 void FPBDJointCachedSolver::InitRotationConstraintDrive(
-			const int32 ConstraintIndex,
-			const FVec3& ConstraintAxis,
-			const FReal Dt,
-			const FReal DeltaAngle)
+	const int32 ConstraintIndex,
+	const FVec3& ConstraintAxis,
+	const FReal Dt,
+	const FReal DeltaAngle)
 {
-	RotationDrives.UpdateDatas(ConstraintIndex, ConstraintAxis, DeltaAngle,0.0f);
+	RotationDrives.UpdateDatas(ConstraintIndex, ConstraintAxis, DeltaAngle, 0.0f);
 
 	InitRotationDatasMass(RotationDrives, ConstraintIndex, Dt);
 }
 
+void FPBDJointCachedSolver::InitRotationConstraintDriveSimd(
+	FVec3 ConstraintAxes[3],
+	const FReal Dt,
+	const FVec3 DeltaAngles
+)
+{
+	FVec3 ConstraintHardIM;
+	FVec3 ConstraintSoftStiffness;
+	FVec3 ConstraintSoftDamping;
+	FVec3 ConstraintSoftIM;
+
+	RotationDrives.Simd.ConstraintCX = MakeVectorRegisterFloatFromDouble(MakeVectorRegister(DeltaAngles[0], DeltaAngles[1], DeltaAngles[2], 0.0f));
+
+	for (int32 ConstraintIndex = 0; ConstraintIndex < 3; ConstraintIndex++)
+	{
+		RotationDrives.bValidDatas[ConstraintIndex] = true;
+		RotationDrives.bLimitsCheck[ConstraintIndex] = true;
+
+		RotationDrives.ConstraintVX[ConstraintIndex] = 0.0;
+		RotationDrives.ConstraintRestitution[ConstraintIndex] = 0.0;
+
+		RotationDrives.Simd.ConstraintArmsSimd[ConstraintIndex][0] = VectorZeroFloat();
+		RotationDrives.Simd.ConstraintArmsSimd[ConstraintIndex][1] = VectorZeroFloat();
+		RotationDrives.Simd.ConstraintAxisSimd[ConstraintIndex] = MakeVectorRegisterFloatFromDouble(MakeVectorRegister(ConstraintAxes[ConstraintIndex][0], ConstraintAxes[ConstraintIndex][1], ConstraintAxes[ConstraintIndex][2], 0.0f));
+
+		// InitRotationDatasMass(RotationDrives, ConstraintIndex, Dt);
+		const VectorRegister4Float& Axis = RotationDrives.Simd.ConstraintAxisSimd[ConstraintIndex];
+		const VectorRegister4Float AxisX = VectorReplicate(Axis, 0);
+		const VectorRegister4Float AxisY = VectorReplicate(Axis, 1);
+		const VectorRegister4Float AxisZ = VectorReplicate(Axis, 2);
+		const FMatrix33& InvI0 = InvI(0);
+		const VectorRegister4Float InvI00 = MakeVectorRegisterFloatFromDouble(MakeVectorRegister(InvI0.M[0][0], InvI0.M[0][1], InvI0.M[0][2], 0.0f));
+		const VectorRegister4Float InvI01 = MakeVectorRegisterFloatFromDouble(MakeVectorRegister(InvI0.M[1][0], InvI0.M[1][1], InvI0.M[1][2], 0.0f));
+		const VectorRegister4Float InvI02 = MakeVectorRegisterFloatFromDouble(MakeVectorRegister(InvI0.M[2][0], InvI0.M[2][1], InvI0.M[2][2], 0.0f));
+		const VectorRegister4Float IA0 = VectorMultiplyAdd(InvI00, AxisX, VectorMultiplyAdd(InvI01, AxisY, VectorMultiply(InvI02, AxisZ)));
+
+		const FMatrix33& InvI1 = InvI(1);
+		const VectorRegister4Float InvI10 = MakeVectorRegisterFloatFromDouble(MakeVectorRegister(InvI1.M[0][0], InvI1.M[0][1], InvI1.M[0][2], 0.0f));
+		const VectorRegister4Float InvI11 = MakeVectorRegisterFloatFromDouble(MakeVectorRegister(InvI1.M[1][0], InvI1.M[1][1], InvI1.M[1][2], 0.0f));
+		const VectorRegister4Float InvI12 = MakeVectorRegisterFloatFromDouble(MakeVectorRegister(InvI1.M[2][0], InvI1.M[2][1], InvI1.M[2][2], 0.0f));
+		const VectorRegister4Float IA1 = VectorMultiplyAdd(InvI10, AxisX, VectorMultiplyAdd(InvI11, AxisY, VectorMultiply(InvI12, AxisZ)));
+
+		const FRealSingle II0 = VectorDot3Scalar(Axis, IA0);
+		const FRealSingle II1 = VectorDot3Scalar(Axis, IA1);
+
+		ConstraintHardIM[ConstraintIndex] = II0 + II1;
+
+		RotationDrives.Simd.ConstraintDRAxisSimd[ConstraintIndex][0] = IA0;
+		RotationDrives.Simd.ConstraintDRAxisSimd[ConstraintIndex][1] = VectorNegate(IA1);
+
+
+		check(RotationDrives.bSoftLimit[ConstraintIndex])
+			const FReal SpringMassScale = (bUsePositionBasedDrives) ? (FReal)1 / (ConstraintHardIM[ConstraintIndex]) : (FReal)1;
+		ConstraintSoftStiffness[ConstraintIndex] = SpringMassScale * RotationDrives.SettingsSoftStiffness[ConstraintIndex] * Dt * Dt;
+		ConstraintSoftDamping[ConstraintIndex] = (bUsePositionBasedDrives) ? SpringMassScale * RotationDrives.SettingsSoftDamping[ConstraintIndex] * Dt : 0;
+		ConstraintSoftIM[ConstraintIndex] = (ConstraintSoftStiffness[ConstraintIndex] + ConstraintSoftDamping[ConstraintIndex]) * ConstraintHardIM[ConstraintIndex] + (FReal)1;
+
+	}
+	RotationDrives.Simd.ConstraintHardIMSimd = MakeVectorRegisterFloatFromDouble(MakeVectorRegister(ConstraintHardIM[0], ConstraintHardIM[1], ConstraintHardIM[2], 0.0f));
+	RotationDrives.Simd.ConstraintSoftStiffnessSimd = MakeVectorRegisterFloatFromDouble(MakeVectorRegister(ConstraintSoftStiffness[0], ConstraintSoftStiffness[1], ConstraintSoftStiffness[2], 0.0f));
+	RotationDrives.Simd.ConstraintSoftDampingSimd = MakeVectorRegisterFloatFromDouble(MakeVectorRegister(ConstraintSoftDamping[0], ConstraintSoftDamping[1], ConstraintSoftDamping[2], 0.0f));
+	RotationDrives.Simd.ConstraintSoftIMSimd = MakeVectorRegisterFloatFromDouble(MakeVectorRegister(ConstraintSoftIM[0], ConstraintSoftIM[1], ConstraintSoftIM[2], 0.0f));
+}
+
 void FPBDJointCachedSolver::InitSwingTwistDrives(
 	const FReal Dt,
 	const FPBDJointSolverSettings& SolverSettings,
@@ -1773,6 +2654,7 @@ void FPBDJointCachedSolver::InitSwingTwistDrives(
 	const int32 S1 = (int32)EJointAngularConstraintIndex::Swing1;
 	const int32 S2 = (int32)EJointAngularConstraintIndex::Swing2;
 
+	// TODO this could be removed if bUseSimd
 	RotationDrives.InitDatas(TW, true, FPBDJointUtilities::GetAngularTwistDriveStiffness(SolverSettings, JointSettings),
 		FPBDJointUtilities::GetAngularTwistDriveDamping(SolverSettings, JointSettings), 0.0);
 	RotationDrives.InitDatas(S1, true, FPBDJointUtilities::GetAngularSwingDriveStiffness(SolverSettings, JointSettings),
@@ -1782,29 +2664,47 @@ void FPBDJointCachedSolver::InitSwingTwistDrives(
 
 	RotationDrives.bAccelerationMode = FPBDJointUtilities::GetAngularDriveAccelerationMode(SolverSettings, JointSettings);
 
-	const bool bUseTwistDrive = bTwistDriveEnabled && (((FMath::Abs(DTwistAngle) > AngleTolerance) && (RotationDrives.ConstraintSoftStiffness[TW] > 0.0f)) || (RotationDrives.ConstraintSoftDamping[TW]  > 0.0f));
-	if (bUseTwistDrive)
+	const bool bUseTwistDrive = bTwistDriveEnabled && (((FMath::Abs(DTwistAngle) > AngleTolerance) && (RotationDrives.D.ConstraintSoftStiffness[TW] > 0.0f)) || (RotationDrives.D.ConstraintSoftDamping[TW] > 0.0f));
+	const bool bUseSwing1Drive = bSwing1DriveEnabled && (((FMath::Abs(DSwing1Angle) > AngleTolerance) && (RotationDrives.D.ConstraintSoftStiffness[S1] > 0.0f)) || (RotationDrives.D.ConstraintSoftDamping[S1] > 0.0f));
+	const bool bUseSwing2Drive = bSwing2DriveEnabled && (((FMath::Abs(DSwing2Angle) > AngleTolerance) && (RotationDrives.D.ConstraintSoftStiffness[S2] > 0.0f)) || (RotationDrives.D.ConstraintSoftDamping[S2] > 0.0f));
+	RotationDrives.bUseSimd = bUseSimd && bUseTwistDrive && bUseSwing1Drive && bUseSwing2Drive;
+
+	if (RotationDrives.bUseSimd)
 	{
-		InitRotationConstraintDrive(TW, ConnectorRs[1] * FJointConstants::TwistAxis(), Dt, DTwistAngle);
+		FVec3 ConstraintAxes[3] = { ConnectorRs[1] * FJointConstants::TwistAxis(), ConnectorRs[1] * FJointConstants::Swing1Axis(), ConnectorRs[1] * FJointConstants::Swing2Axis() };
+		InitRotationConstraintDriveSimd(ConstraintAxes, Dt, FVec3(DTwistAngle, DSwing1Angle, DSwing2Angle));
+
+		RotationDrives.Simd.ConstraintHardStiffnessSimd = VectorZeroFloat();
+		RotationDrives.Simd.ConstraintLambda = VectorZeroFloat();
+
 		RotationDrives.ConstraintVX[TW] = JointSettings.AngularDriveVelocityTarget[TW];
 		RotationDrives.SetMaxForce(TW, JointSettings.AngularDriveMaxTorque[TW], Dt);
-	}
-
-	const bool bUseSwing1Drive = bSwing1DriveEnabled && (((FMath::Abs(DSwing1Angle) > AngleTolerance) && (RotationDrives.ConstraintSoftStiffness[S1] > 0.0f)) || (RotationDrives.ConstraintSoftDamping[S1] > 0.0f));
-	if (bUseSwing1Drive)
-	{
-		InitRotationConstraintDrive(S1, ConnectorRs[1] * FJointConstants::Swing1Axis(),  Dt, DSwing1Angle);
 		RotationDrives.ConstraintVX[S1] = JointSettings.AngularDriveVelocityTarget[S1];
 		RotationDrives.SetMaxForce(S1, JointSettings.AngularDriveMaxTorque[S1], Dt);
-	}
-
-	const bool bUseSwing2Drive = bSwing2DriveEnabled && (((FMath::Abs(DSwing2Angle) > AngleTolerance) && (RotationDrives.ConstraintSoftStiffness[S2] > 0.0f)) || (RotationDrives.ConstraintSoftDamping[S2] > 0.0f));
-	if (bUseSwing2Drive)
-	{
-		InitRotationConstraintDrive(S2, ConnectorRs[1] * FJointConstants::Swing2Axis(),  Dt, DSwing2Angle);
 		RotationDrives.ConstraintVX[S2] = JointSettings.AngularDriveVelocityTarget[S2];
 		RotationDrives.SetMaxForce(S2, JointSettings.AngularDriveMaxTorque[S2], Dt);
 	}
+	else
+	{
+		if (bUseTwistDrive)
+		{
+			InitRotationConstraintDrive(TW, ConnectorRs[1] * FJointConstants::TwistAxis(), Dt, DTwistAngle);
+			RotationDrives.ConstraintVX[TW] = JointSettings.AngularDriveVelocityTarget[TW];
+			RotationDrives.SetMaxForce(TW, JointSettings.AngularDriveMaxTorque[TW], Dt);
+		}
+		if (bUseSwing1Drive)
+		{
+			InitRotationConstraintDrive(S1, ConnectorRs[1] * FJointConstants::Swing1Axis(), Dt, DSwing1Angle);
+			RotationDrives.ConstraintVX[S1] = JointSettings.AngularDriveVelocityTarget[S1];
+			RotationDrives.SetMaxForce(S1, JointSettings.AngularDriveMaxTorque[S1], Dt);
+		}
+		if (bUseSwing2Drive)
+		{
+			InitRotationConstraintDrive(S2, ConnectorRs[1] * FJointConstants::Swing2Axis(), Dt, DSwing2Angle);
+			RotationDrives.ConstraintVX[S2] = JointSettings.AngularDriveVelocityTarget[S2];
+			RotationDrives.SetMaxForce(S2, JointSettings.AngularDriveMaxTorque[S2], Dt);
+		}
+	}
 }
 
 void FPBDJointCachedSolver::InitSLerpDrive(
@@ -1815,7 +2715,7 @@ void FPBDJointCachedSolver::InitSLerpDrive(
 	for (int32 AxisIndex = 0; AxisIndex < 3; ++AxisIndex)
 	{
 		RotationDrives.InitDatas(AxisIndex, true, FPBDJointUtilities::GetAngularSLerpDriveStiffness(SolverSettings, JointSettings),
-						FPBDJointUtilities::GetAngularSLerpDriveDamping(SolverSettings, JointSettings), 0.0);
+			FPBDJointUtilities::GetAngularSLerpDriveDamping(SolverSettings, JointSettings), 0.0);
 	}
 	RotationDrives.bAccelerationMode = FPBDJointUtilities::GetAngularDriveAccelerationMode(SolverSettings, JointSettings);
 
@@ -1824,7 +2724,7 @@ void FPBDJointCachedSolver::InitSLerpDrive(
 	TargetAngPos.EnforceShortestArcWith(R01);
 	const FRotation3 R1Error = TargetAngPos.Inverse() * R01;
 
-	FReal AxisAngles[3] =
+	FVec3 AxisAngles =
 	{
 		2.0f * Utilities::AsinEst(R1Error.X),
 		2.0f * Utilities::AsinEst(R1Error.Y),
@@ -1834,9 +2734,18 @@ void FPBDJointCachedSolver::InitSLerpDrive(
 	FVec3 Axes[3];
 	ConnectorRs[1].ToMatrixAxes(Axes[0], Axes[1], Axes[2]);
 
-	InitRotationConstraintDrive(0, Axes[0], Dt, AxisAngles[0]);
-	InitRotationConstraintDrive(1, Axes[1], Dt, AxisAngles[1]);
-	InitRotationConstraintDrive(2, Axes[2], Dt, AxisAngles[2]);
+	if (RotationDrives.bUseSimd)
+	{
+		RotationDrives.Simd.ConstraintLambda = VectorZeroFloat();
+		RotationDrives.Simd.ConstraintHardStiffnessSimd = VectorZeroFloat();
+		InitRotationConstraintDriveSimd(Axes, Dt, AxisAngles);
+	}
+	else
+	{
+		InitRotationConstraintDrive(0, Axes[0], Dt, AxisAngles[0]);
+		InitRotationConstraintDrive(1, Axes[1], Dt, AxisAngles[1]);
+		InitRotationConstraintDrive(2, Axes[2], Dt, AxisAngles[2]);
+	}
 
 	RotationDrives.SetMaxForce(0, JointSettings.AngularDriveMaxTorque[0], Dt);
 	RotationDrives.SetMaxForce(1, JointSettings.AngularDriveMaxTorque[1], Dt);
@@ -1857,11 +2766,18 @@ void FPBDJointCachedSolver::InitSLerpDrive(
 void FPBDJointCachedSolver::ApplyRotationDrives(
 	const FReal Dt)
 {
-	for(int32 ConstraintIndex = 0; ConstraintIndex < 3; ++ConstraintIndex)
+	if (RotationDrives.bUseSimd)
+	{
+		ApplyRotationDrivesSimd(Dt);
+	}
+	else
 	{
-		if(RotationDrives.bValidDatas[ConstraintIndex])
+		for (int32 ConstraintIndex = 0; ConstraintIndex < 3; ++ConstraintIndex)
 		{
-			ApplyAxisRotationDrive(ConstraintIndex, Dt);
+			if (RotationDrives.bValidDatas[ConstraintIndex])
+			{
+				ApplyAxisRotationDrive(ConstraintIndex, Dt);
+			}
 		}
 	}
 }
@@ -1870,16 +2786,16 @@ void FPBDJointCachedSolver::ApplyAxisRotationDrive(
 	const int32 ConstraintIndex,
 	const FReal Dt)
 {
-	const FReal Stiffness = RotationDrives.ConstraintSoftStiffness[ConstraintIndex];
-	const FReal Damping = RotationDrives.ConstraintSoftDamping[ConstraintIndex];
-	const FReal IM = RotationDrives.ConstraintSoftIM[ConstraintIndex];
+	const FReal Stiffness = RotationDrives.D.ConstraintSoftStiffness[ConstraintIndex];
+	const FReal Damping = RotationDrives.D.ConstraintSoftDamping[ConstraintIndex];
+	const FReal IM = RotationDrives.D.ConstraintSoftIM[ConstraintIndex];
 
 	// Stiffness position delta
 	FReal CX = 0;
 	if (Stiffness > UE_KINDA_SMALL_NUMBER)
 	{
-		const FReal DX = FVec3::DotProduct(Body(1).DQ() - Body(0).DQ(), RotationDrives.ConstraintAxis[ConstraintIndex]);
-		const FReal TargetX = RotationDrives.ConstraintCX[ConstraintIndex];
+		const FReal DX = FVec3::DotProduct(Body(1).DQ() - Body(0).DQ(), RotationDrives.D.ConstraintAxis[ConstraintIndex]);
+		const FReal TargetX = RotationDrives.D.ConstraintCX[ConstraintIndex];
 		CX = TargetX + DX;
 	}
 
@@ -1890,17 +2806,95 @@ void FPBDJointCachedSolver::ApplyAxisRotationDrive(
 		const FVec3 W0Dt = FVec3(Body(0).DQ()) + ConnectorWDts[0];
 		const FVec3 W1Dt = FVec3(Body(1).DQ()) + ConnectorWDts[1];
 		const FReal TargetW = RotationDrives.ConstraintVX[ConstraintIndex];
-		CVDt = (TargetW * Dt) + FVec3::DotProduct(RotationDrives.ConstraintAxis[ConstraintIndex] , W0Dt - W1Dt);
+		CVDt = (TargetW * Dt) + FVec3::DotProduct(RotationDrives.D.ConstraintAxis[ConstraintIndex], W0Dt - W1Dt);
 	}
 
-	FReal Lambda = RotationDrives.ConstraintLambda[ConstraintIndex];
+	FReal Lambda = RotationDrives.D.ConstraintLambda[ConstraintIndex];
 	FReal DeltaLambda = SolverStiffness * (Stiffness * CX - Damping * CVDt - Lambda) / IM;
 	Lambda += DeltaLambda;
 
 	RotationDrives.ApplyMaxLambda(ConstraintIndex, DeltaLambda, Lambda);
-	RotationDrives.ConstraintLambda[ConstraintIndex] = Lambda;
+	RotationDrives.D.ConstraintLambda[ConstraintIndex] = Lambda;
+
+	SolveRotationConstraintDelta(ConstraintIndex, DeltaLambda, true, RotationDrives);
+}
+
+void FPBDJointCachedSolver::ApplyRotationDrivesSimd(
+	const FReal Dtd)
+{
+	const VectorRegister4Float Body0DQ = MakeVectorRegisterFloatFromDouble(MakeVectorRegister(Body(0).DQ()[0], Body(0).DQ()[1], Body(0).DQ()[2], 0.0f));
+	const VectorRegister4Float Body1DQ = MakeVectorRegisterFloatFromDouble(MakeVectorRegister(Body(1).DQ()[0], Body(1).DQ()[1], Body(1).DQ()[2], 0.0f));
+
+	const VectorRegister4Float TargetW = MakeVectorRegisterFloatFromDouble(MakeVectorRegister(RotationDrives.ConstraintVX[0], RotationDrives.ConstraintVX[1], RotationDrives.ConstraintVX[2], 0.0f));
+
+	const VectorRegister4Float DQDiff = VectorSubtract(Body1DQ, Body0DQ);
+	VectorRegister4Float DXArray[3];
+	VectorRegister4Float CVDtArray[3];
+	const VectorRegister4Float W0Dt = VectorAdd(Body0DQ, ConnectorWDtsSimd[0]);
+	const VectorRegister4Float W1Dt = VectorAdd(Body1DQ, ConnectorWDtsSimd[1]);
+	const VectorRegister4Float WDiff = VectorSubtract(W0Dt, W1Dt);
+
+	for (int32 ConstraintIndex = 0; ConstraintIndex < 3; ConstraintIndex++)
+	{
+		DXArray[ConstraintIndex] = Private::VectorDot3FastX(DQDiff, RotationDrives.Simd.ConstraintAxisSimd[ConstraintIndex]);
+		CVDtArray[ConstraintIndex] = Private::VectorDot3FastX(RotationDrives.Simd.ConstraintAxisSimd[ConstraintIndex], WDiff);
+	}
+	VectorRegister4Float DX = VectorUnpackLo(DXArray[0], DXArray[1]);
+	DX = VectorMoveLh(DX, DXArray[2]);
+	const VectorRegister4Float CX = VectorAdd(RotationDrives.Simd.ConstraintCX, DX);
+
+	VectorRegister4Float CVDt = VectorUnpackLo(CVDtArray[0], CVDtArray[1]);
+	CVDt = VectorMoveLh(CVDt, CVDtArray[2]);
+
+	FRealSingle Dtf = FRealSingle(Dtd);
+	VectorRegister4Float Dt = VectorLoadFloat1(&Dtf);
+
+	CVDt = VectorMultiplyAdd(TargetW, Dt, CVDt);
+
+	const FRealSingle SolverStiffnessf = FRealSingle(SolverStiffness);
+	VectorRegister4Float Stiffness = VectorLoadFloat1(&SolverStiffnessf);
+	VectorRegister4Float DeltaLambda = VectorMultiply(Stiffness, VectorDivide(VectorSubtract(VectorMultiply(RotationDrives.Simd.ConstraintSoftStiffnessSimd, CX),
+		VectorMultiplyAdd(RotationDrives.Simd.ConstraintSoftDampingSimd, CVDt, RotationDrives.Simd.ConstraintLambda)), RotationDrives.Simd.ConstraintSoftIMSimd));
+	RotationDrives.Simd.ConstraintLambda = VectorAdd(RotationDrives.Simd.ConstraintLambda, DeltaLambda);
+
+	// Should check MaxLambda eventually
+
+	VectorRegister4Float DeltaLambdas[3];
+	DeltaLambdas[0] = VectorReplicate(DeltaLambda, 0);
+	DeltaLambdas[1] = VectorReplicate(DeltaLambda, 1);
+	DeltaLambdas[2] = VectorReplicate(DeltaLambda, 2);
+
+	VectorRegister4Float DeltaImpulses[3];
+	for (int32 ConstraintIndex = 0; ConstraintIndex < 3; ConstraintIndex++)
+	{
+		DeltaImpulses[ConstraintIndex] = VectorMultiply(RotationDrives.Simd.ConstraintAxisSimd[ConstraintIndex], DeltaLambdas[ConstraintIndex]);
+	}
 
-	SolveRotationConstraintDelta(ConstraintIndex, DeltaLambda,true, RotationDrives);
+	if (Body(0).IsDynamic())
+	{
+		VectorRegister4Float DR0 = VectorZeroFloat();
+		for (int32 ConstraintIndex = 0; ConstraintIndex < 3; ConstraintIndex++)
+		{
+			const VectorRegister4Float Axis = VectorDot3(RotationDrives.Simd.ConstraintAxisSimd[ConstraintIndex], RotationDrives.Simd.ConstraintDRAxisSimd[ConstraintIndex][0]);
+			DR0 = VectorMultiplyAdd(DeltaImpulses[ConstraintIndex], Axis, DR0);
+		}
+		FVec3f DR0f;
+		VectorStoreFloat3(DR0, &DR0f[0]);
+		ApplyRotationDelta(0, FVec3(DR0f));
+	}
+	if (Body(1).IsDynamic())
+	{
+		VectorRegister4Float DR1 = VectorZeroFloat();
+		for (int32 ConstraintIndex = 0; ConstraintIndex < 3; ConstraintIndex++)
+		{
+			const VectorRegister4Float Axis = VectorDot3(RotationDrives.Simd.ConstraintAxisSimd[ConstraintIndex], RotationDrives.Simd.ConstraintDRAxisSimd[ConstraintIndex][1]);
+			DR1 = VectorMultiplyAdd(DeltaImpulses[ConstraintIndex], Axis, DR1);
+		}
+		FVec3f DR1f;
+		VectorStoreFloat3(DR1, &DR1f[0]);
+		ApplyRotationDelta(1, FVec3(DR1f));
+	}
+	NumActiveConstraints += 3;
 }
 
 void FPBDJointCachedSolver::ApplyRotationVelocityDrives(
@@ -1924,6 +2918,7 @@ void FPBDJointCachedSolver::ApplyAxisRotationVelocityDrive(
 	const int32 ConstraintIndex,
 	const FReal Dt)
 {
+	check(RotationDrives.bUseSimd == false);
 	// NOTE: Using the actual damping, not the PBD modified value
 	const FReal Damping = RotationDrives.SettingsSoftDamping[ConstraintIndex] * Dt;
 	if (Damping < UE_SMALL_NUMBER)
@@ -1931,11 +2926,11 @@ void FPBDJointCachedSolver::ApplyAxisRotationVelocityDrive(
 		return;
 	}
 
-	const FReal MassScale = RotationDrives.bAccelerationMode ? (FReal(1) / RotationDrives.ConstraintHardIM[ConstraintIndex]) : FReal(1);
-	const FReal IM = MassScale * Damping * RotationDrives.ConstraintHardIM[ConstraintIndex] + (FReal)1;
+	const FReal MassScale = RotationDrives.bAccelerationMode ? (FReal(1) / RotationDrives.D.ConstraintHardIM[ConstraintIndex]) : FReal(1);
+	const FReal IM = MassScale * Damping * RotationDrives.D.ConstraintHardIM[ConstraintIndex] + (FReal)1;
 
 	// Angular velocity error to correct
-	const FReal WRel = FVec3::DotProduct(W(1) - W(0), RotationDrives.ConstraintAxis[ConstraintIndex]);
+	const FReal WRel = FVec3::DotProduct(W(1) - W(0), RotationDrives.D.ConstraintAxis[ConstraintIndex]);
 	const FReal TargetW = RotationDrives.ConstraintVX[ConstraintIndex];
 	const FReal CV = (WRel - TargetW);
 
@@ -1946,27 +2941,27 @@ void FPBDJointCachedSolver::ApplyAxisRotationVelocityDrive(
 	// Apply limits and accumulate total impulse 
 	// (NOTE: Limits and net impulses are position based, not velocity based)
 	FReal DeltaLambda = DeltaLambdaVel * Dt;
-	FReal Lambda = RotationDrives.ConstraintLambda[ConstraintIndex] + DeltaLambda;
+	FReal Lambda = RotationDrives.D.ConstraintLambda[ConstraintIndex] + DeltaLambda;
 	RotationDrives.ApplyMaxLambda(ConstraintIndex, DeltaLambda, Lambda);
-	RotationDrives.ConstraintLambda[ConstraintIndex] = Lambda;
+	RotationDrives.D.ConstraintLambda[ConstraintIndex] = Lambda;
 	DeltaLambdaVel = DeltaLambda / Dt;
 
 	LambdaVel += DeltaLambdaVel;
-	const FVec3 Impulse = DeltaLambdaVel * RotationDrives.ConstraintAxis[ConstraintIndex];
+	const FVec3 Impulse = DeltaLambdaVel * RotationDrives.D.ConstraintAxis[ConstraintIndex];
 
 	if (Body(0).IsDynamic())
 	{
-		const FVec3 DW0 = RotationDrives.ConstraintDRAxis[ConstraintIndex][0] * DeltaLambdaVel;
+		const FVec3 DW0 = RotationDrives.D.ConstraintDRAxis[ConstraintIndex][0] * DeltaLambdaVel;
 		Body(0).ApplyAngularVelocityDelta(DW0);
 	}
 	if (Body(1).IsDynamic())
 	{
-		const FVec3 DW1 = RotationDrives.ConstraintDRAxis[ConstraintIndex][1] * DeltaLambdaVel;
+		const FVec3 DW1 = RotationDrives.D.ConstraintDRAxis[ConstraintIndex][1] * DeltaLambdaVel;
 		Body(1).ApplyAngularVelocityDelta(DW1);
 	}
 }
 
- // Joint utilities
+// Joint utilities
 
 void FPBDJointCachedSolver::ApplyPositionDelta(
 	const int32 BodyIndex,
@@ -1991,9 +2986,9 @@ void FAxisConstraintDatas::InitDatas(
 	const bool bResetLambdas)
 {
 	bSoftLimit[ConstraintIndex] = bHasSoftLimits;
-	ConstraintHardStiffness[ConstraintIndex] = HardStiffness;
-	ConstraintSoftStiffness[ConstraintIndex] = SoftStiffness;
-	ConstraintSoftDamping[ConstraintIndex] = SoftDamping;
+	D.ConstraintHardStiffness[ConstraintIndex] = HardStiffness;
+	D.ConstraintSoftStiffness[ConstraintIndex] = SoftStiffness;
+	D.ConstraintSoftDamping[ConstraintIndex] = SoftDamping;
 	ConstraintMaxLambda[ConstraintIndex] = 0;
 	SettingsSoftStiffness[ConstraintIndex] = SoftStiffness;
 	SettingsSoftDamping[ConstraintIndex] = SoftDamping;
@@ -2002,9 +2997,9 @@ void FAxisConstraintDatas::InitDatas(
 	MotionType[ConstraintIndex] = EJointMotionType::Free;
 	if (bResetLambdas)
 	{
-		ConstraintLambda = FVec3::Zero();
+		D.ConstraintLambda = FVec3::Zero();
 		ConstraintLambdaVelocity = FVec3::Zero();
-		ConstraintLimits = FVec3::Zero();
+		D.ConstraintLimits = FVec3::Zero();
 	}
 }
 
@@ -2014,40 +3009,44 @@ void FAxisConstraintDatas::UpdateDatas(
 	const FReal DatasCX,
 	const FReal DatasRestitution,
 	const bool bCheckLimit,
-	const FVec3& DatasArm0 ,
-	const FVec3& DatasArm1 ,
+	const FVec3& DatasArm0,
+	const FVec3& DatasArm1,
 	const FReal DatasVX)
 {
 	bValidDatas[ConstraintIndex] = true;
 	bLimitsCheck[ConstraintIndex] = bCheckLimit;
 
-	ConstraintCX[ConstraintIndex] = DatasCX;
+	D.ConstraintCX[ConstraintIndex] = DatasCX;
 	ConstraintVX[ConstraintIndex] = DatasVX;
-	ConstraintAxis[ConstraintIndex] = DatasAxis;
 	ConstraintRestitution[ConstraintIndex] = DatasRestitution;
-	ConstraintArms[ConstraintIndex][0] = DatasArm0;
-	ConstraintArms[ConstraintIndex][1] = DatasArm1;
+
+
+	D.ConstraintArms[ConstraintIndex][0] = DatasArm0;
+	D.ConstraintArms[ConstraintIndex][1] = DatasArm1;
+	D.ConstraintAxis[ConstraintIndex] = DatasAxis;
 }
 
 void FAxisConstraintDatas::UpdateMass(
-		const int32 ConstraintIndex,
-		const FVec3& DatasIA0,
-		const FVec3& DatasIA1,
-		const FReal DatasIM,
-		const FReal Dt,
-		const bool bUsePositionBasedDrives)
+	const int32 ConstraintIndex,
+	const FVec3& DatasIA0,
+	const FVec3& DatasIA1,
+	const FReal DatasIM,
+	const FReal Dt,
+	const bool bUsePositionBasedDrives)
 {
-	ConstraintDRAxis[ConstraintIndex][0] = DatasIA0;
-	ConstraintDRAxis[ConstraintIndex][1] = -DatasIA1;
-	ConstraintHardIM[ConstraintIndex] = DatasIM;
+	D.ConstraintHardIM[ConstraintIndex] = DatasIM;
+
 
-	if(bSoftLimit[ConstraintIndex])
+	D.ConstraintDRAxis[ConstraintIndex][0] = DatasIA0;
+	D.ConstraintDRAxis[ConstraintIndex][1] = -DatasIA1;
+
+	if (bSoftLimit[ConstraintIndex])
 	{
 		// If bUsePositionBasedDrives is false, we apply the velocity drive in the velocity solver phase so we don't include in the PBD settings
-		const FReal SpringMassScale = (bAccelerationMode) ? (FReal)1 / (ConstraintHardIM[ConstraintIndex]) : (FReal)1;
-		ConstraintSoftStiffness[ConstraintIndex] = SpringMassScale * SettingsSoftStiffness[ConstraintIndex] * Dt * Dt;
-		ConstraintSoftDamping[ConstraintIndex] = (bUsePositionBasedDrives) ? SpringMassScale * SettingsSoftDamping[ConstraintIndex] * Dt : 0;
-		ConstraintSoftIM[ConstraintIndex] = (ConstraintSoftStiffness[ConstraintIndex] + ConstraintSoftDamping[ConstraintIndex]) * ConstraintHardIM[ConstraintIndex] + (FReal)1;
+		const FReal SpringMassScale = (bAccelerationMode) ? (FReal)1 / (D.ConstraintHardIM[ConstraintIndex]) : (FReal)1;
+		D.ConstraintSoftStiffness[ConstraintIndex] = SpringMassScale * SettingsSoftStiffness[ConstraintIndex] * Dt * Dt;
+		D.ConstraintSoftDamping[ConstraintIndex] = (bUsePositionBasedDrives) ? SpringMassScale * SettingsSoftDamping[ConstraintIndex] * Dt : 0;
+		D.ConstraintSoftIM[ConstraintIndex] = (D.ConstraintSoftStiffness[ConstraintIndex] + D.ConstraintSoftDamping[ConstraintIndex]) * D.ConstraintHardIM[ConstraintIndex] + (FReal)1;
 	}
 }
 
@@ -2065,7 +3064,7 @@ void FAxisConstraintDatas::SetMaxForce(
 		FReal MaxLambda = InMaxForce * Dt * Dt;
 		if (bAccelerationMode)
 		{
-			MaxLambda /= ConstraintHardIM[ConstraintIndex];
+			MaxLambda /= D.ConstraintHardIM[ConstraintIndex];
 		}
 		ConstraintMaxLambda[ConstraintIndex] = MaxLambda;
 	}
@@ -2080,18 +3079,17 @@ void FAxisConstraintDatas::ApplyMaxLambda(
 	{
 		if (Lambda > ConstraintMaxLambda[ConstraintIndex])
 		{
-			DeltaLambda = ConstraintMaxLambda[ConstraintIndex] - ConstraintLambda[ConstraintIndex];
+			DeltaLambda = ConstraintMaxLambda[ConstraintIndex] - D.ConstraintLambda[ConstraintIndex];
 			Lambda = ConstraintMaxLambda[ConstraintIndex];
 		}
 		else if (Lambda < -ConstraintMaxLambda[ConstraintIndex])
 		{
-			DeltaLambda = -ConstraintMaxLambda[ConstraintIndex] - ConstraintLambda[ConstraintIndex];
+			DeltaLambda = -ConstraintMaxLambda[ConstraintIndex] - D.ConstraintLambda[ConstraintIndex];
 			Lambda = -ConstraintMaxLambda[ConstraintIndex];
 		}
 	}
 }
 
-	
 }
 
 
diff --git a/Engine/Source/Runtime/Experimental/Chaos/Private/Chaos/PBDJointConstraints.cpp b/Engine/Source/Runtime/Experimental/Chaos/Private/Chaos/PBDJointConstraints.cpp
--- a/Engine/Source/Runtime/Experimental/Chaos/Private/Chaos/PBDJointConstraints.cpp
+++ b/Engine/Source/Runtime/Experimental/Chaos/Private/Chaos/PBDJointConstraints.cpp
@@ -479,6 +479,7 @@ namespace Chaos
 		, NumIterationsAtMaxSolverStiffness(1)
 		, NumShockPropagationIterations(0)
 		, bUseLinearSolver(true)
+		, bUseSimd(false)
 		, bSortEnabled(false)
 		, bSolvePositionLast(true)
 		, bUsePositionBasedDrives(true)
diff --git a/Engine/Source/Runtime/Experimental/Chaos/Public/Chaos/Joint/PBDJointCachedSolverGaussSeidel.h b/Engine/Source/Runtime/Experimental/Chaos/Public/Chaos/Joint/PBDJointCachedSolverGaussSeidel.h
--- a/Engine/Source/Runtime/Experimental/Chaos/Public/Chaos/Joint/PBDJointCachedSolverGaussSeidel.h
+++ b/Engine/Source/Runtime/Experimental/Chaos/Public/Chaos/Joint/PBDJointCachedSolverGaussSeidel.h
@@ -16,6 +16,12 @@ namespace Chaos
 /** Cached axis joint datas that will be used during the apply */
 struct FAxisConstraintDatas
 {
+	FAxisConstraintDatas() {}
+	FAxisConstraintDatas(const FAxisConstraintDatas&) = delete;
+	FAxisConstraintDatas& operator=(const FAxisConstraintDatas&) = delete;
+	FAxisConstraintDatas(FAxisConstraintDatas&&) = delete;
+	FAxisConstraintDatas& operator=(FAxisConstraintDatas&&) = delete;
+
 	/** Init the axis joint datas with stiffness / damping */
 	void InitDatas(
 		const int32 ConstraintIndex,
@@ -57,31 +63,63 @@ struct FAxisConstraintDatas
 		FReal& DeltaLambda, 
 		FReal& Lambda);
 	
-	FVec3 ConstraintHardStiffness;
-	FVec3 ConstraintSoftStiffness;
-	FVec3 ConstraintSoftDamping;
-	FVec3 ConstraintMaxLambda;
+	
+
+	struct alignas(16) FDataSimd
+	{
+		FDataSimd() 
+			: ConstraintSoftStiffnessSimd(VectorZeroFloat())
+			, ConstraintSoftDampingSimd(VectorZeroFloat())
+			, ConstraintLimitsSimd(VectorZeroFloat())
+		{}
+
+		VectorRegister4Float ConstraintHardStiffnessSimd;
+		VectorRegister4Float ConstraintSoftStiffnessSimd;
+		VectorRegister4Float ConstraintSoftDampingSimd;
+		VectorRegister4Float ConstraintArmsSimd[3][2];
+		VectorRegister4Float ConstraintAxisSimd[3];
+		VectorRegister4Float ConstraintLimitsSimd;
+		VectorRegister4Float ConstraintSoftIMSimd;
+		VectorRegister4Float ConstraintHardIMSimd;
+		VectorRegister4Float ConstraintDRAxisSimd[3][2];
+		VectorRegister4Float ConstraintCX;
+		VectorRegister4Float ConstraintLambda;
+	};
 
+	struct FData
+	{
+		FVec3 ConstraintHardStiffness;
+		FVec3 ConstraintSoftStiffness;
+		FVec3 ConstraintSoftDamping;
+		FVec3 ConstraintArms[3][2];
+		FVec3 ConstraintAxis[3];
+		FVec3 ConstraintLimits;
+		FVec3 ConstraintSoftIM;
+		FVec3 ConstraintHardIM;
+		FVec3 ConstraintDRAxis[3][2];
+		FVec3 ConstraintCX;
+		FVec3 ConstraintLambda;
+	};
+
+	union
+	{
+		FDataSimd Simd;
+		FData D;
+	};
+
+	FVec3 ConstraintMaxLambda;
 	FVec3 SettingsSoftDamping;
 	FVec3 SettingsSoftStiffness;
-	
-	FVec3 ConstraintArms[3][2];
-	FVec3 ConstraintAxis[3];
-	FVec3 ConstraintLimits;
-	
-	FVec3 ConstraintCX;
 	FVec3 ConstraintVX;
-	FVec3 ConstraintSoftIM;
-	FVec3 ConstraintHardIM;
-	FVec3 ConstraintDRAxis[3][2];
+
 	FVec3 ConstraintRestitution;
-	FVec3 ConstraintLambda;
 	FVec3 ConstraintLambdaVelocity;
 	
 	bool bValidDatas[3];
 	bool bLimitsCheck[3];
-	bool bAccelerationMode;
 	bool bSoftLimit[3];
+	bool bAccelerationMode;
+	bool bUseSimd;
 	
 	EJointMotionType MotionType[3];
 };
@@ -190,15 +228,45 @@ struct FAxisConstraintDatas
 		inline FVec3 GetNetLinearImpulse() const
 		{
 			FVec3 Impulse = FVec3(0);
+			if (PositionConstraints.bUseSimd || PositionDrives.bUseSimd)
+			{
+				VectorRegister4Float ImpulseSimd = VectorZeroFloat();
+				if (PositionConstraints.bUseSimd)
+				{
+					VectorRegister4Float ConstraintLambdas[3];
+					ConstraintLambdas[0] = VectorReplicate(PositionConstraints.Simd.ConstraintLambda, 0);
+					ConstraintLambdas[1] = VectorReplicate(PositionConstraints.Simd.ConstraintLambda, 1);
+					ConstraintLambdas[2] = VectorReplicate(PositionConstraints.Simd.ConstraintLambda, 2);
+					for (int32 Axis = 0; Axis < 3; ++Axis)
+					{
+						ImpulseSimd = VectorMultiplyAdd(ConstraintLambdas[Axis], PositionConstraints.Simd.ConstraintAxisSimd[Axis], ImpulseSimd);
+					}
+				}
+				if (PositionDrives.bUseSimd)
+				{
+					VectorRegister4Float ConstraintLambdas[3];
+					ConstraintLambdas[0] = VectorReplicate(PositionDrives.Simd.ConstraintLambda, 0);
+					ConstraintLambdas[1] = VectorReplicate(PositionDrives.Simd.ConstraintLambda, 1);
+					ConstraintLambdas[2] = VectorReplicate(PositionDrives.Simd.ConstraintLambda, 2);
+					for (int32 Axis = 0; Axis < 3; ++Axis)
+					{
+						ImpulseSimd = VectorMultiplyAdd(ConstraintLambdas[Axis], PositionDrives.Simd.ConstraintAxisSimd[Axis], ImpulseSimd);
+					}
+				}
+				FVec3f Impulsef;
+				VectorStoreFloat3(ImpulseSimd, &Impulsef[0]);
+				Impulse = FVec3(Impulsef);
+			}
+
 			for (int32 Axis = 0; Axis < 3; ++Axis)
 			{
-				if (PositionConstraints.bValidDatas[Axis])
+				if (!PositionConstraints.bUseSimd && PositionConstraints.bValidDatas[Axis])
 				{
-					Impulse += PositionConstraints.ConstraintLambda[Axis] * PositionConstraints.ConstraintAxis[Axis];
+					Impulse += PositionConstraints.D.ConstraintLambda[Axis] * PositionConstraints.D.ConstraintAxis[Axis];
 				}
-				if (PositionDrives.bValidDatas[Axis])
+				if (!PositionDrives.bUseSimd && PositionDrives.bValidDatas[Axis])
 				{
-					Impulse += PositionDrives.ConstraintLambda[Axis] * PositionDrives.ConstraintAxis[Axis];
+					Impulse += PositionDrives.D.ConstraintLambda[Axis] * PositionDrives.D.ConstraintAxis[Axis];
 				}
 			}
 			return Impulse;
@@ -208,15 +276,45 @@ struct FAxisConstraintDatas
 		inline FVec3 GetNetAngularImpulse() const
 		{
 			FVec3 Impulse = FVec3(0);
+			if (RotationConstraints.bUseSimd || RotationDrives.bUseSimd)
+			{
+				VectorRegister4Float ImpulseSimd = VectorZeroFloat();
+				if (RotationConstraints.bUseSimd)
+				{
+					VectorRegister4Float ConstraintLambdas[3];
+					ConstraintLambdas[0] = VectorReplicate(RotationConstraints.Simd.ConstraintLambda, 0);
+					ConstraintLambdas[1] = VectorReplicate(RotationConstraints.Simd.ConstraintLambda, 1);
+					ConstraintLambdas[2] = VectorReplicate(RotationConstraints.Simd.ConstraintLambda, 2);
+					for (int32 Axis = 0; Axis < 3; ++Axis)
+					{
+						ImpulseSimd = VectorMultiplyAdd(ConstraintLambdas[Axis], RotationConstraints.Simd.ConstraintAxisSimd[Axis], ImpulseSimd);
+					}
+				}
+				if (RotationDrives.bUseSimd)
+				{
+					VectorRegister4Float ConstraintLambdas[3];
+					ConstraintLambdas[0] = VectorReplicate(RotationDrives.Simd.ConstraintLambda, 0);
+					ConstraintLambdas[1] = VectorReplicate(RotationDrives.Simd.ConstraintLambda, 1);
+					ConstraintLambdas[2] = VectorReplicate(RotationDrives.Simd.ConstraintLambda, 2);
+					for (int32 Axis = 0; Axis < 3; ++Axis)
+					{
+						ImpulseSimd = VectorMultiplyAdd(ConstraintLambdas[Axis], RotationDrives.Simd.ConstraintAxisSimd[Axis], ImpulseSimd);
+					}
+				}
+				FVec3f Impulsef;
+				VectorStoreFloat3(ImpulseSimd, &Impulsef[0]);
+				Impulse = FVec3(Impulsef);
+			}
+
 			for (int32 Axis = 0; Axis < 3; ++Axis)
 			{
-				if (RotationConstraints.bValidDatas[Axis])
+				if (!RotationConstraints.bUseSimd && RotationConstraints.bValidDatas[Axis])
 				{
-					Impulse += RotationConstraints.ConstraintLambda[Axis] * RotationConstraints.ConstraintAxis[Axis];
+					Impulse += RotationConstraints.D.ConstraintLambda[Axis] * RotationConstraints.D.ConstraintAxis[Axis];
 				}
-				if (RotationDrives.bValidDatas[Axis])
+				if (!RotationDrives.bUseSimd && RotationDrives.bValidDatas[Axis])
 				{
-					Impulse += RotationDrives.ConstraintLambda[Axis] * RotationDrives.ConstraintAxis[Axis];
+					Impulse += RotationDrives.D.ConstraintLambda[Axis] * RotationDrives.D.ConstraintAxis[Axis];
 				}
 			}
 			return Impulse;
@@ -291,11 +389,27 @@ struct FAxisConstraintDatas
 			const FPBDJointSolverSettings& SolverSettings,
 			const FPBDJointSettings& JointSettings);
 
+		void ApplyAxisPositionProjection(
+			const FReal LinearProjection,
+			const int32 ConstraintIndex);
+
+		void ApplyPositionProjectionSimd(
+			const FReal LinearProjection);
+
 		void ApplyRotationProjection(
 			const FReal Dt,
 			const FPBDJointSolverSettings& SolverSettings,
 			const FPBDJointSettings& JointSettings);
 
+		void ApplyAxisRotationProjection(
+			const FReal AngularProjection,
+			const bool bLinearLocked,
+			const int32 ConstraintIndex);
+
+		void ApplyRotationProjectionSimd(
+			const FRealSingle AngularProjection,
+			const bool bLinearLocked);
+
 		void ApplyTeleports(
 			const FReal Dt,
 			const FPBDJointSolverSettings& SolverSettings,
@@ -306,6 +420,13 @@ struct FAxisConstraintDatas
 			const FPBDJointSolverSettings& SolverSettings,
 			const FPBDJointSettings& JointSettings);
 
+		void ApplyAxisPositionTeleport(
+			const FReal TeleportDistance,
+			const int32 ConstraintIndex);
+
+		void ApplyPositionTeleportSimd(
+			const FRealSingle TeleportDistance);
+
 		void ApplyRotationTeleport(
 			const FReal Dt,
 			const FPBDJointSolverSettings& SolverSettings,
@@ -387,6 +508,11 @@ struct FAxisConstraintDatas
 			const FReal Dt,
 			const TVec3<EJointMotionType>& LinearMotion);
 
+		void InitLockedPositionConstraintSimd(
+			const FPBDJointSettings& JointSettings,
+			const FReal Dt,
+			const TVec3<EJointMotionType>& LinearMotion);
+
 		void InitSphericalPositionConstraint(
 			const FPBDJointSettings& JointSettings,
 			const FReal Dt);
@@ -409,6 +535,9 @@ struct FAxisConstraintDatas
 		void ApplyAxisPositionConstraint(
 			const int32 ConstraintIndex,
 			const FReal Dt);
+
+		void ApplyPositionConstraintsSimd(
+			const FReal Dt);
 		
 		void SolvePositionConstraintDelta(
 			const int32 ConstraintIndex, 
@@ -432,6 +561,8 @@ struct FAxisConstraintDatas
 		void ApplyAxisVelocityConstraint(
 			const int32 ConstraintIndex);
 
+		void ApplyVelocityConstraintSimd();
+
 		void SolveLinearVelocityConstraint(
 			const int32 ConstraintIndex,
 			const FReal TargetVel);
@@ -443,7 +574,11 @@ struct FAxisConstraintDatas
 			const FPBDJointSolverSettings& SolverSettings,
 			const FPBDJointSettings& JointSettings,
 			const bool bResetLambdas);
-		
+
+		void InitRotationConstraintsSimd(
+			const FPBDJointSettings& JointSettings,
+			const FReal Dt);
+
 		void CorrectAxisAngleConstraint(
 			const FPBDJointSettings& JointSettings,
 			const int32 ConstraintIndex,
@@ -513,6 +648,9 @@ struct FAxisConstraintDatas
 			const int32 ConstraintIndex,
 			const FReal Dt);
 
+		void ApplyRotationSoftConstraintsSimd(
+			const FReal Dt);
+
 		void SolveRotationConstraintDelta(
 			const int32 ConstraintIndex, 
 			const FReal DeltaLambda,
@@ -540,6 +678,8 @@ struct FAxisConstraintDatas
 		void ApplyAngularVelocityConstraint(
 			const int32 ConstraintIndex);
 
+		void ApplyAngularVelocityConstraintSimd();
+
 		/** Init Position Drives */
 		
 		void InitPositionDrives(
@@ -582,6 +722,11 @@ struct FAxisConstraintDatas
 			const FVec3& ConstraintAxis,
 			const FReal Dt,
 			const FReal DeltaAngle);
+
+		void InitRotationConstraintDriveSimd(
+			FVec3 ConstraintAxes[3],
+			const FReal Dt,
+			const FVec3 DeltaAngles);
 		
 		void InitSwingTwistDrives(
 			const FReal Dt,
@@ -605,6 +750,9 @@ struct FAxisConstraintDatas
 			const int32 ConstraintIndex,
 			const FReal Dt);
 
+		void ApplyRotationDrivesSimd(
+			const FReal Dt);
+
 		void ApplyRotationVelocityDrives(
 			const FReal Dt);
 
@@ -627,6 +775,7 @@ struct FAxisConstraintDatas
 		FVec3 ConnectorXs[MaxConstrainedBodies];			// World-space joint connector positions
 		FRotation3 ConnectorRs[MaxConstrainedBodies];		// World-space joint connector rotations
 		FVec3 ConnectorWDts[MaxConstrainedBodies];			// World-space joint connector angular velocities * dt
+		VectorRegister4Float ConnectorWDtsSimd[MaxConstrainedBodies];			// World-space joint connector angular velocities * dt
 
 		// XPBD Initial iteration world-space body state
 		FVec3 InitConnectorXs[MaxConstrainedBodies];		// World-space joint connector positions
@@ -650,8 +799,8 @@ struct FAxisConstraintDatas
 		FVec3 LastDQs[MaxConstrainedBodies];				// Rotations at the beginning of the iteration
 		FVec3 CurrentPs[MaxConstrainedBodies];				// Positions at the beginning of the iteration
 		FRotation3 CurrentQs[MaxConstrainedBodies];			// Rotations at the beginning of the iteration
-		FVec3 InitConstraintAxisLinearVelocities;			// Linear velocities along the constraint axes at the begining of the frame, used by restitution
-		FVec3 InitConstraintAxisAngularVelocities;			// Angular velocities along the constraint axes at the begining of the frame, used by restitution
+		FVec3 InitConstraintAxisLinearVelocities;			// Linear velocities along the constraint axes at the beginning of the frame, used by restitution
+		FVec3 InitConstraintAxisAngularVelocities;			// Angular velocities along the constraint axes at the beginning of the frame, used by restitution
 		int32 NumActiveConstraints;							// The number of active constraints and drives in the last iteration (-1 initial value)
 		bool bIsActive;										// Whether any constraints actually moved any bodies in last iteration
 		bool bUsePositionBasedDrives;						// Whether to apply velocity drive in the PBD step of the VBD step
@@ -663,6 +812,7 @@ struct FAxisConstraintDatas
 		FAxisConstraintDatas RotationDrives;
 
 		bool bIsBroken;
+		bool bUseSimd;
 
 		// dummy indices
 		static constexpr int32 PointPositionConstraintIndex = 0;
diff --git a/Engine/Source/Runtime/Experimental/Chaos/Public/Chaos/PBDJointConstraintTypes.h b/Engine/Source/Runtime/Experimental/Chaos/Public/Chaos/PBDJointConstraintTypes.h
--- a/Engine/Source/Runtime/Experimental/Chaos/Public/Chaos/PBDJointConstraintTypes.h
+++ b/Engine/Source/Runtime/Experimental/Chaos/Public/Chaos/PBDJointConstraintTypes.h
@@ -223,6 +223,9 @@ namespace Chaos
 		// Whether to use the linear or non-linear joint solver
 		bool bUseLinearSolver;
 
+		// Whether use vectorization in single precision, only available on the linear solver
+		bool bUseSimd;
+
 		// Whether the joints need to be sorted (only required for RBAN - the world solver uses the constraint graph for ordering)
 		bool bSortEnabled;
 
diff --git a/Engine/Source/Runtime/Experimental/Chaos/Public/Chaos/PBDJointConstraints.h b/Engine/Source/Runtime/Experimental/Chaos/Public/Chaos/PBDJointConstraints.h
--- a/Engine/Source/Runtime/Experimental/Chaos/Public/Chaos/PBDJointConstraints.h
+++ b/Engine/Source/Runtime/Experimental/Chaos/Public/Chaos/PBDJointConstraints.h
@@ -175,6 +175,11 @@ namespace Chaos
 		*/
 		void SetUseLinearJointSolver(const bool bInEnable) { Settings.bUseLinearSolver = bInEnable; }
 
+		/**
+		* Whether to use simd in single precision on the linear joint solver. 
+		*/
+		void SetUseSimd(const bool bInEnable) { Settings.bUseSimd = bInEnable; }
+
 		/**
 		 * Whether to sort the joints internally. Sort will be triggered on any tick when a joint was added. Only needed for RBAN.
 		*/
diff --git a/Engine/Source/Runtime/Experimental/ChaosCore/Public/Chaos/VectorUtility.h b/Engine/Source/Runtime/Experimental/ChaosCore/Public/Chaos/VectorUtility.h
--- a/Engine/Source/Runtime/Experimental/ChaosCore/Public/Chaos/VectorUtility.h
+++ b/Engine/Source/Runtime/Experimental/ChaosCore/Public/Chaos/VectorUtility.h
@@ -2,6 +2,7 @@
 #pragma once
 
 #include "Math/VectorRegister.h"
+#include "Chaos/Core.h"
 
 /**
  * Cast VectorRegister4Int in VectorRegister4Float
@@ -93,6 +94,43 @@ FORCEINLINE VectorRegister4Float VectorMoveLh(const VectorRegister4Float& A, con
 }
 
 
+namespace Chaos::Private
+{
+	/**
+	 * Calculates the dot product of two vectors and returns a vector with the result in the first component.
+	 * W values should be set to 0 on the vector input. This function should be used with caution. Only the returned X should be read.
+	 *
+	 * @param Vec1	1st vector set with W equal to 0
+	 * @param Vec2	2nd vector set with W equal to 0
+	 * @return		d = dot3X(Vec1.xyz0, Vec2.xyz0), VectorRegister4Float( d, ?, ?, ? )
+	 */
+	FORCEINLINE VectorRegister4Float VectorDot3FastX(const VectorRegister4Float& Vec1, const VectorRegister4Float& Vec2)
+	{
+#if PLATFORM_ENABLE_VECTORINTRINSICS_NEON
+		VectorRegister4Float Temp = VectorMultiply(Vec1, Vec2); // Multiply 2 vector
+		float32x2_t sum = vpadd_f32(vget_low_f32(Temp), vget_high_f32(Temp));
+		sum = vpadd_f32(sum, sum);
+		return vcombine_f32(sum, sum);
+#elif PLATFORM_ENABLE_VECTORINTRINSICS
+		return _mm_dp_ps(Vec1, Vec2, 0xFF);
+#else
+		return VectorDot3(Vec1, Vec2);
+#endif
+	}
+
+	FORCEINLINE VectorRegister4Float VectorMatrixMultiply(const VectorRegister4Float& Vec, const FMatrix33& M)
+	{
+		const VectorRegister4Float VecX = VectorReplicate(Vec, 0);
+		const VectorRegister4Float VecY = VectorReplicate(Vec, 1);
+		const VectorRegister4Float VecZ = VectorReplicate(Vec, 2);
+
+		const VectorRegister4Float R0 = MakeVectorRegisterFloatFromDouble(MakeVectorRegister(M.M[0][0], M.M[0][1], M.M[0][2], 0.0));
+		const VectorRegister4Float R1 = MakeVectorRegisterFloatFromDouble(MakeVectorRegister(M.M[1][0], M.M[1][1], M.M[1][2], 0.0));
+		const VectorRegister4Float R2 = MakeVectorRegisterFloatFromDouble(MakeVectorRegister(M.M[2][0], M.M[2][1], M.M[2][2], 0.0));
+		return VectorMultiplyAdd(R0, VecX, VectorMultiplyAdd(R1, VecY, VectorMultiply(R2, VecZ)));
+	}
+}
+
 /**
  * Combines two vectors using bitwise NOT AND (treating each vector as a 128 bit field)
  *
