Fix for ASAN crash in Iris ObjectReplicationBridge - use cached version of FilterProfile instead of dereferencing a potentially invalid pointer due to array reallocation
#tests Rebuilt dedicated server and tested against Cooked data

diff --git a/Engine/Source/Runtime/Experimental/Iris/Core/Private/Iris/ReplicationSystem/ObjectReplicationBridge.cpp b/Engine/Source/Runtime/Experimental/Iris/Core/Private/Iris/ReplicationSystem/ObjectReplicationBridge.cpp
--- a/Engine/Source/Runtime/Experimental/Iris/Core/Private/Iris/ReplicationSystem/ObjectReplicationBridge.cpp
+++ b/Engine/Source/Runtime/Experimental/Iris/Core/Private/Iris/ReplicationSystem/ObjectReplicationBridge.cpp
@@ -1660,7 +1660,7 @@ UE::Net::FNetObjectFilterHandle UObjectReplicationBridge::GetDynamicFilter(const
 
 				if (bUseFilter)
 				{
-					OutFilterProfile = FilterInfoPtr->FilterProfile;
+					OutFilterProfile = FilterInfo.FilterProfile;
 					return FilterInfo.FilterHandle;
 				}
 				else
