Fix bug in Linear Solver
Linear Soft contraints shouldn't be solved for velocity

diff --git a/Engine/Source/Runtime/Experimental/Chaos/Private/Chaos/Joint/PBDJointCachedSolverGaussSeidel.cpp b/Engine/Source/Runtime/Experimental/Chaos/Private/Chaos/Joint/PBDJointCachedSolverGaussSeidel.cpp
--- a/Engine/Source/Runtime/Experimental/Chaos/Private/Chaos/Joint/PBDJointCachedSolverGaussSeidel.cpp
+++ b/Engine/Source/Runtime/Experimental/Chaos/Private/Chaos/Joint/PBDJointCachedSolverGaussSeidel.cpp
@@ -936,13 +936,16 @@ void FPBDJointCachedSolver::ApplyLinearVelocityConstraints()
 {
 	if (PositionConstraints.bUseSimd)
 	{
+		check(PositionConstraints.bSoftLimit[0] == false);
+		check(PositionConstraints.bSoftLimit[1] == false);
+		check(PositionConstraints.bSoftLimit[2] == false);
 		ApplyVelocityConstraintSimd();
 	}
 	else
 	{
 		for (int32 ConstraintIndex = 0; ConstraintIndex < 3; ++ConstraintIndex)
 		{
-			if (PositionConstraints.bValidDatas[ConstraintIndex])
+			if (PositionConstraints.bValidDatas[ConstraintIndex] && !PositionConstraints.bSoftLimit[ConstraintIndex])
 			{
 				ApplyAxisVelocityConstraint(ConstraintIndex);
 			}
