[Physics] Fix FPhysicsObject AddForce calls for async physics

diff --git a/Engine/Source/Runtime/Engine/Private/PrimitiveComponentPhysics.cpp b/Engine/Source/Runtime/Engine/Private/PrimitiveComponentPhysics.cpp
--- a/Engine/Source/Runtime/Engine/Private/PrimitiveComponentPhysics.cpp
+++ b/Engine/Source/Runtime/Engine/Private/PrimitiveComponentPhysics.cpp
@@ -341,29 +341,12 @@ void UPrimitiveComponent::AddTorqueInRadians(FVector Torque, FName BoneName, boo
 				return !Interface->AreAllDisabled({ &Object, 1 });
 			}
 		);
-
-		if (bAccelChange)
-		{
-			const float Mass = Interface->GetMass(PhysicsObjects);
-			Interface->AddTorque(PhysicsObjects, Torque * Mass, /*bInvalidate*/true);
-		}
-		else
-		{
-			Interface->AddTorque(PhysicsObjects, Torque, /*bInvalidate*/true);
-		}
+		Interface->AddTorque(PhysicsObjects, Torque, /*bInvalidate*/true, bAccelChange);
 	}
 	else if (Chaos::FPhysicsObject* Object = GetPhysicsObjectByName(BoneName))
 	{
 		FLockedWritePhysicsObjectExternalInterface Interface = FPhysicsObjectExternalInterface::LockWrite({ &Object, 1 });
-		if (bAccelChange)
-		{
-			const float Mass = Interface->GetMass({ &Object, 1 });
-			Interface->AddTorque({ &Object, 1 }, Torque * Mass, /*bInvalidate*/true);
-		}
-		else
-		{
-			Interface->AddTorque({ &Object, 1 }, Torque, /*bInvalidate*/true);
-		}
+		Interface->AddTorque({ &Object, 1 }, Torque, /*bInvalidate*/true, bAccelChange);
 	}
 }
 
diff --git a/Engine/Source/Runtime/Experimental/Chaos/Private/Chaos/PhysicsObjectInterface.cpp b/Engine/Source/Runtime/Experimental/Chaos/Private/Chaos/PhysicsObjectInterface.cpp
--- a/Engine/Source/Runtime/Experimental/Chaos/Private/Chaos/PhysicsObjectInterface.cpp
+++ b/Engine/Source/Runtime/Experimental/Chaos/Private/Chaos/PhysicsObjectInterface.cpp
@@ -985,10 +985,29 @@ namespace Chaos
 				}
 			}
 		}
+
+		if constexpr (Id == EThreadContext::External)
+		{
+			if (Chaos::FPhysicsSolverBase* Solver = FPhysicsObjectInterface::GetSolver(InObjects[0]))
+			{
+				Solver->EnqueueCommandImmediate(
+					[AllObjects = TArray<FPhysicsObjectHandle>{ InObjects }, Force, bInvalidate]() {
+						Chaos::FWritePhysicsObjectInterface_Internal Interface = Chaos::FPhysicsObjectInternalInterface::GetWrite();
+						Interface.AddForce(AllObjects, Force, bInvalidate);
+					}
+				);
+			}
+		}
 	}
 
 	template<EThreadContext Id>
 	void FWritePhysicsObjectInterface<Id>::AddTorque(TArrayView<const FPhysicsObjectHandle> InObjects, const FVector& Torque, bool bInvalidate)
+	{
+		AddTorque(InObjects, Torque, bInvalidate, false);
+	}
+
+	template<EThreadContext Id>
+	void FWritePhysicsObjectInterface<Id>::AddTorque(TArrayView<const FPhysicsObjectHandle> InObjects, const FVector& Torque, bool bInvalidate, bool bAccelChange)
 	{
 		for (const FPhysicsObjectHandle Object : InObjects)
 		{
@@ -1008,11 +1027,37 @@ namespace Chaos
 							SetParticleStateHelper<Id>(Object, EObjectStateType::Dynamic);
 						}
 
-						Rigid->AddTorque(Torque, bInvalidate);
+						if (bAccelChange)
+						{
+							FRotation3 QCom = Rigid->GetR() * Rigid->RotationOfMass();
+							if constexpr (Id == EThreadContext::Internal)
+							{
+								QCom = Rigid->GetQ() * Rigid->RotationOfMass();
+							}
+							const FMatrix33 WorldInertia = Utilities::ComputeWorldSpaceInertia(QCom, Rigid->I());
+							Rigid->AddTorque(WorldInertia * Torque, bInvalidate);
+						}
+						else 
+						{
+							Rigid->AddTorque(Torque, bInvalidate);
+						}
 					}
 				}
 			}
 		}
+
+		if constexpr (Id == EThreadContext::External)
+		{
+			if (Chaos::FPhysicsSolverBase* Solver = FPhysicsObjectInterface::GetSolver(InObjects[0]))
+			{
+				Solver->EnqueueCommandImmediate(
+					[AllObjects = TArray<FPhysicsObjectHandle>{ InObjects }, Torque, bInvalidate, bAccelChange]() {
+						Chaos::FWritePhysicsObjectInterface_Internal Interface = Chaos::FPhysicsObjectInternalInterface::GetWrite();
+						Interface.AddTorque(AllObjects, Torque, bInvalidate, bAccelChange);
+					}
+				);
+			}
+		}
 	}
 
 	template<EThreadContext Id>
diff --git a/Engine/Source/Runtime/Experimental/Chaos/Public/Chaos/PhysicsObjectInterface.h b/Engine/Source/Runtime/Experimental/Chaos/Public/Chaos/PhysicsObjectInterface.h
--- a/Engine/Source/Runtime/Experimental/Chaos/Public/Chaos/PhysicsObjectInterface.h
+++ b/Engine/Source/Runtime/Experimental/Chaos/Public/Chaos/PhysicsObjectInterface.h
@@ -130,10 +130,12 @@ namespace Chaos
 		CHAOS_API void WakeUp(TArrayView<const FPhysicsObjectHandle> InObjects);
 		CHAOS_API void ForceKinematic(TArrayView<const FPhysicsObjectHandle> InObjects);
 		CHAOS_API void AddForce(TArrayView<const FPhysicsObjectHandle> InObjects, const FVector& Force, bool bInvalidate);
+		UE_DEPRECATED(5.5, "This version of AddTorque has been deprecated. Please use the version with @param bAccelChange")
 		CHAOS_API void AddTorque(TArrayView<const FPhysicsObjectHandle> InObjects, const FVector& Torque, bool bInvalidate);
+		CHAOS_API void AddTorque(TArrayView<const FPhysicsObjectHandle> InObjects, const FVector& Torque, bool bInvalidate, bool bAccelChange = false);
 		CHAOS_API void SetLinearImpulseVelocity(TArrayView<const FPhysicsObjectHandle> InObjects, const FVector& Impulse, bool bVelChange);
 
-		UE_DEPRECATED(5.4, "This version AddRadialImpulse has been deprecated. Please use the version where the strain value is passed explicitly")
+		UE_DEPRECATED(5.4, "This version of AddRadialImpulse has been deprecated. Please use the version where the strain value is passed explicitly")
 		CHAOS_API void AddRadialImpulse(TArrayView<const FPhysicsObjectHandle> InObjects, FVector Origin, float Radius, float Strength, enum ERadialImpulseFalloff Falloff, bool bApplyStrain, bool bInvalidate, bool bVelChange = false);
 
 		CHAOS_API void AddRadialImpulse(TArrayView<const FPhysicsObjectHandle> InObjects, FVector Origin, float Radius, float Strength, enum ERadialImpulseFalloff Falloff, bool bApplyStrain, float Strain, bool bInvalidate, bool bVelChange = false, float MinValue = 0.f, float MaxValue = 1.f);
