Clean up Linear Solver code
-Renames couples of variables

diff --git a/Engine/Source/Runtime/Experimental/Chaos/Private/Chaos/Joint/PBDJointCachedSolverGaussSeidel.cpp b/Engine/Source/Runtime/Experimental/Chaos/Private/Chaos/Joint/PBDJointCachedSolverGaussSeidel.cpp
--- a/Engine/Source/Runtime/Experimental/Chaos/Private/Chaos/Joint/PBDJointCachedSolverGaussSeidel.cpp
+++ b/Engine/Source/Runtime/Experimental/Chaos/Private/Chaos/Joint/PBDJointCachedSolverGaussSeidel.cpp
@@ -399,7 +399,7 @@ void FPBDJointCachedSolver::InitPositionConstraints(
 		PositionConstraints.bSoftLimit[1] = false;
 		PositionConstraints.bSoftLimit[2] = false;
 		FRealSingle HardStiffness = FRealSingle(FPBDJointUtilities::GetLinearStiffness(SolverSettings, JointSettings));
-		PositionConstraints.Simd.ConstraintHardStiffnessSimd = VectorLoadFloat1(&HardStiffness);
+		PositionConstraints.Simd.ConstraintHardStiffness = VectorLoadFloat1(&HardStiffness);
 		PositionConstraints.ConstraintMaxLambda = FVec3(0);
 		PositionConstraints.bValidDatas[0] = false;
 		PositionConstraints.bValidDatas[1] = false;
@@ -414,7 +414,7 @@ void FPBDJointCachedSolver::InitPositionConstraints(
 		{
 			PositionConstraints.Simd.ConstraintLambda = VectorZeroFloat();
 			PositionConstraints.ConstraintLambdaVelocity = FVec3::Zero();
-			PositionConstraints.Simd.ConstraintLimitsSimd = VectorZeroFloat();
+			PositionConstraints.Simd.ConstraintLimits = VectorZeroFloat();
 		}
 
 	}
@@ -492,8 +492,8 @@ void FPBDJointCachedSolver::InitPositionDatasMass(
 	const int32 ConstraintIndex,
 	const FReal Dt)
 {
-	const FVec3 AngularAxis0 = FVec3::CrossProduct(PositionDatas.D.ConstraintArms[ConstraintIndex][0], PositionDatas.D.ConstraintAxis[ConstraintIndex]);
-	const FVec3 AngularAxis1 = FVec3::CrossProduct(PositionDatas.D.ConstraintArms[ConstraintIndex][1], PositionDatas.D.ConstraintAxis[ConstraintIndex]);
+	const FVec3 AngularAxis0 = FVec3::CrossProduct(PositionDatas.Data.ConstraintArms[ConstraintIndex][0], PositionDatas.Data.ConstraintAxis[ConstraintIndex]);
+	const FVec3 AngularAxis1 = FVec3::CrossProduct(PositionDatas.Data.ConstraintArms[ConstraintIndex][1], PositionDatas.Data.ConstraintAxis[ConstraintIndex]);
 	const FVec3 IA0 = Utilities::Multiply(InvI(0), AngularAxis0);
 	const FVec3 IA1 = Utilities::Multiply(InvI(1), AngularAxis1);
 	const FReal II0 = FVec3::DotProduct(AngularAxis0, IA0);
@@ -531,13 +531,13 @@ void FPBDJointCachedSolver::InitPositionConstraintDatas(
 
 	if (JointType == EJointMotionType::Locked)
 	{
-		PositionConstraints.D.ConstraintLimits[ConstraintIndex] = 0.0f;
+		PositionConstraints.Data.ConstraintLimits[ConstraintIndex] = 0.0f;
 		PositionConstraints.UpdateDatas(ConstraintIndex, LocalAxis, LocalDelta,
 			0.0, false, ConstraintArm0, ConstraintArm1);
 	}
 	else if (JointType == EJointMotionType::Limited)
 	{
-		PositionConstraints.D.ConstraintLimits[ConstraintIndex] = ConstraintLimit;
+		PositionConstraints.Data.ConstraintLimits[ConstraintIndex] = ConstraintLimit;
 		PositionConstraints.UpdateDatas(ConstraintIndex, LocalAxis, LocalDelta,
 			ConstraintRestitution, true, ConstraintArm0, ConstraintArm1);
 	}
@@ -623,28 +623,28 @@ void FPBDJointCachedSolver::InitLockedPositionConstraintSimd(
 		PositionConstraints.ConstraintVX[ConstraintIndex] = 0.0;
 		PositionConstraints.ConstraintRestitution[ConstraintIndex] = 0.0;
 
-		PositionConstraints.Simd.ConstraintArmsSimd[ConstraintIndex][0] = MakeVectorRegisterFloatFromDouble(MakeVectorRegister(ConstraintArm0[0], ConstraintArm0[1], ConstraintArm0[2], 0.0f));
-		PositionConstraints.Simd.ConstraintArmsSimd[ConstraintIndex][1] = MakeVectorRegisterFloatFromDouble(MakeVectorRegister(ConstraintArm1[0], ConstraintArm1[1], ConstraintArm1[2], 0.0f));
-		PositionConstraints.Simd.ConstraintAxisSimd[ConstraintIndex] = MakeVectorRegisterFloatFromDouble(MakeVectorRegister(LocalAxis[0], LocalAxis[1], LocalAxis[2], 0.0f));
+		PositionConstraints.Simd.ConstraintArms[ConstraintIndex][0] = MakeVectorRegisterFloatFromDouble(MakeVectorRegister(ConstraintArm0[0], ConstraintArm0[1], ConstraintArm0[2], 0.0f));
+		PositionConstraints.Simd.ConstraintArms[ConstraintIndex][1] = MakeVectorRegisterFloatFromDouble(MakeVectorRegister(ConstraintArm1[0], ConstraintArm1[1], ConstraintArm1[2], 0.0f));
+		PositionConstraints.Simd.ConstraintAxis[ConstraintIndex] = MakeVectorRegisterFloatFromDouble(MakeVectorRegister(LocalAxis[0], LocalAxis[1], LocalAxis[2], 0.0f));
 
 		InitConstraintAxisLinearVelocities[ConstraintIndex] = FVec3::DotProduct(InitConstraintVelocity, LocalAxis);
 
-		const VectorRegister4Float AngularAxis0 = VectorCross(PositionConstraints.Simd.ConstraintArmsSimd[ConstraintIndex][0], PositionConstraints.Simd.ConstraintAxisSimd[ConstraintIndex]);
+		const VectorRegister4Float AngularAxis0 = VectorCross(PositionConstraints.Simd.ConstraintArms[ConstraintIndex][0], PositionConstraints.Simd.ConstraintAxis[ConstraintIndex]);
 		const VectorRegister4Float IA0 = Private::VectorMatrixMultiply(AngularAxis0, InvI(0));
 
-		const VectorRegister4Float AngularAxis1 = VectorCross(PositionConstraints.Simd.ConstraintArmsSimd[ConstraintIndex][1], PositionConstraints.Simd.ConstraintAxisSimd[ConstraintIndex]);
+		const VectorRegister4Float AngularAxis1 = VectorCross(PositionConstraints.Simd.ConstraintArms[ConstraintIndex][1], PositionConstraints.Simd.ConstraintAxis[ConstraintIndex]);
 		const VectorRegister4Float IA1 = Private::VectorMatrixMultiply(AngularAxis1, InvI(1));
 
 		const FRealSingle II0 = VectorDot3Scalar(AngularAxis0, IA0);
 		const FRealSingle II1 = VectorDot3Scalar(AngularAxis1, IA1);
 		HardIM[ConstraintIndex] += II0 + II1;
-		PositionConstraints.Simd.ConstraintDRAxisSimd[ConstraintIndex][0] = IA0;
-		PositionConstraints.Simd.ConstraintDRAxisSimd[ConstraintIndex][1] = VectorNegate(IA1);
+		PositionConstraints.Simd.ConstraintDRAxis[ConstraintIndex][0] = IA0;
+		PositionConstraints.Simd.ConstraintDRAxis[ConstraintIndex][1] = VectorNegate(IA1);
 
 
 	}
 	PositionConstraints.Simd.ConstraintCX = MakeVectorRegisterFloat(LocalDeltas[0], LocalDeltas[1], LocalDeltas[2], 0.0f);
-	PositionConstraints.Simd.ConstraintHardIMSimd = MakeVectorRegisterFloatFromDouble(MakeVectorRegister(HardIM[0], HardIM[1], HardIM[2], 0.0f));
+	PositionConstraints.Simd.ConstraintHardIM = MakeVectorRegisterFloatFromDouble(MakeVectorRegister(HardIM[0], HardIM[1], HardIM[2], 0.0f));
 	SetInitConstraintVelocity(ConstraintArm0, ConstraintArm1);
 
 }
@@ -751,19 +751,19 @@ void FPBDJointCachedSolver::SolvePositionConstraintDelta(
 	const FReal DeltaLambda,
 	const FAxisConstraintDatas& ConstraintDatas)
 {
-	const FVec3 DX = ConstraintDatas.D.ConstraintAxis[ConstraintIndex] * DeltaLambda;
+	const FVec3 DX = ConstraintDatas.Data.ConstraintAxis[ConstraintIndex] * DeltaLambda;
 
 	if (Body(0).IsDynamic())
 	{
 		const FVec3 DP0 = InvM(0) * DX;
-		const FVec3 DR0 = ConstraintDatas.D.ConstraintDRAxis[ConstraintIndex][0] * DeltaLambda;
+		const FVec3 DR0 = ConstraintDatas.Data.ConstraintDRAxis[ConstraintIndex][0] * DeltaLambda;
 		ApplyPositionDelta(0, DP0);
 		ApplyRotationDelta(0, DR0);
 	}
 	if (Body(1).IsDynamic())
 	{
 		const FVec3 DP1 = -InvM(1) * DX;
-		const FVec3 DR1 = ConstraintDatas.D.ConstraintDRAxis[ConstraintIndex][1] * DeltaLambda;
+		const FVec3 DR1 = ConstraintDatas.Data.ConstraintDRAxis[ConstraintIndex][1] * DeltaLambda;
 		ApplyPositionDelta(1, DP1);
 		ApplyRotationDelta(1, DR1);
 	}
@@ -775,10 +775,10 @@ void FPBDJointCachedSolver::SolvePositionConstraintHard(
 	const int32 ConstraintIndex,
 	const FReal DeltaConstraint)
 {
-	const FReal DeltaLambda = SolverStiffness * PositionConstraints.D.ConstraintHardStiffness[ConstraintIndex] * DeltaConstraint /
-		PositionConstraints.D.ConstraintHardIM[ConstraintIndex];
+	const FReal DeltaLambda = SolverStiffness * PositionConstraints.Data.ConstraintHardStiffness[ConstraintIndex] * DeltaConstraint /
+		PositionConstraints.Data.ConstraintHardIM[ConstraintIndex];
 
-	PositionConstraints.D.ConstraintLambda[ConstraintIndex] += DeltaLambda;
+	PositionConstraints.Data.ConstraintLambda[ConstraintIndex] += DeltaLambda;
 	SolvePositionConstraintDelta(ConstraintIndex, DeltaLambda, PositionConstraints);
 }
 
@@ -790,16 +790,16 @@ void FPBDJointCachedSolver::SolvePositionConstraintSoft(
 {
 	check(PositionConstraints.bUseSimd == false);
 	FReal VelDt = 0;
-	if (PositionConstraints.D.ConstraintSoftDamping[ConstraintIndex] > UE_KINDA_SMALL_NUMBER)
+	if (PositionConstraints.Data.ConstraintSoftDamping[ConstraintIndex] > UE_KINDA_SMALL_NUMBER)
 	{
-		const FVec3 V0Dt = FVec3::CalculateVelocity(InitConnectorXs[0], ConnectorXs[0] + Body(0).DP() + FVec3::CrossProduct(Body(0).DQ(), PositionConstraints.D.ConstraintArms[ConstraintIndex][0]), 1.0f);
-		const FVec3 V1Dt = FVec3::CalculateVelocity(InitConnectorXs[1], ConnectorXs[1] + Body(1).DP() + FVec3::CrossProduct(Body(1).DQ(), PositionConstraints.D.ConstraintArms[ConstraintIndex][1]), 1.0f);
-		VelDt = TargetVel * Dt + FVec3::DotProduct(V0Dt - V1Dt, PositionConstraints.D.ConstraintAxis[ConstraintIndex]);
+		const FVec3 V0Dt = FVec3::CalculateVelocity(InitConnectorXs[0], ConnectorXs[0] + Body(0).DP() + FVec3::CrossProduct(Body(0).DQ(), PositionConstraints.Data.ConstraintArms[ConstraintIndex][0]), 1.0f);
+		const FVec3 V1Dt = FVec3::CalculateVelocity(InitConnectorXs[1], ConnectorXs[1] + Body(1).DP() + FVec3::CrossProduct(Body(1).DQ(), PositionConstraints.Data.ConstraintArms[ConstraintIndex][1]), 1.0f);
+		VelDt = TargetVel * Dt + FVec3::DotProduct(V0Dt - V1Dt, PositionConstraints.Data.ConstraintAxis[ConstraintIndex]);
 	}
 
-	const FReal DeltaLambda = SolverStiffness * (PositionConstraints.D.ConstraintSoftStiffness[ConstraintIndex] * DeltaConstraint - PositionConstraints.D.ConstraintSoftDamping[ConstraintIndex] * VelDt - PositionConstraints.D.ConstraintLambda[ConstraintIndex]) /
-		PositionConstraints.D.ConstraintSoftIM[ConstraintIndex];
-	PositionConstraints.D.ConstraintLambda[ConstraintIndex] += DeltaLambda;
+	const FReal DeltaLambda = SolverStiffness * (PositionConstraints.Data.ConstraintSoftStiffness[ConstraintIndex] * DeltaConstraint - PositionConstraints.Data.ConstraintSoftDamping[ConstraintIndex] * VelDt - PositionConstraints.Data.ConstraintLambda[ConstraintIndex]) /
+		PositionConstraints.Data.ConstraintSoftIM[ConstraintIndex];
+	PositionConstraints.Data.ConstraintLambda[ConstraintIndex] += DeltaLambda;
 
 	SolvePositionConstraintDelta(ConstraintIndex, DeltaLambda, PositionConstraints);
 }
@@ -810,22 +810,22 @@ void FPBDJointCachedSolver::ApplyAxisPositionConstraint(
 
 	check(PositionConstraints.bUseSimd == false);
 	const FVec3 CX = Body(1).DP() - Body(0).DP() +
-		FVec3::CrossProduct(Body(1).DQ(), PositionConstraints.D.ConstraintArms[ConstraintIndex][1]) -
-		FVec3::CrossProduct(Body(0).DQ(), PositionConstraints.D.ConstraintArms[ConstraintIndex][0]);
+		FVec3::CrossProduct(Body(1).DQ(), PositionConstraints.Data.ConstraintArms[ConstraintIndex][1]) -
+		FVec3::CrossProduct(Body(0).DQ(), PositionConstraints.Data.ConstraintArms[ConstraintIndex][0]);
 
-	FReal DeltaPosition = PositionConstraints.D.ConstraintCX[ConstraintIndex] + FVec3::DotProduct(CX, PositionConstraints.D.ConstraintAxis[ConstraintIndex]);
+	FReal DeltaPosition = PositionConstraints.Data.ConstraintCX[ConstraintIndex] + FVec3::DotProduct(CX, PositionConstraints.Data.ConstraintAxis[ConstraintIndex]);
 
 	bool NeedsSolve = false;
 	if (PositionConstraints.bLimitsCheck[ConstraintIndex])
 	{
-		if (DeltaPosition > PositionConstraints.D.ConstraintLimits[ConstraintIndex])
+		if (DeltaPosition > PositionConstraints.Data.ConstraintLimits[ConstraintIndex])
 		{
-			DeltaPosition -= PositionConstraints.D.ConstraintLimits[ConstraintIndex];
+			DeltaPosition -= PositionConstraints.Data.ConstraintLimits[ConstraintIndex];
 			NeedsSolve = true;
 		}
-		else if (DeltaPosition < -PositionConstraints.D.ConstraintLimits[ConstraintIndex])
+		else if (DeltaPosition < -PositionConstraints.Data.ConstraintLimits[ConstraintIndex])
 		{
-			DeltaPosition += PositionConstraints.D.ConstraintLimits[ConstraintIndex];
+			DeltaPosition += PositionConstraints.Data.ConstraintLimits[ConstraintIndex];
 			NeedsSolve = true;
 		}
 	}
@@ -857,12 +857,12 @@ void FPBDJointCachedSolver::ApplyPositionConstraintsSimd(
 	VectorRegister4Float DeltaPositions[3];
 	for (int32 CIndex = 0; CIndex < 3; ++CIndex)
 	{
-		VectorRegister4Float  Cross1 = VectorCross(Body1DQ, PositionConstraints.Simd.ConstraintArmsSimd[CIndex][1]);
-		VectorRegister4Float  Cross0 = VectorCross(Body0DQ, PositionConstraints.Simd.ConstraintArmsSimd[CIndex][0]);
+		VectorRegister4Float  Cross1 = VectorCross(Body1DQ, PositionConstraints.Simd.ConstraintArms[CIndex][1]);
+		VectorRegister4Float  Cross0 = VectorCross(Body0DQ, PositionConstraints.Simd.ConstraintArms[CIndex][0]);
 		VectorRegister4Float  CrossDiff = VectorSubtract(Cross1, Cross0);
 		VectorRegister4Float CX = VectorAdd(DPDiff, CrossDiff);
 
-		DeltaPositions[CIndex] = Private::VectorDot3FastX(CX, PositionConstraints.Simd.ConstraintAxisSimd[CIndex]);
+		DeltaPositions[CIndex] = Private::VectorDot3FastX(CX, PositionConstraints.Simd.ConstraintAxis[CIndex]);
 	}
 
 	VectorRegister4Float DeltaPosition = VectorUnpackLo(DeltaPositions[0], DeltaPositions[1]);
@@ -873,9 +873,9 @@ void FPBDJointCachedSolver::ApplyPositionConstraintsSimd(
 	const FRealSingle SolverStiffnessf = FRealSingle(SolverStiffness);
 	VectorRegister4Float Stiffness = VectorLoadFloat1(&SolverStiffnessf);
 
-	Stiffness = VectorMultiply(Stiffness, PositionConstraints.Simd.ConstraintHardStiffnessSimd);
+	Stiffness = VectorMultiply(Stiffness, PositionConstraints.Simd.ConstraintHardStiffness);
 	DeltaPosition = VectorMultiply(Stiffness, DeltaPosition);
-	const VectorRegister4Float DeltaLambda = VectorDivide(DeltaPosition, PositionConstraints.Simd.ConstraintHardIMSimd);
+	const VectorRegister4Float DeltaLambda = VectorDivide(DeltaPosition, PositionConstraints.Simd.ConstraintHardIM);
 	PositionConstraints.Simd.ConstraintLambda = VectorAdd(PositionConstraints.Simd.ConstraintLambda, DeltaLambda);
 
 	VectorRegister4Float DLambda[3];
@@ -886,7 +886,7 @@ void FPBDJointCachedSolver::ApplyPositionConstraintsSimd(
 	VectorRegister4Float DX[3];
 	for (int32 CIndex = 0; CIndex < 3; CIndex++)
 	{
-		DX[CIndex] = VectorMultiply(PositionConstraints.Simd.ConstraintAxisSimd[CIndex], DLambda[CIndex]);
+		DX[CIndex] = VectorMultiply(PositionConstraints.Simd.ConstraintAxis[CIndex], DLambda[CIndex]);
 	}
 
 	if (Body(0).IsDynamic())
@@ -899,7 +899,7 @@ void FPBDJointCachedSolver::ApplyPositionConstraintsSimd(
 		for (int32 CIndex = 0; CIndex < 3; CIndex++)
 		{
 			DP0 = VectorMultiplyAdd(Inv0, DX[CIndex], DP0);
-			DR0 = VectorMultiplyAdd(PositionConstraints.Simd.ConstraintDRAxisSimd[CIndex][0], DLambda[CIndex], DR0);
+			DR0 = VectorMultiplyAdd(PositionConstraints.Simd.ConstraintDRAxis[CIndex][0], DLambda[CIndex], DR0);
 		}
 		FVec3f DP0f;
 		VectorStoreFloat3(DP0, &DP0f[0]);
@@ -918,7 +918,7 @@ void FPBDJointCachedSolver::ApplyPositionConstraintsSimd(
 		for (int32 CIndex = 0; CIndex < 3; CIndex++)
 		{
 			DP1 = VectorSubtract(DP1, VectorMultiply(Inv1, DX[CIndex]));
-			DR1 = VectorMultiplyAdd(PositionConstraints.Simd.ConstraintDRAxisSimd[CIndex][1], DLambda[CIndex], DR1);
+			DR1 = VectorMultiplyAdd(PositionConstraints.Simd.ConstraintDRAxis[CIndex][1], DLambda[CIndex], DR1);
 		}
 		FVec3f DP1f;
 		VectorStoreFloat3(DP1, &DP1f[0]);
@@ -958,29 +958,29 @@ void FPBDJointCachedSolver::SolveLinearVelocityConstraint(
 	const FReal TargetVel)
 {
 	check(PositionConstraints.bUseSimd == false);
-	const FVec3 CV0 = V(0) + FVec3::CrossProduct(W(0), PositionConstraints.D.ConstraintArms[ConstraintIndex][0]);
-	const FVec3 CV1 = V(1) + FVec3::CrossProduct(W(1), PositionConstraints.D.ConstraintArms[ConstraintIndex][1]);
+	const FVec3 CV0 = V(0) + FVec3::CrossProduct(W(0), PositionConstraints.Data.ConstraintArms[ConstraintIndex][0]);
+	const FVec3 CV1 = V(1) + FVec3::CrossProduct(W(1), PositionConstraints.Data.ConstraintArms[ConstraintIndex][1]);
 	const FVec3 CV = CV1 - CV0;
 
-	const FReal DeltaLambda = SolverStiffness * PositionConstraints.D.ConstraintHardStiffness[ConstraintIndex] *
-		(FVec3::DotProduct(CV, PositionConstraints.D.ConstraintAxis[ConstraintIndex]) - TargetVel) / PositionConstraints.D.ConstraintHardIM[ConstraintIndex];
+	const FReal DeltaLambda = SolverStiffness * PositionConstraints.Data.ConstraintHardStiffness[ConstraintIndex] *
+		(FVec3::DotProduct(CV, PositionConstraints.Data.ConstraintAxis[ConstraintIndex]) - TargetVel) / PositionConstraints.Data.ConstraintHardIM[ConstraintIndex];
 
 	// @todo(chaos): We should be adding to the net positional impulse here
-	//PositionConstraints.D.ConstraintLambda[ConstraintIndex] += DeltaLambda * Dt;
+	//PositionConstraints.Data.ConstraintLambda[ConstraintIndex] += DeltaLambda * Dt;
 
-	const FVec3 MDV = DeltaLambda * PositionConstraints.D.ConstraintAxis[ConstraintIndex];
+	const FVec3 MDV = DeltaLambda * PositionConstraints.Data.ConstraintAxis[ConstraintIndex];
 
 	if (Body(0).IsDynamic())
 	{
 		const FVec3 DV0 = InvM(0) * MDV;
-		const FVec3 DW0 = PositionConstraints.D.ConstraintDRAxis[ConstraintIndex][0] * DeltaLambda;
+		const FVec3 DW0 = PositionConstraints.Data.ConstraintDRAxis[ConstraintIndex][0] * DeltaLambda;
 
 		Body(0).ApplyVelocityDelta(DV0, DW0);
 	}
 	if (Body(1).IsDynamic())
 	{
 		const FVec3 DV1 = -InvM(1) * MDV;
-		const FVec3 DW1 = PositionConstraints.D.ConstraintDRAxis[ConstraintIndex][1] * DeltaLambda;
+		const FVec3 DW1 = PositionConstraints.Data.ConstraintDRAxis[ConstraintIndex][1] * DeltaLambda;
 
 		Body(1).ApplyVelocityDelta(DV1, DW1);
 	}
@@ -992,7 +992,7 @@ void FPBDJointCachedSolver::ApplyAxisVelocityConstraint(
 	check(PositionConstraints.bUseSimd == false);
 	// Apply restitution for limited joints when we have exceeded the limits
 	// We also drive the velocity to zero for locked constraints (ignoring restitution)
-	if (FMath::Abs(PositionConstraints.D.ConstraintLambda[ConstraintIndex]) > UE_SMALL_NUMBER)
+	if (FMath::Abs(PositionConstraints.Data.ConstraintLambda[ConstraintIndex]) > UE_SMALL_NUMBER)
 	{
 		FReal TargetVel = 0.0f;
 		const FReal Restitution = PositionConstraints.ConstraintRestitution[ConstraintIndex];
@@ -1029,19 +1029,19 @@ void FPBDJointCachedSolver::ApplyVelocityConstraintSimd()
 
 		const FRealSingle SolverStiffnessf = FRealSingle(SolverStiffness);
 		VectorRegister4Float Stiffness = VectorLoadFloat1(&SolverStiffnessf);
-		Stiffness = VectorMultiply(Stiffness, PositionConstraints.Simd.ConstraintHardStiffnessSimd);
+		Stiffness = VectorMultiply(Stiffness, PositionConstraints.Simd.ConstraintHardStiffness);
 
 		VectorRegister4Float HardIM[3];
-		HardIM[0] = VectorReplicate(PositionConstraints.Simd.ConstraintHardIMSimd, 0);
-		HardIM[1] = VectorReplicate(PositionConstraints.Simd.ConstraintHardIMSimd, 1);
-		HardIM[2] = VectorReplicate(PositionConstraints.Simd.ConstraintHardIMSimd, 2);
+		HardIM[0] = VectorReplicate(PositionConstraints.Simd.ConstraintHardIM, 0);
+		HardIM[1] = VectorReplicate(PositionConstraints.Simd.ConstraintHardIM, 1);
+		HardIM[2] = VectorReplicate(PositionConstraints.Simd.ConstraintHardIM, 2);
 		VectorRegister4Float DeltaLambdas[3];
 		for (int32 CIndex = 0; CIndex < 3; CIndex++)
 		{
-			const VectorRegister4Float CV0 = VectorAdd(V0, VectorCross(W0, PositionConstraints.Simd.ConstraintArmsSimd[CIndex][0]));
-			const VectorRegister4Float CV1 = VectorAdd(V1, VectorCross(W1, PositionConstraints.Simd.ConstraintArmsSimd[CIndex][1]));
+			const VectorRegister4Float CV0 = VectorAdd(V0, VectorCross(W0, PositionConstraints.Simd.ConstraintArms[CIndex][0]));
+			const VectorRegister4Float CV1 = VectorAdd(V1, VectorCross(W1, PositionConstraints.Simd.ConstraintArms[CIndex][1]));
 			const VectorRegister4Float CV = VectorSubtract(CV1, CV0);
-			const VectorRegister4Float Dot3 = VectorDot3(CV, PositionConstraints.Simd.ConstraintAxisSimd[CIndex]);
+			const VectorRegister4Float Dot3 = VectorDot3(CV, PositionConstraints.Simd.ConstraintAxis[CIndex]);
 
 			DeltaLambdas[CIndex] = VectorDivide(VectorMultiply(Stiffness, Dot3), HardIM[CIndex]);
 		}
@@ -1053,8 +1053,8 @@ void FPBDJointCachedSolver::ApplyVelocityConstraintSimd()
 			VectorRegister4Float DW0 = VectorZeroFloat();
 			for (int32 CIndex = 0; CIndex < 3; CIndex++)
 			{
-				DV0 = VectorMultiplyAdd(InvM0, VectorMultiply(DeltaLambdas[CIndex], PositionConstraints.Simd.ConstraintAxisSimd[CIndex]), DV0);
-				DW0 = VectorMultiplyAdd(PositionConstraints.Simd.ConstraintDRAxisSimd[CIndex][0], DeltaLambdas[CIndex], DW0);
+				DV0 = VectorMultiplyAdd(InvM0, VectorMultiply(DeltaLambdas[CIndex], PositionConstraints.Simd.ConstraintAxis[CIndex]), DV0);
+				DW0 = VectorMultiplyAdd(PositionConstraints.Simd.ConstraintDRAxis[CIndex][0], DeltaLambdas[CIndex], DW0);
 			}
 			FVec3f DV0f;
 			VectorStoreFloat3(DV0, &DV0f[0]);
@@ -1070,8 +1070,8 @@ void FPBDJointCachedSolver::ApplyVelocityConstraintSimd()
 			VectorRegister4Float DW1 = VectorZeroFloat();
 			for (int32 CIndex = 0; CIndex < 3; CIndex++)
 			{
-				DV1 = VectorMultiplyAdd(InvM1, VectorMultiply(DeltaLambdas[CIndex], PositionConstraints.Simd.ConstraintAxisSimd[CIndex]), DV1);
-				DW1 = VectorMultiplyAdd(PositionConstraints.Simd.ConstraintDRAxisSimd[CIndex][1], DeltaLambdas[CIndex], DW1);
+				DV1 = VectorMultiplyAdd(InvM1, VectorMultiply(DeltaLambdas[CIndex], PositionConstraints.Simd.ConstraintAxis[CIndex]), DV1);
+				DW1 = VectorMultiplyAdd(PositionConstraints.Simd.ConstraintDRAxis[CIndex][1], DeltaLambdas[CIndex], DW1);
 			}
 			FVec3f DV1f;
 			VectorStoreFloat3(DV1, &DV1f[0]);
@@ -1153,7 +1153,7 @@ void FPBDJointCachedSolver::InitRotationConstraints(
 			RotationConstraints.ConstraintLambdaVelocity = FVec3::Zero();
 		}
 
-		RotationConstraints.Simd.ConstraintHardStiffnessSimd = MakeVectorRegisterFloatFromDouble(MakeVectorRegister(FPBDJointUtilities::GetTwistStiffness(SolverSettings, JointSettings),
+		RotationConstraints.Simd.ConstraintHardStiffness = MakeVectorRegisterFloatFromDouble(MakeVectorRegister(FPBDJointUtilities::GetTwistStiffness(SolverSettings, JointSettings),
 			FPBDJointUtilities::GetSwingStiffness(SolverSettings, JointSettings), FPBDJointUtilities::GetSwingStiffness(SolverSettings, JointSettings), 0.0f));
 		InitRotationConstraintsSimd(JointSettings, Dt);
 	}
@@ -1293,15 +1293,15 @@ void FPBDJointCachedSolver::InitRotationConstraintsSimd(
 		RotationConstraints.bLimitsCheck[ConstraintIndex] = true;
 		RotationConstraints.ConstraintVX[ConstraintIndex] = 0;
 
-		RotationConstraints.Simd.ConstraintArmsSimd[ConstraintIndex][0] = VectorZeroFloat();
-		RotationConstraints.Simd.ConstraintArmsSimd[ConstraintIndex][1] = VectorZeroFloat();
+		RotationConstraints.Simd.ConstraintArms[ConstraintIndex][0] = VectorZeroFloat();
+		RotationConstraints.Simd.ConstraintArms[ConstraintIndex][1] = VectorZeroFloat();
 
 		InitConstraintAxisAngularVelocities[ConstraintIndex] = FVec3::DotProduct(W(1) - W(0), LocalAxis);
 
 		// InitRotationDatasMass(RotationConstraints, ConstraintIndex, Dt);
 
-		RotationConstraints.Simd.ConstraintAxisSimd[ConstraintIndex] = MakeVectorRegisterFloatFromDouble(MakeVectorRegister(LocalAxis[0], LocalAxis[1], LocalAxis[2], 0.0f));
-		const VectorRegister4Float& Axis = RotationConstraints.Simd.ConstraintAxisSimd[ConstraintIndex];
+		RotationConstraints.Simd.ConstraintAxis[ConstraintIndex] = MakeVectorRegisterFloatFromDouble(MakeVectorRegister(LocalAxis[0], LocalAxis[1], LocalAxis[2], 0.0f));
+		const VectorRegister4Float& Axis = RotationConstraints.Simd.ConstraintAxis[ConstraintIndex];
 
 		const VectorRegister4Float AxisX = VectorReplicate(Axis, 0);
 		const VectorRegister4Float AxisY = VectorReplicate(Axis, 1);
@@ -1320,8 +1320,8 @@ void FPBDJointCachedSolver::InitRotationConstraintsSimd(
 
 		const FRealSingle II0 = VectorDot3Scalar(Axis, IA0);
 		const FRealSingle II1 = VectorDot3Scalar(Axis, IA1);
-		RotationConstraints.Simd.ConstraintDRAxisSimd[ConstraintIndex][0] = IA0;
-		RotationConstraints.Simd.ConstraintDRAxisSimd[ConstraintIndex][1] = VectorNegate(IA1);
+		RotationConstraints.Simd.ConstraintDRAxis[ConstraintIndex][0] = IA0;
+		RotationConstraints.Simd.ConstraintDRAxis[ConstraintIndex][1] = VectorNegate(IA1);
 
 		ConstraintHardIM[ConstraintIndex] = II0 + II1;
 		const FReal SpringMassScale = (bUsePositionBasedDrives) ? (FReal)1 / (ConstraintHardIM[ConstraintIndex]) : (FReal)1;
@@ -1332,11 +1332,11 @@ void FPBDJointCachedSolver::InitRotationConstraintsSimd(
 	}
 
 	RotationConstraints.Simd.ConstraintCX = MakeVectorRegisterFloat(LocalAngles[0], LocalAngles[1], LocalAngles[2], 0.0f);
-	RotationConstraints.Simd.ConstraintHardIMSimd = MakeVectorRegisterFloat(ConstraintHardIM[0], ConstraintHardIM[1], ConstraintHardIM[2], 0.0f);
-	RotationConstraints.Simd.ConstraintSoftStiffnessSimd = MakeVectorRegisterFloat(ConstraintSoftStiffness[0], ConstraintSoftStiffness[1], ConstraintSoftStiffness[2], 0.0f);
-	RotationConstraints.Simd.ConstraintSoftDampingSimd = MakeVectorRegisterFloat(ConstraintSoftDamping[0], ConstraintSoftDamping[1], ConstraintSoftDamping[2], 0.0f);
-	RotationConstraints.Simd.ConstraintSoftIMSimd = MakeVectorRegisterFloat(ConstraintSoftIM[0], ConstraintSoftIM[1], ConstraintSoftIM[2], 0.0f);
-	RotationConstraints.Simd.ConstraintLimitsSimd = MakeVectorRegisterFloatFromDouble(MakeVectorRegister(JointSettings.AngularLimits[0], JointSettings.AngularLimits[1], JointSettings.AngularLimits[2], UE_BIG_NUMBER));
+	RotationConstraints.Simd.ConstraintHardIM = MakeVectorRegisterFloat(ConstraintHardIM[0], ConstraintHardIM[1], ConstraintHardIM[2], 0.0f);
+	RotationConstraints.Simd.ConstraintSoftStiffness = MakeVectorRegisterFloat(ConstraintSoftStiffness[0], ConstraintSoftStiffness[1], ConstraintSoftStiffness[2], 0.0f);
+	RotationConstraints.Simd.ConstraintSoftDamping = MakeVectorRegisterFloat(ConstraintSoftDamping[0], ConstraintSoftDamping[1], ConstraintSoftDamping[2], 0.0f);
+	RotationConstraints.Simd.ConstraintSoftIM = MakeVectorRegisterFloat(ConstraintSoftIM[0], ConstraintSoftIM[1], ConstraintSoftIM[2], 0.0f);
+	RotationConstraints.Simd.ConstraintLimits = MakeVectorRegisterFloatFromDouble(MakeVectorRegister(JointSettings.AngularLimits[0], JointSettings.AngularLimits[1], JointSettings.AngularLimits[2], UE_BIG_NUMBER));
 }
 
 void FPBDJointCachedSolver::InitRotationDatasMass(
@@ -1344,10 +1344,10 @@ void FPBDJointCachedSolver::InitRotationDatasMass(
 	const int32 ConstraintIndex,
 	const FReal Dt)
 {
-	const FVec3 IA0 = Utilities::Multiply(InvI(0), RotationDatas.D.ConstraintAxis[ConstraintIndex]);
-	const FVec3 IA1 = Utilities::Multiply(InvI(1), RotationDatas.D.ConstraintAxis[ConstraintIndex]);
-	const FReal II0 = FVec3::DotProduct(RotationDatas.D.ConstraintAxis[ConstraintIndex], IA0);
-	const FReal II1 = FVec3::DotProduct(RotationDatas.D.ConstraintAxis[ConstraintIndex], IA1);
+	const FVec3 IA0 = Utilities::Multiply(InvI(0), RotationDatas.Data.ConstraintAxis[ConstraintIndex]);
+	const FVec3 IA1 = Utilities::Multiply(InvI(1), RotationDatas.Data.ConstraintAxis[ConstraintIndex]);
+	const FReal II0 = FVec3::DotProduct(RotationDatas.Data.ConstraintAxis[ConstraintIndex], IA0);
+	const FReal II1 = FVec3::DotProduct(RotationDatas.Data.ConstraintAxis[ConstraintIndex], IA1);
 
 	RotationDatas.UpdateMass(ConstraintIndex, IA0, IA1, II0 + II1, Dt, bUsePositionBasedDrives);
 }
@@ -1367,7 +1367,7 @@ void FPBDJointCachedSolver::InitRotationConstraintDatas(
 	RotationConstraints.UpdateDatas(ConstraintIndex, LocalAxis, LocalAngle, ConstraintRestitution, bCheckLimit);
 
 
-	RotationConstraints.D.ConstraintLimits[ConstraintIndex] = JointSettings.AngularLimits[ConstraintIndex];
+	RotationConstraints.Data.ConstraintLimits[ConstraintIndex] = JointSettings.AngularLimits[ConstraintIndex];
 
 	InitConstraintAxisAngularVelocities[ConstraintIndex] = FVec3::DotProduct(W(1) - W(0), LocalAxis);
 
@@ -1553,17 +1553,17 @@ void FPBDJointCachedSolver::SolveRotationConstraintDelta(
 	const bool bIsSoftConstraint,
 	const FAxisConstraintDatas& ConstraintDatas)
 {
-	const FVec3 DeltaImpulse = ConstraintDatas.D.ConstraintAxis[ConstraintIndex] * DeltaLambda;
+	const FVec3 DeltaImpulse = ConstraintDatas.Data.ConstraintAxis[ConstraintIndex] * DeltaLambda;
 	if (Body(0).IsDynamic())
 	{
-		const FVec3 DR0 = !bIsSoftConstraint ? ConstraintDatas.D.ConstraintDRAxis[ConstraintIndex][0] * DeltaLambda :
-			DeltaImpulse * (FVec3::DotProduct(ConstraintDatas.D.ConstraintAxis[ConstraintIndex], ConstraintDatas.D.ConstraintDRAxis[ConstraintIndex][0]));
+		const FVec3 DR0 = !bIsSoftConstraint ? ConstraintDatas.Data.ConstraintDRAxis[ConstraintIndex][0] * DeltaLambda :
+			DeltaImpulse * (FVec3::DotProduct(ConstraintDatas.Data.ConstraintAxis[ConstraintIndex], ConstraintDatas.Data.ConstraintDRAxis[ConstraintIndex][0]));
 		ApplyRotationDelta(0, DR0);
 	}
 	if (Body(1).IsDynamic())
 	{
-		const FVec3 DR1 = !bIsSoftConstraint ? ConstraintDatas.D.ConstraintDRAxis[ConstraintIndex][1] * DeltaLambda :
-			DeltaImpulse * (FVec3::DotProduct(ConstraintDatas.D.ConstraintAxis[ConstraintIndex], ConstraintDatas.D.ConstraintDRAxis[ConstraintIndex][1]));
+		const FVec3 DR1 = !bIsSoftConstraint ? ConstraintDatas.Data.ConstraintDRAxis[ConstraintIndex][1] * DeltaLambda :
+			DeltaImpulse * (FVec3::DotProduct(ConstraintDatas.Data.ConstraintAxis[ConstraintIndex], ConstraintDatas.Data.ConstraintDRAxis[ConstraintIndex][1]));
 		ApplyRotationDelta(1, DR1);
 	}
 	++NumActiveConstraints;
@@ -1573,10 +1573,10 @@ void FPBDJointCachedSolver::SolveRotationConstraintHard(
 	const int32 ConstraintIndex,
 	const FReal DeltaConstraint)
 {
-	const FReal DeltaLambda = SolverStiffness * RotationConstraints.D.ConstraintHardStiffness[ConstraintIndex] * DeltaConstraint /
-		RotationConstraints.D.ConstraintHardIM[ConstraintIndex];
+	const FReal DeltaLambda = SolverStiffness * RotationConstraints.Data.ConstraintHardStiffness[ConstraintIndex] * DeltaConstraint /
+		RotationConstraints.Data.ConstraintHardIM[ConstraintIndex];
 
-	RotationConstraints.D.ConstraintLambda[ConstraintIndex] += DeltaLambda;
+	RotationConstraints.Data.ConstraintLambda[ConstraintIndex] += DeltaLambda;
 	SolveRotationConstraintDelta(ConstraintIndex, DeltaLambda, false, RotationConstraints);
 }
 
@@ -1588,17 +1588,17 @@ void FPBDJointCachedSolver::SolveRotationConstraintSoft(
 {
 	// Damping angular velocity
 	FReal AngVelDt = 0;
-	if (RotationConstraints.D.ConstraintSoftDamping[ConstraintIndex] > UE_KINDA_SMALL_NUMBER)
+	if (RotationConstraints.Data.ConstraintSoftDamping[ConstraintIndex] > UE_KINDA_SMALL_NUMBER)
 	{
 		const FVec3 W0Dt = FVec3(Body(0).DQ()) + ConnectorWDts[0];
 		const FVec3 W1Dt = FVec3(Body(1).DQ()) + ConnectorWDts[1];
-		AngVelDt = TargetVel * Dt + FVec3::DotProduct(RotationConstraints.D.ConstraintAxis[ConstraintIndex], W0Dt - W1Dt);
+		AngVelDt = TargetVel * Dt + FVec3::DotProduct(RotationConstraints.Data.ConstraintAxis[ConstraintIndex], W0Dt - W1Dt);
 	}
 
-	const FReal DeltaLambda = SolverStiffness * (RotationConstraints.D.ConstraintSoftStiffness[ConstraintIndex] * DeltaConstraint -
-		RotationConstraints.D.ConstraintSoftDamping[ConstraintIndex] * AngVelDt - RotationConstraints.D.ConstraintLambda[ConstraintIndex]) /
-		RotationConstraints.D.ConstraintSoftIM[ConstraintIndex];
-	RotationConstraints.D.ConstraintLambda[ConstraintIndex] += DeltaLambda;
+	const FReal DeltaLambda = SolverStiffness * (RotationConstraints.Data.ConstraintSoftStiffness[ConstraintIndex] * DeltaConstraint -
+		RotationConstraints.Data.ConstraintSoftDamping[ConstraintIndex] * AngVelDt - RotationConstraints.Data.ConstraintLambda[ConstraintIndex]) /
+		RotationConstraints.Data.ConstraintSoftIM[ConstraintIndex];
+	RotationConstraints.Data.ConstraintLambda[ConstraintIndex] += DeltaLambda;
 
 	SolveRotationConstraintDelta(ConstraintIndex, DeltaLambda, false, RotationConstraints);
 }
@@ -1607,20 +1607,20 @@ void FPBDJointCachedSolver::ApplyRotationConstraint(
 	const int32 ConstraintIndex,
 	const FReal Dt)
 {
-	FReal DeltaAngle = RotationConstraints.D.ConstraintCX[ConstraintIndex] +
-		FVec3::DotProduct(Body(1).DQ() - Body(0).DQ(), RotationConstraints.D.ConstraintAxis[ConstraintIndex]);
+	FReal DeltaAngle = RotationConstraints.Data.ConstraintCX[ConstraintIndex] +
+		FVec3::DotProduct(Body(1).DQ() - Body(0).DQ(), RotationConstraints.Data.ConstraintAxis[ConstraintIndex]);
 
 	bool NeedsSolve = false;
 	if (RotationConstraints.bLimitsCheck[ConstraintIndex])
 	{
-		if (DeltaAngle > RotationConstraints.D.ConstraintLimits[ConstraintIndex])
+		if (DeltaAngle > RotationConstraints.Data.ConstraintLimits[ConstraintIndex])
 		{
-			DeltaAngle -= RotationConstraints.D.ConstraintLimits[ConstraintIndex];
+			DeltaAngle -= RotationConstraints.Data.ConstraintLimits[ConstraintIndex];
 			NeedsSolve = true;
 		}
-		else if (DeltaAngle < -RotationConstraints.D.ConstraintLimits[ConstraintIndex])
+		else if (DeltaAngle < -RotationConstraints.Data.ConstraintLimits[ConstraintIndex])
 		{
-			DeltaAngle += RotationConstraints.D.ConstraintLimits[ConstraintIndex];
+			DeltaAngle += RotationConstraints.Data.ConstraintLimits[ConstraintIndex];
 			NeedsSolve = true;
 
 		}
@@ -1649,18 +1649,18 @@ void FPBDJointCachedSolver::ApplyRotationSoftConstraintsSimd(
 	VectorRegister4Float ProjAxes[3];
 	for (int32 CIndex = 0; CIndex < 3; ++CIndex)
 	{
-		ProjAxes[CIndex] = Private::VectorDot3FastX(DQDiff, RotationConstraints.Simd.ConstraintAxisSimd[CIndex]);
+		ProjAxes[CIndex] = Private::VectorDot3FastX(DQDiff, RotationConstraints.Simd.ConstraintAxis[CIndex]);
 	}
 	VectorRegister4Float ProjAxis = VectorUnpackLo(ProjAxes[0], ProjAxes[1]);
 	ProjAxis = VectorMoveLh(ProjAxis, ProjAxes[2]);
 	VectorRegister4Float DeltaAngle = VectorAdd(RotationConstraints.Simd.ConstraintCX, ProjAxis);
 
 
-	const VectorRegister4Float AngleGT = VectorCompareGT(DeltaAngle, RotationConstraints.Simd.ConstraintLimitsSimd);
-	const VectorRegister4Float AngleLT = VectorCompareLT(DeltaAngle, VectorNegate(RotationConstraints.Simd.ConstraintLimitsSimd));
+	const VectorRegister4Float AngleGT = VectorCompareGT(DeltaAngle, RotationConstraints.Simd.ConstraintLimits);
+	const VectorRegister4Float AngleLT = VectorCompareLT(DeltaAngle, VectorNegate(RotationConstraints.Simd.ConstraintLimits));
 
-	DeltaAngle = VectorSelect(AngleGT, VectorSubtract(DeltaAngle, RotationConstraints.Simd.ConstraintLimitsSimd),
-		VectorSelect(AngleLT, VectorAdd(DeltaAngle, RotationConstraints.Simd.ConstraintLimitsSimd), DeltaAngle));
+	DeltaAngle = VectorSelect(AngleGT, VectorSubtract(DeltaAngle, RotationConstraints.Simd.ConstraintLimits),
+		VectorSelect(AngleLT, VectorAdd(DeltaAngle, RotationConstraints.Simd.ConstraintLimits), DeltaAngle));
 
 	FRealSingle AngleTolerancef = FRealSingle(AngleTolerance);
 	const VectorRegister4Float AngleToleranceSimd = VectorLoadFloat1(&AngleTolerancef);
@@ -1673,15 +1673,15 @@ void FPBDJointCachedSolver::ApplyRotationSoftConstraintsSimd(
 		VectorRegister4Float WDiff = VectorSubtract(ConnectorWDtsSimd[0], VectorAdd(DQDiff, ConnectorWDtsSimd[1]));
 		for (int32 CIndex = 0; CIndex < 3; ++CIndex)
 		{
-			AngVelDts[CIndex] = Private::VectorDot3FastX(RotationConstraints.Simd.ConstraintAxisSimd[CIndex], WDiff);
+			AngVelDts[CIndex] = Private::VectorDot3FastX(RotationConstraints.Simd.ConstraintAxis[CIndex], WDiff);
 		}
 		VectorRegister4Float AngVelDt = VectorUnpackLo(AngVelDts[0], AngVelDts[1]);
 		AngVelDt = VectorMoveLh(AngVelDt, AngVelDts[2]);
 
 		const FRealSingle SolverStiffnessf = FRealSingle(SolverStiffness);
 		VectorRegister4Float Stiffness = VectorLoadFloat1(&SolverStiffnessf);
-		VectorRegister4Float DeltaLambda = VectorMultiply(Stiffness, VectorDivide(VectorSubtract(VectorMultiply(RotationConstraints.Simd.ConstraintSoftStiffnessSimd, DeltaAngle),
-			VectorSubtract(VectorMultiply(RotationConstraints.Simd.ConstraintSoftDampingSimd, AngVelDt), RotationConstraints.Simd.ConstraintLambda)), RotationConstraints.Simd.ConstraintSoftIMSimd));
+		VectorRegister4Float DeltaLambda = VectorMultiply(Stiffness, VectorDivide(VectorSubtract(VectorMultiply(RotationConstraints.Simd.ConstraintSoftStiffness, DeltaAngle),
+			VectorSubtract(VectorMultiply(RotationConstraints.Simd.ConstraintSoftDamping, AngVelDt), RotationConstraints.Simd.ConstraintLambda)), RotationConstraints.Simd.ConstraintSoftIM));
 
 		DeltaLambda = VectorSelect(NeedsSolve, DeltaLambda, VectorZeroFloat());
 		RotationConstraints.Simd.ConstraintLambda = VectorAdd(RotationConstraints.Simd.ConstraintLambda, DeltaLambda);
@@ -1696,7 +1696,7 @@ void FPBDJointCachedSolver::ApplyRotationSoftConstraintsSimd(
 			VectorRegister4Float DR0 = VectorZeroFloat();
 			for (int32 CIndex = 0; CIndex < 3; ++CIndex)
 			{
-				DR0 = VectorMultiplyAdd(RotationConstraints.Simd.ConstraintDRAxisSimd[CIndex][0], DeltaLambdas[CIndex], DR0);
+				DR0 = VectorMultiplyAdd(RotationConstraints.Simd.ConstraintDRAxis[CIndex][0], DeltaLambdas[CIndex], DR0);
 			}
 			FVec3f DR0f;
 			VectorStoreFloat3(DR0, &DR0f[0]);
@@ -1707,7 +1707,7 @@ void FPBDJointCachedSolver::ApplyRotationSoftConstraintsSimd(
 			VectorRegister4Float DR1 = VectorZeroFloat();
 			for (int32 CIndex = 0; CIndex < 3; ++CIndex)
 			{
-				DR1 = VectorMultiplyAdd(RotationConstraints.Simd.ConstraintDRAxisSimd[CIndex][1], DeltaLambdas[CIndex], DR1);
+				DR1 = VectorMultiplyAdd(RotationConstraints.Simd.ConstraintDRAxis[CIndex][1], DeltaLambdas[CIndex], DR1);
 			}
 			FVec3f DR1f;
 			VectorStoreFloat3(DR1, &DR1f[0]);
@@ -1744,21 +1744,21 @@ void FPBDJointCachedSolver::SolveAngularVelocityConstraint(
 {
 	const FVec3 CW = W(1) - W(0);
 
-	const FReal DeltaLambda = SolverStiffness * RotationConstraints.D.ConstraintHardStiffness[ConstraintIndex] *
-		(FVec3::DotProduct(CW, RotationConstraints.D.ConstraintAxis[ConstraintIndex]) - TargetVel) / RotationConstraints.D.ConstraintHardIM[ConstraintIndex];
+	const FReal DeltaLambda = SolverStiffness * RotationConstraints.Data.ConstraintHardStiffness[ConstraintIndex] *
+		(FVec3::DotProduct(CW, RotationConstraints.Data.ConstraintAxis[ConstraintIndex]) - TargetVel) / RotationConstraints.Data.ConstraintHardIM[ConstraintIndex];
 
 	// @todo(chaos): we should be adding to the net positional impulse here
-	// RotationConstraints.D.ConstraintLambda[ConsraintIndex] += DeltaLambda * Dt;
+	// RotationConstraints.Data.ConstraintLambda[ConsraintIndex] += DeltaLambda * Dt;
 
 	if (Body(0).IsDynamic())
 	{
-		const FVec3 DW0 = RotationConstraints.D.ConstraintDRAxis[ConstraintIndex][0] * DeltaLambda;
+		const FVec3 DW0 = RotationConstraints.Data.ConstraintDRAxis[ConstraintIndex][0] * DeltaLambda;
 
 		Body(0).ApplyAngularVelocityDelta(DW0);
 	}
 	if (Body(1).IsDynamic())
 	{
-		const FVec3 DW1 = RotationConstraints.D.ConstraintDRAxis[ConstraintIndex][1] * DeltaLambda;
+		const FVec3 DW1 = RotationConstraints.Data.ConstraintDRAxis[ConstraintIndex][1] * DeltaLambda;
 
 		Body(1).ApplyAngularVelocityDelta(DW1);
 	}
@@ -1769,7 +1769,7 @@ void FPBDJointCachedSolver::ApplyAngularVelocityConstraint(
 {
 	// Apply restitution for limited joints when we have exceeded the limits
 	// We also drive the velocity to zero for locked constraints (ignoring restitution)
-	if (FMath::Abs(RotationConstraints.D.ConstraintLambda[ConstraintIndex]) > UE_SMALL_NUMBER)
+	if (FMath::Abs(RotationConstraints.Data.ConstraintLambda[ConstraintIndex]) > UE_SMALL_NUMBER)
 	{
 		FReal TargetVel = 0.0f;
 		if ((RotationConstraints.MotionType[ConstraintIndex] == EJointMotionType::Limited) && (RotationConstraints.ConstraintRestitution[ConstraintIndex] != 0.0f))
@@ -1810,16 +1810,16 @@ void FPBDJointCachedSolver::ApplyAngularVelocityConstraintSimd()
 		const VectorRegister4Float CW = VectorSubtract(W1, W0);
 		const FRealSingle SolverStiffnessf = FRealSingle(SolverStiffness);
 		VectorRegister4Float Stiffness = VectorLoadFloat1(&SolverStiffnessf);
-		Stiffness = VectorMultiply(Stiffness, RotationConstraints.Simd.ConstraintHardStiffnessSimd);
+		Stiffness = VectorMultiply(Stiffness, RotationConstraints.Simd.ConstraintHardStiffness);
 
 		VectorRegister4Float DeltaLambdas[3];
 		for (int32 CIndex = 0; CIndex < 3; CIndex++)
 		{
-			DeltaLambdas[CIndex] = Private::VectorDot3FastX(CW, RotationConstraints.Simd.ConstraintAxisSimd[CIndex]);
+			DeltaLambdas[CIndex] = Private::VectorDot3FastX(CW, RotationConstraints.Simd.ConstraintAxis[CIndex]);
 		}
 		VectorRegister4Float DeltaLambda = VectorUnpackLo(DeltaLambdas[0], DeltaLambdas[1]);
 		DeltaLambda = VectorMoveLh(DeltaLambda, DeltaLambdas[2]);
-		DeltaLambda = VectorDivide(VectorMultiply(Stiffness, VectorSubtract(DeltaLambda, TargetVel)), RotationConstraints.Simd.ConstraintHardIMSimd);
+		DeltaLambda = VectorDivide(VectorMultiply(Stiffness, VectorSubtract(DeltaLambda, TargetVel)), RotationConstraints.Simd.ConstraintHardIM);
 
 		DeltaLambda = VectorSelect(IsGTEps, DeltaLambda, VectorZeroFloat());
 		DeltaLambdas[0] = VectorReplicate(DeltaLambda, 0);
@@ -1831,7 +1831,7 @@ void FPBDJointCachedSolver::ApplyAngularVelocityConstraintSimd()
 			VectorRegister4Float DW0 = VectorZeroFloat();
 			for (int32 CIndex = 0; CIndex < 3; CIndex++)
 			{
-				DW0 = VectorMultiplyAdd(RotationConstraints.Simd.ConstraintDRAxisSimd[CIndex][0], DeltaLambdas[CIndex], DW0);
+				DW0 = VectorMultiplyAdd(RotationConstraints.Simd.ConstraintDRAxis[CIndex][0], DeltaLambdas[CIndex], DW0);
 			}
 
 			FVec3f DW0f;
@@ -1843,7 +1843,7 @@ void FPBDJointCachedSolver::ApplyAngularVelocityConstraintSimd()
 			VectorRegister4Float DW1 = VectorZeroFloat();
 			for (int32 CIndex = 0; CIndex < 3; CIndex++)
 			{
-				DW1 = VectorMultiplyAdd(RotationConstraints.Simd.ConstraintDRAxisSimd[CIndex][1], DeltaLambdas[CIndex], DW1);
+				DW1 = VectorMultiplyAdd(RotationConstraints.Simd.ConstraintDRAxis[CIndex][1], DeltaLambdas[CIndex], DW1);
 			}
 			FVec3f DW1f;
 			VectorStoreFloat3(DW1, &DW1f[0]);
@@ -1895,7 +1895,7 @@ void FPBDJointCachedSolver::InitPositionDrives(
 						FPBDJointUtilities::GetLinearDriveDamping(SolverSettings, JointSettings, AxisIndex), 0);
 					const FVec3 Axis = R0M.GetAxis(AxisIndex);
 
-					if ((FMath::Abs(FVec3::DotProduct(CX, Axis)) > PositionTolerance) || (PositionDrives.D.ConstraintSoftDamping[AxisIndex] > 0.0f))
+					if ((FMath::Abs(FVec3::DotProduct(CX, Axis)) > PositionTolerance) || (PositionDrives.Data.ConstraintSoftDamping[AxisIndex] > 0.0f))
 					{
 						InitAxisPositionDrive(AxisIndex, Axis, CX, VTarget, Dt);
 					}
@@ -2004,20 +2004,20 @@ void FPBDJointCachedSolver::ApplyAxisRotationProjection(
 	check(RotationConstraints.bUseSimd == false);
 	if (RotationConstraints.bValidDatas[ConstraintIndex] && !RotationConstraints.bSoftLimit[ConstraintIndex])
 	{
-		FReal DeltaAngle = RotationConstraints.D.ConstraintCX[ConstraintIndex] +
-			FVec3::DotProduct(Body(1).DQ() - Body(0).DQ(), RotationConstraints.D.ConstraintAxis[ConstraintIndex]);
+		FReal DeltaAngle = RotationConstraints.Data.ConstraintCX[ConstraintIndex] +
+			FVec3::DotProduct(Body(1).DQ() - Body(0).DQ(), RotationConstraints.Data.ConstraintAxis[ConstraintIndex]);
 
 		bool NeedsSolve = false;
 		if (RotationConstraints.bLimitsCheck[ConstraintIndex])
 		{
-			if (DeltaAngle > RotationConstraints.D.ConstraintLimits[ConstraintIndex])
+			if (DeltaAngle > RotationConstraints.Data.ConstraintLimits[ConstraintIndex])
 			{
-				DeltaAngle -= RotationConstraints.D.ConstraintLimits[ConstraintIndex];
+				DeltaAngle -= RotationConstraints.Data.ConstraintLimits[ConstraintIndex];
 				NeedsSolve = true;
 			}
-			else if (DeltaAngle < -RotationConstraints.D.ConstraintLimits[ConstraintIndex])
+			else if (DeltaAngle < -RotationConstraints.Data.ConstraintLimits[ConstraintIndex])
 			{
-				DeltaAngle += RotationConstraints.D.ConstraintLimits[ConstraintIndex];
+				DeltaAngle += RotationConstraints.Data.ConstraintLimits[ConstraintIndex];
 				NeedsSolve = true;
 
 			}
@@ -2025,10 +2025,10 @@ void FPBDJointCachedSolver::ApplyAxisRotationProjection(
 
 		if (!RotationConstraints.bLimitsCheck[ConstraintIndex] || (RotationConstraints.bLimitsCheck[ConstraintIndex] && NeedsSolve && FMath::Abs(DeltaAngle) > AngleTolerance))
 		{
-			const FReal IM = -FVec3::DotProduct(RotationConstraints.D.ConstraintAxis[ConstraintIndex], RotationConstraints.D.ConstraintDRAxis[ConstraintIndex][1]);
-			const FReal DeltaLambda = SolverStiffness * RotationConstraints.D.ConstraintHardStiffness[ConstraintIndex] * DeltaAngle / IM;
+			const FReal IM = -FVec3::DotProduct(RotationConstraints.Data.ConstraintAxis[ConstraintIndex], RotationConstraints.Data.ConstraintDRAxis[ConstraintIndex][1]);
+			const FReal DeltaLambda = SolverStiffness * RotationConstraints.Data.ConstraintHardStiffness[ConstraintIndex] * DeltaAngle / IM;
 
-			const FVec3 DR1 = AngularProjection * RotationConstraints.D.ConstraintDRAxis[ConstraintIndex][1] * DeltaLambda;
+			const FVec3 DR1 = AngularProjection * RotationConstraints.Data.ConstraintDRAxis[ConstraintIndex][1] * DeltaLambda;
 			ApplyRotationDelta(1, DR1);
 
 			if (bLinearLocked)
@@ -2036,12 +2036,12 @@ void FPBDJointCachedSolver::ApplyAxisRotationProjection(
 				FVec3 PositionConstraintsArms1;
 				if (!PositionConstraints.bUseSimd)
 				{
-					PositionConstraintsArms1 = PositionConstraints.D.ConstraintArms[ConstraintIndex][1];
+					PositionConstraintsArms1 = PositionConstraints.Data.ConstraintArms[ConstraintIndex][1];
 				}
 				else
 				{
 					FVec3f PositionConstraintsArmsf;
-					VectorStoreFloat3(PositionConstraints.Simd.ConstraintArmsSimd[ConstraintIndex][1], &PositionConstraintsArmsf[0]);
+					VectorStoreFloat3(PositionConstraints.Simd.ConstraintArms[ConstraintIndex][1], &PositionConstraintsArmsf[0]);
 					PositionConstraintsArms1 = FVec3(PositionConstraintsArmsf);
 				}
 				const FVec3 DP1 = -AngularProjection * FVec3::CrossProduct(DR1, PositionConstraintsArms1);
@@ -2062,18 +2062,18 @@ void FPBDJointCachedSolver::ApplyRotationProjectionSimd(
 	VectorRegister4Float ProjAxes[3];
 	for (int32 CIndex = 0; CIndex < 3; ++CIndex)
 	{
-		ProjAxes[CIndex] = Private::VectorDot3FastX(DQDiff, RotationConstraints.Simd.ConstraintAxisSimd[CIndex]);
+		ProjAxes[CIndex] = Private::VectorDot3FastX(DQDiff, RotationConstraints.Simd.ConstraintAxis[CIndex]);
 	}
 	VectorRegister4Float ProjAxis = VectorUnpackLo(ProjAxes[0], ProjAxes[1]);
 	ProjAxis = VectorMoveLh(ProjAxis, ProjAxes[2]);
 	VectorRegister4Float DeltaAngle = VectorAdd(RotationConstraints.Simd.ConstraintCX, ProjAxis);
 
 
-	const VectorRegister4Float AngleGT = VectorCompareGT(DeltaAngle, RotationConstraints.Simd.ConstraintLimitsSimd);
-	const VectorRegister4Float AngleLT = VectorCompareLT(DeltaAngle, VectorNegate(RotationConstraints.Simd.ConstraintLimitsSimd));
+	const VectorRegister4Float AngleGT = VectorCompareGT(DeltaAngle, RotationConstraints.Simd.ConstraintLimits);
+	const VectorRegister4Float AngleLT = VectorCompareLT(DeltaAngle, VectorNegate(RotationConstraints.Simd.ConstraintLimits));
 
-	DeltaAngle = VectorSelect(AngleGT, VectorSubtract(DeltaAngle, RotationConstraints.Simd.ConstraintLimitsSimd),
-		VectorSelect(AngleLT, VectorAdd(DeltaAngle, RotationConstraints.Simd.ConstraintLimitsSimd), DeltaAngle));
+	DeltaAngle = VectorSelect(AngleGT, VectorSubtract(DeltaAngle, RotationConstraints.Simd.ConstraintLimits),
+		VectorSelect(AngleLT, VectorAdd(DeltaAngle, RotationConstraints.Simd.ConstraintLimits), DeltaAngle));
 
 	FRealSingle AngleTolerancef = FRealSingle(AngleTolerance);
 	const VectorRegister4Float AngleToleranceSimd = VectorLoadFloat1(&AngleTolerancef);
@@ -2085,7 +2085,7 @@ void FPBDJointCachedSolver::ApplyRotationProjectionSimd(
 		VectorRegister4Float IMs[3];
 		for (int32 CIndex = 0; CIndex < 3; ++CIndex)
 		{
-			IMs[CIndex] = Private::VectorDot3FastX(RotationConstraints.Simd.ConstraintAxisSimd[CIndex], RotationConstraints.Simd.ConstraintDRAxisSimd[CIndex][1]);
+			IMs[CIndex] = Private::VectorDot3FastX(RotationConstraints.Simd.ConstraintAxis[CIndex], RotationConstraints.Simd.ConstraintDRAxis[CIndex][1]);
 		}
 		VectorRegister4Float IM = VectorUnpackLo(IMs[0], IMs[1]);
 		IM = VectorMoveLh(IM, IMs[2]);
@@ -2093,7 +2093,7 @@ void FPBDJointCachedSolver::ApplyRotationProjectionSimd(
 
 		const FRealSingle SolverStiffnessf = FRealSingle(SolverStiffness);
 		VectorRegister4Float Stiffness = VectorLoadFloat1(&SolverStiffnessf);
-		VectorRegister4Float DeltaLambda = VectorDivide(VectorMultiply(VectorMultiply(Stiffness, RotationConstraints.Simd.ConstraintHardStiffnessSimd), DeltaAngle), IM);
+		VectorRegister4Float DeltaLambda = VectorDivide(VectorMultiply(VectorMultiply(Stiffness, RotationConstraints.Simd.ConstraintHardStiffness), DeltaAngle), IM);
 		DeltaLambda = VectorSelect(NeedsSolve, DeltaLambda, VectorZeroFloat());
 
 		VectorRegister4Float AngularProjection = VectorLoadFloat1(&AngularProjectionf);
@@ -2108,7 +2108,7 @@ void FPBDJointCachedSolver::ApplyRotationProjectionSimd(
 		VectorRegister4Float DR1s[3];
 		for (int32 CIndex = 0; CIndex < 3; ++CIndex)
 		{
-			DR1s[CIndex] = VectorMultiply(RotationConstraints.Simd.ConstraintDRAxisSimd[CIndex][1], DeltaLambdas[CIndex]);
+			DR1s[CIndex] = VectorMultiply(RotationConstraints.Simd.ConstraintDRAxis[CIndex][1], DeltaLambdas[CIndex]);
 			DR1 = VectorAdd(DR1s[CIndex], DR1);
 		}
 		FVec3f DR1f;
@@ -2124,11 +2124,11 @@ void FPBDJointCachedSolver::ApplyRotationProjectionSimd(
 				VectorRegister4Float PositionConstraintsArms1;
 				if (PositionConstraints.bUseSimd)
 				{
-					PositionConstraintsArms1 = PositionConstraints.Simd.ConstraintArmsSimd[CIndex][1];
+					PositionConstraintsArms1 = PositionConstraints.Simd.ConstraintArms[CIndex][1];
 				}
 				else
 				{
-					PositionConstraintsArms1 = MakeVectorRegisterFloatFromDouble(MakeVectorRegister(PositionConstraints.D.ConstraintArms[CIndex][1][0], PositionConstraints.D.ConstraintArms[CIndex][1][1], PositionConstraints.D.ConstraintArms[CIndex][1][2], 0.0));
+					PositionConstraintsArms1 = MakeVectorRegisterFloatFromDouble(MakeVectorRegister(PositionConstraints.Data.ConstraintArms[CIndex][1][0], PositionConstraints.Data.ConstraintArms[CIndex][1][1], PositionConstraints.Data.ConstraintArms[CIndex][1][2], 0.0));
 				}
 				DP1 = VectorMultiplyAdd(AngularProjection, VectorCross(DR1, PositionConstraintsArms1), DP1);
 			}
@@ -2174,35 +2174,35 @@ void FPBDJointCachedSolver::ApplyAxisPositionProjection(
 	if (PositionConstraints.bValidDatas[ConstraintIndex] && !PositionConstraints.bSoftLimit[ConstraintIndex])
 	{
 		const FVec3 CX = Body(1).DP() - Body(0).DP() +
-			FVec3::CrossProduct(Body(1).DQ(), PositionConstraints.D.ConstraintArms[ConstraintIndex][1]) -
-			FVec3::CrossProduct(Body(0).DQ(), PositionConstraints.D.ConstraintArms[ConstraintIndex][0]);
+			FVec3::CrossProduct(Body(1).DQ(), PositionConstraints.Data.ConstraintArms[ConstraintIndex][1]) -
+			FVec3::CrossProduct(Body(0).DQ(), PositionConstraints.Data.ConstraintArms[ConstraintIndex][0]);
 
-		FReal DeltaPosition = PositionConstraints.D.ConstraintCX[ConstraintIndex] + FVec3::DotProduct(CX, PositionConstraints.D.ConstraintAxis[ConstraintIndex]);
+		FReal DeltaPosition = PositionConstraints.Data.ConstraintCX[ConstraintIndex] + FVec3::DotProduct(CX, PositionConstraints.Data.ConstraintAxis[ConstraintIndex]);
 
 		bool NeedsSolve = false;
 		if (PositionConstraints.bLimitsCheck[ConstraintIndex])
 		{
-			if (DeltaPosition > PositionConstraints.D.ConstraintLimits[ConstraintIndex])
+			if (DeltaPosition > PositionConstraints.Data.ConstraintLimits[ConstraintIndex])
 			{
-				DeltaPosition -= PositionConstraints.D.ConstraintLimits[ConstraintIndex];
+				DeltaPosition -= PositionConstraints.Data.ConstraintLimits[ConstraintIndex];
 				NeedsSolve = true;
 			}
-			else if (DeltaPosition < -PositionConstraints.D.ConstraintLimits[ConstraintIndex])
+			else if (DeltaPosition < -PositionConstraints.Data.ConstraintLimits[ConstraintIndex])
 			{
-				DeltaPosition += PositionConstraints.D.ConstraintLimits[ConstraintIndex];
+				DeltaPosition += PositionConstraints.Data.ConstraintLimits[ConstraintIndex];
 				NeedsSolve = true;
 			}
 		}
 		if (!PositionConstraints.bLimitsCheck[ConstraintIndex] || (PositionConstraints.bLimitsCheck[ConstraintIndex] && NeedsSolve && FMath::Abs(DeltaPosition) > PositionTolerance))
 		{
-			const FVec3 AngularAxis1 = FVec3::CrossProduct(PositionConstraints.D.ConstraintArms[ConstraintIndex][1], PositionConstraints.D.ConstraintAxis[ConstraintIndex]);
-			const FReal IM = InvM(1) - FVec3::DotProduct(AngularAxis1, PositionConstraints.D.ConstraintDRAxis[ConstraintIndex][1]);
-			const FReal DeltaLambda = SolverStiffness * PositionConstraints.D.ConstraintHardStiffness[ConstraintIndex] * DeltaPosition / IM;
+			const FVec3 AngularAxis1 = FVec3::CrossProduct(PositionConstraints.Data.ConstraintArms[ConstraintIndex][1], PositionConstraints.Data.ConstraintAxis[ConstraintIndex]);
+			const FReal IM = InvM(1) - FVec3::DotProduct(AngularAxis1, PositionConstraints.Data.ConstraintDRAxis[ConstraintIndex][1]);
+			const FReal DeltaLambda = SolverStiffness * PositionConstraints.Data.ConstraintHardStiffness[ConstraintIndex] * DeltaPosition / IM;
 
-			const FVec3 DX = PositionConstraints.D.ConstraintAxis[ConstraintIndex] * DeltaLambda;
+			const FVec3 DX = PositionConstraints.Data.ConstraintAxis[ConstraintIndex] * DeltaLambda;
 
 			const FVec3 DP1 = -LinearProjection * InvM(1) * DX;
-			const FVec3 DR1 = LinearProjection * PositionConstraints.D.ConstraintDRAxis[ConstraintIndex][1] * DeltaLambda;
+			const FVec3 DR1 = LinearProjection * PositionConstraints.Data.ConstraintDRAxis[ConstraintIndex][1] * DeltaLambda;
 
 			ApplyPositionDelta(1, DP1);
 			ApplyRotationDelta(1, DR1);
@@ -2224,12 +2224,12 @@ void FPBDJointCachedSolver::ApplyPositionProjectionSimd(
 	VectorRegister4Float DeltaPositions[3];
 	for (int32 CIndex = 0; CIndex < 3; ++CIndex)
 	{
-		VectorRegister4Float  Cross1 = VectorCross(Body1DQ, PositionConstraints.Simd.ConstraintArmsSimd[CIndex][1]);
-		VectorRegister4Float  Cross0 = VectorCross(Body0DQ, PositionConstraints.Simd.ConstraintArmsSimd[CIndex][0]);
+		VectorRegister4Float  Cross1 = VectorCross(Body1DQ, PositionConstraints.Simd.ConstraintArms[CIndex][1]);
+		VectorRegister4Float  Cross0 = VectorCross(Body0DQ, PositionConstraints.Simd.ConstraintArms[CIndex][0]);
 		VectorRegister4Float  CrossDiff = VectorSubtract(Cross1, Cross0);
 		VectorRegister4Float CX = VectorAdd(DPDiff, CrossDiff);
 
-		DeltaPositions[CIndex] = Private::VectorDot3FastX(CX, PositionConstraints.Simd.ConstraintAxisSimd[CIndex]);
+		DeltaPositions[CIndex] = Private::VectorDot3FastX(CX, PositionConstraints.Simd.ConstraintAxis[CIndex]);
 	}
 
 	VectorRegister4Float DeltaPosition = VectorUnpackLo(DeltaPositions[0], DeltaPositions[1]);
@@ -2240,7 +2240,7 @@ void FPBDJointCachedSolver::ApplyPositionProjectionSimd(
 	const FRealSingle SolverStiffnessf = FRealSingle(SolverStiffness);
 	VectorRegister4Float Stiffness = VectorLoadFloat1(&SolverStiffnessf);
 
-	Stiffness = VectorMultiply(Stiffness, PositionConstraints.Simd.ConstraintHardStiffnessSimd);
+	Stiffness = VectorMultiply(Stiffness, PositionConstraints.Simd.ConstraintHardStiffness);
 	DeltaPosition = VectorMultiply(Stiffness, DeltaPosition);
 
 	FRealSingle InvM1f = FRealSingle(InvM(1));
@@ -2249,8 +2249,8 @@ void FPBDJointCachedSolver::ApplyPositionProjectionSimd(
 	VectorRegister4Float IMs[3];
 	for (int32 CIndex = 0; CIndex < 3; CIndex++)
 	{
-		const VectorRegister4Float AngularAxis1 = VectorCross(PositionConstraints.Simd.ConstraintArmsSimd[CIndex][1], PositionConstraints.Simd.ConstraintAxisSimd[CIndex]);
-		IMs[CIndex] = Private::VectorDot3FastX(AngularAxis1, PositionConstraints.Simd.ConstraintDRAxisSimd[CIndex][1]);
+		const VectorRegister4Float AngularAxis1 = VectorCross(PositionConstraints.Simd.ConstraintArms[CIndex][1], PositionConstraints.Simd.ConstraintAxis[CIndex]);
+		IMs[CIndex] = Private::VectorDot3FastX(AngularAxis1, PositionConstraints.Simd.ConstraintDRAxis[CIndex][1]);
 	}
 	VectorRegister4Float IM = VectorUnpackLo(IMs[0], IMs[1]);
 	IM = VectorMoveLh(IM, IMs[2]);
@@ -2271,9 +2271,9 @@ void FPBDJointCachedSolver::ApplyPositionProjectionSimd(
 	VectorRegister4Float DR1 = VectorZeroFloat();
 	for (int32 CIndex = 0; CIndex < 3; CIndex++)
 	{
-		VectorRegister4Float DX = VectorMultiply(PositionConstraints.Simd.ConstraintAxisSimd[CIndex], DeltaLambdas[CIndex]);
+		VectorRegister4Float DX = VectorMultiply(PositionConstraints.Simd.ConstraintAxis[CIndex], DeltaLambdas[CIndex]);
 		DP1 = VectorMultiplyAdd(NegLinPorjInvM, DX, DP1);
-		DR1 = VectorMultiplyAdd(VectorMultiply(LinearProjectionSimd, PositionConstraints.Simd.ConstraintDRAxisSimd[CIndex][1]), DeltaLambdas[CIndex], DR1);
+		DR1 = VectorMultiplyAdd(VectorMultiply(LinearProjectionSimd, PositionConstraints.Simd.ConstraintDRAxis[CIndex][1]), DeltaLambdas[CIndex], DR1);
 	}
 	FVec3f DP1f;
 	VectorStoreFloat3(DP1, &DP1f[0]);
@@ -2324,19 +2324,19 @@ void FPBDJointCachedSolver::ApplyAxisPositionTeleport(
 
 	if (PositionConstraints.bValidDatas[ConstraintIndex] && !PositionConstraints.bSoftLimit[ConstraintIndex])
 	{
-		FReal DeltaPosition = PositionConstraints.D.ConstraintCX[ConstraintIndex];
+		FReal DeltaPosition = PositionConstraints.Data.ConstraintCX[ConstraintIndex];
 
 		bool NeedsSolve = false;
 		if (PositionConstraints.bLimitsCheck[ConstraintIndex])
 		{
-			if (DeltaPosition > PositionConstraints.D.ConstraintLimits[ConstraintIndex])
+			if (DeltaPosition > PositionConstraints.Data.ConstraintLimits[ConstraintIndex])
 			{
-				DeltaPosition -= PositionConstraints.D.ConstraintLimits[ConstraintIndex];
+				DeltaPosition -= PositionConstraints.Data.ConstraintLimits[ConstraintIndex];
 				NeedsSolve = true;
 			}
-			else if (DeltaPosition < -PositionConstraints.D.ConstraintLimits[ConstraintIndex])
+			else if (DeltaPosition < -PositionConstraints.Data.ConstraintLimits[ConstraintIndex])
 			{
-				DeltaPosition += PositionConstraints.D.ConstraintLimits[ConstraintIndex];
+				DeltaPosition += PositionConstraints.Data.ConstraintLimits[ConstraintIndex];
 				NeedsSolve = true;
 			}
 		}
@@ -2344,7 +2344,7 @@ void FPBDJointCachedSolver::ApplyAxisPositionTeleport(
 		{
 			if (FMath::Abs(DeltaPosition) > TeleportDistance)
 			{
-				const FVec3 DP1 = -DeltaPosition * PositionConstraints.D.ConstraintAxis[ConstraintIndex];
+				const FVec3 DP1 = -DeltaPosition * PositionConstraints.Data.ConstraintAxis[ConstraintIndex];
 				ApplyPositionDelta(1, DP1);
 			}
 		}
@@ -2365,10 +2365,10 @@ void FPBDJointCachedSolver::ApplyPositionTeleportSimd(
 		ConstraintCXs[0] = VectorReplicate(ConstraintCX, 0);
 		ConstraintCXs[1] = VectorReplicate(ConstraintCX, 1);
 		ConstraintCXs[2] = VectorReplicate(ConstraintCX, 2);
-		VectorRegister4Float DP1 = VectorMultiply(PositionConstraints.Simd.ConstraintAxisSimd[0], ConstraintCXs[0]);
+		VectorRegister4Float DP1 = VectorMultiply(PositionConstraints.Simd.ConstraintAxis[0], ConstraintCXs[0]);
 		for (int32 ConstraintIndex = 1; ConstraintIndex < 3; ++ConstraintIndex)
 		{
-			DP1 = VectorMultiplyAdd(PositionConstraints.Simd.ConstraintAxisSimd[ConstraintIndex], ConstraintCXs[ConstraintIndex], DP1);
+			DP1 = VectorMultiplyAdd(PositionConstraints.Simd.ConstraintAxis[ConstraintIndex], ConstraintCXs[ConstraintIndex], DP1);
 		}
 		FVec3f DP1f;
 		VectorStoreFloat3(DP1, &DP1f[0]);
@@ -2408,13 +2408,13 @@ void FPBDJointCachedSolver::ApplyAxisPositionDrive(
 	const FReal Dt)
 {
 	check(PositionDrives.bUseSimd == false);
-	const FReal Stiffness = PositionDrives.D.ConstraintSoftStiffness[ConstraintIndex];
-	const FReal Damping = PositionDrives.D.ConstraintSoftDamping[ConstraintIndex];
-	const FReal IM = PositionDrives.D.ConstraintSoftIM[ConstraintIndex];
+	const FReal Stiffness = PositionDrives.Data.ConstraintSoftStiffness[ConstraintIndex];
+	const FReal Damping = PositionDrives.Data.ConstraintSoftDamping[ConstraintIndex];
+	const FReal IM = PositionDrives.Data.ConstraintSoftIM[ConstraintIndex];
 
-	const FVec3 Delta0 = Body(0).DP() + FVec3::CrossProduct(Body(0).DQ(), PositionDrives.D.ConstraintArms[ConstraintIndex][0]);
-	const FVec3 Delta1 = Body(1).DP() + FVec3::CrossProduct(Body(1).DQ(), PositionDrives.D.ConstraintArms[ConstraintIndex][1]);
-	const FReal CX = PositionDrives.D.ConstraintCX[ConstraintIndex] + FVec3::DotProduct(Delta1 - Delta0, PositionDrives.D.ConstraintAxis[ConstraintIndex]);
+	const FVec3 Delta0 = Body(0).DP() + FVec3::CrossProduct(Body(0).DQ(), PositionDrives.Data.ConstraintArms[ConstraintIndex][0]);
+	const FVec3 Delta1 = Body(1).DP() + FVec3::CrossProduct(Body(1).DQ(), PositionDrives.Data.ConstraintArms[ConstraintIndex][1]);
+	const FReal CX = PositionDrives.Data.ConstraintCX[ConstraintIndex] + FVec3::DotProduct(Delta1 - Delta0, PositionDrives.Data.ConstraintAxis[ConstraintIndex]);
 
 	FReal CVDt = 0;
 	if (Damping > UE_KINDA_SMALL_NUMBER)
@@ -2422,15 +2422,15 @@ void FPBDJointCachedSolver::ApplyAxisPositionDrive(
 		const FVec3 V0Dt = FVec3::CalculateVelocity(InitConnectorXs[0], ConnectorXs[0] + Delta0, 1.0f);
 		const FVec3 V1Dt = FVec3::CalculateVelocity(InitConnectorXs[1], ConnectorXs[1] + Delta1, 1.0f);
 		const FReal TargetVDt = PositionDrives.ConstraintVX[ConstraintIndex] * Dt;
-		CVDt = TargetVDt + FVec3::DotProduct(V0Dt - V1Dt, PositionDrives.D.ConstraintAxis[ConstraintIndex]);
+		CVDt = TargetVDt + FVec3::DotProduct(V0Dt - V1Dt, PositionDrives.Data.ConstraintAxis[ConstraintIndex]);
 	}
 
-	FReal Lambda = PositionDrives.D.ConstraintLambda[ConstraintIndex];
+	FReal Lambda = PositionDrives.Data.ConstraintLambda[ConstraintIndex];
 	FReal DeltaLambda = SolverStiffness * (Stiffness * CX - Damping * CVDt - Lambda) / IM;
 	Lambda += DeltaLambda;
 
 	PositionDrives.ApplyMaxLambda(ConstraintIndex, DeltaLambda, Lambda);
-	PositionDrives.D.ConstraintLambda[ConstraintIndex] = Lambda;
+	PositionDrives.Data.ConstraintLambda[ConstraintIndex] = Lambda;
 
 	SolvePositionConstraintDelta(ConstraintIndex, DeltaLambda, PositionDrives);
 }
@@ -2464,13 +2464,13 @@ void FPBDJointCachedSolver::ApplyAxisPositionVelocityDrive(
 		return;
 	}
 
-	const FReal MassScale = PositionDrives.bAccelerationMode ? (FReal(1) / PositionDrives.D.ConstraintHardIM[ConstraintIndex]) : FReal(1);
-	const FReal IM = MassScale * Damping * PositionDrives.D.ConstraintHardIM[ConstraintIndex] + (FReal)1;
+	const FReal MassScale = PositionDrives.bAccelerationMode ? (FReal(1) / PositionDrives.Data.ConstraintHardIM[ConstraintIndex]) : FReal(1);
+	const FReal IM = MassScale * Damping * PositionDrives.Data.ConstraintHardIM[ConstraintIndex] + (FReal)1;
 
 	// Velocity error to correct
-	const FVec3 V0 = V(0) + FVec3::CrossProduct(W(0), PositionDrives.D.ConstraintArms[ConstraintIndex][0]);
-	const FVec3 V1 = V(1) + FVec3::CrossProduct(W(1), PositionDrives.D.ConstraintArms[ConstraintIndex][1]);
-	const FReal VRel = FVec3::DotProduct(V1 - V0, PositionDrives.D.ConstraintAxis[ConstraintIndex]);
+	const FVec3 V0 = V(0) + FVec3::CrossProduct(W(0), PositionDrives.Data.ConstraintArms[ConstraintIndex][0]);
+	const FVec3 V1 = V(1) + FVec3::CrossProduct(W(1), PositionDrives.Data.ConstraintArms[ConstraintIndex][1]);
+	const FReal VRel = FVec3::DotProduct(V1 - V0, PositionDrives.Data.ConstraintAxis[ConstraintIndex]);
 	const FReal TargetV = PositionDrives.ConstraintVX[ConstraintIndex];
 	const FReal CV = (VRel - TargetV);
 
@@ -2481,25 +2481,25 @@ void FPBDJointCachedSolver::ApplyAxisPositionVelocityDrive(
 	// Apply limits and accumulate total impulse 
 	// (NOTE: Limits and net impulses are position based, not velocity based)
 	FReal DeltaLambda = DeltaLambdaVel * Dt;
-	FReal Lambda = PositionDrives.D.ConstraintLambda[ConstraintIndex] + DeltaLambda;
+	FReal Lambda = PositionDrives.Data.ConstraintLambda[ConstraintIndex] + DeltaLambda;
 	PositionDrives.ApplyMaxLambda(ConstraintIndex, DeltaLambda, Lambda);
-	PositionDrives.D.ConstraintLambda[ConstraintIndex] = Lambda;
+	PositionDrives.Data.ConstraintLambda[ConstraintIndex] = Lambda;
 	DeltaLambdaVel = DeltaLambda / Dt;
 
 	LambdaVel += DeltaLambdaVel;
-	const FVec3 Impulse = DeltaLambdaVel * PositionDrives.D.ConstraintAxis[ConstraintIndex];
+	const FVec3 Impulse = DeltaLambdaVel * PositionDrives.Data.ConstraintAxis[ConstraintIndex];
 
 	if (Body(0).IsDynamic())
 	{
 		const FVec3 DV0 = InvM(0) * Impulse;
-		const FVec3 DW0 = PositionDrives.D.ConstraintDRAxis[ConstraintIndex][0] * DeltaLambdaVel;
+		const FVec3 DW0 = PositionDrives.Data.ConstraintDRAxis[ConstraintIndex][0] * DeltaLambdaVel;
 
 		Body(0).ApplyVelocityDelta(DV0, DW0);
 	}
 	if (Body(1).IsDynamic())
 	{
 		const FVec3 DV1 = -InvM(1) * Impulse;
-		const FVec3 DW1 = PositionDrives.D.ConstraintDRAxis[ConstraintIndex][1] * DeltaLambdaVel;
+		const FVec3 DW1 = PositionDrives.Data.ConstraintDRAxis[ConstraintIndex][1] * DeltaLambdaVel;
 
 		Body(1).ApplyVelocityDelta(DV1, DW1);
 	}
@@ -2592,12 +2592,12 @@ void FPBDJointCachedSolver::InitRotationConstraintDriveSimd(
 		RotationDrives.ConstraintVX[ConstraintIndex] = 0.0;
 		RotationDrives.ConstraintRestitution[ConstraintIndex] = 0.0;
 
-		RotationDrives.Simd.ConstraintArmsSimd[ConstraintIndex][0] = VectorZeroFloat();
-		RotationDrives.Simd.ConstraintArmsSimd[ConstraintIndex][1] = VectorZeroFloat();
-		RotationDrives.Simd.ConstraintAxisSimd[ConstraintIndex] = MakeVectorRegisterFloatFromDouble(MakeVectorRegister(ConstraintAxes[ConstraintIndex][0], ConstraintAxes[ConstraintIndex][1], ConstraintAxes[ConstraintIndex][2], 0.0f));
+		RotationDrives.Simd.ConstraintArms[ConstraintIndex][0] = VectorZeroFloat();
+		RotationDrives.Simd.ConstraintArms[ConstraintIndex][1] = VectorZeroFloat();
+		RotationDrives.Simd.ConstraintAxis[ConstraintIndex] = MakeVectorRegisterFloatFromDouble(MakeVectorRegister(ConstraintAxes[ConstraintIndex][0], ConstraintAxes[ConstraintIndex][1], ConstraintAxes[ConstraintIndex][2], 0.0f));
 
 		// InitRotationDatasMass(RotationDrives, ConstraintIndex, Dt);
-		const VectorRegister4Float& Axis = RotationDrives.Simd.ConstraintAxisSimd[ConstraintIndex];
+		const VectorRegister4Float& Axis = RotationDrives.Simd.ConstraintAxis[ConstraintIndex];
 		const VectorRegister4Float AxisX = VectorReplicate(Axis, 0);
 		const VectorRegister4Float AxisY = VectorReplicate(Axis, 1);
 		const VectorRegister4Float AxisZ = VectorReplicate(Axis, 2);
@@ -2618,8 +2618,8 @@ void FPBDJointCachedSolver::InitRotationConstraintDriveSimd(
 
 		ConstraintHardIM[ConstraintIndex] = II0 + II1;
 
-		RotationDrives.Simd.ConstraintDRAxisSimd[ConstraintIndex][0] = IA0;
-		RotationDrives.Simd.ConstraintDRAxisSimd[ConstraintIndex][1] = VectorNegate(IA1);
+		RotationDrives.Simd.ConstraintDRAxis[ConstraintIndex][0] = IA0;
+		RotationDrives.Simd.ConstraintDRAxis[ConstraintIndex][1] = VectorNegate(IA1);
 
 
 		check(RotationDrives.bSoftLimit[ConstraintIndex])
@@ -2629,10 +2629,10 @@ void FPBDJointCachedSolver::InitRotationConstraintDriveSimd(
 		ConstraintSoftIM[ConstraintIndex] = (ConstraintSoftStiffness[ConstraintIndex] + ConstraintSoftDamping[ConstraintIndex]) * ConstraintHardIM[ConstraintIndex] + (FReal)1;
 
 	}
-	RotationDrives.Simd.ConstraintHardIMSimd = MakeVectorRegisterFloatFromDouble(MakeVectorRegister(ConstraintHardIM[0], ConstraintHardIM[1], ConstraintHardIM[2], 0.0f));
-	RotationDrives.Simd.ConstraintSoftStiffnessSimd = MakeVectorRegisterFloatFromDouble(MakeVectorRegister(ConstraintSoftStiffness[0], ConstraintSoftStiffness[1], ConstraintSoftStiffness[2], 0.0f));
-	RotationDrives.Simd.ConstraintSoftDampingSimd = MakeVectorRegisterFloatFromDouble(MakeVectorRegister(ConstraintSoftDamping[0], ConstraintSoftDamping[1], ConstraintSoftDamping[2], 0.0f));
-	RotationDrives.Simd.ConstraintSoftIMSimd = MakeVectorRegisterFloatFromDouble(MakeVectorRegister(ConstraintSoftIM[0], ConstraintSoftIM[1], ConstraintSoftIM[2], 0.0f));
+	RotationDrives.Simd.ConstraintHardIM = MakeVectorRegisterFloatFromDouble(MakeVectorRegister(ConstraintHardIM[0], ConstraintHardIM[1], ConstraintHardIM[2], 0.0f));
+	RotationDrives.Simd.ConstraintSoftStiffness = MakeVectorRegisterFloatFromDouble(MakeVectorRegister(ConstraintSoftStiffness[0], ConstraintSoftStiffness[1], ConstraintSoftStiffness[2], 0.0f));
+	RotationDrives.Simd.ConstraintSoftDamping = MakeVectorRegisterFloatFromDouble(MakeVectorRegister(ConstraintSoftDamping[0], ConstraintSoftDamping[1], ConstraintSoftDamping[2], 0.0f));
+	RotationDrives.Simd.ConstraintSoftIM = MakeVectorRegisterFloatFromDouble(MakeVectorRegister(ConstraintSoftIM[0], ConstraintSoftIM[1], ConstraintSoftIM[2], 0.0f));
 }
 
 void FPBDJointCachedSolver::InitSwingTwistDrives(
@@ -2667,9 +2667,9 @@ void FPBDJointCachedSolver::InitSwingTwistDrives(
 
 	RotationDrives.bAccelerationMode = FPBDJointUtilities::GetAngularDriveAccelerationMode(SolverSettings, JointSettings);
 
-	const bool bUseTwistDrive = bTwistDriveEnabled && (((FMath::Abs(DTwistAngle) > AngleTolerance) && (RotationDrives.D.ConstraintSoftStiffness[TW] > 0.0f)) || (RotationDrives.D.ConstraintSoftDamping[TW] > 0.0f));
-	const bool bUseSwing1Drive = bSwing1DriveEnabled && (((FMath::Abs(DSwing1Angle) > AngleTolerance) && (RotationDrives.D.ConstraintSoftStiffness[S1] > 0.0f)) || (RotationDrives.D.ConstraintSoftDamping[S1] > 0.0f));
-	const bool bUseSwing2Drive = bSwing2DriveEnabled && (((FMath::Abs(DSwing2Angle) > AngleTolerance) && (RotationDrives.D.ConstraintSoftStiffness[S2] > 0.0f)) || (RotationDrives.D.ConstraintSoftDamping[S2] > 0.0f));
+	const bool bUseTwistDrive = bTwistDriveEnabled && (((FMath::Abs(DTwistAngle) > AngleTolerance) && (RotationDrives.Data.ConstraintSoftStiffness[TW] > 0.0f)) || (RotationDrives.Data.ConstraintSoftDamping[TW] > 0.0f));
+	const bool bUseSwing1Drive = bSwing1DriveEnabled && (((FMath::Abs(DSwing1Angle) > AngleTolerance) && (RotationDrives.Data.ConstraintSoftStiffness[S1] > 0.0f)) || (RotationDrives.Data.ConstraintSoftDamping[S1] > 0.0f));
+	const bool bUseSwing2Drive = bSwing2DriveEnabled && (((FMath::Abs(DSwing2Angle) > AngleTolerance) && (RotationDrives.Data.ConstraintSoftStiffness[S2] > 0.0f)) || (RotationDrives.Data.ConstraintSoftDamping[S2] > 0.0f));
 	RotationDrives.bUseSimd = bUseSimd && bUseTwistDrive && bUseSwing1Drive && bUseSwing2Drive;
 
 	if (RotationDrives.bUseSimd)
@@ -2677,7 +2677,7 @@ void FPBDJointCachedSolver::InitSwingTwistDrives(
 		FVec3 ConstraintAxes[3] = { ConnectorRs[1] * FJointConstants::TwistAxis(), ConnectorRs[1] * FJointConstants::Swing1Axis(), ConnectorRs[1] * FJointConstants::Swing2Axis() };
 		InitRotationConstraintDriveSimd(ConstraintAxes, Dt, FVec3(DTwistAngle, DSwing1Angle, DSwing2Angle));
 
-		RotationDrives.Simd.ConstraintHardStiffnessSimd = VectorZeroFloat();
+		RotationDrives.Simd.ConstraintHardStiffness = VectorZeroFloat();
 		RotationDrives.Simd.ConstraintLambda = VectorZeroFloat();
 
 		RotationDrives.ConstraintVX[TW] = JointSettings.AngularDriveVelocityTarget[TW];
@@ -2740,7 +2740,7 @@ void FPBDJointCachedSolver::InitSLerpDrive(
 	if (RotationDrives.bUseSimd)
 	{
 		RotationDrives.Simd.ConstraintLambda = VectorZeroFloat();
-		RotationDrives.Simd.ConstraintHardStiffnessSimd = VectorZeroFloat();
+		RotationDrives.Simd.ConstraintHardStiffness = VectorZeroFloat();
 		InitRotationConstraintDriveSimd(Axes, Dt, AxisAngles);
 	}
 	else
@@ -2789,16 +2789,16 @@ void FPBDJointCachedSolver::ApplyAxisRotationDrive(
 	const int32 ConstraintIndex,
 	const FReal Dt)
 {
-	const FReal Stiffness = RotationDrives.D.ConstraintSoftStiffness[ConstraintIndex];
-	const FReal Damping = RotationDrives.D.ConstraintSoftDamping[ConstraintIndex];
-	const FReal IM = RotationDrives.D.ConstraintSoftIM[ConstraintIndex];
+	const FReal Stiffness = RotationDrives.Data.ConstraintSoftStiffness[ConstraintIndex];
+	const FReal Damping = RotationDrives.Data.ConstraintSoftDamping[ConstraintIndex];
+	const FReal IM = RotationDrives.Data.ConstraintSoftIM[ConstraintIndex];
 
 	// Stiffness position delta
 	FReal CX = 0;
 	if (Stiffness > UE_KINDA_SMALL_NUMBER)
 	{
-		const FReal DX = FVec3::DotProduct(Body(1).DQ() - Body(0).DQ(), RotationDrives.D.ConstraintAxis[ConstraintIndex]);
-		const FReal TargetX = RotationDrives.D.ConstraintCX[ConstraintIndex];
+		const FReal DX = FVec3::DotProduct(Body(1).DQ() - Body(0).DQ(), RotationDrives.Data.ConstraintAxis[ConstraintIndex]);
+		const FReal TargetX = RotationDrives.Data.ConstraintCX[ConstraintIndex];
 		CX = TargetX + DX;
 	}
 
@@ -2809,15 +2809,15 @@ void FPBDJointCachedSolver::ApplyAxisRotationDrive(
 		const FVec3 W0Dt = FVec3(Body(0).DQ()) + ConnectorWDts[0];
 		const FVec3 W1Dt = FVec3(Body(1).DQ()) + ConnectorWDts[1];
 		const FReal TargetW = RotationDrives.ConstraintVX[ConstraintIndex];
-		CVDt = (TargetW * Dt) + FVec3::DotProduct(RotationDrives.D.ConstraintAxis[ConstraintIndex], W0Dt - W1Dt);
+		CVDt = (TargetW * Dt) + FVec3::DotProduct(RotationDrives.Data.ConstraintAxis[ConstraintIndex], W0Dt - W1Dt);
 	}
 
-	FReal Lambda = RotationDrives.D.ConstraintLambda[ConstraintIndex];
+	FReal Lambda = RotationDrives.Data.ConstraintLambda[ConstraintIndex];
 	FReal DeltaLambda = SolverStiffness * (Stiffness * CX - Damping * CVDt - Lambda) / IM;
 	Lambda += DeltaLambda;
 
 	RotationDrives.ApplyMaxLambda(ConstraintIndex, DeltaLambda, Lambda);
-	RotationDrives.D.ConstraintLambda[ConstraintIndex] = Lambda;
+	RotationDrives.Data.ConstraintLambda[ConstraintIndex] = Lambda;
 
 	SolveRotationConstraintDelta(ConstraintIndex, DeltaLambda, true, RotationDrives);
 }
@@ -2839,8 +2839,8 @@ void FPBDJointCachedSolver::ApplyRotationDrivesSimd(
 
 	for (int32 ConstraintIndex = 0; ConstraintIndex < 3; ConstraintIndex++)
 	{
-		DXArray[ConstraintIndex] = Private::VectorDot3FastX(DQDiff, RotationDrives.Simd.ConstraintAxisSimd[ConstraintIndex]);
-		CVDtArray[ConstraintIndex] = Private::VectorDot3FastX(RotationDrives.Simd.ConstraintAxisSimd[ConstraintIndex], WDiff);
+		DXArray[ConstraintIndex] = Private::VectorDot3FastX(DQDiff, RotationDrives.Simd.ConstraintAxis[ConstraintIndex]);
+		CVDtArray[ConstraintIndex] = Private::VectorDot3FastX(RotationDrives.Simd.ConstraintAxis[ConstraintIndex], WDiff);
 	}
 	VectorRegister4Float DX = VectorUnpackLo(DXArray[0], DXArray[1]);
 	DX = VectorMoveLh(DX, DXArray[2]);
@@ -2856,8 +2856,8 @@ void FPBDJointCachedSolver::ApplyRotationDrivesSimd(
 
 	const FRealSingle SolverStiffnessf = FRealSingle(SolverStiffness);
 	VectorRegister4Float Stiffness = VectorLoadFloat1(&SolverStiffnessf);
-	VectorRegister4Float DeltaLambda = VectorMultiply(Stiffness, VectorDivide(VectorSubtract(VectorMultiply(RotationDrives.Simd.ConstraintSoftStiffnessSimd, CX),
-		VectorMultiplyAdd(RotationDrives.Simd.ConstraintSoftDampingSimd, CVDt, RotationDrives.Simd.ConstraintLambda)), RotationDrives.Simd.ConstraintSoftIMSimd));
+	VectorRegister4Float DeltaLambda = VectorMultiply(Stiffness, VectorDivide(VectorSubtract(VectorMultiply(RotationDrives.Simd.ConstraintSoftStiffness, CX),
+		VectorMultiplyAdd(RotationDrives.Simd.ConstraintSoftDamping, CVDt, RotationDrives.Simd.ConstraintLambda)), RotationDrives.Simd.ConstraintSoftIM));
 	RotationDrives.Simd.ConstraintLambda = VectorAdd(RotationDrives.Simd.ConstraintLambda, DeltaLambda);
 
 	// Should check MaxLambda eventually
@@ -2870,7 +2870,7 @@ void FPBDJointCachedSolver::ApplyRotationDrivesSimd(
 	VectorRegister4Float DeltaImpulses[3];
 	for (int32 ConstraintIndex = 0; ConstraintIndex < 3; ConstraintIndex++)
 	{
-		DeltaImpulses[ConstraintIndex] = VectorMultiply(RotationDrives.Simd.ConstraintAxisSimd[ConstraintIndex], DeltaLambdas[ConstraintIndex]);
+		DeltaImpulses[ConstraintIndex] = VectorMultiply(RotationDrives.Simd.ConstraintAxis[ConstraintIndex], DeltaLambdas[ConstraintIndex]);
 	}
 
 	if (Body(0).IsDynamic())
@@ -2878,7 +2878,7 @@ void FPBDJointCachedSolver::ApplyRotationDrivesSimd(
 		VectorRegister4Float DR0 = VectorZeroFloat();
 		for (int32 ConstraintIndex = 0; ConstraintIndex < 3; ConstraintIndex++)
 		{
-			const VectorRegister4Float Axis = VectorDot3(RotationDrives.Simd.ConstraintAxisSimd[ConstraintIndex], RotationDrives.Simd.ConstraintDRAxisSimd[ConstraintIndex][0]);
+			const VectorRegister4Float Axis = VectorDot3(RotationDrives.Simd.ConstraintAxis[ConstraintIndex], RotationDrives.Simd.ConstraintDRAxis[ConstraintIndex][0]);
 			DR0 = VectorMultiplyAdd(DeltaImpulses[ConstraintIndex], Axis, DR0);
 		}
 		FVec3f DR0f;
@@ -2890,7 +2890,7 @@ void FPBDJointCachedSolver::ApplyRotationDrivesSimd(
 		VectorRegister4Float DR1 = VectorZeroFloat();
 		for (int32 ConstraintIndex = 0; ConstraintIndex < 3; ConstraintIndex++)
 		{
-			const VectorRegister4Float Axis = VectorDot3(RotationDrives.Simd.ConstraintAxisSimd[ConstraintIndex], RotationDrives.Simd.ConstraintDRAxisSimd[ConstraintIndex][1]);
+			const VectorRegister4Float Axis = VectorDot3(RotationDrives.Simd.ConstraintAxis[ConstraintIndex], RotationDrives.Simd.ConstraintDRAxis[ConstraintIndex][1]);
 			DR1 = VectorMultiplyAdd(DeltaImpulses[ConstraintIndex], Axis, DR1);
 		}
 		FVec3f DR1f;
@@ -2929,11 +2929,11 @@ void FPBDJointCachedSolver::ApplyAxisRotationVelocityDrive(
 		return;
 	}
 
-	const FReal MassScale = RotationDrives.bAccelerationMode ? (FReal(1) / RotationDrives.D.ConstraintHardIM[ConstraintIndex]) : FReal(1);
-	const FReal IM = MassScale * Damping * RotationDrives.D.ConstraintHardIM[ConstraintIndex] + (FReal)1;
+	const FReal MassScale = RotationDrives.bAccelerationMode ? (FReal(1) / RotationDrives.Data.ConstraintHardIM[ConstraintIndex]) : FReal(1);
+	const FReal IM = MassScale * Damping * RotationDrives.Data.ConstraintHardIM[ConstraintIndex] + (FReal)1;
 
 	// Angular velocity error to correct
-	const FReal WRel = FVec3::DotProduct(W(1) - W(0), RotationDrives.D.ConstraintAxis[ConstraintIndex]);
+	const FReal WRel = FVec3::DotProduct(W(1) - W(0), RotationDrives.Data.ConstraintAxis[ConstraintIndex]);
 	const FReal TargetW = RotationDrives.ConstraintVX[ConstraintIndex];
 	const FReal CV = (WRel - TargetW);
 
@@ -2944,22 +2944,22 @@ void FPBDJointCachedSolver::ApplyAxisRotationVelocityDrive(
 	// Apply limits and accumulate total impulse 
 	// (NOTE: Limits and net impulses are position based, not velocity based)
 	FReal DeltaLambda = DeltaLambdaVel * Dt;
-	FReal Lambda = RotationDrives.D.ConstraintLambda[ConstraintIndex] + DeltaLambda;
+	FReal Lambda = RotationDrives.Data.ConstraintLambda[ConstraintIndex] + DeltaLambda;
 	RotationDrives.ApplyMaxLambda(ConstraintIndex, DeltaLambda, Lambda);
-	RotationDrives.D.ConstraintLambda[ConstraintIndex] = Lambda;
+	RotationDrives.Data.ConstraintLambda[ConstraintIndex] = Lambda;
 	DeltaLambdaVel = DeltaLambda / Dt;
 
 	LambdaVel += DeltaLambdaVel;
-	const FVec3 Impulse = DeltaLambdaVel * RotationDrives.D.ConstraintAxis[ConstraintIndex];
+	const FVec3 Impulse = DeltaLambdaVel * RotationDrives.Data.ConstraintAxis[ConstraintIndex];
 
 	if (Body(0).IsDynamic())
 	{
-		const FVec3 DW0 = RotationDrives.D.ConstraintDRAxis[ConstraintIndex][0] * DeltaLambdaVel;
+		const FVec3 DW0 = RotationDrives.Data.ConstraintDRAxis[ConstraintIndex][0] * DeltaLambdaVel;
 		Body(0).ApplyAngularVelocityDelta(DW0);
 	}
 	if (Body(1).IsDynamic())
 	{
-		const FVec3 DW1 = RotationDrives.D.ConstraintDRAxis[ConstraintIndex][1] * DeltaLambdaVel;
+		const FVec3 DW1 = RotationDrives.Data.ConstraintDRAxis[ConstraintIndex][1] * DeltaLambdaVel;
 		Body(1).ApplyAngularVelocityDelta(DW1);
 	}
 }
@@ -2989,9 +2989,9 @@ void FAxisConstraintDatas::InitDatas(
 	const bool bResetLambdas)
 {
 	bSoftLimit[ConstraintIndex] = bHasSoftLimits;
-	D.ConstraintHardStiffness[ConstraintIndex] = HardStiffness;
-	D.ConstraintSoftStiffness[ConstraintIndex] = SoftStiffness;
-	D.ConstraintSoftDamping[ConstraintIndex] = SoftDamping;
+	Data.ConstraintHardStiffness[ConstraintIndex] = HardStiffness;
+	Data.ConstraintSoftStiffness[ConstraintIndex] = SoftStiffness;
+	Data.ConstraintSoftDamping[ConstraintIndex] = SoftDamping;
 	ConstraintMaxLambda[ConstraintIndex] = 0;
 	SettingsSoftStiffness[ConstraintIndex] = SoftStiffness;
 	SettingsSoftDamping[ConstraintIndex] = SoftDamping;
@@ -3000,9 +3000,9 @@ void FAxisConstraintDatas::InitDatas(
 	MotionType[ConstraintIndex] = EJointMotionType::Free;
 	if (bResetLambdas)
 	{
-		D.ConstraintLambda = FVec3::Zero();
+		Data.ConstraintLambda = FVec3::Zero();
 		ConstraintLambdaVelocity = FVec3::Zero();
-		D.ConstraintLimits = FVec3::Zero();
+		Data.ConstraintLimits = FVec3::Zero();
 	}
 }
 
@@ -3019,14 +3019,14 @@ void FAxisConstraintDatas::UpdateDatas(
 	bValidDatas[ConstraintIndex] = true;
 	bLimitsCheck[ConstraintIndex] = bCheckLimit;
 
-	D.ConstraintCX[ConstraintIndex] = DatasCX;
+	Data.ConstraintCX[ConstraintIndex] = DatasCX;
 	ConstraintVX[ConstraintIndex] = DatasVX;
 	ConstraintRestitution[ConstraintIndex] = DatasRestitution;
 
 
-	D.ConstraintArms[ConstraintIndex][0] = DatasArm0;
-	D.ConstraintArms[ConstraintIndex][1] = DatasArm1;
-	D.ConstraintAxis[ConstraintIndex] = DatasAxis;
+	Data.ConstraintArms[ConstraintIndex][0] = DatasArm0;
+	Data.ConstraintArms[ConstraintIndex][1] = DatasArm1;
+	Data.ConstraintAxis[ConstraintIndex] = DatasAxis;
 }
 
 void FAxisConstraintDatas::UpdateMass(
@@ -3037,19 +3037,19 @@ void FAxisConstraintDatas::UpdateMass(
 	const FReal Dt,
 	const bool bUsePositionBasedDrives)
 {
-	D.ConstraintHardIM[ConstraintIndex] = DatasIM;
+	Data.ConstraintHardIM[ConstraintIndex] = DatasIM;
 
 
-	D.ConstraintDRAxis[ConstraintIndex][0] = DatasIA0;
-	D.ConstraintDRAxis[ConstraintIndex][1] = -DatasIA1;
+	Data.ConstraintDRAxis[ConstraintIndex][0] = DatasIA0;
+	Data.ConstraintDRAxis[ConstraintIndex][1] = -DatasIA1;
 
 	if (bSoftLimit[ConstraintIndex])
 	{
 		// If bUsePositionBasedDrives is false, we apply the velocity drive in the velocity solver phase so we don't include in the PBD settings
-		const FReal SpringMassScale = (bAccelerationMode) ? (FReal)1 / (D.ConstraintHardIM[ConstraintIndex]) : (FReal)1;
-		D.ConstraintSoftStiffness[ConstraintIndex] = SpringMassScale * SettingsSoftStiffness[ConstraintIndex] * Dt * Dt;
-		D.ConstraintSoftDamping[ConstraintIndex] = (bUsePositionBasedDrives) ? SpringMassScale * SettingsSoftDamping[ConstraintIndex] * Dt : 0;
-		D.ConstraintSoftIM[ConstraintIndex] = (D.ConstraintSoftStiffness[ConstraintIndex] + D.ConstraintSoftDamping[ConstraintIndex]) * D.ConstraintHardIM[ConstraintIndex] + (FReal)1;
+		const FReal SpringMassScale = (bAccelerationMode) ? (FReal)1 / (Data.ConstraintHardIM[ConstraintIndex]) : (FReal)1;
+		Data.ConstraintSoftStiffness[ConstraintIndex] = SpringMassScale * SettingsSoftStiffness[ConstraintIndex] * Dt * Dt;
+		Data.ConstraintSoftDamping[ConstraintIndex] = (bUsePositionBasedDrives) ? SpringMassScale * SettingsSoftDamping[ConstraintIndex] * Dt : 0;
+		Data.ConstraintSoftIM[ConstraintIndex] = (Data.ConstraintSoftStiffness[ConstraintIndex] + Data.ConstraintSoftDamping[ConstraintIndex]) * Data.ConstraintHardIM[ConstraintIndex] + (FReal)1;
 	}
 }
 
@@ -3067,7 +3067,7 @@ void FAxisConstraintDatas::SetMaxForce(
 		FReal MaxLambda = InMaxForce * Dt * Dt;
 		if (bAccelerationMode)
 		{
-			MaxLambda /= D.ConstraintHardIM[ConstraintIndex];
+			MaxLambda /= Data.ConstraintHardIM[ConstraintIndex];
 		}
 		ConstraintMaxLambda[ConstraintIndex] = MaxLambda;
 	}
@@ -3082,12 +3082,12 @@ void FAxisConstraintDatas::ApplyMaxLambda(
 	{
 		if (Lambda > ConstraintMaxLambda[ConstraintIndex])
 		{
-			DeltaLambda = ConstraintMaxLambda[ConstraintIndex] - D.ConstraintLambda[ConstraintIndex];
+			DeltaLambda = ConstraintMaxLambda[ConstraintIndex] - Data.ConstraintLambda[ConstraintIndex];
 			Lambda = ConstraintMaxLambda[ConstraintIndex];
 		}
 		else if (Lambda < -ConstraintMaxLambda[ConstraintIndex])
 		{
-			DeltaLambda = -ConstraintMaxLambda[ConstraintIndex] - D.ConstraintLambda[ConstraintIndex];
+			DeltaLambda = -ConstraintMaxLambda[ConstraintIndex] - Data.ConstraintLambda[ConstraintIndex];
 			Lambda = -ConstraintMaxLambda[ConstraintIndex];
 		}
 	}
diff --git a/Engine/Source/Runtime/Experimental/Chaos/Public/Chaos/Joint/PBDJointCachedSolverGaussSeidel.h b/Engine/Source/Runtime/Experimental/Chaos/Public/Chaos/Joint/PBDJointCachedSolverGaussSeidel.h
--- a/Engine/Source/Runtime/Experimental/Chaos/Public/Chaos/Joint/PBDJointCachedSolverGaussSeidel.h
+++ b/Engine/Source/Runtime/Experimental/Chaos/Public/Chaos/Joint/PBDJointCachedSolverGaussSeidel.h
@@ -68,20 +68,20 @@ struct FAxisConstraintDatas
 	struct alignas(16) FDataSimd
 	{
 		FDataSimd() 
-			: ConstraintSoftStiffnessSimd(VectorZeroFloat())
-			, ConstraintSoftDampingSimd(VectorZeroFloat())
-			, ConstraintLimitsSimd(VectorZeroFloat())
+			: ConstraintSoftStiffness(VectorZeroFloat())
+			, ConstraintSoftDamping(VectorZeroFloat())
+			, ConstraintLimits(VectorZeroFloat())
 		{}
 
-		VectorRegister4Float ConstraintHardStiffnessSimd;
-		VectorRegister4Float ConstraintSoftStiffnessSimd;
-		VectorRegister4Float ConstraintSoftDampingSimd;
-		VectorRegister4Float ConstraintArmsSimd[3][2];
-		VectorRegister4Float ConstraintAxisSimd[3];
-		VectorRegister4Float ConstraintLimitsSimd;
-		VectorRegister4Float ConstraintSoftIMSimd;
-		VectorRegister4Float ConstraintHardIMSimd;
-		VectorRegister4Float ConstraintDRAxisSimd[3][2];
+		VectorRegister4Float ConstraintHardStiffness;
+		VectorRegister4Float ConstraintSoftStiffness;
+		VectorRegister4Float ConstraintSoftDamping;
+		VectorRegister4Float ConstraintArms[3][2];
+		VectorRegister4Float ConstraintAxis[3];
+		VectorRegister4Float ConstraintLimits;
+		VectorRegister4Float ConstraintSoftIM;
+		VectorRegister4Float ConstraintHardIM;
+		VectorRegister4Float ConstraintDRAxis[3][2];
 		VectorRegister4Float ConstraintCX;
 		VectorRegister4Float ConstraintLambda;
 	};
@@ -104,7 +104,7 @@ struct FAxisConstraintDatas
 	union
 	{
 		FDataSimd Simd;
-		FData D;
+		FData Data;
 	};
 
 	FVec3 ConstraintMaxLambda;
@@ -239,7 +239,7 @@ struct FAxisConstraintDatas
 					ConstraintLambdas[2] = VectorReplicate(PositionConstraints.Simd.ConstraintLambda, 2);
 					for (int32 Axis = 0; Axis < 3; ++Axis)
 					{
-						ImpulseSimd = VectorMultiplyAdd(ConstraintLambdas[Axis], PositionConstraints.Simd.ConstraintAxisSimd[Axis], ImpulseSimd);
+						ImpulseSimd = VectorMultiplyAdd(ConstraintLambdas[Axis], PositionConstraints.Simd.ConstraintAxis[Axis], ImpulseSimd);
 					}
 				}
 				if (PositionDrives.bUseSimd)
@@ -250,7 +250,7 @@ struct FAxisConstraintDatas
 					ConstraintLambdas[2] = VectorReplicate(PositionDrives.Simd.ConstraintLambda, 2);
 					for (int32 Axis = 0; Axis < 3; ++Axis)
 					{
-						ImpulseSimd = VectorMultiplyAdd(ConstraintLambdas[Axis], PositionDrives.Simd.ConstraintAxisSimd[Axis], ImpulseSimd);
+						ImpulseSimd = VectorMultiplyAdd(ConstraintLambdas[Axis], PositionDrives.Simd.ConstraintAxis[Axis], ImpulseSimd);
 					}
 				}
 				FVec3f Impulsef;
@@ -262,11 +262,11 @@ struct FAxisConstraintDatas
 			{
 				if (!PositionConstraints.bUseSimd && PositionConstraints.bValidDatas[Axis])
 				{
-					Impulse += PositionConstraints.D.ConstraintLambda[Axis] * PositionConstraints.D.ConstraintAxis[Axis];
+					Impulse += PositionConstraints.Data.ConstraintLambda[Axis] * PositionConstraints.Data.ConstraintAxis[Axis];
 				}
 				if (!PositionDrives.bUseSimd && PositionDrives.bValidDatas[Axis])
 				{
-					Impulse += PositionDrives.D.ConstraintLambda[Axis] * PositionDrives.D.ConstraintAxis[Axis];
+					Impulse += PositionDrives.Data.ConstraintLambda[Axis] * PositionDrives.Data.ConstraintAxis[Axis];
 				}
 			}
 			return Impulse;
@@ -287,7 +287,7 @@ struct FAxisConstraintDatas
 					ConstraintLambdas[2] = VectorReplicate(RotationConstraints.Simd.ConstraintLambda, 2);
 					for (int32 Axis = 0; Axis < 3; ++Axis)
 					{
-						ImpulseSimd = VectorMultiplyAdd(ConstraintLambdas[Axis], RotationConstraints.Simd.ConstraintAxisSimd[Axis], ImpulseSimd);
+						ImpulseSimd = VectorMultiplyAdd(ConstraintLambdas[Axis], RotationConstraints.Simd.ConstraintAxis[Axis], ImpulseSimd);
 					}
 				}
 				if (RotationDrives.bUseSimd)
@@ -298,7 +298,7 @@ struct FAxisConstraintDatas
 					ConstraintLambdas[2] = VectorReplicate(RotationDrives.Simd.ConstraintLambda, 2);
 					for (int32 Axis = 0; Axis < 3; ++Axis)
 					{
-						ImpulseSimd = VectorMultiplyAdd(ConstraintLambdas[Axis], RotationDrives.Simd.ConstraintAxisSimd[Axis], ImpulseSimd);
+						ImpulseSimd = VectorMultiplyAdd(ConstraintLambdas[Axis], RotationDrives.Simd.ConstraintAxis[Axis], ImpulseSimd);
 					}
 				}
 				FVec3f Impulsef;
@@ -310,11 +310,11 @@ struct FAxisConstraintDatas
 			{
 				if (!RotationConstraints.bUseSimd && RotationConstraints.bValidDatas[Axis])
 				{
-					Impulse += RotationConstraints.D.ConstraintLambda[Axis] * RotationConstraints.D.ConstraintAxis[Axis];
+					Impulse += RotationConstraints.Data.ConstraintLambda[Axis] * RotationConstraints.Data.ConstraintAxis[Axis];
 				}
 				if (!RotationDrives.bUseSimd && RotationDrives.bValidDatas[Axis])
 				{
-					Impulse += RotationDrives.D.ConstraintLambda[Axis] * RotationDrives.D.ConstraintAxis[Axis];
+					Impulse += RotationDrives.Data.ConstraintLambda[Axis] * RotationDrives.Data.ConstraintAxis[Axis];
 				}
 			}
 			return Impulse;
