Move ondemand I/O store out of IAS dispatcher backend

diff --git a/Engine/Source/Runtime/Experimental/IoStoreOnDemand/Private/IoStoreOnDemand.cpp b/Engine/Source/Runtime/Experimental/IoStoreOnDemand/Private/IoStoreOnDemand.cpp
--- a/Engine/Source/Runtime/Experimental/IoStoreOnDemand/Private/IoStoreOnDemand.cpp
+++ b/Engine/Source/Runtime/Experimental/IoStoreOnDemand/Private/IoStoreOnDemand.cpp
@@ -19,6 +19,7 @@
 #include "Misc/Parse.h"
 #include "Misc/PathViews.h"
 #include "Modules/ModuleManager.h"
+#include "OnDemandIoStore.h"
 #include "OnDemandHttpClient.h"
 #include "OnDemandIoDispatcherBackend.h"
 #include "Serialization/Archive.h"
@@ -51,6 +52,7 @@ DEFINE_LOG_CATEGORY(LogIas);
 namespace UE::IoStore
 {
 
+////////////////////////////////////////////////////////////////////////////////
 FString GIasOnDemandTocExt = TEXT(".uondemandtoc");
 
 bool GIasSuspendSystem = false;
@@ -140,7 +142,7 @@ static bool ApplyEncryptionKeyFromString(const FString& GuidKeyPair)
 }
 
 ////////////////////////////////////////////////////////////////////////////////
-static bool TryParseConfigContent(const FString& ConfigContent, const FString& ConfigFileName, FOnDemandEndpoint& OutEndpoint)
+static bool TryParseConfigContent(const FString& ConfigContent, const FString& ConfigFileName, FOnDemandEndpointConfig& OutEndpoint)
 {
 	if (ConfigContent.IsEmpty())
 	{
@@ -192,7 +194,7 @@ static bool TryParseConfigContent(const FString& ConfigContent, const FString& C
 }
 
 ////////////////////////////////////////////////////////////////////////////////
-static bool TryParseConfigFileFromPlatformPackage(FOnDemandEndpoint& OutEndpoint)
+static bool TryParseConfigFileFromPlatformPackage(FOnDemandEndpointConfig& OutConfig)
 {
 	const FString ConfigFileName = TEXT("IoStoreOnDemand.ini");
 	const FString ConfigPath = FPaths::Combine(TEXT("Cloud"), ConfigFileName);
@@ -200,7 +202,7 @@ static bool TryParseConfigFileFromPlatformPackage(FOnDemandEndpoint& OutEndpoint
 	if (FPlatformMisc::FileExistsInPlatformPackage(ConfigPath))
 	{
 		const FString ConfigContent = FPlatformMisc::LoadTextFileFromPlatformPackage(ConfigPath);
-		return TryParseConfigContent(ConfigContent, ConfigFileName, OutEndpoint);
+		return TryParseConfigContent(ConfigContent, ConfigFileName, OutConfig);
 	}
 	else
 	{
@@ -209,17 +211,41 @@ static bool TryParseConfigFileFromPlatformPackage(FOnDemandEndpoint& OutEndpoint
 }
 
 ////////////////////////////////////////////////////////////////////////////////
-bool TryParseConfigFile(const FString& ConfigPath, FOnDemandEndpoint& OutEndpoint)
+bool TryParseEndpointConfig(const TCHAR* CommandLine, FOnDemandEndpointConfig& OutConfig)
 {
-	const FString ConfigFileName = TEXT("IoStoreOnDemand.ini");
-	FString ConfigContent; 
-
-	if (!FFileHelper::LoadFileToString(ConfigContent, &IPlatformFile::GetPlatformPhysical(), *ConfigPath))
+	OutConfig = FOnDemandEndpointConfig();
+#if !UE_BUILD_SHIPPING
+	FString UrlParam;
+	if (FParse::Value(CommandLine, TEXT("Ias.TocUrl="), UrlParam))
 	{
-		return false;
+		FStringView UrlView(UrlParam);
+		if (UrlView.StartsWith(TEXTVIEW("http://")) && UrlView.EndsWith(TEXTVIEW(".iochunktoc")))
+		{
+			int32 Delim = INDEX_NONE;
+			if (UrlView.RightChop(7).FindChar(TEXT('/'), Delim))
+			{
+				OutConfig.ServiceUrls.Add(FString(UrlView.Left(7 +  Delim)));
+				OutConfig.TocPath = UrlView.RightChop(OutConfig.ServiceUrls[0].Len() + 1);
+			}
+		}
+	}
+	else
+#endif
+	{
+		if (TryParseConfigFileFromPlatformPackage(OutConfig))
+		{
+			TStringBuilder<256> TocFilePath;
+			FPathViews::Append(TocFilePath, TEXT("Cloud"), FPaths::GetBaseFilename(OutConfig.TocPath));
+			TocFilePath.Append(TEXT(".iochunktoc"));
+
+			if (FPlatformMisc::FileExistsInPlatformPackage(*TocFilePath))
+			{
+				OutConfig.TocFilePath = TocFilePath;
+			}
+		}
 	}
 
-	return TryParseConfigContent(ConfigContent, ConfigFileName, OutEndpoint);
+	return OutConfig.IsValid();
 }
 
 ////////////////////////////////////////////////////////////////////////////////
@@ -1387,8 +1413,9 @@ TIoStatusOr<FIoStoreUploadResult> UploadContainerFiles(
 	FString ChunksRelativePath = UploadParams.BucketPrefix.IsEmpty() ? TEXT("Chunks") : FString::Printf(TEXT("%s/Chunks"), *UploadParams.BucketPrefix);
 	ChunksRelativePath.ToLowerInline();
 
-	bool bWritePerContainerToc = false;
-	GConfig->GetBool(TEXT("Ias"), TEXT("CreatePerContainerTocs"), bWritePerContainerToc, GEngineIni);
+	bool bPerContainerTocsConfigValue = false;
+	GConfig->GetBool(TEXT("Ias"), TEXT("CreatePerContainerTocs"), bPerContainerTocsConfigValue, GEngineIni);
+	const bool bWritePerContainerToc = bPerContainerTocsConfigValue || UploadParams.bPerContainerTocs;
 
 	uint64 TotalUploadedChunks = 0;
 	uint64 TotalUploadedBytes = 0;
@@ -2551,41 +2578,27 @@ void FIoStoreOnDemandModule::InitializeInternal()
 	// Make sure we haven't called initialize before
 	check(!HttpIoDispatcherBackend.IsValid());
 
-	FOnDemandEndpoint Endpoint;
-	
-	FString UrlParam;
-	if (FParse::Value(CommandLine, TEXT("Ias.TocUrl="), UrlParam))
+	FOnDemandEndpointConfig EndpointConfig;
+	if (TryParseEndpointConfig(CommandLine, EndpointConfig) == false)
 	{
-		FStringView UrlView(UrlParam);
-		if (UrlView.StartsWith(TEXTVIEW("http://")) && UrlView.EndsWith(TEXTVIEW(".iochunktoc")))
-		{
-			int32 Delim = INDEX_NONE;
-			if (UrlView.RightChop(7).FindChar(TEXT('/'), Delim))
-			{
-				Endpoint.ServiceUrls.Add(FString(UrlView.Left(7 +  Delim)));
-				Endpoint.TocPath = UrlView.RightChop(Endpoint.ServiceUrls[0].Len() + 1);
-
-				// Since the user has provided a url to download the .iochunktoc from we should
-				// assume that they want to download it rather than use anything on disk.
-				Endpoint.bForceTocDownload = true;
-			}
-		}
+		return;
 	}
 
+	if (IoStore.IsValid() == false)
 	{
-		FString EncryptionKey;
-		if (FParse::Value(CommandLine, TEXT("Ias.EncryptionKey="), EncryptionKey))
+		IoStore = MakeUnique<FOnDemandIoStore>();
+		if (FIoStatus Status = IoStore->Initialize(); !Status.IsOk())
 		{
-			ApplyEncryptionKeyFromString(EncryptionKey);
+			UE_LOG(LogIas, Error, TEXT("Failed to initialize I/O store on demand, reason '%s'"), *Status.ToString());
+			return;
 		}
 	}
 
-	if (!Endpoint.IsValid())
 	{
-		Endpoint = FOnDemandEndpoint();
-		if (!TryParseConfigFileFromPlatformPackage(Endpoint))
+		FString EncryptionKey;
+		if (FParse::Value(CommandLine, TEXT("Ias.EncryptionKey="), EncryptionKey))
 		{
-			return;
+			ApplyEncryptionKeyFromString(EncryptionKey);
 		}
 	}
 
@@ -2608,8 +2621,8 @@ void FIoStoreOnDemandModule::InitializeInternal()
 			(CacheConfig.DiskQuota > 0) ? TEXT("init-fail") : TEXT("zero-quota"));
 	}
 
-	HttpIoDispatcherBackend = MakeOnDemandIoDispatcherBackend(MoveTemp(Cache));
-	HttpIoDispatcherBackend->Mount(Endpoint);
+	HttpIoDispatcherBackend = MakeOnDemandIoDispatcherBackend(EndpointConfig, *IoStore, MoveTemp(Cache));
+
 	int32 BackendPriority = -10;
 #if !UE_BUILD_SHIPPING
 	if (FParse::Param(CommandLine, TEXT("Ias")))
@@ -2630,6 +2643,52 @@ void FIoStoreOnDemandModule::InitializeInternal()
 	}
 	
 	FIoDispatcher::Get().Mount(HttpIoDispatcherBackend.ToSharedRef(), BackendPriority);
+
+	bool bUsePerContainerTocsConfigValue = false;
+	if (GConfig)
+	{
+		GConfig->GetBool(TEXT("Ias"), TEXT("UsePerContainerTocs"), bUsePerContainerTocsConfigValue, GEngineIni);
+	}
+	bool bUsePerContainerTocsParam = false;
+#if !UE_BUILD_SHIPPING
+	bUsePerContainerTocsParam = FParse::Param(CommandLine, TEXT("Ias.UsePerContainerTocs"));
+#endif
+
+	const bool bUsePerContainerTocs = bUsePerContainerTocsConfigValue || bUsePerContainerTocsParam;
+	UE_LOG(LogIas, Log, TEXT("Using per container TOCs=%s"), bUsePerContainerTocs ? TEXT("True") : TEXT("False"));
+
+	TOptional<FOnDemandMountArgs> MountArgs;
+	if (EndpointConfig.TocFilePath.IsEmpty() == false)
+	{
+		if (bUsePerContainerTocs == false)
+		{
+			MountArgs.Emplace(FOnDemandMountArgs
+			{
+				.MountId = EndpointConfig.TocFilePath,
+				.FilePath = EndpointConfig.TocFilePath
+			});
+		}
+	}
+	else if (!EndpointConfig.ServiceUrls.IsEmpty() && !EndpointConfig.TocPath.IsEmpty())
+	{
+		const FString TocUrl = EndpointConfig.ServiceUrls[0] / EndpointConfig.TocPath;
+		MountArgs.Emplace(FOnDemandMountArgs
+		{
+			.MountId = TocUrl,
+			.Url = TocUrl
+		});
+	}
+
+	if (MountArgs)
+	{
+		IoStore->Mount(
+			MoveTemp(MountArgs.GetValue()),
+			[](TIoStatusOr<FOnDemandMountResult> MountResult)
+			{
+				UE_CLOG(!MountResult.IsOk(), LogIas, Error,
+					TEXT("Failed to mount TOC, reason '%s'"), *MountResult.Status().ToString());
+			});
+	}
 }
 	
 void FIoStoreOnDemandModule::StartupModule()
diff --git a/Engine/Source/Runtime/Experimental/IoStoreOnDemand/Private/OnDemandIoDispatcherBackend.cpp b/Engine/Source/Runtime/Experimental/IoStoreOnDemand/Private/OnDemandIoDispatcherBackend.cpp
--- a/Engine/Source/Runtime/Experimental/IoStoreOnDemand/Private/OnDemandIoDispatcherBackend.cpp
+++ b/Engine/Source/Runtime/Experimental/IoStoreOnDemand/Private/OnDemandIoDispatcherBackend.cpp
@@ -2,7 +2,6 @@
 
 #include "OnDemandIoDispatcherBackend.h"
 
-#include "Algo/Transform.h"
 #include "AnalyticsEventAttribute.h"
 #include "Containers/BitArray.h"
 #include "Containers/StringView.h"
@@ -10,15 +9,12 @@
 #include "DistributionEndpoints.h"
 #include "GenericPlatform/GenericPlatformCrashContext.h"
 #include "HAL/Event.h"
-#include "HAL/FileManagerGeneric.h"
 #include "HAL/LowLevelMemTracker.h"
 #include "HAL/Platform.h"
-#include "HAL/PlatformFileManager.h"
 #include "HAL/PlatformTime.h"
 #include "HAL/PreprocessorHelpers.h"
 #include "HAL/Runnable.h"
 #include "HAL/RunnableThread.h"
-#include "HttpManager.h"
 #include "IO/IoAllocators.h"
 #include "IO/IoChunkEncoding.h"
 #include "IO/IoDispatcher.h"
@@ -29,20 +25,14 @@
 #include "IasCache.h"
 #include "Logging/StructuredLog.h"
 #include "Math/NumericLimits.h"
-#include "Misc/CommandLine.h"
 #include "Misc/ConfigCacheIni.h"
-#include "Misc/CoreDelegatesInternal.h"
-#include "Misc/EncryptionKeyManager.h"
 #include "Misc/EnumClassFlags.h"
 #include "Misc/PathViews.h"
 #include "Misc/Paths.h"
-#include "Misc/ScopeExit.h"
 #include "Misc/ScopeLock.h"
 #include "Misc/ScopeRWLock.h"
 #include "OnDemandHttpClient.h"
-#include "Serialization/CompactBinary.h"
-#include "Serialization/CompactBinarySerialization.h"
-#include "Serialization/MemoryReader.h"
+#include "OnDemandIoStore.h"
 #include "Statistics.h"
 #include "Tasks/Task.h"
 
@@ -180,17 +170,6 @@ static FAutoConsoleVariableRef CVar_IoReportAnalytics(
 	TEXT("Enables reporting statics to the analytics system")
 );
 
-int32 GIasTocMode = 0;
-static FAutoConsoleVariableRef CVar_IasTocMode(
-	TEXT("ias.TocMode"),
-	GIasTocMode,
-	TEXT("How should the IAS system load it's toc (see ETocMode).\n")
-	TEXT("0 = Load a single .iochunktoc\n")
-	TEXT("1 = Try to find a .uondemandtoc each time a pak file is mounted\n")
-	TEXT("2 = Download the toc from the target CDN"),
-	ECVF_ReadOnly
-);
-
 static int32 GIasHttpRangeRequestMinSizeKiB = 128;
 static FAutoConsoleVariableRef CVar_IasHttpRangeRequestMinSizeKiB(
 	TEXT("ias.HttpRangeRequestMinSizeKiB"),
@@ -236,29 +215,6 @@ static void LatencyTest(FStringView Url, FStringView Path)
 }
 #endif // !UE_BUILD_SHIPPING
 
-///////////////////////////////////////////////////////////////////////////////
-enum class ETocMode : int32
-{
-	LoadTocFromDisk = 0,	// <- Current Default
-	LoadTocFromMountedPaks,
-	LoadTocFromNetwork
-};
-
-static void ForceTocMode(ETocMode Mode)
-{
-	GIasTocMode = static_cast<int32>(Mode);
-}
-
-static ETocMode GetTocMode()
-{
-	if (GIasTocMode < 0 || GIasTocMode > static_cast<int32>(ETocMode::LoadTocFromNetwork))
-	{
-		UE_LOG(LogIas, Log, TEXT("ias.TocMode set to invalid value, defaulting to ETocMode::LoadTocFromDisk"));
-		return ETocMode::LoadTocFromDisk;
-	}
-
-	return static_cast<ETocMode>(GIasTocMode);
-}
 ///////////////////////////////////////////////////////////////////////////////
 static int32 LatencyTest(TConstArrayView<FString> Urls, FStringView Path, std::atomic_bool& bCancel)
 {
@@ -313,273 +269,6 @@ FIoHash GetChunkKey(const FIoHash& ChunkHash, const FIoOffsetAndLength& Range)
 	return HashBuilder.Finalize();
 }
 
-///////////////////////////////////////////////////////////////////////////////
-class FOnDemandIoStore
-{
-public:
-	struct FTocEntry
-	{
-		uint32 RawSize = 0;
-		uint32 EncodedSize = 0;
-		uint32 BlockOffset = ~uint32(0);
-		uint32 BlockCount = 0; 
-		FIoHash Hash;
-	};
-
-	struct FToc;
-
-	struct FContainer
-	{
-		const FToc* Toc = nullptr;
-		FAES::FAESKey EncryptionKey;
-		FString Name; // TODO: Consider removing when NO_LOGGING == 1 (only used for logging at the moment)
-		FString EncryptionKeyGuid;
-		FString ChunksDirectory;
-		FName CompressionFormat;
-		uint32 BlockSize = 0;
-
-		TMap<FIoChunkId, FTocEntry> TocEntries;
-		TArray<uint32> BlockSizes;
-		TArray<FIoBlockHash> BlockHashes;
-	};
-
-	struct FToc
-	{
-		FString TocPath;
-		TArray<FContainer> Containers;
-	};
-
-	using FTocArray = TChunkedArray<FToc, sizeof(FToc) * 4>;
-
-	struct FChunkInfo
-	{
-		const FContainer* Container = nullptr;
-		const FTocEntry* Entry = nullptr;
-
-		bool IsValid() const { return Container && Entry; }
-		operator bool() const { return IsValid(); }
-
-		TConstArrayView<uint32> GetBlocks() const
-		{
-			check(Container != nullptr && Entry != nullptr);
-			return TConstArrayView<uint32>(Container->BlockSizes.GetData() + Entry->BlockOffset, Entry->BlockCount);
-		}
-		
-		TConstArrayView<FIoBlockHash> GetBlockHashes() const
-		{
-			check(Container != nullptr && Entry != nullptr);
-			return Container->BlockHashes.IsEmpty()
-				? TConstArrayView<FIoBlockHash>()
-				: TConstArrayView<FIoBlockHash>(Container->BlockHashes.GetData() + Entry->BlockOffset, Entry->BlockCount);
-		}
-	};
-
-	FOnDemandIoStore();
-	~FOnDemandIoStore();
-
-	void AddToc(FStringView TocPath, FOnDemandToc&& Toc);
-	void RemoveToc(FStringView TocPath);
-
-	TIoStatusOr<uint64> GetChunkSize(const FIoChunkId& ChunkId);
-	FChunkInfo GetChunkInfo(const FIoChunkId& ChunkId);
-	FString GetFirstTocPath() const;
-
-	TArray<FIoChunkId> GetAllChunkIds(bool bIncludeOptionalChunks);
-
-private:
-	void AddDeferredContainers();
-	void OnEncryptionKeyAdded(const FGuid& Id, const FAES::FAESKey& Key);
-
-	FTocArray Tocs;
-	TArray<FContainer*> RegisteredContainers;
-	TArray<FContainer*> DeferredContainers;
-	mutable FRWLock Lock;
-};
-
-FOnDemandIoStore::FOnDemandIoStore()
-{
-	FEncryptionKeyManager::Get().OnKeyAdded().AddRaw(this, &FOnDemandIoStore::OnEncryptionKeyAdded);
-}
-
-FOnDemandIoStore::~FOnDemandIoStore()
-{
-	FEncryptionKeyManager::Get().OnKeyAdded().RemoveAll(this);
-}
-
-void FOnDemandIoStore::AddToc(FStringView TocPath, FOnDemandToc&& Toc)
-{
-	UE_LOGFMT(LogIas, Log, "Adding TOC '{FileName}'", TocPath);
-
-	TRACE_CPUPROFILER_EVENT_SCOPE(IasBackend::AddToc);
-
-	TStringBuilder<128> ChunksDirectory;
-
-	{
-		// First attempt to parse the prefix from the given path (legacy)
-		int32 Idx = INDEX_NONE;
-		if (TocPath.FindLastChar(TEXT('/'), Idx))
-		{
-			Algo::Transform(TocPath.Left(Idx), AppendChars(ChunksDirectory), FChar::ToLower);
-		}
-
-		// If there was no prefix in the path we should just use the ChunksDirectory from the toc itself
-		if (ChunksDirectory.Len() == 0)
-		{
-			Algo::Transform(Toc.Header.ChunksDirectory, AppendChars(ChunksDirectory), FChar::ToLower);
-			
-		}
-
-		FPathViews::Append(ChunksDirectory, TEXT("chunks"));
-	}
-
-	{
-		FWriteScopeLock _(Lock);
-
-		const FOnDemandTocHeader& Header = Toc.Header;
-		FToc* NewToc = new(Tocs) FToc{FString(TocPath)};
-		NewToc->Containers.SetNum(Toc.Containers.Num()); // List of containers can never change
-
-		const FName CompressionFormat(Header.CompressionFormat);
-		int32 ContainerIndex = 0;
-
-		for (FOnDemandTocContainerEntry& Container : Toc.Containers)
-		{
-			FContainer* NewContainer = &NewToc->Containers[ContainerIndex++];
-			NewContainer->Toc = NewToc;
-			NewContainer->Name = MoveTemp(Container.ContainerName);
-			NewContainer->ChunksDirectory = ChunksDirectory;
-			NewContainer->CompressionFormat = CompressionFormat;
-			NewContainer->BlockSize = Header.BlockSize;
-			NewContainer->EncryptionKeyGuid = Container.EncryptionKeyGuid;
-			
-			NewContainer->TocEntries.Reserve(Container.Entries.Num());
-			for (const FOnDemandTocEntry& TocEntry : Container.Entries)
-			{
-				check(TocEntry.RawSize <= 0xffff'ffffull);
-				check(TocEntry.EncodedSize <= 0xffff'ffffull);
-				NewContainer->TocEntries.Add(TocEntry.ChunkId, FTocEntry
-				{
-					uint32(TocEntry.RawSize),
-					uint32(TocEntry.EncodedSize),
-					TocEntry.BlockOffset,
-					TocEntry.BlockCount,
-					TocEntry.Hash,
-				});
-			}
-
-			NewContainer->BlockSizes = MoveTemp(Container.BlockSizes);
-			NewContainer->BlockHashes = MoveTemp(Container.BlockHashes);
-
-			DeferredContainers.Add(NewContainer);
-		}
-	}
-
-	AddDeferredContainers();
-}
-
-void FOnDemandIoStore::RemoveToc(FStringView TocPath)
-{
-	// TODO: Need to test this with content bundles!
-	checkNoEntry();
-}
-
-TIoStatusOr<uint64> FOnDemandIoStore::GetChunkSize(const FIoChunkId& ChunkId)
-{
-	if (FChunkInfo Info = GetChunkInfo(ChunkId))
-	{
-		return Info.Entry->RawSize;
-	}
-
-	return FIoStatus(EIoErrorCode::UnknownChunkID);
-}
-
-FOnDemandIoStore::FChunkInfo FOnDemandIoStore::GetChunkInfo(const FIoChunkId& ChunkId)
-{
-	FReadScopeLock _(Lock);
-
-	for (const FContainer* Container : RegisteredContainers)
-	{
-		if (const FTocEntry* Entry = Container->TocEntries.Find(ChunkId))
-		{
-			return FChunkInfo{Container, Entry};
-		}
-	}
-
-	return {};
-}
-
-TArray<FIoChunkId> FOnDemandIoStore::GetAllChunkIds(bool bIncludeOptionalChunks)
-{
-	FReadScopeLock _(Lock);
-
-	TArray<FIoChunkId> ChunkIds;
-	int32 NumChunks = 0;
-	for (const FContainer* Container : RegisteredContainers)
-	{
-		NumChunks += Container->TocEntries.Num();
-	}
-
-	ChunkIds.Reserve(NumChunks);
-
-	for (const FContainer* Container : RegisteredContainers)
-	{
-		for (const TPair<FIoChunkId, FTocEntry>& Entry : Container->TocEntries)
-		{
-			if (bIncludeOptionalChunks || Entry.Key.GetChunkType() != EIoChunkType::OptionalBulkData)
-			{
-				ChunkIds.Add(Entry.Key);
-			}
-		}
-	}
-
-	return ChunkIds;
-}
-
-void FOnDemandIoStore::AddDeferredContainers()
-{
-	FWriteScopeLock _(Lock);
-
-	for (auto It = DeferredContainers.CreateIterator(); It; ++It)
-	{
-		FContainer* Container = *It;
-		if (Container->EncryptionKeyGuid.IsEmpty())
-		{
-			check(Container->EncryptionKey.IsValid() == false);
-			UE_LOG(LogIas, Log, TEXT("Mounting container '%s' (%d entries)"), *Container->Name, Container->TocEntries.Num());
-			RegisteredContainers.Add(Container);
-			It.RemoveCurrent();
-		}
-		else
-		{
-			FGuid KeyGuid;
-			ensure(FGuid::Parse(Container->EncryptionKeyGuid, KeyGuid));
-			if (FEncryptionKeyManager::Get().TryGetKey(KeyGuid, Container->EncryptionKey))
-			{
-				UE_LOG(LogIas, Log, TEXT("Mounting container '%s' (%d entries)"), *Container->Name, Container->TocEntries.Num());
-				RegisteredContainers.Add(Container);
-				It.RemoveCurrent();
-			}
-			else
-			{
-				UE_LOG(LogIas, Log, TEXT("Deferring container '%s', encryption key '%s' not available"), *Container->Name, *Container->EncryptionKeyGuid);
-			}
-		}
-	}
-}
-
-void FOnDemandIoStore::OnEncryptionKeyAdded(const FGuid& Id, const FAES::FAESKey& Key)
-{
-	LLM_SCOPE_BYTAG(Ias);
-	TRACE_CPUPROFILER_EVENT_SCOPE(IasBackend::OnEncryptionKeyAdded);
-	AddDeferredContainers();
-}
-
-FString FOnDemandIoStore::GetFirstTocPath() const
-{
-	FReadScopeLock _(Lock);
-	return Tocs.Num() > 0 ? Tocs[0].TocPath : FString();
-}
-
 ///////////////////////////////////////////////////////////////////////////////
 template<typename T>
 class TThreadSafeIntrusiveQueue
@@ -718,29 +407,29 @@ private:
 ///////////////////////////////////////////////////////////////////////////////
 struct FChunkRequestParams
 {
-	static FChunkRequestParams Create(const FIoOffsetAndLength& OffsetLength, FOnDemandIoStore::FChunkInfo ChunkInfo)
+	static FChunkRequestParams Create(const FIoOffsetAndLength& OffsetLength, FOnDemandChunkInfo ChunkInfo)
 	{
 		FIoOffsetAndLength ChunkRange;
-		if (ChunkInfo.Entry->EncodedSize <= (uint64(GIasHttpRangeRequestMinSizeKiB) << 10))
+		if (ChunkInfo.EncodedSize() <= (uint64(GIasHttpRangeRequestMinSizeKiB) << 10))
 		{
-			ChunkRange = FIoOffsetAndLength(0, ChunkInfo.Entry->EncodedSize);
+			ChunkRange = FIoOffsetAndLength(0, ChunkInfo.EncodedSize());
 		}
 		else
 		{
-			const uint64 RawSize = FMath::Min<uint64>(OffsetLength.GetLength(), ChunkInfo.Entry->RawSize - OffsetLength.GetOffset());
+			const uint64 RawSize = FMath::Min<uint64>(OffsetLength.GetLength(), ChunkInfo.RawSize() - OffsetLength.GetOffset());
 
 			ChunkRange = FIoChunkEncoding::GetChunkRange(
-				ChunkInfo.Entry->RawSize,
-				ChunkInfo.Container->BlockSize,
-				ChunkInfo.GetBlocks(),
+				ChunkInfo.RawSize(),
+				ChunkInfo.BlockSize(),
+				ChunkInfo.Blocks(),
 				OffsetLength.GetOffset(),
 				RawSize).ConsumeValueOrDie();
 		}
 
-		return FChunkRequestParams{GetChunkKey(ChunkInfo.Entry->Hash, ChunkRange), ChunkRange, ChunkInfo};
+		return FChunkRequestParams{GetChunkKey(ChunkInfo.Hash(), ChunkRange), ChunkRange, ChunkInfo};
 	}
 
-	static FChunkRequestParams Create(FIoRequestImpl* Request, FOnDemandIoStore::FChunkInfo ChunkInfo)
+	static FChunkRequestParams Create(FIoRequestImpl* Request, FOnDemandChunkInfo ChunkInfo)
 	{
 		check(Request);
 		check(Request->NextRequest == nullptr);
@@ -749,28 +438,26 @@ struct FChunkRequestParams
 
 	const FIoHash& GetUrlHash() const
 	{
-		return ChunkInfo.Entry->Hash;
+		return ChunkInfo.Hash();
 	}
 
 	void GetUrl(FAnsiStringBuilderBase& Url) const
 	{
-		const FString HashString = LexToString(ChunkInfo.Entry->Hash);
-		Url << "/" << ChunkInfo.Container->ChunksDirectory
+		const FString HashString = LexToString(ChunkInfo.Hash());
+		Url << "/" << ChunkInfo.ChunksDirectory()
 			<< "/" << HashString.Left(2)
 			<< "/" << HashString << ANSITEXTVIEW(".iochunk");
 	}
 
 	FIoChunkDecodingParams GetDecodingParams() const
 	{
-		const FAES::FAESKey& EncryptionKey = ChunkInfo.Container->EncryptionKey;
-
 		FIoChunkDecodingParams Params;
-		Params.EncryptionKey = EncryptionKey.IsValid() ? MakeMemoryView(EncryptionKey.Key, FAES::FAESKey::KeySize) : FMemoryView();
-		Params.CompressionFormat = ChunkInfo.Container->CompressionFormat;
-		Params.BlockSize = ChunkInfo.Container->BlockSize;
-		Params.TotalRawSize = ChunkInfo.Entry->RawSize;
-		Params.EncodedBlockSize = ChunkInfo.GetBlocks(); 
-		Params.BlockHash = ChunkInfo.GetBlockHashes(); 
+		Params.EncryptionKey = ChunkInfo.EncryptionKey(); 
+		Params.CompressionFormat = ChunkInfo.CompressionFormat();
+		Params.BlockSize = ChunkInfo.BlockSize();
+		Params.TotalRawSize = ChunkInfo.RawSize();
+		Params.EncodedBlockSize = ChunkInfo.Blocks(); 
+		Params.BlockHash = ChunkInfo.BlockHashes(); 
 		Params.EncodedOffset = ChunkRange.GetOffset();
 
 		return Params;
@@ -778,7 +465,7 @@ struct FChunkRequestParams
 
 	FIoHash ChunkKey;
 	FIoOffsetAndLength ChunkRange;
-	FOnDemandIoStore::FChunkInfo ChunkInfo;
+	FOnDemandChunkInfo ChunkInfo;
 };
 
 ///////////////////////////////////////////////////////////////////////////////
@@ -944,84 +631,6 @@ static void LogIoResult(
 		PrioToString(Priority));
 };
 
-/**
- * Utility to create a FArchive capable of reading from disk using the exact same pathing
- * rules as FPlatformMisc::LoadTextFileFromPlatformPackage but without forcing the entire
- * file to be loaded at once.
- */
-static TUniquePtr<FArchive> CreateReaderFromPlatformPackage(const FString& RelPath)
-{
-	TRACE_CPUPROFILER_EVENT_SCOPE(IasBackend::CreateReaderFromPlatformPackage);
-
-	const FString AbsPath = FPaths::Combine(FGenericPlatformMisc::RootDir(), RelPath);
-
-	IFileHandle* File = IPlatformFile::GetPlatformPhysical().OpenRead(*AbsPath);
-	if (File)
-	{
-#if PLATFORM_ANDROID
-		// This is a handle to an asset so we need to call Seek(0) to move the internal
-		// offset to the start of the asset file.
-		File->Seek(0);
-#endif //PLATFORM_ANDROID
-		const uint32 ReadBufferSize = 256 * 1024;
-		return MakeUnique<FArchiveFileReaderGeneric>(File, *AbsPath, File->Size(), ReadBufferSize);
-	}
-	else
-	{
-		return TUniquePtr<FArchive>();
-	}	
-}
-
-/** Loads a single .iochunktoc (containing all of our container tocs) from disk */
-static TIoStatusOr<FOnDemandToc> LoadOnDemandTocFromDisk(FStringView TocHash)
-{
-	TRACE_CPUPROFILER_EVENT_SCOPE(IasBackend::LoadOnDemandTocFromDisk);
-
-	FOnDemandToc OutToc;
-
-	const ETocMode TocMode = GetTocMode();
-	if (TocMode == ETocMode::LoadTocFromDisk)
-	{
-		FString TocFileName = FString(TocHash);
-		TocFileName .Append(TEXT(".iochunktoc"));
-
-		const FString TocPath = FPaths::Combine(TEXT("Cloud"), TocFileName);
-
-		if (FPlatformMisc::FileExistsInPlatformPackage(TocPath))
-		{	
-			TUniquePtr<FArchive> Ar = CreateReaderFromPlatformPackage(TocPath);
-			if (!Ar.IsValid())
-			{
-				return FIoStatus(EIoErrorCode::FileOpenFailed, WriteToString<256>(TEXT("Failed to open '"), TocPath, TEXT("' from disk")));
-			}
-
-			*Ar << OutToc;
-
-			Ar->Close();
-
-			if (!Ar->IsError() && !Ar->IsCriticalError())
-			{
-				UE_LOG(LogIas, Log, TEXT("Loaded '%s' from disk"), *TocPath);
-				return OutToc;
-			}
-			else
-			{
-				return FIoStatus(EIoErrorCode::ReadError, WriteToString<256>(TEXT("Failed to serialize '"), TocPath, TEXT("' from disk")));
-			}
-		}
-		else
-		{
-			return FIoStatus(EIoErrorCode::NotFound, WriteToString<256>(TEXT("Unable to find '"), TocPath, TEXT("' on disk")));
-		}
-	}
-	else
-	{
-		checkNoEntry();
-	}
-
-	return OutToc;
-}
-
 ///////////////////////////////////////////////////////////////////////////////
 struct FBackendStatus
 {
@@ -1356,7 +965,7 @@ class FOnDemandIoBackend final
 	};
 public:
 
-	FOnDemandIoBackend(TUniquePtr<IIasCache>&& InCache);
+	FOnDemandIoBackend(const FOnDemandEndpointConfig& InConfig, FOnDemandIoStore& InIoStore, TUniquePtr<IIasCache>&& InCache);
 	virtual ~FOnDemandIoBackend();
 
 	// I/O dispatcher backend
@@ -1373,7 +982,6 @@ public:
 	virtual TIoStatusOr<FIoMappedRegion> OpenMapped(const FIoChunkId& ChunkId, const FIoReadOptions& Options) override;
 
 	// I/O Http backend
-	virtual void Mount(const FOnDemandEndpoint& Endpoint) override;
 	virtual void SetBulkOptionalEnabled(bool bEnabled) override;
 	virtual void SetEnabled(bool bEnabled) override;
 	virtual bool IsEnabled() const override;
@@ -1387,51 +995,21 @@ public:
 
 private:
 	bool Resolve(FIoRequestImpl* Request);
-	void MountContainer(FStringView ContainerPath);
-	void UnmountContainer(FStringView ContainerPath);
 
 	FString GetEndpointTestPath() const;
 	void ConditionallyStartBackendThread();
 	void CompleteRequest(FChunkRequest* ChunkRequest);
 	void CompleteMaterialize(FChunkRequest* ChunkRequest);
-
-	FIoStatus ApplyLoadedOnDemandToc(const FString& TocPath);
-	FIoStatus DownloadoadOnDemandToc(const FString& CdnUrl, const FString& TocPath);
-
 	bool ResolveDistributedEndpoint(const FDistributedEndpointUrl& Url);
-	void InitializePrimaryEndpoint();
-
-	/** Mode to control which types of .iochunktoc are flush */
-	enum class EFlushMode : uint8
-	{
-		/** Flush tocs that require disk access */
-		Disk	= 1 << 0,
-		/** Flush tocs that require network access */
-		Network = 1 << 1,
-
-		/** Flush all types of tocs */
-		All		= Disk | Network,
-	};
-
-	void FlushDeferredTocs(EFlushMode FlushMode);
-
-	bool SetupHttpThread();
 	void ProcessHttpRequests(FHttpClient& HttpClient, FBitWindow& HttpErrors, int32 MaxConcurrentRequests);
 	int32 WaitForCompleteRequestTasks(float WaitTimeSeconds, float PollTimeSeconds);
 	void DrainHttpRequests();
 
-	struct FTocParams
-	{
-		FString Path;
-		bool bForceDownload;
-	};
-
+	FOnDemandIoStore& IoStore;
 	TUniquePtr<IIasCache> Cache;
-	TUniquePtr<FOnDemandIoStore> IoStore;
 	TSharedPtr<const FIoDispatcherBackendContext> BackendContext;
 	TUniquePtr<FRunnableThread> BackendThread;
 	FEventRef TickBackendEvent;
-	TArray<FTocParams> DeferredTocs;
 	FChunkRequests ChunkRequests;
 	FIoRequestQueue CompletedRequests;
 	FChunkRequestQueue HttpRequests;
@@ -1444,26 +1022,38 @@ private:
 
 	FString EndpointTestPath;
 
-	FDelegateHandle PakDelegateHandle;
-
 	mutable FRWLock Lock;
 	std::atomic_uint32_t InflightCacheRequestCount{0};
 	std::atomic_bool bStopRequested{false};
 
-	bool bGeneratedOnDemandToc = false;
-	UE::Tasks::TTask<TIoStatusOr<FOnDemandToc>> LoadingOnDemandTocTask;
-
 #if UE_IAS_DEBUG_CONSOLE_CMDS
 	TArray<IConsoleCommand*> DynamicConsoleCommands;
 #endif // UE_IAS_DEBUG_CONSOLE_CMDS
 };
 
 ///////////////////////////////////////////////////////////////////////////////
-FOnDemandIoBackend::FOnDemandIoBackend(TUniquePtr<IIasCache>&& InCache)
-	: Cache(MoveTemp(InCache))
+FOnDemandIoBackend::FOnDemandIoBackend(const FOnDemandEndpointConfig& Config, FOnDemandIoStore& InIoStore, TUniquePtr<IIasCache>&& InCache)
+	: IoStore(InIoStore)
+	, Cache(MoveTemp(InCache))
 {
-	IoStore = MakeUnique<FOnDemandIoStore>();
-	BackendStatus.SetHttpEnabled(true);
+	EndpointTestPath = Config.TocPath;
+	if (Config.DistributionUrl.IsEmpty() == false)
+	{
+		DistributionUrl = { Config.DistributionUrl, Config.FallbackUrl };
+	}
+	else
+	{
+		for (const FString& ServiceUrl : Config.ServiceUrls)
+		{
+			for (const FString& Url : Config.ServiceUrls)
+			{
+				AvailableEps.Urls.Add(Url.Replace(TEXT("https"), TEXT("http")).ToLower());
+			}
+		}
+	}
+
+	// Don't enable HTTP until the background thread has been started
+	BackendStatus.SetHttpEnabled(false);
 	BackendStatus.SetCacheEnabled(Cache.IsValid());
 
 #if UE_IAS_DEBUG_CONSOLE_CMDS
@@ -1481,25 +1071,10 @@ FOnDemandIoBackend::FOnDemandIoBackend(TUniquePtr<IIasCache>&& InCache)
 		ECVF_Cheat)
 	);
 #endif // UE_IAS_DEBUG_CONSOLE_CMDS
-
-	{
-		bool bUsePerContainerTocs = false;
-		GConfig->GetBool(TEXT("Ias"), TEXT("UsePerContainerTocs"), bUsePerContainerTocs, GEngineIni);
-
-		if (bUsePerContainerTocs)
-		{
-			ForceTocMode(ETocMode::LoadTocFromMountedPaks);
-		}
-	}
 }
 
 FOnDemandIoBackend::~FOnDemandIoBackend()
 {
-	if (PakDelegateHandle.IsValid())
-	{
-		FCoreInternalDelegates::GetOnPakMountOperation().Remove(PakDelegateHandle);
-	}
-
 #if UE_IAS_DEBUG_CONSOLE_CMDS
 	for (IConsoleCommand* Cmd : DynamicConsoleCommands)
 	{
@@ -1545,91 +1120,9 @@ void FOnDemandIoBackend::Shutdown()
 	BackendContext.Reset();
 }
 
-void FOnDemandIoBackend::MountContainer(FStringView ContainerPath)
-{
-	// TODO: For now we just log an error if a container cannot be mounted, but for production
-	// we will need to make sure that the user is informed so that they can verify their installation.
-	IPlatformFile& PlatformFile = IPlatformFile::GetPlatformPhysical();
-
-	const FString TocPath = FPathViews::ChangeExtension(ContainerPath, GIasOnDemandTocExt);
-
-	if (PlatformFile.FileExists(*TocPath))
-	{
-		UE_LOG(LogIas, Log, TEXT("Mounting '%s'"), *TocPath);
-
-		TUniquePtr<FArchive> Ar(IFileManager::Get().CreateFileReader(*TocPath));
-
-		if (!Ar.IsValid())
-		{
-			UE_LOG(LogIas, Error, TEXT("Failed to open '%s' for reading"), *TocPath);
-			return;
-		}
-
-		/// First we read and validate the sentinel at the end of the file to check for potential
-		// file corruption that could lead to crashing during serialization.
-		{
-			const int64 SentinelPos = Ar->TotalSize() - FOnDemandTocSentinel::SentinelSize;
-
-			if (SentinelPos < 0)
-			{
-				UE_LOG(LogIas, Error, TEXT("The file '%s' is smaller than expected and quite possible corrupted"), *TocPath);
-				return;
-			}
-
-			Ar->Seek(SentinelPos);
-
-			FOnDemandTocSentinel Sentinel;
-			(*Ar) << Sentinel;
-
-			if (!Sentinel.IsValid())
-			{
-				UE_LOG(LogIas, Error, TEXT("File corruption detected when serializing '%s'"), *TocPath);
-				return;
-			}
-
-			Ar->Seek(0);
-		}
-
-		FOnDemandToc ContainerToc;
-		(*Ar) << ContainerToc;
-
-		if (Ar->IsError() || Ar->IsCriticalError())
-		{
-			UE_LOG(LogIas, Error, TEXT("Failed to serialize '%s'"), *TocPath);
-			return;
-		}
-
-		const FStringView FileName = FPathViews::GetCleanFilename(TocPath);
-		IoStore->AddToc(FileName, MoveTemp(ContainerToc));
-	}
-}
-
-void FOnDemandIoBackend::UnmountContainer(FStringView ContainerPath)
-{
-	const FString TocPath = FPathViews::ChangeExtension(ContainerPath, GIasOnDemandTocExt);
-	const FStringView FileName = FPathViews::GetCleanFilename(TocPath);
-
-	IoStore->RemoveToc(FileName);
-}
-
 FString FOnDemandIoBackend::GetEndpointTestPath() const
 {
-	if (!EndpointTestPath.IsEmpty())
-	{
-		return EndpointTestPath;
-	}
-
-	// Older fallback path, it probably shouldn't be possible to get this far and can be removed later.
-	FString TestPath = IoStore.IsValid() ? IoStore->GetFirstTocPath() : FString();
-	if (TestPath.IsEmpty())
-	{
-		FReadScopeLock _(Lock);
-		if (DeferredTocs.IsEmpty() == false)
-		{
-			TestPath = DeferredTocs[0].Path;
-		}
-	}
-	return TestPath;
+	return EndpointTestPath;
 }
 
 void FOnDemandIoBackend::ConditionallyStartBackendThread()
@@ -1672,7 +1165,7 @@ void FOnDemandIoBackend::CompleteRequest(FChunkRequest* ChunkRequest)
 		if (Chunk.GetSize() > 0)
 		{
 			TRACE_CPUPROFILER_EVENT_SCOPE(IasBackend::DecodeBlocks);
-			const uint64 RawSize = FMath::Min<uint64>(Request->Options.GetSize(), ChunkRequest->Params.ChunkInfo.Entry->RawSize);
+			const uint64 RawSize = FMath::Min<uint64>(Request->Options.GetSize(), ChunkRequest->Params.ChunkInfo.RawSize());
 			Request->CreateBuffer(RawSize);
 			DecodingParams.RawOffset = Request->Options.GetOffset(); 
 			bDecoded = FIoChunkEncoding::Decode(DecodingParams, Chunk.GetView(), Request->GetBuffer().GetMutableView());
@@ -1691,7 +1184,7 @@ void FOnDemandIoBackend::CompleteRequest(FChunkRequest* ChunkRequest)
 			Stats.OnIoRequestComplete(Request->GetBuffer().GetSize(), DurationMs);
 			LogIoResult(Request->ChunkId, ChunkRequest->Params.GetUrlHash(), DurationMs,
 				Request->GetBuffer().DataSize(), Request->Options.GetOffset(),
-				ChunkRequest->Params.ChunkRange, ChunkRequest->Params.ChunkInfo.Entry->EncodedSize,
+				ChunkRequest->Params.ChunkRange, ChunkRequest->Params.ChunkInfo.EncodedSize(),
 				ChunkRequest->Priority, ChunkRequest->bCached);
 				
 		}
@@ -1703,7 +1196,7 @@ void FOnDemandIoBackend::CompleteRequest(FChunkRequest* ChunkRequest)
 			Stats.OnIoRequestError();
 			LogIoResult(Request->ChunkId, ChunkRequest->Params.GetUrlHash(), DurationMs,
 				0, Request->Options.GetOffset(),
-				ChunkRequest->Params.ChunkRange, ChunkRequest->Params.ChunkInfo.Entry->EncodedSize,
+				ChunkRequest->Params.ChunkRange, ChunkRequest->Params.ChunkInfo.EncodedSize(),
 				ChunkRequest->Priority, ChunkRequest->bCached);
 		}
 
@@ -1765,7 +1258,7 @@ bool FOnDemandIoBackend::Resolve(FIoRequestImpl* Request)
 {
 	using namespace UE::Tasks;
 
-	FOnDemandIoStore::FChunkInfo ChunkInfo = IoStore->GetChunkInfo(Request->ChunkId);
+	FOnDemandChunkInfo ChunkInfo = IoStore.GetChunkInfo(Request->ChunkId);
 	if (!ChunkInfo.IsValid())
 	{
 		return false;
@@ -1894,14 +1387,14 @@ TIoStatusOr<uint64> FOnDemandIoBackend::GetSizeForChunk(const FIoChunkId& ChunkI
 {
 	OutAvailable = 0;
 
-	const FOnDemandIoStore::FChunkInfo ChunkInfo = IoStore.IsValid() ? IoStore->GetChunkInfo(ChunkId) : FOnDemandIoStore::FChunkInfo();
+	const FOnDemandChunkInfo ChunkInfo = IoStore.GetChunkInfo(ChunkId);
 	if (ChunkInfo.IsValid() == false)
 	{
 		return FIoStatus(EIoErrorCode::UnknownChunkID);
 	}
 
-	FIoOffsetAndLength RequestedRange(ChunkRange.GetOffset(), FMath::Min<uint64>(ChunkInfo.Entry->RawSize, ChunkRange.GetLength()));
-	OutAvailable = ChunkInfo.Entry->RawSize;
+	FIoOffsetAndLength RequestedRange(ChunkRange.GetOffset(), FMath::Min<uint64>(ChunkInfo.RawSize(), ChunkRange.GetLength()));
+	OutAvailable = ChunkInfo.RawSize();
 
 	if (BackendStatus.IsHttpEnabled(ChunkId.GetChunkType()) == false)
 	{
@@ -1922,7 +1415,7 @@ TIoStatusOr<uint64> FOnDemandIoBackend::GetSizeForChunk(const FIoChunkId& ChunkI
 		OutAvailable = RequestedRange.GetLength();
 	}
 
-	return TIoStatusOr<uint64>(ChunkInfo.Entry->RawSize);
+	return TIoStatusOr<uint64>(ChunkInfo.RawSize());
 }
 
 FIoRequestImpl* FOnDemandIoBackend::GetCompletedIoRequests()
@@ -1943,56 +1436,6 @@ TIoStatusOr<FIoMappedRegion> FOnDemandIoBackend::OpenMapped(const FIoChunkId& Ch
 	return FIoStatus::Unknown;
 }
 
-FIoStatus FOnDemandIoBackend::ApplyLoadedOnDemandToc(const FString& TocPath)
-{
-	TRACE_CPUPROFILER_EVENT_SCOPE(IasBackend::ApplyLoadedOnDemandToc);
-
-	if (bGeneratedOnDemandToc)
-	{
-		return FIoStatus::Ok;
-	}
-
-	TIoStatusOr<FOnDemandToc> GeneratedTocResult;
-
-	if (LoadingOnDemandTocTask.IsValid())
-	{
-		GeneratedTocResult = MoveTemp(LoadingOnDemandTocTask.GetResult());
-	}
-	else
-	{
-		GeneratedTocResult = LoadOnDemandTocFromDisk(FPathViews::GetBaseFilename(TocPath));
-	}
-	
-	if (!GeneratedTocResult.IsOk())
-	{
-		return GeneratedTocResult.Status();
-	}
-
-	IoStore->AddToc(TocPath, GeneratedTocResult.ConsumeValueOrDie());
-
-	bGeneratedOnDemandToc = true;
-
-	return FIoStatus::Ok;
-}
-
-FIoStatus FOnDemandIoBackend::DownloadoadOnDemandToc(const FString& CdnUrl, const FString& TocPath)
-{
-	TRACE_CPUPROFILER_EVENT_SCOPE(IasBackend::LoadOnDemandToc);
-
-	UE_LOG(LogIas, Log, TEXT("Downloading OnDemandToc from CDN"));
-
-	TIoStatusOr<FOnDemandToc> TocResult = LoadTocFromUrl(CdnUrl, TocPath, 1);
-	if (TocResult.IsOk())
-	{
-		IoStore->AddToc(TocPath, TocResult.ConsumeValueOrDie());
-		return FIoStatus::Ok;
-	}
-	else
-	{
-		return TocResult.Status();
-	}
-}
-
 bool FOnDemandIoBackend::ResolveDistributedEndpoint(const FDistributedEndpointUrl& DistributedEndpointUrl)
 {
 	TRACE_CPUPROFILER_EVENT_SCOPE(IasBackend::ResolveDistributedEndpoint);
@@ -2047,231 +1490,14 @@ bool FOnDemandIoBackend::ResolveDistributedEndpoint(const FDistributedEndpointUr
 	return false;
 }
 
-void FOnDemandIoBackend::InitializePrimaryEndpoint()
-{
-	TRACE_CPUPROFILER_EVENT_SCOPE(IasBackend::InitializePrimaryEndpoint);
-
-	// We only need to run this code if we do not have an currently selected, most likely because
-	// a list of endpoints was polled from online rather than a url being provided by the cmdline/config system.
-	if (!AvailableEps.HasCurrent())
-	{
-		UE_LOG(LogIas, Log, TEXT("Attempting to find the primary endpoint..."));
-
-		TConstArrayView<FString> Urls = AvailableEps.Urls;
-		const int32 MaxUrls = FMath::Min(GIasMaxEndpointTestCountAtStartup, AvailableEps.Urls.Num());
-		const FString TestPath = GetEndpointTestPath();
-		if (int32 Idx = LatencyTest(Urls.Left(MaxUrls), TestPath, bStopRequested); Idx != INDEX_NONE)
-		{
-			FOnDemandIoBackendStats::Get()->OnHttpConnected();
-
-			AvailableEps.Current = Idx;
-			UE_LOG(LogIas, Log, TEXT("Using endpoint '%s'"), *AvailableEps.GetCurrent());
-		}
-		else
-		{
-			UE_LOG(LogIas, Error, TEXT("Unable to connect to any valid endpoint"));
-			BackendStatus.SetHttpError(true);
-			return;
-		}
-	}
-
-	// Now we have a valid endpoint, we need to flush any pending .iochunktoc download requests
-	FlushDeferredTocs(EFlushMode::All);
-}
-
-void FOnDemandIoBackend::FlushDeferredTocs(EFlushMode FlushMode)
-{
-	TRACE_CPUPROFILER_EVENT_SCOPE(IasBackend::FlushDeferredTocs);
-
-	TArray<FTocParams> Tocs;
-	{
-		FWriteScopeLock _(Lock);
-		if (DeferredTocs.IsEmpty())
-		{
-			return;
-		}
-		Tocs = MoveTemp(DeferredTocs);
-	}
-
-	const ETocMode TocMode = GetTocMode();
-
-	for (FTocParams& TocParams : Tocs)
-	{
-		if (TocMode == ETocMode::LoadTocFromDisk && !TocParams.bForceDownload)
-		{
-			if (EnumHasAnyFlags(FlushMode, EFlushMode::Disk))
-			{
-				FIoStatus Result = ApplyLoadedOnDemandToc(TocParams.Path);
-				if (!Result.IsOk())
-				{
-					UE_LOG(LogIas, Error, TEXT("Failed to add generated toc', reason '%s'"), *Result.ToString());
-				}
-			}
-			else
-			{
-				FWriteScopeLock _(Lock);
-				DeferredTocs.Emplace(MoveTemp(TocParams));
-			}
-		}
-		else
-		{
-			if (EnumHasAnyFlags(FlushMode, EFlushMode::Network) && AvailableEps.HasCurrent())
-			{
-				FIoStatus Result = DownloadoadOnDemandToc(AvailableEps.GetCurrent(), TocParams.Path);
-				if (!Result.IsOk())
-				{
-					UE_LOG(LogIas, Error, TEXT("Failed to add TOC '%s/%s', reason '%s'"), *AvailableEps.GetCurrent(), *TocParams.Path, *Result.ToString());
-				}
-			}
-			else
-			{
-				FWriteScopeLock _(Lock);
-				DeferredTocs.Emplace(MoveTemp(TocParams));
-			}
-		}
-	}
-}
-
-void FOnDemandIoBackend::Mount(const FOnDemandEndpoint& Endpoint)
-{
-	LLM_SCOPE_BYTAG(Ias);
-	TRACE_CPUPROFILER_EVENT_SCOPE(IasBackend::Mount);
-
-	if ((Endpoint.DistributionUrl.IsEmpty() && Endpoint.ServiceUrls.IsEmpty()) || Endpoint.TocPath.IsEmpty())
-	{
-		UE_LOG(LogIas, Error, TEXT("Trying to mount an invalid on demand endpoint"));
-		return;
-	}
-
-	EndpointTestPath = Endpoint.TocPath;
-
-	// TODO: Should pass this info around via DeferredTocs
-	if (Endpoint.bForceTocDownload)
-	{
-		ForceTocMode(ETocMode::LoadTocFromNetwork);
-	}
-
-	if (GetTocMode() == ETocMode::LoadTocFromDisk)
-	{
-		UE_LOG(LogIas, Display, TEXT("Loading .iochunktoc from disk"));
-
-		FString TocHash = FPaths::GetBaseFilename(Endpoint.TocPath);
-		LoadingOnDemandTocTask = UE::Tasks::Launch(UE_SOURCE_LOCATION, [TocHash = MoveTemp(TocHash)]() -> TIoStatusOr<FOnDemandToc>
-			{
-				return LoadOnDemandTocFromDisk(TocHash);
-			});
-	}
-	else if (GetTocMode() == ETocMode::LoadTocFromMountedPaks)
-	{
-		UE_LOG(LogIas, Display, TEXT("Loading .uondemandtoc from mounted pakfiles"));
-
-		// First make sure that we receive notifications of any new pakfiles being mounted so that we can react to them
-		PakDelegateHandle = FCoreInternalDelegates::GetOnPakMountOperation().AddLambda(
-			[this](EMountOperation Operation, const TCHAR* ContainerPath, int32 Order) -> void
-			{
-				switch (Operation)
-				{
-					case EMountOperation::Mount:
-						this->MountContainer(ContainerPath);
-						break;
-					case EMountOperation::Unmount:
-						this->UnmountContainer(ContainerPath);
-						break;
-					default:
-						checkNoEntry();
-				}
-			}
-		);
-
-		FCurrentlyMountedPaksDelegate& Delegate = FCoreInternalDelegates::GetCurrentlyMountedPaksDelegate();
-		if (Delegate.IsBound())
-		{
-			TArray<FMountedPakInfo> PreExistingMountedPaks = Delegate.Execute();
-			for (const FMountedPakInfo& Info : PreExistingMountedPaks)
-			{
-				check(Info.PakFile != nullptr);
-				MountContainer(Info.PakFile->PakGetPakFilename());
-			}
-		}
-	}
-
-	{
-		FWriteScopeLock _(Lock);
-		if (Endpoint.ServiceUrls.IsEmpty())
-		{
-			if (AvailableEps.HasCurrent() == false)
-			{
-				if (!DistributionUrl.IsValid())
-				{
-					DistributionUrl = { Endpoint.DistributionUrl, Endpoint.FallbackUrl };
-				}
-
-				if (GetTocMode() != ETocMode::LoadTocFromMountedPaks)
-				{
-					DeferredTocs.Add(FTocParams{ Endpoint.TocPath, Endpoint.bForceTocDownload });
-				}
-
-				return;
-			}
-		}
-		else if (AvailableEps.Urls.IsEmpty())
-		{
-			for (const FString& Url : Endpoint.ServiceUrls)
-			{
-				AvailableEps.Urls.Add(Url.Replace(TEXT("https"), TEXT("http")));
-			}
-		}
-	}
-
-	if (GetTocMode() == ETocMode::LoadTocFromNetwork)
-	{
-		if (AvailableEps.HasCurrent())
-		{
-			FIoStatus Result = DownloadoadOnDemandToc(AvailableEps.GetCurrent(), Endpoint.TocPath);
-			if (!Result.IsOk())
-			{
-				UE_LOG(LogIas, Error, TEXT("Deferring TOC '%s/%s' due to '%s'"), *AvailableEps.GetCurrent(), *Endpoint.TocPath, *Result.ToString());
-				BackendStatus.SetHttpError(true);
-				FWriteScopeLock _(Lock);
-				DeferredTocs.Add(FTocParams{ Endpoint.TocPath, Endpoint.bForceTocDownload });
-			}
-		}
-		else
-		{
-			FWriteScopeLock _(Lock);
-			DeferredTocs.Add(FTocParams{ Endpoint.TocPath, Endpoint.bForceTocDownload });
-		}
-	}
-	else if (GetTocMode() == ETocMode::LoadTocFromDisk)
-	{
-		FIoStatus GeneratedResult = ApplyLoadedOnDemandToc(Endpoint.TocPath);
-		if (!GeneratedResult.IsOk())
-		{
-			UE_LOG(LogIas, Error, TEXT("Failed to add generated toc', reason '%s'"), *GeneratedResult.ToString());
-		}	
-	}
-}
-
 void FOnDemandIoBackend::SetBulkOptionalEnabled(bool bEnabled)
 {
-	// Ignore enable/disable messages if we are mounted from the pak file system
-	// TODO: Remove SetBulkOptionalEnabled entirely once LoadTocFromMountedPaks is
-	// the only valid mode
-	if (GetTocMode() != ETocMode::LoadTocFromMountedPaks)
-	{
-		BackendStatus.SetHttpOptionalBulkEnabled(bEnabled);
-	}
+	BackendStatus.SetHttpOptionalBulkEnabled(bEnabled);
 }
 
 void FOnDemandIoBackend::SetEnabled(bool bEnabled)
 {
-	// Ignore enable/disable messages if we are mounted from the pak file system
-	// TODO: Remove SetEnabled entirely once LoadTocFromMountedPaks is the only
-	// valid mode
-	if (GetTocMode() != ETocMode::LoadTocFromMountedPaks)
-	{
-		BackendStatus.SetHttpEnabled(bEnabled);
-	}
+	BackendStatus.SetHttpEnabled(bEnabled);
 }
 
 bool FOnDemandIoBackend::IsEnabled() const
@@ -2453,41 +1679,30 @@ void FOnDemandIoBackend::DrainHttpRequests()
 	}
 }
 
-bool FOnDemandIoBackend::SetupHttpThread()
+uint32 FOnDemandIoBackend::Run()
 {
-	// A disk only flush should be fairly quick at this point, especially if 's.IasEnableThreadedTocGeneration' is true
-	// and this will help mount .iochunktoc as soon as possible rather than waiting for the primary endpoint to be
-	// chosen.
-	FlushDeferredTocs(EFlushMode::Disk);
+	LLM_SCOPE_BYTAG(Ias);
 
-	// Note that the following method will block until the distributed endpoint (if we have one) has
-	// been resolved as the thread cannot function without it.
 	if (DistributionUrl.IsValid())
 	{
-		const bool bResult = ResolveDistributedEndpoint(DistributionUrl);
-		DistributionUrl.Reset();
-
-		if (!bResult)
+		if (ResolveDistributedEndpoint(DistributionUrl))
 		{
-			return false;
+			UE_LOG(LogIas, Error, TEXT("Failed to resolve CDN endpoints from distribution URL"));
+			BackendStatus.SetHttpEnabled(false);
+			return 0;
 		}
-		
 	}
 
-	InitializePrimaryEndpoint();
-
-	return true;
-}
-
-uint32 FOnDemandIoBackend::Run()
-{
-	LLM_SCOPE_BYTAG(Ias);
-
-	// This call can take some time depending on the IAS setup and the network conditions
-	if (!SetupHttpThread())
+	if (AvailableEps.Urls.IsEmpty())
 	{
+		UE_LOG(LogIas, Error, TEXT("HTTP streaming disabled, no valid endpoint"));
+		BackendStatus.SetHttpEnabled(false);
 		return 0;
 	}
+	AvailableEps.Current = FMath::Min(GIasHttpPrimaryEndpoint, AvailableEps.Urls.Num() - 1);
+
+	BackendStatus.SetHttpEnabled(true);
+	FOnDemandIoBackendStats::Get()->OnHttpConnected();
 
 	FBitWindow HttpErrors;
 	HttpErrors.Reset(GIasHttpErrorSampleCount);
@@ -2495,7 +1710,7 @@ uint32 FOnDemandIoBackend::Run()
 	TUniquePtr<FHttpClient> HttpClient = FHttpClient::Create(FHttpClientConfig
 	{
 		.Endpoints = AvailableEps.Urls,
-		.PrimaryEndpoint = FMath::Min(GIasHttpPrimaryEndpoint, AvailableEps.Urls.Num() -1),
+		.PrimaryEndpoint = AvailableEps.Current,
 		.MaxConnectionCount = GIasHttpConnectionCount,
 		.MaxRetryCount = FMath::Max(AvailableEps.Urls.Num() + 1, GIasHttpRetryCount),
 		.ReceiveBufferSize = GIasHttpRecvBufKiB >= 0 ? GIasHttpRecvBufKiB << 10 : -1,
@@ -2541,7 +1756,6 @@ uint32 FOnDemandIoBackend::Run()
 					HttpClient->SetEndpoint(Idx);
 					BackendStatus.SetHttpError(false);
 					UE_LOG(LogIas, Log, TEXT("Successfully reconnected to '%s'"), *AvailableEps.GetCurrent());
-					FlushDeferredTocs(EFlushMode::All);
 				}
 			}
 			else if (HttpClient->IsUsingPrimaryEndpoint() == false)
@@ -2592,9 +1806,13 @@ int32 FOnDemandIoBackend::WaitForCompleteRequestTasks(float WaitTimeSeconds, flo
 	return ChunkRequests.Num();
 }
 
-TSharedPtr<IOnDemandIoDispatcherBackend> MakeOnDemandIoDispatcherBackend(TUniquePtr<IIasCache>&& Cache)
+TSharedPtr<IOnDemandIoDispatcherBackend> MakeOnDemandIoDispatcherBackend(
+	const FOnDemandEndpointConfig& Config,
+	FOnDemandIoStore& IoStore,
+	TUniquePtr<IIasCache>&& Cache)
 {
-	return MakeShareable<IOnDemandIoDispatcherBackend>(new FOnDemandIoBackend(MoveTemp(Cache)));
+	return MakeShareable<IOnDemandIoDispatcherBackend>(
+		new FOnDemandIoBackend(Config, IoStore, MoveTemp(Cache)));
 }
 
 } // namespace UE::IoStore
diff --git a/Engine/Source/Runtime/Experimental/IoStoreOnDemand/Private/OnDemandIoDispatcherBackend.h b/Engine/Source/Runtime/Experimental/IoStoreOnDemand/Private/OnDemandIoDispatcherBackend.h
--- a/Engine/Source/Runtime/Experimental/IoStoreOnDemand/Private/OnDemandIoDispatcherBackend.h
+++ b/Engine/Source/Runtime/Experimental/IoStoreOnDemand/Private/OnDemandIoDispatcherBackend.h
@@ -11,6 +11,8 @@ struct FAnalyticsEventAttribute;
 namespace UE::IoStore
 {
 
+class FOnDemandIoStore;
+
 struct FDistributedEndpointUrl
 {
 	FString EndpointUrl;
@@ -33,15 +35,14 @@ struct FDistributedEndpointUrl
 	}
 };
 
-struct FOnDemandEndpoint
+struct FOnDemandEndpointConfig
 {
 	FString DistributionUrl;
 	FString FallbackUrl;
 
 	TArray<FString> ServiceUrls;
 	FString TocPath;
-
-	bool bForceTocDownload = false;
+	FString TocFilePath;
 
 	bool IsValid() const
 	{
@@ -55,7 +56,6 @@ class IOnDemandIoDispatcherBackend
 public:
 	virtual ~IOnDemandIoDispatcherBackend() = default;
 
-	virtual void Mount(const FOnDemandEndpoint& Endpoint) = 0;
 	virtual void SetBulkOptionalEnabled(bool bInEnabled) = 0;
 	virtual void SetEnabled(bool bInEnabled) = 0;
 	virtual bool IsEnabled() const = 0;
@@ -63,6 +63,9 @@ public:
 	virtual void ReportAnalytics(TArray<FAnalyticsEventAttribute>& OutAnalyticsArray) const = 0;
 };
 
-TSharedPtr<IOnDemandIoDispatcherBackend> MakeOnDemandIoDispatcherBackend(TUniquePtr<IIasCache>&& Cache);
+TSharedPtr<IOnDemandIoDispatcherBackend> MakeOnDemandIoDispatcherBackend(
+	const FOnDemandEndpointConfig& Config,
+	FOnDemandIoStore& IoStore,
+	TUniquePtr<IIasCache>&& Cache);
 
 } // namespace UE::IoStore
diff --git a/Engine/Source/Runtime/Experimental/IoStoreOnDemand/Private/OnDemandIoStore.cpp b/Engine/Source/Runtime/Experimental/IoStoreOnDemand/Private/OnDemandIoStore.cpp
new file mode 100644
--- /dev/null
+++ b/Engine/Source/Runtime/Experimental/IoStoreOnDemand/Private/OnDemandIoStore.cpp
@@ -0,0 +1,575 @@
+// Copyright Epic Games, Inc. All Rights Reserved.
+
+#include "OnDemandIoStore.h"
+#include "OnDemandHttpClient.h"
+
+#include "Algo/RemoveIf.h"
+#include "Algo/Transform.h"
+#include "Async/ManualResetEvent.h"
+#include "Async/UniqueLock.h"
+#include "Containers/StringConv.h"
+#include "HAL/FileManager.h"
+#include "HAL/FileManagerGeneric.h"
+#include "HAL/PlatformFile.h"
+#include "HAL/PlatformFileManager.h"
+#include "HAL/PlatformTime.h"
+#include "HAL/PlatformProcess.h"
+#include "Misc/CoreDelegatesInternal.h"
+#include "Misc/EncryptionKeyManager.h"
+#include "Misc/Guid.h"
+#include "Misc/Paths.h"
+#include "Misc/PathViews.h"
+#include "Serialization/MemoryReader.h"
+
+namespace UE::IoStore
+{
+
+extern FString GIasOnDemandTocExt;
+
+///////////////////////////////////////////////////////////////////////////////
+namespace Private
+{
+
+static void SplitHostUrl(const FStringView& Url, FStringView& OutHost, FStringView& OutRemainder)
+{
+	OutHost = OutRemainder = FStringView();
+
+	if (Url.StartsWith(TEXTVIEW("http")))
+	{
+		int32 Delim = INDEX_NONE;
+		ensure(Url.FindChar(':', Delim));
+		const int32 ProtocolDelim = Delim + 3;
+		ensure(Url.RightChop(ProtocolDelim).FindChar('/', Delim));
+		OutHost = Url.Left(ProtocolDelim + Delim);
+	}
+
+	OutRemainder = Url.RightChop(OutHost.Len());
+}
+
+/**
+ * Utility to create a FArchive capable of reading from disk using the exact same pathing
+ * rules as FPlatformMisc::LoadTextFileFromPlatformPackage but without forcing the entire
+ * file to be loaded at once.
+ */
+static TUniquePtr<FArchive> CreateReaderFromPlatformPackage(const FString& RelPath)
+{
+	TRACE_CPUPROFILER_EVENT_SCOPE(IasBackend::CreateReaderFromPlatformPackage);
+
+	const FString AbsPath = FPaths::Combine(FGenericPlatformMisc::RootDir(), RelPath);
+
+	IFileHandle* File = IPlatformFile::GetPlatformPhysical().OpenRead(*AbsPath);
+	if (File)
+	{
+#if PLATFORM_ANDROID
+		// This is a handle to an asset so we need to call Seek(0) to move the internal
+		// offset to the start of the asset file.
+		File->Seek(0);
+#endif //PLATFORM_ANDROID
+		const uint32 ReadBufferSize = 256 * 1024;
+		return MakeUnique<FArchiveFileReaderGeneric>(File, *AbsPath, File->Size(), ReadBufferSize);
+	}
+	else
+	{
+		return TUniquePtr<FArchive>();
+	}	
+}
+
+} // namespace UE::IoStore::Private
+
+///////////////////////////////////////////////////////////////////////////////
+const FOnDemandChunkEntry FOnDemandChunkEntry::Null = {};
+
+///////////////////////////////////////////////////////////////////////////////
+FOnDemandIoStore::FOnDemandIoStore()
+{
+	FEncryptionKeyManager::Get().OnKeyAdded().AddRaw(this, &FOnDemandIoStore::OnEncryptionKeyAdded);
+}
+
+FOnDemandIoStore::~FOnDemandIoStore()
+{
+	FEncryptionKeyManager::Get().OnKeyAdded().RemoveAll(this);
+
+	if (OnMountPakHandle.IsValid())
+	{
+		FCoreInternalDelegates::GetOnPakMountOperation().Remove(OnMountPakHandle);
+	}
+
+	if (TickFuture.IsValid())
+	{
+		TickFuture.Wait();
+	}
+}
+
+FIoStatus FOnDemandIoStore::Initialize()
+{
+	OnMountPakHandle = FCoreInternalDelegates::GetOnPakMountOperation().AddLambda(
+		[this](EMountOperation Operation, const TCHAR* ContainerPath, int32 Order) -> void
+		{
+			IPlatformFile& Ipf = IPlatformFile::GetPlatformPhysical();
+			const FString OnDemandTocPath = FPathViews::ChangeExtension(ContainerPath, GIasOnDemandTocExt);
+
+			if (Ipf.FileExists(*OnDemandTocPath) == false)
+			{
+				return;
+			}
+
+			switch (Operation)
+			{
+			case EMountOperation::Mount:
+			{
+				UE::FManualResetEvent DoneEvent;
+				Mount(FOnDemandMountArgs
+				{
+					.MountId = OnDemandTocPath,
+					.FilePath = OnDemandTocPath
+				},
+				[&DoneEvent](TIoStatusOr<FOnDemandMountResult> Result)
+				{
+					UE_CLOG(!Result.IsOk(), LogIoStoreOnDemand, Error,
+						TEXT("Failed to mount container, reason '%s'"), *Result.Status().ToString());
+					DoneEvent.Notify();
+				});
+				DoneEvent.Wait();
+				break;
+			}
+			case EMountOperation::Unmount:
+			{
+				const FIoStatus Status = Unmount(OnDemandTocPath);
+				UE_CLOG(!Status.IsOk(), LogIoStoreOnDemand, Error,
+					TEXT("Failed to unmount container, reason '%s'"), *Status.ToString());
+				break;
+			}
+			default:
+				checkNoEntry();
+			}
+		}
+	);
+
+	{
+		IPlatformFile& Ipf = IPlatformFile::GetPlatformPhysical();
+		FCurrentlyMountedPaksDelegate& Delegate = FCoreInternalDelegates::GetCurrentlyMountedPaksDelegate();
+		if (Delegate.IsBound())
+		{
+			TArray<FMountedPakInfo> PakInfo = Delegate.Execute();
+			for (const FMountedPakInfo& Info : PakInfo)
+			{
+				check(Info.PakFile != nullptr);
+				const FString OnDemandTocPath = FPathViews::ChangeExtension(Info.PakFile->PakGetPakFilename(), GIasOnDemandTocExt);
+				if (Ipf.FileExists(*OnDemandTocPath))
+				{
+					MountRequests.Add(OnDemandTocPath, MakeShared<FMountRequest>(FMountRequest
+					{
+						.MountArgs = FOnDemandMountArgs
+						{
+							.MountId = OnDemandTocPath,
+							.FilePath = OnDemandTocPath
+						},
+						.OnCompleted = [](TIoStatusOr<FOnDemandMountResult> Result)
+						{
+							UE_CLOG(!Result.IsOk(), LogIoStoreOnDemand, Error, TEXT("Failed to mount container, reason '%s'"),
+								*Result.Status().ToString());
+						}
+					}));
+				}
+			}
+		}
+
+		// Process mount requests synchronously at startup
+		while (Tick(MAX_int64));
+	}
+
+	return EIoErrorCode::Ok;
+}
+
+void FOnDemandIoStore::Mount(FOnDemandMountArgs&& Args, FOnDemandMountCompleted&& OnCompleted)
+{
+	if (Args.MountId.IsEmpty())
+	{
+		return OnCompleted(FIoStatus(EIoErrorCode::InvalidParameter, TEXT("Invalid Mount ID")));
+	}
+
+	{
+		UE::TUniqueLock Lock(MountRequestMutex);
+		FSharedMountRequest& MountRequest = MountRequests.FindOrAdd(Args.MountId);
+
+		if (MountRequest.IsValid())
+		{
+			UE_LOG(LogIoStoreOnDemand, Warning, TEXT("Mount request '%s' is already mounting"), *Args.MountId);
+			return OnCompleted(FIoStatus(EIoErrorCode::InvalidParameter));
+		}
+
+		UE_LOG(LogIoStoreOnDemand, Log, TEXT("Enqueing mount request, MountId='%s'"), *Args.MountId);
+		MountRequest = MakeShared<FMountRequest>(FMountRequest
+		{ 
+			.MountArgs		= MoveTemp(Args),
+			.OnCompleted	= MoveTemp(OnCompleted),
+		});
+	}
+
+	ConditionallyStartTicking();
+}
+
+void FOnDemandIoStore::ConditionallyStartTicking()
+{
+	check(FPlatformProcess::SupportsMultithreading());
+
+	if (FPlatformProcess::SupportsMultithreading() && GIOThreadPool != nullptr)
+	{
+		bool bExpected = false;
+		if (bTicking.compare_exchange_strong(bExpected, true))
+		{
+			TickFuture = AsyncPool(
+				*GIOThreadPool,
+				[this]
+				{
+					while (Tick(MAX_int64));
+				},
+				nullptr,
+				EQueuedWorkPriority::Low);
+		}
+	}
+	else
+	{
+		while (Tick(MAX_int64));
+	}
+}
+
+FIoStatus FOnDemandIoStore::Unmount(FStringView MountId)
+{
+	UE_LOG(LogIoStoreOnDemand, Log, TEXT("Unmounting '%s'"), *WriteToString<256>(MountId));
+
+	{
+		UE::TUniqueLock Lock(MountRequestMutex);
+		if (FSharedMountRequest* MountRequest = MountRequests.Find(FString(MountId)))
+		{
+			(*MountRequest)->bCancelled = true;
+		}
+	}
+
+	{
+		TUniqueLock Lock(ContainerMutex);
+
+		Containers.SetNum(Algo::RemoveIf(Containers, [&MountId](const FSharedOnDemandContainer& Container)
+		{
+			if (Container->MountId == MountId)
+			{
+				UE_LOG(LogIoStoreOnDemand, Log, TEXT("Unmounting container, ContainerName='%s', MountId='%s'"),
+					*WriteToString<128>(Container->Name), *WriteToString<128>(Container->MountId));
+				return true;
+			}
+
+			return false;
+		}));
+	}
+
+	return EIoErrorCode::Ok;
+}
+
+FOnDemandChunkInfo FOnDemandIoStore::GetChunkInfo(const FIoChunkId& ChunkId)
+{
+	TUniqueLock Lock(ContainerMutex);
+
+	for (FSharedOnDemandContainer& Container : Containers)
+	{
+		if (EnumHasAnyFlags(Container->Flags, EOnDemandContainerFlags::Mounted))
+		{
+			if (const FOnDemandChunkEntry* Entry = Container->ChunkEntries.Find(ChunkId))
+			{
+				return FOnDemandChunkInfo(Container, *Entry);
+			}
+		}
+	}
+
+	return FOnDemandChunkInfo();
+}
+
+bool FOnDemandIoStore::Tick(int64 MaxCycles)
+{
+	check(!bTicking);
+	bTicking = true;
+	ON_SCOPE_EXIT{ bTicking = false; };
+
+	TArray<FSharedMountRequest> Requests;
+	{
+		UE::TUniqueLock Lock(MountRequestMutex);
+		for (TPair<FString, FSharedMountRequest>& Kv : MountRequests)
+		{
+			Requests.Add(Kv.Value);
+		}
+	}
+
+	if (Requests.IsEmpty())
+	{
+		return false;
+	}
+
+	for (FSharedMountRequest& Request : Requests)
+	{
+		FIoStatus Status = ProcessMountRequest(*Request);
+
+		{
+			UE::TUniqueLock Lock(MountRequestMutex);
+			MountRequests.Remove(Request->MountArgs.MountId);
+			if (Request->bCancelled)
+			{
+				Status = EIoErrorCode::Cancelled;
+			}
+		}
+
+		TIoStatusOr<FOnDemandMountResult> MountStatus;
+		if (Status.IsOk())
+		{
+			UE::TUniqueLock Lock(ContainerMutex);
+			for (const FSharedOnDemandContainer& Container : Request->Containers)
+			{
+				Containers.Add(Container);
+				if (Container->EncryptionKeyGuid.IsEmpty() == false)
+				{
+					FGuid KeyGuid;
+					ensure(FGuid::Parse(Container->EncryptionKeyGuid, KeyGuid));
+					if (FEncryptionKeyManager::Get().TryGetKey(KeyGuid, Container->EncryptionKey) == false)
+					{
+						UE_LOG(LogIoStoreOnDemand, Log, TEXT("Deferring container '%s' until encryption key '%s' becomes available"),
+							*Container->Name, *Container->EncryptionKeyGuid);
+						EnumAddFlags(Container->Flags, EOnDemandContainerFlags::PendingEncryptionKey);
+						continue;
+					}
+				}
+
+				UE_LOG(LogIoStoreOnDemand, Log, TEXT("Mounting container '%s' (%d entries)"),
+					*Container->Name, Container->ChunkEntries.Num());
+				EnumAddFlags(Container->Flags, EOnDemandContainerFlags::Mounted);
+			}
+
+			MountStatus = TIoStatusOr<FOnDemandMountResult>(FOnDemandMountResult
+			{
+				.MountId = Request->MountArgs.MountId
+			});
+		}
+		else
+		{
+			MountStatus = TIoStatusOr<FOnDemandMountResult>(Status);	
+		}
+
+		FOnDemandMountCompleted OnCompleted = MoveTemp(Request->OnCompleted);
+		OnCompleted(MountStatus);
+	}
+
+	return true;
+}
+
+FIoStatus FOnDemandIoStore::ProcessMountRequest(FMountRequest& MountRequest)
+{
+	UE_LOG(LogIoStoreOnDemand, Log, TEXT("Processing mount request, MountId='%s'"), *MountRequest.MountArgs.MountId);
+
+	FOnDemandMountArgs& Args = MountRequest.MountArgs;
+	FStringView Host, TocRelUrl;
+	Private::SplitHostUrl(Args.Url, Host, TocRelUrl);
+	const FStringView TocPath = FPathViews::GetPath(TocRelUrl);
+	
+	if (Args.Toc.IsSet())
+	{
+		CreateContainersFromToc(Args.MountId, TocPath, Args.Toc.GetValue(), MountRequest.Containers);
+	}
+	else if (Args.FilePath.IsEmpty() == false)
+	{
+		UE_LOG(LogIoStoreOnDemand, Log, TEXT("Loading TOC from file '%s'"), *Args.FilePath);
+
+		TUniquePtr<FArchive> Ar;
+		if (FPlatformMisc::FileExistsInPlatformPackage(Args.FilePath))
+		{
+			Ar = Private::CreateReaderFromPlatformPackage(Args.FilePath);
+		}
+		else
+		{
+			Ar.Reset(IFileManager::Get().CreateFileReader(*Args.FilePath));
+		}
+
+		if (Ar.IsValid() == false)
+		{
+			return FIoStatusBuilder(EIoErrorCode::FileNotOpen)
+				<< TEXT("Failed to open '")
+				<< Args.FilePath
+				<< TEXT("'");
+		}
+
+		// TODO: Add sentinel to all serialization paths for TOC files
+		const bool bIncludeDot = true;	
+		const bool bValidate = FPathViews::GetExtension(Args.FilePath, bIncludeDot) == GIasOnDemandTocExt;
+		if (bValidate)
+		{
+			const int64 SentinelPos = Ar->TotalSize() - FOnDemandTocSentinel::SentinelSize;
+
+			if (SentinelPos < 0)
+			{
+				UE_LOG(LogIoStoreOnDemand, Error, TEXT("The file '%s' is smaller than expected and quite possible corrupted"), *Args.FilePath);
+				return EIoErrorCode::CorruptToc;
+			}
+
+			Ar->Seek(SentinelPos);
+
+			FOnDemandTocSentinel Sentinel;
+			(*Ar) << Sentinel;
+
+			if (!Sentinel.IsValid())
+			{
+				UE_LOG(LogIas, Error, TEXT("File corruption detected when serializing '%s'"), *Args.FilePath);
+				return EIoErrorCode::CorruptToc;
+			}
+
+			Ar->Seek(0);
+		}
+
+		FOnDemandToc Toc;
+		*Ar << Toc;
+
+		if (Ar->IsError() || Ar->IsCriticalError())
+		{
+			return FIoStatusBuilder(EIoErrorCode::FileNotOpen) << TEXT("Failed to serialize TOC file");
+		}
+
+		CreateContainersFromToc(Args.MountId, TocPath, Toc, MountRequest.Containers);
+	}
+	else if (Args.Url.IsEmpty() == false)
+	{
+		UE_LOG(LogIoStoreOnDemand, Log, TEXT("Loading TOC from URL '%s'"), *Args.Url);
+
+		if (Host.IsEmpty())
+		{
+			return FIoStatusBuilder(EIoErrorCode::InvalidParameter) << TEXT("No valid host URL");
+		}
+
+		FHttpClientConfig HttpCfg;
+		HttpCfg.Endpoints.Add(FString(Host));
+		HttpCfg.PrimaryEndpoint = 0; 
+		HttpCfg.MaxRetryCount = 2;
+		TUniquePtr<FHttpClient> Client = FHttpClient::Create(MoveTemp(HttpCfg));
+
+		TAnsiStringBuilder<256> Url;
+		if (!TocRelUrl.StartsWith(TCHAR('/')))
+		{
+			Url << "/";
+		}
+		Url << TocRelUrl;
+
+		FIoBuffer Body;
+		FIoStatus HttpStatus;
+		Client->Get(Url.ToView(), [&Body, &HttpStatus](TIoStatusOr<FIoBuffer> Response, uint64)
+			{
+				if (Response.IsOk())
+				{
+					Body = Response.ConsumeValueOrDie();
+				}
+				else
+				{
+					HttpStatus = Response.Status();
+				}
+			});
+
+		while (Client->Tick());
+
+		if (HttpStatus.IsOk() == false)
+		{
+			return HttpStatus;
+		}
+
+		FOnDemandToc Toc;
+		FMemoryReaderView Ar(Body.GetView());
+		Ar << Toc;
+
+		if (Ar.IsError() || Ar.IsCriticalError())
+		{
+			return FIoStatusBuilder(EIoErrorCode::FileNotOpen) << TEXT("Failed to serialize TOC from HTTP response");
+		}
+
+		CreateContainersFromToc(Args.MountId, TocPath, Toc, MountRequest.Containers);
+	}
+	else
+	{
+		return FIoStatusBuilder(EIoErrorCode::InvalidParameter) << TEXT("Invalid mount arguments");
+	}
+
+	return EIoErrorCode::Ok;
+}
+
+void FOnDemandIoStore::OnEncryptionKeyAdded(const FGuid& Id, const FAES::FAESKey& Key)
+{
+	TUniqueLock Lock(ContainerMutex);
+
+	for (FSharedOnDemandContainer& Container : Containers)
+	{
+		if (EnumHasAnyFlags(Container->Flags, EOnDemandContainerFlags::PendingEncryptionKey))
+		{
+			FGuid KeyGuid;
+			ensure(FGuid::Parse(FString(StringCast<TCHAR>(*Container->EncryptionKeyGuid)), KeyGuid));
+
+			if (FEncryptionKeyManager::Get().TryGetKey(KeyGuid, Container->EncryptionKey))
+			{
+				UE_LOG(LogIoStoreOnDemand, Log, TEXT("Mounting container '%s' (%d entries)"),
+					*WriteToString<128>(Container->Name), Container->ChunkEntries.Num());
+
+				EnumRemoveFlags(Container->Flags, EOnDemandContainerFlags::PendingEncryptionKey);
+				EnumAddFlags(Container->Flags, EOnDemandContainerFlags::Mounted);
+			}
+		}
+	}
+}
+
+void FOnDemandIoStore::CreateContainersFromToc(
+	FStringView MountId,
+	FStringView TocPath,
+	FOnDemandToc& Toc,
+	TArray<FSharedOnDemandContainer>& Out)
+{
+	const FOnDemandTocHeader& Header = Toc.Header;
+	const FName CompressionFormat(Header.CompressionFormat);
+
+	TStringBuilder<128> ChunksDirectory;
+
+	{
+		if (TocPath.IsEmpty() == false)
+		{
+			Algo::Transform(TocPath, AppendChars(ChunksDirectory), FChar::ToLower);
+		}
+		else
+		{
+			Algo::Transform(Toc.Header.ChunksDirectory, AppendChars(ChunksDirectory), FChar::ToLower);
+		}
+
+		FPathViews::Append(ChunksDirectory, TEXT("chunks"));
+	}
+
+	for (FOnDemandTocContainerEntry& ContainerEntry : Toc.Containers)
+	{
+		FSharedOnDemandContainer Container = MakeShared<FOnDemandContainer>();
+		Container->Name					= MoveTemp(ContainerEntry.ContainerName);
+		Container->MountId				= MountId;
+		Container->ChunksDirectory		= StringCast<ANSICHAR>(ChunksDirectory.ToString());
+		Container->EncryptionKeyGuid	= MoveTemp(ContainerEntry.EncryptionKeyGuid);
+		Container->BlockSize			= Header.BlockSize;
+		Container->BlockSizes			= MoveTemp(ContainerEntry.BlockSizes);
+		Container->BlockHashes			= MoveTemp(ContainerEntry.BlockHashes);
+		Container->ContainerId			= ContainerEntry.ContainerId;
+		Container->CompressionFormats.Add(CompressionFormat);
+
+		Container->ChunkEntries.Reserve(ContainerEntry.Entries.Num());
+		for (const FOnDemandTocEntry& TocEntry : ContainerEntry.Entries)
+		{
+			check(TocEntry.RawSize <= 0xffff'ffffull);
+			check(TocEntry.EncodedSize <= 0xffff'ffffull);
+			Container->ChunkEntries.Add(TocEntry.ChunkId, FOnDemandChunkEntry
+			{
+				.Hash					= TocEntry.Hash,
+				.RawSize				= uint32(TocEntry.RawSize),
+				.EncodedSize			= uint32(TocEntry.EncodedSize),
+				.BlockOffset			= TocEntry.BlockOffset,
+				.BlockCount				= TocEntry.BlockCount,
+				.CompressionFormatIndex	= 0
+			});
+		}
+
+		Out.Add(MoveTemp(Container));
+	}
+}
+
+} // namespace UE::IoStore
diff --git a/Engine/Source/Runtime/Experimental/IoStoreOnDemand/Private/OnDemandIoStore.h b/Engine/Source/Runtime/Experimental/IoStoreOnDemand/Private/OnDemandIoStore.h
new file mode 100644
--- /dev/null
+++ b/Engine/Source/Runtime/Experimental/IoStoreOnDemand/Private/OnDemandIoStore.h
@@ -0,0 +1,157 @@
+// Copyright Epic Games, Inc. All Rights Reserved.
+
+#pragma once
+
+#include "Async/Async.h"
+#include "Async/Mutex.h"
+#include "Containers/AnsiString.h"
+#include "IO/IoStoreOnDemand.h"
+#include "IO/IoHash.h"
+#include "IO/IoChunkEncoding.h"
+#include "IO/IoChunkId.h"
+#include "Misc/AES.h"
+#include "Misc/EnumClassFlags.h"
+
+#include <atomic>
+
+namespace UE::IoStore
+{
+
+///////////////////////////////////////////////////////////////////////////////
+enum class EOnDemandContainerFlags : uint8
+{
+	None					= 0,
+	PendingEncryptionKey	= (1 << 0),
+	Mounted					= (1 << 1),
+};
+ENUM_CLASS_FLAGS(EOnDemandContainerFlags);
+
+///////////////////////////////////////////////////////////////////////////////
+struct FOnDemandChunkEntry
+{
+	static const FOnDemandChunkEntry Null;
+
+	FIoHash	Hash;
+	uint32	RawSize = 0;
+	uint32	EncodedSize = 0;
+	uint32	BlockOffset = ~uint32(0);
+	uint32	BlockCount = 0; 
+	uint8	CompressionFormatIndex = 0;
+};
+
+///////////////////////////////////////////////////////////////////////////////
+struct FOnDemandContainer
+{
+	using FChunkEntryMap = TMap<FIoChunkId, FOnDemandChunkEntry>;
+
+	FAES::FAESKey			EncryptionKey;
+	FIoHash					PakHash;
+	FChunkEntryMap 			ChunkEntries;
+	FString					EncryptionKeyGuid;
+	FString					Name;
+	FString					MountId;
+	FAnsiString				ChunksDirectory;
+	TArray<FName>			CompressionFormats;
+	TArray<uint32>			BlockSizes;
+	TArray<FIoBlockHash>	BlockHashes;
+	FIoContainerId			ContainerId;
+	uint32					BlockSize = 0;
+	EOnDemandContainerFlags Flags = EOnDemandContainerFlags ::None;
+};
+
+using FSharedOnDemandContainer = TSharedPtr<FOnDemandContainer, ESPMode::ThreadSafe>;
+
+///////////////////////////////////////////////////////////////////////////////
+struct FOnDemandChunkInfo
+{
+	FOnDemandChunkInfo()
+		: Entry(FOnDemandChunkEntry::Null)
+	{ }
+
+	const									FIoHash& Hash() const { return Entry.Hash; }
+	uint32									RawSize() const { return Entry.RawSize; }
+	uint32									EncodedSize() const { return Entry.EncodedSize; }
+	uint32									BlockSize() const { return SharedContainer->BlockSize; }
+	FName									CompressionFormat() const { return SharedContainer->CompressionFormats[Entry.CompressionFormatIndex]; }
+	FMemoryView								EncryptionKey() const { return FMemoryView(SharedContainer->EncryptionKey.Key, FAES::FAESKey::KeySize); }
+	inline TConstArrayView<uint32>			Blocks() const;
+	inline TConstArrayView<FIoBlockHash>	BlockHashes() const;
+	FAnsiStringView							ChunksDirectory() const { return SharedContainer->ChunksDirectory; }
+
+	bool									IsValid() const { return SharedContainer.IsValid(); }
+	operator								bool() const { return IsValid(); }
+
+private:
+	friend class FOnDemandIoStore;
+
+	FOnDemandChunkInfo(FSharedOnDemandContainer InContainer, const FOnDemandChunkEntry& InEntry)
+		: SharedContainer(InContainer)
+		, Entry(InEntry)
+	{ }
+
+	FSharedOnDemandContainer	SharedContainer;
+	const FOnDemandChunkEntry&	Entry;
+};
+
+TConstArrayView<uint32> FOnDemandChunkInfo::Blocks() const
+{
+	return TConstArrayView<uint32>(SharedContainer->BlockSizes.GetData() + Entry.BlockOffset, Entry.BlockCount);
+}
+
+TConstArrayView<FIoBlockHash> FOnDemandChunkInfo::BlockHashes() const
+{
+	return SharedContainer->BlockHashes.IsEmpty()
+		? TConstArrayView<FIoBlockHash>()
+		: TConstArrayView<FIoBlockHash>(SharedContainer->BlockHashes.GetData() + Entry.BlockOffset, Entry.BlockCount);
+}
+
+///////////////////////////////////////////////////////////////////////////////
+class FOnDemandIoStore
+{
+	struct FMountRequest
+	{
+		FOnDemandMountArgs					MountArgs;
+		FOnDemandMountCompleted				OnCompleted;
+		TArray<FSharedOnDemandContainer>	Containers;
+		bool								bCancelled = false;
+	};
+
+	using FSharedMountRequest	= TSharedPtr<FMountRequest>;
+	using FMountRequestMap		= TMap<FString, FSharedMountRequest>;
+
+public:
+	FOnDemandIoStore();
+	~FOnDemandIoStore();
+	FOnDemandIoStore(const FOnDemandIoStore&) = delete;
+	FOnDemandIoStore(FOnDemandIoStore&&) = delete;
+	FOnDemandIoStore& operator=(const FOnDemandIoStore&) = delete;
+	FOnDemandIoStore& operator=(FOnDemandIoStore&&) = delete;
+
+	FIoStatus				Initialize();
+	void					Mount(FOnDemandMountArgs&& Args, FOnDemandMountCompleted&& OnCompleted);
+	FIoStatus				Unmount(FStringView MountId);
+	FOnDemandChunkInfo		GetChunkInfo(const FIoChunkId& ChunkId);
+
+private:
+	bool					Tick(int64 MaxCycles);
+	FIoStatus				ProcessMountRequest(FMountRequest& MountRequest);
+	void					ConditionallyStartTicking();
+	void					OnEncryptionKeyAdded(const FGuid& Id, const FAES::FAESKey& Key);
+	static void				CreateContainersFromToc(
+								FStringView MountId,
+								FStringView TocPath,
+								FOnDemandToc& Toc,
+								TArray<FSharedOnDemandContainer>& Out);
+
+	FDelegateHandle						OnMountPakHandle;
+	TArray<FSharedOnDemandContainer>	Containers;
+	UE::FMutex							ContainerMutex;
+
+	FMountRequestMap					MountRequests;
+	UE::FMutex							MountRequestMutex;
+
+	std::atomic_bool					bTicking{false};
+	TFuture<void>						TickFuture;
+};
+
+} // namespace UE::IoStore
diff --git a/Engine/Source/Runtime/Experimental/IoStoreOnDemand/Private/Tool/Upload.cpp b/Engine/Source/Runtime/Experimental/IoStoreOnDemand/Private/Tool/Upload.cpp
--- a/Engine/Source/Runtime/Experimental/IoStoreOnDemand/Private/Tool/Upload.cpp
+++ b/Engine/Source/Runtime/Experimental/IoStoreOnDemand/Private/Tool/Upload.cpp
@@ -46,6 +46,7 @@ static FIoStoreUploadParams BuildUploadParams(const FContext& Context)
 	Ret.BuildVersion			= Context.Get<FStringView>(TEXT("-BuildVersion"),			Ret.BuildVersion);
 	Ret.TargetPlatform			= Context.Get<FStringView>(TEXT("-TargetPlatform"),			Ret.TargetPlatform);
 	Ret.bWriteTocToDisk			= Context.Get<bool>(TEXT("-WriteTocToDisk"),				Ret.bWriteTocToDisk);
+	Ret.bPerContainerTocs		= Context.Get<bool>(TEXT("-PerContainerTocs"),				Ret.bPerContainerTocs);
 	Ret.MaxConcurrentUploads	= Context.Get<int32>(TEXT("-MaxConcurrentUploads"),			Ret.MaxConcurrentUploads);
 
 	Ret.bDeleteContainerFiles	= !Context.Get<bool>(TEXT("-KeepContainerFiles"),	!Ret.bDeleteContainerFiles);
@@ -146,6 +147,7 @@ static FCommand UploadCommand(
 		TArgument<bool>(TEXT("-KeepContainerFiles"),	TEXT("Do not delete container files after upload")),
 		TArgument<bool>(TEXT("-KeepPakFiles"),			TEXT("Do not delete the springboard pak files")),
 		TArgument<bool>(TEXT("-WriteTocToDisk"),		TEXT("Output the TOC to disk as well as uploading")),
+		TArgument<bool>(TEXT("-PerContainerTocs"),		TEXT("Whether to generate TOC's for each container file(s)")),
 		TArgument<int32>(TEXT("-MaxConcurrentUploads"),	TEXT("Number of simultaneous uploads")),
 		S3Arguments,
 	}
diff --git a/Engine/Source/Runtime/Experimental/IoStoreOnDemand/Public/IO/IoStoreOnDemand.h b/Engine/Source/Runtime/Experimental/IoStoreOnDemand/Public/IO/IoStoreOnDemand.h
--- a/Engine/Source/Runtime/Experimental/IoStoreOnDemand/Public/IO/IoStoreOnDemand.h
+++ b/Engine/Source/Runtime/Experimental/IoStoreOnDemand/Public/IO/IoStoreOnDemand.h
@@ -28,6 +28,7 @@ struct FAnalyticsEventAttribute;
 struct FIoContainerSettings;
 struct FIoStoreWriterSettings;
 namespace UE::IoStore { struct FOnDemandEndpoint; }
+namespace UE::IoStore { class FOnDemandIoStore; }
 using FIoBlockHash = uint32;
 
 // Custom initialization allows users to control when
@@ -209,6 +210,7 @@ struct FIoStoreUploadParams
 	int32 MaxConcurrentUploads = 16;
 	bool bDeleteContainerFiles = true;
 	bool bDeletePakFiles = true;
+	bool bPerContainerTocs = false;
 
 	/** If we should write out the .iochunktoc to disk as well as uploading it. */
 	bool bWriteTocToDisk = false;
@@ -344,10 +346,18 @@ enum class EOnDemandInitResult
 
 struct FOnDemandMountArgs
 {
+	/** Mount an already serialized TOC. */
+	TOptional<FOnDemandToc> Toc;
+	/** Mandatory ID to be used for unmounting all container file(s) included in the TOC. */
 	FString MountId;
+	/** Download the TOC from the specified URL. */
 	FString Url;
+	/** Serialize the TOC from the specified file path. */
+	FString FilePath;
+	/** Directory path for any additonal content downloaded as part of the TOC. */
 	FString InstallDirectory;
-	bool bInstall = false;
+	/** Whether to download and install container chunks to local storage. */
+	bool bInstall = true;
 };
 
 struct FOnDemandMountResult
@@ -368,6 +378,7 @@ private:
 	TOptional<bool> DeferredEnabled;
 	TOptional<bool> DeferredAbandonCache;
 	TOptional<bool> DeferredBulkOptionalEnabled;
+	TUniquePtr<FOnDemandIoStore> IoStore;
 
 public:
 	UE_API void SetBulkOptionalEnabled(bool bInEnabled);
