Chaos - improve but disable new convex-mesh collision by default

Inner-face collisions were not being rejected following the prior convex-mesh collision improvement. However, it still has issues so disabling it for now.

diff --git a/Engine/Source/Runtime/Experimental/Chaos/Private/Chaos/CollisionOneShotManifolds.cpp b/Engine/Source/Runtime/Experimental/Chaos/Private/Chaos/CollisionOneShotManifolds.cpp
--- a/Engine/Source/Runtime/Experimental/Chaos/Private/Chaos/CollisionOneShotManifolds.cpp
+++ b/Engine/Source/Runtime/Experimental/Chaos/Private/Chaos/CollisionOneShotManifolds.cpp
@@ -137,13 +137,22 @@ namespace Chaos
 
 	bool bChaos_Collision_UseCapsuleTriMesh2 = true;
 	FAutoConsoleVariableRef CVarChaos_Collision_UseCapsuleTriMesh2(TEXT("p.Chaos.Collision.UseCapsuleTriMesh2"), bChaos_Collision_UseCapsuleTriMesh2, TEXT(""));
+	
 	// @todo(chaos): (2) is the intended solution. The others are for rollback if blocking bugs are found. Remove this once (2) is tested
 	// 0: Use GJK + EPA to generate closest feature, generate multi-point manifold from its normal, and then fix manifold based on mesh topology
 	// 1: Use SAT to generate closest feature, generate multi-point manifold from its normal, and then fix manifold based on mesh topology
 	// 2: Use GJK + SAT to generate closest feature, then fix its normal based on mesh topology, then generate the multi-point manifold
-	int32 Chaos_Collision_ConvexTriMeshMode = 2;
+	int32 Chaos_Collision_ConvexTriMeshMode = 1;
 	FAutoConsoleVariableRef CVarChaos_Collision_UseConvexTriMesh2(TEXT("p.Chaos.Collision.ConvexTriMeshMode"), Chaos_Collision_ConvexTriMeshMode, TEXT(""));
 
+	// Enable backface culling in ConstructConvexTriangleOneShotManifold3
+	bool bChaos_Collision_ConvexTriMeshBackFaceCull = true;
+	FAutoConsoleVariableRef CVarChaos_Collision_ConvexTriMeshBackFaceCull(TEXT("p.Chaos.Collision.ConvexTriMeshBackFaceCull"), bChaos_Collision_ConvexTriMeshBackFaceCull, TEXT(""));
+
+	// Enable convex-inside-triangle culling in ConstructConvexTriangleOneShotManifold3
+	bool bChaos_Collision_ConvexTriMeshInsideCull = true;
+	FAutoConsoleVariableRef CVarChaos_Collision_ConvexTriMeshInsideCull(TEXT("p.Chaos.Collision.ConvexTriMeshInsideCull"), bChaos_Collision_ConvexTriMeshInsideCull, TEXT(""));
+
 	// true: use GJK and then SAT if the shapes overlap
 	// false: use GJK and then EPA if the shapes overlap
 	bool bChaos_Collision_UseConvexTriangleGJKSAT = true;
diff --git a/Engine/Source/Runtime/Experimental/Chaos/Private/Chaos/CollisionOneShotManifoldsMeshShapes.cpp b/Engine/Source/Runtime/Experimental/Chaos/Private/Chaos/CollisionOneShotManifoldsMeshShapes.cpp
--- a/Engine/Source/Runtime/Experimental/Chaos/Private/Chaos/CollisionOneShotManifoldsMeshShapes.cpp
+++ b/Engine/Source/Runtime/Experimental/Chaos/Private/Chaos/CollisionOneShotManifoldsMeshShapes.cpp
@@ -39,7 +39,9 @@ namespace Chaos
 
 	extern bool bChaos_Collision_UseCapsuleTriMesh2;
 	extern int32 Chaos_Collision_ConvexTriMeshMode;
-	
+	extern bool bChaos_Collision_ConvexTriMeshInsideCull;
+	extern bool bChaos_Collision_ConvexTriMeshBackFaceCull;
+
 	namespace CVars
 	{
 #if CHAOS_DEBUG_DRAW
@@ -102,6 +104,16 @@ namespace Chaos
 			const FTriangle& Triangle = ContactGenerator.GetTriangle(TriangleIndex);
 			const FVec3 TriangleNormal = ContactGenerator.GetTriangleNormal(TriangleIndex);
 
+			// If the convex origin is inside the triangle, ignore it
+			if (bChaos_Collision_ConvexTriMeshInsideCull)
+			{
+				const FReal ConvexDistance = FVec3::DotProduct(Triangle.GetCentroid() - Convex.GetCenterOfMass(), TriangleNormal);
+				if (ConvexDistance > 0)
+				{
+					return;
+				}
+			}
+
 			// Find the closest feature between the Convex at its initial position X and the Triangle
 			Private::FConvexContactPoint ClosestContact;
 			const bool bFoundClosestContact = Private::FindClosestFeatures(Convex, Triangle, TriangleNormal, FVec3(0), CullDistance, ClosestContact);
@@ -120,6 +132,23 @@ namespace Chaos
 				}
 				#endif
 
+				// Backface cull based on closest contact normal
+				const FReal TriangleDotNormal = FVec3::DotProduct(TriangleNormal, ClosestContact.ShapeContactNormal);
+				if (bChaos_Collision_ConvexTriMeshBackFaceCull)
+				{
+					if (TriangleDotNormal < 0)
+					{
+						return;
+					}
+				}
+
+				// Cull distance is zero for back faces
+				const FReal EffectiveCullDistance = (TriangleDotNormal < 0) ? FReal(0.0) : CullDistance;
+				if (ClosestContact.Phi > EffectiveCullDistance)
+				{
+					return;
+				}
+
 				// Use the mesh info to correct the normal - this corrects edge and vertex normals if they are
 				// outside the range allowed by the set of triangles sharing the feature
 				if (ContactGenerator.FixFeature(TriangleIndex, ClosestContact.Features[1].FeatureType, ClosestContact.Features[1].PlaneFeatureIndex, ClosestContact.ShapeContactNormal))
@@ -140,8 +169,8 @@ namespace Chaos
 				#endif
 
 				// Back face culling based on the corrected feature
-				const FReal TriangleDotNormal = FVec3::DotProduct(TriangleNormal, ClosestContact.ShapeContactNormal);
-				if (TriangleDotNormal < 0)
+				const FReal TriangleDotCorrectedNormal = FVec3::DotProduct(TriangleNormal, ClosestContact.ShapeContactNormal);
+				if (TriangleDotCorrectedNormal < 0)
 				{
 					return;
 				}
