Added on-demand package store backend

diff --git a/Engine/Source/Runtime/Experimental/IoStoreOnDemand/IoStoreOnDemand.Build.cs b/Engine/Source/Runtime/Experimental/IoStoreOnDemand/IoStoreOnDemand.Build.cs
--- a/Engine/Source/Runtime/Experimental/IoStoreOnDemand/IoStoreOnDemand.Build.cs
+++ b/Engine/Source/Runtime/Experimental/IoStoreOnDemand/IoStoreOnDemand.Build.cs
@@ -10,6 +10,7 @@ public class IoStoreOnDemand : ModuleRules
 		PublicDependencyModuleNames.Add("TraceLog");
 		PrivateDependencyModuleNames.AddRange(
 			new string[] {
+				"CoreUObject",
 				"HTTP",
 				"Json",
 				"Analytics"
diff --git a/Engine/Source/Runtime/Experimental/IoStoreOnDemand/Private/OnDemandPackageStoreBackend.cpp b/Engine/Source/Runtime/Experimental/IoStoreOnDemand/Private/OnDemandPackageStoreBackend.cpp
new file mode 100644
--- /dev/null
+++ b/Engine/Source/Runtime/Experimental/IoStoreOnDemand/Private/OnDemandPackageStoreBackend.cpp
@@ -0,0 +1,225 @@
+// Copyright Epic Games, Inc. All Rights Reserved.
+
+#include "OnDemandPackageStoreBackend.h"
+
+#include "Algo/Accumulate.h"
+#include "Async/Mutex.h"
+#include "Async/UniqueLock.h"
+#include "IO/IoContainerHeader.h"
+#include "IO/IoStoreOnDemand.h"
+#include "IO/PackageId.h"
+#include "Internationalization/PackageLocalizationManager.h"
+
+namespace UE::IoStore
+{
+
+///////////////////////////////////////////////////////////////////////////////
+class FOnDemandPackageStoreBackend final
+	: public IOnDemandPackageStoreBackend
+{
+	struct FContainer
+	{
+		FString				ContainerName;
+		FIoContainerHeader	ContainerHeader;
+	};
+
+	using FSharedBackendContext = TSharedPtr<const FPackageStoreBackendContext>;
+	using FEntryMap				= TMap<FPackageId, const FFilePackageStoreEntry*>;
+	using FUniqueContainer		= TUniquePtr<FContainer>;
+	using FRedirect				= TTuple<FName, FPackageId>;
+	using FLocalizedMap			= TMap<FPackageId, FName>;
+	using FRedirectMap			= TMap<FPackageId, FRedirect>;
+
+public:
+						FOnDemandPackageStoreBackend();
+						virtual ~FOnDemandPackageStoreBackend ();
+
+	virtual FIoStatus	Mount(FString ContainerName, FIoContainerHeader&& ContainerHeader) override;
+	virtual FIoStatus	Unmount(const FString& ContainerName) override;
+
+	virtual void		OnMounted(TSharedRef<const FPackageStoreBackendContext> Context) override;
+	virtual void 		BeginRead() override;
+	virtual void 		EndRead() override;
+
+	virtual bool		GetPackageRedirectInfo(
+							FPackageId PackageId,
+							FName& OutSourcePackageName,
+							FPackageId& OutRedirectedToPackageId) override;
+
+	virtual EPackageStoreEntryStatus GetPackageStoreEntry(
+							FPackageId PackageId,
+							FName PackageName,
+							FPackageStoreEntry& OutPackageStoreEntry) override;
+
+private:
+	void				UpdateLookupTables();
+
+	FSharedBackendContext		BackendContext;
+	TArray<FUniqueContainer>	Containers;
+	FEntryMap					EntryMap;
+	FLocalizedMap				LocalizedMap;
+	FRedirectMap				RedirectMap;
+	UE::FMutex					Mutex;
+	bool						bNeedsUpdate = false;
+};
+
+///////////////////////////////////////////////////////////////////////////////
+FOnDemandPackageStoreBackend::FOnDemandPackageStoreBackend()
+{
+}
+
+FOnDemandPackageStoreBackend::~FOnDemandPackageStoreBackend()
+{
+}
+
+FIoStatus FOnDemandPackageStoreBackend::Mount(FString ContainerName, FIoContainerHeader&& ContainerHeader)
+{
+	{
+		UE::TUniqueLock Lock(Mutex);
+		Containers.Add(MakeUnique<FContainer>(MoveTemp(ContainerName),MoveTemp(ContainerHeader)));
+		bNeedsUpdate = true;
+	}
+
+	if (BackendContext)
+	{
+		BackendContext->PendingEntriesAdded.Broadcast();
+	}
+
+	return EIoErrorCode::Ok;
+}
+
+FIoStatus FOnDemandPackageStoreBackend::Unmount(const FString& ContainerName)
+{
+	return EIoErrorCode::Ok;
+}
+
+void FOnDemandPackageStoreBackend::OnMounted(TSharedRef<const FPackageStoreBackendContext> Context)
+{
+	BackendContext = Context;
+}
+
+void FOnDemandPackageStoreBackend::BeginRead()
+{
+	ensure(Mutex.TryLock());
+	UpdateLookupTables();
+}
+
+void FOnDemandPackageStoreBackend::EndRead()
+{
+	Mutex.Unlock();
+}
+
+EPackageStoreEntryStatus FOnDemandPackageStoreBackend::GetPackageStoreEntry(
+	FPackageId PackageId,
+	FName PackageName,
+	FPackageStoreEntry& OutPackageStoreEntry)
+{
+	if (const FFilePackageStoreEntry* Entry = EntryMap.FindRef(PackageId))
+	{
+		OutPackageStoreEntry.ImportedPackageIds =
+			MakeArrayView(Entry->ImportedPackages.Data(), Entry->ImportedPackages.Num());
+		OutPackageStoreEntry.ShaderMapHashes =
+			MakeArrayView(Entry->ShaderMapHashes.Data(), Entry->ShaderMapHashes.Num());
+
+		return EPackageStoreEntryStatus::Ok;
+	}
+
+	return EPackageStoreEntryStatus::Missing;
+}
+
+bool FOnDemandPackageStoreBackend::GetPackageRedirectInfo(
+	FPackageId PackageId,
+	FName& OutSourcePackageName,
+	FPackageId& OutRedirectedToPackageId)
+{
+	if (const FRedirect* Redirect = RedirectMap.Find(PackageId))
+	{
+		OutSourcePackageName		= Redirect->Key;
+		OutRedirectedToPackageId	= Redirect->Value;
+		return true;
+	}
+	
+	if (const FName* SourcePkgName = LocalizedMap.Find(PackageId))
+	{
+		const FName LocalizedPkgName = FPackageLocalizationManager::Get().FindLocalizedPackageName(*SourcePkgName);
+		if (LocalizedPkgName.IsNone() == false)
+		{
+			const FPackageId LocalizedPkgId = FPackageId::FromName(LocalizedPkgName);
+			if (EntryMap.Find(LocalizedPkgId))
+			{
+				OutSourcePackageName		= *SourcePkgName;
+				OutRedirectedToPackageId	= LocalizedPkgId;
+				return true;
+			}
+		}
+	}
+
+	return false;
+}
+
+void FOnDemandPackageStoreBackend::UpdateLookupTables()
+{
+	if (!bNeedsUpdate)
+	{
+		return;
+	}
+
+	bNeedsUpdate = false;
+	check(Mutex.IsLocked());
+
+	EntryMap.Empty();
+	LocalizedMap.Empty();
+	RedirectMap.Empty();
+
+	const int32 PackageCount = Algo::TransformAccumulate(
+		Containers,
+		[](const FUniqueContainer& C) { return C->ContainerHeader.PackageIds.Num(); },
+		0);
+
+	EntryMap.Reserve(PackageCount);
+
+	for (const FUniqueContainer& Container : Containers)
+	{
+		const FIoContainerHeader& Hdr = Container->ContainerHeader;
+		TConstArrayView<FFilePackageStoreEntry> Entries(
+			reinterpret_cast<const FFilePackageStoreEntry*>(Hdr.StoreEntries.GetData()),
+			Hdr.PackageIds.Num());
+
+		int32 Idx = 0;
+		for (const FFilePackageStoreEntry& Entry : Entries)
+		{
+			const FPackageId PkgId = Hdr.PackageIds[Idx++];
+			EntryMap.Add(PkgId, &Entry);
+		}
+
+		for (const FIoContainerHeaderLocalizedPackage& Localized : Hdr.LocalizedPackages)
+		{
+			FName& SourcePackageName = LocalizedMap.FindOrAdd(Localized.SourcePackageId);
+			if (SourcePackageName.IsNone())
+			{
+				FDisplayNameEntryId NameEntry = Hdr.RedirectsNameMap[Localized.SourcePackageName.GetIndex()];
+				SourcePackageName = NameEntry.ToName(Localized.SourcePackageName.GetNumber());
+			}
+		}
+
+		for (const FIoContainerHeaderPackageRedirect& Redirect : Hdr.PackageRedirects)
+		{
+			FRedirect& RedirectEntry = RedirectMap.FindOrAdd(Redirect.SourcePackageId);
+			FName& SourcePackageName = RedirectEntry.Key;
+			if (SourcePackageName.IsNone())
+			{
+				FDisplayNameEntryId NameEntry = Hdr.RedirectsNameMap[Redirect.SourcePackageName.GetIndex()];
+				SourcePackageName = NameEntry.ToName(Redirect.SourcePackageName.GetNumber());
+				RedirectEntry.Value = Redirect.TargetPackageId;
+			}
+		}
+	}
+}
+
+///////////////////////////////////////////////////////////////////////////////
+TSharedPtr<IOnDemandPackageStoreBackend> MakeOnDemandPackageStoreBackend()
+{
+	return MakeShareable<IOnDemandPackageStoreBackend>(new FOnDemandPackageStoreBackend());
+}
+
+} // namespace UE
diff --git a/Engine/Source/Runtime/Experimental/IoStoreOnDemand/Private/OnDemandPackageStoreBackend.h b/Engine/Source/Runtime/Experimental/IoStoreOnDemand/Private/OnDemandPackageStoreBackend.h
new file mode 100644
--- /dev/null
+++ b/Engine/Source/Runtime/Experimental/IoStoreOnDemand/Private/OnDemandPackageStoreBackend.h
@@ -0,0 +1,26 @@
+// Copyright Epic Games, Inc. All Rights Reserved.
+
+#pragma once
+
+#include "IO/IoStatus.h"
+#include "IO/PackageStore.h"
+#include "Templates/SharedPointer.h"
+
+struct FIoContainerHeader;
+
+namespace UE::IoStore
+{
+
+class IOnDemandPackageStoreBackend
+	: public IPackageStoreBackend
+{
+public:
+	virtual ~IOnDemandPackageStoreBackend() = default;
+
+	virtual FIoStatus Mount(FString ContainerName, FIoContainerHeader&& ContainerHeader) = 0;
+	virtual FIoStatus Unmount(const FString& ContainerName) = 0;
+};
+
+TSharedPtr<IOnDemandPackageStoreBackend> MakeOnDemandPackageStoreBackend();
+
+} // namespace UE
