[Backout] - CL32955946
Original CL Desc
-----------------------------------------------------------------
Added on-demand package store backend

diff --git a/Engine/Source/Runtime/Experimental/IoStoreOnDemand/IoStoreOnDemand.Build.cs b/Engine/Source/Runtime/Experimental/IoStoreOnDemand/IoStoreOnDemand.Build.cs
--- a/Engine/Source/Runtime/Experimental/IoStoreOnDemand/IoStoreOnDemand.Build.cs
+++ b/Engine/Source/Runtime/Experimental/IoStoreOnDemand/IoStoreOnDemand.Build.cs
@@ -10,7 +10,6 @@ public class IoStoreOnDemand : ModuleRules
 		PublicDependencyModuleNames.Add("TraceLog");
 		PrivateDependencyModuleNames.AddRange(
 			new string[] {
-				"CoreUObject",
 				"HTTP",
 				"Json",
 				"Analytics"
diff --git a/Engine/Source/Runtime/Experimental/IoStoreOnDemand/Private/OnDemandPackageStoreBackend.cpp b/Engine/Source/Runtime/Experimental/IoStoreOnDemand/Private/OnDemandPackageStoreBackend.cpp
deleted file mode 100644
diff --git a/Engine/Source/Runtime/Experimental/IoStoreOnDemand/Private/OnDemandPackageStoreBackend.h b/Engine/Source/Runtime/Experimental/IoStoreOnDemand/Private/OnDemandPackageStoreBackend.h
deleted file mode 100644
