Updated debug proxying to use outcome objects

#rnx

diff --git a/Engine/Source/Runtime/Experimental/IoStoreOnDemand/Private/CoreHttp/Http.cpp b/Engine/Source/Runtime/Experimental/IoStoreOnDemand/Private/CoreHttp/Http.cpp
--- a/Engine/Source/Runtime/Experimental/IoStoreOnDemand/Private/CoreHttp/Http.cpp
+++ b/Engine/Source/Runtime/Experimental/IoStoreOnDemand/Private/CoreHttp/Http.cpp
@@ -1107,7 +1107,7 @@ static uint32 GetSocksIpAddress()
 }
 
 ////////////////////////////////////////////////////////////////////////////////
-static int32 ConnectSocks4(FSocket& Socket, uint32 IpAddress, uint32 Port)
+static FOutcome ConnectSocks4(FSocket& Socket, uint32 IpAddress, uint32 Port)
 {
 	struct FSocks4Request
 	{
@@ -1128,14 +1128,14 @@ static int32 ConnectSocks4(FSocket& Socket, uint32 IpAddress, uint32 Port)
 	uint32 SocksIpAddress = GetSocksIpAddress();
 	if (!SocksIpAddress)
 	{
-		return -1;
+		return FOutcome::Error("Invalid socks IP address");
 	}
 
 	FOutcome Outcome = FOutcome::None();
 
 	if (Outcome = Socket.Connect(SocksIpAddress, GSocksPort); Outcome.IsError())
 	{
-		return -1;
+		return Outcome;
 	}
 
 	FSocks4Request Request = {
@@ -1144,20 +1144,20 @@ static int32 ConnectSocks4(FSocket& Socket, uint32 IpAddress, uint32 Port)
 	};
 	if (Outcome = Socket.Send((const char*)&Request, sizeof(Request)); Outcome.IsError())
 	{
-		return -1;
+		return Outcome;
 	}
 
 	FSocks4Reply Reply;
 	if (Outcome = Socket.Recv((char*)&Reply, sizeof(Reply)); Outcome.IsError())
 	{
-		return -1;
+		return Outcome;
 	}
 
-	return 1;
+	return FOutcome::Ok(1);
 }
 
 ////////////////////////////////////////////////////////////////////////////////
-static int32 ConnectSocks5(FSocket& Socket, uint32 IpAddress, uint32 Port)
+static FOutcome ConnectSocks5(FSocket& Socket, uint32 IpAddress, uint32 Port)
 {
 #ifdef _MSC_VER
 	// MSVC's static analysis doesn't see that 'Result' from recv() is checked
@@ -1169,36 +1169,28 @@ static int32 ConnectSocks5(FSocket& Socket, uint32 IpAddress, uint32 Port)
 	uint32 SocksIpAddress = GetSocksIpAddress();
 	if (!SocksIpAddress)
 	{
-		return -1;
+		return FOutcome::Error("Invalid socks5 IP address");
 	}
 
 	FOutcome Outcome = FOutcome::None();
 
 	if (Outcome = Socket.Connect(SocksIpAddress, GSocksPort); Outcome.IsError())
 	{
-		return -1;
+		return Outcome;
 	}
 
 	// Greeting
 	const char Greeting[] = { 5, 1, 0 };
 	Outcome = Socket.Send(Greeting, sizeof(Greeting));
-	if (!Outcome.IsOk() || Outcome.GetResult() != sizeof(Greeting))
-	{
-		return -1;
-	}
+	if (!Outcome.IsOk())						return Outcome;
+	if (Outcome.GetResult() != sizeof(Greeting))return FOutcome::Error("Could not send socks5 greeting");
 
 	// Server auth-choice
-	char ServerChoice[1 + 1];
-	Outcome = Socket.Recv(ServerChoice, sizeof(ServerChoice));
-	if (!Outcome.IsOk() || Outcome.GetResult() != sizeof(ServerChoice))
-	{
-		return -1;
-	}
-
-	if (ServerChoice[0] != 0x05 || ServerChoice[1] != 0x00)
-	{
-		return -1;
-	}
+	char Choice[1 + 1];
+	Outcome = Socket.Recv(Choice, sizeof(Choice));
+	if (!Outcome.IsOk())						return Outcome;
+	if (Outcome.GetResult() != sizeof(Choice))	return FOutcome::Error("Recv too short from socks5 server");
+	if (Choice[0] != 0x05 || Choice[1] != 0x00) return FOutcome::Error("Got unexpected socks5 version from server");
 
 	// Connection request
 	IpAddress = htonl(IpAddress);
@@ -1207,25 +1199,17 @@ static int32 ConnectSocks5(FSocket& Socket, uint32 IpAddress, uint32 Port)
 	std::memcpy(Request + 4, &IpAddress, sizeof(IpAddress));
 	std::memcpy(Request + 8, &NsPort, sizeof(NsPort));
 	Outcome = Socket.Send(Request, sizeof(Request));
-	if (!Outcome.IsOk() || Outcome.GetResult() != sizeof(Request))
-	{
-		return -1;
-	}
+	if (!Outcome.IsOk())						return Outcome;
+	if (Outcome.GetResult() != sizeof(Request)) return FOutcome::Error("Sent too little to socks5 server");
 
 	// Connect reply
 	char Reply[3 + (1 + 4) + 2];
 	Outcome = Socket.Recv(Reply, sizeof(Reply));
-	if (!Outcome.IsOk() || Outcome.GetResult() != sizeof(Reply))
-	{
-		return -1;
-	}
+	if (!Outcome.IsOk())						return Outcome;
+	if (Outcome.GetResult() != sizeof(Reply))	return FOutcome::Error("Socks5 reply too short");
+	if (Reply[0] != 0x05 || Reply[1] != 0x00)	return FOutcome::Error("Reply has unexpected socks5 version");
 
-	if (Reply[0] != 0x05 || Reply[1] != 0x00)
-	{
-		return -1;
-	}
-
-	return 1;
+	return FOutcome::Ok(1);
 
 #ifdef _MSC_VER
 #pragma warning(pop)
@@ -1235,14 +1219,14 @@ static int32 ConnectSocks5(FSocket& Socket, uint32 IpAddress, uint32 Port)
 #endif // UE_BUILD_SHIPPING
 
 ////////////////////////////////////////////////////////////////////////////////
-static int32 MaybeConnectSocks(FSocket& Socket, uint32 IpAddress, uint32 Port)
+static FOutcome MaybeConnectSocks(FSocket& Socket, uint32 IpAddress, uint32 Port)
 {
 #if UE_BUILD_SHIPPING
-	return 0;
+	return FOutcome::Ok();
 #else
 	if (GSocksIp.IsEmpty())
 	{
-		return 0;
+		return FOutcome::Ok();
 	}
 
 	switch (GSocksVersion)
@@ -1251,7 +1235,7 @@ static int32 MaybeConnectSocks(FSocket& Socket, uint32 IpAddress, uint32 Port)
 	case 5: return ConnectSocks5(Socket, IpAddress, Port);
 	}
 
-	return -1;
+	return FOutcome::Error("Unsupported socks version");
 #endif // UE_BUILD_SHIPPING
 }
 
@@ -1392,9 +1376,11 @@ int32 FHost::IsResolved() const
 ////////////////////////////////////////////////////////////////////////////////
 FOutcome FHost::Connect(FSocket& Socket)
 {
+	FOutcome Outcome = FOutcome::None();
+
 	if (IsResolved() <= 0)
 	{
-		if (FOutcome Outcome = ResolveHostName(); Outcome.IsError())
+		if (Outcome = ResolveHostName(); Outcome.IsError())
 		{
 			return Outcome;
 		}
@@ -1412,16 +1398,13 @@ FOutcome FHost::Connect(FSocket& Socket)
 	}
 
 	// Attempt a SOCKS connect
-	bool bSocksConnected = false;
-	if (int32 Result = MaybeConnectSocks(Candidate, IpAddress, Port); Result)
+	Outcome = MaybeConnectSocks(Candidate, IpAddress, Port);
+	if (Outcome.IsError())
 	{
-		if (Result < 0)
-		{
-			return FOutcome::Error("Failed establishing SOCKS connection");
-		}
-
-		bSocksConnected = true;
+		return Outcome;
 	}
+	check(Outcome.IsOk());
+	bool bSocksConnected = (Outcome.GetResult() == 1);
 
 	// Condition the socket
 	if (!Candidate.SetBlocking(false))
@@ -1447,7 +1430,7 @@ FOutcome FHost::Connect(FSocket& Socket)
 	}
 
 	// Issue the connect - this is done non-blocking so we need to wait (ret=0)
-	if (FOutcome Outcome = Candidate.Connect(IpAddress, Port); Outcome.IsError())
+	if (Outcome = Candidate.Connect(IpAddress, Port); Outcome.IsError())
 	{
 		return Outcome;
 	}
