Use outcome objects for Do*() functions

#rnx

diff --git a/Engine/Source/Runtime/Experimental/IoStoreOnDemand/Private/CoreHttp/Http.cpp b/Engine/Source/Runtime/Experimental/IoStoreOnDemand/Private/CoreHttp/Http.cpp
--- a/Engine/Source/Runtime/Experimental/IoStoreOnDemand/Private/CoreHttp/Http.cpp
+++ b/Engine/Source/Runtime/Experimental/IoStoreOnDemand/Private/CoreHttp/Http.cpp
@@ -2005,7 +2005,7 @@ FTicketPerf::FSample FTicketPerf::GetSample() const
 // {{{1 event-loop-int .........................................................
 
 ////////////////////////////////////////////////////////////////////////////////
-static int32 DoSend(FActivity* Activity, FSocket& Socket)
+static FOutcome DoSend(FActivity* Activity, FSocket& Socket)
 {
 	Trace(Activity, ETrace::StateChange, Activity->State);
 
@@ -2029,12 +2029,12 @@ static int32 DoSend(FActivity* Activity, FSocket& Socket)
 	if (Outcome.IsError())
 	{
 		Activity_SetError(Activity, Outcome.GetMessage().GetData());
-		return -1;
+		return Outcome;
 	}
 
 	if (Outcome.IsWaiting())
 	{
-		return 0;
+		return Outcome;
 	}
 
 	check(Outcome.IsOk());
@@ -2051,11 +2051,11 @@ static int32 DoSend(FActivity* Activity, FSocket& Socket)
 #endif
 
 	Activity_ChangeState(Activity, FActivity::EState::RecvMessage, Buffer.GetSize());
-	return Result;
+	return FOutcome::Ok(Result);
 }
 
 ////////////////////////////////////////////////////////////////////////////////
-static int32 DoRecvMessage(FActivity* Activity, FSocket& Socket)
+static FOutcome DoRecvMessage(FActivity* Activity, FSocket& Socket)
 {
 	TRACE_CPUPROFILER_EVENT_SCOPE(IasHttp::DoRecvMessage);
 
@@ -2078,12 +2078,12 @@ static int32 DoRecvMessage(FActivity* Activity, FSocket& Socket)
 		if (Outcome.IsError())
 		{
 			Activity_SetError(Activity, Outcome.GetMessage().GetData());
-			return -1;
+			return Outcome;
 		}
 
 		if (Outcome.IsWaiting())
 		{
-			return 1;
+			return Outcome;
 		}
 
 		check(Outcome.IsOk());
@@ -2105,7 +2105,7 @@ static int32 DoRecvMessage(FActivity* Activity, FSocket& Socket)
 			if (Buffer.GetSize() > (8 << 10))
 			{
 				Activity_SetError(Activity, "Headers have grown larger than expected");
-				return -1;
+				return FOutcome::Error(Activity->ErrorReason);
 			}
 
 			continue;
@@ -2124,7 +2124,7 @@ static int32 DoRecvMessage(FActivity* Activity, FSocket& Socket)
 	if (ParseMessage(ResponseView, Internal.Offsets) < 0)
 	{
 		Activity_SetError(Activity, "Failed to parse message status");
-		return -1;
+		return FOutcome::Error(Activity->ErrorReason);
 	}
 
 	// Parse headers
@@ -2171,7 +2171,7 @@ static int32 DoRecvMessage(FActivity* Activity, FSocket& Socket)
 		}
 
 		Activity_SetError(Activity, "Unknown content length value");
-		return -1;
+		return FOutcome::Error(Activity->ErrorReason);
 	}
 
 	// Call out to the sink to get a content destination
@@ -2188,7 +2188,7 @@ static int32 DoRecvMessage(FActivity* Activity, FSocket& Socket)
 		if (Activity->Dest == PriorDest)
 		{
 			Activity_SetError(Activity, "User did not provide a destination buffer");
-			return -1;
+			return FOutcome::Error(Activity->ErrorReason);
 		}
 
 		// The user seems to have forgotten something. Let's help them along
@@ -2204,7 +2204,7 @@ static int32 DoRecvMessage(FActivity* Activity, FSocket& Socket)
 	if (AlreadyReceived > uint32(ContentLength))
 	{
 		Activity_SetError(Activity, "More data received that expected");
-		return -1;
+		return FOutcome::Error(Activity->ErrorReason);
 	}
 
 	if (Activity->NoContent == 1)
@@ -2212,10 +2212,10 @@ static int32 DoRecvMessage(FActivity* Activity, FSocket& Socket)
 		if (AlreadyReceived)
 		{
 			Activity_SetError(Activity, "Received content when none was expected");
-			return -1;
+			return FOutcome::Error(Activity->ErrorReason);
 		}
 		Activity_ChangeState(Activity, FActivity::EState::RecvDone);
-		return 0;
+		return FOutcome::Ok();
 	}
 
 	check(Activity->Dest != nullptr);
@@ -2227,7 +2227,7 @@ static int32 DoRecvMessage(FActivity* Activity, FSocket& Socket)
 
 	if (AlreadyReceived == 0)
 	{
-		return 0;
+		return FOutcome::Ok();
 	}
 
 	FMutableMemoryView DestView = Activity->Dest->GetMutableView();
@@ -2235,7 +2235,7 @@ static int32 DoRecvMessage(FActivity* Activity, FSocket& Socket)
 	if (!bStreamed || AlreadyReceived < DestView.GetSize())
 	{
 		::memcpy(DestView.GetData(), Cursor, AlreadyReceived);
-		return 0;
+		return FOutcome::Ok();
 	}
 
 #if 1
@@ -2254,11 +2254,11 @@ static int32 DoRecvMessage(FActivity* Activity, FSocket& Socket)
 	while (AlreadyReceived);
 #endif
 
-	return 0;
+	return FOutcome::Ok();
 }
 
 ////////////////////////////////////////////////////////////////////////////////
-static int32 DoRecvContent(FActivity* Activity, FSocket& Socket, int32& MaxRecvSize)
+static FOutcome DoRecvContent(FActivity* Activity, FSocket& Socket, int32& MaxRecvSize)
 {
 	TRACE_CPUPROFILER_EVENT_SCOPE(IasHttp::DoRecvContent);
 
@@ -2277,7 +2277,7 @@ static int32 DoRecvContent(FActivity* Activity, FSocket& Socket, int32& MaxRecvS
 		check(Size >= 0);
 		if (Size == 0)
 		{
-			return 1;
+			return FOutcome::Waiting();
 		}
 
 		Trace(Activity, ETrace::StateChange, Activity->State);
@@ -2288,13 +2288,13 @@ static int32 DoRecvContent(FActivity* Activity, FSocket& Socket, int32& MaxRecvS
 
 		if (Outcome.IsWaiting())
 		{
-			return 1;
+			return Outcome;
 		}
 
 		if (Outcome.IsError())
 		{
 			Activity_SetError(Activity, Outcome.GetMessage().GetData());
-			return -1;
+			return Outcome;
 		}
 
 		check(Outcome.IsOk());
@@ -2309,7 +2309,7 @@ static int32 DoRecvContent(FActivity* Activity, FSocket& Socket, int32& MaxRecvS
 	if (!FLatencyInjector::Begin(FLatencyInjector::EType::Network, Activity->StateParam))
 	{
 		Activity_SetError(Activity, "Forced random failure");
-		return -1;
+		return FOutcome::Error(Activity->ErrorReason);
 	}
 
 #if IAS_HTTP_WITH_PERF
@@ -2317,18 +2317,18 @@ static int32 DoRecvContent(FActivity* Activity, FSocket& Socket, int32& MaxRecvS
 #endif
 
 	Activity_ChangeState(Activity, FActivity::EState::RecvDone);
-	return 0;
+	return FOutcome::Ok();
 }
 
 ////////////////////////////////////////////////////////////////////////////////
-static int32 DoRecvStream(FActivity*, FSocket&, uint32)
+static FOutcome DoRecvStream(FActivity*, FSocket&, uint32)
 {
 	check(false); // not yet implemented
-	return -1;
+	return FOutcome::Error("Not supported");
 }
 
 ////////////////////////////////////////////////////////////////////////////////
-static int32 DoRecv(FActivity* Activity, FSocket& Socket, int32& MaxRecvSize)
+static FOutcome DoRecv(FActivity* Activity, FSocket& Socket, int32& MaxRecvSize)
 {
 	using EState = FActivity::EState;
 
@@ -2340,7 +2340,7 @@ static int32 DoRecv(FActivity* Activity, FSocket& Socket, int32& MaxRecvSize)
 	if (State == EState::RecvStream)	return DoRecvStream(Activity, Socket, MaxRecvSize);
 	
 	check(false); // it is not expected that we'll get here
-	return -1;
+	return FOutcome::Error("unreachable");
 }
 
 ////////////////////////////////////////////////////////////////////////////////
@@ -2727,12 +2727,12 @@ void FSocketGroup::RecvInternal(FTickState& State)
 	FActivity* Activity = Recv;
 	check(IsReceiving(Activity));
 
-	int32 Result = DoRecv(Activity, Socket, State.RecvAllowance);
+	FOutcome Outcome = DoRecv(Activity, Socket, State.RecvAllowance);
 
 	// Any sort of error here is unrecoverable
-	if (Result < 0)
+	if (Outcome.IsError())
 	{
-		Fail(State, Activity->ErrorReason);
+		Fail(State, Outcome.GetMessage().GetData());
 		return;
 	}
 
@@ -2762,7 +2762,7 @@ void FSocketGroup::RecvInternal(FTickState& State)
 	}
 
 	// If there was no data available this is far as receiving can go
-	if (bWaiting = (Result > 0); bWaiting)
+	if (bWaiting = Outcome.IsWaiting(); bWaiting)
 	{
 		return;
 	}
@@ -2803,18 +2803,18 @@ void FSocketGroup::SendInternal(FTickState& State)
 
 	FActivity* Activity = Send;
 
-	int32 Result = DoSend(Activity, Socket);
+	FOutcome Outcome = DoSend(Activity, Socket);
 
-	if (Result == 0)
+	if (Outcome.IsWaiting())
 	{
 		// For now we'll not add the socket as a waiter. It is unlikely that we
 		// send enough to need to wait currently.
 		return;
 	}
 
-	if (Result < 0)
+	if (Outcome.IsError())
 	{
-		Fail(State, Activity->ErrorReason);
+		Fail(State, Outcome.GetMessage().GetData());
 		return;
 	}
 
