Added missing header

diff --git a/Engine/Source/Runtime/Experimental/IoStoreOnDemand/Private/OnDemandCacheIoDispatcherBackend.h b/Engine/Source/Runtime/Experimental/IoStoreOnDemand/Private/OnDemandCacheIoDispatcherBackend.h
new file mode 100644
--- /dev/null
+++ b/Engine/Source/Runtime/Experimental/IoStoreOnDemand/Private/OnDemandCacheIoDispatcherBackend.h
@@ -0,0 +1,23 @@
+// Copyright Epic Games, Inc. All Rights Reserved.
+
+#pragma once
+
+#include "IO/IoDispatcherBackend.h"
+#include "Templates/SharedPointer.h"
+
+namespace UE::IoStore
+{
+
+class IIoCache;
+class FOnDemandIoStore;
+
+class IOnDemandCacheIoDispatcherBackend
+	: public IIoDispatcherBackend
+{
+public:
+	virtual ~IOnDemandCacheIoDispatcherBackend() = default;
+};
+
+TSharedPtr<IOnDemandCacheIoDispatcherBackend> MakeOnDemandCacheIoDispatcherBackend(FOnDemandIoStore& IoStore, IIoCache& Cache);
+
+} // namespace UE::IoStore
