Added missing package store header

diff --git a/Engine/Source/Runtime/Experimental/IoStoreOnDemand/Private/OnDemandPackageStoreBackend.h b/Engine/Source/Runtime/Experimental/IoStoreOnDemand/Private/OnDemandPackageStoreBackend.h
new file mode 100644
--- /dev/null
+++ b/Engine/Source/Runtime/Experimental/IoStoreOnDemand/Private/OnDemandPackageStoreBackend.h
@@ -0,0 +1,26 @@
+// Copyright Epic Games, Inc. All Rights Reserved.
+
+#pragma once
+
+#include "IO/IoStatus.h"
+#include "IO/PackageStore.h"
+#include "Templates/SharedPointer.h"
+
+struct FIoContainerHeader;
+
+namespace UE::IoStore
+{
+
+class IOnDemandPackageStoreBackend
+	: public IPackageStoreBackend
+{
+public:
+	virtual ~IOnDemandPackageStoreBackend() = default;
+
+	virtual FIoStatus Mount(FString ContainerName, FIoContainerHeader&& ContainerHeader) = 0;
+	virtual FIoStatus Unmount(const FString& ContainerName) = 0;
+};
+
+TSharedPtr<IOnDemandPackageStoreBackend> MakeOnDemandPackageStoreBackend();
+
+} // namespace UE
