Start to distance ourselves from the underlying transport.

#rnx

diff --git a/Engine/Source/Runtime/Experimental/IoStoreOnDemand/Private/CoreHttp/Http.cpp b/Engine/Source/Runtime/Experimental/IoStoreOnDemand/Private/CoreHttp/Http.cpp
--- a/Engine/Source/Runtime/Experimental/IoStoreOnDemand/Private/CoreHttp/Http.cpp
+++ b/Engine/Source/Runtime/Experimental/IoStoreOnDemand/Private/CoreHttp/Http.cpp
@@ -2004,8 +2004,10 @@ FTicketPerf::FSample FTicketPerf::GetSample() const
 
 // {{{1 event-loop-int .........................................................
 
+using FPeerType = FSocket;
+
 ////////////////////////////////////////////////////////////////////////////////
-static FOutcome DoSend(FActivity* Activity, FSocket& Socket)
+static FOutcome DoSend(FActivity* Activity, FPeerType& Peer)
 {
 	Trace(Activity, ETrace::StateChange, Activity->State);
 
@@ -2024,7 +2026,7 @@ static FOutcome DoSend(FActivity* Activity, FSocket& Socket)
 	SendSize -= AlreadySent;
 	check(SendSize > 0);
 
-	FOutcome Outcome = Socket.Send(SendData, SendSize);
+	FOutcome Outcome = Peer.Send(SendData, SendSize);
 
 	if (Outcome.IsError())
 	{
@@ -2043,7 +2045,7 @@ static FOutcome DoSend(FActivity* Activity, FSocket& Socket)
 	Activity->StateParam += Result;
 	if (Activity->StateParam < Buffer.GetSize())
 	{
-		return DoSend(Activity, Socket);
+		return DoSend(Activity, Peer);
 	}
 
 #if IAS_HTTP_WITH_PERF
@@ -2055,7 +2057,7 @@ static FOutcome DoSend(FActivity* Activity, FSocket& Socket)
 }
 
 ////////////////////////////////////////////////////////////////////////////////
-static FOutcome DoRecvMessage(FActivity* Activity, FSocket& Socket)
+static FOutcome DoRecvMessage(FActivity* Activity, FPeerType& Peer)
 {
 	TRACE_CPUPROFILER_EVENT_SCOPE(IasHttp::DoRecvMessage);
 
@@ -2073,7 +2075,7 @@ static FOutcome DoRecvMessage(FActivity* Activity, FSocket& Socket)
 #endif
 
 		auto [Dest, DestSize] = Buffer.GetMutableFree(0, PageSize);
-		FOutcome Outcome = Socket.Recv(Dest, DestSize);
+		FOutcome Outcome = Peer.Recv(Dest, DestSize);
 
 		if (Outcome.IsError())
 		{
@@ -2258,7 +2260,7 @@ static FOutcome DoRecvMessage(FActivity* Activity, FSocket& Socket)
 }
 
 ////////////////////////////////////////////////////////////////////////////////
-static FOutcome DoRecvContent(FActivity* Activity, FSocket& Socket, int32& MaxRecvSize)
+static FOutcome DoRecvContent(FActivity* Activity, FPeerType& Peer, int32& MaxRecvSize)
 {
 	TRACE_CPUPROFILER_EVENT_SCOPE(IasHttp::DoRecvContent);
 
@@ -2284,7 +2286,7 @@ static FOutcome DoRecvContent(FActivity* Activity, FSocket& Socket, int32& MaxRe
 
 		char* Cursor = (char*)(DestView.GetData()) + Activity->StateParam;
 
-		FOutcome Outcome = Socket.Recv(Cursor, Size);
+		FOutcome Outcome = Peer.Recv(Cursor, Size);
 
 		if (Outcome.IsWaiting())
 		{
@@ -2321,23 +2323,23 @@ static FOutcome DoRecvContent(FActivity* Activity, FSocket& Socket, int32& MaxRe
 }
 
 ////////////////////////////////////////////////////////////////////////////////
-static FOutcome DoRecvStream(FActivity*, FSocket&, uint32)
+static FOutcome DoRecvStream(FActivity*, FPeerType&, uint32)
 {
 	check(false); // not yet implemented
 	return FOutcome::Error("Not supported");
 }
 
 ////////////////////////////////////////////////////////////////////////////////
-static FOutcome DoRecv(FActivity* Activity, FSocket& Socket, int32& MaxRecvSize)
+static FOutcome DoRecv(FActivity* Activity, FPeerType& Peer, int32& MaxRecvSize)
 {
 	using EState = FActivity::EState;
 
 	EState State = Activity->State; 
 	check(State >= EState::RecvMessage && State < EState::RecvDone);
 
-	if (State == EState::RecvMessage)	return DoRecvMessage(Activity, Socket);
-	if (State == EState::RecvContent)	return DoRecvContent(Activity, Socket, MaxRecvSize);
-	if (State == EState::RecvStream)	return DoRecvStream(Activity, Socket, MaxRecvSize);
+	if (State == EState::RecvMessage)	return DoRecvMessage(Activity, Peer);
+	if (State == EState::RecvContent)	return DoRecvContent(Activity, Peer, MaxRecvSize);
+	if (State == EState::RecvStream)	return DoRecvStream(Activity, Peer, MaxRecvSize);
 	
 	check(false); // it is not expected that we'll get here
 	return FOutcome::Error("unreachable");
@@ -2515,7 +2517,7 @@ int32 FThrottler::GetWaitEstimateMs() const
  *  Loop:
  *    FHostGroup[HostPtr]:
  *	    Work: Act0 -> Act1 -> Act2 -> Act3 -> ...
- *      FSocketGroup[0...HostMaxConnections]:
+ *      FPeerGroup[0...HostMaxConnections]:
  *			Act.Send
  *			Act.Recv
  */
@@ -2635,14 +2637,14 @@ void FWorkQueue::TickCancels(FTickState& State)
 
 
 ////////////////////////////////////////////////////////////////////////////////
-class FSocketGroup
+class FPeerGroup
 {
 public:
-						FSocketGroup() = default;
-						~FSocketGroup();
-	bool				operator == (FSocket* Rhs) const { return &Socket == Rhs; }
+						FPeerGroup() = default;
+						~FPeerGroup();
+	bool				operator == (FPeerType* Rhs) const { return &Peer == Rhs; }
 	void				Unwait()			{ check(bWaiting); bWaiting = false; }
-	FSocket::FWaiter	GetWaiter() const;
+	FPeerType::FWaiter	GetWaiter() const;
 	bool 				Tick(FTickState& State);
 	void				TickSend(FTickState& State, FHost& Host);
 	void				Fail(FTickState& State, const char* Reason);
@@ -2652,36 +2654,36 @@ private:
 	void				SendInternal(FTickState& State);
 	FActivity*			Send = nullptr;
 	FActivity*			Recv = nullptr;
-	FSocket				Socket;
+	FPeerType			Peer;
 	uint32				LastUseMs = 0;
 	uint8				IsKeepAlive = 0;
 	bool				bWaiting = false;
 
-	UE_NONCOPYABLE(FSocketGroup);
+	UE_NONCOPYABLE(FPeerGroup);
 };
 
 ////////////////////////////////////////////////////////////////////////////////
-FSocketGroup::~FSocketGroup()
+FPeerGroup::~FPeerGroup()
 {
 	check(Send == nullptr);
 	check(Recv == nullptr);
 }
 
 ////////////////////////////////////////////////////////////////////////////////
-FSocket::FWaiter FSocketGroup::GetWaiter() const
+FPeerType::FWaiter FPeerGroup::GetWaiter() const
 {
 	if (!bWaiting)
 	{
-		return FSocket::FWaiter();
+		return FPeerType::FWaiter();
 	}
 
-	using EWhat = FSocket::FWaiter::EWhat;
+	using EWhat = FPeerType::FWaiter::EWhat;
 	EWhat What = (Recv != nullptr) ? EWhat::Recv : EWhat::Send;
-	return FSocket::FWaiter(Socket, What);
+	return FPeerType::FWaiter(Peer, What);
 }
 
 ////////////////////////////////////////////////////////////////////////////////
-void FSocketGroup::Fail(FTickState& State, const char* Reason)
+void FPeerGroup::Fail(FTickState& State, const char* Reason)
 {
 	// Any send left at this point is unrecoverable
 	if (Send != nullptr)
@@ -2706,14 +2708,14 @@ void FSocketGroup::Fail(FTickState& State, const char* Reason)
 		Activity = Next;
 	}
 
-	Socket = FSocket();
+	Peer = FPeerType();
 	Send = Recv = nullptr;
 	bWaiting = false;
 	IsKeepAlive = 0;
 }
 
 ////////////////////////////////////////////////////////////////////////////////
-void FSocketGroup::RecvInternal(FTickState& State)
+void FPeerGroup::RecvInternal(FTickState& State)
 {
 	check(Recv != nullptr);
 
@@ -2727,7 +2729,7 @@ void FSocketGroup::RecvInternal(FTickState& State)
 	FActivity* Activity = Recv;
 	check(IsReceiving(Activity));
 
-	FOutcome Outcome = DoRecv(Activity, Socket, State.RecvAllowance);
+	FOutcome Outcome = DoRecv(Activity, Peer, State.RecvAllowance);
 
 	// Any sort of error here is unrecoverable
 	if (Outcome.IsError())
@@ -2753,7 +2755,7 @@ void FSocketGroup::RecvInternal(FTickState& State)
 				Send = Next;
 				SendInternal(State);
 
-				if (!Socket.IsValid())
+				if (!Peer.IsValid())
 				{
 					return;
 				}
@@ -2792,18 +2794,18 @@ void FSocketGroup::RecvInternal(FTickState& State)
 		return;
 	}
 
-	Socket = FSocket();
+	Peer = FPeerType();
 }
 
 ////////////////////////////////////////////////////////////////////////////////
-void FSocketGroup::SendInternal(FTickState& State)
+void FPeerGroup::SendInternal(FTickState& State)
 {
 	check(IsKeepAlive == 1);
 	check(Send != nullptr);
 
 	FActivity* Activity = Send;
 
-	FOutcome Outcome = DoSend(Activity, Socket);
+	FOutcome Outcome = DoSend(Activity, Peer);
 
 	if (Outcome.IsWaiting())
 	{
@@ -2832,7 +2834,7 @@ void FSocketGroup::SendInternal(FTickState& State)
 }
 
 ////////////////////////////////////////////////////////////////////////////////
-bool FSocketGroup::Tick(FTickState& State)
+bool FPeerGroup::Tick(FTickState& State)
 {
 	if (Send != nullptr)
 	{
@@ -2848,7 +2850,7 @@ bool FSocketGroup::Tick(FTickState& State)
 }
 
 ////////////////////////////////////////////////////////////////////////////////
-void FSocketGroup::TickSend(FTickState& State, FHost& Host)
+void FPeerGroup::TickSend(FTickState& State, FHost& Host)
 {
 	// This path is only for those that are idle and have nothing to do
 	if (Send != nullptr || Recv != nullptr)
@@ -2861,19 +2863,19 @@ void FSocketGroup::TickSend(FTickState& State, FHost& Host)
 	check(Pending != nullptr);
 
 	// Close idle sockets
-	if (Socket.IsValid() && LastUseMs + GIdleMs < State.NowMs)
+	if (Peer.IsValid() && LastUseMs + GIdleMs < State.NowMs)
 	{
 		LastUseMs = State.NowMs;
-		Socket = FSocket();
+		Peer = FPeerType();
 	}
 
 	// We don't have a connected socket on first use, or if a keep-alive:close
 	// was received from the server. So we connect here.
 	bool bWillBlock = false;
-	if (!Socket.IsValid())
+	if (!Peer.IsValid())
 	{
 		IsKeepAlive = 1;
-		FOutcome Outcome = Host.Connect(Socket);
+		FOutcome Outcome = Host.Connect(Peer);
 
 		// We failed to connect, let's bail.
 		if (Outcome.IsError())
@@ -2912,7 +2914,7 @@ public:
 
 private:
 	int32						Wait(const FTickState& State);
-	TArray<FSocketGroup>		SocketGroups;
+	TArray<FPeerGroup>			PeerGroups;
 	FWorkQueue					Work;
 	FHost&						Host;
 	uint32						BusyCount = 0;
@@ -2926,17 +2928,17 @@ FHostGroup::FHostGroup(FHost& InHost)
 : Host(InHost)
 {
 	uint32 Num = InHost.GetMaxConnections();
-	SocketGroups.SetNum(Num);
+	PeerGroups.SetNum(Num);
 }
 
 ////////////////////////////////////////////////////////////////////////////////
 int32 FHostGroup::Wait(const FTickState& State)
 {
 	// Collect groups that are waiting on something
-	TArray<FSocket::FWaiter, TFixedAllocator<64>> Waiters;
-	for (FSocketGroup& Group : SocketGroups)
+	TArray<FPeerType::FWaiter, TFixedAllocator<64>> Waiters;
+	for (FPeerGroup& Group : PeerGroups)
 	{
-		FSocket::FWaiter Waiter = Group.GetWaiter();
+		FPeerType::FWaiter Waiter = Group.GetWaiter();
 		if (Waiter.IsValid())
 		{
 			Waiters.Add(Waiter);
@@ -2960,7 +2962,7 @@ int32 FHostGroup::Wait(const FTickState& State)
 	}
 
 	// Actually do the wait
-	int32 Result = FSocket::Wait(Waiters, PollTimeoutMs);
+	int32 Result = FPeerType::Wait(Waiters, PollTimeoutMs);
 	if (Result <= 0)
 	{
 		// If the user opts to not block then we don't accumulate wait time and
@@ -2986,9 +2988,9 @@ int32 FHostGroup::Wait(const FTickState& State)
 			continue;
 		}
 
-		auto* Candidate = (FSocket*)(Waiters[i].Candidate);
+		auto* Candidate = (FPeerType*)(Waiters[i].Candidate);
 		auto Pred = [Candidate] (auto& Lhs) { return Lhs == Candidate; };
-		FSocketGroup* Group = SocketGroups.FindByPredicate(Pred);
+		FPeerGroup* Group = PeerGroups.FindByPredicate(Pred);
 		check(Group != nullptr);
 		Group->Unwait();
 
@@ -3010,7 +3012,7 @@ void FHostGroup::Tick(FTickState& State)
 		Work.TickCancels(State);
 
 		// Get available work out on idle sockets as soon as possible
-		for (FSocketGroup& Group : SocketGroups)
+		for (FPeerGroup& Group : PeerGroups)
 		{
 			if (!Work.HasWork())
 			{
@@ -3028,7 +3030,7 @@ void FHostGroup::Tick(FTickState& State)
 			? "FailTimeout hit"
 			: "poll() returned an unexpected error";
 
-		for (FSocketGroup& Group : SocketGroups)
+		for (FPeerGroup& Group : PeerGroups)
 		{
 			Group.Fail(State, Reason);
 		}
@@ -3037,7 +3039,7 @@ void FHostGroup::Tick(FTickState& State)
 	}
 
 	// Tick everything, starting with groups that are maybe closest to finishing
-	for (FSocketGroup& Group : SocketGroups)
+	for (FPeerGroup& Group : PeerGroups)
 	{
 		BusyCount += (Group.Tick(State) == true);
 	}
