Implemented dirty joint tracking to avoid syncing back every joint on each physics tick.

Added a bit array tracking which constraints in the internal constraint list should be present in the constraint graph. If they're in the graph they will be solved and we will have results to marshall back to the game thread.

diff --git a/Engine/Source/Runtime/Experimental/Chaos/Private/Chaos/PBDJointConstraints.cpp b/Engine/Source/Runtime/Experimental/Chaos/Private/Chaos/PBDJointConstraints.cpp
--- a/Engine/Source/Runtime/Experimental/Chaos/Private/Chaos/PBDJointConstraints.cpp
+++ b/Engine/Source/Runtime/Experimental/Chaos/Private/Chaos/PBDJointConstraints.cpp
@@ -695,10 +695,13 @@ namespace Chaos
 		TArray<FPBDJointSettings> SortedConstraintSettings;
 		TArray<FParticlePair> SortedConstraintParticles;
 		TArray<FPBDJointState> SortedConstraintStates;
-		SortedHandles.Reserve(SortedIndices.Num());
-		SortedConstraintSettings.Reserve(SortedIndices.Num());
-		SortedConstraintParticles.Reserve(SortedIndices.Num());
-		SortedConstraintStates.Reserve(SortedIndices.Num());
+		TBitArray<> SortedConstraintDirtyFlags;
+		const int32 NumToSort = SortedIndices.Num();
+		SortedHandles.Reserve(NumToSort);
+		SortedConstraintSettings.Reserve(NumToSort);
+		SortedConstraintParticles.Reserve(NumToSort);
+		SortedConstraintStates.Reserve(NumToSort);
+		SortedConstraintDirtyFlags.Reserve(NumToSort);
 
 		for (int32 SortedConstraintIndex = 0; SortedConstraintIndex < SortedIndices.Num(); ++SortedConstraintIndex)
 		{
@@ -708,6 +711,7 @@ namespace Chaos
 			SortedConstraintSettings.Add(ConstraintSettings[UnsortedConstraintIndex]);
 			SortedConstraintParticles.Add(ConstraintParticles[UnsortedConstraintIndex]);
 			SortedConstraintStates.Add(ConstraintStates[UnsortedConstraintIndex]);
+			SortedConstraintDirtyFlags.Add(ConstraintDirtyFlags[UnsortedConstraintIndex]);
 
 			Handles[UnsortedConstraintIndex]->SetConstraintIndex(SortedConstraintIndex);
 		}
@@ -716,6 +720,7 @@ namespace Chaos
 		Swap(ConstraintSettings, SortedConstraintSettings);
 		Swap(ConstraintParticles, SortedConstraintParticles);
 		Swap(ConstraintStates, SortedConstraintStates);
+		Swap(ConstraintDirtyFlags, SortedConstraintDirtyFlags);
 	}
 
 
@@ -954,7 +959,10 @@ namespace Chaos
 
 	void FPBDJointConstraints::AddConstraintsToGraph(Private::FPBDIslandManager& IslandManager)
 	{
-		for (int32 ConstraintIndex = 0; ConstraintIndex < GetNumConstraints(); ++ConstraintIndex)
+		const int32 NumJointConstraints = GetNumConstraints();
+		ConstraintDirtyFlags.Init(false, NumJointConstraints);
+
+		for (int32 ConstraintIndex = 0; ConstraintIndex < NumJointConstraints; ++ConstraintIndex)
 		{
 			FPBDJointConstraintHandle* ConstraintHandle = GetConstraintHandle(ConstraintIndex);
 			check(ConstraintHandle != nullptr);
@@ -962,9 +970,14 @@ namespace Chaos
 			const bool bIsInGraph = ConstraintHandle->IsInConstraintGraph();
 			const bool bShouldBeInGraph = ShouldBeInGraph(ConstraintIndex);
 
-			if (bShouldBeInGraph && !bIsInGraph)
+			if (bShouldBeInGraph)
 			{
-				IslandManager.AddConstraint(ContainerId, ConstraintHandle, GetConstrainedParticles(ConstraintIndex));
+				ConstraintDirtyFlags[ConstraintIndex] = true;
+
+				if(!bIsInGraph)
+				{
+					IslandManager.AddConstraint(ContainerId, ConstraintHandle, GetConstrainedParticles(ConstraintIndex));
+				}
 			}
 			else if (bIsInGraph && !bShouldBeInGraph)
 			{
@@ -989,6 +1002,11 @@ namespace Chaos
 		}
 	}
 
+	const TBitArray<>& FPBDJointConstraints::GetDirtyConstraintFlags()
+	{
+		return ConstraintDirtyFlags;
+	}
+
 	bool FPBDJointConstraints::ShouldBeInGraph(const int32 ConstraintIndex) const
 	{
 		if (!IsConstraintEnabled(ConstraintIndex))
diff --git a/Engine/Source/Runtime/Experimental/Chaos/Private/PBDRigidsSolver.cpp b/Engine/Source/Runtime/Experimental/Chaos/Private/PBDRigidsSolver.cpp
--- a/Engine/Source/Runtime/Experimental/Chaos/Private/PBDRigidsSolver.cpp
+++ b/Engine/Source/Runtime/Experimental/Chaos/Private/PBDRigidsSolver.cpp
@@ -2181,10 +2181,11 @@ namespace Chaos
 			ensure(PullData->DirtyJointConstraints.Num() == 0);	//we only fill this once per frame
 			PullData->DirtyJointConstraints.Reserve(JointConstraintPhysicsProxies_Internal.Num());
 
-			for(int32 Idx = 0; Idx < JointConstraintPhysicsProxies_Internal.Num(); ++Idx)
+			TBitArray<>::FConstIterator It(GetJointConstraints().GetDirtyConstraintFlags());
+			for(; It; ++It)
 			{
 				PullData->DirtyJointConstraints.AddDefaulted();
-				JointConstraintPhysicsProxies_Internal[Idx]->BufferPhysicsResults(PullData->DirtyJointConstraints.Last());
+				JointConstraintPhysicsProxies_Internal[It.GetIndex()]->BufferPhysicsResults(PullData->DirtyJointConstraints.Last());
 			}
 		}
 
diff --git a/Engine/Source/Runtime/Experimental/Chaos/Public/Chaos/PBDJointConstraints.h b/Engine/Source/Runtime/Experimental/Chaos/Public/Chaos/PBDJointConstraints.h
--- a/Engine/Source/Runtime/Experimental/Chaos/Public/Chaos/PBDJointConstraints.h
+++ b/Engine/Source/Runtime/Experimental/Chaos/Public/Chaos/PBDJointConstraints.h
@@ -360,6 +360,8 @@ namespace Chaos
 		// Called by the joint solver at the end of the constraint solver phase
 		CHAOS_API void SetSolverResults(const int32 ConstraintIndex, const FVec3& LinearImpulse, const FVec3& AngularImpulse, const bool bIsBroken, const FSolverBody* SolverBody0, const FSolverBody* SolverBody1);
 
+		CHAOS_API const TBitArray<>& GetDirtyConstraintFlags();
+
 		// @todo(chaos): only needed for RBAN, and should be private or moved to the solver
 		CHAOS_API int32 GetConstraintIsland(int32 ConstraintIndex) const;
 		CHAOS_API int32 GetConstraintLevel(int32 ConstraintIndex) const;
@@ -416,6 +418,7 @@ namespace Chaos
 		TArray<FPBDJointSettings> ConstraintSettings;
 		TArray<FParticlePair> ConstraintParticles;
 		TArray<FPBDJointState> ConstraintStates;
+		TBitArray<> ConstraintDirtyFlags;
 
 		FHandles Handles;
 		FConstraintHandleAllocator HandleAllocator;
diff --git a/Engine/Source/Runtime/Experimental/Chaos/Public/PhysicsProxy/JointConstraintProxy.h b/Engine/Source/Runtime/Experimental/Chaos/Public/PhysicsProxy/JointConstraintProxy.h
--- a/Engine/Source/Runtime/Experimental/Chaos/Public/PhysicsProxy/JointConstraintProxy.h
+++ b/Engine/Source/Runtime/Experimental/Chaos/Public/PhysicsProxy/JointConstraintProxy.h
@@ -104,7 +104,6 @@ private:
 	FJointConstraint* Constraint_GT;
 	FPBDJointConstraintHandle* Constraint_PT;
 	FParticlePair OriginalParticleHandles_PT;
-	bool bInitialized = false;
 };
 
 
