[Backout] - CL32971360
Original CL Desc
-----------------------------------------------------------------
Implemented dirty joint tracking to avoid syncing back every joint on each physics tick.

Added a bit array tracking which constraints in the internal constraint list should be present in the constraint graph. If they're in the graph they will be solved and we will have results to marshall back to the game thread.

diff --git a/Engine/Source/Runtime/Experimental/Chaos/Private/Chaos/PBDJointConstraints.cpp b/Engine/Source/Runtime/Experimental/Chaos/Private/Chaos/PBDJointConstraints.cpp
--- a/Engine/Source/Runtime/Experimental/Chaos/Private/Chaos/PBDJointConstraints.cpp
+++ b/Engine/Source/Runtime/Experimental/Chaos/Private/Chaos/PBDJointConstraints.cpp
@@ -695,13 +695,10 @@ namespace Chaos
 		TArray<FPBDJointSettings> SortedConstraintSettings;
 		TArray<FParticlePair> SortedConstraintParticles;
 		TArray<FPBDJointState> SortedConstraintStates;
-		TBitArray<> SortedConstraintDirtyFlags;
-		const int32 NumToSort = SortedIndices.Num();
-		SortedHandles.Reserve(NumToSort);
-		SortedConstraintSettings.Reserve(NumToSort);
-		SortedConstraintParticles.Reserve(NumToSort);
-		SortedConstraintStates.Reserve(NumToSort);
-		SortedConstraintDirtyFlags.Reserve(NumToSort);
+		SortedHandles.Reserve(SortedIndices.Num());
+		SortedConstraintSettings.Reserve(SortedIndices.Num());
+		SortedConstraintParticles.Reserve(SortedIndices.Num());
+		SortedConstraintStates.Reserve(SortedIndices.Num());
 
 		for (int32 SortedConstraintIndex = 0; SortedConstraintIndex < SortedIndices.Num(); ++SortedConstraintIndex)
 		{
@@ -711,7 +708,6 @@ namespace Chaos
 			SortedConstraintSettings.Add(ConstraintSettings[UnsortedConstraintIndex]);
 			SortedConstraintParticles.Add(ConstraintParticles[UnsortedConstraintIndex]);
 			SortedConstraintStates.Add(ConstraintStates[UnsortedConstraintIndex]);
-			SortedConstraintDirtyFlags.Add(ConstraintDirtyFlags[UnsortedConstraintIndex]);
 
 			Handles[UnsortedConstraintIndex]->SetConstraintIndex(SortedConstraintIndex);
 		}
@@ -720,7 +716,6 @@ namespace Chaos
 		Swap(ConstraintSettings, SortedConstraintSettings);
 		Swap(ConstraintParticles, SortedConstraintParticles);
 		Swap(ConstraintStates, SortedConstraintStates);
-		Swap(ConstraintDirtyFlags, SortedConstraintDirtyFlags);
 	}
 
 
@@ -959,10 +954,7 @@ namespace Chaos
 
 	void FPBDJointConstraints::AddConstraintsToGraph(Private::FPBDIslandManager& IslandManager)
 	{
-		const int32 NumJointConstraints = GetNumConstraints();
-		ConstraintDirtyFlags.Init(false, NumJointConstraints);
-
-		for (int32 ConstraintIndex = 0; ConstraintIndex < NumJointConstraints; ++ConstraintIndex)
+		for (int32 ConstraintIndex = 0; ConstraintIndex < GetNumConstraints(); ++ConstraintIndex)
 		{
 			FPBDJointConstraintHandle* ConstraintHandle = GetConstraintHandle(ConstraintIndex);
 			check(ConstraintHandle != nullptr);
@@ -970,14 +962,9 @@ namespace Chaos
 			const bool bIsInGraph = ConstraintHandle->IsInConstraintGraph();
 			const bool bShouldBeInGraph = ShouldBeInGraph(ConstraintIndex);
 
-			if (bShouldBeInGraph)
+			if (bShouldBeInGraph && !bIsInGraph)
 			{
-				ConstraintDirtyFlags[ConstraintIndex] = true;
-
-				if(!bIsInGraph)
-				{
-					IslandManager.AddConstraint(ContainerId, ConstraintHandle, GetConstrainedParticles(ConstraintIndex));
-				}
+				IslandManager.AddConstraint(ContainerId, ConstraintHandle, GetConstrainedParticles(ConstraintIndex));
 			}
 			else if (bIsInGraph && !bShouldBeInGraph)
 			{
@@ -1002,11 +989,6 @@ namespace Chaos
 		}
 	}
 
-	const TBitArray<>& FPBDJointConstraints::GetDirtyConstraintFlags()
-	{
-		return ConstraintDirtyFlags;
-	}
-
 	bool FPBDJointConstraints::ShouldBeInGraph(const int32 ConstraintIndex) const
 	{
 		if (!IsConstraintEnabled(ConstraintIndex))
diff --git a/Engine/Source/Runtime/Experimental/Chaos/Private/PBDRigidsSolver.cpp b/Engine/Source/Runtime/Experimental/Chaos/Private/PBDRigidsSolver.cpp
--- a/Engine/Source/Runtime/Experimental/Chaos/Private/PBDRigidsSolver.cpp
+++ b/Engine/Source/Runtime/Experimental/Chaos/Private/PBDRigidsSolver.cpp
@@ -2181,11 +2181,10 @@ namespace Chaos
 			ensure(PullData->DirtyJointConstraints.Num() == 0);	//we only fill this once per frame
 			PullData->DirtyJointConstraints.Reserve(JointConstraintPhysicsProxies_Internal.Num());
 
-			TBitArray<>::FConstIterator It(GetJointConstraints().GetDirtyConstraintFlags());
-			for(; It; ++It)
+			for(int32 Idx = 0; Idx < JointConstraintPhysicsProxies_Internal.Num(); ++Idx)
 			{
 				PullData->DirtyJointConstraints.AddDefaulted();
-				JointConstraintPhysicsProxies_Internal[It.GetIndex()]->BufferPhysicsResults(PullData->DirtyJointConstraints.Last());
+				JointConstraintPhysicsProxies_Internal[Idx]->BufferPhysicsResults(PullData->DirtyJointConstraints.Last());
 			}
 		}
 
diff --git a/Engine/Source/Runtime/Experimental/Chaos/Public/Chaos/PBDJointConstraints.h b/Engine/Source/Runtime/Experimental/Chaos/Public/Chaos/PBDJointConstraints.h
--- a/Engine/Source/Runtime/Experimental/Chaos/Public/Chaos/PBDJointConstraints.h
+++ b/Engine/Source/Runtime/Experimental/Chaos/Public/Chaos/PBDJointConstraints.h
@@ -360,8 +360,6 @@ namespace Chaos
 		// Called by the joint solver at the end of the constraint solver phase
 		CHAOS_API void SetSolverResults(const int32 ConstraintIndex, const FVec3& LinearImpulse, const FVec3& AngularImpulse, const bool bIsBroken, const FSolverBody* SolverBody0, const FSolverBody* SolverBody1);
 
-		CHAOS_API const TBitArray<>& GetDirtyConstraintFlags();
-
 		// @todo(chaos): only needed for RBAN, and should be private or moved to the solver
 		CHAOS_API int32 GetConstraintIsland(int32 ConstraintIndex) const;
 		CHAOS_API int32 GetConstraintLevel(int32 ConstraintIndex) const;
@@ -418,7 +416,6 @@ namespace Chaos
 		TArray<FPBDJointSettings> ConstraintSettings;
 		TArray<FParticlePair> ConstraintParticles;
 		TArray<FPBDJointState> ConstraintStates;
-		TBitArray<> ConstraintDirtyFlags;
 
 		FHandles Handles;
 		FConstraintHandleAllocator HandleAllocator;
diff --git a/Engine/Source/Runtime/Experimental/Chaos/Public/PhysicsProxy/JointConstraintProxy.h b/Engine/Source/Runtime/Experimental/Chaos/Public/PhysicsProxy/JointConstraintProxy.h
--- a/Engine/Source/Runtime/Experimental/Chaos/Public/PhysicsProxy/JointConstraintProxy.h
+++ b/Engine/Source/Runtime/Experimental/Chaos/Public/PhysicsProxy/JointConstraintProxy.h
@@ -104,6 +104,7 @@ private:
 	FJointConstraint* Constraint_GT;
 	FPBDJointConstraintHandle* Constraint_PT;
 	FParticlePair OriginalParticleHandles_PT;
+	bool bInitialized = false;
 };
 
 
