Removed unused variables in Chaos.
Fixed uninitialized members in TSQVisitor.

diff --git a/Engine/Source/Runtime/Experimental/Chaos/Private/Chaos/HeightField.cpp b/Engine/Source/Runtime/Experimental/Chaos/Private/Chaos/HeightField.cpp
--- a/Engine/Source/Runtime/Experimental/Chaos/Private/Chaos/HeightField.cpp
+++ b/Engine/Source/Runtime/Experimental/Chaos/Private/Chaos/HeightField.cpp
@@ -1319,7 +1319,6 @@ namespace Chaos
 		FVec2 ClippedFlatRayEnd;
 
 		// Data for fast box cast
-		FVec3 Min, Max, HitPoint;
 		bool bParallel[3];
 		FVec3 InvDir;
 
@@ -1487,7 +1486,6 @@ namespace Chaos
 			Seen.Add(StartCell);
 
 			// Data for fast box cast
-			FVec3 Min, Max, HitPoint;
 			bool bParallel[3];
 			FVec3 InvDir;
 
diff --git a/Engine/Source/Runtime/Experimental/Chaos/Private/GeometryCollection/GeometryCollectionConvexUtility.cpp b/Engine/Source/Runtime/Experimental/Chaos/Private/GeometryCollection/GeometryCollectionConvexUtility.cpp
--- a/Engine/Source/Runtime/Experimental/Chaos/Private/GeometryCollection/GeometryCollectionConvexUtility.cpp
+++ b/Engine/Source/Runtime/Experimental/Chaos/Private/GeometryCollection/GeometryCollectionConvexUtility.cpp
@@ -626,7 +626,6 @@ void CreateNonoverlappingConvexHulls(
 	{
 		OutCenter = (CloseA + CloseB) * .5;
 		OutNormal = Normal;
-		Chaos::FVec3 GJKCenter;
 		Chaos::FReal BestScore = ScoreCutPlane(A, B, FChaosPlane(OutCenter, OutNormal), bOneSidedCut, OutCenter, OutNormal);
 		Chaos::FVec3 MassSepNormal = (B.GetCenterOfMass() - A.GetCenterOfMass());
 		if (MassSepNormal.Normalize() && BestScore > 0)
diff --git a/Engine/Source/Runtime/Experimental/Chaos/Public/Chaos/AABB.h b/Engine/Source/Runtime/Experimental/Chaos/Public/Chaos/AABB.h
--- a/Engine/Source/Runtime/Experimental/Chaos/Public/Chaos/AABB.h
+++ b/Engine/Source/Runtime/Experimental/Chaos/Public/Chaos/AABB.h
@@ -114,7 +114,6 @@ namespace Chaos
 
 		FORCEINLINE TAABB<T, d> GetIntersection(const TAABB<T, d>& Other) const
 		{
-			TVector<T, 3> Tmp;
 			return TAABB<T, d>(MMin.ComponentwiseMax(Other.MMin), MMax.ComponentwiseMin(Other.MMax));
 		}
 
diff --git a/Engine/Source/Runtime/Experimental/Chaos/Public/Chaos/Collision/SimSweep.h b/Engine/Source/Runtime/Experimental/Chaos/Public/Chaos/Collision/SimSweep.h
--- a/Engine/Source/Runtime/Experimental/Chaos/Public/Chaos/Collision/SimSweep.h
+++ b/Engine/Source/Runtime/Experimental/Chaos/Public/Chaos/Collision/SimSweep.h
@@ -160,7 +160,6 @@ namespace Chaos
 						{
 							const FRigidTransform3 OtherShapeWorldTransform = FRigidTransform3(OtherShape->GetLeafRelativeTransform()) * OtherWorldTransform;
 
-							FMTDInfo MTDInfo;
 							const bool bOverlapHit = OverlapQuery(
 								*OtherImplicit,
 								OtherShapeWorldTransform,
diff --git a/Engine/Source/Runtime/Experimental/Chaos/Public/Chaos/Simplex.h b/Engine/Source/Runtime/Experimental/Chaos/Public/Chaos/Simplex.h
--- a/Engine/Source/Runtime/Experimental/Chaos/Public/Chaos/Simplex.h
+++ b/Engine/Source/Runtime/Experimental/Chaos/Public/Chaos/Simplex.h
@@ -635,7 +635,6 @@ namespace Chaos
 	template <typename T>
 	TVec3<T> SimplexFindClosestToOrigin2(TVec3<T>* Simplex, int32& NumVerts, T* OutBarycentric, TVec3<T>* A, TVec3<T>* B)
 	{
-		TVec3<T> ClosestPoint;
 		switch (NumVerts)
 		{
 		case 1:
diff --git a/Engine/Source/Runtime/PhysicsCore/Public/SQVisitor.h b/Engine/Source/Runtime/PhysicsCore/Public/SQVisitor.h
--- a/Engine/Source/Runtime/PhysicsCore/Public/SQVisitor.h
+++ b/Engine/Source/Runtime/PhysicsCore/Public/SQVisitor.h
@@ -446,9 +446,9 @@ private:
 	}
 #endif
 
-	const FVector StartPoint;
-	const FVector Dir;
-	const FVector HalfExtents;
+	const FVector StartPoint = FVector::ZeroVector;
+	const FVector Dir = FVector::ZeroVector;
+	const FVector HalfExtents = FVector::ZeroVector;
 	FHitFlags OutputFlags;
 	bool bAnyHit;
 	const ChaosInterface::FQueryDebugParams DebugParams;
