Fix spam debug log  in AddRadialImpulseHelper

diff --git a/Engine/Source/Runtime/Experimental/Chaos/Private/Chaos/PhysicsObjectInterface.cpp b/Engine/Source/Runtime/Experimental/Chaos/Private/Chaos/PhysicsObjectInterface.cpp
--- a/Engine/Source/Runtime/Experimental/Chaos/Private/Chaos/PhysicsObjectInterface.cpp
+++ b/Engine/Source/Runtime/Experimental/Chaos/Private/Chaos/PhysicsObjectInterface.cpp
@@ -131,7 +131,6 @@ namespace
 						const float PrevFalloffAlpha = FalloffAlpha;
 						// remap the value within MinValue/MaxValue within the radius
 						FalloffAlpha = FMath::Lerp(MinValue, MaxValue, FalloffAlpha);
-						UE_LOG(LogChaos, Warning, TEXT("FalloffAlpha = %f = Lerp(%f, %f, %f)"), FalloffAlpha, MinValue, MaxValue, PrevFalloffAlpha);
 					}
 					else
 					{
