Reorganised the large .cpp into more manageable units.

#rnx

diff --git a/Engine/Source/Runtime/Experimental/IoStoreOnDemand/Private/CoreHttp/Activity.inl b/Engine/Source/Runtime/Experimental/IoStoreOnDemand/Private/CoreHttp/Activity.inl
new file mode 100644
--- /dev/null
+++ b/Engine/Source/Runtime/Experimental/IoStoreOnDemand/Private/CoreHttp/Activity.inl
@@ -0,0 +1,193 @@
+// Copyright Epic Games, Inc. All Rights Reserved.
+
+#pragma once
+
+namespace UE::IoStore::HTTP
+{
+
+// {{{1 activity ...............................................................
+
+#if IAS_HTTP_WITH_PERF
+
+////////////////////////////////////////////////////////////////////////////////
+class FStopwatch
+{
+public:
+	uint64	GetInterval(uint32 i) const;
+	void	SendStart()		{ Impl(0); }
+	void	SendEnd()		{ Impl(1); }
+	void	RecvStart()		{ Impl(2); }
+	void	RecvEnd()		{ Impl(3); }
+
+private:
+	void	Impl(uint32 Index);
+	uint64	Samples[4] = {};
+	uint32	Counts[2] = {};
+};
+
+////////////////////////////////////////////////////////////////////////////////
+uint64 FStopwatch::GetInterval(uint32 i) const
+{
+	if (i >= UE_ARRAY_COUNT(Samples) - 1)
+	{
+		return 0;
+	}
+	return Samples[i + 1] - Samples[i];
+}
+
+////////////////////////////////////////////////////////////////////////////////
+void FStopwatch::Impl(uint32 Index)
+{
+	if (uint64& Out = Samples[Index]; Out == 0)
+	{
+		Out = FPlatformTime::Cycles64();
+	}
+	Counts[Index >> 1] += !(Index & 1);
+}
+
+#endif // IAS_HTTP_WITH_PERF
+
+////////////////////////////////////////////////////////////////////////////////
+struct FResponseInternal
+{
+	FMessageOffsets Offsets;
+	int32			ContentLength = 0;
+	uint16			MessageLength;
+	mutable int16	Code;
+};
+
+////////////////////////////////////////////////////////////////////////////////
+struct alignas(16) FActivity
+{
+	enum class EState : uint8
+	{
+		None,
+		Build,
+		Send,
+		RecvMessage,
+		RecvContent,
+		RecvStream,
+		RecvDone,
+		Completed,
+		Cancelled,
+		Failed,
+		_Num,
+	};
+
+	FActivity*			Next = nullptr;
+	int8				Slot = -1;
+	EState				State = EState::None;
+	uint8				IsKeepAlive : 1;
+	uint8				NoContent : 1;
+	uint8				bFollow30x : 1;
+	uint8				_Unused0 : 5;
+	uint32				StateParam = 0;
+#if IAS_HTTP_WITH_PERF
+	FStopwatch			Stopwatch;
+#endif
+	union {
+		FHost*			Host;
+		FIoBuffer*		Dest;
+		const char*		ErrorReason;
+	};
+	UPTRINT				SinkParam;
+	FTicketSink			Sink;
+	FSocket				Socket;
+	FResponseInternal	Response;
+	FBuffer				Buffer;
+};
+
+////////////////////////////////////////////////////////////////////////////////
+static void Activity_TraceStateNames()
+{
+	FAnsiStringView StateNames[] = {
+		"None", "Build", "Send", "RecvMessage", "RecvContent", "RecvStream",
+		"RecvDone", "Completed", "Cancelled", "Failed", "$",
+	};
+	static_assert(UE_ARRAY_COUNT(StateNames) == int32(FActivity::EState::_Num) + 1);
+
+	TraceEnum(StateNames);
+}
+
+////////////////////////////////////////////////////////////////////////////////
+static void Activity_ChangeState(FActivity* Activity, FActivity::EState InState, uint32 Param=0)
+{
+	Trace(Activity, ETrace::StateChange, InState);
+
+	check(Activity->State != InState);
+	Activity->State = InState;
+	Activity->StateParam = Param;
+}
+
+////////////////////////////////////////////////////////////////////////////////
+static int32 Activity_Rewind(FActivity* Activity)
+{
+	using EState = FActivity::EState;
+
+	if (Activity->State == EState::Send)
+	{
+		Activity->StateParam = 0;
+		return 0;
+	}
+
+	if (Activity->State == EState::RecvMessage)
+	{
+		Activity->Buffer.Resize(Activity->StateParam);
+		Activity_ChangeState(Activity, EState::Send);
+		return 1;
+	}
+
+	return -1;
+}
+
+////////////////////////////////////////////////////////////////////////////////
+static uint32 Activity_RemainingKiB(FActivity* Activity)
+{
+	if (Activity->State < FActivity::EState::RecvContent) return MAX_uint32;
+	if (Activity->State > FActivity::EState::RecvContent) return 0;
+
+	uint32 ContentLength = uint32(Activity->Response.ContentLength);
+	check(Activity->StateParam <= ContentLength);
+	return (ContentLength - Activity->StateParam) >> 10;
+}
+
+////////////////////////////////////////////////////////////////////////////////
+static FActivity* Activity_Alloc(uint32 BufferSize)
+{
+	BufferSize = (BufferSize + 15) & ~15;
+
+	uint32 Size = BufferSize + sizeof(FActivity);
+	auto* Activity = (FActivity*)FMemory::Malloc(Size, alignof(FActivity));
+
+	new (Activity) FActivity();
+
+	auto* Scratch = (char*)(Activity + 1);
+	uint32 ScratchSize = BufferSize;
+	Activity->Buffer = FBuffer(Scratch, ScratchSize);
+
+	Activity_ChangeState(Activity, FActivity::EState::Build);
+
+	return Activity;
+}
+
+////////////////////////////////////////////////////////////////////////////////
+static void Activity_Free(FActivity* Activity)
+{
+	Trace(Activity, ETrace::ActivityDestroy);
+
+	Activity->~FActivity();
+	FMemory::Free(Activity);
+}
+
+////////////////////////////////////////////////////////////////////////////////
+static void Activity_SetError(FActivity* Activity, const char* Reason)
+{
+	Activity->IsKeepAlive = 0;
+	Activity->ErrorReason = Reason;
+
+	Activity_ChangeState(Activity, FActivity::EState::Failed, LastSocketResult());
+}
+
+// }}}
+
+} // namespace UE::IoStore::HTTP
diff --git a/Engine/Source/Runtime/Experimental/IoStoreOnDemand/Private/CoreHttp/Api.inl b/Engine/Source/Runtime/Experimental/IoStoreOnDemand/Private/CoreHttp/Api.inl
new file mode 100644
--- /dev/null
+++ b/Engine/Source/Runtime/Experimental/IoStoreOnDemand/Private/CoreHttp/Api.inl
@@ -0,0 +1,297 @@
+// Copyright Epic Games, Inc. All Rights Reserved.
+
+#pragma once
+
+namespace UE::IoStore::HTTP
+{
+
+// {{{1 request ................................................................
+
+////////////////////////////////////////////////////////////////////////////////
+FRequest::~FRequest()
+{
+	if (Ptr == nullptr)
+	{
+		return;
+	}
+
+	Activity_Free(Ptr);
+}
+
+////////////////////////////////////////////////////////////////////////////////
+FRequest::FRequest(FRequest&& Rhs)
+{
+	this->~FRequest();
+	Swap(Ptr, Rhs.Ptr);
+}
+
+////////////////////////////////////////////////////////////////////////////////
+FRequest&& FRequest::Accept(EMimeType MimeType)
+{
+	switch (MimeType)
+	{
+	case EMimeType::Text:				return Header("Accept", "text/html");
+	case EMimeType::Binary:				return Header("Accept", "application/octet-stream");
+	case EMimeType::Json:				return Header("Accept", "application/json");
+	case EMimeType::Xml:				return Header("Accept", "application/xml");
+	/* UE_CUSTOM_MIME_TYPES
+	case EMimeType::CbObject:			return Header("Accept", "application/x-ue-cb");
+	case EMimeType::CbPackage:			return Header("Accept", "application/x-ue-pkg");
+	case EMimeType::CompressedBuffer:	return Header("Accept", "application/x-ue-comp");
+	*/
+	}
+
+	return MoveTemp(*this);
+}
+
+////////////////////////////////////////////////////////////////////////////////
+FRequest&& FRequest::Accept(FAnsiStringView MimeType)
+{
+	return Header("Accept", MimeType);
+}
+
+////////////////////////////////////////////////////////////////////////////////
+FRequest&& FRequest::Header(FAnsiStringView Key, FAnsiStringView Value)
+{
+	checkSlow(Ptr->State == FActivity::EState::Build);
+	FMessageBuilder(Ptr->Buffer) << Key << ": " << Value << "\r\n";
+	return MoveTemp(*this);
+}
+
+
+
+// {{{1 response ...............................................................
+
+////////////////////////////////////////////////////////////////////////////////
+EStatusCodeClass FResponse::GetStatus() const
+{
+	uint32 Code = GetStatusCode();
+	if (Code <= 199) return EStatusCodeClass::Informational;
+	if (Code <= 299) return EStatusCodeClass::Successful;
+	if (Code <= 399) return EStatusCodeClass::Redirection;
+	if (Code <= 499) return EStatusCodeClass::ClientError;
+	if (Code <= 599) return EStatusCodeClass::ServerError;
+	return EStatusCodeClass::Unknown;
+}
+
+////////////////////////////////////////////////////////////////////////////////
+uint32 FResponse::GetStatusCode() const
+{
+	const auto* Activity = (const FActivity*)this;
+	const FResponseInternal& Internal = Activity->Response;
+
+	const char* MessageData = Activity->Buffer.GetData() + Activity->StateParam;
+	if (Internal.Code < 0)
+	{
+		const char* CodePtr = MessageData + Internal.Offsets.StatusCode;
+		Internal.Code = uint16(CrudeToInt(FAnsiStringView(CodePtr, 3)));
+	}
+
+	return Internal.Code;
+}
+
+////////////////////////////////////////////////////////////////////////////////
+FAnsiStringView FResponse::GetStatusMessage() const
+{
+	const auto* Activity = (const FActivity*)this;
+	const FResponseInternal& Internal = Activity->Response;
+
+	const char* MessageData = Activity->Buffer.GetData() + Activity->StateParam;
+	return FAnsiStringView(
+		MessageData + Internal.Offsets.Message,
+		Internal.Offsets.Headers - Internal.Offsets.Message
+	);
+}
+
+////////////////////////////////////////////////////////////////////////////////
+int64 FResponse::GetContentLength() const
+{
+	const auto* Activity = (const FActivity*)this;
+	const FResponseInternal& Internal = Activity->Response;
+	return Internal.ContentLength;
+}
+
+////////////////////////////////////////////////////////////////////////////////
+EMimeType FResponse::GetContentType() const
+{
+	FAnsiStringView Value;
+	GetContentType(Value);
+
+	if (Value == "text/html")					return EMimeType::Text;
+	if (Value == "application/octet-stream")	return EMimeType::Binary;
+	if (Value == "application/json")			return EMimeType::Json;
+	if (Value == "application/xml")				return EMimeType::Xml;
+	/* UE_CUSTOM_MIME_TYPES
+	if (Value == "application/x-ue-cb")			return EMimeType::CbObject;
+	if (Value == "application/x-ue-pkg")		return EMimeType::CbPackage;
+	if (Value == "application/x-ue-comp")		return EMimeType::CompressedBuffer;
+	*/
+
+	return EMimeType::Unknown;
+}
+
+////////////////////////////////////////////////////////////////////////////////
+void FResponse::GetContentType(FAnsiStringView& Out) const
+{
+	Out = GetHeader("Accept");
+
+	int32 SemiColon;
+	if (Out.FindChar(';', SemiColon))
+	{
+		Out = Out.Mid(SemiColon).TrimEnd();
+	}
+}
+
+////////////////////////////////////////////////////////////////////////////////
+FAnsiStringView FResponse::GetHeader(FAnsiStringView Name) const
+{
+	FAnsiStringView Result;
+	ReadHeaders([&Result, Name] (FAnsiStringView Candidate, FAnsiStringView Value)
+	{
+		if (Candidate != Name)
+		{
+			return true;
+		}
+
+		Result = Value;
+		return false;
+	});
+	return Result;
+}
+
+////////////////////////////////////////////////////////////////////////////////
+void FResponse::ReadHeaders(FHeaderSink Sink) const
+{
+	const auto* Activity = (const FActivity*)this;
+	const FResponseInternal& Internal = Activity->Response;
+
+	const char* MessageData = Activity->Buffer.GetData() + Activity->StateParam;
+	FAnsiStringView Result, Headers(
+		MessageData + Internal.Offsets.Headers,
+		Internal.MessageLength - Internal.Offsets.Headers
+	);
+
+	EnumerateHeaders(Headers, Sink);
+}
+
+////////////////////////////////////////////////////////////////////////////////
+void FResponse::SetDestination(FIoBuffer* Buffer)
+{
+	auto* Activity = (FActivity*)this;
+	Activity->Dest = Buffer;
+}
+
+
+
+// {{{1 ticket-status ..........................................................
+
+////////////////////////////////////////////////////////////////////////////////
+FTicketStatus::EId FTicketStatus::GetId() const
+{
+	const auto* Activity = (FActivity*)this;
+	switch (Activity->State)
+	{
+	case FActivity::EState::RecvMessage:	return EId::Response;
+	case FActivity::EState::RecvDone:		return EId::Content;
+	case FActivity::EState::Cancelled:		return EId::Cancelled;
+	case FActivity::EState::Failed:			return EId::Error;
+	default:								check(false);
+	}
+	return EId::Error;
+}
+
+////////////////////////////////////////////////////////////////////////////////
+UPTRINT FTicketStatus::GetParam() const
+{
+	const auto* Activity = (FActivity*)this;
+	return Activity->SinkParam;
+}
+
+////////////////////////////////////////////////////////////////////////////////
+FTicket FTicketStatus::GetTicket() const
+{
+	const auto* Activity = (FActivity*)this;
+	return 1ull << Activity->Slot;
+}
+
+////////////////////////////////////////////////////////////////////////////////
+uint32 FTicketStatus::GetIndex() const
+{
+	const auto* Activity = (FActivity*)this;
+	return Activity->Slot;
+}
+
+////////////////////////////////////////////////////////////////////////////////
+FResponse& FTicketStatus::GetResponse() const
+{
+	check(GetId() < EId::Content);
+	const auto* Activity = (FActivity*)this;
+	return *(FResponse*)Activity;
+}
+
+////////////////////////////////////////////////////////////////////////////////
+uint32 FTicketStatus::GetContentLength() const
+{
+	check(GetId() <= EId::Content);
+	const auto* Activity = (FActivity*)this;
+	return Activity->Response.ContentLength;
+}
+
+////////////////////////////////////////////////////////////////////////////////
+const FTicketPerf& FTicketStatus::GetPerf() const
+{
+	check(GetId() == EId::Content);
+	const auto* Activity = (FActivity*)this;
+	return *(FTicketPerf*)Activity;
+}
+
+////////////////////////////////////////////////////////////////////////////////
+const FIoBuffer& FTicketStatus::GetContent() const
+{
+	check(GetId() == EId::Content);
+	const auto* Activity = (FActivity*)this;
+	return *(Activity->Dest);
+}
+
+////////////////////////////////////////////////////////////////////////////////
+const char* FTicketStatus::GetErrorReason() const
+{
+	check(GetId() == EId::Error);
+	const auto* Activity = (FActivity*)this;
+	return Activity->ErrorReason;
+}
+
+
+
+// {{{1 perf ...................................................................
+
+#if IAS_HTTP_WITH_PERF
+
+////////////////////////////////////////////////////////////////////////////////
+FTicketPerf::FSample FTicketPerf::GetSample() const
+{
+	const auto* Activity = (FActivity*)this;
+
+	static uint64 Freq;
+	if (Freq == 0)
+	{
+		Freq = uint64(1.0 / FPlatformTime::GetSecondsPerCycle());
+	}
+
+	auto ToMs = [] (uint64 Value) { return uint32((Value * 1000ull) / Freq); };
+
+	const FStopwatch& Stopwatch = Activity->Stopwatch;
+	return {
+		ToMs(Stopwatch.GetInterval(0)),
+		ToMs(Stopwatch.GetInterval(1)),
+		ToMs(Stopwatch.GetInterval(2)),
+	};
+}
+
+#endif // IAS_HTTP_WITH_PERF
+
+
+
+// }}}
+
+} // namespace UE::IoStore::HTTP
diff --git a/Engine/Source/Runtime/Experimental/IoStoreOnDemand/Private/CoreHttp/Client.h b/Engine/Source/Runtime/Experimental/IoStoreOnDemand/Private/CoreHttp/Client.h
--- a/Engine/Source/Runtime/Experimental/IoStoreOnDemand/Private/CoreHttp/Client.h
+++ b/Engine/Source/Runtime/Experimental/IoStoreOnDemand/Private/CoreHttp/Client.h
@@ -3,7 +3,7 @@
 #pragma once
 
 #if !defined(NO_UE_INCLUDES)
-#include "Containers/StringView.h"
+#include <Containers/StringView.h>
 #endif
 
 #if !defined(IAS_HTTP_WITH_PERF)
diff --git a/Engine/Source/Runtime/Experimental/IoStoreOnDemand/Private/CoreHttp/ConnectionPool.inl b/Engine/Source/Runtime/Experimental/IoStoreOnDemand/Private/CoreHttp/ConnectionPool.inl
new file mode 100644
--- /dev/null
+++ b/Engine/Source/Runtime/Experimental/IoStoreOnDemand/Private/CoreHttp/ConnectionPool.inl
@@ -0,0 +1,302 @@
+// Copyright Epic Games, Inc. All Rights Reserved.
+
+#pragma once
+
+namespace UE::IoStore::HTTP
+{
+
+// {{{1 connection-pool ........................................................
+
+////////////////////////////////////////////////////////////////////////////////
+class FHost
+{
+public:
+	enum class EDirection : uint8 { Send, Recv };
+	static const uint32 InvalidIp = 0x00ff'ffff;
+
+	struct FParams
+	{
+		const ANSICHAR* HostName;
+		uint32			Port = 0;
+		uint32			MaxConnections = 1;
+	};
+
+					FHost(const FParams& Params);
+	void			SetBufferSize(EDirection Dir, int32 Size);
+	int32			GetBufferSize(EDirection Dir) const;
+	FOutcome		Connect(FSocket& Socket);
+	int32			IsResolved() const;
+	FOutcome		ResolveHostName();
+	uint32			GetMaxConnections() const	{ return MaxConnections; }
+	uint32			GetIpAddress() const		{ return IpAddresses[0]; }
+	FAnsiStringView	GetHostName() const			{ return HostName; }
+	uint32			GetPort() const				{ return Port; }
+
+private:
+	const ANSICHAR*	HostName;
+	uint32			IpAddresses[4] = {};
+	int16			SendBufKb = -1;
+	int16			RecvBufKb = -1;
+	uint16			Port;
+	uint8			MaxConnections;
+};
+
+////////////////////////////////////////////////////////////////////////////////
+FHost::FHost(const FParams& Params)
+: HostName(Params.HostName)
+, Port(uint16(Params.Port))
+, MaxConnections(uint8(Params.MaxConnections))
+{
+	check(MaxConnections && MaxConnections == Params.MaxConnections);
+
+	if (Port == 0)
+	{
+		Port = 80;
+	}
+}
+
+////////////////////////////////////////////////////////////////////////////////
+void FHost::SetBufferSize(EDirection Dir, int32 Size)
+{
+	(Dir == EDirection::Send) ? SendBufKb : RecvBufKb = uint16(Size >> 10);
+}
+
+////////////////////////////////////////////////////////////////////////////////
+int32 FHost::GetBufferSize(EDirection Dir) const
+{
+	return int32((Dir == EDirection::Send) ? SendBufKb : RecvBufKb) << 10;
+}
+
+////////////////////////////////////////////////////////////////////////////////
+FOutcome FHost::ResolveHostName()
+{
+	// todo: GetAddrInfoW() for async resolve on Windows
+
+	TRACE_CPUPROFILER_EVENT_SCOPE(IasHttp::PoolResolve);
+
+	IpAddresses[0] = 1;
+
+	addrinfo* Info = nullptr;
+	ON_SCOPE_EXIT { if (Info != nullptr) freeaddrinfo(Info); };
+
+	addrinfo Hints = {};
+	Hints.ai_family = AF_INET;
+	Hints.ai_socktype = SOCK_STREAM;
+	Hints.ai_protocol = IPPROTO_TCP;
+	auto Result = getaddrinfo(HostName, nullptr, &Hints, &Info);
+	if (uint32(Result) || Info == nullptr)
+	{
+		return FOutcome::Error("Error encountered resolving");
+	}
+
+	if (Info->ai_family != AF_INET)
+	{
+		return FOutcome::Error("Unexpected address family during resolve");
+	}
+
+	uint32 AddressCount = 0;
+	for (const addrinfo* Cursor = Info; Cursor != nullptr; Cursor = Cursor->ai_next)
+	{
+		const auto* AddrInet = (sockaddr_in*)(Cursor->ai_addr);
+		if (AddrInet->sin_family != AF_INET)
+		{
+			continue;
+		}
+
+		uint32 IpAddress = 0;
+		memcpy(&IpAddress, &(AddrInet->sin_addr), sizeof(uint32));
+
+		if (IpAddress == 0)
+		{
+			break;
+		}
+
+		IpAddresses[AddressCount] = htonl(IpAddress);
+		if (++AddressCount >= UE_ARRAY_COUNT(IpAddresses))
+		{
+			break;
+		}
+	}
+
+	if (AddressCount > 0)
+	{
+		return FOutcome::Ok(AddressCount);
+	}
+
+	return FOutcome::Error("Unable to resolve host");
+}
+
+////////////////////////////////////////////////////////////////////////////////
+int32 FHost::IsResolved() const
+{
+	switch (IpAddresses[0])
+	{
+	case 0:  return 0;
+	case 1:  return -1;
+	default: return 1;
+	}
+}
+
+////////////////////////////////////////////////////////////////////////////////
+FOutcome FHost::Connect(FSocket& Socket)
+{
+	FOutcome Outcome = FOutcome::None();
+
+	if (IsResolved() <= 0)
+	{
+		if (Outcome = ResolveHostName(); Outcome.IsError())
+		{
+			return Outcome;
+		}
+	}
+
+	check(IsResolved() > 0);
+	check(!Socket.IsValid());
+
+	uint32 IpAddress = GetIpAddress();
+
+	FSocket Candidate;
+	if (!Candidate.Create())
+	{
+		return FOutcome::Error("Failed to create socket");
+	}
+
+	// Attempt a SOCKS connect
+	Outcome = MaybeConnectSocks(Candidate, IpAddress, Port);
+	if (Outcome.IsError())
+	{
+		return Outcome;
+	}
+	check(Outcome.IsOk());
+	bool bSocksConnected = (Outcome.GetResult() == 1);
+
+	// Condition the socket
+	if (!Candidate.SetBlocking(false))
+	{
+		return FOutcome::Error("Unable to set socket non-blocking");
+	}
+
+	if (int32 OptValue = GetBufferSize(FHost::EDirection::Send); OptValue >= 0)
+	{
+		Candidate.SetSendBufSize(OptValue);
+	}
+
+	if (int32 OptValue = GetBufferSize(FHost::EDirection::Recv); OptValue >= 0)
+	{
+		Candidate.SetRecvBufSize(OptValue);
+	}
+
+	// Socks connect in a blocking fashion so we're all set (ret=1)
+	if (bSocksConnected)
+	{
+		Socket = MoveTemp(Candidate);
+		return FOutcome::Ok();
+	}
+
+	// Issue the connect - this is done non-blocking so we need to wait (ret=0)
+	if (Outcome = Candidate.Connect(IpAddress, Port); Outcome.IsError())
+	{
+		return Outcome;
+	}
+
+	Socket = MoveTemp(Candidate);
+	return FOutcome::Waiting();
+}
+
+
+
+////////////////////////////////////////////////////////////////////////////////
+int32 FConnectionPool::FParams::SetHostFromUrl(FAnsiStringView Url)
+{
+	FUrlOffsets Offsets;
+	if (ParseUrl(Url, Offsets) < 0)
+	{
+		return -1;
+	}
+
+	HostName = Offsets.HostName.Get(Url);
+
+	if (Offsets.Port)
+	{
+		FAnsiStringView PortView = Offsets.Port.Get(Url);
+		Port = uint16(CrudeToInt(PortView));
+	}
+
+	return Offsets.Path;
+}
+
+////////////////////////////////////////////////////////////////////////////////
+FConnectionPool::FConnectionPool(const FParams& Params)
+{
+	check(Params.ConnectionCount - 1 <= 63u);
+	check(Params.Port <= 0xffffu);
+
+	// Alloc a new internal object
+	uint32 HostNameLen = Params.HostName.Len();
+	uint32 AllocSize = sizeof(FHost) + (HostNameLen + 1);
+	auto* Internal = (FHost*)FMemory::Malloc(AllocSize, alignof(FHost));
+
+	// Copy host
+	char* HostDest = (char*)(Internal + 1);
+	memcpy(HostDest, Params.HostName.GetData(), HostNameLen);
+	HostDest[HostNameLen] = '\0';
+
+	// Init internal object
+	new (Internal) FHost({
+		.HostName		= HostDest,
+		.Port			= Params.Port,
+		.MaxConnections	= Params.ConnectionCount,
+	});
+	Internal->SetBufferSize(FHost::EDirection::Send, Params.SendBufSize);
+	Internal->SetBufferSize(FHost::EDirection::Recv, Params.RecvBufSize);
+
+	Ptr = Internal;
+}
+
+////////////////////////////////////////////////////////////////////////////////
+FConnectionPool::~FConnectionPool()
+{
+	if (Ptr != nullptr)
+	{
+		FMemory::Free(Ptr);
+	}
+}
+
+////////////////////////////////////////////////////////////////////////////////
+bool FConnectionPool::Resolve()
+{
+	return Ptr->ResolveHostName().IsOk();
+}
+
+////////////////////////////////////////////////////////////////////////////////
+void FConnectionPool::Describe(FAnsiStringBuilderBase& OutString) const
+{
+	const FAnsiStringView HostName = Ptr->GetHostName();
+	OutString.Appendf("%.*s", HostName.Len(), HostName.GetData());
+	if (!!Ptr->IsResolved())
+	{
+		const auto IpAddress = Ptr->GetIpAddress();
+		OutString.Appendf(" (%u.%u.%u.%u)",
+						(IpAddress >> 24) & 0xff,
+						(IpAddress >> 16) & 0xff,
+						(IpAddress >> 8) & 0xff,
+						IpAddress & 0xff
+		);
+	}
+	else
+	{
+		OutString.Append(" (unresolved)");
+	}
+}
+
+////////////////////////////////////////////////////////////////////////////////
+bool FConnectionPool::IsValidHostUrl(FAnsiStringView Url)
+{
+	FUrlOffsets Tmp;
+	return ParseUrl(Url, Tmp) >= 0;
+}
+
+// }}}
+
+} // namespace UE::IoStore::HTTP
+
diff --git a/Engine/Source/Runtime/Experimental/IoStoreOnDemand/Private/CoreHttp/Http.cpp b/Engine/Source/Runtime/Experimental/IoStoreOnDemand/Private/CoreHttp/Http.cpp
--- a/Engine/Source/Runtime/Experimental/IoStoreOnDemand/Private/CoreHttp/Http.cpp
+++ b/Engine/Source/Runtime/Experimental/IoStoreOnDemand/Private/CoreHttp/Http.cpp
@@ -3,4284 +3,28 @@
 #include "Client.h"
 
 #if !defined(NO_UE_INCLUDES)
-#include "Containers/Array.h"
-#include "HAL/CriticalSection.h"
-#include "HAL/IConsoleManager.h"
-#include "IO/IoBuffer.h"
-#include "LatencyInjector.h"
-#include "Math/UnrealMathUtility.h"
-#include "Memory/MemoryView.h"
-#include "Misc/ScopeExit.h"
-#include "Misc/ScopeLock.h"
-#include "Misc/StringBuilder.h"
-#include "ProfilingDebugging/CpuProfilerTrace.h"
-#include "Tasks/Task.h"
-#include "Trace/Trace.h"
+#include <Containers/Array.h>
+#include <HAL/CriticalSection.h>
+#include <HAL/IConsoleManager.h>
+#include <IO/IoBuffer.h>
+#include <LatencyInjector.h>
+#include <Math/UnrealMathUtility.h>
+#include <Memory/MemoryView.h>
+#include <Misc/ScopeExit.h>
+#include <Misc/ScopeLock.h>
+#include <Misc/StringBuilder.h>
+#include <ProfilingDebugging/CpuProfilerTrace.h>
+#include <Tasks/Task.h>
+#include <Trace/Trace.h>
 #endif
 
 #include <atomic>
 
-// {{{1 platforms ..............................................................
-
-#if PLATFORM_MICROSOFT
-#	if !defined(NO_UE_INCLUDES)
-#	include "Windows/AllowWindowsPlatformTypes.h"
-#		include <winsock2.h>
-#		include <ws2tcpip.h>
-#	include "Windows/HideWindowsPlatformTypes.h"
-#	else
-#		include <winsock2.h>
-#		include <ws2tcpip.h>
-#		pragma comment(lib, "Ws2_32.lib")
-#	endif // NO_UE_INCLUDES
-	using SocketType	= SOCKET;
-	using MsgFlagType	= int;
-
-	enum { SHUT_RDWR = SD_BOTH };
-
-#	define IAS_HTTP_USE_POLL
-	template <typename... ArgTypes> auto poll(ArgTypes... Args)
-	{
-		return WSAPoll(Forward<ArgTypes>(Args)...);
-	}
-
-	// Winsock defines "PF_MAX" indicating the protocol family count. This
-	// however competes with UE definitions related to pixel formats.
-	#if defined(PF_MAX)
-	#	undef PF_MAX
-	#endif
-
-#elif PLATFORM_APPLE | PLATFORM_UNIX | PLATFORM_ANDROID
-#	include <arpa/inet.h>
-#	include <fcntl.h>
-#	include <netdb.h>
-#	include <poll.h>
-#	include <sys/select.h>
-#	include <sys/socket.h>
-#	include <unistd.h>
-	using SocketType	= int;
-	using MsgFlagType	= int;
-	static int32 closesocket(int32 Socket) { return close(Socket); }
-#	define IAS_HTTP_USE_POLL
-#else
-#	include "CoreHttp/Http.inl"
-#endif
-
-static_assert(sizeof(sockaddr_in::sin_addr) == sizeof(uint32));
-
-#if PLATFORM_MICROSOFT
-	static int32 LastSocketResult() { return WSAGetLastError(); }
-#	define IsSocketResult(err)		(LastSocketResult() == WSA##err)
-#else
-	static int32 LastSocketResult() { return errno; }
-#	define IsSocketResult(err)		(LastSocketResult() == err)
-	static_assert(EWOULDBLOCK == EAGAIN);
-#endif
-
-static const SocketType InvalidSocket = ~SocketType(0);
-
-// }}}
-
-namespace UE::IoStore::HTTP
-{
-
-// {{{1 trace ..................................................................
-
-////////////////////////////////////////////////////////////////////////////////
-#define MAKE_TRACE_ENUM(x) \
-	x(LoopCreate) \
-	x(LoopTick) \
-	x(LoopDestroy) \
-	x(ActivityCreate) \
-	x(ActivityDestroy) \
-	x(SocketCreate) \
-	x(SocketDestroy) \
-	x(RequestBegin) \
-	x(StateChange) \
-	x(Wait) \
-	x(Unwait) \
-	x(Connect) \
-	x(Send) \
-	x(Recv) \
-	x(StartWork) \
-	x($)
-
-enum class ETrace
-{
-#define TRACE_ENUM_VALUES(x) x,
-	MAKE_TRACE_ENUM(TRACE_ENUM_VALUES)
-#undef TRACE_ENUM_VALUES
-};
-
-static const FAnsiStringView TraceEnumNames[] = {
-#define TRACE_ENUM_VALUES(x) #x ,
-	MAKE_TRACE_ENUM(TRACE_ENUM_VALUES)
-#undef TRACE_ENUM_VALUES
-};
-
-#undef MAKE_TRACE_ENUM
-
-UE_TRACE_CHANNEL(IasHttpChannel);
-
-UE_TRACE_EVENT_BEGIN(IasHttp, Enum, NoSync)
-	UE_TRACE_EVENT_FIELD(UE::Trace::AnsiString, Name)
-UE_TRACE_EVENT_END()
-
-UE_TRACE_EVENT_BEGIN(IasHttp, Event, NoSync)
-	UE_TRACE_EVENT_FIELD(uint64, Cycle)
-	UE_TRACE_EVENT_FIELD(uint32, Id)
-	UE_TRACE_EVENT_FIELD(uint32, Param)
-	UE_TRACE_EVENT_FIELD(uint8, Action)
-UE_TRACE_EVENT_END()
-
-static void Activity_TraceStateNames();
-
-////////////////////////////////////////////////////////////////////////////////
-static void TraceEnum(const FAnsiStringView* Names)
-{
-	for (;; ++Names)
-	{
-		UE_TRACE_LOG(IasHttp, Enum, IasHttpChannel)
-			<< Enum.Name(Names[0].GetData(), Names[0].Len());
-
-		if (*Names == "$")
-		{
-			break;
-		}
-	}
-}
-
-////////////////////////////////////////////////////////////////////////////////
-static void TraceInner(const UPTRINT Id, ETrace Action, UPTRINT Param)
-{
-#if UE_TRACE_ENABLED
-	static bool Once = [] {
-		TraceEnum(TraceEnumNames);
-		Activity_TraceStateNames();
-		return true;
-	}();
-
-	UE_TRACE_LOG(IasHttp, Event, IasHttpChannel)
-		<< Event.Cycle(FPlatformTime::Cycles64())
-		<< Event.Id(uint32(Id))
-		<< Event.Param(uint32(Param))
-		<< Event.Action(uint8(Action));
-#endif // UE_TRACE_ENABLED
-}
-
-template <typename T, typename U=UPTRINT>
-static void Trace(T Id, ETrace Action, U Param=0)
-{
-	TraceInner(UPTRINT(Id), Action, UPTRINT(Param));
-}
-
-
-
-// {{{1 misc ...................................................................
-
-////////////////////////////////////////////////////////////////////////////////
-#define IAS_CVAR(Type, Name, Default, Desc, ...) \
-	Type G##Name = Default; \
-	static FAutoConsoleVariableRef CVar_Ias##Name( \
-		TEXT("ias.Http" #Name), \
-		G##Name, \
-		TEXT(Desc) \
-		__VA_ARGS__ \
-	)
-
-////////////////////////////////////////////////////////////////////////////////
-static IAS_CVAR(int32, RecvWorkThresholdKiB,80,		"Threshold of data remaining at which next request is sent (in KiB)");
-static IAS_CVAR(int32, IdleMs,				50'000,	"Time in seconds to close idle connections or fail waits");
-
-////////////////////////////////////////////////////////////////////////////////
-class FOutcome
-{
-public:
-	static FOutcome Ok(uint32 Result=0);
-	static FOutcome Waiting();
-	static FOutcome Error(const char* Message, int32 Code=-1);
-	static FOutcome None()				{ return Error(""); }
-	bool			IsError() const		{ return Message < 0x8000'0000'0000; }
-	bool			IsWaiting() const	{ return Tag == WaitTag; }
-	bool			IsOk() const		{ return Tag == OkTag; }
-	FAnsiStringView GetMessage() const	{ check(IsError()); return (const char*)(Message); }
-	int32			GetErrorCode() const{ check(IsError()); return int32(Code); }
-	uint32			GetResult() const	{ check(IsOk()); return Result; }
-
-private:
-					FOutcome() = default;
-
-	static uint32 const OkTag	= 0x0000'8000;
-	static uint32 const WaitTag = 0x0001'8000;
-
-	union {
-		struct {
-			UPTRINT	Message : 48;
-			PTRINT	Code	: 16;
-		};
-		struct {
-			uint32	Result;
-			uint32	Tag;
-		};
-	};
-};
-static_assert(sizeof(FOutcome) == sizeof(void*));
-
-////////////////////////////////////////////////////////////////////////////////
-FOutcome FOutcome::Ok(uint32 Result)
-{
-	FOutcome Outcome;
-	Outcome.Tag = OkTag;
-	Outcome.Result = Result;
-	return Outcome;
-}
-
-////////////////////////////////////////////////////////////////////////////////
-FOutcome FOutcome::Waiting()
-{
-	FOutcome Outcome;
-	Outcome.Tag = WaitTag;
-	return Outcome;
-}
-
-////////////////////////////////////////////////////////////////////////////////
-FOutcome FOutcome::Error(const char* Message, int32 Code)
-{
-	check(Message != nullptr);
-	check(Code <= 0xffff && Code >= -0xffff);
-
-	FOutcome Outcome;
-	Outcome.Message = UPTRINT(Message);
-	Outcome.Code = int16(Code);
-	return Outcome;
-}
-
-
-
-////////////////////////////////////////////////////////////////////////////////
-template <typename LambdaType>
-static void EnumerateHeaders(FAnsiStringView Headers, LambdaType&& Lambda)
-{
-	// NB. here we are assuming that we will be dealing with servers that will
-	// not be returning headers with "obsolete line folding".
-
-	auto IsOws = [] (int32 c) { return (c == ' ') | (c == '\t'); };
-
-	const char* Cursor = Headers.GetData();
-	const char* End = Cursor + Headers.Len();
-	do
-	{
-		int32 ColonIndex = 0;
-		for (; Cursor + ColonIndex < End; ++ColonIndex)
-		{
-			if (Cursor[ColonIndex] == ':')
-			{
-				break;
-			}
-		}
-
-		Cursor += ColonIndex;
-
-		const char* Right = Cursor + 1;
-		while (Right < End)
-		{
-			if (Right[0] != '\r' || Right[1] != '\n')
-			{
-				Right += 1 + (Right[1] != '\r');
-				continue;
-			}
-
-			FAnsiStringView Name(Cursor - ColonIndex, ColonIndex);
-
-			const char* Left = Cursor + 1;
-			for (; IsOws(Left[0]); ++Left);
-
-			Cursor = Right;
-			for (; Cursor > Left + 1 && IsOws(Cursor[-1]); --Cursor);
-
-			FAnsiStringView Value (Left, int32(ptrdiff_t(Cursor - Left)));
-
-			if (!Lambda(Name, Value))
-			{
-				Right = End;
-			}
-
-			break;
-		}
-
-		Cursor = Right + 2;
-	} while (Cursor < End);
-}
-
-////////////////////////////////////////////////////////////////////////////////
-static int32 FindMessageTerminal(const char* Data, uint32 Length)
-{
-	for (uint32 i = 4; i <= Length; ++i)
-	{
-		uint32 Candidate;
-		::memcpy(&Candidate, Data + i - 4, sizeof(Candidate));
-		if (Candidate == 0x0a0d0a0d)
-		{
-			return i;
-		}
-
-		i += (Data[i - 1] > 0x0d) ? 3 : 0;
-	}
-
-	return -1;
-}
-
-////////////////////////////////////////////////////////////////////////////////
-static int64 CrudeToInt(FAnsiStringView View)
-{
-	// FCStringAnsi::* is not used to mitigate any locale hiccups. By
-	// initialising 'Value' with MSB set we can detect cases where View did not
-	// start with digits. This works as we won't be using this on huge numbers.
-	int64 Value = 0x8000'0000'0000'0000ll;
-	for (int32 c : View)
-	{
-		uint32 Digit = c - '0';
-		if (Digit > 9u)
-		{
-			break;
-		}
-		Value *= 10;
-		Value += Digit;
-	}
-	return Value;
-};
-
-////////////////////////////////////////////////////////////////////////////////
-struct FMessageOffsets
-{
-	uint8	StatusCode;
-	uint8	Message;
-	uint16	Headers;
-};
-
-static int32 ParseMessage(FAnsiStringView Message, FMessageOffsets& Out)
-{
-	const FAnsiStringView Protocol("HTTP/1.1 ");
-
-	// Check there's enough data
-	if (Message.Len() < Protocol.Len() + 1) // "+1" accounts for at least one digit
-	{
-		return -1;
-	}
-
-	const char* Cursor = Message.GetData();
-
-	// Check for the expected protocol
-	if (FAnsiStringView(Cursor, 9) != Protocol)
-	{
-		return -1;
-	}
-	int32 i = Protocol.Len();
-
-	// Trim left and tightly reject anything adventurous
-	for (int n = 32; Cursor[i] == ' ' && i < n; ++i);
-	Out.StatusCode = uint8(i);
-
-	// At least one status line digit. (Note to self; expect exactly three)
-	for (int n = 32; uint32(Cursor[i] - 0x30) <= 9 && i < n; ++i);
-	if (uint32(i - Out.StatusCode - 1) > 32)
-	{
-		return -1;
-	}
-
-	// Trim left
-	for (int n = 32; Cursor[i] == ' ' && i < n; ++i);
-	Out.Message = uint8(i);
-
-	// Extra conservative length allowance
-	if (i > 32)
-	{
-		return -1;
-	}
-
-	// Find \r\n
-	for (; Cursor[i] != '\r'; ++i)
-	{
-		if (i >= 2048)
-		{
-			return -1;
-		}
-	}
-	if (Cursor[i + 1] != '\n')
-	{
-		return -1;
-	}
-	Out.Headers = uint16(i + 2);
-
-	return 1;
-}
-
-////////////////////////////////////////////////////////////////////////////////
-struct FUrlOffsets
-{
-	struct Slice
-	{
-						Slice() = default;
-						Slice(int32 l, int32 r) : Left(uint8(l)), Right(uint8(r)) {}
-		FAnsiStringView	Get(FAnsiStringView Url) const { return Url.Mid(Left, Right - Left); }
-						operator bool () const { return Left > 0; }
-		int32			Len() const { return Right - Left; }
-		uint8			Left;
-		uint8			Right;
-	};
-	Slice				UserInfo;
-	Slice				HostName;
-	Slice				Port;
-	uint8				Path;
-	uint8				SchemeLength;
-};
-
-static int32 ParseUrl(FAnsiStringView Url, FUrlOffsets& Out)
-{
-	if (Url.Len() < 5)
-	{
-		return -1;
-	}
-
-	Out = {};
-
-	const char* Start = Url.GetData();
-	const char* Cursor = Start;
-
-	// Scheme
-	int32 i = 0;
-	for (; i < 5; ++i)
-	{
-		if (uint32(Cursor[i] - 'a') > uint32('z' - 'a'))
-		{
-			break;
-		}
-	}
-
-	Out.SchemeLength = uint8(i);
-	FAnsiStringView Scheme = Url.Left(i);
-	if (Scheme != "http" && Scheme != "https")
-	{
-		return -1;
-	}
-
-	// Separator and authority
-	if (Cursor[i] != ':' || Cursor[i + 1] != '/' || Cursor[i + 2] != '/')
-	{
-		return -1;
-	}
-	i += 3;
-
-	struct { int32 c; int32 i; } Seps[2];
-	int32 SepCount = 0;
-	for (; i < Url.Len(); ++i)
-	{
-		int32 c = Cursor[i];
-		if (c < '-')							break;
-		if (c != ':' && c != '@' && c != '/')	continue;
-		if (c == '/' || SepCount >= 2)			break;
-
-		if (c == '@' && SepCount)
-		{
-			SepCount -= (Seps[SepCount - 1].c == ':');
-		}
-		Seps[SepCount++] = { c, i };
-	}
-
-	if (i > 0xff || i <= Scheme.Len() + 3)
-	{
-		return -1;
-	}
-
-	if (i < Url.Len())
-	{
-		Out.Path = uint8(i);
-	}
-
-	Out.HostName = { uint8(Scheme.Len() + 3), uint8(i) };
-
-	switch (SepCount)
-	{
-	case 0:
-		break;
-
-	case 1:
-		if (Seps[0].c == ':')
-		{
-			Out.Port = { Seps[0].i + 1, i };
-			Out.HostName.Right = uint8(Seps[0].i);
-		}
-		else
-		{
-			Out.UserInfo = { Out.HostName.Left, Seps[0].i };
-			Out.HostName.Left += uint8(Seps[0].i + 1);
-			Out.HostName.Right += uint8(Seps[0].i + 1);
-		}
-		break;
-
-	case 2:
-		if ((Seps[0].c != '@') | (Seps[1].c != ':'))
-		{
-			return -1;
-		}
-		Out.UserInfo = { Out.HostName.Left, Seps[0].i };
-		Out.Port.Left = uint8(Seps[1].i + 1);
-		Out.Port.Right = Out.HostName.Right;
-		Out.HostName.Left = Out.UserInfo.Right + 1;
-		Out.HostName.Right = Out.Port.Left - 1;
-		break;
-
-	default:
-		return -1;
-	}
-
-	bool Bad = false;
-	Bad |= (Out.HostName.Len() == 0);
-	Bad |= bool(Out.UserInfo) & (Out.UserInfo.Len() == 0);
-
-	if (Out.Port.Left)
-	{
-		Bad |= (Out.Port.Len() == 0);
-		for (int32 j = 0, n = Out.Port.Len(); j < n; ++j)
-		{
-			Bad |= (uint32(Start[Out.Port.Left + j] - '0') > 9);
-		}
-	}
-
-	return Bad ? -1 : 1;
-}
-
-
-
-// {{{1 buffer .................................................................
-
-////////////////////////////////////////////////////////////////////////////////
-class alignas(16) FBuffer
-{
-public:
-	struct FMutableSection
-	{
-		char*		Data;
-		uint32		Size;
-	};
-
-								FBuffer() = default;
-								FBuffer(char* InData, uint32 InMax);
-								~FBuffer();
-	FBuffer&					operator = (FBuffer&& Rhs);
-	void						Fix();
-	void						Resize(uint32 Size);
-	const char*					GetData() const;
-	uint32						GetSize() const;
-	uint32						GetCapacity() const;
-	template <typename T> T*	Alloc(uint32 Count=1);
-	FMutableSection				GetMutableFree(uint32 MinSize, uint32 PageSize=256);
-	void						AdvanceUsed(uint32 Delta);
-
-private:
-	char*						GetDataPtr();
-	void						Extend(uint32 AtLeast, uint32 PageSize);
-	union
-	{
-		struct
-		{
-			UPTRINT				Data : 63;
-			UPTRINT				Inline : 1;
-		};
-		UPTRINT					DataInline = 0;
-	};
-	uint32						Max = 0;
-	uint32						Used = 0;
-
-private:
-								FBuffer(FBuffer&&) = delete;
-								FBuffer(const FBuffer&) = delete;
-	FBuffer&					operator = (const FBuffer&) = delete;
-};
-
-////////////////////////////////////////////////////////////////////////////////
-FBuffer::FBuffer(char* InData, uint32 InMax)
-: Data(UPTRINT(InData))
-, Inline(1)
-, Max(InMax)
-{
-}
-
-////////////////////////////////////////////////////////////////////////////////
-FBuffer::~FBuffer()
-{
-	if (Data && !Inline)
-	{
-		FMemory::Free(GetDataPtr());
-	}
-}
-
-////////////////////////////////////////////////////////////////////////////////
-FBuffer& FBuffer::operator = (FBuffer&& Rhs)
-{
-	Swap(DataInline, Rhs.DataInline);
-	Swap(Max, Rhs.Max);
-	Swap(Used, Rhs.Used);
-	return *this;
-}
-
-////////////////////////////////////////////////////////////////////////////////
-void FBuffer::Fix()
-{
-	check(Inline);
-	Data += Used;
-	Max -= Used;
-	Used = 0;
-}
-
-////////////////////////////////////////////////////////////////////////////////
-void FBuffer::Resize(uint32 Size)
-{
-	check(Size <= Max);
-	Used = Size;
-}
-
-////////////////////////////////////////////////////////////////////////////////
-char* FBuffer::GetDataPtr()
-{
-	return (char*)Data;
-}
-
-////////////////////////////////////////////////////////////////////////////////
-const char* FBuffer::GetData() const
-{
-	return (char*)Data;
-}
-
-////////////////////////////////////////////////////////////////////////////////
-uint32 FBuffer::GetSize() const
-{
-	return Used;
-}
-
-////////////////////////////////////////////////////////////////////////////////
-uint32 FBuffer::GetCapacity() const
-{
-	return Max;
-}
-
-////////////////////////////////////////////////////////////////////////////////
-template <typename T>
-T* FBuffer::Alloc(uint32 Count)
-{
-	uint32 AlignBias = uint32(Data) & (alignof(T) - 1);
-	if (AlignBias)
-	{
-		AlignBias = alignof(T) - AlignBias;
-	}
-
-	uint32 PotentialUsed = Used + AlignBias + (sizeof(T) * Count);
-	if (PotentialUsed > Max)
-	{
-		Extend(PotentialUsed, 256);
-	}
-
-	void* Ret = GetDataPtr() + Used + AlignBias;
-	Used = PotentialUsed;
-	return (T*)Ret;
-}
-
-////////////////////////////////////////////////////////////////////////////////
-FBuffer::FMutableSection FBuffer::GetMutableFree(uint32 MinSize, uint32 PageSize)
-{
-	MinSize = (MinSize == 0 && Used == Max) ? PageSize : MinSize;
-
-	uint32 PotentialUsed = Used + MinSize;
-	if (PotentialUsed > Max)
-	{
-		Extend(PotentialUsed, PageSize);
-	}
-
-	return FMutableSection{ GetDataPtr() + Used, Max - Used };
-}
-
-////////////////////////////////////////////////////////////////////////////////
-void FBuffer::AdvanceUsed(uint32 Delta)
-{
-	Used += Delta;
-	check(Used <= Max);
-}
-
-////////////////////////////////////////////////////////////////////////////////
-void FBuffer::Extend(uint32 AtLeast, uint32 PageSize)
-{
-	checkSlow((PageSize - 1 & PageSize) == 0);
-
-	--PageSize;
-	Max = (AtLeast + PageSize) & ~PageSize;
-
-	if (!Inline)
-	{
-		Data = UPTRINT(FMemory::Realloc(GetDataPtr(), Max, alignof(FBuffer)));
-		return;
-	}
-
-	const char* PrevData = GetDataPtr();
-	Data = UPTRINT(FMemory::Malloc(Max, alignof(FBuffer)));
-	::memcpy(GetDataPtr(), PrevData, Used);
-	Inline = 0;
-}
-
-
-
-////////////////////////////////////////////////////////////////////////////////
-class FMessageBuilder
-{
-public:
-						FMessageBuilder(FBuffer& Ruffeb);
-	FMessageBuilder&	operator << (FAnsiStringView Lhs);
-
-private:
-	FBuffer&		Buffer;
-};
-
-////////////////////////////////////////////////////////////////////////////////
-FMessageBuilder::FMessageBuilder(FBuffer& Ruffeb)
-: Buffer(Ruffeb)
-{
-}
-
-////////////////////////////////////////////////////////////////////////////////
-FMessageBuilder& FMessageBuilder::operator << (FAnsiStringView Lhs)
-{
-	uint32 Length = uint32(Lhs.Len());
-	FBuffer::FMutableSection Section = Buffer.GetMutableFree(Length);
-	::memcpy(Section.Data, Lhs.GetData(), Length);
-	Buffer.AdvanceUsed(Length);
-	return *this;
-}
-
-
-
-// {{{1 socket .................................................................
-
-////////////////////////////////////////////////////////////////////////////////
-class FSocket
-{
-public:
-	struct FWaiter
-	{
-		enum class EWhat { Send = 0b01, Recv = 0b10, Both = Send|Recv };
-				FWaiter() { std::memset(this, 0, sizeof(*this)); }
-				FWaiter(const FSocket& Socket, EWhat InWaitOn);
-		bool	IsValid() const { UPTRINT x{0}; return std::memcmp(this, &x, sizeof(*this)) != 0; }
-		bool	operator == (FSocket& Rhs) const { return UPTRINT(&Rhs) == Candidate; }
-		UPTRINT	Candidate : 60;
-		UPTRINT	WaitOn : 2;
-		UPTRINT	Ready : 2;
-	};
-
-				FSocket() = default;
-				~FSocket()					{ Destroy(); }
-				FSocket(FSocket&& Rhs)		{ Move(MoveTemp(Rhs)); }
-	FSocket&	operator = (FSocket&& Rhs)	{ Move(MoveTemp(Rhs)); return *this; }
-	bool		IsValid() const				{ return Socket != InvalidSocket; }
-	bool		Create();
-	void		Destroy();
-	FOutcome	Connect(uint32 Ip, uint32 Port);
-	void		Disconnect();
-	FOutcome	Send(const char* Data, uint32 Size);
-	FOutcome	Recv(char* Dest, uint32 Size);
-	bool		SetBlocking(bool bBlocking);
-	bool		SetSendBufSize(int32 Size);
-	bool		SetRecvBufSize(int32 Size);
-	static int	Wait(TArrayView<FWaiter> Waiters, int32 TimeoutMs);
-
-private:
-	void		Move(FSocket&& Rhs);
-	SocketType	Socket = InvalidSocket;
-
-				FSocket(const FSocket&) = delete;
-	FSocket&	operator = (const FSocket&) = delete;
-};
-
-////////////////////////////////////////////////////////////////////////////////
-FSocket::FWaiter::FWaiter(const FSocket& Socket, EWhat InWaitOn)
-: Candidate(UPTRINT(&Socket))
-, WaitOn(UPTRINT(InWaitOn))
-, Ready(0)
-{
-	static_assert(sizeof(*this) == sizeof(void*));
-}
-
-////////////////////////////////////////////////////////////////////////////////
-void FSocket::Move(FSocket&& Rhs)
-{
-	check(!IsValid() || !Rhs.IsValid()); // currently we only want to pass one around
-	Swap(Socket, Rhs.Socket);
-}
-
-////////////////////////////////////////////////////////////////////////////////
-bool FSocket::Create()
-{
-	check(!IsValid());
-	Socket = socket(AF_INET, SOCK_STREAM, IPPROTO_TCP);
-
-	if (!IsValid())
-	{
-		return false;
-	}
-
-	Trace(Socket, ETrace::SocketCreate);
-	return true;
-}
-
-////////////////////////////////////////////////////////////////////////////////
-void FSocket::Destroy()
-{
-	if (Socket == InvalidSocket)
-	{
-		return;
-	}
-
-	Trace(Socket, ETrace::SocketDestroy);
-
-	closesocket(Socket);
-	Socket = InvalidSocket;
-}
-
-////////////////////////////////////////////////////////////////////////////////
-FOutcome FSocket::Connect(uint32 IpAddress, uint32 Port)
-{
-	check(IsValid());
-
-	Trace(Socket, ETrace::Connect, IpAddress);
-
-	IpAddress = htonl(IpAddress);
-
-	sockaddr_in AddrInet = { sizeof(sockaddr_in) };
-	AddrInet.sin_family = AF_INET;
-	AddrInet.sin_port = htons(uint16(Port));
-	memcpy(&(AddrInet.sin_addr), &IpAddress, sizeof(IpAddress));
-
-	int32 Result = connect(Socket, &(sockaddr&)AddrInet, sizeof(AddrInet));
-
-	if (IsSocketResult(EWOULDBLOCK) | IsSocketResult(EINPROGRESS))
-	{
-		return FOutcome::Waiting();
-	}
-
-	if (Result < 0)
-	{
-		return FOutcome::Error("Socket connect failed", Result);
-	}
-
-	return FOutcome::Ok();
-}
-
-////////////////////////////////////////////////////////////////////////////////
-void FSocket::Disconnect()
-{
-	check(IsValid());
-	shutdown(Socket, SHUT_RDWR);
-}
-
-////////////////////////////////////////////////////////////////////////////////
-FOutcome FSocket::Send(const char* Data, uint32 Size)
-{
-	Trace(Socket, ETrace::Send, -1);
-	int32 Result = send(Socket, Data, Size, MsgFlagType(0));
-	Trace(Socket, ETrace::Send, FMath::Max(Result, 0));
-
-	if (Result > 0)						return FOutcome::Ok(Result);
-	if (Result == 0)					return FOutcome::Error("Send ATH0");
-	if (IsSocketResult(EWOULDBLOCK))	return FOutcome::Waiting();
-
-	if (IsSocketResult(ENOTCONN))
-	{
-		int32 Error = 0;
-		socklen_t ErrorSize = sizeof(Error);
-		Result = getsockopt(Socket, SOL_SOCKET, SO_ERROR, (char*)&Error, &ErrorSize);
-		if (Result < 0 || Error != 0)
-		{
-			return FOutcome::Error("Error while connecting");
-		}
-
-		return FOutcome::Waiting();
-	}
-
-	Result = LastSocketResult();
-	return FOutcome::Error("Send", Result);
-}
-
-////////////////////////////////////////////////////////////////////////////////
-FOutcome FSocket::Recv(char* Dest, uint32 Size)
-{
-	Trace(Socket, ETrace::Recv, -1);
-	int32 Result = recv(Socket, Dest, Size, MsgFlagType(0));
-	Trace(Socket, ETrace::Recv, FMath::Max(0, Result));
-
-	if (Result > 0)						return FOutcome::Ok(Result);
-	if (Result == 0)					return FOutcome::Error("Recv ATH0");
-	if (IsSocketResult(EWOULDBLOCK))	return FOutcome::Waiting();
-
-	Result = LastSocketResult();
-	return FOutcome::Error("Recv", Result);
-}
-
-////////////////////////////////////////////////////////////////////////////////
-bool FSocket::SetBlocking(bool bBlocking)
-{
-#if defined(IAS_HTTP_HAS_NONBLOCK_IMPL)
-	return SetNonBlockingSocket(Socket, !bBlocking);
-#elif PLATFORM_MICROSOFT
-	unsigned long NonBlockingMode = (bBlocking != true);
-	return (ioctlsocket(Socket, FIONBIO, &NonBlockingMode) != SOCKET_ERROR);
-#else
-	int32 Flags = fcntl(Socket, F_GETFL, 0);
-	if (Flags == -1)
-	{
-		return false;
-	}
-
-	int32 NewFlags = bBlocking
-		? (Flags & ~int32(O_NONBLOCK))
-		: (Flags |	int32(O_NONBLOCK));
-	return (Flags == NewFlags) || (fcntl(Socket, F_SETFL, Flags) >= 0);
-#endif
-}
-
-////////////////////////////////////////////////////////////////////////////////
-bool FSocket::SetSendBufSize(int32 Size)
-{
-	return 0 == setsockopt(Socket, SOL_SOCKET, SO_SNDBUF, &(char&)Size, sizeof(Size));
-}
-
-////////////////////////////////////////////////////////////////////////////////
-bool FSocket::SetRecvBufSize(int32 Size)
-{
-	return 0 == setsockopt(Socket, SOL_SOCKET, SO_RCVBUF, &(char&)Size, sizeof(Size));
-}
-
-////////////////////////////////////////////////////////////////////////////////
-int32 FSocket::Wait(TArrayView<FWaiter> Waiters, int32 TimeoutMs)
-{
-#if !defined(IAS_HTTP_USE_POLL)
-	struct FSelect
-	{
-		SocketType	fd;
-		int32		events;
-		int32		revents;
-	};
-	static const int32 POLLIN  = 1 << 0;
-	static const int32 POLLOUT = 1 << 1;
-	static const int32 POLLERR = 1 << 2;
-	static const int32 POLLHUP = POLLERR;
-	static const int32 POLLNVAL= POLLERR;
-#else
-	using FSelect = pollfd;
-#endif
-
-	// The following looks odd because POLLFD varies subtly from one platform
-	// to the next. To cleanly set members to zero and to not get narrowing
-	// warnings from the compiler, we list-init and don't assume POD types.
-	using PollEventType = decltype(FSelect::events);
-	PollEventType Events[] = { POLLERR, POLLOUT, POLLIN, POLLOUT|POLLOUT };
-	TArray<FSelect, TFixedAllocator<64>> Selects;
-	for (FWaiter& Waiter : Waiters)
-	{
-		Selects.Emplace_GetRef() = {
-			((FSocket*)Waiter.Candidate)->Socket,
-			Events[Waiter.WaitOn],
-			/* 0 */
-		};
-	}
-
-	// Poll the sockets
-#if defined(IAS_HTTP_USE_POLL)
-	int32 Result = poll(Selects.GetData(), Selects.Num(), TimeoutMs);
-	if (Result <= 0)
-	{
-		return Result;
-	}
-#else
-	timeval TimeVal = {};
-	timeval* TimeValPtr = (TimeoutMs >= 0 ) ? &TimeVal : nullptr;
-	if (TimeoutMs > 0)
-	{
-		TimeVal = { TimeoutMs >> 10, TimeoutMs & ((1 << 10) - 1) };
-	}
-
-	fd_set FdSetRead;	FD_ZERO(&FdSetRead);
-	fd_set FdSetWrite;	FD_ZERO(&FdSetWrite);
-	fd_set FdSetExcept; FD_ZERO(&FdSetExcept);
-
-	SocketType MaxFd = 0;
-	for (FSelect& Select : Selects)
-	{
-		fd_set* RwSet = (Select.events & POLLIN) ? &FdSetRead : &FdSetWrite;
-		FD_SET(Select.fd, RwSet);
-		FD_SET(Select.fd, &FdSetExcept);
-		MaxFd = FMath::Max(Select.fd, MaxFd);
-	}
-
-	int32 Result = select(int32(MaxFd + 1), &FdSetRead, &FdSetWrite, &FdSetExcept, TimeValPtr);
-	if (Result <= 0)
-	{
-		return Result;
-	}
-
-	for (FSelect& Select : Selects)
-	{
-		if (FD_ISSET(Select.fd, &FdSetExcept))
-		{
-			Select.revents = POLLERR;
-			continue;
-		}
-
-		fd_set* RwSet = (Select.events & POLLIN) ? &FdSetRead : &FdSetWrite;
-		if (FD_ISSET(Select.fd, RwSet))
-		{
-			Select.revents = Select.events;
-		}
-	}
-#endif // IAS_HTTP_USE_POLL
-
-	// Transfer poll results to the input sockets. We don't transfer across error
-	// states. Subsequent sockets ops can take care of that instead.
-	static const auto TestBits = POLLIN|POLLOUT|POLLERR|POLLHUP|POLLNVAL;
-	for (uint32 i = 0, n = Waiters.Num(); i < n; ++i)
-	{
-		auto RetEvents = Selects[i].revents;
-		if (!(RetEvents & TestBits))
-		{
-			continue;
-		}
-
-		uint32 Value = 0;
-		if (!!(RetEvents & POLLOUT)) Value |= uint32(FSocket::FWaiter::EWhat::Send);
-		if (!!(RetEvents & POLLIN))	 Value |= uint32(FSocket::FWaiter::EWhat::Recv);
-		Waiters[i].Ready = Value ? Value : uint32(FSocket::FWaiter::EWhat::Both);
-	}
-
-	return Result;
-}
-
-
-
-// {{{1 socks ..................................................................
-
-#if !UE_BUILD_SHIPPING
-
-///////////////////////////////////////////////////////////////////////////////
-static IAS_CVAR(int32,		SocksVersion,	5,		"SOCKS proxy protocol version to use");
-static IAS_CVAR(FString,	SocksIp,		"",		"Routes all IAS HTTP traffic through the given SOCKS proxy");
-static IAS_CVAR(int32,		SocksPort,		1080,	"Port of the SOCKS proxy to use");
-
-////////////////////////////////////////////////////////////////////////////////
-static uint32 GetSocksIpAddress()
-{
-	const TCHAR* Value = *GSocksIp;
-	uint32 IpAddress = 0;
-	uint32 Accumulator = 0;
-	while (true)
-	{
-		uint32 c = *Value++;
-
-		if (c - '0' <= '9' - '0')
-		{
-			Accumulator *= 10;
-			Accumulator += (c - '0');
-			continue;
-		}
-
-		if (c == '.' || c == '\0')
-		{
-			IpAddress <<= 8;
-			IpAddress |= Accumulator;
-			Accumulator = 0;
-			if (c == '\0')
-			{
-				break;
-			}
-			continue;
-		}
-
-		return 0;
-	}
-	return IpAddress;
-}
-
-////////////////////////////////////////////////////////////////////////////////
-static FOutcome ConnectSocks4(FSocket& Socket, uint32 IpAddress, uint32 Port)
-{
-	struct FSocks4Request
-	{
-		uint8	Version = 4;
-		uint8	Command = 1;
-		uint16	Port;
-		uint32	IpAddress;
-	};
-
-	struct FSocks4Reply
-	{
-		uint8	Version;
-		uint8	Code;
-		uint16	Port;
-		uint32	IpAddress;
-	};
-
-	uint32 SocksIpAddress = GetSocksIpAddress();
-	if (!SocksIpAddress)
-	{
-		return FOutcome::Error("Invalid socks IP address");
-	}
-
-	FOutcome Outcome = FOutcome::None();
-
-	if (Outcome = Socket.Connect(SocksIpAddress, GSocksPort); Outcome.IsError())
-	{
-		return Outcome;
-	}
-
-	FSocks4Request Request = {
-		.Port		= htons(uint16(Port)),
-		.IpAddress	= htonl(IpAddress),
-	};
-	if (Outcome = Socket.Send((const char*)&Request, sizeof(Request)); Outcome.IsError())
-	{
-		return Outcome;
-	}
-
-	FSocks4Reply Reply;
-	if (Outcome = Socket.Recv((char*)&Reply, sizeof(Reply)); Outcome.IsError())
-	{
-		return Outcome;
-	}
-
-	return FOutcome::Ok(1);
-}
-
-////////////////////////////////////////////////////////////////////////////////
-static FOutcome ConnectSocks5(FSocket& Socket, uint32 IpAddress, uint32 Port)
-{
-#ifdef _MSC_VER
-	// MSVC's static analysis doesn't see that 'Result' from recv() is checked
-	// to be the exact size of the destination buffer.
-#pragma warning(push)
-#pragma warning(disable : 6385)
-#endif
-
-	uint32 SocksIpAddress = GetSocksIpAddress();
-	if (!SocksIpAddress)
-	{
-		return FOutcome::Error("Invalid socks5 IP address");
-	}
-
-	FOutcome Outcome = FOutcome::None();
-
-	if (Outcome = Socket.Connect(SocksIpAddress, GSocksPort); Outcome.IsError())
-	{
-		return Outcome;
-	}
-
-	// Greeting
-	const char Greeting[] = { 5, 1, 0 };
-	Outcome = Socket.Send(Greeting, sizeof(Greeting));
-	if (!Outcome.IsOk())						return Outcome;
-	if (Outcome.GetResult() != sizeof(Greeting))return FOutcome::Error("Could not send socks5 greeting");
-
-	// Server auth-choice
-	char Choice[1 + 1];
-	Outcome = Socket.Recv(Choice, sizeof(Choice));
-	if (!Outcome.IsOk())						return Outcome;
-	if (Outcome.GetResult() != sizeof(Choice))	return FOutcome::Error("Recv too short from socks5 server");
-	if (Choice[0] != 0x05 || Choice[1] != 0x00) return FOutcome::Error("Got unexpected socks5 version from server");
-
-	// Connection request
-	IpAddress = htonl(IpAddress);
-	uint16 NsPort = htons(uint16(Port));
-	char Request[] = { 5, 1, 0, 1, 0x11,0x11,0x11,0x11, 0x22,0x22 };
-	std::memcpy(Request + 4, &IpAddress, sizeof(IpAddress));
-	std::memcpy(Request + 8, &NsPort, sizeof(NsPort));
-	Outcome = Socket.Send(Request, sizeof(Request));
-	if (!Outcome.IsOk())						return Outcome;
-	if (Outcome.GetResult() != sizeof(Request)) return FOutcome::Error("Sent too little to socks5 server");
-
-	// Connect reply
-	char Reply[3 + (1 + 4) + 2];
-	Outcome = Socket.Recv(Reply, sizeof(Reply));
-	if (!Outcome.IsOk())						return Outcome;
-	if (Outcome.GetResult() != sizeof(Reply))	return FOutcome::Error("Socks5 reply too short");
-	if (Reply[0] != 0x05 || Reply[1] != 0x00)	return FOutcome::Error("Reply has unexpected socks5 version");
-
-	return FOutcome::Ok(1);
-
-#ifdef _MSC_VER
-#pragma warning(pop)
-#endif
-}
-
-#endif // UE_BUILD_SHIPPING
-
-////////////////////////////////////////////////////////////////////////////////
-static FOutcome MaybeConnectSocks(FSocket& Socket, uint32 IpAddress, uint32 Port)
-{
-#if UE_BUILD_SHIPPING
-	return FOutcome::Ok();
-#else
-	if (GSocksIp.IsEmpty())
-	{
-		return FOutcome::Ok();
-	}
-
-	switch (GSocksVersion)
-	{
-	case 4: return ConnectSocks4(Socket, IpAddress, Port);
-	case 5: return ConnectSocks5(Socket, IpAddress, Port);
-	}
-
-	return FOutcome::Error("Unsupported socks version");
-#endif // UE_BUILD_SHIPPING
-}
-
-
-
-// {{{1 connection-pool ........................................................
-
-////////////////////////////////////////////////////////////////////////////////
-class FHost
-{
-public:
-	enum class EDirection : uint8 { Send, Recv };
-	static const uint32 InvalidIp = 0x00ff'ffff;
-
-	struct FParams
-	{
-		const ANSICHAR* HostName;
-		uint32			Port = 0;
-		uint32			MaxConnections = 1;
-	};
-
-					FHost(const FParams& Params);
-	void			SetBufferSize(EDirection Dir, int32 Size);
-	int32			GetBufferSize(EDirection Dir) const;
-	FOutcome		Connect(FSocket& Socket);
-	int32			IsResolved() const;
-	FOutcome		ResolveHostName();
-	uint32			GetMaxConnections() const	{ return MaxConnections; }
-	uint32			GetIpAddress() const		{ return IpAddresses[0]; }
-	FAnsiStringView	GetHostName() const			{ return HostName; }
-	uint32			GetPort() const				{ return Port; }
-
-private:
-	const ANSICHAR*	HostName;
-	uint32			IpAddresses[4] = {};
-	int16			SendBufKb = -1;
-	int16			RecvBufKb = -1;
-	uint16			Port;
-	uint8			MaxConnections;
-};
-
-////////////////////////////////////////////////////////////////////////////////
-FHost::FHost(const FParams& Params)
-: HostName(Params.HostName)
-, Port(uint16(Params.Port))
-, MaxConnections(uint8(Params.MaxConnections))
-{
-	check(MaxConnections && MaxConnections == Params.MaxConnections);
-
-	if (Port == 0)
-	{
-		Port = 80;
-	}
-}
-
-////////////////////////////////////////////////////////////////////////////////
-void FHost::SetBufferSize(EDirection Dir, int32 Size)
-{
-	(Dir == EDirection::Send) ? SendBufKb : RecvBufKb = uint16(Size >> 10);
-}
-
-////////////////////////////////////////////////////////////////////////////////
-int32 FHost::GetBufferSize(EDirection Dir) const
-{
-	return int32((Dir == EDirection::Send) ? SendBufKb : RecvBufKb) << 10;
-}
-
-////////////////////////////////////////////////////////////////////////////////
-FOutcome FHost::ResolveHostName()
-{
-	// todo: GetAddrInfoW() for async resolve on Windows
-
-	TRACE_CPUPROFILER_EVENT_SCOPE(IasHttp::PoolResolve);
-
-	IpAddresses[0] = 1;
-
-	addrinfo* Info = nullptr;
-	ON_SCOPE_EXIT { if (Info != nullptr) freeaddrinfo(Info); };
-
-	addrinfo Hints = {};
-	Hints.ai_family = AF_INET;
-	Hints.ai_socktype = SOCK_STREAM;
-	Hints.ai_protocol = IPPROTO_TCP;
-	auto Result = getaddrinfo(HostName, nullptr, &Hints, &Info);
-	if (uint32(Result) || Info == nullptr)
-	{
-		return FOutcome::Error("Error encountered resolving");
-	}
-
-	if (Info->ai_family != AF_INET)
-	{
-		return FOutcome::Error("Unexpected address family during resolve");
-	}
-
-	uint32 AddressCount = 0;
-	for (const addrinfo* Cursor = Info; Cursor != nullptr; Cursor = Cursor->ai_next)
-	{
-		const auto* AddrInet = (sockaddr_in*)(Cursor->ai_addr);
-		if (AddrInet->sin_family != AF_INET)
-		{
-			continue;
-		}
-
-		uint32 IpAddress = 0;
-		memcpy(&IpAddress, &(AddrInet->sin_addr), sizeof(uint32));
-
-		if (IpAddress == 0)
-		{
-			break;
-		}
-
-		IpAddresses[AddressCount] = htonl(IpAddress);
-		if (++AddressCount >= UE_ARRAY_COUNT(IpAddresses))
-		{
-			break;
-		}
-	}
-
-	if (AddressCount > 0)
-	{
-		return FOutcome::Ok(AddressCount);
-	}
-
-	return FOutcome::Error("Unable to resolve host");
-}
-
-////////////////////////////////////////////////////////////////////////////////
-int32 FHost::IsResolved() const
-{
-	switch (IpAddresses[0])
-	{
-	case 0:  return 0;
-	case 1:  return -1;
-	default: return 1;
-	}
-}
-
-////////////////////////////////////////////////////////////////////////////////
-FOutcome FHost::Connect(FSocket& Socket)
-{
-	FOutcome Outcome = FOutcome::None();
-
-	if (IsResolved() <= 0)
-	{
-		if (Outcome = ResolveHostName(); Outcome.IsError())
-		{
-			return Outcome;
-		}
-	}
-
-	check(IsResolved() > 0);
-	check(!Socket.IsValid());
-
-	uint32 IpAddress = GetIpAddress();
-
-	FSocket Candidate;
-	if (!Candidate.Create())
-	{
-		return FOutcome::Error("Failed to create socket");
-	}
-
-	// Attempt a SOCKS connect
-	Outcome = MaybeConnectSocks(Candidate, IpAddress, Port);
-	if (Outcome.IsError())
-	{
-		return Outcome;
-	}
-	check(Outcome.IsOk());
-	bool bSocksConnected = (Outcome.GetResult() == 1);
-
-	// Condition the socket
-	if (!Candidate.SetBlocking(false))
-	{
-		return FOutcome::Error("Unable to set socket non-blocking");
-	}
-
-	if (int32 OptValue = GetBufferSize(FHost::EDirection::Send); OptValue >= 0)
-	{
-		Candidate.SetSendBufSize(OptValue);
-	}
-
-	if (int32 OptValue = GetBufferSize(FHost::EDirection::Recv); OptValue >= 0)
-	{
-		Candidate.SetRecvBufSize(OptValue);
-	}
-
-	// Socks connect in a blocking fashion so we're all set (ret=1)
-	if (bSocksConnected)
-	{
-		Socket = MoveTemp(Candidate);
-		return FOutcome::Ok();
-	}
-
-	// Issue the connect - this is done non-blocking so we need to wait (ret=0)
-	if (Outcome = Candidate.Connect(IpAddress, Port); Outcome.IsError())
-	{
-		return Outcome;
-	}
-
-	Socket = MoveTemp(Candidate);
-	return FOutcome::Waiting();
-}
-
-
-
-////////////////////////////////////////////////////////////////////////////////
-int32 FConnectionPool::FParams::SetHostFromUrl(FAnsiStringView Url)
-{
-	FUrlOffsets Offsets;
-	if (ParseUrl(Url, Offsets) < 0)
-	{
-		return -1;
-	}
-
-	HostName = Offsets.HostName.Get(Url);
-
-	if (Offsets.Port)
-	{
-		FAnsiStringView PortView = Offsets.Port.Get(Url);
-		Port = uint16(CrudeToInt(PortView));
-	}
-
-	return Offsets.Path;
-}
-
-////////////////////////////////////////////////////////////////////////////////
-FConnectionPool::FConnectionPool(const FParams& Params)
-{
-	check(Params.ConnectionCount - 1 <= 63u);
-	check(Params.Port <= 0xffffu);
-
-	// Alloc a new internal object
-	uint32 HostNameLen = Params.HostName.Len();
-	uint32 AllocSize = sizeof(FHost) + (HostNameLen + 1);
-	auto* Internal = (FHost*)FMemory::Malloc(AllocSize, alignof(FHost));
-
-	// Copy host
-	char* HostDest = (char*)(Internal + 1);
-	memcpy(HostDest, Params.HostName.GetData(), HostNameLen);
-	HostDest[HostNameLen] = '\0';
-
-	// Init internal object
-	new (Internal) FHost({
-		.HostName		= HostDest,
-		.Port			= Params.Port,
-		.MaxConnections	= Params.ConnectionCount,
-	});
-	Internal->SetBufferSize(FHost::EDirection::Send, Params.SendBufSize);
-	Internal->SetBufferSize(FHost::EDirection::Recv, Params.RecvBufSize);
-
-	Ptr = Internal;
-}
-
-////////////////////////////////////////////////////////////////////////////////
-FConnectionPool::~FConnectionPool()
-{
-	if (Ptr != nullptr)
-	{
-		FMemory::Free(Ptr);
-	}
-}
-
-////////////////////////////////////////////////////////////////////////////////
-bool FConnectionPool::Resolve()
-{
-	return Ptr->ResolveHostName().IsOk();
-}
-
-////////////////////////////////////////////////////////////////////////////////
-void FConnectionPool::Describe(FAnsiStringBuilderBase& OutString) const
-{
-	const FAnsiStringView HostName = Ptr->GetHostName();
-	OutString.Appendf("%.*s", HostName.Len(), HostName.GetData());
-	if (!!Ptr->IsResolved())
-	{
-		const auto IpAddress = Ptr->GetIpAddress();
-		OutString.Appendf(" (%u.%u.%u.%u)",
-						(IpAddress >> 24) & 0xff,
-						(IpAddress >> 16) & 0xff,
-						(IpAddress >> 8) & 0xff,
-						IpAddress & 0xff
-		);
-	}
-	else
-	{
-		OutString.Append(" (unresolved)");
-	}
-}
-
-////////////////////////////////////////////////////////////////////////////////
-bool FConnectionPool::IsValidHostUrl(FAnsiStringView Url)
-{
-	FUrlOffsets Tmp;
-	return ParseUrl(Url, Tmp) >= 0;
-}
-
-// {{{1 activity ...............................................................
-
-#if IAS_HTTP_WITH_PERF
-
-////////////////////////////////////////////////////////////////////////////////
-class FStopwatch
-{
-public:
-	uint64	GetInterval(uint32 i) const;
-	void	SendStart()		{ Impl(0); }
-	void	SendEnd()		{ Impl(1); }
-	void	RecvStart()		{ Impl(2); }
-	void	RecvEnd()		{ Impl(3); }
-
-private:
-	void	Impl(uint32 Index);
-	uint64	Samples[4] = {};
-	uint32	Counts[2] = {};
-};
-
-////////////////////////////////////////////////////////////////////////////////
-uint64 FStopwatch::GetInterval(uint32 i) const
-{
-	if (i >= UE_ARRAY_COUNT(Samples) - 1)
-	{
-		return 0;
-	}
-	return Samples[i + 1] - Samples[i];
-}
-
-////////////////////////////////////////////////////////////////////////////////
-void FStopwatch::Impl(uint32 Index)
-{
-	if (uint64& Out = Samples[Index]; Out == 0)
-	{
-		Out = FPlatformTime::Cycles64();
-	}
-	Counts[Index >> 1] += !(Index & 1);
-}
-
-#endif // IAS_HTTP_WITH_PERF
-
-////////////////////////////////////////////////////////////////////////////////
-struct FResponseInternal
-{
-	FMessageOffsets Offsets;
-	int32			ContentLength = 0;
-	uint16			MessageLength;
-	mutable int16	Code;
-};
-
-////////////////////////////////////////////////////////////////////////////////
-struct alignas(16) FActivity
-{
-	enum class EState : uint8
-	{
-		None,
-		Build,
-		Send,
-		RecvMessage,
-		RecvContent,
-		RecvStream,
-		RecvDone,
-		Completed,
-		Cancelled,
-		Failed,
-		_Num,
-	};
-
-	FActivity*			Next = nullptr;
-	int8				Slot = -1;
-	EState				State = EState::None;
-	uint8				IsKeepAlive : 1;
-	uint8				NoContent : 1;
-	uint8				bFollow30x : 1;
-	uint8				_Unused0 : 5;
-	uint32				StateParam = 0;
-#if IAS_HTTP_WITH_PERF
-	FStopwatch			Stopwatch;
-#endif
-	union {
-		FHost*			Host;
-		FIoBuffer*		Dest;
-		const char*		ErrorReason;
-	};
-	UPTRINT				SinkParam;
-	FTicketSink			Sink;
-	FSocket				Socket;
-	FResponseInternal	Response;
-	FBuffer				Buffer;
-};
-
-////////////////////////////////////////////////////////////////////////////////
-static void Activity_TraceStateNames()
-{
-	FAnsiStringView StateNames[] = {
-		"None", "Build", "Send", "RecvMessage", "RecvContent", "RecvStream",
-		"RecvDone", "Completed", "Cancelled", "Failed", "$",
-	};
-	static_assert(UE_ARRAY_COUNT(StateNames) == int32(FActivity::EState::_Num) + 1);
-
-	TraceEnum(StateNames);
-}
-
-////////////////////////////////////////////////////////////////////////////////
-static void Activity_ChangeState(FActivity* Activity, FActivity::EState InState, uint32 Param=0)
-{
-	Trace(Activity, ETrace::StateChange, InState);
-
-	check(Activity->State != InState);
-	Activity->State = InState;
-	Activity->StateParam = Param;
-}
-
-////////////////////////////////////////////////////////////////////////////////
-static int32 Activity_Rewind(FActivity* Activity)
-{
-	using EState = FActivity::EState;
-
-	if (Activity->State == EState::Send)
-	{
-		Activity->StateParam = 0;
-		return 0;
-	}
-
-	if (Activity->State == EState::RecvMessage)
-	{
-		Activity->Buffer.Resize(Activity->StateParam);
-		Activity_ChangeState(Activity, EState::Send);
-		return 1;
-	}
-
-	return -1;
-}
-
-////////////////////////////////////////////////////////////////////////////////
-static uint32 Activity_RemainingKiB(FActivity* Activity)
-{
-	if (Activity->State < FActivity::EState::RecvContent) return MAX_uint32;
-	if (Activity->State > FActivity::EState::RecvContent) return 0;
-
-	uint32 ContentLength = uint32(Activity->Response.ContentLength);
-	check(Activity->StateParam <= ContentLength);
-	return (ContentLength - Activity->StateParam) >> 10;
-}
-
-////////////////////////////////////////////////////////////////////////////////
-static FActivity* Activity_Alloc(uint32 BufferSize)
-{
-	BufferSize = (BufferSize + 15) & ~15;
-
-	uint32 Size = BufferSize + sizeof(FActivity);
-	auto* Activity = (FActivity*)FMemory::Malloc(Size, alignof(FActivity));
-
-	new (Activity) FActivity();
-
-	auto* Scratch = (char*)(Activity + 1);
-	uint32 ScratchSize = BufferSize;
-	Activity->Buffer = FBuffer(Scratch, ScratchSize);
-
-	Activity_ChangeState(Activity, FActivity::EState::Build);
-
-	return Activity;
-}
-
-////////////////////////////////////////////////////////////////////////////////
-static void Activity_Free(FActivity* Activity)
-{
-	Trace(Activity, ETrace::ActivityDestroy);
-
-	Activity->~FActivity();
-	FMemory::Free(Activity);
-}
-
-////////////////////////////////////////////////////////////////////////////////
-static void Activity_SetError(FActivity* Activity, const char* Reason)
-{
-	Activity->IsKeepAlive = 0;
-	Activity->ErrorReason = Reason;
-
-	Activity_ChangeState(Activity, FActivity::EState::Failed, LastSocketResult());
-}
-
-// {{{1 request ................................................................
-
-////////////////////////////////////////////////////////////////////////////////
-FRequest::~FRequest()
-{
-	if (Ptr == nullptr)
-	{
-		return;
-	}
-
-	Activity_Free(Ptr);
-}
-
-////////////////////////////////////////////////////////////////////////////////
-FRequest::FRequest(FRequest&& Rhs)
-{
-	this->~FRequest();
-	Swap(Ptr, Rhs.Ptr);
-}
-
-////////////////////////////////////////////////////////////////////////////////
-FRequest&& FRequest::Accept(EMimeType MimeType)
-{
-	switch (MimeType)
-	{
-	case EMimeType::Text:				return Header("Accept", "text/html");
-	case EMimeType::Binary:				return Header("Accept", "application/octet-stream");
-	case EMimeType::Json:				return Header("Accept", "application/json");
-	case EMimeType::Xml:				return Header("Accept", "application/xml");
-	/* UE_CUSTOM_MIME_TYPES
-	case EMimeType::CbObject:			return Header("Accept", "application/x-ue-cb");
-	case EMimeType::CbPackage:			return Header("Accept", "application/x-ue-pkg");
-	case EMimeType::CompressedBuffer:	return Header("Accept", "application/x-ue-comp");
-	*/
-	}
-
-	return MoveTemp(*this);
-}
-
-////////////////////////////////////////////////////////////////////////////////
-FRequest&& FRequest::Accept(FAnsiStringView MimeType)
-{
-	return Header("Accept", MimeType);
-}
-
-////////////////////////////////////////////////////////////////////////////////
-FRequest&& FRequest::Header(FAnsiStringView Key, FAnsiStringView Value)
-{
-	checkSlow(Ptr->State == FActivity::EState::Build);
-	FMessageBuilder(Ptr->Buffer) << Key << ": " << Value << "\r\n";
-	return MoveTemp(*this);
-}
-
-
-
-// {{{1 response ...............................................................
-
-////////////////////////////////////////////////////////////////////////////////
-EStatusCodeClass FResponse::GetStatus() const
-{
-	uint32 Code = GetStatusCode();
-	if (Code <= 199) return EStatusCodeClass::Informational;
-	if (Code <= 299) return EStatusCodeClass::Successful;
-	if (Code <= 399) return EStatusCodeClass::Redirection;
-	if (Code <= 499) return EStatusCodeClass::ClientError;
-	if (Code <= 599) return EStatusCodeClass::ServerError;
-	return EStatusCodeClass::Unknown;
-}
-
-////////////////////////////////////////////////////////////////////////////////
-uint32 FResponse::GetStatusCode() const
-{
-	const auto* Activity = (const FActivity*)this;
-	const FResponseInternal& Internal = Activity->Response;
-
-	const char* MessageData = Activity->Buffer.GetData() + Activity->StateParam;
-	if (Internal.Code < 0)
-	{
-		const char* CodePtr = MessageData + Internal.Offsets.StatusCode;
-		Internal.Code = uint16(CrudeToInt(FAnsiStringView(CodePtr, 3)));
-	}
-
-	return Internal.Code;
-}
-
-////////////////////////////////////////////////////////////////////////////////
-FAnsiStringView FResponse::GetStatusMessage() const
-{
-	const auto* Activity = (const FActivity*)this;
-	const FResponseInternal& Internal = Activity->Response;
-
-	const char* MessageData = Activity->Buffer.GetData() + Activity->StateParam;
-	return FAnsiStringView(
-		MessageData + Internal.Offsets.Message,
-		Internal.Offsets.Headers - Internal.Offsets.Message
-	);
-}
-
-////////////////////////////////////////////////////////////////////////////////
-int64 FResponse::GetContentLength() const
-{
-	const auto* Activity = (const FActivity*)this;
-	const FResponseInternal& Internal = Activity->Response;
-	return Internal.ContentLength;
-}
-
-////////////////////////////////////////////////////////////////////////////////
-EMimeType FResponse::GetContentType() const
-{
-	FAnsiStringView Value;
-	GetContentType(Value);
-
-	if (Value == "text/html")					return EMimeType::Text;
-	if (Value == "application/octet-stream")	return EMimeType::Binary;
-	if (Value == "application/json")			return EMimeType::Json;
-	if (Value == "application/xml")				return EMimeType::Xml;
-	/* UE_CUSTOM_MIME_TYPES
-	if (Value == "application/x-ue-cb")			return EMimeType::CbObject;
-	if (Value == "application/x-ue-pkg")		return EMimeType::CbPackage;
-	if (Value == "application/x-ue-comp")		return EMimeType::CompressedBuffer;
-	*/
-
-	return EMimeType::Unknown;
-}
-
-////////////////////////////////////////////////////////////////////////////////
-void FResponse::GetContentType(FAnsiStringView& Out) const
-{
-	Out = GetHeader("Accept");
-
-	int32 SemiColon;
-	if (Out.FindChar(';', SemiColon))
-	{
-		Out = Out.Mid(SemiColon).TrimEnd();
-	}
-}
-
-////////////////////////////////////////////////////////////////////////////////
-FAnsiStringView FResponse::GetHeader(FAnsiStringView Name) const
-{
-	FAnsiStringView Result;
-	ReadHeaders([&Result, Name] (FAnsiStringView Candidate, FAnsiStringView Value)
-	{
-		if (Candidate != Name)
-		{
-			return true;
-		}
-
-		Result = Value;
-		return false;
-	});
-	return Result;
-}
-
-////////////////////////////////////////////////////////////////////////////////
-void FResponse::ReadHeaders(FHeaderSink Sink) const
-{
-	const auto* Activity = (const FActivity*)this;
-	const FResponseInternal& Internal = Activity->Response;
-
-	const char* MessageData = Activity->Buffer.GetData() + Activity->StateParam;
-	FAnsiStringView Result, Headers(
-		MessageData + Internal.Offsets.Headers,
-		Internal.MessageLength - Internal.Offsets.Headers
-	);
-
-	EnumerateHeaders(Headers, Sink);
-}
-
-////////////////////////////////////////////////////////////////////////////////
-void FResponse::SetDestination(FIoBuffer* Buffer)
-{
-	auto* Activity = (FActivity*)this;
-	Activity->Dest = Buffer;
-}
-
-
-
-// {{{1 ticket-status ..........................................................
-
-////////////////////////////////////////////////////////////////////////////////
-FTicketStatus::EId FTicketStatus::GetId() const
-{
-	const auto* Activity = (FActivity*)this;
-	switch (Activity->State)
-	{
-	case FActivity::EState::RecvMessage:	return EId::Response;
-	case FActivity::EState::RecvDone:		return EId::Content;
-	case FActivity::EState::Cancelled:		return EId::Cancelled;
-	case FActivity::EState::Failed:			return EId::Error;
-	default:								check(false);
-	}
-	return EId::Error;
-}
-
-////////////////////////////////////////////////////////////////////////////////
-UPTRINT FTicketStatus::GetParam() const
-{
-	const auto* Activity = (FActivity*)this;
-	return Activity->SinkParam;
-}
-
-////////////////////////////////////////////////////////////////////////////////
-FTicket FTicketStatus::GetTicket() const
-{
-	const auto* Activity = (FActivity*)this;
-	return 1ull << Activity->Slot;
-}
-
-////////////////////////////////////////////////////////////////////////////////
-uint32 FTicketStatus::GetIndex() const
-{
-	const auto* Activity = (FActivity*)this;
-	return Activity->Slot;
-}
-
-////////////////////////////////////////////////////////////////////////////////
-FResponse& FTicketStatus::GetResponse() const
-{
-	check(GetId() < EId::Content);
-	const auto* Activity = (FActivity*)this;
-	return *(FResponse*)Activity;
-}
-
-////////////////////////////////////////////////////////////////////////////////
-uint32 FTicketStatus::GetContentLength() const
-{
-	check(GetId() <= EId::Content);
-	const auto* Activity = (FActivity*)this;
-	return Activity->Response.ContentLength;
-}
-
-////////////////////////////////////////////////////////////////////////////////
-const FTicketPerf& FTicketStatus::GetPerf() const
-{
-	check(GetId() == EId::Content);
-	const auto* Activity = (FActivity*)this;
-	return *(FTicketPerf*)Activity;
-}
-
-////////////////////////////////////////////////////////////////////////////////
-const FIoBuffer& FTicketStatus::GetContent() const
-{
-	check(GetId() == EId::Content);
-	const auto* Activity = (FActivity*)this;
-	return *(Activity->Dest);
-}
-
-////////////////////////////////////////////////////////////////////////////////
-const char* FTicketStatus::GetErrorReason() const
-{
-	check(GetId() == EId::Error);
-	const auto* Activity = (FActivity*)this;
-	return Activity->ErrorReason;
-}
-
-
-
-// {{{1 perf ...................................................................
-
-#if IAS_HTTP_WITH_PERF
-
-////////////////////////////////////////////////////////////////////////////////
-FTicketPerf::FSample FTicketPerf::GetSample() const
-{
-	const auto* Activity = (FActivity*)this;
-
-	static uint64 Freq;
-	if (Freq == 0)
-	{
-		Freq = uint64(1.0 / FPlatformTime::GetSecondsPerCycle());
-	}
-
-	auto ToMs = [] (uint64 Value) { return uint32((Value * 1000ull) / Freq); };
-
-	const FStopwatch& Stopwatch = Activity->Stopwatch;
-	return {
-		ToMs(Stopwatch.GetInterval(0)),
-		ToMs(Stopwatch.GetInterval(1)),
-		ToMs(Stopwatch.GetInterval(2)),
-	};
-}
-
-#endif // IAS_HTTP_WITH_PERF
-
-
-
-// {{{1 event-loop-int .........................................................
-
-using FPeerType = FSocket;
-
-////////////////////////////////////////////////////////////////////////////////
-static FOutcome DoSend(FActivity* Activity, FPeerType& Peer)
-{
-	Trace(Activity, ETrace::StateChange, Activity->State);
-
-#if IAS_HTTP_WITH_PERF
-	Activity->Stopwatch.SendStart();
-#endif
-
-	TRACE_CPUPROFILER_EVENT_SCOPE(IasHttp::DoSend);
-
-	FBuffer& Buffer = Activity->Buffer;
-	const char* SendData = Buffer.GetData();
-	int32 SendSize = Buffer.GetSize();
-
-	uint32 AlreadySent = Activity->StateParam;
-	SendData += AlreadySent;
-	SendSize -= AlreadySent;
-	check(SendSize > 0);
-
-	FOutcome Outcome = Peer.Send(SendData, SendSize);
-
-	if (Outcome.IsError())
-	{
-		Activity_SetError(Activity, Outcome.GetMessage().GetData());
-		return Outcome;
-	}
-
-	if (Outcome.IsWaiting())
-	{
-		return Outcome;
-	}
-
-	check(Outcome.IsOk());
-
-	int32 Result = Outcome.GetResult();
-	Activity->StateParam += Result;
-	if (Activity->StateParam < Buffer.GetSize())
-	{
-		return DoSend(Activity, Peer);
-	}
-
-#if IAS_HTTP_WITH_PERF
-	Activity->Stopwatch.SendEnd();
-#endif
-
-	Activity_ChangeState(Activity, FActivity::EState::RecvMessage, Buffer.GetSize());
-	return FOutcome::Ok(Result);
-}
-
-////////////////////////////////////////////////////////////////////////////////
-static FOutcome DoRecvMessage(FActivity* Activity, FPeerType& Peer)
-{
-	TRACE_CPUPROFILER_EVENT_SCOPE(IasHttp::DoRecvMessage);
-
-	static const uint32 PageSize = 256;
-
-	FBuffer& Buffer = Activity->Buffer;
-
-	const char* MessageRight;
-	while (true)
-	{
-		Trace(Activity, ETrace::StateChange, Activity->State);
-
-#if IAS_HTTP_WITH_PERF
-		Activity->Stopwatch.RecvStart();
-#endif
-
-		auto [Dest, DestSize] = Buffer.GetMutableFree(0, PageSize);
-		FOutcome Outcome = Peer.Recv(Dest, DestSize);
-
-		if (Outcome.IsError())
-		{
-			Activity_SetError(Activity, Outcome.GetMessage().GetData());
-			return Outcome;
-		}
-
-		if (Outcome.IsWaiting())
-		{
-			return Outcome;
-		}
-
-		check(Outcome.IsOk());
-
-		int32 Result = Outcome.GetResult();
-		Buffer.AdvanceUsed(Result);
-
-		// Rewind a little to cover cases where the terminal is fragmented across
-		// recv() calls
-		uint32 DestBias = 0;
-		if (Dest - 3 >= Buffer.GetData() + Activity->StateParam)
-		{
-			Dest -= (DestBias = 3);
-		}
-
-		int32 MessageEnd = FindMessageTerminal(Dest, Result + DestBias);
-		if (MessageEnd < 0)
-		{
-			if (Buffer.GetSize() > (8 << 10))
-			{
-				Activity_SetError(Activity, "Headers have grown larger than expected");
-				return FOutcome::Error(Activity->ErrorReason);
-			}
-
-			continue;
-		}
-
-		MessageRight = Dest + MessageEnd;
-		break;
-	}
-
-	// Fill out the internal response object
-	FResponseInternal& Internal = Activity->Response;
-	const char* MessageData = Buffer.GetData() + Activity->StateParam;
-	Internal.MessageLength = uint16(ptrdiff_t(MessageRight - MessageData));
-
-	FAnsiStringView ResponseView(MessageData, Internal.MessageLength);
-	if (ParseMessage(ResponseView, Internal.Offsets) < 0)
-	{
-		Activity_SetError(Activity, "Failed to parse message status");
-		return FOutcome::Error(Activity->ErrorReason);
-	}
-
-	// Parse headers
-	FAnsiStringView Headers = ResponseView.Mid(
-		Internal.Offsets.Headers,
-		Internal.MessageLength - Internal.Offsets.Headers - 2 // "-2" trims off '\r\n' that signals end of headers
-	);
-
-	int32 Count = 2;
-	bool IsKeepAlive = true;
-	int32 ContentLength = -1;
-	EnumerateHeaders(
-		Headers,
-		[&ContentLength, &IsKeepAlive, &Count] (FAnsiStringView Name, FAnsiStringView Value)
-		{
-			// todo; may need smarter value handling; ;/, separated options & key-value pairs (ex. in rfc2068)
-
-			// "Keep-Alive"			- deprecated
-			// "Transfer-Encoding"	- may be required later
-
-			if (Name.Equals("Content-Length", ESearchCase::IgnoreCase))
-			{
-				ContentLength = int32(CrudeToInt(Value));
-				Count--;
-			}
-
-			else if (Name.Equals("Connection", ESearchCase::IgnoreCase))
-			{
-				IsKeepAlive = !Value.Equals("close");
-				Count--;
-			}
-
-			return Count > 0;
-		}
-	);
-
-	Activity->IsKeepAlive &= IsKeepAlive;
-
-	if (ContentLength < 0)
-	{
-		if (ContentLength == -1)
-		{
-			// todo; query Transfer-Encoding for chunked
-		}
-
-		Activity_SetError(Activity, "Unknown content length value");
-		return FOutcome::Error(Activity->ErrorReason);
-	}
-
-	// Call out to the sink to get a content destination
-	FIoBuffer* PriorDest = Activity->Dest; // to retain unioned Host ptr
-	Internal.Code = -1;
-	Internal.ContentLength = ContentLength;
-	{
-		FTicketStatus& SinkArg = *(FTicketStatus*)Activity;
-		Activity->Sink(SinkArg);
-	}
-
-	if (Activity->NoContent == 0)
-	{
-		if (Activity->Dest == PriorDest)
-		{
-			Activity_SetError(Activity, "User did not provide a destination buffer");
-			return FOutcome::Error(Activity->ErrorReason);
-		}
-
-		// The user seems to have forgotten something. Let's help them along
-		if (Activity->Dest->GetSize() == 0)
-		{
-			*Activity->Dest = FIoBuffer(ContentLength);
-		}
-	}
-
-	// Perhaps we have some of the content already?
-	const char* BufferRight = Buffer.GetData() + Buffer.GetSize();
-	uint32 AlreadyReceived = uint32(ptrdiff_t(BufferRight - MessageRight));
-	if (AlreadyReceived > uint32(ContentLength))
-	{
-		Activity_SetError(Activity, "More data received that expected");
-		return FOutcome::Error(Activity->ErrorReason);
-	}
-
-	if (Activity->NoContent == 1)
-	{
-		if (AlreadyReceived)
-		{
-			Activity_SetError(Activity, "Received content when none was expected");
-			return FOutcome::Error(Activity->ErrorReason);
-		}
-		Activity_ChangeState(Activity, FActivity::EState::RecvDone);
-		return FOutcome::Ok();
-	}
-
-	check(Activity->Dest != nullptr);
-
-	const bool bStreamed = Activity->Dest->GetSize() < ContentLength;
-
-	auto NextState = bStreamed ? FActivity::EState::RecvStream : FActivity::EState::RecvContent;
-	Activity_ChangeState(Activity, NextState, AlreadyReceived);
-
-	if (AlreadyReceived == 0)
-	{
-		return FOutcome::Ok();
-	}
-
-	FMutableMemoryView DestView = Activity->Dest->GetMutableView();
-	const char* Cursor = BufferRight - AlreadyReceived;
-	if (!bStreamed || AlreadyReceived < DestView.GetSize())
-	{
-		::memcpy(DestView.GetData(), Cursor, AlreadyReceived);
-		return FOutcome::Ok();
-	}
-
-#if 1
-	check(false); // not implemented yet
-#else
-	do
-	{
-		uint32 Size = FMath::Min<uint32>(DestView.GetSize(), AlreadyReceived);
-		::memcpy(DestView.GetData(), Cursor, Size);
-
-		/* send sink */
-
-		AlreadyReceived -= Size;
-		Cursor += Size;
-	}
-	while (AlreadyReceived);
-#endif
-
-	return FOutcome::Ok();
-}
-
-////////////////////////////////////////////////////////////////////////////////
-static FOutcome DoRecvContent(FActivity* Activity, FPeerType& Peer, int32& MaxRecvSize)
-{
-	TRACE_CPUPROFILER_EVENT_SCOPE(IasHttp::DoRecvContent);
-
-	FResponseInternal& Response = Activity->Response;
-	FMutableMemoryView DestView = Activity->Dest->GetMutableView();
-
-	while (true)
-	{
-		int32 Size = (Response.ContentLength - Activity->StateParam);
-		if (Size == 0)
-		{
-			break;
-		}
-
-		Size = FMath::Min(Size, MaxRecvSize);
-		check(Size >= 0);
-		if (Size == 0)
-		{
-			return FOutcome::Waiting();
-		}
-
-		Trace(Activity, ETrace::StateChange, Activity->State);
-
-		char* Cursor = (char*)(DestView.GetData()) + Activity->StateParam;
-
-		FOutcome Outcome = Peer.Recv(Cursor, Size);
-
-		if (Outcome.IsWaiting())
-		{
-			return Outcome;
-		}
-
-		if (Outcome.IsError())
-		{
-			Activity_SetError(Activity, Outcome.GetMessage().GetData());
-			return Outcome;
-		}
-
-		check(Outcome.IsOk());
-
-		int32 Result = Outcome.GetResult();
-		check(Result <= MaxRecvSize);
-
-		Activity->StateParam += Result;
-		MaxRecvSize -= Result;
-	}
-
-	if (!FLatencyInjector::Begin(FLatencyInjector::EType::Network, Activity->StateParam))
-	{
-		Activity_SetError(Activity, "Forced random failure");
-		return FOutcome::Error(Activity->ErrorReason);
-	}
-
-#if IAS_HTTP_WITH_PERF
-	Activity->Stopwatch.RecvEnd();
-#endif
-
-	Activity_ChangeState(Activity, FActivity::EState::RecvDone);
-	return FOutcome::Ok();
-}
-
-////////////////////////////////////////////////////////////////////////////////
-static FOutcome DoRecvStream(FActivity*, FPeerType&, uint32)
-{
-	check(false); // not yet implemented
-	return FOutcome::Error("Not supported");
-}
-
-////////////////////////////////////////////////////////////////////////////////
-static FOutcome DoRecv(FActivity* Activity, FPeerType& Peer, int32& MaxRecvSize)
-{
-	using EState = FActivity::EState;
-
-	EState State = Activity->State; 
-	check(State >= EState::RecvMessage && State < EState::RecvDone);
-
-	if (State == EState::RecvMessage)	return DoRecvMessage(Activity, Peer);
-	if (State == EState::RecvContent)	return DoRecvContent(Activity, Peer, MaxRecvSize);
-	if (State == EState::RecvStream)	return DoRecvStream(Activity, Peer, MaxRecvSize);
-	
-	check(false); // it is not expected that we'll get here
-	return FOutcome::Error("unreachable");
-}
-
-////////////////////////////////////////////////////////////////////////////////
-static int32 DoRecvDone(FActivity* Activity)
-{
-	if (!FLatencyInjector::HasExpired(Activity->StateParam))
-	{
-		return 1;
-	}
-
-	// Notify the user we've received everything
-	FTicketStatus& SinkArg = *(FTicketStatus*)Activity;
-	Activity->Sink(SinkArg);
-
-	Activity_ChangeState(Activity, FActivity::EState::Completed);
-	return 0;
-}
-
-////////////////////////////////////////////////////////////////////////////////
-static void DoCancel(FActivity* Activity)
-{
-	if (Activity->State >= FActivity::EState::Completed)
-	{
-		return;
-	}
-
-	Activity_ChangeState(Activity, FActivity::EState::Cancelled);
-
-	FTicketStatus& SinkArg = *(FTicketStatus*)Activity;
-	Activity->Sink(SinkArg);
-}
-
-////////////////////////////////////////////////////////////////////////////////
-static void DoFail(FActivity* Activity)
-{
-	check(Activity->State == FActivity::EState::Failed);
-
-	// Notify the user we've received everything
-	FTicketStatus& SinkArg = *(FTicketStatus*)Activity;
-	Activity->Sink(SinkArg);
-}
-
-
-
-// {{{1 throttler ..............................................................
-
-////////////////////////////////////////////////////////////////////////////////
-static void ThrottleTest(FAnsiStringView);
-
-////////////////////////////////////////////////////////////////////////////////
-class FThrottler
-{
-public:
-			FThrottler();
-	void	SetLimit(uint32 KiBPerSec);
-	int32	GetAllowance();
-	void	ReturnUnused(uint32 Unused);
-
-private:
-	friend	void ThrottleTest(FAnsiStringView);
-	int32	GetAllowance(uint64 CycleDelta);
-	int32	GetWaitEstimateMs() const;
-	uint64	CycleFreq;
-	uint64	CycleLast;
-	uint64	CycleIdle;
-	uint32	Limit = 0;
-	int32	Available = 0;
-
-	enum {
-		LIMITLESS	= MAX_int32,
-		THRESHOLD	= 2 << 10,
-	};
-};
-
-////////////////////////////////////////////////////////////////////////////////
-FThrottler::FThrottler()
-{
-	CycleFreq = uint64(1.0 / FPlatformTime::GetSecondsPerCycle());
-	CycleLast = FPlatformTime::Cycles64() - CycleFreq;
-	CycleIdle = CycleFreq * 8;
-}
-
-////////////////////////////////////////////////////////////////////////////////
-void FThrottler::SetLimit(uint32 KiBPerSec)
-{
-	// 512MiB/s might as well be limitless.
-	KiBPerSec = (KiBPerSec < (512 << 10)) ? KiBPerSec : 0;
-	Limit = KiBPerSec << 10;
-}
-
-////////////////////////////////////////////////////////////////////////////////
-int32 FThrottler::GetAllowance()
-{
-	int64 Cycle = FPlatformTime::Cycles64();
-	int64 CycleDelta = Cycle - CycleLast;
-	if (CycleDelta < 0)
-	{
-		return Limit;
-	}
-	CycleLast = Cycle;
-	return GetAllowance(CycleDelta);
-}
-
-////////////////////////////////////////////////////////////////////////////////
-int32 FThrottler::GetAllowance(uint64 CycleDelta)
-{
-	if (Limit == 0)
-	{
-		return LIMITLESS;
-	}
-
-	// If we're idle for too long then reset the throttling
-	if (CycleDelta >= CycleIdle)
-	{
-		Available = 0;
-		return Limit;
-	}
-
-	uint64 Delta = (uint64(Limit) * CycleDelta) / CycleFreq;
-
-	// A gate against lost precision
-	if (Delta == 0)
-	{
-		CycleLast -= CycleDelta;
-		return 0 - GetWaitEstimateMs();
-	}
-
-	// Don't let available run away
-	uint64 Next = FMath::Min<uint64>(uint64(Available) + Delta, Limit * 4);
-
-	Available = uint32(Next);
-
-	// Doesn't make sense to trickle out tiny allowances
-	if (Available < THRESHOLD)
-	{
-		return 0 - GetWaitEstimateMs();
-	}
-
-	int32 Released = Available;
-	Available = 0;
-	return Released;
-}
-
-////////////////////////////////////////////////////////////////////////////////
-void FThrottler::ReturnUnused(uint32 Unused)
-{
-	Available += Unused;
-}
-
-////////////////////////////////////////////////////////////////////////////////
-int32 FThrottler::GetWaitEstimateMs() const
-{
-	// Calculate an approximate time to wait for more allowance
-	int64 Estimate = THRESHOLD - Available;
-	Estimate = (Estimate * 1000ll) / int64(Limit);
-	return FMath::Max(int32(Estimate), 0);
-}
-
-
-
-// {{{1 groups .................................................................
-
-/*
- * - Activities (requests send with a loop) are managed in singly-linked lists
- * - Each activity has an associated host it is talking to.
- * - Hosts are ephemeral, or represented externally via a FConnectionPool object
- * - Loop has a group for each host, and each host-group has a bunch of socket-groups
- * - Host-group has a list of work; pending activities waiting to start
- * - Socket-groups own up to two activities; one sending, one receiving
- * - As it recvs, a socket-group will, if possible, fetch more work from the host
- *
- *  Loop:
- *    FHostGroup[HostPtr]:
- *	    Work: Act0 -> Act1 -> Act2 -> Act3 -> ...
- *      FPeerGroup[0...HostMaxConnections]:
- *			Act.Send
- *			Act.Recv
- */
-
-////////////////////////////////////////////////////////////////////////////////
-struct FTickState
-{
-	FActivity*					DoneList;
-	uint64						Cancels;
-	int32&						RecvAllowance;
-	int32						PollTimeoutMs;
-	int32						FailTimeoutMs;
-	uint32						NowMs;
-	class FWorkQueue*			Work;
-};
-
-
-
-////////////////////////////////////////////////////////////////////////////////
-class FWorkQueue
-{
-public:
-						FWorkQueue() = default;
-						~FWorkQueue();
-	bool				HasWork() const { return List != nullptr; }
-	void				AddActivity(FActivity* Activity);
-	FActivity*			PopActivity();
-	void				TickCancels(FTickState& State);
-
-private:
-	FActivity*			List = nullptr;
-	FActivity*			ListTail = nullptr;
-	uint64				ActiveSlots = 0;
-
-	UE_NONCOPYABLE(FWorkQueue);
-};
-
-////////////////////////////////////////////////////////////////////////////////
-FWorkQueue::~FWorkQueue()
-{
-	check(List == nullptr);
-	check(ListTail == nullptr);
-}
-
-////////////////////////////////////////////////////////////////////////////////
-void FWorkQueue::AddActivity(FActivity* Activity)
-{
-	// We use a tail pointer here to maintain order that requests were made
-
-	check(Activity->Next == nullptr);
-
-	if (ListTail != nullptr)
-	{
-		ListTail->Next = Activity;
-	}
-	List = (List == nullptr) ? Activity : List;
-	ListTail = Activity;
-
-	ActiveSlots |= (1ull << Activity->Slot);
-}
-
-////////////////////////////////////////////////////////////////////////////////
-FActivity* FWorkQueue::PopActivity()
-{
-	if (List == nullptr)
-	{
-		return nullptr;
-	}
-
-	FActivity* Activity = List;
-	if ((List = List->Next) == nullptr)
-	{
-		ListTail = nullptr;
-	}
-
-	check(ActiveSlots & (1ull << Activity->Slot));
-	ActiveSlots ^= (1ull << Activity->Slot);
-
-	Activity->Next = nullptr;
-	return Activity;
-}
-
-////////////////////////////////////////////////////////////////////////////////
-void FWorkQueue::TickCancels(FTickState& State)
-{
-	if (State.Cancels == 0 || (State.Cancels & ActiveSlots) == 0)
-	{
-		return;
-	}
-
-	// We are going to rebuild the list of activities to maintain order as the
-	// activity list is singular.
-
-	check(List != nullptr);
-	FActivity* Activity = List;
-	List = ListTail = nullptr;
-	ActiveSlots = 0;
-
-	for (FActivity* Next; Activity != nullptr; Activity = Next)
-	{
-		Next = Activity->Next;
-
-		if (uint64 Slot = (1ull << Activity->Slot); (State.Cancels & Slot) == 0)
-		{
-			Activity->Next = nullptr;
-			AddActivity(Activity);
-			continue;
-		}
-
-		DoCancel(Activity);
-
-		Activity->Next = State.DoneList;
-		State.DoneList = Activity;
-	}
-}
-
-
-
-////////////////////////////////////////////////////////////////////////////////
-class FPeerGroup
-{
-public:
-						FPeerGroup() = default;
-						~FPeerGroup();
-	bool				operator == (FPeerType* Rhs) const { return &Peer == Rhs; }
-	void				Unwait()			{ check(bWaiting); bWaiting = false; }
-	FPeerType::FWaiter	GetWaiter() const;
-	bool 				Tick(FTickState& State);
-	void				TickSend(FTickState& State, FHost& Host);
-	void				Fail(FTickState& State, const char* Reason);
-
-private:
-	void				RecvInternal(FTickState& State);
-	void				SendInternal(FTickState& State);
-	FActivity*			Send = nullptr;
-	FActivity*			Recv = nullptr;
-	FPeerType			Peer;
-	uint32				LastUseMs = 0;
-	uint8				IsKeepAlive = 0;
-	bool				bWaiting = false;
-
-	UE_NONCOPYABLE(FPeerGroup);
-};
-
-////////////////////////////////////////////////////////////////////////////////
-FPeerGroup::~FPeerGroup()
-{
-	check(Send == nullptr);
-	check(Recv == nullptr);
-}
-
-////////////////////////////////////////////////////////////////////////////////
-FPeerType::FWaiter FPeerGroup::GetWaiter() const
-{
-	if (!bWaiting)
-	{
-		return FPeerType::FWaiter();
-	}
-
-	using EWhat = FPeerType::FWaiter::EWhat;
-	EWhat What = (Recv != nullptr) ? EWhat::Recv : EWhat::Send;
-	return FPeerType::FWaiter(Peer, What);
-}
-
-////////////////////////////////////////////////////////////////////////////////
-void FPeerGroup::Fail(FTickState& State, const char* Reason)
-{
-	// Any send left at this point is unrecoverable
-	if (Send != nullptr)
-	{
-		Send->Next = Recv;
-		Recv = Send;
-	}
-
-	// Failure is quite terminal and we need to abort everything
-	for (FActivity* Activity = Recv; Activity != nullptr;)
-	{
-		if (Activity->State != FActivity::EState::Failed)
-		{
-			Activity_SetError(Activity, Reason);
-		}
-
-		DoFail(Activity);
-
-		FActivity* Next = Activity->Next;
-		Activity->Next = State.DoneList;
-		State.DoneList = Activity;
-		Activity = Next;
-	}
-
-	Peer = FPeerType();
-	Send = Recv = nullptr;
-	bWaiting = false;
-	IsKeepAlive = 0;
-}
-
-////////////////////////////////////////////////////////////////////////////////
-void FPeerGroup::RecvInternal(FTickState& State)
-{
-	check(Recv != nullptr);
-
-	// Another helper lambda
-	auto IsReceiving = [] (const FActivity* Act)
-	{
-		using EState = FActivity::EState;
-		return (Act->State >= EState::RecvMessage) & (Act->State < EState::RecvDone);
-	};
-
-	FActivity* Activity = Recv;
-	check(IsReceiving(Activity));
-
-	FOutcome Outcome = DoRecv(Activity, Peer, State.RecvAllowance);
-
-	// Any sort of error here is unrecoverable
-	if (Outcome.IsError())
-	{
-		Fail(State, Outcome.GetMessage().GetData());
-		return;
-	}
-
-	IsKeepAlive &= Activity->IsKeepAlive;
-	LastUseMs = State.NowMs;
-
-	// If we've only a small amount left to receive we can start more work
-	if (IsKeepAlive & (Recv->Next == nullptr))
-	{
-		uint32 Remaining = Activity_RemainingKiB(Activity);
-		if (Remaining < uint32(GRecvWorkThresholdKiB))
-		{
-			if (FActivity* Next = State.Work->PopActivity(); Next != nullptr)
-			{
-				Trace(Activity, ETrace::StartWork);
-	
-				check(Send == nullptr);
-				Send = Next;
-				SendInternal(State);
-
-				if (!Peer.IsValid())
-				{
-					return;
-				}
-			}
-		}
-	}
-
-	// If there was no data available this is far as receiving can go
-	if (bWaiting = Outcome.IsWaiting(); bWaiting)
-	{
-		return;
-	}
-
-	// If we're still in a receiving state we will just try again otherwise it
-	// is finished and we will let DoneList recipient finish it off.
-	if (IsReceiving(Activity))
-	{
-		return;
-	}
-
-	DoRecvDone(Activity);
-
-	Recv = Activity->Next;
-	Activity->Next = State.DoneList;
-	State.DoneList = Activity;
-
-	// If the server wants to close the socket we need to rewind the send
-	if (IsKeepAlive != 0)
-	{
-		return;
-	}
-
-	if (Send != nullptr && Activity_Rewind(Send) < 0)
-	{
-		Fail(State, "Unable to rewind on keep-alive close");
-		return;
-	}
-
-	Peer = FPeerType();
-}
-
-////////////////////////////////////////////////////////////////////////////////
-void FPeerGroup::SendInternal(FTickState& State)
-{
-	check(IsKeepAlive == 1);
-	check(Send != nullptr);
-
-	FActivity* Activity = Send;
-
-	FOutcome Outcome = DoSend(Activity, Peer);
-
-	if (Outcome.IsWaiting())
-	{
-		// For now we'll not add the socket as a waiter. It is unlikely that we
-		// send enough to need to wait currently.
-		return;
-	}
-
-	if (Outcome.IsError())
-	{
-		Fail(State, Outcome.GetMessage().GetData());
-		return;
-	}
-
-	Send = nullptr;
-
-	// Pass along this send to be received
-	if (Recv == nullptr)
-	{
-		Recv = Activity;
-		return;
-	}
-
-	check(Recv->Next == nullptr);
-	Recv->Next = Activity;
-}
-
-////////////////////////////////////////////////////////////////////////////////
-bool FPeerGroup::Tick(FTickState& State)
-{
-	if (Send != nullptr)
-	{
-		SendInternal(State);
-	}
-
-	if (Recv != nullptr && State.RecvAllowance)
-	{
-		RecvInternal(State);
-	}
-
-	return !!IsKeepAlive | !!(UPTRINT(Send) | UPTRINT(Recv));
-}
-
-////////////////////////////////////////////////////////////////////////////////
-void FPeerGroup::TickSend(FTickState& State, FHost& Host)
-{
-	// This path is only for those that are idle and have nothing to do
-	if (Send != nullptr || Recv != nullptr)
-	{
-		return;
-	}
-
-	// Failing will try and recover work which we don't want to happen yet
-	FActivity* Pending = State.Work->PopActivity();
-	check(Pending != nullptr);
-
-	// Close idle sockets
-	if (Peer.IsValid() && LastUseMs + GIdleMs < State.NowMs)
-	{
-		LastUseMs = State.NowMs;
-		Peer = FPeerType();
-	}
-
-	// We don't have a connected socket on first use, or if a keep-alive:close
-	// was received from the server. So we connect here.
-	bool bWillBlock = false;
-	if (!Peer.IsValid())
-	{
-		IsKeepAlive = 1;
-		FOutcome Outcome = Host.Connect(Peer);
-
-		// We failed to connect, let's bail.
-		if (Outcome.IsError())
-		{
-			Pending->Next = Recv;
-			Recv = Pending;
-			Fail(State, Outcome.GetMessage().GetData());
-			return;
-		}
-
-		bWillBlock = Outcome.IsWaiting();
-	}
-
-	Send = Pending;
-
-	if (!bWillBlock)
-	{
-		return SendInternal(State);
-	}
-
-	// Non-blocking connect
-	bWaiting = true;
-}
-
-
-
-////////////////////////////////////////////////////////////////////////////////
-class FHostGroup
-{
-public:
-								FHostGroup(FHost& InHost);
-	bool						IsBusy() const	{ return BusyCount != 0; }
-	const FHost&				GetHost() const	{ return Host; }
-	void						Tick(FTickState& State);
-	void						AddActivity(FActivity* Activity);
-
-private:
-	int32						Wait(const FTickState& State);
-	TArray<FPeerGroup>			PeerGroups;
-	FWorkQueue					Work;
-	FHost&						Host;
-	uint32						BusyCount = 0;
-	int32						WaitTimeAccum = 0;
-
-	UE_NONCOPYABLE(FHostGroup);
-};
-
-////////////////////////////////////////////////////////////////////////////////
-FHostGroup::FHostGroup(FHost& InHost)
-: Host(InHost)
-{
-	uint32 Num = InHost.GetMaxConnections();
-	PeerGroups.SetNum(Num);
-}
-
-////////////////////////////////////////////////////////////////////////////////
-int32 FHostGroup::Wait(const FTickState& State)
-{
-	// Collect groups that are waiting on something
-	TArray<FPeerType::FWaiter, TFixedAllocator<64>> Waiters;
-	for (FPeerGroup& Group : PeerGroups)
-	{
-		FPeerType::FWaiter Waiter = Group.GetWaiter();
-		if (Waiter.IsValid())
-		{
-			Waiters.Add(Waiter);
-		}
-	}
-
-	if (Waiters.IsEmpty())
-	{
-		return 0;
-	}
-
-	Trace(0, ETrace::Wait);
-	ON_SCOPE_EXIT { Trace(0, ETrace::Unwait); };
-
-	// If the poll timeout is negative then treat that as a fatal timeout
-	check(State.FailTimeoutMs);
-	int32 PollTimeoutMs = State.PollTimeoutMs;
-	if (PollTimeoutMs < 0)
-	{
-		PollTimeoutMs = State.FailTimeoutMs;
-	}
-
-	// Actually do the wait
-	int32 Result = FPeerType::Wait(Waiters, PollTimeoutMs);
-	if (Result <= 0)
-	{
-		// If the user opts to not block then we don't accumulate wait time and
-		// leave it to them to manage time a fail timoue
-		WaitTimeAccum += PollTimeoutMs;
-
-		if (State.PollTimeoutMs < 0 || WaitTimeAccum >= State.FailTimeoutMs)
-		{
-			return MIN_int32;
-		}
-
-		return Result;
-	}
-
-	WaitTimeAccum = 0;
-
-	// For each waiter that's ready, find the associated group "unwait" them.
-	int32 Count = 0;
-	for (int32 i = 0, n = Waiters.Num(); i < n; ++i)
-	{
-		if (Waiters[i].Ready == 0)
-		{
-			continue;
-		}
-
-		auto* Candidate = (FPeerType*)(Waiters[i].Candidate);
-		auto Pred = [Candidate] (auto& Lhs) { return Lhs == Candidate; };
-		FPeerGroup* Group = PeerGroups.FindByPredicate(Pred);
-		check(Group != nullptr);
-		Group->Unwait();
-
-		Waiters.RemoveAtSwap(i, EAllowShrinking::No);
-		--n, --i, ++Count;
-	}
-	check(Count == Result);
-
-	return Result;
-}
-
-////////////////////////////////////////////////////////////////////////////////
-void FHostGroup::Tick(FTickState& State)
-{
-	State.Work = &Work;
-
-	if (BusyCount = Work.HasWork(); BusyCount)
-	{
-		Work.TickCancels(State);
-
-		// Get available work out on idle sockets as soon as possible
-		for (FPeerGroup& Group : PeerGroups)
-		{
-			if (!Work.HasWork())
-			{
-				break;
-			}
-
-			Group.TickSend(State, Host);
-		}
-	}
-
-	// Wait on the groups that are
-	if (int32 Result = Wait(State); Result < 0)
-	{
-		const char* Reason = (Result == MIN_int32)
-			? "FailTimeout hit"
-			: "poll() returned an unexpected error";
-
-		for (FPeerGroup& Group : PeerGroups)
-		{
-			Group.Fail(State, Reason);
-		}
-
-		return;
-	}
-
-	// Tick everything, starting with groups that are maybe closest to finishing
-	for (FPeerGroup& Group : PeerGroups)
-	{
-		BusyCount += (Group.Tick(State) == true);
-	}
-}
-
-////////////////////////////////////////////////////////////////////////////////
-void FHostGroup::AddActivity(FActivity* Activity)
-{
-	Work.AddActivity(Activity);
-}
-
-
-
-// {{{1 event-loop .............................................................
-
-////////////////////////////////////////////////////////////////////////////////
-static const FEventLoop::FRequestParams GDefaultParams;
-
-////////////////////////////////////////////////////////////////////////////////
-class FEventLoop::FImpl
-{
-public:
-							~FImpl();
-	uint32					Tick(int32 PollTimeoutMs=0);
-	bool					IsIdle() const;
-	void					Throttle(uint32 KiBPerSec);
-	void					SetFailTimeout(int32 TimeoutMs);
-	void					Cancel(FTicket Ticket);
-	FRequest				Request(FAnsiStringView Method, FAnsiStringView Path, FActivity* Activity);
-	FTicket					Send(FActivity* Activity);
-
-private:
-	void					ReceiveWork();
-	FCriticalSection		Lock;
-	std::atomic<uint64>		FreeSlots		= ~0ull;
-	std::atomic<uint64>		Cancels			= 0;
-	uint64					PrevFreeSlots	= ~0ull;
-	FActivity*				Pending			= nullptr;
-	FThrottler				Throttler;
-	TArray<FHostGroup>		Groups;
-	int32					FailTimeoutMs	= GIdleMs;
-	uint32					BusyCount		= 0;
-};
-
-////////////////////////////////////////////////////////////////////////////////
-FEventLoop::FImpl::~FImpl()
-{
-	check(BusyCount == 0);
-}
-
-////////////////////////////////////////////////////////////////////////////////
-FRequest FEventLoop::FImpl::Request(
-	FAnsiStringView Method,
-	FAnsiStringView Path,
-	FActivity* Activity)
-{
-	Trace(Activity, ETrace::ActivityCreate, this);
-
-	if (Path.Len() == 0)
-	{
-		Path = "/";
-	}
-
-	Activity->NoContent = (Method == "HEAD");
-
-	FMessageBuilder Builder(Activity->Buffer);
-
-	Builder << Method << " " << Path << " HTTP/1.1" "\r\n"
-		"Host: " << Activity->Host->GetHostName() << "\r\n";
-
-	// HTTP/1.1 is persistent by default thus "Connection" header isn't required
-	if (!Activity->IsKeepAlive)
-	{
-		Builder << "Connection: close\r\n";
-	}
-
-	FRequest Ret;
-	Ret.Ptr = Activity;
-	return Ret;
-}
-
-////////////////////////////////////////////////////////////////////////////////
-FTicket FEventLoop::FImpl::Send(FActivity* Activity)
-{
-	Trace(Activity, ETrace::RequestBegin);
-
-	FMessageBuilder(Activity->Buffer) << "\r\n";
-	Activity_ChangeState(Activity, FActivity::EState::Send);
-
-	uint64 Slot;
-	{
-		FScopeLock _(&Lock);
-
-		for (;; FPlatformProcess::SleepNoStats(0.0f))
-		{
-			uint64 FreeSlotsLoad = FreeSlots.load(std::memory_order_relaxed);
-			if (!FreeSlotsLoad)
-			{
-				// we don't handle oversubscription at the moment. Could return
-				// activity to Reqeust and return a 0 ticket.
-				check(false);
-			}
-			Slot = -int64(FreeSlotsLoad) & FreeSlotsLoad;
-			if (FreeSlots.compare_exchange_weak(FreeSlotsLoad, FreeSlotsLoad - Slot, std::memory_order_relaxed))
-			{
-				break;
-			}
-		}
-		Activity->Slot = int8(63 - FMath::CountLeadingZeros64(Slot));
-
-		// This puts pending requests in reverse order of when they were made
-		// but this will be undone when ReceiveWork() traverses the list
-		Activity->Next = Pending;
-		Pending = Activity;
-	}
-
-	return Slot;
-}
-
-////////////////////////////////////////////////////////////////////////////////
-bool FEventLoop::FImpl::IsIdle() const
-{
-	return FreeSlots.load(std::memory_order_relaxed) == ~0ull;
-}
-
-////////////////////////////////////////////////////////////////////////////////
-void FEventLoop::FImpl::Throttle(uint32 KiBPerSec)
-{
-	Throttler.SetLimit(KiBPerSec);
-}
-
-////////////////////////////////////////////////////////////////////////////////
-void FEventLoop::FImpl::SetFailTimeout(int32 TimeoutMs)
-{
-	// While TimeoutMs must be >=0, it is signed so that MAX_uint32 can't be used
-	check(TimeoutMs > 0);
-	FailTimeoutMs = TimeoutMs;
-}
-
-////////////////////////////////////////////////////////////////////////////////
-void FEventLoop::FImpl::Cancel(FTicket Ticket)
-{
-	Cancels.fetch_or(Ticket, std::memory_order_relaxed);
-}
-
-////////////////////////////////////////////////////////////////////////////////
-void FEventLoop::FImpl::ReceiveWork()
-{
-	uint64 FreeSlotsLoad = FreeSlots.load(std::memory_order_relaxed);
-	if (FreeSlots == PrevFreeSlots)
-	{
-		return;
-	}
-	PrevFreeSlots = FreeSlotsLoad;
-
-	// Fetch the pending activities from out in the wild
-	FActivity* Activity = nullptr;
-	{
-		FScopeLock _(&Lock);
-		Swap(Activity, Pending);
-	}
-
-	// Pending is in the reverse of the order that requests were made
-	FActivity* Reverse = nullptr;
-	for (FActivity* Next; Activity != nullptr; Activity = Next)
-	{
-		Next = Activity->Next;
-		Activity->Next = Reverse;
-		Reverse = Activity;
-	}
-	Activity = Reverse;
-
-	// Group activities by their host.
-	for (FActivity* Next; Activity != nullptr; Activity = Next)
-	{
-		Next = Activity->Next;
-		Activity->Next = nullptr;
-
-		FHost& Host = *(Activity->Host);
-		auto Pred = [&Host] (const FHostGroup& Lhs) { return &Lhs.GetHost() == &Host; };
-		FHostGroup* Group = Groups.FindByPredicate(Pred);
-		if (Group == nullptr)
-		{
-			Group = &(Groups.Emplace_GetRef(Host));
-		}
-
-		Group->AddActivity(Activity);
-		++BusyCount;
-	}
-}
-
-////////////////////////////////////////////////////////////////////////////////
-uint32 FEventLoop::FImpl::Tick(int32 PollTimeoutMs)
-{
-	TRACE_CPUPROFILER_EVENT_SCOPE(IasHttp::Tick);
-
-	ReceiveWork();
-
-	// We limit recv sizes as a way to control bandwidth use.
-	int32 RecvAllowance = Throttler.GetAllowance();
-	if (RecvAllowance <= 0)
-	{
-		if (PollTimeoutMs == 0)
-		{
-			return BusyCount;
-		}
-
-		int32 ThrottleWaitMs = -RecvAllowance;
-		if (PollTimeoutMs > 0)
-		{
-			ThrottleWaitMs = FMath::Min(ThrottleWaitMs, PollTimeoutMs);
-		}
-		FPlatformProcess::SleepNoStats(float(ThrottleWaitMs) / 1000.0f);
-
-		RecvAllowance = Throttler.GetAllowance();
-		if (RecvAllowance <= 0)
-		{
-			return BusyCount;
-		}
-	}
-
-	uint64 CancelsLoad = Cancels.load(std::memory_order_relaxed);
-
-	uint32 NowMs;
-	{
-		// 4.2MM seconds will give us 50 days of uptime.
-		static uint64 Freq = 0;
-		static uint64 Base = 0;
-		if (Freq == 0)
-		{
-			Freq = uint64(1.0 / FPlatformTime::GetSecondsPerCycle());
-			Base = FPlatformTime::Cycles64();
-		}
-		uint64 NowBig = ((FPlatformTime::Cycles64() - Base) * 1000) / Freq;
-		NowMs = uint32(NowBig);
-		check(NowMs == NowBig);
-	}
-
-	// Tick groups and then remove ones that are idle
-	FTickState TickState = {
-		.DoneList = nullptr,
-		.Cancels = CancelsLoad,
-		.RecvAllowance = RecvAllowance,
-		.PollTimeoutMs = PollTimeoutMs,
-		.FailTimeoutMs = FailTimeoutMs,
-		.NowMs = NowMs,
-	};
-	for (FHostGroup& Group : Groups)
-	{
-		Group.Tick(TickState);
-	}
-
-	for (uint32 i = 0, n = Groups.Num(); i < n; ++i)
-	{
-		FHostGroup& Group = Groups[i];
-		if (Group.IsBusy())
-		{
-			continue;
-		}
-
-		Groups.RemoveAtSwap(i, EAllowShrinking::No);
-		--n, --i;
-	}
-
-	Throttler.ReturnUnused(RecvAllowance);
-
-	uint64 ReturnedSlots = 0;
-	for (FActivity* Activity = TickState.DoneList; Activity != nullptr;)
-	{
-		FActivity* Next = Activity->Next;
-		ReturnedSlots |= (1ull << Activity->Slot);
-		Activity_Free(Activity);
-		--BusyCount;
-		Activity = Next;
-	}
-
-	uint32 BusyBias = 0;
-	if (ReturnedSlots)
-	{
-		uint64 LatestFree = FreeSlots.fetch_add(ReturnedSlots, std::memory_order_relaxed);
-		BusyBias += (LatestFree != PrevFreeSlots);
-		PrevFreeSlots += ReturnedSlots;
-	}
-
-	if (CancelsLoad)
-	{
-		Cancels.fetch_and(~CancelsLoad, std::memory_order_relaxed);
-	}
-
-	return BusyCount + BusyBias;
-}
-
-
-
-////////////////////////////////////////////////////////////////////////////////
-FEventLoop::FEventLoop()						{ Impl = new FEventLoop::FImpl(); Trace(Impl, ETrace::LoopCreate); }
-FEventLoop::~FEventLoop()						{ Trace(Impl, ETrace::LoopDestroy); delete Impl; }
-uint32 FEventLoop::Tick(int32 PollTimeoutMs)	{ return Impl->Tick(PollTimeoutMs); }
-bool FEventLoop::IsIdle() const					{ return Impl->IsIdle(); }
-void FEventLoop::Cancel(FTicket Ticket)			{ return Impl->Cancel(Ticket); }
-void FEventLoop::Throttle(uint32 KiBPerSec)		{ return Impl->Throttle(KiBPerSec); }
-void FEventLoop::SetFailTimeout(int32 Ms)		{ return Impl->SetFailTimeout(Ms); }
-
-////////////////////////////////////////////////////////////////////////////////
-FRequest FEventLoop::Request(
-	FAnsiStringView Method,
-	FAnsiStringView Url,
-	const FRequestParams* Params)
-{
-	// Parse the URL into its components
-	FUrlOffsets UrlOffsets;
-	if (ParseUrl(Url, UrlOffsets) < 0)
-	{
-		return FRequest();
-	}
-
-	FAnsiStringView HostName = UrlOffsets.HostName.Get(Url);
-
-	uint32 Port = 0;
-	if (UrlOffsets.Port)
-	{
-		FAnsiStringView PortView = UrlOffsets.Port.Get(Url);
-		Port = uint32(CrudeToInt(PortView));
-	}
-
-	FAnsiStringView Path;
-	if (UrlOffsets.Path > 0)
-	{
-		Path = Url.Mid(UrlOffsets.Path);
-	}
-
-	// Create an activity and an emphemeral host
-	Params = (Params != nullptr) ? Params : &GDefaultParams;
-
-	uint32 BufferSize = Params->BufferSize;
-	BufferSize = (BufferSize >= 128) ? BufferSize : 128;
-	BufferSize += sizeof(FHost) + HostName.Len();
-	FActivity* Activity = Activity_Alloc(BufferSize);
-
-	FBuffer& Buffer = Activity->Buffer;
-
-	FHost* Host = Activity->Host = Buffer.Alloc<FHost>();
-	Activity->IsKeepAlive = 0;
-	Activity->bFollow30x = (Params->bAutoRedirect == true);
-
-	uint32 HostNameLength = HostName.Len();
-	char* HostNamePtr = Buffer.Alloc<char>(HostNameLength + 1);
-
-	Buffer.Fix();
-
-	memcpy(HostNamePtr, HostName.GetData(), HostNameLength);
-	HostNamePtr[HostNameLength] = '\0';
-	new (Host) FHost({
-		.HostName	= HostNamePtr,
-		.Port		= Port,
-	});
-
-	return Impl->Request(Method, Path, Activity);
-}
-
-////////////////////////////////////////////////////////////////////////////////
-FRequest FEventLoop::Request(
-	FAnsiStringView Method,
-	FAnsiStringView Path,
-	FConnectionPool& Pool,
-	const FRequestParams* Params)
-{
-	check(Pool.Ptr != nullptr);
-
-	Params = (Params != nullptr) ? Params : &GDefaultParams;
-
-	uint32 BufferSize = Params->BufferSize;
-	BufferSize = (BufferSize >= 128) ? BufferSize : 128;
-	FActivity* Activity = Activity_Alloc(BufferSize);
-
-	Activity->Host = Pool.Ptr;
-	Activity->IsKeepAlive = 1;
-	Activity->bFollow30x = (Params->bAutoRedirect == true);
-
-	return Impl->Request(Method, Path, Activity);
-}
-
-////////////////////////////////////////////////////////////////////////////////
-bool FEventLoop::Redirect(const FTicketStatus& Status, FTicketSink& OuterSink)
-{
-	const FResponse& Response = Status.GetResponse();
-
-	switch (Response.GetStatusCode())
-	{
-		case 301:				// RedirectMoved
-		case 302:				// RedirectFound
-		case 307:				// RedirectTemp
-		case 308: break;		// RedirectPerm
-		default: return false;
-	}
-
-	FAnsiStringView Location = Response.GetHeader("Location");
-	if (Location.IsEmpty())
-	{
-		return false;
-	}
-
-	check(Response.GetContentLength() == 0); // should we ever hit this, we'll fix it
-
-	const auto& Activity = (FActivity&)Response; // todo: yuk
-
-	// Original method should remain unchanged
-	const char* Data = Activity.Buffer.GetData();
-	FAnsiStringView Method;
-	for (uint32 i = 0; i < 5; ++i)
-	{
-		if (Data[i] <= ' ')
-		{
-			Method = FAnsiStringView(Data, i);
-			break;
-		}
-	}
-	check(!Method.IsEmpty());
-
-	FRequest ForwardRequest;
-	if (!Location.StartsWith("http://") && !Location.StartsWith("https://"))
-	{
-		if (Location[0] != '/')
-		{
-			return false;
-		}
-
-		FHost& Host = *(Activity.Host);
-
-		TAnsiStringBuilder<256> Url;
-		Url << "http://";
-		Url << Host.GetHostName();
-		Url << ":" << Host.GetPort();
-		Url << Location;
-
-		new (&ForwardRequest) FRequest(Request(Method, Url));
-	}
-	else
-	{
-		new (&ForwardRequest) FRequest(Request(Method, Location));
-	}
-
-	// Transfer original request headers
-	check(Activity.State == FActivity::EState::RecvMessage);
-	uint32 Length = 0;
-	for (const char* End = Data + Activity.StateParam; Data + Length < End; ++Length)
-	{
-		if (Data[Length] != '\n')
-		{
-			continue;
-		}
-
-		Data = Data + Length + 1;
-		Length = uint32(ptrdiff_t(End - Data));
-		break;
-	}
-
-	FAnsiStringView OriginalHeaders(Data, Length);
-	EnumerateHeaders(OriginalHeaders, [&ForwardRequest] (FAnsiStringView Name, FAnsiStringView Value)
-	{
-		if (Name != "Host" && Name != "Connection")
-		{
-			ForwardRequest.Header(Name, Value);
-		}
-		return true;
-	});
-
-	// Send the request
-	Send(MoveTemp(ForwardRequest), MoveTemp(OuterSink), Status.GetParam());
-
-	return true;
-}
-
-////////////////////////////////////////////////////////////////////////////////
-FTicket FEventLoop::Send(FRequest&& Request, FTicketSink Sink, UPTRINT SinkParam)
-{
-	FActivity* Activity = nullptr;
-	Swap(Activity, Request.Ptr);
-	Activity->SinkParam = SinkParam;
-	Activity->Sink = Sink;
-
-	// Intercept sink calls to catch 30x status codes and follow them
-	if (Activity->bFollow30x)
-	{
-		auto RedirectSink = [
-				this,
-				OuterSink=MoveTemp(Activity->Sink)
-			] (const FTicketStatus& Status) mutable
-		{
-			if (Status.GetId() == FTicketStatus::EId::Response)
-			{
-				if (Redirect(Status, OuterSink))
-				{
-					return;
-				}
-			}
-
-			if (OuterSink)
-			{
-				return OuterSink(Status);
-			}
-		};
-		Activity->Sink = RedirectSink;
-	}
-
-	return Impl->Send(Activity);
-}
-
-
-
-// {{{1 test ...................................................................
-
-#if !(UE_BUILD_SHIPPING|UE_BUILD_TEST)
-
-////////////////////////////////////////////////////////////////////////////////
-static void MiscTest()
-{
-#define CRLF "\r\n"
-	struct {
-		FAnsiStringView Input;
-		int32 Output;
-	} FmtTestCases[] = {
-		{ "", -1 },
-		{ "abcd", -1 },
-		{ "abcd\r", -1 },
-		{ CRLF "\r\r", -1 },
-		{ CRLF CRLF, 4 },
-		{ "abc" CRLF CRLF, 7 },
-	};
-	for (const auto [Input, Output] : FmtTestCases)
-	{
-		check(FindMessageTerminal(Input.GetData(), Input.Len()) == Output);
-	}
-
-	FMessageOffsets MsgOut;
-	check(ParseMessage("", MsgOut) == -1);
-	check(ParseMessage("MR", MsgOut) == -1);
-	check(ParseMessage("HTTP/1.1", MsgOut) == -1);
-	check(ParseMessage("HTTP/1.1 ", MsgOut) == -1);
-	check(ParseMessage("HTTP/1.1 1" CRLF, MsgOut) > 0);
-	check(ParseMessage("HTTP/1.1    1" CRLF, MsgOut) > 0);
-	check(ParseMessage("HTTP/1.1 100 " CRLF, MsgOut) > 0);
-	check(ParseMessage("HTTP/1.1 100  Message of some sort    " CRLF, MsgOut) > 0);
-	check(ParseMessage("HTTP/1.1 100 _Message with a \r in it" CRLF, MsgOut) == -1);
-
-	bool AllIsWell = true;
-	auto NotExpectedToBeCalled = [&AllIsWell] (auto, auto)
-	{
-		AllIsWell = false;
-		return false;
-	};
-
-	EnumerateHeaders("",		NotExpectedToBeCalled); check(AllIsWell);
-	EnumerateHeaders(CRLF,		NotExpectedToBeCalled); check(AllIsWell);
-	EnumerateHeaders("foo",		NotExpectedToBeCalled); check(AllIsWell);
-	EnumerateHeaders(" foo",	NotExpectedToBeCalled); check(AllIsWell);
-	EnumerateHeaders(" foo ",	NotExpectedToBeCalled); check(AllIsWell);
-	EnumerateHeaders("foo:bar",	NotExpectedToBeCalled); check(AllIsWell);
-
-	auto IsBar = [&] (auto, auto Value) { return AllIsWell = (Value == "bar"); };
-	EnumerateHeaders("foo: bar" CRLF,		IsBar); check(AllIsWell);
-	EnumerateHeaders("foo: bar \t" CRLF,	IsBar); check(AllIsWell);
-	EnumerateHeaders("foo:\tbar " CRLF,		IsBar); check(AllIsWell);
-	EnumerateHeaders("foo:bar " CRLF,		IsBar); check(AllIsWell);
-	EnumerateHeaders("foo:bar" CRLF "!",	IsBar); check(AllIsWell);
-	EnumerateHeaders("foo:bar" CRLF " ",	IsBar); check(AllIsWell);
-	EnumerateHeaders("foo:bar" CRLF "n:ej",	IsBar); check(AllIsWell);
-
-	check(CrudeToInt("") < 0);
-	check(CrudeToInt("X") < 0);
-	check(CrudeToInt("/") < 0);
-	check(CrudeToInt(":") < 0);
-	check(CrudeToInt("-1") < -1);
-	check(CrudeToInt("0") == 0);
-	check(CrudeToInt("9") == 9);
-	check(CrudeToInt("493") == 493);
-
-	FUrlOffsets UrlOut;
-
-	check(ParseUrl("", UrlOut) == -1);
-	check(ParseUrl("abc://asd/", UrlOut) == -1);
-	check(ParseUrl("http://", UrlOut) == -1);
-	check(ParseUrl("http://:/", UrlOut) == -1);
-	check(ParseUrl("http://@:/", UrlOut) == -1);
-	check(ParseUrl("http://foo:ba:r/", UrlOut) == -1);
-	check(ParseUrl("http://foo@ba:r/", UrlOut) == -1);
-	check(ParseUrl("http://foo@ba:r", UrlOut) == -1);
-	check(ParseUrl("http://foo@ba:/", UrlOut) == -1);
-	check(ParseUrl("http://foo@ba@9/", UrlOut) == -1);
-	check(ParseUrl("http://@ba:9/", UrlOut) == -1);
-	check(ParseUrl(
-		"http://zzzzzzzzzzzzzzzzzzzzzzzzzzzzzzzzzzzzzzzzzzzzzzzzzzzzzzzzzzzzzzzzzzzzzzzzzzzzzzzzzzzzzzzzzzzzzzzzzzzzzzzzzzzzzzzzzzzzzzzzzz"
-		"zzzzzzzzzzzzzzzzzzzzzzzzzzzzzzzzzzzzzzzzzzzzzzzzzzzzzzzzzzzzzzzzzzzzzzzzzzzzzzzzzzzzzzzzzzzzzzzzzzzzzzzzzzzzzzzzzzzzzzzzzzzzzzzzz"
-		"zzzzzzzzzzzzzzzzzzzzzzzzzzzzzzzzzzzzzzzzzzzzzzzzzzzzzzzzzzzzzzzzzzzzzzzzzzzzzzzzzzzzzzzzzzzzzzzzzzzzzzzzzzzzzzzzzzzzzzzzzzzzzzzzz"
-		"zzzzzzzzzzzzzzzzzzzzzzzzzzzzzzzzzzzzzzzzzzzzzzzzzzzzzzzzzzzzzzzzzzzzzzzzzzzzzzzzzzzzzzzzzzzzzzzzzzzzzzzzzzzzzzzzzzzzzzzzzzzzzzzzz"
-		"zzzzzzzzzzzzzzzzzzzzzzzzzzzzzzzzzzzzzzzzzzzzzzzzzzzzzzzzzzzzzzzzzzzzzzzzzzzzzzzzzzzzzzzzzzzzzzzzzzzzzzzzzzzzzzzzzzzzzzzzzzzzz.com",
-		UrlOut) == -1);
-
-	check(ParseUrl("http://ab-c.com/", UrlOut) > 0);
-	check(ParseUrl("http:[REDACTED_EMAIL]/", UrlOut) > 0);
-	check(ParseUrl("https://abc.com", UrlOut) > 0);
-	check(ParseUrl("https://abc.com:999", UrlOut) > 0);
-	check(ParseUrl("https://abc.com:999/", UrlOut) > 0);
-	check(ParseUrl("https://foo:[REDACTED_EMAIL]:999", UrlOut) > 0);
-	check(ParseUrl("https://foo:[REDACTED_EMAIL]:999/", UrlOut) > 0);
-	check(ParseUrl("https:[REDACTED_EMAIL]:999", UrlOut) > 0);
-	check(ParseUrl("https:[REDACTED_EMAIL]:999/", UrlOut) > 0);
-
-	for (int32 i : { 0x10, 0x20, 0x40, 0x7f, 0xff })
-	{
-		char Url[] = "http://stockholm.patchercache.epicgames.net:123";
-		char Buffer[512];
-		std::memset(Buffer, i, sizeof(Buffer));
-		std::memcpy(Buffer, Url, sizeof(Url) - 1);
-		check(ParseUrl(FAnsiStringView(Buffer, sizeof(Url) - 1), UrlOut) > 0);
-		check(UrlOut.Port.Get(Url) == "123");
-	}
-
-	FAnsiStringView Url = "http://abc:[REDACTED_EMAIL]:999/";
-	check(ParseUrl(Url, UrlOut) > 0);
-	check(UrlOut.SchemeLength == 4);
-	check(UrlOut.UserInfo.Get(Url) == "abc:123");
-	check(UrlOut.HostName.Get(Url) == "bc.com");
-	check(UrlOut.Port.Get(Url) == "999");
-	check(UrlOut.Path == 25);
-#undef CRLF
-
-	static const char* OutcomeMsg = "\x4d\x52";
-	check(FOutcome::Error(OutcomeMsg, -5).IsOk() == false);
-	check(FOutcome::Error(OutcomeMsg, -5).IsWaiting() == false);
-	check(FOutcome::Error(OutcomeMsg, -5).IsError());
-	check(FOutcome::Error(OutcomeMsg, -5).GetErrorCode() == -5);
-	check(FOutcome::Error(OutcomeMsg,  5).GetErrorCode() ==  5);
-	check(FOutcome::Error(OutcomeMsg, -5).GetMessage() == OutcomeMsg);
-
-	check(FOutcome::Ok(  0).IsOk());
-	check(FOutcome::Ok(-13).IsWaiting() == false);
-	check(FOutcome::Ok( 13).IsError() == false);
-
-	check(FOutcome::Waiting().IsOk() == false);
-	check(FOutcome::Waiting().IsWaiting());
-	check(FOutcome::Waiting().IsError() == false);
-}
-
-////////////////////////////////////////////////////////////////////////////////
-static void ThrottleTest(FAnsiStringView TestUrl)
-{
-	enum { TheMax = 0x7fff'fffful };
-
-	check(FThrottler().GetAllowance() >= TheMax);
-
-	FThrottler Throttler;
-	uint64 OneSecond = Throttler.CycleFreq;
-
-	for (uint32 TargetElapsed : { 1, 5, 7 })
-	{
-		uint64 CycleTest = FPlatformTime::Cycles64();
-		for (uint32 i = 0; i < TargetElapsed; ++i)
-		{
-			FPlatformProcess::Sleep(1.0f);
-		}
-		CycleTest = FPlatformTime::Cycles64() - CycleTest;
-		check((CycleTest + (OneSecond / 8)) / OneSecond == TargetElapsed);
-	}
-
-
-	Throttler.SetLimit(0);
-	check(Throttler.GetAllowance(0) >= TheMax);
-	check(Throttler.GetAllowance()  >= TheMax);
-	check(Throttler.GetAllowance()  >= TheMax);
-
-	for (uint32 i : { 10, 63, 100 })
-	{
-		Throttler.SetLimit(i);
-		check(Throttler.GetAllowance( OneSecond          ) == (i << 10));
-		check(Throttler.GetAllowance((OneSecond + 1) >> 1) == (i << 9));
-		check(Throttler.GetAllowance((OneSecond + 3) >> 2) == (i << 8));
-	}
-
-	uint32 Limit = 17 << 10;
-	Throttler = FThrottler();
-	Throttler.SetLimit(Limit >> 10);
-	Throttler.GetAllowance(OneSecond);
-	Throttler.ReturnUnused(Limit >> 1);
-	check(Throttler.GetAllowance(OneSecond) == (Limit + (Limit >> 1)));
-
-	// runaway
-	check(Throttler.GetAllowance(OneSecond * 3) == Limit * 3);
-	check(Throttler.GetAllowance(OneSecond * 4) == Limit * 4);
-	check(Throttler.GetAllowance(OneSecond * 5) == Limit * 4);
-
-	// idle
-	check(Throttler.GetAllowance(OneSecond * 64) == Limit);
-
-	// threshold
-	Limit = FThrottler::THRESHOLD;
-	Throttler = FThrottler();
-	Throttler.SetLimit(Limit >> 10);
-	for (int64 Counter = OneSecond;;)
-	{
-		int64 Delta = (OneSecond + 15) >> 4;
-		Counter -= Delta;
-		int32 Allowance = Throttler.GetAllowance(Delta);
-		if (Allowance > 0)
-		{
-			check(Allowance == Limit);
-			check(Counter < 10);
-			break;
-		}
-	};
-
-	// overflow(ish)
-	Throttler = FThrottler();
-	Throttler.SetLimit((512 << 10) - 1);
-	Throttler.GetAllowance((OneSecond * 790) / 100);
-
-	// timing test
-	FIoBuffer RecvData;
-	for (uint32 SizeKiB : { 10, 25, 60 })
-	{
-		const uint32 ThrottleKiB = 5;
-
-		TAnsiStringBuilder<128> Url;
-		Url << TestUrl;
-		Url << (SizeKiB << 10);
-
-		FEventLoop Loop;
-		Loop.Throttle(ThrottleKiB);
-
-		FRequest Request = Loop.Request("GET", Url).Accept("*/*");
-		Loop.Send(MoveTemp(Request), [&] (const FTicketStatus& Status) {
-			check(Status.GetId() != FTicketStatus::EId::Error);
-			if (Status.GetId() == FTicketStatus::EId::Response)
-			{
-				Status.GetResponse().SetDestination(&RecvData);
-			}
-		});
-
-		int32 Timeout = -1;
-		if (SizeKiB < 25) Timeout = 123;
-		if (SizeKiB > 25) Timeout = 4567;
-
-		uint64 Time = FPlatformTime::Cycles64();
-		while (Loop.Tick(Timeout));
-		Time = FPlatformTime::Cycles64() - Time;
-		Time /= OneSecond;
-
-		// It's dangerous stuff testing elapsed time you know. The +1 is because
-		// throttling assumes one second has already passed when initialised.
-#if PLATFORM_WINDOWS
-		check(Time + 1 == (SizeKiB / ThrottleKiB));
-#endif
-
-		RecvData = FIoBuffer();
-	}
-}
-
-////////////////////////////////////////////////////////////////////////////////
-static void RedirectTest(const ANSICHAR* TestHost)
-{
-	FEventLoop Loop;
-
-	auto WaitForLoopIdle = [&] {
-		for (; Loop.Tick(-1); FPlatformProcess::SleepNoStats(0.02f));
-	};
-
-	FEventLoop::FRequestParams RequestParams = {
-		.bAutoRedirect = true,
-	};
-
-	enum ReTyp { ReAbs, ReRel };
-	enum { RecvDataSize = 48 };
-
-	TAnsiStringBuilder<64> Builder;
-	auto BuildUrl = [&] (ReTyp Typ, uint32 Code) -> const auto&
-	{
-		Builder.Reset();
-		Builder << "https://";
-		Builder << TestHost;
-		Builder << ":9493";
-		Builder << "/redirect";
-		Builder << ((Typ < ReRel) ? "/abs/" : "/rel/");
-		Builder << Code;
-		Builder << "/data/" << uint32(RecvDataSize);
-		return Builder;
-	};
-
-	UPTRINT SinkParam = 0xaa'493'493'493'493'bbull;
-
-	uint32 RecvCount;
-	auto OkSink = [Dest=FIoBuffer(), SinkParam, &RecvCount] (const FTicketStatus& Status) mutable {
-		check(Status.GetParam() == SinkParam);
-		check(Status.GetId() != FTicketStatus::EId::Error);
-		if (Status.GetId() == FTicketStatus::EId::Response)
-		{
-			FResponse& Response = Status.GetResponse();
-			check(Response.GetStatusCode() == 200);
-			Response.SetDestination(&Dest);
-			return;
-		}
-		check(Status.GetId() == FTicketStatus::EId::Content);
-		RecvCount += uint32(Dest.GetSize());
-	};
-
-	uint32 TestCodes[] = { 301, 302, 307, 308 };
-
-	for (auto ReTest : { ReAbs, ReRel })
-	{
-		RecvCount = 0;
-		for (uint32 Code : TestCodes)
-		{
-			FRequest Request = Loop.Get(BuildUrl(ReTest, Code), &RequestParams);
-			if (Code > TestCodes[1])
-			{
-				Request.Header("TestCodeHeader", "Header-Of-Test-Codes");
-			}
-			Loop.Send(MoveTemp(Request), OkSink, SinkParam);
-		}
-		WaitForLoopIdle();
-		check(RecvCount == RecvDataSize * UE_ARRAY_COUNT(TestCodes));
-	}
-
-	RequestParams = FEventLoop::FRequestParams();
-	RequestParams.bAutoRedirect = true;
-
-	FConnectionPool::FParams Params;
-	Params.SetHostFromUrl(BuildUrl(ReAbs, 0));
-	Params.ConnectionCount = 4;
-	FConnectionPool Pool(Params);
-
-	RecvCount = 0;
-	Loop.Send(Loop.Get("/redirect/abs/307/data/55", Pool, &RequestParams), OkSink, SinkParam);
-	WaitForLoopIdle();
-	check(RecvCount == 55);
-}
-
-////////////////////////////////////////////////////////////////////////////////
-static void HttpTest(const ANSICHAR* TestHost)
-{
-	TAnsiStringBuilder<64> Ret;
-	auto BuildUrl = [&] (const ANSICHAR* Suffix=nullptr, uint32 Port=9493) -> const auto&
-	{
-		Ret.Reset();
-		Ret << "http://";
-		Ret << TestHost;
-		Ret << ":" << Port;
-		return (Suffix != nullptr) ? (Ret << Suffix) : Ret;
-	};
-
-	struct
-	{
-		FIoBuffer Dest;
-		uint64 Hash = 0;
-	} Content[64];
-
-	auto HashSink = [&] (const FTicketStatus& Status) -> FIoBuffer*
-	{
-		check(Status.GetId() != FTicketStatus::EId::Error);
-
-		uint32 Index = Status.GetIndex();
-
-		if (Status.GetId() == FTicketStatus::EId::Response)
-		{
-			FResponse& Response = Status.GetResponse();
-			check(Response.GetStatus() == EStatusCodeClass::Successful);
-			check(Response.GetStatusCode() == 200);
-			check(Response.GetContentLength() == Status.GetContentLength());
-
-			FAnsiStringView HashView = Response.GetHeader("X-TestServer-Hash");
-			Content[Index].Hash = CrudeToInt(HashView);
-			check(int64(Content[Index].Hash) > 0);
-
-			Content[Index].Dest = FIoBuffer();
-			Response.SetDestination(&(Content[Index].Dest));
-			return nullptr;
-		}
-
-		uint32 ReceivedHash = 0x493;
-		FMemoryView ContentView = Content[Index].Dest.GetView();
-		check(ContentView.GetSize() == Status.GetContentLength());
-		for (uint32 i = 0; i < Status.GetContentLength(); ++i)
-		{
-			uint8 c = ((const uint8*)(ContentView.GetData()))[i];
-			ReceivedHash = (ReceivedHash + c) * 0x493;
-		}
-		check(Content[Index].Hash == ReceivedHash);
-		Content[Index].Hash = 0;
-
-		return nullptr;
-	};
-
-	auto NullSink = [] (const FTicketStatus&) {};
-
-	auto NoErrorSink = [&] (const FTicketStatus& Status)
-	{
-		check(Status.GetId() != FTicketStatus::EId::Error);
-		if (Status.GetId() != FTicketStatus::EId::Response)
-		{
-			return;
-		}
-
-		uint32 Index = Status.GetIndex();
-
-		FResponse& Response = Status.GetResponse();
-		Content[Index].Dest = FIoBuffer();
-		Response.SetDestination(&(Content[Index].Dest));
-	};
-
-	FEventLoop Loop;
-	volatile bool LoopStop = false;
-	volatile bool LoopTickDelay = false;
-	auto LoopTask = UE::Tasks::Launch(TEXT("IasHttpTest.Loop"), [&] () {
-		uint32 DelaySeed = 493;
-		while (!LoopStop)
-		{
-			while (Loop.Tick())
-			{
-				if (!LoopTickDelay)
-				{
-					continue;
-				}
-
-				float DelayFloat = float(DelaySeed % 75) / 1000.0f;
-				FPlatformProcess::SleepNoStats(DelayFloat);
-				DelaySeed *= 0xa93;
-			}
-
-			FPlatformProcess::SleepNoStats(0.005f);
-		}
-	});
-
-	auto WaitForLoopIdle = [&Loop] ()
-	{
-		FPlatformProcess::SleepNoStats(0.25f);
-		while (!Loop.IsIdle())
-		{
-			FPlatformProcess::SleepNoStats(0.03f);
-		}
-	};
-
-	// unused request
-	{
-		FRequest Request = Loop.Request("GET", BuildUrl("/data"));
-	}
-
-	// foundational
-	{
-		FRequest Request = Loop.Request("GET", BuildUrl("/seed/493"));
-		Request.Accept(EMimeType::Json);
-
-		FTicket Ticket = Loop.Send(MoveTemp(Request), NullSink);
-
-		WaitForLoopIdle();
-	}
-
-	// convenience
-	{
-		FRequest Request = Loop.Get(BuildUrl("/data")).Accept(EMimeType::Json);
-
-		FTicket Tickets[] = {
-			Loop.Send(MoveTemp(Request), HashSink),
-			Loop.Send(Loop.Get(BuildUrl("/data")).Accept(EMimeType::Json), HashSink),
-			Loop.Send(Loop.Get("http://httpbin.org/get"), NoErrorSink),
-		};
-		WaitForLoopIdle();
-	}
-
-	// convenience
-	{
-		FRequest Request = Loop.Get(BuildUrl("/data")).Accept(EMimeType::Json);
-		FTicket Ticket = Loop.Send(MoveTemp(Request), HashSink);
-		WaitForLoopIdle();
-	}
-
-	// pool
-	for (uint16 i = 1; i < 64; ++i)
-	{
-		FConnectionPool::FParams Params;
-		Params.SetHostFromUrl(BuildUrl());
-		Params.ConnectionCount = (i % 2) + 1;
-		FConnectionPool Pool(Params);
-		for (int32 j = 0; j < i; ++j)
-		{
-			FRequest Request = Loop.Get("/data", Pool);
-			Loop.Send(MoveTemp(Request), HashSink);
-		}
-		WaitForLoopIdle();
-	}
-
-	// fatal timeout
-	for (int32 i = 0; i < 14; ++i)
-	{
-		bool bExpectFailTimeout = !!(i & 1);
-		auto Sink = [bExpectFailTimeout, Dest=FIoBuffer()] (const FTicketStatus& Status) mutable
-		{
-			if (Status.GetId() == FTicketStatus::EId::Response)
-			{
-				FResponse& Response = Status.GetResponse();
-				Response.SetDestination(&Dest);
-				return;
-			}
-
-			check(Status.GetId() == FTicketStatus::EId::Error);
-
-			const char* Reason = Status.GetErrorReason();
-			bool IsFailTimeout = (FCStringAnsi::Strstr(Reason, "FailTimeout") != nullptr);
-			check(IsFailTimeout == bExpectFailTimeout);
-		};
-
-		auto ErrorSink = [] (const FTicketStatus& Status)
-		{
-			check(Status.GetId() == FTicketStatus::EId::Error);
-		};
-
-		FConnectionPool::FParams Params;
-		Params.SetHostFromUrl(BuildUrl());
-		FConnectionPool Pool(Params);
-
-		FEventLoop Loop2;
-		Loop2.Send(Loop2.Get("/data?stall=1", Pool), Sink);
-
-		// Requests are pipelined. The second one will get went during the stall so
-		// we expect it to fail. The subsequent ones are expected to succeed.
-		Loop2.Send(Loop2.Get("/data", Pool), ErrorSink);
-		Loop2.Send(Loop2.Get("/data", Pool), HashSink);
-		Loop2.Send(Loop2.Get("/data", Pool), HashSink);
-
-		int32 PollTimeoutMs = -1;
-		if (bExpectFailTimeout)
-		{
-			Loop2.SetFailTimeout(1000);
-
-			if ((i & 3) == 1)
-			{
-				PollTimeoutMs = 1000;
-			}
-		}
-		while (Loop2.Tick(PollTimeoutMs));
-
-		Loop2.Send(Loop2.Get("/data/23", Pool), NoErrorSink);
-		while (Loop2.Tick(PollTimeoutMs));
-	}
-
-	// no connect
-	{
-		FRequest Requests[] = {
-			Loop.Request("GET", BuildUrl(nullptr, 10930)),
-			Loop.Request("GET", "http://thisdoesnotexistihope/"),
-		};
-		Loop.Send(MoveTemp(Requests[0]), NullSink);
-		Loop.Send(MoveTemp(Requests[1]), NullSink);
-		WaitForLoopIdle();
-	}
-
-	// head and large requests
-	{
-		auto MixTh = [Th=uint32(0)] () mutable { return (Th = (Th * 75) + 74) & 255; };
-
-		char AsciiData[257];
-		for (char& c : AsciiData)
-		{
-			int32 i = int32(ptrdiff_t(&c - AsciiData));
-			c = 0x41 + (MixTh() % 26);
-			c += (MixTh() & 2) ? 0x20 : 0;
-		}
-
-		for (int32 i = 0; (i += 69493) < 2 << 20;)
-		{
-			FRequest Request = Loop.Request("HEAD", BuildUrl("/data"));
-			for (int32 j = i; j > 0;)
-			{
-				FAnsiStringView Name(AsciiData, MixTh() + 1);
-				FAnsiStringView Value(AsciiData, MixTh() + 1);
-				Request.Header(Name, Value);
-				j -= Name.Len() + Value.Len();
-
-			}
-
-			Loop.Send(MoveTemp(Request), [] (const FTicketStatus& Status) {
-				if (Status.GetId() == FTicketStatus::EId::Response)
-				{
-					FResponse& Response = Status.GetResponse();
-					check(Response.GetStatusCode() == 431); // "too many headers"
-				}
-			});
-
-			WaitForLoopIdle();
-		}
-	}
-
-	// stress 1
-	{
-		const uint32 StressLoad = 32;
-
-		struct {
-			const ANSICHAR* Uri;
-			bool Disconnect;
-		} StressUrls[] = {
-			{ "/data?slowly=1",		false },
-			{ "/data?disconnect=1",	true },
-		};
-
-		uint64 Errors = 0;
-		auto ErrorSink = [&] (const FTicketStatus& Status)
-		{
-			FTicket Ticket = Status.GetTicket();
-			uint32 Index = uint32(63 - FMath::CountLeadingZeros64(uint64(Ticket)));
-
-			if (Status.GetId() == FTicketStatus::EId::Error)
-			{
-				Errors |= 1ull << Index;
-				return;
-			}
-
-			else if (Status.GetId() == FTicketStatus::EId::Response)
-			{
-				FResponse& Response = Status.GetResponse();
-				Content[Index].Dest = FIoBuffer();
-				Response.SetDestination(&(Content[Index].Dest));
-				return;
-			}
-
-			check(false);
-		};
-
-		for (const auto& [StressUri, ExpectDisconnect] : StressUrls)
-		{
-			FTicketSink Sink = ExpectDisconnect ? FTicketSink(ErrorSink) : FTicketSink(HashSink);
-
-			const auto& StressUrl = BuildUrl(StressUri);
-			for (bool AddDelay : {false, true})
-			{
-				FTicket Tickets[StressLoad];
-				for (FTicket& Ticket : Tickets)
-				{
-					Ticket = Loop.Send(Loop.Get(StressUrl).Header("Accept", "*/*"), Sink);
-				}
-
-				LoopTickDelay = AddDelay;
-
-				WaitForLoopIdle();
-			}
-
-			LoopTickDelay = false;
-		}
-	}
-
-	// stress 2
-	{
-		const uint32 StressLoad = 3;
-		const uint32 StressTaskCount = 7;
-		static_assert(StressLoad * StressTaskCount <= 32);
-
-		FAnsiStringView Url = BuildUrl("/data");
-
-		auto StressTaskEntry = [&] {
-			for (uint32 i = 0; i < StressLoad; ++i)
-			{
-				FTicket Ticket = Loop.Send(Loop.Get(Url), HashSink);
-				if (!Ticket)
-				{
-					FPlatformProcess::SleepNoStats(0.01f);
-					--i;
-				}
-			}
-		};
-
-		UE::Tasks::FTask StressTasks[StressTaskCount];
-		for (auto& StressTask : StressTasks)
-		{
-			StressTask = UE::Tasks::Launch(TEXT("StressTask"), [&] { StressTaskEntry(); });
-		}
-		for (auto& StressTask : StressTasks)
-		{
-			StressTask.Wait();
-		}
-
-		WaitForLoopIdle();
-	}
-
-	// tamper
-	for (int32 i = 1; i <= 100; ++i)
-	{
-		TAnsiStringBuilder<32> TamperUrl;
-		TamperUrl << "/data?tamper=" << i;
-		FAnsiStringView Url = BuildUrl(TamperUrl.ToString());
-
-		for (int j = 0; j < 48; ++j)
-		{
-			FRequest Request = Loop.Request("GET", Url);
-			Loop.Send(MoveTemp(Request), NullSink);
-		}
-
-		WaitForLoopIdle();
-	}
-
-	LoopStop = true;
-	LoopTask.Wait();
-
-	check(Loop.IsIdle());
-
-#if IS_PROGRAM
-	ThrottleTest(BuildUrl("/data/"));
-#endif
-
-	// pre-generated headers
-	// request-with-body
-	// proxy
-	// chunked transfer encoding
-	// gzip / deflate
-	// redirects
-	// loop multi-req.
-	// tls
-	// url auth credentials
-	// transfer-file / splice / sendfile
-	// (header field parser)
-	// (form-data)
-	// (cookies)
-	// (cache)
-	// (websocket)
-	// (ipv6)
-	// (utf-8 host names)
-}
-
-////////////////////////////////////////////////////////////////////////////////
-IOSTOREONDEMAND_API void IasHttpTest(const ANSICHAR* TestHost="localhost")
-{
-#if PLATFORM_WINDOWS
-	WSADATA WsaData;
-	if (WSAStartup(MAKEWORD(2, 2), &WsaData) == 0x0a9e0493)
-		return;
-	ON_SCOPE_EXIT { WSACleanup(); };
-#endif
-
-	MiscTest();
-	HttpTest(TestHost);
-	RedirectTest(TestHost);
-}
-
-#endif // !SHIP|TEST
-
-// }}}
-
-} // namespace UE::IoStore::HTTP
+#include "Misc.inl"
+#include "Socket.inl"
+#include "Socks.inl"
+#include "ConnectionPool.inl"
+#include "Activity.inl"
+#include "Api.inl"
+#include "Loop.inl"
+#include "Test.inl"
diff --git a/Engine/Source/Runtime/Experimental/IoStoreOnDemand/Private/CoreHttp/Loop.inl b/Engine/Source/Runtime/Experimental/IoStoreOnDemand/Private/CoreHttp/Loop.inl
new file mode 100644
--- /dev/null
+++ b/Engine/Source/Runtime/Experimental/IoStoreOnDemand/Private/CoreHttp/Loop.inl
@@ -0,0 +1,1448 @@
+// Copyright Epic Games, Inc. All Rights Reserved.
+
+#pragma once
+
+namespace UE::IoStore::HTTP
+{
+
+// {{{1 event-loop-int .........................................................
+
+using FPeerType = FSocket;
+
+////////////////////////////////////////////////////////////////////////////////
+static FOutcome DoSend(FActivity* Activity, FPeerType& Peer)
+{
+	Trace(Activity, ETrace::StateChange, Activity->State);
+
+#if IAS_HTTP_WITH_PERF
+	Activity->Stopwatch.SendStart();
+#endif
+
+	TRACE_CPUPROFILER_EVENT_SCOPE(IasHttp::DoSend);
+
+	FBuffer& Buffer = Activity->Buffer;
+	const char* SendData = Buffer.GetData();
+	int32 SendSize = Buffer.GetSize();
+
+	uint32 AlreadySent = Activity->StateParam;
+	SendData += AlreadySent;
+	SendSize -= AlreadySent;
+	check(SendSize > 0);
+
+	FOutcome Outcome = Peer.Send(SendData, SendSize);
+
+	if (Outcome.IsError())
+	{
+		Activity_SetError(Activity, Outcome.GetMessage().GetData());
+		return Outcome;
+	}
+
+	if (Outcome.IsWaiting())
+	{
+		return Outcome;
+	}
+
+	check(Outcome.IsOk());
+
+	int32 Result = Outcome.GetResult();
+	Activity->StateParam += Result;
+	if (Activity->StateParam < Buffer.GetSize())
+	{
+		return DoSend(Activity, Peer);
+	}
+
+#if IAS_HTTP_WITH_PERF
+	Activity->Stopwatch.SendEnd();
+#endif
+
+	Activity_ChangeState(Activity, FActivity::EState::RecvMessage, Buffer.GetSize());
+	return FOutcome::Ok(Result);
+}
+
+////////////////////////////////////////////////////////////////////////////////
+static FOutcome DoRecvMessage(FActivity* Activity, FPeerType& Peer)
+{
+	TRACE_CPUPROFILER_EVENT_SCOPE(IasHttp::DoRecvMessage);
+
+	static const uint32 PageSize = 256;
+
+	FBuffer& Buffer = Activity->Buffer;
+
+	const char* MessageRight;
+	while (true)
+	{
+		Trace(Activity, ETrace::StateChange, Activity->State);
+
+#if IAS_HTTP_WITH_PERF
+		Activity->Stopwatch.RecvStart();
+#endif
+
+		auto [Dest, DestSize] = Buffer.GetMutableFree(0, PageSize);
+		FOutcome Outcome = Peer.Recv(Dest, DestSize);
+
+		if (Outcome.IsError())
+		{
+			Activity_SetError(Activity, Outcome.GetMessage().GetData());
+			return Outcome;
+		}
+
+		if (Outcome.IsWaiting())
+		{
+			return Outcome;
+		}
+
+		check(Outcome.IsOk());
+
+		int32 Result = Outcome.GetResult();
+		Buffer.AdvanceUsed(Result);
+
+		// Rewind a little to cover cases where the terminal is fragmented across
+		// recv() calls
+		uint32 DestBias = 0;
+		if (Dest - 3 >= Buffer.GetData() + Activity->StateParam)
+		{
+			Dest -= (DestBias = 3);
+		}
+
+		int32 MessageEnd = FindMessageTerminal(Dest, Result + DestBias);
+		if (MessageEnd < 0)
+		{
+			if (Buffer.GetSize() > (8 << 10))
+			{
+				Activity_SetError(Activity, "Headers have grown larger than expected");
+				return FOutcome::Error(Activity->ErrorReason);
+			}
+
+			continue;
+		}
+
+		MessageRight = Dest + MessageEnd;
+		break;
+	}
+
+	// Fill out the internal response object
+	FResponseInternal& Internal = Activity->Response;
+	const char* MessageData = Buffer.GetData() + Activity->StateParam;
+	Internal.MessageLength = uint16(ptrdiff_t(MessageRight - MessageData));
+
+	FAnsiStringView ResponseView(MessageData, Internal.MessageLength);
+	if (ParseMessage(ResponseView, Internal.Offsets) < 0)
+	{
+		Activity_SetError(Activity, "Failed to parse message status");
+		return FOutcome::Error(Activity->ErrorReason);
+	}
+
+	// Parse headers
+	FAnsiStringView Headers = ResponseView.Mid(
+		Internal.Offsets.Headers,
+		Internal.MessageLength - Internal.Offsets.Headers - 2 // "-2" trims off '\r\n' that signals end of headers
+	);
+
+	int32 Count = 2;
+	bool IsKeepAlive = true;
+	int32 ContentLength = -1;
+	EnumerateHeaders(
+		Headers,
+		[&ContentLength, &IsKeepAlive, &Count] (FAnsiStringView Name, FAnsiStringView Value)
+		{
+			// todo; may need smarter value handling; ;/, separated options & key-value pairs (ex. in rfc2068)
+
+			// "Keep-Alive"			- deprecated
+			// "Transfer-Encoding"	- may be required later
+
+			if (Name.Equals("Content-Length", ESearchCase::IgnoreCase))
+			{
+				ContentLength = int32(CrudeToInt(Value));
+				Count--;
+			}
+
+			else if (Name.Equals("Connection", ESearchCase::IgnoreCase))
+			{
+				IsKeepAlive = !Value.Equals("close");
+				Count--;
+			}
+
+			return Count > 0;
+		}
+	);
+
+	Activity->IsKeepAlive &= IsKeepAlive;
+
+	if (ContentLength < 0)
+	{
+		if (ContentLength == -1)
+		{
+			// todo; query Transfer-Encoding for chunked
+		}
+
+		Activity_SetError(Activity, "Unknown content length value");
+		return FOutcome::Error(Activity->ErrorReason);
+	}
+
+	// Call out to the sink to get a content destination
+	FIoBuffer* PriorDest = Activity->Dest; // to retain unioned Host ptr
+	Internal.Code = -1;
+	Internal.ContentLength = ContentLength;
+	{
+		FTicketStatus& SinkArg = *(FTicketStatus*)Activity;
+		Activity->Sink(SinkArg);
+	}
+
+	if (Activity->NoContent == 0)
+	{
+		if (Activity->Dest == PriorDest)
+		{
+			Activity_SetError(Activity, "User did not provide a destination buffer");
+			return FOutcome::Error(Activity->ErrorReason);
+		}
+
+		// The user seems to have forgotten something. Let's help them along
+		if (Activity->Dest->GetSize() == 0)
+		{
+			*Activity->Dest = FIoBuffer(ContentLength);
+		}
+	}
+
+	// Perhaps we have some of the content already?
+	const char* BufferRight = Buffer.GetData() + Buffer.GetSize();
+	uint32 AlreadyReceived = uint32(ptrdiff_t(BufferRight - MessageRight));
+	if (AlreadyReceived > uint32(ContentLength))
+	{
+		Activity_SetError(Activity, "More data received that expected");
+		return FOutcome::Error(Activity->ErrorReason);
+	}
+
+	if (Activity->NoContent == 1)
+	{
+		if (AlreadyReceived)
+		{
+			Activity_SetError(Activity, "Received content when none was expected");
+			return FOutcome::Error(Activity->ErrorReason);
+		}
+		Activity_ChangeState(Activity, FActivity::EState::RecvDone);
+		return FOutcome::Ok();
+	}
+
+	check(Activity->Dest != nullptr);
+
+	const bool bStreamed = Activity->Dest->GetSize() < ContentLength;
+
+	auto NextState = bStreamed ? FActivity::EState::RecvStream : FActivity::EState::RecvContent;
+	Activity_ChangeState(Activity, NextState, AlreadyReceived);
+
+	if (AlreadyReceived == 0)
+	{
+		return FOutcome::Ok();
+	}
+
+	FMutableMemoryView DestView = Activity->Dest->GetMutableView();
+	const char* Cursor = BufferRight - AlreadyReceived;
+	if (!bStreamed || AlreadyReceived < DestView.GetSize())
+	{
+		::memcpy(DestView.GetData(), Cursor, AlreadyReceived);
+		return FOutcome::Ok();
+	}
+
+#if 1
+	check(false); // not implemented yet
+#else
+	do
+	{
+		uint32 Size = FMath::Min<uint32>(DestView.GetSize(), AlreadyReceived);
+		::memcpy(DestView.GetData(), Cursor, Size);
+
+		/* send sink */
+
+		AlreadyReceived -= Size;
+		Cursor += Size;
+	}
+	while (AlreadyReceived);
+#endif
+
+	return FOutcome::Ok();
+}
+
+////////////////////////////////////////////////////////////////////////////////
+static FOutcome DoRecvContent(FActivity* Activity, FPeerType& Peer, int32& MaxRecvSize)
+{
+	TRACE_CPUPROFILER_EVENT_SCOPE(IasHttp::DoRecvContent);
+
+	FResponseInternal& Response = Activity->Response;
+	FMutableMemoryView DestView = Activity->Dest->GetMutableView();
+
+	while (true)
+	{
+		int32 Size = (Response.ContentLength - Activity->StateParam);
+		if (Size == 0)
+		{
+			break;
+		}
+
+		Size = FMath::Min(Size, MaxRecvSize);
+		check(Size >= 0);
+		if (Size == 0)
+		{
+			return FOutcome::Waiting();
+		}
+
+		Trace(Activity, ETrace::StateChange, Activity->State);
+
+		char* Cursor = (char*)(DestView.GetData()) + Activity->StateParam;
+
+		FOutcome Outcome = Peer.Recv(Cursor, Size);
+
+		if (Outcome.IsWaiting())
+		{
+			return Outcome;
+		}
+
+		if (Outcome.IsError())
+		{
+			Activity_SetError(Activity, Outcome.GetMessage().GetData());
+			return Outcome;
+		}
+
+		check(Outcome.IsOk());
+
+		int32 Result = Outcome.GetResult();
+		check(Result <= MaxRecvSize);
+
+		Activity->StateParam += Result;
+		MaxRecvSize -= Result;
+	}
+
+	if (!FLatencyInjector::Begin(FLatencyInjector::EType::Network, Activity->StateParam))
+	{
+		Activity_SetError(Activity, "Forced random failure");
+		return FOutcome::Error(Activity->ErrorReason);
+	}
+
+#if IAS_HTTP_WITH_PERF
+	Activity->Stopwatch.RecvEnd();
+#endif
+
+	Activity_ChangeState(Activity, FActivity::EState::RecvDone);
+	return FOutcome::Ok();
+}
+
+////////////////////////////////////////////////////////////////////////////////
+static FOutcome DoRecvStream(FActivity*, FPeerType&, uint32)
+{
+	check(false); // not yet implemented
+	return FOutcome::Error("Not supported");
+}
+
+////////////////////////////////////////////////////////////////////////////////
+static FOutcome DoRecv(FActivity* Activity, FPeerType& Peer, int32& MaxRecvSize)
+{
+	using EState = FActivity::EState;
+
+	EState State = Activity->State; 
+	check(State >= EState::RecvMessage && State < EState::RecvDone);
+
+	if (State == EState::RecvMessage)	return DoRecvMessage(Activity, Peer);
+	if (State == EState::RecvContent)	return DoRecvContent(Activity, Peer, MaxRecvSize);
+	if (State == EState::RecvStream)	return DoRecvStream(Activity, Peer, MaxRecvSize);
+	
+	check(false); // it is not expected that we'll get here
+	return FOutcome::Error("unreachable");
+}
+
+////////////////////////////////////////////////////////////////////////////////
+static int32 DoRecvDone(FActivity* Activity)
+{
+	if (!FLatencyInjector::HasExpired(Activity->StateParam))
+	{
+		return 1;
+	}
+
+	// Notify the user we've received everything
+	FTicketStatus& SinkArg = *(FTicketStatus*)Activity;
+	Activity->Sink(SinkArg);
+
+	Activity_ChangeState(Activity, FActivity::EState::Completed);
+	return 0;
+}
+
+////////////////////////////////////////////////////////////////////////////////
+static void DoCancel(FActivity* Activity)
+{
+	if (Activity->State >= FActivity::EState::Completed)
+	{
+		return;
+	}
+
+	Activity_ChangeState(Activity, FActivity::EState::Cancelled);
+
+	FTicketStatus& SinkArg = *(FTicketStatus*)Activity;
+	Activity->Sink(SinkArg);
+}
+
+////////////////////////////////////////////////////////////////////////////////
+static void DoFail(FActivity* Activity)
+{
+	check(Activity->State == FActivity::EState::Failed);
+
+	// Notify the user we've received everything
+	FTicketStatus& SinkArg = *(FTicketStatus*)Activity;
+	Activity->Sink(SinkArg);
+}
+
+
+
+// {{{1 work-queue .............................................................
+
+/*
+ * - Activities (requests send with a loop) are managed in singly-linked lists
+ * - Each activity has an associated host it is talking to.
+ * - Hosts are ephemeral, or represented externally via a FConnectionPool object
+ * - Loop has a group for each host, and each host-group has a bunch of socket-groups
+ * - Host-group has a list of work; pending activities waiting to start
+ * - Socket-groups own up to two activities; one sending, one receiving
+ * - As it recvs, a socket-group will, if possible, fetch more work from the host
+ *
+ *  Loop:
+ *    FHostGroup[HostPtr]:
+ *	    Work: Act0 -> Act1 -> Act2 -> Act3 -> ...
+ *      FPeerGroup[0...HostMaxConnections]:
+ *			Act.Send
+ *			Act.Recv
+ */
+
+////////////////////////////////////////////////////////////////////////////////
+struct FTickState
+{
+	FActivity*					DoneList;
+	uint64						Cancels;
+	int32&						RecvAllowance;
+	int32						PollTimeoutMs;
+	int32						FailTimeoutMs;
+	uint32						NowMs;
+	class FWorkQueue*			Work;
+};
+
+
+
+////////////////////////////////////////////////////////////////////////////////
+class FWorkQueue
+{
+public:
+						FWorkQueue() = default;
+						~FWorkQueue();
+	bool				HasWork() const { return List != nullptr; }
+	void				AddActivity(FActivity* Activity);
+	FActivity*			PopActivity();
+	void				TickCancels(FTickState& State);
+
+private:
+	FActivity*			List = nullptr;
+	FActivity*			ListTail = nullptr;
+	uint64				ActiveSlots = 0;
+
+	UE_NONCOPYABLE(FWorkQueue);
+};
+
+////////////////////////////////////////////////////////////////////////////////
+FWorkQueue::~FWorkQueue()
+{
+	check(List == nullptr);
+	check(ListTail == nullptr);
+}
+
+////////////////////////////////////////////////////////////////////////////////
+void FWorkQueue::AddActivity(FActivity* Activity)
+{
+	// We use a tail pointer here to maintain order that requests were made
+
+	check(Activity->Next == nullptr);
+
+	if (ListTail != nullptr)
+	{
+		ListTail->Next = Activity;
+	}
+	List = (List == nullptr) ? Activity : List;
+	ListTail = Activity;
+
+	ActiveSlots |= (1ull << Activity->Slot);
+}
+
+////////////////////////////////////////////////////////////////////////////////
+FActivity* FWorkQueue::PopActivity()
+{
+	if (List == nullptr)
+	{
+		return nullptr;
+	}
+
+	FActivity* Activity = List;
+	if ((List = List->Next) == nullptr)
+	{
+		ListTail = nullptr;
+	}
+
+	check(ActiveSlots & (1ull << Activity->Slot));
+	ActiveSlots ^= (1ull << Activity->Slot);
+
+	Activity->Next = nullptr;
+	return Activity;
+}
+
+////////////////////////////////////////////////////////////////////////////////
+void FWorkQueue::TickCancels(FTickState& State)
+{
+	if (State.Cancels == 0 || (State.Cancels & ActiveSlots) == 0)
+	{
+		return;
+	}
+
+	// We are going to rebuild the list of activities to maintain order as the
+	// activity list is singular.
+
+	check(List != nullptr);
+	FActivity* Activity = List;
+	List = ListTail = nullptr;
+	ActiveSlots = 0;
+
+	for (FActivity* Next; Activity != nullptr; Activity = Next)
+	{
+		Next = Activity->Next;
+
+		if (uint64 Slot = (1ull << Activity->Slot); (State.Cancels & Slot) == 0)
+		{
+			Activity->Next = nullptr;
+			AddActivity(Activity);
+			continue;
+		}
+
+		DoCancel(Activity);
+
+		Activity->Next = State.DoneList;
+		State.DoneList = Activity;
+	}
+}
+
+
+
+// {{{1 peer-group .............................................................
+
+////////////////////////////////////////////////////////////////////////////////
+class FPeerGroup
+{
+public:
+						FPeerGroup() = default;
+						~FPeerGroup();
+	bool				operator == (FPeerType* Rhs) const { return &Peer == Rhs; }
+	void				Unwait()			{ check(bWaiting); bWaiting = false; }
+	FPeerType::FWaiter	GetWaiter() const;
+	bool 				Tick(FTickState& State);
+	void				TickSend(FTickState& State, FHost& Host);
+	void				Fail(FTickState& State, const char* Reason);
+
+private:
+	void				RecvInternal(FTickState& State);
+	void				SendInternal(FTickState& State);
+	FActivity*			Send = nullptr;
+	FActivity*			Recv = nullptr;
+	FPeerType			Peer;
+	uint32				LastUseMs = 0;
+	uint8				IsKeepAlive = 0;
+	bool				bWaiting = false;
+
+	UE_NONCOPYABLE(FPeerGroup);
+};
+
+////////////////////////////////////////////////////////////////////////////////
+FPeerGroup::~FPeerGroup()
+{
+	check(Send == nullptr);
+	check(Recv == nullptr);
+}
+
+////////////////////////////////////////////////////////////////////////////////
+FPeerType::FWaiter FPeerGroup::GetWaiter() const
+{
+	if (!bWaiting)
+	{
+		return FPeerType::FWaiter();
+	}
+
+	using EWhat = FPeerType::FWaiter::EWhat;
+	EWhat What = (Recv != nullptr) ? EWhat::Recv : EWhat::Send;
+	return FPeerType::FWaiter(Peer, What);
+}
+
+////////////////////////////////////////////////////////////////////////////////
+void FPeerGroup::Fail(FTickState& State, const char* Reason)
+{
+	// Any send left at this point is unrecoverable
+	if (Send != nullptr)
+	{
+		Send->Next = Recv;
+		Recv = Send;
+	}
+
+	// Failure is quite terminal and we need to abort everything
+	for (FActivity* Activity = Recv; Activity != nullptr;)
+	{
+		if (Activity->State != FActivity::EState::Failed)
+		{
+			Activity_SetError(Activity, Reason);
+		}
+
+		DoFail(Activity);
+
+		FActivity* Next = Activity->Next;
+		Activity->Next = State.DoneList;
+		State.DoneList = Activity;
+		Activity = Next;
+	}
+
+	Peer = FPeerType();
+	Send = Recv = nullptr;
+	bWaiting = false;
+	IsKeepAlive = 0;
+}
+
+////////////////////////////////////////////////////////////////////////////////
+void FPeerGroup::RecvInternal(FTickState& State)
+{
+	check(Recv != nullptr);
+
+	// Another helper lambda
+	auto IsReceiving = [] (const FActivity* Act)
+	{
+		using EState = FActivity::EState;
+		return (Act->State >= EState::RecvMessage) & (Act->State < EState::RecvDone);
+	};
+
+	FActivity* Activity = Recv;
+	check(IsReceiving(Activity));
+
+	FOutcome Outcome = DoRecv(Activity, Peer, State.RecvAllowance);
+
+	// Any sort of error here is unrecoverable
+	if (Outcome.IsError())
+	{
+		Fail(State, Outcome.GetMessage().GetData());
+		return;
+	}
+
+	IsKeepAlive &= Activity->IsKeepAlive;
+	LastUseMs = State.NowMs;
+
+	// If we've only a small amount left to receive we can start more work
+	if (IsKeepAlive & (Recv->Next == nullptr))
+	{
+		uint32 Remaining = Activity_RemainingKiB(Activity);
+		if (Remaining < uint32(GRecvWorkThresholdKiB))
+		{
+			if (FActivity* Next = State.Work->PopActivity(); Next != nullptr)
+			{
+				Trace(Activity, ETrace::StartWork);
+	
+				check(Send == nullptr);
+				Send = Next;
+				SendInternal(State);
+
+				if (!Peer.IsValid())
+				{
+					return;
+				}
+			}
+		}
+	}
+
+	// If there was no data available this is far as receiving can go
+	if (bWaiting = Outcome.IsWaiting(); bWaiting)
+	{
+		return;
+	}
+
+	// If we're still in a receiving state we will just try again otherwise it
+	// is finished and we will let DoneList recipient finish it off.
+	if (IsReceiving(Activity))
+	{
+		return;
+	}
+
+	DoRecvDone(Activity);
+
+	Recv = Activity->Next;
+	Activity->Next = State.DoneList;
+	State.DoneList = Activity;
+
+	// If the server wants to close the socket we need to rewind the send
+	if (IsKeepAlive != 0)
+	{
+		return;
+	}
+
+	if (Send != nullptr && Activity_Rewind(Send) < 0)
+	{
+		Fail(State, "Unable to rewind on keep-alive close");
+		return;
+	}
+
+	Peer = FPeerType();
+}
+
+////////////////////////////////////////////////////////////////////////////////
+void FPeerGroup::SendInternal(FTickState& State)
+{
+	check(IsKeepAlive == 1);
+	check(Send != nullptr);
+
+	FActivity* Activity = Send;
+
+	FOutcome Outcome = DoSend(Activity, Peer);
+
+	if (Outcome.IsWaiting())
+	{
+		// For now we'll not add the socket as a waiter. It is unlikely that we
+		// send enough to need to wait currently.
+		return;
+	}
+
+	if (Outcome.IsError())
+	{
+		Fail(State, Outcome.GetMessage().GetData());
+		return;
+	}
+
+	Send = nullptr;
+
+	// Pass along this send to be received
+	if (Recv == nullptr)
+	{
+		Recv = Activity;
+		return;
+	}
+
+	check(Recv->Next == nullptr);
+	Recv->Next = Activity;
+}
+
+////////////////////////////////////////////////////////////////////////////////
+bool FPeerGroup::Tick(FTickState& State)
+{
+	if (Send != nullptr)
+	{
+		SendInternal(State);
+	}
+
+	if (Recv != nullptr && State.RecvAllowance)
+	{
+		RecvInternal(State);
+	}
+
+	return !!IsKeepAlive | !!(UPTRINT(Send) | UPTRINT(Recv));
+}
+
+////////////////////////////////////////////////////////////////////////////////
+void FPeerGroup::TickSend(FTickState& State, FHost& Host)
+{
+	// This path is only for those that are idle and have nothing to do
+	if (Send != nullptr || Recv != nullptr)
+	{
+		return;
+	}
+
+	// Failing will try and recover work which we don't want to happen yet
+	FActivity* Pending = State.Work->PopActivity();
+	check(Pending != nullptr);
+
+	// Close idle sockets
+	if (Peer.IsValid() && LastUseMs + GIdleMs < State.NowMs)
+	{
+		LastUseMs = State.NowMs;
+		Peer = FPeerType();
+	}
+
+	// We don't have a connected socket on first use, or if a keep-alive:close
+	// was received from the server. So we connect here.
+	bool bWillBlock = false;
+	if (!Peer.IsValid())
+	{
+		IsKeepAlive = 1;
+		FOutcome Outcome = Host.Connect(Peer);
+
+		// We failed to connect, let's bail.
+		if (Outcome.IsError())
+		{
+			Pending->Next = Recv;
+			Recv = Pending;
+			Fail(State, Outcome.GetMessage().GetData());
+			return;
+		}
+
+		bWillBlock = Outcome.IsWaiting();
+	}
+
+	Send = Pending;
+
+	if (!bWillBlock)
+	{
+		return SendInternal(State);
+	}
+
+	// Non-blocking connect
+	bWaiting = true;
+}
+
+
+
+// {{{1 host-group .............................................................
+
+////////////////////////////////////////////////////////////////////////////////
+class FHostGroup
+{
+public:
+								FHostGroup(FHost& InHost);
+	bool						IsBusy() const	{ return BusyCount != 0; }
+	const FHost&				GetHost() const	{ return Host; }
+	void						Tick(FTickState& State);
+	void						AddActivity(FActivity* Activity);
+
+private:
+	int32						Wait(const FTickState& State);
+	TArray<FPeerGroup>			PeerGroups;
+	FWorkQueue					Work;
+	FHost&						Host;
+	uint32						BusyCount = 0;
+	int32						WaitTimeAccum = 0;
+
+	UE_NONCOPYABLE(FHostGroup);
+};
+
+////////////////////////////////////////////////////////////////////////////////
+FHostGroup::FHostGroup(FHost& InHost)
+: Host(InHost)
+{
+	uint32 Num = InHost.GetMaxConnections();
+	PeerGroups.SetNum(Num);
+}
+
+////////////////////////////////////////////////////////////////////////////////
+int32 FHostGroup::Wait(const FTickState& State)
+{
+	// Collect groups that are waiting on something
+	TArray<FPeerType::FWaiter, TFixedAllocator<64>> Waiters;
+	for (FPeerGroup& Group : PeerGroups)
+	{
+		FPeerType::FWaiter Waiter = Group.GetWaiter();
+		if (Waiter.IsValid())
+		{
+			Waiters.Add(Waiter);
+		}
+	}
+
+	if (Waiters.IsEmpty())
+	{
+		return 0;
+	}
+
+	Trace(0, ETrace::Wait);
+	ON_SCOPE_EXIT { Trace(0, ETrace::Unwait); };
+
+	// If the poll timeout is negative then treat that as a fatal timeout
+	check(State.FailTimeoutMs);
+	int32 PollTimeoutMs = State.PollTimeoutMs;
+	if (PollTimeoutMs < 0)
+	{
+		PollTimeoutMs = State.FailTimeoutMs;
+	}
+
+	// Actually do the wait
+	int32 Result = FPeerType::Wait(Waiters, PollTimeoutMs);
+	if (Result <= 0)
+	{
+		// If the user opts to not block then we don't accumulate wait time and
+		// leave it to them to manage time a fail timoue
+		WaitTimeAccum += PollTimeoutMs;
+
+		if (State.PollTimeoutMs < 0 || WaitTimeAccum >= State.FailTimeoutMs)
+		{
+			return MIN_int32;
+		}
+
+		return Result;
+	}
+
+	WaitTimeAccum = 0;
+
+	// For each waiter that's ready, find the associated group "unwait" them.
+	int32 Count = 0;
+	for (int32 i = 0, n = Waiters.Num(); i < n; ++i)
+	{
+		if (Waiters[i].Ready == 0)
+		{
+			continue;
+		}
+
+		auto* Candidate = (FPeerType*)(Waiters[i].Candidate);
+		auto Pred = [Candidate] (auto& Lhs) { return Lhs == Candidate; };
+		FPeerGroup* Group = PeerGroups.FindByPredicate(Pred);
+		check(Group != nullptr);
+		Group->Unwait();
+
+		Waiters.RemoveAtSwap(i, EAllowShrinking::No);
+		--n, --i, ++Count;
+	}
+	check(Count == Result);
+
+	return Result;
+}
+
+////////////////////////////////////////////////////////////////////////////////
+void FHostGroup::Tick(FTickState& State)
+{
+	State.Work = &Work;
+
+	if (BusyCount = Work.HasWork(); BusyCount)
+	{
+		Work.TickCancels(State);
+
+		// Get available work out on idle sockets as soon as possible
+		for (FPeerGroup& Group : PeerGroups)
+		{
+			if (!Work.HasWork())
+			{
+				break;
+			}
+
+			Group.TickSend(State, Host);
+		}
+	}
+
+	// Wait on the groups that are
+	if (int32 Result = Wait(State); Result < 0)
+	{
+		const char* Reason = (Result == MIN_int32)
+			? "FailTimeout hit"
+			: "poll() returned an unexpected error";
+
+		for (FPeerGroup& Group : PeerGroups)
+		{
+			Group.Fail(State, Reason);
+		}
+
+		return;
+	}
+
+	// Tick everything, starting with groups that are maybe closest to finishing
+	for (FPeerGroup& Group : PeerGroups)
+	{
+		BusyCount += (Group.Tick(State) == true);
+	}
+}
+
+////////////////////////////////////////////////////////////////////////////////
+void FHostGroup::AddActivity(FActivity* Activity)
+{
+	Work.AddActivity(Activity);
+}
+
+
+
+// {{{1 event-loop .............................................................
+
+////////////////////////////////////////////////////////////////////////////////
+static const FEventLoop::FRequestParams GDefaultParams;
+
+////////////////////////////////////////////////////////////////////////////////
+class FEventLoop::FImpl
+{
+public:
+							~FImpl();
+	uint32					Tick(int32 PollTimeoutMs=0);
+	bool					IsIdle() const;
+	void					Throttle(uint32 KiBPerSec);
+	void					SetFailTimeout(int32 TimeoutMs);
+	void					Cancel(FTicket Ticket);
+	FRequest				Request(FAnsiStringView Method, FAnsiStringView Path, FActivity* Activity);
+	FTicket					Send(FActivity* Activity);
+
+private:
+	void					ReceiveWork();
+	FCriticalSection		Lock;
+	std::atomic<uint64>		FreeSlots		= ~0ull;
+	std::atomic<uint64>		Cancels			= 0;
+	uint64					PrevFreeSlots	= ~0ull;
+	FActivity*				Pending			= nullptr;
+	FThrottler				Throttler;
+	TArray<FHostGroup>		Groups;
+	int32					FailTimeoutMs	= GIdleMs;
+	uint32					BusyCount		= 0;
+};
+
+////////////////////////////////////////////////////////////////////////////////
+FEventLoop::FImpl::~FImpl()
+{
+	check(BusyCount == 0);
+}
+
+////////////////////////////////////////////////////////////////////////////////
+FRequest FEventLoop::FImpl::Request(
+	FAnsiStringView Method,
+	FAnsiStringView Path,
+	FActivity* Activity)
+{
+	Trace(Activity, ETrace::ActivityCreate, this);
+
+	if (Path.Len() == 0)
+	{
+		Path = "/";
+	}
+
+	Activity->NoContent = (Method == "HEAD");
+
+	FMessageBuilder Builder(Activity->Buffer);
+
+	Builder << Method << " " << Path << " HTTP/1.1" "\r\n"
+		"Host: " << Activity->Host->GetHostName() << "\r\n";
+
+	// HTTP/1.1 is persistent by default thus "Connection" header isn't required
+	if (!Activity->IsKeepAlive)
+	{
+		Builder << "Connection: close\r\n";
+	}
+
+	FRequest Ret;
+	Ret.Ptr = Activity;
+	return Ret;
+}
+
+////////////////////////////////////////////////////////////////////////////////
+FTicket FEventLoop::FImpl::Send(FActivity* Activity)
+{
+	Trace(Activity, ETrace::RequestBegin);
+
+	FMessageBuilder(Activity->Buffer) << "\r\n";
+	Activity_ChangeState(Activity, FActivity::EState::Send);
+
+	uint64 Slot;
+	{
+		FScopeLock _(&Lock);
+
+		for (;; FPlatformProcess::SleepNoStats(0.0f))
+		{
+			uint64 FreeSlotsLoad = FreeSlots.load(std::memory_order_relaxed);
+			if (!FreeSlotsLoad)
+			{
+				// we don't handle oversubscription at the moment. Could return
+				// activity to Reqeust and return a 0 ticket.
+				check(false);
+			}
+			Slot = -int64(FreeSlotsLoad) & FreeSlotsLoad;
+			if (FreeSlots.compare_exchange_weak(FreeSlotsLoad, FreeSlotsLoad - Slot, std::memory_order_relaxed))
+			{
+				break;
+			}
+		}
+		Activity->Slot = int8(63 - FMath::CountLeadingZeros64(Slot));
+
+		// This puts pending requests in reverse order of when they were made
+		// but this will be undone when ReceiveWork() traverses the list
+		Activity->Next = Pending;
+		Pending = Activity;
+	}
+
+	return Slot;
+}
+
+////////////////////////////////////////////////////////////////////////////////
+bool FEventLoop::FImpl::IsIdle() const
+{
+	return FreeSlots.load(std::memory_order_relaxed) == ~0ull;
+}
+
+////////////////////////////////////////////////////////////////////////////////
+void FEventLoop::FImpl::Throttle(uint32 KiBPerSec)
+{
+	Throttler.SetLimit(KiBPerSec);
+}
+
+////////////////////////////////////////////////////////////////////////////////
+void FEventLoop::FImpl::SetFailTimeout(int32 TimeoutMs)
+{
+	// While TimeoutMs must be >=0, it is signed so that MAX_uint32 can't be used
+	check(TimeoutMs > 0);
+	FailTimeoutMs = TimeoutMs;
+}
+
+////////////////////////////////////////////////////////////////////////////////
+void FEventLoop::FImpl::Cancel(FTicket Ticket)
+{
+	Cancels.fetch_or(Ticket, std::memory_order_relaxed);
+}
+
+////////////////////////////////////////////////////////////////////////////////
+void FEventLoop::FImpl::ReceiveWork()
+{
+	uint64 FreeSlotsLoad = FreeSlots.load(std::memory_order_relaxed);
+	if (FreeSlots == PrevFreeSlots)
+	{
+		return;
+	}
+	PrevFreeSlots = FreeSlotsLoad;
+
+	// Fetch the pending activities from out in the wild
+	FActivity* Activity = nullptr;
+	{
+		FScopeLock _(&Lock);
+		Swap(Activity, Pending);
+	}
+
+	// Pending is in the reverse of the order that requests were made
+	FActivity* Reverse = nullptr;
+	for (FActivity* Next; Activity != nullptr; Activity = Next)
+	{
+		Next = Activity->Next;
+		Activity->Next = Reverse;
+		Reverse = Activity;
+	}
+	Activity = Reverse;
+
+	// Group activities by their host.
+	for (FActivity* Next; Activity != nullptr; Activity = Next)
+	{
+		Next = Activity->Next;
+		Activity->Next = nullptr;
+
+		FHost& Host = *(Activity->Host);
+		auto Pred = [&Host] (const FHostGroup& Lhs) { return &Lhs.GetHost() == &Host; };
+		FHostGroup* Group = Groups.FindByPredicate(Pred);
+		if (Group == nullptr)
+		{
+			Group = &(Groups.Emplace_GetRef(Host));
+		}
+
+		Group->AddActivity(Activity);
+		++BusyCount;
+	}
+}
+
+////////////////////////////////////////////////////////////////////////////////
+uint32 FEventLoop::FImpl::Tick(int32 PollTimeoutMs)
+{
+	TRACE_CPUPROFILER_EVENT_SCOPE(IasHttp::Tick);
+
+	ReceiveWork();
+
+	// We limit recv sizes as a way to control bandwidth use.
+	int32 RecvAllowance = Throttler.GetAllowance();
+	if (RecvAllowance <= 0)
+	{
+		if (PollTimeoutMs == 0)
+		{
+			return BusyCount;
+		}
+
+		int32 ThrottleWaitMs = -RecvAllowance;
+		if (PollTimeoutMs > 0)
+		{
+			ThrottleWaitMs = FMath::Min(ThrottleWaitMs, PollTimeoutMs);
+		}
+		FPlatformProcess::SleepNoStats(float(ThrottleWaitMs) / 1000.0f);
+
+		RecvAllowance = Throttler.GetAllowance();
+		if (RecvAllowance <= 0)
+		{
+			return BusyCount;
+		}
+	}
+
+	uint64 CancelsLoad = Cancels.load(std::memory_order_relaxed);
+
+	uint32 NowMs;
+	{
+		// 4.2MM seconds will give us 50 days of uptime.
+		static uint64 Freq = 0;
+		static uint64 Base = 0;
+		if (Freq == 0)
+		{
+			Freq = uint64(1.0 / FPlatformTime::GetSecondsPerCycle());
+			Base = FPlatformTime::Cycles64();
+		}
+		uint64 NowBig = ((FPlatformTime::Cycles64() - Base) * 1000) / Freq;
+		NowMs = uint32(NowBig);
+		check(NowMs == NowBig);
+	}
+
+	// Tick groups and then remove ones that are idle
+	FTickState TickState = {
+		.DoneList = nullptr,
+		.Cancels = CancelsLoad,
+		.RecvAllowance = RecvAllowance,
+		.PollTimeoutMs = PollTimeoutMs,
+		.FailTimeoutMs = FailTimeoutMs,
+		.NowMs = NowMs,
+	};
+	for (FHostGroup& Group : Groups)
+	{
+		Group.Tick(TickState);
+	}
+
+	for (uint32 i = 0, n = Groups.Num(); i < n; ++i)
+	{
+		FHostGroup& Group = Groups[i];
+		if (Group.IsBusy())
+		{
+			continue;
+		}
+
+		Groups.RemoveAtSwap(i, EAllowShrinking::No);
+		--n, --i;
+	}
+
+	Throttler.ReturnUnused(RecvAllowance);
+
+	uint64 ReturnedSlots = 0;
+	for (FActivity* Activity = TickState.DoneList; Activity != nullptr;)
+	{
+		FActivity* Next = Activity->Next;
+		ReturnedSlots |= (1ull << Activity->Slot);
+		Activity_Free(Activity);
+		--BusyCount;
+		Activity = Next;
+	}
+
+	uint32 BusyBias = 0;
+	if (ReturnedSlots)
+	{
+		uint64 LatestFree = FreeSlots.fetch_add(ReturnedSlots, std::memory_order_relaxed);
+		BusyBias += (LatestFree != PrevFreeSlots);
+		PrevFreeSlots += ReturnedSlots;
+	}
+
+	if (CancelsLoad)
+	{
+		Cancels.fetch_and(~CancelsLoad, std::memory_order_relaxed);
+	}
+
+	return BusyCount + BusyBias;
+}
+
+
+
+////////////////////////////////////////////////////////////////////////////////
+FEventLoop::FEventLoop()						{ Impl = new FEventLoop::FImpl(); Trace(Impl, ETrace::LoopCreate); }
+FEventLoop::~FEventLoop()						{ Trace(Impl, ETrace::LoopDestroy); delete Impl; }
+uint32 FEventLoop::Tick(int32 PollTimeoutMs)	{ return Impl->Tick(PollTimeoutMs); }
+bool FEventLoop::IsIdle() const					{ return Impl->IsIdle(); }
+void FEventLoop::Cancel(FTicket Ticket)			{ return Impl->Cancel(Ticket); }
+void FEventLoop::Throttle(uint32 KiBPerSec)		{ return Impl->Throttle(KiBPerSec); }
+void FEventLoop::SetFailTimeout(int32 Ms)		{ return Impl->SetFailTimeout(Ms); }
+
+////////////////////////////////////////////////////////////////////////////////
+FRequest FEventLoop::Request(
+	FAnsiStringView Method,
+	FAnsiStringView Url,
+	const FRequestParams* Params)
+{
+	// Parse the URL into its components
+	FUrlOffsets UrlOffsets;
+	if (ParseUrl(Url, UrlOffsets) < 0)
+	{
+		return FRequest();
+	}
+
+	FAnsiStringView HostName = UrlOffsets.HostName.Get(Url);
+
+	uint32 Port = 0;
+	if (UrlOffsets.Port)
+	{
+		FAnsiStringView PortView = UrlOffsets.Port.Get(Url);
+		Port = uint32(CrudeToInt(PortView));
+	}
+
+	FAnsiStringView Path;
+	if (UrlOffsets.Path > 0)
+	{
+		Path = Url.Mid(UrlOffsets.Path);
+	}
+
+	// Create an activity and an emphemeral host
+	Params = (Params != nullptr) ? Params : &GDefaultParams;
+
+	uint32 BufferSize = Params->BufferSize;
+	BufferSize = (BufferSize >= 128) ? BufferSize : 128;
+	BufferSize += sizeof(FHost) + HostName.Len();
+	FActivity* Activity = Activity_Alloc(BufferSize);
+
+	FBuffer& Buffer = Activity->Buffer;
+
+	FHost* Host = Activity->Host = Buffer.Alloc<FHost>();
+	Activity->IsKeepAlive = 0;
+	Activity->bFollow30x = (Params->bAutoRedirect == true);
+
+	uint32 HostNameLength = HostName.Len();
+	char* HostNamePtr = Buffer.Alloc<char>(HostNameLength + 1);
+
+	Buffer.Fix();
+
+	memcpy(HostNamePtr, HostName.GetData(), HostNameLength);
+	HostNamePtr[HostNameLength] = '\0';
+	new (Host) FHost({
+		.HostName	= HostNamePtr,
+		.Port		= Port,
+	});
+
+	return Impl->Request(Method, Path, Activity);
+}
+
+////////////////////////////////////////////////////////////////////////////////
+FRequest FEventLoop::Request(
+	FAnsiStringView Method,
+	FAnsiStringView Path,
+	FConnectionPool& Pool,
+	const FRequestParams* Params)
+{
+	check(Pool.Ptr != nullptr);
+
+	Params = (Params != nullptr) ? Params : &GDefaultParams;
+
+	uint32 BufferSize = Params->BufferSize;
+	BufferSize = (BufferSize >= 128) ? BufferSize : 128;
+	FActivity* Activity = Activity_Alloc(BufferSize);
+
+	Activity->Host = Pool.Ptr;
+	Activity->IsKeepAlive = 1;
+	Activity->bFollow30x = (Params->bAutoRedirect == true);
+
+	return Impl->Request(Method, Path, Activity);
+}
+
+////////////////////////////////////////////////////////////////////////////////
+bool FEventLoop::Redirect(const FTicketStatus& Status, FTicketSink& OuterSink)
+{
+	const FResponse& Response = Status.GetResponse();
+
+	switch (Response.GetStatusCode())
+	{
+		case 301:				// RedirectMoved
+		case 302:				// RedirectFound
+		case 307:				// RedirectTemp
+		case 308: break;		// RedirectPerm
+		default: return false;
+	}
+
+	FAnsiStringView Location = Response.GetHeader("Location");
+	if (Location.IsEmpty())
+	{
+		return false;
+	}
+
+	check(Response.GetContentLength() == 0); // should we ever hit this, we'll fix it
+
+	const auto& Activity = (FActivity&)Response; // todo: yuk
+
+	// Original method should remain unchanged
+	const char* Data = Activity.Buffer.GetData();
+	FAnsiStringView Method;
+	for (uint32 i = 0; i < 5; ++i)
+	{
+		if (Data[i] <= ' ')
+		{
+			Method = FAnsiStringView(Data, i);
+			break;
+		}
+	}
+	check(!Method.IsEmpty());
+
+	FRequest ForwardRequest;
+	if (!Location.StartsWith("http://") && !Location.StartsWith("https://"))
+	{
+		if (Location[0] != '/')
+		{
+			return false;
+		}
+
+		FHost& Host = *(Activity.Host);
+
+		TAnsiStringBuilder<256> Url;
+		Url << "http://";
+		Url << Host.GetHostName();
+		Url << ":" << Host.GetPort();
+		Url << Location;
+
+		new (&ForwardRequest) FRequest(Request(Method, Url));
+	}
+	else
+	{
+		new (&ForwardRequest) FRequest(Request(Method, Location));
+	}
+
+	// Transfer original request headers
+	check(Activity.State == FActivity::EState::RecvMessage);
+	uint32 Length = 0;
+	for (const char* End = Data + Activity.StateParam; Data + Length < End; ++Length)
+	{
+		if (Data[Length] != '\n')
+		{
+			continue;
+		}
+
+		Data = Data + Length + 1;
+		Length = uint32(ptrdiff_t(End - Data));
+		break;
+	}
+
+	FAnsiStringView OriginalHeaders(Data, Length);
+	EnumerateHeaders(OriginalHeaders, [&ForwardRequest] (FAnsiStringView Name, FAnsiStringView Value)
+	{
+		if (Name != "Host" && Name != "Connection")
+		{
+			ForwardRequest.Header(Name, Value);
+		}
+		return true;
+	});
+
+	// Send the request
+	Send(MoveTemp(ForwardRequest), MoveTemp(OuterSink), Status.GetParam());
+
+	return true;
+}
+
+////////////////////////////////////////////////////////////////////////////////
+FTicket FEventLoop::Send(FRequest&& Request, FTicketSink Sink, UPTRINT SinkParam)
+{
+	FActivity* Activity = nullptr;
+	Swap(Activity, Request.Ptr);
+	Activity->SinkParam = SinkParam;
+	Activity->Sink = Sink;
+
+	// Intercept sink calls to catch 30x status codes and follow them
+	if (Activity->bFollow30x)
+	{
+		auto RedirectSink = [
+				this,
+				OuterSink=MoveTemp(Activity->Sink)
+			] (const FTicketStatus& Status) mutable
+		{
+			if (Status.GetId() == FTicketStatus::EId::Response)
+			{
+				if (Redirect(Status, OuterSink))
+				{
+					return;
+				}
+			}
+
+			if (OuterSink)
+			{
+				return OuterSink(Status);
+			}
+		};
+		Activity->Sink = RedirectSink;
+	}
+
+	return Impl->Send(Activity);
+}
+
+
+
+
+// }}}
+
+} // namespace UE::IoStore::HTTP
+
diff --git a/Engine/Source/Runtime/Experimental/IoStoreOnDemand/Private/CoreHttp/Misc.inl b/Engine/Source/Runtime/Experimental/IoStoreOnDemand/Private/CoreHttp/Misc.inl
new file mode 100644
--- /dev/null
+++ b/Engine/Source/Runtime/Experimental/IoStoreOnDemand/Private/CoreHttp/Misc.inl
@@ -0,0 +1,797 @@
+// Copyright Epic Games, Inc. All Rights Reserved.
+
+#pragma once
+
+namespace UE::IoStore::HTTP
+{
+
+// {{{1 trace ..................................................................
+
+////////////////////////////////////////////////////////////////////////////////
+#define MAKE_TRACE_ENUM(x) \
+	x(LoopCreate) \
+	x(LoopTick) \
+	x(LoopDestroy) \
+	x(ActivityCreate) \
+	x(ActivityDestroy) \
+	x(SocketCreate) \
+	x(SocketDestroy) \
+	x(RequestBegin) \
+	x(StateChange) \
+	x(Wait) \
+	x(Unwait) \
+	x(Connect) \
+	x(Send) \
+	x(Recv) \
+	x(StartWork) \
+	x($)
+
+enum class ETrace
+{
+#define TRACE_ENUM_VALUES(x) x,
+	MAKE_TRACE_ENUM(TRACE_ENUM_VALUES)
+#undef TRACE_ENUM_VALUES
+};
+
+static const FAnsiStringView TraceEnumNames[] = {
+#define TRACE_ENUM_VALUES(x) #x ,
+	MAKE_TRACE_ENUM(TRACE_ENUM_VALUES)
+#undef TRACE_ENUM_VALUES
+};
+
+#undef MAKE_TRACE_ENUM
+
+UE_TRACE_CHANNEL(IasHttpChannel);
+
+UE_TRACE_EVENT_BEGIN(IasHttp, Enum, NoSync)
+	UE_TRACE_EVENT_FIELD(UE::Trace::AnsiString, Name)
+UE_TRACE_EVENT_END()
+
+UE_TRACE_EVENT_BEGIN(IasHttp, Event, NoSync)
+	UE_TRACE_EVENT_FIELD(uint64, Cycle)
+	UE_TRACE_EVENT_FIELD(uint32, Id)
+	UE_TRACE_EVENT_FIELD(uint32, Param)
+	UE_TRACE_EVENT_FIELD(uint8, Action)
+UE_TRACE_EVENT_END()
+
+static void Activity_TraceStateNames();
+
+////////////////////////////////////////////////////////////////////////////////
+static void TraceEnum(const FAnsiStringView* Names)
+{
+	for (;; ++Names)
+	{
+		UE_TRACE_LOG(IasHttp, Enum, IasHttpChannel)
+			<< Enum.Name(Names[0].GetData(), Names[0].Len());
+
+		if (*Names == "$")
+		{
+			break;
+		}
+	}
+}
+
+////////////////////////////////////////////////////////////////////////////////
+static void TraceInner(const UPTRINT Id, ETrace Action, UPTRINT Param)
+{
+#if UE_TRACE_ENABLED
+	static bool Once = [] {
+		TraceEnum(TraceEnumNames);
+		Activity_TraceStateNames();
+		return true;
+	}();
+
+	UE_TRACE_LOG(IasHttp, Event, IasHttpChannel)
+		<< Event.Cycle(FPlatformTime::Cycles64())
+		<< Event.Id(uint32(Id))
+		<< Event.Param(uint32(Param))
+		<< Event.Action(uint8(Action));
+#endif // UE_TRACE_ENABLED
+}
+
+template <typename T, typename U=UPTRINT>
+static void Trace(T Id, ETrace Action, U Param=0)
+{
+	TraceInner(UPTRINT(Id), Action, UPTRINT(Param));
+}
+
+
+
+// {{{1 misc ...................................................................
+
+////////////////////////////////////////////////////////////////////////////////
+#define IAS_CVAR(Type, Name, Default, Desc, ...) \
+	Type G##Name = Default; \
+	static FAutoConsoleVariableRef CVar_Ias##Name( \
+		TEXT("ias.Http" #Name), \
+		G##Name, \
+		TEXT(Desc) \
+		__VA_ARGS__ \
+	)
+
+////////////////////////////////////////////////////////////////////////////////
+static IAS_CVAR(int32, RecvWorkThresholdKiB,80,		"Threshold of data remaining at which next request is sent (in KiB)");
+static IAS_CVAR(int32, IdleMs,				50'000,	"Time in seconds to close idle connections or fail waits");
+
+////////////////////////////////////////////////////////////////////////////////
+class FOutcome
+{
+public:
+	static FOutcome Ok(uint32 Result=0);
+	static FOutcome Waiting();
+	static FOutcome Error(const char* Message, int32 Code=-1);
+	static FOutcome None()				{ return Error(""); }
+	bool			IsError() const		{ return Message < 0x8000'0000'0000; }
+	bool			IsWaiting() const	{ return Tag == WaitTag; }
+	bool			IsOk() const		{ return Tag == OkTag; }
+	FAnsiStringView GetMessage() const	{ check(IsError()); return (const char*)(Message); }
+	int32			GetErrorCode() const{ check(IsError()); return int32(Code); }
+	uint32			GetResult() const	{ check(IsOk()); return Result; }
+
+private:
+					FOutcome() = default;
+
+	static uint32 const OkTag	= 0x0000'8000;
+	static uint32 const WaitTag = 0x0001'8000;
+
+	union {
+		struct {
+			UPTRINT	Message : 48;
+			PTRINT	Code	: 16;
+		};
+		struct {
+			uint32	Result;
+			uint32	Tag;
+		};
+	};
+};
+static_assert(sizeof(FOutcome) == sizeof(void*));
+
+////////////////////////////////////////////////////////////////////////////////
+FOutcome FOutcome::Ok(uint32 Result)
+{
+	FOutcome Outcome;
+	Outcome.Tag = OkTag;
+	Outcome.Result = Result;
+	return Outcome;
+}
+
+////////////////////////////////////////////////////////////////////////////////
+FOutcome FOutcome::Waiting()
+{
+	FOutcome Outcome;
+	Outcome.Tag = WaitTag;
+	return Outcome;
+}
+
+////////////////////////////////////////////////////////////////////////////////
+FOutcome FOutcome::Error(const char* Message, int32 Code)
+{
+	check(Message != nullptr);
+	check(Code <= 0xffff && Code >= -0xffff);
+
+	FOutcome Outcome;
+	Outcome.Message = UPTRINT(Message);
+	Outcome.Code = int16(Code);
+	return Outcome;
+}
+
+
+
+////////////////////////////////////////////////////////////////////////////////
+template <typename LambdaType>
+static void EnumerateHeaders(FAnsiStringView Headers, LambdaType&& Lambda)
+{
+	// NB. here we are assuming that we will be dealing with servers that will
+	// not be returning headers with "obsolete line folding".
+
+	auto IsOws = [] (int32 c) { return (c == ' ') | (c == '\t'); };
+
+	const char* Cursor = Headers.GetData();
+	const char* End = Cursor + Headers.Len();
+	do
+	{
+		int32 ColonIndex = 0;
+		for (; Cursor + ColonIndex < End; ++ColonIndex)
+		{
+			if (Cursor[ColonIndex] == ':')
+			{
+				break;
+			}
+		}
+
+		Cursor += ColonIndex;
+
+		const char* Right = Cursor + 1;
+		while (Right < End)
+		{
+			if (Right[0] != '\r' || Right[1] != '\n')
+			{
+				Right += 1 + (Right[1] != '\r');
+				continue;
+			}
+
+			FAnsiStringView Name(Cursor - ColonIndex, ColonIndex);
+
+			const char* Left = Cursor + 1;
+			for (; IsOws(Left[0]); ++Left);
+
+			Cursor = Right;
+			for (; Cursor > Left + 1 && IsOws(Cursor[-1]); --Cursor);
+
+			FAnsiStringView Value (Left, int32(ptrdiff_t(Cursor - Left)));
+
+			if (!Lambda(Name, Value))
+			{
+				Right = End;
+			}
+
+			break;
+		}
+
+		Cursor = Right + 2;
+	} while (Cursor < End);
+}
+
+////////////////////////////////////////////////////////////////////////////////
+static int32 FindMessageTerminal(const char* Data, uint32 Length)
+{
+	for (uint32 i = 4; i <= Length; ++i)
+	{
+		uint32 Candidate;
+		::memcpy(&Candidate, Data + i - 4, sizeof(Candidate));
+		if (Candidate == 0x0a0d0a0d)
+		{
+			return i;
+		}
+
+		i += (Data[i - 1] > 0x0d) ? 3 : 0;
+	}
+
+	return -1;
+}
+
+////////////////////////////////////////////////////////////////////////////////
+static int64 CrudeToInt(FAnsiStringView View)
+{
+	// FCStringAnsi::* is not used to mitigate any locale hiccups. By
+	// initialising 'Value' with MSB set we can detect cases where View did not
+	// start with digits. This works as we won't be using this on huge numbers.
+	int64 Value = 0x8000'0000'0000'0000ll;
+	for (int32 c : View)
+	{
+		uint32 Digit = c - '0';
+		if (Digit > 9u)
+		{
+			break;
+		}
+		Value *= 10;
+		Value += Digit;
+	}
+	return Value;
+};
+
+////////////////////////////////////////////////////////////////////////////////
+struct FMessageOffsets
+{
+	uint8	StatusCode;
+	uint8	Message;
+	uint16	Headers;
+};
+
+static int32 ParseMessage(FAnsiStringView Message, FMessageOffsets& Out)
+{
+	const FAnsiStringView Protocol("HTTP/1.1 ");
+
+	// Check there's enough data
+	if (Message.Len() < Protocol.Len() + 1) // "+1" accounts for at least one digit
+	{
+		return -1;
+	}
+
+	const char* Cursor = Message.GetData();
+
+	// Check for the expected protocol
+	if (FAnsiStringView(Cursor, 9) != Protocol)
+	{
+		return -1;
+	}
+	int32 i = Protocol.Len();
+
+	// Trim left and tightly reject anything adventurous
+	for (int n = 32; Cursor[i] == ' ' && i < n; ++i);
+	Out.StatusCode = uint8(i);
+
+	// At least one status line digit. (Note to self; expect exactly three)
+	for (int n = 32; uint32(Cursor[i] - 0x30) <= 9 && i < n; ++i);
+	if (uint32(i - Out.StatusCode - 1) > 32)
+	{
+		return -1;
+	}
+
+	// Trim left
+	for (int n = 32; Cursor[i] == ' ' && i < n; ++i);
+	Out.Message = uint8(i);
+
+	// Extra conservative length allowance
+	if (i > 32)
+	{
+		return -1;
+	}
+
+	// Find \r\n
+	for (; Cursor[i] != '\r'; ++i)
+	{
+		if (i >= 2048)
+		{
+			return -1;
+		}
+	}
+	if (Cursor[i + 1] != '\n')
+	{
+		return -1;
+	}
+	Out.Headers = uint16(i + 2);
+
+	return 1;
+}
+
+////////////////////////////////////////////////////////////////////////////////
+struct FUrlOffsets
+{
+	struct Slice
+	{
+						Slice() = default;
+						Slice(int32 l, int32 r) : Left(uint8(l)), Right(uint8(r)) {}
+		FAnsiStringView	Get(FAnsiStringView Url) const { return Url.Mid(Left, Right - Left); }
+						operator bool () const { return Left > 0; }
+		int32			Len() const { return Right - Left; }
+		uint8			Left;
+		uint8			Right;
+	};
+	Slice				UserInfo;
+	Slice				HostName;
+	Slice				Port;
+	uint8				Path;
+	uint8				SchemeLength;
+};
+
+static int32 ParseUrl(FAnsiStringView Url, FUrlOffsets& Out)
+{
+	if (Url.Len() < 5)
+	{
+		return -1;
+	}
+
+	Out = {};
+
+	const char* Start = Url.GetData();
+	const char* Cursor = Start;
+
+	// Scheme
+	int32 i = 0;
+	for (; i < 5; ++i)
+	{
+		if (uint32(Cursor[i] - 'a') > uint32('z' - 'a'))
+		{
+			break;
+		}
+	}
+
+	Out.SchemeLength = uint8(i);
+	FAnsiStringView Scheme = Url.Left(i);
+	if (Scheme != "http" && Scheme != "https")
+	{
+		return -1;
+	}
+
+	// Separator and authority
+	if (Cursor[i] != ':' || Cursor[i + 1] != '/' || Cursor[i + 2] != '/')
+	{
+		return -1;
+	}
+	i += 3;
+
+	struct { int32 c; int32 i; } Seps[2];
+	int32 SepCount = 0;
+	for (; i < Url.Len(); ++i)
+	{
+		int32 c = Cursor[i];
+		if (c < '-')							break;
+		if (c != ':' && c != '@' && c != '/')	continue;
+		if (c == '/' || SepCount >= 2)			break;
+
+		if (c == '@' && SepCount)
+		{
+			SepCount -= (Seps[SepCount - 1].c == ':');
+		}
+		Seps[SepCount++] = { c, i };
+	}
+
+	if (i > 0xff || i <= Scheme.Len() + 3)
+	{
+		return -1;
+	}
+
+	if (i < Url.Len())
+	{
+		Out.Path = uint8(i);
+	}
+
+	Out.HostName = { uint8(Scheme.Len() + 3), uint8(i) };
+
+	switch (SepCount)
+	{
+	case 0:
+		break;
+
+	case 1:
+		if (Seps[0].c == ':')
+		{
+			Out.Port = { Seps[0].i + 1, i };
+			Out.HostName.Right = uint8(Seps[0].i);
+		}
+		else
+		{
+			Out.UserInfo = { Out.HostName.Left, Seps[0].i };
+			Out.HostName.Left += uint8(Seps[0].i + 1);
+			Out.HostName.Right += uint8(Seps[0].i + 1);
+		}
+		break;
+
+	case 2:
+		if ((Seps[0].c != '@') | (Seps[1].c != ':'))
+		{
+			return -1;
+		}
+		Out.UserInfo = { Out.HostName.Left, Seps[0].i };
+		Out.Port.Left = uint8(Seps[1].i + 1);
+		Out.Port.Right = Out.HostName.Right;
+		Out.HostName.Left = Out.UserInfo.Right + 1;
+		Out.HostName.Right = Out.Port.Left - 1;
+		break;
+
+	default:
+		return -1;
+	}
+
+	bool Bad = false;
+	Bad |= (Out.HostName.Len() == 0);
+	Bad |= bool(Out.UserInfo) & (Out.UserInfo.Len() == 0);
+
+	if (Out.Port.Left)
+	{
+		Bad |= (Out.Port.Len() == 0);
+		for (int32 j = 0, n = Out.Port.Len(); j < n; ++j)
+		{
+			Bad |= (uint32(Start[Out.Port.Left + j] - '0') > 9);
+		}
+	}
+
+	return Bad ? -1 : 1;
+}
+
+
+
+// {{{1 buffer .................................................................
+
+////////////////////////////////////////////////////////////////////////////////
+class alignas(16) FBuffer
+{
+public:
+	struct FMutableSection
+	{
+		char*		Data;
+		uint32		Size;
+	};
+
+								FBuffer() = default;
+								FBuffer(char* InData, uint32 InMax);
+								~FBuffer();
+	FBuffer&					operator = (FBuffer&& Rhs);
+	void						Fix();
+	void						Resize(uint32 Size);
+	const char*					GetData() const;
+	uint32						GetSize() const;
+	uint32						GetCapacity() const;
+	template <typename T> T*	Alloc(uint32 Count=1);
+	FMutableSection				GetMutableFree(uint32 MinSize, uint32 PageSize=256);
+	void						AdvanceUsed(uint32 Delta);
+
+private:
+	char*						GetDataPtr();
+	void						Extend(uint32 AtLeast, uint32 PageSize);
+	union
+	{
+		struct
+		{
+			UPTRINT				Data : 63;
+			UPTRINT				Inline : 1;
+		};
+		UPTRINT					DataInline = 0;
+	};
+	uint32						Max = 0;
+	uint32						Used = 0;
+
+private:
+								FBuffer(FBuffer&&) = delete;
+								FBuffer(const FBuffer&) = delete;
+	FBuffer&					operator = (const FBuffer&) = delete;
+};
+
+////////////////////////////////////////////////////////////////////////////////
+FBuffer::FBuffer(char* InData, uint32 InMax)
+: Data(UPTRINT(InData))
+, Inline(1)
+, Max(InMax)
+{
+}
+
+////////////////////////////////////////////////////////////////////////////////
+FBuffer::~FBuffer()
+{
+	if (Data && !Inline)
+	{
+		FMemory::Free(GetDataPtr());
+	}
+}
+
+////////////////////////////////////////////////////////////////////////////////
+FBuffer& FBuffer::operator = (FBuffer&& Rhs)
+{
+	Swap(DataInline, Rhs.DataInline);
+	Swap(Max, Rhs.Max);
+	Swap(Used, Rhs.Used);
+	return *this;
+}
+
+////////////////////////////////////////////////////////////////////////////////
+void FBuffer::Fix()
+{
+	check(Inline);
+	Data += Used;
+	Max -= Used;
+	Used = 0;
+}
+
+////////////////////////////////////////////////////////////////////////////////
+void FBuffer::Resize(uint32 Size)
+{
+	check(Size <= Max);
+	Used = Size;
+}
+
+////////////////////////////////////////////////////////////////////////////////
+char* FBuffer::GetDataPtr()
+{
+	return (char*)Data;
+}
+
+////////////////////////////////////////////////////////////////////////////////
+const char* FBuffer::GetData() const
+{
+	return (char*)Data;
+}
+
+////////////////////////////////////////////////////////////////////////////////
+uint32 FBuffer::GetSize() const
+{
+	return Used;
+}
+
+////////////////////////////////////////////////////////////////////////////////
+uint32 FBuffer::GetCapacity() const
+{
+	return Max;
+}
+
+////////////////////////////////////////////////////////////////////////////////
+template <typename T>
+T* FBuffer::Alloc(uint32 Count)
+{
+	uint32 AlignBias = uint32(Data) & (alignof(T) - 1);
+	if (AlignBias)
+	{
+		AlignBias = alignof(T) - AlignBias;
+	}
+
+	uint32 PotentialUsed = Used + AlignBias + (sizeof(T) * Count);
+	if (PotentialUsed > Max)
+	{
+		Extend(PotentialUsed, 256);
+	}
+
+	void* Ret = GetDataPtr() + Used + AlignBias;
+	Used = PotentialUsed;
+	return (T*)Ret;
+}
+
+////////////////////////////////////////////////////////////////////////////////
+FBuffer::FMutableSection FBuffer::GetMutableFree(uint32 MinSize, uint32 PageSize)
+{
+	MinSize = (MinSize == 0 && Used == Max) ? PageSize : MinSize;
+
+	uint32 PotentialUsed = Used + MinSize;
+	if (PotentialUsed > Max)
+	{
+		Extend(PotentialUsed, PageSize);
+	}
+
+	return FMutableSection{ GetDataPtr() + Used, Max - Used };
+}
+
+////////////////////////////////////////////////////////////////////////////////
+void FBuffer::AdvanceUsed(uint32 Delta)
+{
+	Used += Delta;
+	check(Used <= Max);
+}
+
+////////////////////////////////////////////////////////////////////////////////
+void FBuffer::Extend(uint32 AtLeast, uint32 PageSize)
+{
+	checkSlow((PageSize - 1 & PageSize) == 0);
+
+	--PageSize;
+	Max = (AtLeast + PageSize) & ~PageSize;
+
+	if (!Inline)
+	{
+		Data = UPTRINT(FMemory::Realloc(GetDataPtr(), Max, alignof(FBuffer)));
+		return;
+	}
+
+	const char* PrevData = GetDataPtr();
+	Data = UPTRINT(FMemory::Malloc(Max, alignof(FBuffer)));
+	::memcpy(GetDataPtr(), PrevData, Used);
+	Inline = 0;
+}
+
+
+
+////////////////////////////////////////////////////////////////////////////////
+class FMessageBuilder
+{
+public:
+						FMessageBuilder(FBuffer& Ruffeb);
+	FMessageBuilder&	operator << (FAnsiStringView Lhs);
+
+private:
+	FBuffer&		Buffer;
+};
+
+////////////////////////////////////////////////////////////////////////////////
+FMessageBuilder::FMessageBuilder(FBuffer& Ruffeb)
+: Buffer(Ruffeb)
+{
+}
+
+////////////////////////////////////////////////////////////////////////////////
+FMessageBuilder& FMessageBuilder::operator << (FAnsiStringView Lhs)
+{
+	uint32 Length = uint32(Lhs.Len());
+	FBuffer::FMutableSection Section = Buffer.GetMutableFree(Length);
+	::memcpy(Section.Data, Lhs.GetData(), Length);
+	Buffer.AdvanceUsed(Length);
+	return *this;
+}
+
+
+
+// {{{1 throttler ..............................................................
+
+////////////////////////////////////////////////////////////////////////////////
+static void ThrottleTest(FAnsiStringView);
+
+////////////////////////////////////////////////////////////////////////////////
+class FThrottler
+{
+public:
+			FThrottler();
+	void	SetLimit(uint32 KiBPerSec);
+	int32	GetAllowance();
+	void	ReturnUnused(uint32 Unused);
+
+private:
+	friend	void ThrottleTest(FAnsiStringView);
+	int32	GetAllowance(uint64 CycleDelta);
+	int32	GetWaitEstimateMs() const;
+	uint64	CycleFreq;
+	uint64	CycleLast;
+	uint64	CycleIdle;
+	uint32	Limit = 0;
+	int32	Available = 0;
+
+	enum {
+		LIMITLESS	= MAX_int32,
+		THRESHOLD	= 2 << 10,
+	};
+};
+
+////////////////////////////////////////////////////////////////////////////////
+FThrottler::FThrottler()
+{
+	CycleFreq = uint64(1.0 / FPlatformTime::GetSecondsPerCycle());
+	CycleLast = FPlatformTime::Cycles64() - CycleFreq;
+	CycleIdle = CycleFreq * 8;
+}
+
+////////////////////////////////////////////////////////////////////////////////
+void FThrottler::SetLimit(uint32 KiBPerSec)
+{
+	// 512MiB/s might as well be limitless.
+	KiBPerSec = (KiBPerSec < (512 << 10)) ? KiBPerSec : 0;
+	Limit = KiBPerSec << 10;
+}
+
+////////////////////////////////////////////////////////////////////////////////
+int32 FThrottler::GetAllowance()
+{
+	int64 Cycle = FPlatformTime::Cycles64();
+	int64 CycleDelta = Cycle - CycleLast;
+	if (CycleDelta < 0)
+	{
+		return Limit;
+	}
+	CycleLast = Cycle;
+	return GetAllowance(CycleDelta);
+}
+
+////////////////////////////////////////////////////////////////////////////////
+int32 FThrottler::GetAllowance(uint64 CycleDelta)
+{
+	if (Limit == 0)
+	{
+		return LIMITLESS;
+	}
+
+	// If we're idle for too long then reset the throttling
+	if (CycleDelta >= CycleIdle)
+	{
+		Available = 0;
+		return Limit;
+	}
+
+	uint64 Delta = (uint64(Limit) * CycleDelta) / CycleFreq;
+
+	// A gate against lost precision
+	if (Delta == 0)
+	{
+		CycleLast -= CycleDelta;
+		return 0 - GetWaitEstimateMs();
+	}
+
+	// Don't let available run away
+	uint64 Next = FMath::Min<uint64>(uint64(Available) + Delta, Limit * 4);
+
+	Available = uint32(Next);
+
+	// Doesn't make sense to trickle out tiny allowances
+	if (Available < THRESHOLD)
+	{
+		return 0 - GetWaitEstimateMs();
+	}
+
+	int32 Released = Available;
+	Available = 0;
+	return Released;
+}
+
+////////////////////////////////////////////////////////////////////////////////
+void FThrottler::ReturnUnused(uint32 Unused)
+{
+	Available += Unused;
+}
+
+////////////////////////////////////////////////////////////////////////////////
+int32 FThrottler::GetWaitEstimateMs() const
+{
+	// Calculate an approximate time to wait for more allowance
+	int64 Estimate = THRESHOLD - Available;
+	Estimate = (Estimate * 1000ll) / int64(Limit);
+	return FMath::Max(int32(Estimate), 0);
+}
+
+// }}}
+
+} // namespace UE::IoStore::HTTP
diff --git a/Engine/Source/Runtime/Experimental/IoStoreOnDemand/Private/CoreHttp/Socket.inl b/Engine/Source/Runtime/Experimental/IoStoreOnDemand/Private/CoreHttp/Socket.inl
new file mode 100644
--- /dev/null
+++ b/Engine/Source/Runtime/Experimental/IoStoreOnDemand/Private/CoreHttp/Socket.inl
@@ -0,0 +1,377 @@
+// Copyright Epic Games, Inc. All Rights Reserved.
+
+#pragma once
+
+// {{{1 platforms ..............................................................
+
+#if PLATFORM_MICROSOFT
+#	if !defined(NO_UE_INCLUDES)
+#	include "Windows/AllowWindowsPlatformTypes.h"
+#		include <winsock2.h>
+#		include <ws2tcpip.h>
+#	include "Windows/HideWindowsPlatformTypes.h"
+#	else
+#		include <winsock2.h>
+#		include <ws2tcpip.h>
+#		pragma comment(lib, "Ws2_32.lib")
+#	endif // NO_UE_INCLUDES
+	using SocketType	= SOCKET;
+	using MsgFlagType	= int;
+
+	enum { SHUT_RDWR = SD_BOTH };
+
+#	define IAS_HTTP_USE_POLL
+	template <typename... ArgTypes> auto poll(ArgTypes... Args)
+	{
+		return WSAPoll(Forward<ArgTypes>(Args)...);
+	}
+
+	// Winsock defines "PF_MAX" indicating the protocol family count. This
+	// however competes with UE definitions related to pixel formats.
+	#if defined(PF_MAX)
+	#	undef PF_MAX
+	#endif
+
+#elif PLATFORM_APPLE | PLATFORM_UNIX | PLATFORM_ANDROID
+#	include <arpa/inet.h>
+#	include <fcntl.h>
+#	include <netdb.h>
+#	include <poll.h>
+#	include <sys/select.h>
+#	include <sys/socket.h>
+#	include <unistd.h>
+	using SocketType	= int;
+	using MsgFlagType	= int;
+	static int32 closesocket(int32 Socket) { return close(Socket); }
+#	define IAS_HTTP_USE_POLL
+#else
+#	include "CoreHttp/Http.inl"
+#endif
+
+static_assert(sizeof(sockaddr_in::sin_addr) == sizeof(uint32));
+
+#if PLATFORM_MICROSOFT
+	static int32 LastSocketResult() { return WSAGetLastError(); }
+#	define IsSocketResult(err)		(LastSocketResult() == WSA##err)
+#else
+	static int32 LastSocketResult() { return errno; }
+#	define IsSocketResult(err)		(LastSocketResult() == err)
+	static_assert(EWOULDBLOCK == EAGAIN);
+#endif
+
+static const SocketType InvalidSocket = ~SocketType(0);
+
+// }}}
+
+namespace UE::IoStore::HTTP
+{
+
+// {{{1 socket .................................................................
+
+////////////////////////////////////////////////////////////////////////////////
+class FSocket
+{
+public:
+	struct FWaiter
+	{
+		enum class EWhat { Send = 0b01, Recv = 0b10, Both = Send|Recv };
+				FWaiter() { std::memset(this, 0, sizeof(*this)); }
+				FWaiter(const FSocket& Socket, EWhat InWaitOn);
+		bool	IsValid() const { UPTRINT x{0}; return std::memcmp(this, &x, sizeof(*this)) != 0; }
+		bool	operator == (FSocket& Rhs) const { return UPTRINT(&Rhs) == Candidate; }
+		UPTRINT	Candidate : 60;
+		UPTRINT	WaitOn : 2;
+		UPTRINT	Ready : 2;
+	};
+
+				FSocket() = default;
+				~FSocket()					{ Destroy(); }
+				FSocket(FSocket&& Rhs)		{ Move(MoveTemp(Rhs)); }
+	FSocket&	operator = (FSocket&& Rhs)	{ Move(MoveTemp(Rhs)); return *this; }
+	bool		IsValid() const				{ return Socket != InvalidSocket; }
+	bool		Create();
+	void		Destroy();
+	FOutcome	Connect(uint32 Ip, uint32 Port);
+	void		Disconnect();
+	FOutcome	Send(const char* Data, uint32 Size);
+	FOutcome	Recv(char* Dest, uint32 Size);
+	bool		SetBlocking(bool bBlocking);
+	bool		SetSendBufSize(int32 Size);
+	bool		SetRecvBufSize(int32 Size);
+	static int	Wait(TArrayView<FWaiter> Waiters, int32 TimeoutMs);
+
+private:
+	void		Move(FSocket&& Rhs);
+	SocketType	Socket = InvalidSocket;
+
+				FSocket(const FSocket&) = delete;
+	FSocket&	operator = (const FSocket&) = delete;
+};
+
+////////////////////////////////////////////////////////////////////////////////
+FSocket::FWaiter::FWaiter(const FSocket& Socket, EWhat InWaitOn)
+: Candidate(UPTRINT(&Socket))
+, WaitOn(UPTRINT(InWaitOn))
+, Ready(0)
+{
+	static_assert(sizeof(*this) == sizeof(void*));
+}
+
+////////////////////////////////////////////////////////////////////////////////
+void FSocket::Move(FSocket&& Rhs)
+{
+	check(!IsValid() || !Rhs.IsValid()); // currently we only want to pass one around
+	Swap(Socket, Rhs.Socket);
+}
+
+////////////////////////////////////////////////////////////////////////////////
+bool FSocket::Create()
+{
+	check(!IsValid());
+	Socket = socket(AF_INET, SOCK_STREAM, IPPROTO_TCP);
+
+	if (!IsValid())
+	{
+		return false;
+	}
+
+	Trace(Socket, ETrace::SocketCreate);
+	return true;
+}
+
+////////////////////////////////////////////////////////////////////////////////
+void FSocket::Destroy()
+{
+	if (Socket == InvalidSocket)
+	{
+		return;
+	}
+
+	Trace(Socket, ETrace::SocketDestroy);
+
+	closesocket(Socket);
+	Socket = InvalidSocket;
+}
+
+////////////////////////////////////////////////////////////////////////////////
+FOutcome FSocket::Connect(uint32 IpAddress, uint32 Port)
+{
+	check(IsValid());
+
+	Trace(Socket, ETrace::Connect, IpAddress);
+
+	IpAddress = htonl(IpAddress);
+
+	sockaddr_in AddrInet = { sizeof(sockaddr_in) };
+	AddrInet.sin_family = AF_INET;
+	AddrInet.sin_port = htons(uint16(Port));
+	memcpy(&(AddrInet.sin_addr), &IpAddress, sizeof(IpAddress));
+
+	int32 Result = connect(Socket, &(sockaddr&)AddrInet, sizeof(AddrInet));
+
+	if (IsSocketResult(EWOULDBLOCK) | IsSocketResult(EINPROGRESS))
+	{
+		return FOutcome::Waiting();
+	}
+
+	if (Result < 0)
+	{
+		return FOutcome::Error("Socket connect failed", Result);
+	}
+
+	return FOutcome::Ok();
+}
+
+////////////////////////////////////////////////////////////////////////////////
+void FSocket::Disconnect()
+{
+	check(IsValid());
+	shutdown(Socket, SHUT_RDWR);
+}
+
+////////////////////////////////////////////////////////////////////////////////
+FOutcome FSocket::Send(const char* Data, uint32 Size)
+{
+	Trace(Socket, ETrace::Send, -1);
+	int32 Result = send(Socket, Data, Size, MsgFlagType(0));
+	Trace(Socket, ETrace::Send, FMath::Max(Result, 0));
+
+	if (Result > 0)						return FOutcome::Ok(Result);
+	if (Result == 0)					return FOutcome::Error("Send ATH0");
+	if (IsSocketResult(EWOULDBLOCK))	return FOutcome::Waiting();
+
+	if (IsSocketResult(ENOTCONN))
+	{
+		int32 Error = 0;
+		socklen_t ErrorSize = sizeof(Error);
+		Result = getsockopt(Socket, SOL_SOCKET, SO_ERROR, (char*)&Error, &ErrorSize);
+		if (Result < 0 || Error != 0)
+		{
+			return FOutcome::Error("Error while connecting");
+		}
+
+		return FOutcome::Waiting();
+	}
+
+	Result = LastSocketResult();
+	return FOutcome::Error("Send", Result);
+}
+
+////////////////////////////////////////////////////////////////////////////////
+FOutcome FSocket::Recv(char* Dest, uint32 Size)
+{
+	Trace(Socket, ETrace::Recv, -1);
+	int32 Result = recv(Socket, Dest, Size, MsgFlagType(0));
+	Trace(Socket, ETrace::Recv, FMath::Max(0, Result));
+
+	if (Result > 0)						return FOutcome::Ok(Result);
+	if (Result == 0)					return FOutcome::Error("Recv ATH0");
+	if (IsSocketResult(EWOULDBLOCK))	return FOutcome::Waiting();
+
+	Result = LastSocketResult();
+	return FOutcome::Error("Recv", Result);
+}
+
+////////////////////////////////////////////////////////////////////////////////
+bool FSocket::SetBlocking(bool bBlocking)
+{
+#if defined(IAS_HTTP_HAS_NONBLOCK_IMPL)
+	return SetNonBlockingSocket(Socket, !bBlocking);
+#elif PLATFORM_MICROSOFT
+	unsigned long NonBlockingMode = (bBlocking != true);
+	return (ioctlsocket(Socket, FIONBIO, &NonBlockingMode) != SOCKET_ERROR);
+#else
+	int32 Flags = fcntl(Socket, F_GETFL, 0);
+	if (Flags == -1)
+	{
+		return false;
+	}
+
+	int32 NewFlags = bBlocking
+		? (Flags & ~int32(O_NONBLOCK))
+		: (Flags |	int32(O_NONBLOCK));
+	return (Flags == NewFlags) || (fcntl(Socket, F_SETFL, Flags) >= 0);
+#endif
+}
+
+////////////////////////////////////////////////////////////////////////////////
+bool FSocket::SetSendBufSize(int32 Size)
+{
+	return 0 == setsockopt(Socket, SOL_SOCKET, SO_SNDBUF, &(char&)Size, sizeof(Size));
+}
+
+////////////////////////////////////////////////////////////////////////////////
+bool FSocket::SetRecvBufSize(int32 Size)
+{
+	return 0 == setsockopt(Socket, SOL_SOCKET, SO_RCVBUF, &(char&)Size, sizeof(Size));
+}
+
+////////////////////////////////////////////////////////////////////////////////
+int32 FSocket::Wait(TArrayView<FWaiter> Waiters, int32 TimeoutMs)
+{
+#if !defined(IAS_HTTP_USE_POLL)
+	struct FSelect
+	{
+		SocketType	fd;
+		int32		events;
+		int32		revents;
+	};
+	static const int32 POLLIN  = 1 << 0;
+	static const int32 POLLOUT = 1 << 1;
+	static const int32 POLLERR = 1 << 2;
+	static const int32 POLLHUP = POLLERR;
+	static const int32 POLLNVAL= POLLERR;
+#else
+	using FSelect = pollfd;
+#endif
+
+	// The following looks odd because POLLFD varies subtly from one platform
+	// to the next. To cleanly set members to zero and to not get narrowing
+	// warnings from the compiler, we list-init and don't assume POD types.
+	using PollEventType = decltype(FSelect::events);
+	PollEventType Events[] = { POLLERR, POLLOUT, POLLIN, POLLOUT|POLLOUT };
+	TArray<FSelect, TFixedAllocator<64>> Selects;
+	for (FWaiter& Waiter : Waiters)
+	{
+		Selects.Emplace_GetRef() = {
+			((FSocket*)Waiter.Candidate)->Socket,
+			Events[Waiter.WaitOn],
+			/* 0 */
+		};
+	}
+
+	// Poll the sockets
+#if defined(IAS_HTTP_USE_POLL)
+	int32 Result = poll(Selects.GetData(), Selects.Num(), TimeoutMs);
+	if (Result <= 0)
+	{
+		return Result;
+	}
+#else
+	timeval TimeVal = {};
+	timeval* TimeValPtr = (TimeoutMs >= 0 ) ? &TimeVal : nullptr;
+	if (TimeoutMs > 0)
+	{
+		TimeVal = { TimeoutMs >> 10, TimeoutMs & ((1 << 10) - 1) };
+	}
+
+	fd_set FdSetRead;	FD_ZERO(&FdSetRead);
+	fd_set FdSetWrite;	FD_ZERO(&FdSetWrite);
+	fd_set FdSetExcept; FD_ZERO(&FdSetExcept);
+
+	SocketType MaxFd = 0;
+	for (FSelect& Select : Selects)
+	{
+		fd_set* RwSet = (Select.events & POLLIN) ? &FdSetRead : &FdSetWrite;
+		FD_SET(Select.fd, RwSet);
+		FD_SET(Select.fd, &FdSetExcept);
+		MaxFd = FMath::Max(Select.fd, MaxFd);
+	}
+
+	int32 Result = select(int32(MaxFd + 1), &FdSetRead, &FdSetWrite, &FdSetExcept, TimeValPtr);
+	if (Result <= 0)
+	{
+		return Result;
+	}
+
+	for (FSelect& Select : Selects)
+	{
+		if (FD_ISSET(Select.fd, &FdSetExcept))
+		{
+			Select.revents = POLLERR;
+			continue;
+		}
+
+		fd_set* RwSet = (Select.events & POLLIN) ? &FdSetRead : &FdSetWrite;
+		if (FD_ISSET(Select.fd, RwSet))
+		{
+			Select.revents = Select.events;
+		}
+	}
+#endif // IAS_HTTP_USE_POLL
+
+	// Transfer poll results to the input sockets. We don't transfer across error
+	// states. Subsequent sockets ops can take care of that instead.
+	static const auto TestBits = POLLIN|POLLOUT|POLLERR|POLLHUP|POLLNVAL;
+	for (uint32 i = 0, n = Waiters.Num(); i < n; ++i)
+	{
+		auto RetEvents = Selects[i].revents;
+		if (!(RetEvents & TestBits))
+		{
+			continue;
+		}
+
+		uint32 Value = 0;
+		if (!!(RetEvents & POLLOUT)) Value |= uint32(FSocket::FWaiter::EWhat::Send);
+		if (!!(RetEvents & POLLIN))	 Value |= uint32(FSocket::FWaiter::EWhat::Recv);
+		Waiters[i].Ready = Value ? Value : uint32(FSocket::FWaiter::EWhat::Both);
+	}
+
+	return Result;
+}
+
+
+
+// }}}
+
+} // namespace UE::IoStore::HTTP
diff --git a/Engine/Source/Runtime/Experimental/IoStoreOnDemand/Private/CoreHttp/Socks.inl b/Engine/Source/Runtime/Experimental/IoStoreOnDemand/Private/CoreHttp/Socks.inl
new file mode 100644
--- /dev/null
+++ b/Engine/Source/Runtime/Experimental/IoStoreOnDemand/Private/CoreHttp/Socks.inl
@@ -0,0 +1,186 @@
+// Copyright Epic Games, Inc. All Rights Reserved.
+
+#pragma once
+
+namespace UE::IoStore::HTTP
+{
+
+// {{{1 socks ..................................................................
+
+#if !UE_BUILD_SHIPPING
+
+///////////////////////////////////////////////////////////////////////////////
+static IAS_CVAR(int32,		SocksVersion,	5,		"SOCKS proxy protocol version to use");
+static IAS_CVAR(FString,	SocksIp,		"",		"Routes all IAS HTTP traffic through the given SOCKS proxy");
+static IAS_CVAR(int32,		SocksPort,		1080,	"Port of the SOCKS proxy to use");
+
+////////////////////////////////////////////////////////////////////////////////
+static uint32 GetSocksIpAddress()
+{
+	const TCHAR* Value = *GSocksIp;
+	uint32 IpAddress = 0;
+	uint32 Accumulator = 0;
+	while (true)
+	{
+		uint32 c = *Value++;
+
+		if (c - '0' <= '9' - '0')
+		{
+			Accumulator *= 10;
+			Accumulator += (c - '0');
+			continue;
+		}
+
+		if (c == '.' || c == '\0')
+		{
+			IpAddress <<= 8;
+			IpAddress |= Accumulator;
+			Accumulator = 0;
+			if (c == '\0')
+			{
+				break;
+			}
+			continue;
+		}
+
+		return 0;
+	}
+	return IpAddress;
+}
+
+////////////////////////////////////////////////////////////////////////////////
+static FOutcome ConnectSocks4(FSocket& Socket, uint32 IpAddress, uint32 Port)
+{
+	struct FSocks4Request
+	{
+		uint8	Version = 4;
+		uint8	Command = 1;
+		uint16	Port;
+		uint32	IpAddress;
+	};
+
+	struct FSocks4Reply
+	{
+		uint8	Version;
+		uint8	Code;
+		uint16	Port;
+		uint32	IpAddress;
+	};
+
+	uint32 SocksIpAddress = GetSocksIpAddress();
+	if (!SocksIpAddress)
+	{
+		return FOutcome::Error("Invalid socks IP address");
+	}
+
+	FOutcome Outcome = FOutcome::None();
+
+	if (Outcome = Socket.Connect(SocksIpAddress, GSocksPort); Outcome.IsError())
+	{
+		return Outcome;
+	}
+
+	FSocks4Request Request = {
+		.Port		= htons(uint16(Port)),
+		.IpAddress	= htonl(IpAddress),
+	};
+	if (Outcome = Socket.Send((const char*)&Request, sizeof(Request)); Outcome.IsError())
+	{
+		return Outcome;
+	}
+
+	FSocks4Reply Reply;
+	if (Outcome = Socket.Recv((char*)&Reply, sizeof(Reply)); Outcome.IsError())
+	{
+		return Outcome;
+	}
+
+	return FOutcome::Ok(1);
+}
+
+////////////////////////////////////////////////////////////////////////////////
+static FOutcome ConnectSocks5(FSocket& Socket, uint32 IpAddress, uint32 Port)
+{
+#ifdef _MSC_VER
+	// MSVC's static analysis doesn't see that 'Result' from recv() is checked
+	// to be the exact size of the destination buffer.
+#pragma warning(push)
+#pragma warning(disable : 6385)
+#endif
+
+	uint32 SocksIpAddress = GetSocksIpAddress();
+	if (!SocksIpAddress)
+	{
+		return FOutcome::Error("Invalid socks5 IP address");
+	}
+
+	FOutcome Outcome = FOutcome::None();
+
+	if (Outcome = Socket.Connect(SocksIpAddress, GSocksPort); Outcome.IsError())
+	{
+		return Outcome;
+	}
+
+	// Greeting
+	const char Greeting[] = { 5, 1, 0 };
+	Outcome = Socket.Send(Greeting, sizeof(Greeting));
+	if (!Outcome.IsOk())						return Outcome;
+	if (Outcome.GetResult() != sizeof(Greeting))return FOutcome::Error("Could not send socks5 greeting");
+
+	// Server auth-choice
+	char Choice[1 + 1];
+	Outcome = Socket.Recv(Choice, sizeof(Choice));
+	if (!Outcome.IsOk())						return Outcome;
+	if (Outcome.GetResult() != sizeof(Choice))	return FOutcome::Error("Recv too short from socks5 server");
+	if (Choice[0] != 0x05 || Choice[1] != 0x00) return FOutcome::Error("Got unexpected socks5 version from server");
+
+	// Connection request
+	IpAddress = htonl(IpAddress);
+	uint16 NsPort = htons(uint16(Port));
+	char Request[] = { 5, 1, 0, 1, 0x11,0x11,0x11,0x11, 0x22,0x22 };
+	std::memcpy(Request + 4, &IpAddress, sizeof(IpAddress));
+	std::memcpy(Request + 8, &NsPort, sizeof(NsPort));
+	Outcome = Socket.Send(Request, sizeof(Request));
+	if (!Outcome.IsOk())						return Outcome;
+	if (Outcome.GetResult() != sizeof(Request)) return FOutcome::Error("Sent too little to socks5 server");
+
+	// Connect reply
+	char Reply[3 + (1 + 4) + 2];
+	Outcome = Socket.Recv(Reply, sizeof(Reply));
+	if (!Outcome.IsOk())						return Outcome;
+	if (Outcome.GetResult() != sizeof(Reply))	return FOutcome::Error("Socks5 reply too short");
+	if (Reply[0] != 0x05 || Reply[1] != 0x00)	return FOutcome::Error("Reply has unexpected socks5 version");
+
+	return FOutcome::Ok(1);
+
+#ifdef _MSC_VER
+#pragma warning(pop)
+#endif
+}
+
+#endif // UE_BUILD_SHIPPING
+
+////////////////////////////////////////////////////////////////////////////////
+static FOutcome MaybeConnectSocks(FSocket& Socket, uint32 IpAddress, uint32 Port)
+{
+#if UE_BUILD_SHIPPING
+	return FOutcome::Ok();
+#else
+	if (GSocksIp.IsEmpty())
+	{
+		return FOutcome::Ok();
+	}
+
+	switch (GSocksVersion)
+	{
+	case 4: return ConnectSocks4(Socket, IpAddress, Port);
+	case 5: return ConnectSocks5(Socket, IpAddress, Port);
+	}
+
+	return FOutcome::Error("Unsupported socks version");
+#endif // UE_BUILD_SHIPPING
+}
+
+// }}}
+
+} // namespace UE::IoStore::HTTP
diff --git a/Engine/Source/Runtime/Experimental/IoStoreOnDemand/Private/CoreHttp/Test.inl b/Engine/Source/Runtime/Experimental/IoStoreOnDemand/Private/CoreHttp/Test.inl
new file mode 100644
--- /dev/null
+++ b/Engine/Source/Runtime/Experimental/IoStoreOnDemand/Private/CoreHttp/Test.inl
@@ -0,0 +1,741 @@
+// Copyright Epic Games, Inc. All Rights Reserved.
+
+#pragma once
+
+namespace UE::IoStore::HTTP
+{
+
+// {{{1 test ...................................................................
+
+#if !(UE_BUILD_SHIPPING|UE_BUILD_TEST)
+
+////////////////////////////////////////////////////////////////////////////////
+static void MiscTest()
+{
+#define CRLF "\r\n"
+	struct {
+		FAnsiStringView Input;
+		int32 Output;
+	} FmtTestCases[] = {
+		{ "", -1 },
+		{ "abcd", -1 },
+		{ "abcd\r", -1 },
+		{ CRLF "\r\r", -1 },
+		{ CRLF CRLF, 4 },
+		{ "abc" CRLF CRLF, 7 },
+	};
+	for (const auto [Input, Output] : FmtTestCases)
+	{
+		check(FindMessageTerminal(Input.GetData(), Input.Len()) == Output);
+	}
+
+	FMessageOffsets MsgOut;
+	check(ParseMessage("", MsgOut) == -1);
+	check(ParseMessage("MR", MsgOut) == -1);
+	check(ParseMessage("HTTP/1.1", MsgOut) == -1);
+	check(ParseMessage("HTTP/1.1 ", MsgOut) == -1);
+	check(ParseMessage("HTTP/1.1 1" CRLF, MsgOut) > 0);
+	check(ParseMessage("HTTP/1.1    1" CRLF, MsgOut) > 0);
+	check(ParseMessage("HTTP/1.1 100 " CRLF, MsgOut) > 0);
+	check(ParseMessage("HTTP/1.1 100  Message of some sort    " CRLF, MsgOut) > 0);
+	check(ParseMessage("HTTP/1.1 100 _Message with a \r in it" CRLF, MsgOut) == -1);
+
+	bool AllIsWell = true;
+	auto NotExpectedToBeCalled = [&AllIsWell] (auto, auto)
+	{
+		AllIsWell = false;
+		return false;
+	};
+
+	EnumerateHeaders("",		NotExpectedToBeCalled); check(AllIsWell);
+	EnumerateHeaders(CRLF,		NotExpectedToBeCalled); check(AllIsWell);
+	EnumerateHeaders("foo",		NotExpectedToBeCalled); check(AllIsWell);
+	EnumerateHeaders(" foo",	NotExpectedToBeCalled); check(AllIsWell);
+	EnumerateHeaders(" foo ",	NotExpectedToBeCalled); check(AllIsWell);
+	EnumerateHeaders("foo:bar",	NotExpectedToBeCalled); check(AllIsWell);
+
+	auto IsBar = [&] (auto, auto Value) { return AllIsWell = (Value == "bar"); };
+	EnumerateHeaders("foo: bar" CRLF,		IsBar); check(AllIsWell);
+	EnumerateHeaders("foo: bar \t" CRLF,	IsBar); check(AllIsWell);
+	EnumerateHeaders("foo:\tbar " CRLF,		IsBar); check(AllIsWell);
+	EnumerateHeaders("foo:bar " CRLF,		IsBar); check(AllIsWell);
+	EnumerateHeaders("foo:bar" CRLF "!",	IsBar); check(AllIsWell);
+	EnumerateHeaders("foo:bar" CRLF " ",	IsBar); check(AllIsWell);
+	EnumerateHeaders("foo:bar" CRLF "n:ej",	IsBar); check(AllIsWell);
+
+	check(CrudeToInt("") < 0);
+	check(CrudeToInt("X") < 0);
+	check(CrudeToInt("/") < 0);
+	check(CrudeToInt(":") < 0);
+	check(CrudeToInt("-1") < -1);
+	check(CrudeToInt("0") == 0);
+	check(CrudeToInt("9") == 9);
+	check(CrudeToInt("493") == 493);
+
+	FUrlOffsets UrlOut;
+
+	check(ParseUrl("", UrlOut) == -1);
+	check(ParseUrl("abc://asd/", UrlOut) == -1);
+	check(ParseUrl("http://", UrlOut) == -1);
+	check(ParseUrl("http://:/", UrlOut) == -1);
+	check(ParseUrl("http://@:/", UrlOut) == -1);
+	check(ParseUrl("http://foo:ba:r/", UrlOut) == -1);
+	check(ParseUrl("http://foo@ba:r/", UrlOut) == -1);
+	check(ParseUrl("http://foo@ba:r", UrlOut) == -1);
+	check(ParseUrl("http://foo@ba:/", UrlOut) == -1);
+	check(ParseUrl("http://foo@ba@9/", UrlOut) == -1);
+	check(ParseUrl("http://@ba:9/", UrlOut) == -1);
+	check(ParseUrl(
+		"http://zzzzzzzzzzzzzzzzzzzzzzzzzzzzzzzzzzzzzzzzzzzzzzzzzzzzzzzzzzzzzzzzzzzzzzzzzzzzzzzzzzzzzzzzzzzzzzzzzzzzzzzzzzzzzzzzzzzzzzzzzz"
+		"zzzzzzzzzzzzzzzzzzzzzzzzzzzzzzzzzzzzzzzzzzzzzzzzzzzzzzzzzzzzzzzzzzzzzzzzzzzzzzzzzzzzzzzzzzzzzzzzzzzzzzzzzzzzzzzzzzzzzzzzzzzzzzzzz"
+		"zzzzzzzzzzzzzzzzzzzzzzzzzzzzzzzzzzzzzzzzzzzzzzzzzzzzzzzzzzzzzzzzzzzzzzzzzzzzzzzzzzzzzzzzzzzzzzzzzzzzzzzzzzzzzzzzzzzzzzzzzzzzzzzzz"
+		"zzzzzzzzzzzzzzzzzzzzzzzzzzzzzzzzzzzzzzzzzzzzzzzzzzzzzzzzzzzzzzzzzzzzzzzzzzzzzzzzzzzzzzzzzzzzzzzzzzzzzzzzzzzzzzzzzzzzzzzzzzzzzzzzz"
+		"zzzzzzzzzzzzzzzzzzzzzzzzzzzzzzzzzzzzzzzzzzzzzzzzzzzzzzzzzzzzzzzzzzzzzzzzzzzzzzzzzzzzzzzzzzzzzzzzzzzzzzzzzzzzzzzzzzzzzzzzzzzzz.com",
+		UrlOut) == -1);
+
+	check(ParseUrl("http://ab-c.com/", UrlOut) > 0);
+	check(ParseUrl("http:[REDACTED_EMAIL]/", UrlOut) > 0);
+	check(ParseUrl("https://abc.com", UrlOut) > 0);
+	check(ParseUrl("https://abc.com:999", UrlOut) > 0);
+	check(ParseUrl("https://abc.com:999/", UrlOut) > 0);
+	check(ParseUrl("https://foo:[REDACTED_EMAIL]:999", UrlOut) > 0);
+	check(ParseUrl("https://foo:[REDACTED_EMAIL]:999/", UrlOut) > 0);
+	check(ParseUrl("https:[REDACTED_EMAIL]:999", UrlOut) > 0);
+	check(ParseUrl("https:[REDACTED_EMAIL]:999/", UrlOut) > 0);
+
+	for (int32 i : { 0x10, 0x20, 0x40, 0x7f, 0xff })
+	{
+		char Url[] = "http://stockholm.patchercache.epicgames.net:123";
+		char Buffer[512];
+		std::memset(Buffer, i, sizeof(Buffer));
+		std::memcpy(Buffer, Url, sizeof(Url) - 1);
+		check(ParseUrl(FAnsiStringView(Buffer, sizeof(Url) - 1), UrlOut) > 0);
+		check(UrlOut.Port.Get(Url) == "123");
+	}
+
+	FAnsiStringView Url = "http://abc:[REDACTED_EMAIL]:999/";
+	check(ParseUrl(Url, UrlOut) > 0);
+	check(UrlOut.SchemeLength == 4);
+	check(UrlOut.UserInfo.Get(Url) == "abc:123");
+	check(UrlOut.HostName.Get(Url) == "bc.com");
+	check(UrlOut.Port.Get(Url) == "999");
+	check(UrlOut.Path == 25);
+#undef CRLF
+
+	static const char* OutcomeMsg = "\x4d\x52";
+	check(FOutcome::Error(OutcomeMsg, -5).IsOk() == false);
+	check(FOutcome::Error(OutcomeMsg, -5).IsWaiting() == false);
+	check(FOutcome::Error(OutcomeMsg, -5).IsError());
+	check(FOutcome::Error(OutcomeMsg, -5).GetErrorCode() == -5);
+	check(FOutcome::Error(OutcomeMsg,  5).GetErrorCode() ==  5);
+	check(FOutcome::Error(OutcomeMsg, -5).GetMessage() == OutcomeMsg);
+
+	check(FOutcome::Ok(  0).IsOk());
+	check(FOutcome::Ok(-13).IsWaiting() == false);
+	check(FOutcome::Ok( 13).IsError() == false);
+
+	check(FOutcome::Waiting().IsOk() == false);
+	check(FOutcome::Waiting().IsWaiting());
+	check(FOutcome::Waiting().IsError() == false);
+}
+
+////////////////////////////////////////////////////////////////////////////////
+static void ThrottleTest(FAnsiStringView TestUrl)
+{
+	enum { TheMax = 0x7fff'fffful };
+
+	check(FThrottler().GetAllowance() >= TheMax);
+
+	FThrottler Throttler;
+	uint64 OneSecond = Throttler.CycleFreq;
+
+	for (uint32 TargetElapsed : { 1, 5, 7 })
+	{
+		uint64 CycleTest = FPlatformTime::Cycles64();
+		for (uint32 i = 0; i < TargetElapsed; ++i)
+		{
+			FPlatformProcess::Sleep(1.0f);
+		}
+		CycleTest = FPlatformTime::Cycles64() - CycleTest;
+		check((CycleTest + (OneSecond / 8)) / OneSecond == TargetElapsed);
+	}
+
+
+	Throttler.SetLimit(0);
+	check(Throttler.GetAllowance(0) >= TheMax);
+	check(Throttler.GetAllowance()  >= TheMax);
+	check(Throttler.GetAllowance()  >= TheMax);
+
+	for (uint32 i : { 10, 63, 100 })
+	{
+		Throttler.SetLimit(i);
+		check(Throttler.GetAllowance( OneSecond          ) == (i << 10));
+		check(Throttler.GetAllowance((OneSecond + 1) >> 1) == (i << 9));
+		check(Throttler.GetAllowance((OneSecond + 3) >> 2) == (i << 8));
+	}
+
+	uint32 Limit = 17 << 10;
+	Throttler = FThrottler();
+	Throttler.SetLimit(Limit >> 10);
+	Throttler.GetAllowance(OneSecond);
+	Throttler.ReturnUnused(Limit >> 1);
+	check(Throttler.GetAllowance(OneSecond) == (Limit + (Limit >> 1)));
+
+	// runaway
+	check(Throttler.GetAllowance(OneSecond * 3) == Limit * 3);
+	check(Throttler.GetAllowance(OneSecond * 4) == Limit * 4);
+	check(Throttler.GetAllowance(OneSecond * 5) == Limit * 4);
+
+	// idle
+	check(Throttler.GetAllowance(OneSecond * 64) == Limit);
+
+	// threshold
+	Limit = FThrottler::THRESHOLD;
+	Throttler = FThrottler();
+	Throttler.SetLimit(Limit >> 10);
+	for (int64 Counter = OneSecond;;)
+	{
+		int64 Delta = (OneSecond + 15) >> 4;
+		Counter -= Delta;
+		int32 Allowance = Throttler.GetAllowance(Delta);
+		if (Allowance > 0)
+		{
+			check(Allowance == Limit);
+			check(Counter < 10);
+			break;
+		}
+	};
+
+	// overflow(ish)
+	Throttler = FThrottler();
+	Throttler.SetLimit((512 << 10) - 1);
+	Throttler.GetAllowance((OneSecond * 790) / 100);
+
+	// timing test
+	FIoBuffer RecvData;
+	for (uint32 SizeKiB : { 10, 25, 60 })
+	{
+		const uint32 ThrottleKiB = 5;
+
+		TAnsiStringBuilder<128> Url;
+		Url << TestUrl;
+		Url << (SizeKiB << 10);
+
+		FEventLoop Loop;
+		Loop.Throttle(ThrottleKiB);
+
+		FRequest Request = Loop.Request("GET", Url).Accept("*/*");
+		Loop.Send(MoveTemp(Request), [&] (const FTicketStatus& Status) {
+			check(Status.GetId() != FTicketStatus::EId::Error);
+			if (Status.GetId() == FTicketStatus::EId::Response)
+			{
+				Status.GetResponse().SetDestination(&RecvData);
+			}
+		});
+
+		int32 Timeout = -1;
+		if (SizeKiB < 25) Timeout = 123;
+		if (SizeKiB > 25) Timeout = 4567;
+
+		uint64 Time = FPlatformTime::Cycles64();
+		while (Loop.Tick(Timeout));
+		Time = FPlatformTime::Cycles64() - Time;
+		Time /= OneSecond;
+
+		// It's dangerous stuff testing elapsed time you know. The +1 is because
+		// throttling assumes one second has already passed when initialised.
+#if PLATFORM_WINDOWS
+		check(Time + 1 == (SizeKiB / ThrottleKiB));
+#endif
+
+		RecvData = FIoBuffer();
+	}
+}
+
+////////////////////////////////////////////////////////////////////////////////
+static void RedirectTest(const ANSICHAR* TestHost)
+{
+	FEventLoop Loop;
+
+	auto WaitForLoopIdle = [&] {
+		for (; Loop.Tick(-1); FPlatformProcess::SleepNoStats(0.02f));
+	};
+
+	FEventLoop::FRequestParams RequestParams = {
+		.bAutoRedirect = true,
+	};
+
+	enum ReTyp { ReAbs, ReRel };
+	enum { RecvDataSize = 48 };
+
+	TAnsiStringBuilder<64> Builder;
+	auto BuildUrl = [&] (ReTyp Typ, uint32 Code) -> const auto&
+	{
+		Builder.Reset();
+		Builder << "https://";
+		Builder << TestHost;
+		Builder << ":9493";
+		Builder << "/redirect";
+		Builder << ((Typ < ReRel) ? "/abs/" : "/rel/");
+		Builder << Code;
+		Builder << "/data/" << uint32(RecvDataSize);
+		return Builder;
+	};
+
+	UPTRINT SinkParam = 0xaa'493'493'493'493'bbull;
+
+	uint32 RecvCount;
+	auto OkSink = [Dest=FIoBuffer(), SinkParam, &RecvCount] (const FTicketStatus& Status) mutable {
+		check(Status.GetParam() == SinkParam);
+		check(Status.GetId() != FTicketStatus::EId::Error);
+		if (Status.GetId() == FTicketStatus::EId::Response)
+		{
+			FResponse& Response = Status.GetResponse();
+			check(Response.GetStatusCode() == 200);
+			Response.SetDestination(&Dest);
+			return;
+		}
+		check(Status.GetId() == FTicketStatus::EId::Content);
+		RecvCount += uint32(Dest.GetSize());
+	};
+
+	uint32 TestCodes[] = { 301, 302, 307, 308 };
+
+	for (auto ReTest : { ReAbs, ReRel })
+	{
+		RecvCount = 0;
+		for (uint32 Code : TestCodes)
+		{
+			FRequest Request = Loop.Get(BuildUrl(ReTest, Code), &RequestParams);
+			if (Code > TestCodes[1])
+			{
+				Request.Header("TestCodeHeader", "Header-Of-Test-Codes");
+			}
+			Loop.Send(MoveTemp(Request), OkSink, SinkParam);
+		}
+		WaitForLoopIdle();
+		check(RecvCount == RecvDataSize * UE_ARRAY_COUNT(TestCodes));
+	}
+
+	RequestParams = FEventLoop::FRequestParams();
+	RequestParams.bAutoRedirect = true;
+
+	FConnectionPool::FParams Params;
+	Params.SetHostFromUrl(BuildUrl(ReAbs, 0));
+	Params.ConnectionCount = 4;
+	FConnectionPool Pool(Params);
+
+	RecvCount = 0;
+	Loop.Send(Loop.Get("/redirect/abs/307/data/55", Pool, &RequestParams), OkSink, SinkParam);
+	WaitForLoopIdle();
+	check(RecvCount == 55);
+}
+
+////////////////////////////////////////////////////////////////////////////////
+static void HttpTest(const ANSICHAR* TestHost)
+{
+	TAnsiStringBuilder<64> Ret;
+	auto BuildUrl = [&] (const ANSICHAR* Suffix=nullptr, uint32 Port=9493) -> const auto&
+	{
+		Ret.Reset();
+		Ret << "http://";
+		Ret << TestHost;
+		Ret << ":" << Port;
+		return (Suffix != nullptr) ? (Ret << Suffix) : Ret;
+	};
+
+	struct
+	{
+		FIoBuffer Dest;
+		uint64 Hash = 0;
+	} Content[64];
+
+	auto HashSink = [&] (const FTicketStatus& Status) -> FIoBuffer*
+	{
+		check(Status.GetId() != FTicketStatus::EId::Error);
+
+		uint32 Index = Status.GetIndex();
+
+		if (Status.GetId() == FTicketStatus::EId::Response)
+		{
+			FResponse& Response = Status.GetResponse();
+			check(Response.GetStatus() == EStatusCodeClass::Successful);
+			check(Response.GetStatusCode() == 200);
+			check(Response.GetContentLength() == Status.GetContentLength());
+
+			FAnsiStringView HashView = Response.GetHeader("X-TestServer-Hash");
+			Content[Index].Hash = CrudeToInt(HashView);
+			check(int64(Content[Index].Hash) > 0);
+
+			Content[Index].Dest = FIoBuffer();
+			Response.SetDestination(&(Content[Index].Dest));
+			return nullptr;
+		}
+
+		uint32 ReceivedHash = 0x493;
+		FMemoryView ContentView = Content[Index].Dest.GetView();
+		check(ContentView.GetSize() == Status.GetContentLength());
+		for (uint32 i = 0; i < Status.GetContentLength(); ++i)
+		{
+			uint8 c = ((const uint8*)(ContentView.GetData()))[i];
+			ReceivedHash = (ReceivedHash + c) * 0x493;
+		}
+		check(Content[Index].Hash == ReceivedHash);
+		Content[Index].Hash = 0;
+
+		return nullptr;
+	};
+
+	auto NullSink = [] (const FTicketStatus&) {};
+
+	auto NoErrorSink = [&] (const FTicketStatus& Status)
+	{
+		check(Status.GetId() != FTicketStatus::EId::Error);
+		if (Status.GetId() != FTicketStatus::EId::Response)
+		{
+			return;
+		}
+
+		uint32 Index = Status.GetIndex();
+
+		FResponse& Response = Status.GetResponse();
+		Content[Index].Dest = FIoBuffer();
+		Response.SetDestination(&(Content[Index].Dest));
+	};
+
+	FEventLoop Loop;
+	volatile bool LoopStop = false;
+	volatile bool LoopTickDelay = false;
+	auto LoopTask = UE::Tasks::Launch(TEXT("IasHttpTest.Loop"), [&] () {
+		uint32 DelaySeed = 493;
+		while (!LoopStop)
+		{
+			while (Loop.Tick())
+			{
+				if (!LoopTickDelay)
+				{
+					continue;
+				}
+
+				float DelayFloat = float(DelaySeed % 75) / 1000.0f;
+				FPlatformProcess::SleepNoStats(DelayFloat);
+				DelaySeed *= 0xa93;
+			}
+
+			FPlatformProcess::SleepNoStats(0.005f);
+		}
+	});
+
+	auto WaitForLoopIdle = [&Loop] ()
+	{
+		FPlatformProcess::SleepNoStats(0.25f);
+		while (!Loop.IsIdle())
+		{
+			FPlatformProcess::SleepNoStats(0.03f);
+		}
+	};
+
+	// unused request
+	{
+		FRequest Request = Loop.Request("GET", BuildUrl("/data"));
+	}
+
+	// foundational
+	{
+		FRequest Request = Loop.Request("GET", BuildUrl("/seed/493"));
+		Request.Accept(EMimeType::Json);
+
+		FTicket Ticket = Loop.Send(MoveTemp(Request), NullSink);
+
+		WaitForLoopIdle();
+	}
+
+	// convenience
+	{
+		FRequest Request = Loop.Get(BuildUrl("/data")).Accept(EMimeType::Json);
+
+		FTicket Tickets[] = {
+			Loop.Send(MoveTemp(Request), HashSink),
+			Loop.Send(Loop.Get(BuildUrl("/data")).Accept(EMimeType::Json), HashSink),
+			Loop.Send(Loop.Get("http://httpbin.org/get"), NoErrorSink),
+		};
+		WaitForLoopIdle();
+	}
+
+	// convenience
+	{
+		FRequest Request = Loop.Get(BuildUrl("/data")).Accept(EMimeType::Json);
+		FTicket Ticket = Loop.Send(MoveTemp(Request), HashSink);
+		WaitForLoopIdle();
+	}
+
+	// pool
+	for (uint16 i = 1; i < 64; ++i)
+	{
+		FConnectionPool::FParams Params;
+		Params.SetHostFromUrl(BuildUrl());
+		Params.ConnectionCount = (i % 2) + 1;
+		FConnectionPool Pool(Params);
+		for (int32 j = 0; j < i; ++j)
+		{
+			FRequest Request = Loop.Get("/data", Pool);
+			Loop.Send(MoveTemp(Request), HashSink);
+		}
+		WaitForLoopIdle();
+	}
+
+	// fatal timeout
+	for (int32 i = 0; i < 14; ++i)
+	{
+		bool bExpectFailTimeout = !!(i & 1);
+		auto Sink = [bExpectFailTimeout, Dest=FIoBuffer()] (const FTicketStatus& Status) mutable
+		{
+			if (Status.GetId() == FTicketStatus::EId::Response)
+			{
+				FResponse& Response = Status.GetResponse();
+				Response.SetDestination(&Dest);
+				return;
+			}
+
+			check(Status.GetId() == FTicketStatus::EId::Error);
+
+			const char* Reason = Status.GetErrorReason();
+			bool IsFailTimeout = (FCStringAnsi::Strstr(Reason, "FailTimeout") != nullptr);
+			check(IsFailTimeout == bExpectFailTimeout);
+		};
+
+		auto ErrorSink = [] (const FTicketStatus& Status)
+		{
+			check(Status.GetId() == FTicketStatus::EId::Error);
+		};
+
+		FConnectionPool::FParams Params;
+		Params.SetHostFromUrl(BuildUrl());
+		FConnectionPool Pool(Params);
+
+		FEventLoop Loop2;
+		Loop2.Send(Loop2.Get("/data?stall=1", Pool), Sink);
+
+		// Requests are pipelined. The second one will get went during the stall so
+		// we expect it to fail. The subsequent ones are expected to succeed.
+		Loop2.Send(Loop2.Get("/data", Pool), ErrorSink);
+		Loop2.Send(Loop2.Get("/data", Pool), HashSink);
+		Loop2.Send(Loop2.Get("/data", Pool), HashSink);
+
+		int32 PollTimeoutMs = -1;
+		if (bExpectFailTimeout)
+		{
+			Loop2.SetFailTimeout(1000);
+
+			if ((i & 3) == 1)
+			{
+				PollTimeoutMs = 1000;
+			}
+		}
+		while (Loop2.Tick(PollTimeoutMs));
+
+		Loop2.Send(Loop2.Get("/data/23", Pool), NoErrorSink);
+		while (Loop2.Tick(PollTimeoutMs));
+	}
+
+	// no connect
+	{
+		FRequest Requests[] = {
+			Loop.Request("GET", BuildUrl(nullptr, 10930)),
+			Loop.Request("GET", "http://thisdoesnotexistihope/"),
+		};
+		Loop.Send(MoveTemp(Requests[0]), NullSink);
+		Loop.Send(MoveTemp(Requests[1]), NullSink);
+		WaitForLoopIdle();
+	}
+
+	// head and large requests
+	{
+		auto MixTh = [Th=uint32(0)] () mutable { return (Th = (Th * 75) + 74) & 255; };
+
+		char AsciiData[257];
+		for (char& c : AsciiData)
+		{
+			int32 i = int32(ptrdiff_t(&c - AsciiData));
+			c = 0x41 + (MixTh() % 26);
+			c += (MixTh() & 2) ? 0x20 : 0;
+		}
+
+		for (int32 i = 0; (i += 69493) < 2 << 20;)
+		{
+			FRequest Request = Loop.Request("HEAD", BuildUrl("/data"));
+			for (int32 j = i; j > 0;)
+			{
+				FAnsiStringView Name(AsciiData, MixTh() + 1);
+				FAnsiStringView Value(AsciiData, MixTh() + 1);
+				Request.Header(Name, Value);
+				j -= Name.Len() + Value.Len();
+
+			}
+
+			Loop.Send(MoveTemp(Request), [] (const FTicketStatus& Status) {
+				if (Status.GetId() == FTicketStatus::EId::Response)
+				{
+					FResponse& Response = Status.GetResponse();
+					check(Response.GetStatusCode() == 431); // "too many headers"
+				}
+			});
+
+			WaitForLoopIdle();
+		}
+	}
+
+	// stress 1
+	{
+		const uint32 StressLoad = 32;
+
+		struct {
+			const ANSICHAR* Uri;
+			bool Disconnect;
+		} StressUrls[] = {
+			{ "/data?slowly=1",		false },
+			{ "/data?disconnect=1",	true },
+		};
+
+		uint64 Errors = 0;
+		auto ErrorSink = [&] (const FTicketStatus& Status)
+		{
+			FTicket Ticket = Status.GetTicket();
+			uint32 Index = uint32(63 - FMath::CountLeadingZeros64(uint64(Ticket)));
+
+			if (Status.GetId() == FTicketStatus::EId::Error)
+			{
+				Errors |= 1ull << Index;
+				return;
+			}
+
+			else if (Status.GetId() == FTicketStatus::EId::Response)
+			{
+				FResponse& Response = Status.GetResponse();
+				Content[Index].Dest = FIoBuffer();
+				Response.SetDestination(&(Content[Index].Dest));
+				return;
+			}
+
+			check(false);
+		};
+
+		for (const auto& [StressUri, ExpectDisconnect] : StressUrls)
+		{
+			FTicketSink Sink = ExpectDisconnect ? FTicketSink(ErrorSink) : FTicketSink(HashSink);
+
+			const auto& StressUrl = BuildUrl(StressUri);
+			for (bool AddDelay : {false, true})
+			{
+				FTicket Tickets[StressLoad];
+				for (FTicket& Ticket : Tickets)
+				{
+					Ticket = Loop.Send(Loop.Get(StressUrl).Header("Accept", "*/*"), Sink);
+				}
+
+				LoopTickDelay = AddDelay;
+
+				WaitForLoopIdle();
+			}
+
+			LoopTickDelay = false;
+		}
+	}
+
+	// stress 2
+	{
+		const uint32 StressLoad = 3;
+		const uint32 StressTaskCount = 7;
+		static_assert(StressLoad * StressTaskCount <= 32);
+
+		FAnsiStringView Url = BuildUrl("/data");
+
+		auto StressTaskEntry = [&] {
+			for (uint32 i = 0; i < StressLoad; ++i)
+			{
+				FTicket Ticket = Loop.Send(Loop.Get(Url), HashSink);
+				if (!Ticket)
+				{
+					FPlatformProcess::SleepNoStats(0.01f);
+					--i;
+				}
+			}
+		};
+
+		UE::Tasks::FTask StressTasks[StressTaskCount];
+		for (auto& StressTask : StressTasks)
+		{
+			StressTask = UE::Tasks::Launch(TEXT("StressTask"), [&] { StressTaskEntry(); });
+		}
+		for (auto& StressTask : StressTasks)
+		{
+			StressTask.Wait();
+		}
+
+		WaitForLoopIdle();
+	}
+
+	// tamper
+	for (int32 i = 1; i <= 100; ++i)
+	{
+		TAnsiStringBuilder<32> TamperUrl;
+		TamperUrl << "/data?tamper=" << i;
+		FAnsiStringView Url = BuildUrl(TamperUrl.ToString());
+
+		for (int j = 0; j < 48; ++j)
+		{
+			FRequest Request = Loop.Request("GET", Url);
+			Loop.Send(MoveTemp(Request), NullSink);
+		}
+
+		WaitForLoopIdle();
+	}
+
+	LoopStop = true;
+	LoopTask.Wait();
+
+	check(Loop.IsIdle());
+
+#if IS_PROGRAM
+	ThrottleTest(BuildUrl("/data/"));
+#endif
+
+	// pre-generated headers
+	// request-with-body
+	// proxy
+	// chunked transfer encoding
+	// gzip / deflate
+	// redirects
+	// loop multi-req.
+	// tls
+	// url auth credentials
+	// transfer-file / splice / sendfile
+	// (header field parser)
+	// (form-data)
+	// (cookies)
+	// (cache)
+	// (websocket)
+	// (ipv6)
+	// (utf-8 host names)
+}
+
+////////////////////////////////////////////////////////////////////////////////
+IOSTOREONDEMAND_API void IasHttpTest(const ANSICHAR* TestHost="localhost")
+{
+#if PLATFORM_WINDOWS
+	WSADATA WsaData;
+	if (WSAStartup(MAKEWORD(2, 2), &WsaData) == 0x0a9e0493)
+		return;
+	ON_SCOPE_EXIT { WSACleanup(); };
+#endif
+
+	MiscTest();
+	HttpTest(TestHost);
+	RedirectTest(TestHost);
+}
+
+#endif // !SHIP|TEST
+// }}}
+
+} // namespace UE::IoStore::HTTP
+
