Moved latency testing "into" IoStoreOnDemand. This is strategic to pair down the client that may be moved out to its own module.

#rnx

diff --git a/Engine/Source/Runtime/Experimental/IoStoreOnDemand/Private/CoreHttp/LatencyTesting.cpp b/Engine/Source/Runtime/Experimental/IoStoreOnDemand/Private/CoreHttp/LatencyTesting.cpp
deleted file mode 100644
diff --git a/Engine/Source/Runtime/Experimental/IoStoreOnDemand/Private/CoreHttp/LatencyTesting.h b/Engine/Source/Runtime/Experimental/IoStoreOnDemand/Private/CoreHttp/LatencyTesting.h
deleted file mode 100644
diff --git a/Engine/Source/Runtime/Experimental/IoStoreOnDemand/Private/LatencyTesting.cpp b/Engine/Source/Runtime/Experimental/IoStoreOnDemand/Private/LatencyTesting.cpp
new file mode 100644
--- /dev/null
+++ b/Engine/Source/Runtime/Experimental/IoStoreOnDemand/Private/LatencyTesting.cpp
@@ -0,0 +1,69 @@
+// Copyright Epic Games, Inc. All Rights Reserved.
+
+#include "Containers/StringConv.h"
+#include "Containers/StringView.h"
+#include "CoreHttp/Client.h"
+#include "HAL/PlatformTime.h"
+#include "IO/IoStoreOnDemand.h"
+#include "Misc/StringBuilder.h"
+#include "Templates/Function.h"
+
+namespace UE::IoStore
+{
+
+void LatencyTest(FStringView InUrl, FStringView InPath, uint32 InTimeOutMs, TArrayView<int32> OutResults)
+{
+	using namespace HTTP;
+
+	auto AnsiUrl = StringCast<ANSICHAR>(InUrl.GetData(), InUrl.Len());
+
+	FConnectionPool::FParams PoolParams;
+	PoolParams.SetHostFromUrl(AnsiUrl);
+	PoolParams.ConnectionCount = 1;
+	FConnectionPool Pool(PoolParams);
+
+	TAnsiStringBuilder<256> AnsiPath;
+	if (!InPath.StartsWith(TEXT('/')))
+	{
+		AnsiPath << '/';
+	}
+	AnsiPath << InPath;
+
+	FEventLoop Loop;
+
+	for (int32& Result : OutResults)
+	{
+		bool Ok = false;
+
+		FRequest Request = Loop.Request("HEAD", AnsiPath, Pool);
+		Loop.Send(MoveTemp(Request), [&](const FTicketStatus& Status)
+			{
+				if (Status.GetId() != FTicketStatus::EId::Response)
+					return;
+
+				const FResponse& Response = Status.GetResponse();
+				Ok = (Response.GetStatus() == EStatusCodeClass::Successful);
+			});
+
+		uint64 Cycles = FPlatformTime::Cycles64();
+		while (Loop.Tick(InTimeOutMs) != 0);
+		Cycles = FPlatformTime::Cycles64() - Cycles;
+
+		Result = Ok ? int32(Cycles) : -1;
+	}
+
+	int64 Freq = int64(1.0 / FPlatformTime::GetSecondsPerCycle());
+	for (int32& Result : OutResults)
+	{
+		if (Result == -1)
+			continue;
+
+		Result = int32((int64(Result) * 1000) / Freq);
+	}
+	
+	TAnsiStringBuilder<512> ConnectionDesc;
+	Pool.Describe(ConnectionDesc);
+	UE_LOG(LogIas, Log, TEXT("Testing endpoint %s"), ANSI_TO_TCHAR(ConnectionDesc.ToString()));
+}
+
+} // namespace UE::IoStore
diff --git a/Engine/Source/Runtime/Experimental/IoStoreOnDemand/Private/OnDemandIoDispatcherBackend.cpp b/Engine/Source/Runtime/Experimental/IoStoreOnDemand/Private/OnDemandIoDispatcherBackend.cpp
--- a/Engine/Source/Runtime/Experimental/IoStoreOnDemand/Private/OnDemandIoDispatcherBackend.cpp
+++ b/Engine/Source/Runtime/Experimental/IoStoreOnDemand/Private/OnDemandIoDispatcherBackend.cpp
@@ -5,7 +5,6 @@
 #include "AnalyticsEventAttribute.h"
 #include "Containers/BitArray.h"
 #include "Containers/StringView.h"
-#include "CoreHttp/LatencyTesting.h"
 #include "DistributionEndpoints.h"
 #include "GenericPlatform/GenericPlatformCrashContext.h"
 #include "HAL/Event.h"
@@ -51,6 +50,7 @@ namespace UE::IoStore
 ///////////////////////////////////////////////////////////////////////////////
 
 extern FString GIasOnDemandTocExt;
+void LatencyTest(FStringView, FStringView, uint32, TArrayView<int32>);
 
 ///////////////////////////////////////////////////////////////////////////////
 int32 GIasHttpPrimaryEndpoint = 0;
@@ -209,7 +209,7 @@ static void LatencyTest(FStringView Url, FStringView Path)
 	TRACE_CPUPROFILER_EVENT_SCOPE(IasBackend::LatencyTest);
 
 	int32 Results[4] = {};
-	UE::IoStore::HTTP::LatencyTest(Url, Path, GIasHttpTimeOutMs, MakeArrayView(Results));
+	LatencyTest(Url, Path, GIasHttpTimeOutMs, MakeArrayView(Results));
 	UE_LOG(LogIas, Log, TEXT("Endpoint '%s' latency test (ms): %d %d %d %d"),
 		Url.GetData(), Results[0], Results[1], Results[2], Results[3]);
 }
@@ -223,7 +223,7 @@ static int32 LatencyTest(TConstArrayView<FString> Urls, FStringView Path, std::a
 	for (int32 Idx = 0; Idx < Urls.Num() && !bCancel.load(std::memory_order_relaxed); ++Idx)
 	{
 		int32 LatencyMs = -1;
-		UE::IoStore::HTTP::LatencyTest(Urls[Idx], Path, GIasHttpTimeOutMs, MakeArrayView(&LatencyMs, 1));
+		LatencyTest(Urls[Idx], Path, GIasHttpTimeOutMs, MakeArrayView(&LatencyMs, 1));
 		if (LatencyMs > 0)
 		{
 			return Idx;
