Made IAS' CoreHttp visible to the big wide world.

Note that DLL exports might not be correctly in place at time of writing.

#rnx

diff --git a/Engine/Source/Runtime/Experimental/IoStoreOnDemand/Private/CoreHttp/Client.h b/Engine/Source/Runtime/Experimental/IoStoreOnDemand/Private/CoreHttp/Client.h
deleted file mode 100644
diff --git a/Engine/Source/Runtime/Experimental/IoStoreOnDemand/Private/CoreHttp/Http.cpp b/Engine/Source/Runtime/Experimental/IoStoreOnDemand/Private/CoreHttp/Http.cpp
--- a/Engine/Source/Runtime/Experimental/IoStoreOnDemand/Private/CoreHttp/Http.cpp
+++ b/Engine/Source/Runtime/Experimental/IoStoreOnDemand/Private/CoreHttp/Http.cpp
@@ -1,6 +1,6 @@
 // Copyright Epic Games, Inc. All Rights Reserved.
 
-#include "Client.h"
+#include "CoreHttp/Client.h"
 
 #if !defined(NO_UE_INCLUDES)
 #include <Containers/Array.h>
diff --git a/Engine/Source/Runtime/Experimental/IoStoreOnDemand/Public/CoreHttp/Client.h b/Engine/Source/Runtime/Experimental/IoStoreOnDemand/Public/CoreHttp/Client.h
new file mode 100644
--- /dev/null
+++ b/Engine/Source/Runtime/Experimental/IoStoreOnDemand/Public/CoreHttp/Client.h
@@ -0,0 +1,223 @@
+// Copyright Epic Games, Inc. All Rights Reserved.
+
+#pragma once
+
+#if !defined(NO_UE_INCLUDES)
+#include <Containers/StringView.h>
+#endif
+
+#if !defined(IAS_HTTP_WITH_PERF)
+#	define IAS_HTTP_WITH_PERF !UE_BUILD_SHIPPING
+#endif
+
+#define UE_API
+
+////////////////////////////////////////////////////////////////////////////////
+class FIoBuffer;
+
+namespace UE::IoStore::HTTP
+{
+
+////////////////////////////////////////////////////////////////////////////////
+enum class EMimeType
+{
+	Unknown = 0,
+	Text,
+	Binary,
+	Json,
+	Xml,
+	/* UE_CUSTOM_MIME_TYPES
+	CbObject,
+	CbPackage,
+	CompressedBuffer,
+	*/
+	Count
+};
+
+////////////////////////////////////////////////////////////////////////////////
+enum class EStatusCodeClass
+{
+	Informational,
+	Successful,
+	Redirection,
+	ClientError,
+	ServerError,
+	Unknown,
+};
+
+////////////////////////////////////////////////////////////////////////////////
+using	FTicket = uint64;
+struct	FActivity;
+
+////////////////////////////////////////////////////////////////////////////////
+class UE_API FConnectionPool
+{
+public:
+	struct FParams
+	{
+		int32				SetHostFromUrl(FAnsiStringView Url);
+		FAnsiStringView		HostName;
+		int32				SendBufSize = -1;
+		int32				RecvBufSize = -1;
+		uint32				Port = 0;
+		uint16				ConnectionCount = 1;
+		/*
+		enum class ProxyType { Http, Socks4 };
+		Proxy = { ip, port, type }
+		 */
+	};
+
+							FConnectionPool() = default;
+							FConnectionPool(const FParams& Params);
+							~FConnectionPool();
+							FConnectionPool(FConnectionPool&& Rhs)	{ *this = MoveTemp(Rhs); }
+	FConnectionPool&		operator = (FConnectionPool&& Rhs)		{ Swap(Ptr, Rhs.Ptr); return *this; }
+	bool					Resolve();
+	void					Describe(FAnsiStringBuilderBase&) const;
+
+	static bool				IsValidHostUrl(FAnsiStringView Url);
+
+private:
+	friend					class FEventLoop;
+	class FHost*			Ptr = nullptr;
+
+private:
+							FConnectionPool(const FConnectionPool&) = delete;
+	FConnectionPool&		operator = (const FConnectionPool&) = delete;
+};
+
+////////////////////////////////////////////////////////////////////////////////
+class UE_API FRequest
+{
+public:
+						~FRequest();
+						FRequest(FRequest&& Rhs);
+	bool				IsValid() const { return Ptr != nullptr; }
+	FRequest&&			Accept(EMimeType MimeType);
+	FRequest&&			Accept(FAnsiStringView MimeType);
+	FRequest&&			Header(FAnsiStringView Key, FAnsiStringView Value);
+	void				Content(const void* Data, SIZE_T Size, EMimeType MimeType);
+	void				Content(const void* Data, SIZE_T Size, FAnsiStringView MimeType);
+
+private:
+	friend				class FEventLoop;
+						FRequest() = default;
+	FActivity*			Ptr = nullptr;
+
+private:
+						FRequest(const FRequest&) = delete;
+	FRequest&			operator = (const FRequest&) = delete;
+	FRequest&			operator = (FRequest&&) = delete;
+};
+
+////////////////////////////////////////////////////////////////////////////////
+class UE_API FResponse
+{
+public:
+	using FHeaderSink = TFunction<bool (FAnsiStringView, FAnsiStringView)>;
+
+	EStatusCodeClass	GetStatus() const;
+	uint32				GetStatusCode() const;
+	FAnsiStringView 	GetStatusMessage() const;
+	int64				GetContentLength() const;
+	EMimeType			GetContentType() const;
+	void				GetContentType(FAnsiStringView& Out) const;
+	FAnsiStringView 	GetHeader(FAnsiStringView Name) const;
+	void				ReadHeaders(FHeaderSink Sink) const;
+	void				SetDestination(FIoBuffer* Buffer);
+
+private:
+						FResponse() = delete;
+						FResponse(const FResponse&) = delete;
+						FResponse(FResponse&&) = delete;
+	FResponse&			operator = (const FResponse&) = delete;
+	FResponse&			operator = (FResponse&&) = delete;
+};
+
+////////////////////////////////////////////////////////////////////////////////
+class UE_API FTicketPerf
+{
+public:
+	struct FSample
+	{
+		uint32			SendMs;
+		uint32			WaitMs;
+		uint32			RecvMs;
+	};
+
+	FSample				GetSample() const;
+};
+
+#if !IAS_HTTP_WITH_PERF
+inline FTicketPerf::FSample FTicketPerf::GetSample() const { return {}; }
+#endif
+
+////////////////////////////////////////////////////////////////////////////////
+class UE_API FTicketStatus
+{
+public:
+	enum class EId : uint8 { Response, Content, Cancelled, Error };
+	EId					GetId() const;
+	UPTRINT				GetParam() const;
+	FTicket				GetTicket() const;
+	uint32				GetIndex() const;
+	FResponse&			GetResponse() const;		// if GetId() == EId::Response
+	const FIoBuffer&	GetContent() const;			// if GetId() == EId::Content
+	uint32				GetContentLength() const;	//  |
+	const FTicketPerf&	GetPerf() const;			// _|_
+	const char*			GetErrorReason() const;		// if GetId() == EId::Error
+
+private:
+						FTicketStatus() = delete;
+						FTicketStatus(const FTicketStatus&) = delete;
+						FTicketStatus(FTicketStatus&&) = delete;
+	FTicketStatus&		operator = (const FTicketStatus&) = delete;
+	FTicketStatus&		operator = (FTicketStatus&&) = delete;
+};
+
+using FTicketSink = TFunction<void (const FTicketStatus&)>;
+
+////////////////////////////////////////////////////////////////////////////////
+class UE_API FEventLoop
+{
+	class FImpl;
+
+public:
+	static const uint32		MaxActiveTickets = 64;
+
+	struct FRequestParams
+	{
+		uint32	BufferSize	= 256;
+		bool	bAutoRedirect = false;
+	};
+
+	template <typename... T> [[nodiscard]] FRequest Get(T&&... t)  { return Request("GET",  Forward<T&&>(t)...); }
+	template <typename... T> [[nodiscard]] FRequest Post(T&&... t) { return Request("POST", Forward<T&&>(t)...); }
+
+							FEventLoop();
+							~FEventLoop();
+	uint32					Tick(int32 PollTimeoutMs=0);
+	void					Throttle(uint32 KiBPerSec);
+	void					SetFailTimeout(int32 TimeoutMs);
+	bool					IsIdle() const;
+	void					Cancel(FTicket Ticket);
+	[[nodiscard]] FRequest	Request(FAnsiStringView Method, FAnsiStringView Url, const FRequestParams* Params=nullptr);
+	[[nodiscard]] FRequest	Request(FAnsiStringView Method, FAnsiStringView Path, FConnectionPool& Pool, const FRequestParams* Params=nullptr);
+	FTicket					Send(FRequest&& Request, FTicketSink Sink, UPTRINT SinkParam=0);
+
+private:
+	bool					Redirect(const FTicketStatus& Status, FTicketSink& OuterSink);
+	FImpl*					Impl;
+
+private:
+							FEventLoop(const FEventLoop&)	= delete;
+							FEventLoop(FEventLoop&&)		= delete;
+	FEventLoop&				operator = (const FEventLoop&)	= delete;
+	FEventLoop&				operator = (FEventLoop&&)		= delete;
+};
+
+} // namespace UE::IoStore::HTTP
+
+#undef UE_API
+
+/* vim: set noet : */
