Fix vectorized solver
Add a unit test
The vectorized initialization position locked was not correct.

diff --git a/Engine/Source/Programs/HeadlessChaos/Private/HeadlessChaosTestConstraints.cpp b/Engine/Source/Programs/HeadlessChaos/Private/HeadlessChaosTestConstraints.cpp
--- a/Engine/Source/Programs/HeadlessChaos/Private/HeadlessChaosTestConstraints.cpp
+++ b/Engine/Source/Programs/HeadlessChaos/Private/HeadlessChaosTestConstraints.cpp
@@ -83,7 +83,7 @@ namespace ChaosTest {
 	 * Joint constraints test with the fixed body held in place with a position constraint.
 	 * Joint body swings under the fixed body at fixed distance.
 	 */
-	template<typename TEvolution>
+	template<typename TEvolution, bool bUseSimd = false>
 	void PositionAndJoint()
 	{
 		const int32 Iterations = 10;
@@ -106,6 +106,7 @@ namespace ChaosTest {
 
 		TVec2<TGeometryParticleHandle<FReal, 3>*> JointParticles = { Dynamics[0], Dynamics[1] };
 		FPBDJointConstraints JointConstraints;
+		JointConstraints.SetUseSimd(bUseSimd);
 		JointConstraints.AddConstraint(JointParticles, FRigidTransform3(JointConstraintPosition, FRotation3::FromIdentity()));
 		Evolution.AddConstraintContainer(JointConstraints);
 
@@ -331,6 +332,8 @@ namespace ChaosTest {
 	{
 		ChaosTest::Position<FPBDRigidsEvolutionGBF>();
 		ChaosTest::PositionAndJoint<FPBDRigidsEvolutionGBF>();
+		constexpr bool bTestSimdSolver = true;
+		ChaosTest::PositionAndJoint<FPBDRigidsEvolutionGBF, bTestSimdSolver>();
 		ChaosTest::SuspensionConstraintHardstop<FPBDRigidsEvolutionGBF>();
 		ChaosTest::SuspensionConstraintSpring<FPBDRigidsEvolutionGBF>();
 
diff --git a/Engine/Source/Runtime/Experimental/Chaos/Private/Chaos/Joint/PBDJointCachedSolverGaussSeidel.cpp b/Engine/Source/Runtime/Experimental/Chaos/Private/Chaos/Joint/PBDJointCachedSolverGaussSeidel.cpp
--- a/Engine/Source/Runtime/Experimental/Chaos/Private/Chaos/Joint/PBDJointCachedSolverGaussSeidel.cpp
+++ b/Engine/Source/Runtime/Experimental/Chaos/Private/Chaos/Joint/PBDJointCachedSolverGaussSeidel.cpp
@@ -623,12 +623,19 @@ void FPBDJointCachedSolver::InitLockedPositionConstraintSimd(
 		PositionConstraints.ConstraintVX[ConstraintIndex] = 0.0;
 		PositionConstraints.ConstraintRestitution[ConstraintIndex] = 0.0;
 
-		PositionConstraints.Simd.ConstraintArms[ConstraintIndex][0] = MakeVectorRegisterFloatFromDouble(MakeVectorRegister(ConstraintArm0[0], ConstraintArm0[1], ConstraintArm0[2], 0.0f));
-		PositionConstraints.Simd.ConstraintArms[ConstraintIndex][1] = MakeVectorRegisterFloatFromDouble(MakeVectorRegister(ConstraintArm1[0], ConstraintArm1[1], ConstraintArm1[2], 0.0f));
 		PositionConstraints.Simd.ConstraintAxis[ConstraintIndex] = MakeVectorRegisterFloatFromDouble(MakeVectorRegister(LocalAxis[0], LocalAxis[1], LocalAxis[2], 0.0f));
+	}
+	PositionConstraints.Simd.ConstraintCX = MakeVectorRegisterFloat(LocalDeltas[0], LocalDeltas[1], LocalDeltas[2], 0.0f);
 
-		InitConstraintAxisLinearVelocities[ConstraintIndex] = FVec3::DotProduct(InitConstraintVelocity, LocalAxis);
+	PositionConstraints.Simd.ConstraintArms[0][0] = MakeVectorRegisterFloatFromDouble(MakeVectorRegister(ConstraintArm0[0], ConstraintArm0[1], ConstraintArm0[2], 0.0f));
+	PositionConstraints.Simd.ConstraintArms[0][1] = MakeVectorRegisterFloatFromDouble(MakeVectorRegister(ConstraintArm1[0], ConstraintArm1[1], ConstraintArm1[2], 0.0f));
+	PositionConstraints.Simd.ConstraintArms[1][0] = PositionConstraints.Simd.ConstraintArms[0][0];
+	PositionConstraints.Simd.ConstraintArms[1][1] = PositionConstraints.Simd.ConstraintArms[0][1];
+	PositionConstraints.Simd.ConstraintArms[2][0] = PositionConstraints.Simd.ConstraintArms[0][0];
+	PositionConstraints.Simd.ConstraintArms[2][1] = PositionConstraints.Simd.ConstraintArms[0][1];
 
+	for (int32 ConstraintIndex = 0; ConstraintIndex < 3; ++ConstraintIndex)
+	{
 		const VectorRegister4Float AngularAxis0 = VectorCross(PositionConstraints.Simd.ConstraintArms[ConstraintIndex][0], PositionConstraints.Simd.ConstraintAxis[ConstraintIndex]);
 		const VectorRegister4Float IA0 = Private::VectorMatrixMultiply(AngularAxis0, InvI(0));
 
@@ -640,10 +647,7 @@ void FPBDJointCachedSolver::InitLockedPositionConstraintSimd(
 		HardIM[ConstraintIndex] += II0 + II1;
 		PositionConstraints.Simd.ConstraintDRAxis[ConstraintIndex][0] = IA0;
 		PositionConstraints.Simd.ConstraintDRAxis[ConstraintIndex][1] = VectorNegate(IA1);
-
-
 	}
-	PositionConstraints.Simd.ConstraintCX = MakeVectorRegisterFloat(LocalDeltas[0], LocalDeltas[1], LocalDeltas[2], 0.0f);
 	PositionConstraints.Simd.ConstraintHardIM = MakeVectorRegisterFloatFromDouble(MakeVectorRegister(HardIM[0], HardIM[1], HardIM[2], 0.0f));
 	SetInitConstraintVelocity(ConstraintArm0, ConstraintArm1);
 
