Refactored AddElementToExistingGrid to avoid potentially-uninitialized FVector warning.

diff --git a/Engine/Source/Runtime/Experimental/Chaos/Public/Chaos/BoundingVolume.h b/Engine/Source/Runtime/Experimental/Chaos/Public/Chaos/BoundingVolume.h
--- a/Engine/Source/Runtime/Experimental/Chaos/Public/Chaos/BoundingVolume.h
+++ b/Engine/Source/Runtime/Experimental/Chaos/Public/Chaos/BoundingVolume.h
@@ -1213,49 +1213,47 @@ private:
 
 		if(bHasBounds)
 		{
-			bool bDirty = bIsEmpty;
-			TVector<int32, 3> StartIndex;
-			TVector<int32, 3> EndIndex;
+			PayloadInfo.GlobalPayloadIdx = INDEX_NONE;
 
 			if (bIsEmpty == false)
 			{
 				//add payload to appropriate cells
-				StartIndex = MGrid.CellUnsafe(NewBounds.Min());
-				EndIndex = MGrid.CellUnsafe(NewBounds.Max());
+				TVector<int32, 3> StartIndex = MGrid.CellUnsafe(NewBounds.Min());
+				TVector<int32, 3> EndIndex = MGrid.CellUnsafe(NewBounds.Max());
 
+				bool bDirty = false;
 				for (int Axis = 0; Axis < d; ++Axis)
 				{
 					if (StartIndex[Axis] < 0 || EndIndex[Axis] >= MGrid.Counts()[Axis])
 					{
 						bDirty = true;
+						break;
 					}
 				}
-			}
 
-			PayloadInfo.GlobalPayloadIdx = INDEX_NONE;
-
-			if (!bDirty)
-			{
-				PayloadInfo.DirtyPayloadIdx = INDEX_NONE;
-				PayloadInfo.StartIdx = StartIndex;
-				PayloadInfo.EndIdx = EndIndex;
-
-				for (int32 x = StartIndex[0]; x <= EndIndex[0]; ++x)
+				if (!bDirty)
 				{
-					for (int32 y = StartIndex[1]; y <= EndIndex[1]; ++y)
+					PayloadInfo.DirtyPayloadIdx = INDEX_NONE;
+					PayloadInfo.StartIdx = StartIndex;
+					PayloadInfo.EndIdx = EndIndex;
+
+					for (int32 x = StartIndex[0]; x <= EndIndex[0]; ++x)
 					{
-						for (int32 z = StartIndex[2]; z <= EndIndex[2]; ++z)
+						for (int32 y = StartIndex[1]; y <= EndIndex[1]; ++y)
 						{
-							MElements(x, y, z).Add({ NewBounds, Payload, StartIndex, EndIndex });
+							for (int32 z = StartIndex[2]; z <= EndIndex[2]; ++z)
+							{
+								MElements(x, y, z).Add({ NewBounds, Payload, StartIndex, EndIndex });
+							}
 						}
 					}
+
+					return;
 				}
 			}
-			else
-			{
-				PayloadInfo.DirtyPayloadIdx = MDirtyElements.Num();
-				MDirtyElements.Add({ NewBounds, Payload });
-			}
+
+			PayloadInfo.DirtyPayloadIdx = MDirtyElements.Num();
+			MDirtyElements.Add({ NewBounds, Payload });
 		}
 		else
 		{
