Dataflow : Context Caching
- Check if the FStructProperty supports hashing when validating the node properties.
- Return empty hash for the node if a property does not support the hash.

diff --git a/Engine/Source/Runtime/Experimental/Dataflow/Core/Private/Dataflow/DataflowNode.cpp b/Engine/Source/Runtime/Experimental/Dataflow/Core/Private/Dataflow/DataflowNode.cpp
--- a/Engine/Source/Runtime/Experimental/Dataflow/Core/Private/Dataflow/DataflowNode.cpp
+++ b/Engine/Source/Runtime/Experimental/Dataflow/Core/Private/Dataflow/DataflowNode.cpp
@@ -518,12 +518,27 @@ uint32 FDataflowNode::GetValueHash()
 			{
 				if (const FProperty* const Property = PropertyIt.Key())
 				{
+					if (const FStructProperty* StructProperty = CastField<FStructProperty>(Property))
+					{
+						//
+						// Note : [CacheContextPropertySupport]
+						// 
+						// Some UPROPERTIES do not support hash values.For example, FFilePath, is a struct 
+						// that is not defined using USTRUCT, and does not support the GetTypeValue() function.
+						// These types of attributes need to return a Zero(0) hash, to indicate that the Hash 
+						// is not supported.To add property hashing support, add GetTypeValue to the properties 
+						// supporting USTRUCT(See Class.h  UScriptStruct::GetStructTypeHash)
+						// 
+						if (!StructProperty->Struct) return 0;
+						if (!StructProperty->Struct->GetCppStructOps()) return 0;
+					}
+
 					if (Property->PropertyFlags & CPF_HasGetValueTypeHash)
 					{
 						// uint32 CrcHash = FCrc::MemCrc32(PropertyIt.Value(), Property->ElementSize);
 						// UE_LOG(LogChaos, Warning, TEXT("( %lu \t%s"), (unsigned long)CrcHash, *Property->GetName())
 
-						if ( Property->PropertyFlags & CPF_TObjectPtr )
+						if (Property->PropertyFlags & CPF_TObjectPtr)
 						{
 							// @todo(dataflow) : Do something about TObjectPtr<T>
 						}
@@ -539,6 +554,34 @@ uint32 FDataflowNode::GetValueHash()
 	return Hash;
 }
 
+void FDataflowNode::ValidateProperties()
+{
+	if (const TUniquePtr<FStructOnScope> ScriptOnStruct = TUniquePtr<FStructOnScope>(NewStructOnScope()))
+	{
+		if (const UStruct* const Struct = ScriptOnStruct->GetStruct())
+		{
+			for (FPropertyValueIterator PropertyIt(FProperty::StaticClass(), Struct, this); PropertyIt; ++PropertyIt)
+			{
+				if (const FProperty* const Property = PropertyIt.Key())
+				{
+					if (const FStructProperty* StructProperty = CastField<FStructProperty>(Property))
+					{
+						if (!StructProperty->Struct || !StructProperty->Struct->GetCppStructOps())
+						{
+							// See Note : [CacheContextPropertySupport]
+							FString StructPropertyName;
+							StructProperty->GetName(StructPropertyName);
+							UE_LOG(LogChaos, Warning, 
+								TEXT("Dataflow: Context caching disable for graphs with node '%s' due to non-hashed UPROPERTY '%s'."), 
+								*GetName().ToString(), *StructPropertyName)
+						}
+					}
+				}
+			}
+		}
+	}
+}
+
 bool FDataflowNode::ValidateConnections()
 {
 	bHasValidConnections = true;
diff --git a/Engine/Source/Runtime/Experimental/Dataflow/Core/Public/Dataflow/DataflowNode.h b/Engine/Source/Runtime/Experimental/Dataflow/Core/Public/Dataflow/DataflowNode.h
--- a/Engine/Source/Runtime/Experimental/Dataflow/Core/Public/Dataflow/DataflowNode.h
+++ b/Engine/Source/Runtime/Experimental/Dataflow/Core/Public/Dataflow/DataflowNode.h
@@ -275,6 +275,8 @@ struct FDataflowNode
 
 	DATAFLOWCORE_API virtual bool ValidateConnections();
 
+	DATAFLOWCORE_API virtual void ValidateProperties();
+
 	bool HasValidConnections() const { return bHasValidConnections; }
 
 	virtual bool IsA(FName InType) const 
@@ -347,7 +349,7 @@ namespace Dataflow
 		[](const ::Dataflow::FNewNodeParameters& InParam){							\
 				TUniquePtr<A> Val = MakeUnique<A>(::Dataflow::FNodeParameters{		\
 					InParam.Name, InParam.OwningObject}, InParam.Guid);				\
-				Val->ValidateConnections(); return Val;});
+				Val->ValidateProperties();Val->ValidateConnections();return Val;});
 
 #define DATAFLOW_NODE_RENDER_TYPE(A, B)												\
 	virtual TArray<::Dataflow::FRenderingParameter> GetRenderParametersImpl() const {		\
