Fixed bug when resolving CDNs from distribution endpoint

diff --git a/Engine/Source/Runtime/Experimental/IoStoreOnDemand/Private/OnDemandIoDispatcherBackend.cpp b/Engine/Source/Runtime/Experimental/IoStoreOnDemand/Private/OnDemandIoDispatcherBackend.cpp
--- a/Engine/Source/Runtime/Experimental/IoStoreOnDemand/Private/OnDemandIoDispatcherBackend.cpp
+++ b/Engine/Source/Runtime/Experimental/IoStoreOnDemand/Private/OnDemandIoDispatcherBackend.cpp
@@ -1685,7 +1685,7 @@ uint32 FOnDemandIoBackend::Run()
 
 	if (DistributionUrl.IsValid())
 	{
-		if (ResolveDistributedEndpoint(DistributionUrl))
+		if (ResolveDistributedEndpoint(DistributionUrl) == false)
 		{
 			UE_LOG(LogIas, Error, TEXT("Failed to resolve CDN endpoints from distribution URL"));
 			BackendStatus.SetHttpEnabled(false);
