Enable AutoRTFM for Chaos module

diff --git a/Engine/Source/Runtime/Experimental/Chaos/Chaos.Build.cs b/Engine/Source/Runtime/Experimental/Chaos/Chaos.Build.cs
--- a/Engine/Source/Runtime/Experimental/Chaos/Chaos.Build.cs
+++ b/Engine/Source/Runtime/Experimental/Chaos/Chaos.Build.cs
@@ -44,6 +44,8 @@ namespace UnrealBuildTool.Rules
 			
 			UnsafeTypeCastWarningLevel = WarningLevel.Error;
 
+			bAllowAutoRTFMInstrumentation = true;
+
 			StaticAnalyzerDisabledCheckers.Add("cplusplus.NewDeleteLeaks"); // To be reevalulated, believed to be invalid warnings.
 			StaticAnalyzerDisabledCheckers.Add("core.UndefinedBinaryOperatorResult"); // Invalid warning in mass property calculation.
 		}
