Dataflow : Node rendering
- Node color remapping

diff --git a/Engine/Plugins/Experimental/Dataflow/Source/DataflowAssetTools/Private/Dataflow/CollectionRenderingPatternUtility.cpp b/Engine/Plugins/Experimental/Dataflow/Source/DataflowAssetTools/Private/Dataflow/CollectionRenderingPatternUtility.cpp
--- a/Engine/Plugins/Experimental/Dataflow/Source/DataflowAssetTools/Private/Dataflow/CollectionRenderingPatternUtility.cpp
+++ b/Engine/Plugins/Experimental/Dataflow/Source/DataflowAssetTools/Private/Dataflow/CollectionRenderingPatternUtility.cpp
@@ -80,7 +80,7 @@ namespace Dataflow
 					const FIndex3i ColorElementTri = ColorOverlay->GetTriangle(TriangleID);
 					for (int TriVertIndex = 0; TriVertIndex < 3; ++TriVertIndex)
 					{
-						FVector4f Color(Colors[Tri[TriVertIndex]]); Color.W = 1.0f;
+						FVector4f Color(Colors[Remapping[Tri[TriVertIndex]]]); Color.W = 1.0f;
 						ColorOverlay->SetElement(ColorElementTri[TriVertIndex], Color);
 					}
 				};
diff --git a/Engine/Source/Runtime/Experimental/GeometryCollectionEngine/Private/GeometryCollection/GeometryCollectionEngineConversion.cpp b/Engine/Source/Runtime/Experimental/GeometryCollectionEngine/Private/GeometryCollection/GeometryCollectionEngineConversion.cpp
--- a/Engine/Source/Runtime/Experimental/GeometryCollectionEngine/Private/GeometryCollection/GeometryCollectionEngineConversion.cpp
+++ b/Engine/Source/Runtime/Experimental/GeometryCollectionEngine/Private/GeometryCollection/GeometryCollectionEngineConversion.cpp
@@ -1339,8 +1339,11 @@ bool FGeometryCollectionEngineConversion::AppendSkeletalMesh(const USkeletalMesh
 								if (ColorOverlay->IsElement(SourceInstanceID))
 								{
 									FVector4f InstColor = ColorOverlay->GetElement(SourceInstanceID);
-									UE::Geometry::LinearColors::SRGBToLinear(InstColor);
-									//Color[TargetVertexIndex] = FLinearColor(InstColor);
+									if (Converter.bTransformVertexColorsLinearToSRGB)
+									{
+										UE::Geometry::LinearColors::SRGBToLinear(InstColor);
+									}
+									Color[TargetVertexIndex] = FLinearColor(InstColor);
 								}
 							}
 
