GeometryCollection : SkeletalMeshConversion
- Remove usage of DynamicMesh

diff --git a/Engine/Source/Runtime/Experimental/GeometryCollectionEngine/Private/GeometryCollection/GeometryCollectionEngineConversion.cpp b/Engine/Source/Runtime/Experimental/GeometryCollectionEngine/Private/GeometryCollection/GeometryCollectionEngineConversion.cpp
--- a/Engine/Source/Runtime/Experimental/GeometryCollectionEngine/Private/GeometryCollection/GeometryCollectionEngineConversion.cpp
+++ b/Engine/Source/Runtime/Experimental/GeometryCollectionEngine/Private/GeometryCollection/GeometryCollectionEngineConversion.cpp
@@ -7,8 +7,6 @@
 #include "Async/ParallelFor.h"
 #include "Components/SkeletalMeshComponent.h"
 #include "Components/StaticMeshComponent.h"
-#include "DynamicMesh/DynamicMesh3.h"
-#include "DynamicMesh/DynamicMeshAttributeSet.h"
 #include "Engine/Selection.h"
 #include "Engine/SkeletalMesh.h"
 #include "Engine/SkinnedAssetCommon.h"
@@ -31,7 +29,6 @@
 #include "Materials/Material.h"
 #include "MeshDescription.h"
 #include "MeshDescriptionBuilder.h"
-#include "MeshDescriptionToDynamicMesh.h"
 #include "Misc/ScopedSlowTask.h"
 #include "Physics/Experimental/ChaosInterfaceUtils.h"
 #include "PhysicsEngine/BodySetup.h"
@@ -1177,10 +1174,6 @@ bool FGeometryCollectionEngineConversion::AppendSkeletalMesh(const USkeletalMesh
 	TManagedArray<int32>& MaterialIndex = GeometryCollection.ModifyAttribute<int32>("MaterialIndex", FGeometryCollection::FacesGroup);
 	TManagedArray<int32>& MaterialID = GeometryCollection.ModifyAttribute<int32>("MaterialID", FGeometryCollection::FacesGroup);
 
-	FMeshDescriptionToDynamicMesh Converter;
-	UE::Geometry::FDynamicMesh3 DynamicMesh;
-	Converter.Convert(&MeshDescription, DynamicMesh);
-
 	//
 	// Convert the transform hierarchy
 	//
@@ -1221,8 +1214,8 @@ bool FGeometryCollectionEngineConversion::AppendSkeletalMesh(const USkeletalMesh
 	TArray<int32> SourceVertexToComponentMap;				// Map from mesh vertex index to target vertex index
 	TArray<TArray<FIntVector>> ComponentsSourceIndices;		// Mesh triangle indices of each component.
 	TArray<TArray<FIntVector2>> SourceToTargetTriangleMap;	// Mesh triangle index of each triangle in the component. 
-	GeometryCollectionEngineUtility::GenerateConnectedComponents(InSkeletalMesh, ComponentsSourceIndices, 
-		SourceToTargetTriangleMap, SourceVertexToComponentMap,TriangleCount, VertexCount);
+	GeometryCollectionEngineUtility::GenerateConnectedComponents(InSkeletalMesh, ComponentsSourceIndices,
+		SourceToTargetTriangleMap, SourceVertexToComponentMap, TriangleCount, VertexCount);
 
 
 	//
@@ -1282,19 +1275,14 @@ bool FGeometryCollectionEngineConversion::AppendSkeletalMesh(const USkeletalMesh
 	TArray<FTransform> ComponentTransform;
 	GeometryCollectionAlgo::GlobalMatrices(LocalSpaceTransform, Parent, ComponentTransform);
 
-
+	bool bHasVertexColors = false;
 	FLinearColor DefaultColor = FLinearColor::White;
-	TVertexInstanceAttributesConstRef<FVector4f> InstanceColors;
-	UE::Geometry::FDynamicMeshColorOverlay* ColorOverlay = nullptr;
-	if (UE::Geometry::FDynamicMeshAttributeSet* Attributes = DynamicMesh.Attributes())
+	FSkeletalMeshConstAttributes SkeletalMeshConstAttributes(MeshDescription);
+	TVertexInstanceAttributesConstRef<FVector4f> InstanceColors = SkeletalMeshConstAttributes.GetVertexInstanceColors();
+	if (InstanceColors.IsValid())
 	{
-		FSkeletalMeshConstAttributes SkeletalMeshConstAttributes(MeshDescription);
-		InstanceColors = SkeletalMeshConstAttributes.GetVertexInstanceColors();
-		if (InstanceColors.IsValid())
-		{
-			ColorOverlay = Attributes->PrimaryColors();
-			DefaultColor = InstanceColors.GetDefaultValue();
-		}
+		bHasVertexColors = true;
+		DefaultColor = InstanceColors.GetDefaultValue();
 	}
 
 	// @todo(GeometryCollectionConversion) : Add support for UV's, Normals  
@@ -1322,29 +1310,20 @@ bool FGeometryCollectionEngineConversion::AppendSkeletalMesh(const USkeletalMesh
 							//  @todo(Skin Wieghts) : Currently its a static pose in component space. 
 							BoneMap[TargetVertexIndex] = ComponentToTransformGroupIndex[ComponentIndex];
 
-							FVector SourceVertex = (FVector)DynamicMesh.GetVertex(SourceVertexIndex);
+							FVector SourceVertex = (FVector)MeshDescription.GetVertexPosition(SourceVertexIndex);
 							FMatrix M = ComponentTransform[BoneMap[TargetVertexIndex]].ToInverseMatrixWithScale();
 							Vertex[TargetVertexIndex] = (FVector4f)M.TransformPosition(SourceVertex);
 
 							TargetVertexVisited[TargetVertexIndex] = true;
 
-							//colors
+							// transfer color
 							Color[TargetVertexIndex] = DefaultColor;
-							if (ColorOverlay)
+							if (bHasVertexColors)
 							{
-								// k : 0 -> 2
-								// SourceTriangleIndex is a TriangleID from the DynamicMesh
-								// TargetVertexIndex is the index for the vertex in the GeometryCollection
-								int32 SourceInstanceID = ColorOverlay->GetTriangle(SourceTriangleIndex)[k];
-								if (ColorOverlay->IsElement(SourceInstanceID))
-								{
-									FVector4f InstColor = ColorOverlay->GetElement(SourceInstanceID);
-									if (Converter.bTransformVertexColorsLinearToSRGB)
-									{
-										UE::Geometry::LinearColors::SRGBToLinear(InstColor);
-									}
-									Color[TargetVertexIndex] = FLinearColor(InstColor);
-								}
+								TArrayView<const FVertexInstanceID> SourceInstanceTri = MeshDescription.GetTriangleVertexInstances(SourceTriangleIndex);
+								FVector4f InstColor = InstanceColors[SourceInstanceTri[k]];
+								UE::Geometry::LinearColors::SRGBToLinear(InstColor);
+								Color[TargetVertexIndex] = FLinearColor(InstColor);
 							}
 
 							// @todo(GeometryCollectionConversion) : Add support for UV's, Normals  
@@ -1361,7 +1340,7 @@ bool FGeometryCollectionEngineConversion::AppendSkeletalMesh(const USkeletalMesh
 			}
 		}
 	}
-	
+
 
 	// Geometry Group
 	TArray<int32> GeometryIndices;
diff --git a/Engine/Source/Runtime/Experimental/GeometryCollectionEngine/Private/GeometryCollection/GeometryCollectionEngineUtility.cpp b/Engine/Source/Runtime/Experimental/GeometryCollectionEngine/Private/GeometryCollection/GeometryCollectionEngineUtility.cpp
--- a/Engine/Source/Runtime/Experimental/GeometryCollectionEngine/Private/GeometryCollection/GeometryCollectionEngineUtility.cpp
+++ b/Engine/Source/Runtime/Experimental/GeometryCollectionEngine/Private/GeometryCollection/GeometryCollectionEngineUtility.cpp
@@ -3,7 +3,6 @@
 #include "GeometryCollection/GeometryCollectionEngineUtility.h"
 
 #include "Chaos/Utilities.h"
-#include "DynamicMesh/DynamicMesh3.h"
 #include "Engine/SkeletalMesh.h"
 #include "GeometryCollection/GeometryCollection.h"
 #include "GeometryCollection/GeometryCollectionCache.h"
@@ -14,6 +13,7 @@
 #include "Modules/ModuleInterface.h"
 #include "Modules/ModuleManager.h"
 
+
 #if WITH_EDITOR
 #include "MeshUtilities.h"
 #endif
@@ -473,14 +473,11 @@ void GeometryCollectionEngineUtility::GenerateConnectedComponents(const USkeleta
 	constexpr int32 LODIndex = 0;
 	if (InSkeletalMesh->HasMeshDescription(LODIndex))
 	{
-		FMeshDescription SourceMesh;
-		UE::Geometry::FDynamicMesh3 DynamicMesh;
-		FMeshDescriptionToDynamicMesh Converter;
-		InSkeletalMesh->CloneMeshDescription(LODIndex, SourceMesh);
-		Converter.Convert(&SourceMesh, DynamicMesh);
+		FMeshDescription MeshDescription;
+		InSkeletalMesh->CloneMeshDescription(LODIndex, MeshDescription);
 
 		TArray<int32> VertexComponentID;
-		VertexComponentID.Init(INDEX_NONE, DynamicMesh.MaxVertexID());
+		VertexComponentID.Init(INDEX_NONE, MeshDescription.Vertices().Num());
 
 		int32 MaxComponentID = INDEX_NONE;
 		for (int VisitedVertexIdx = 0; VisitedVertexIdx < VertexComponentID.Num(); VisitedVertexIdx++)
@@ -494,10 +491,10 @@ void GeometryCollectionEngineUtility::GenerateConnectedComponents(const USkeleta
 				while (!Neighbors.IsEmpty())
 				{
 					int32 CurrentVertex = Neighbors.Pop();
-					for (int32 EdgeID : DynamicMesh.VtxEdgesItr(CurrentVertex))
+					for (int32 EdgeID : MeshDescription.GetVertexConnectedEdgeIDs(CurrentVertex))
 					{
-						const FDynamicMesh3::FEdge& Edge = DynamicMesh.GetEdge(EdgeID);
-						int NewVertex = Edge.Vert[0] == CurrentVertex ? Edge.Vert[1] : Edge.Vert[0];
+						const TArrayView<const FVertexID>& Edge = MeshDescription.GetEdgeVertices(EdgeID);
+						int NewVertex = Edge[0] == CurrentVertex ? Edge[1] : Edge[0];
 						if (VertexComponentID[NewVertex] == INDEX_NONE)
 						{
 							Neighbors.Push(NewVertex);
@@ -512,9 +509,9 @@ void GeometryCollectionEngineUtility::GenerateConnectedComponents(const USkeleta
 		TriangleCount = 0;
 		SourceTriangleIndex.Init(TArray< FIntVector2 >(), MaxComponentID + 1);
 		SourceTriangleVertices.Init(TArray<FIntVector>(), MaxComponentID + 1);
-		for (int32 TriangleID : DynamicMesh.TriangleIndicesItr())
+		for (FTriangleID TriangleID : MeshDescription.Triangles().GetElementIDs())
 		{
-			const FIntVector3& Triangle = DynamicMesh.GetTriangle(TriangleID);
+			const TArrayView<const FVertexID>& Triangle = MeshDescription.GetTriangleVertices(TriangleID);
 			int32 ComponentID = VertexComponentID[Triangle[0]];
 			for (int k = 1; k < 3; k++)
 			{
@@ -523,14 +520,14 @@ void GeometryCollectionEngineUtility::GenerateConnectedComponents(const USkeleta
 					return;
 				}
 			}
-			SourceTriangleVertices[ComponentID].Add(Triangle);
+			SourceTriangleVertices[ComponentID].Add(FIntVector3(Triangle[0], Triangle[1], Triangle[2]));
 			SourceTriangleIndex[ComponentID].Add(FIntVector2(TriangleID, INDEX_NONE));
 			TriangleCount++;
 		}
 
 		// Build Remapping
 		int32 CurrentRemapIndex = 0;
-		VertexComponentMap.Init(INDEX_NONE, DynamicMesh.MaxVertexID());
+		VertexComponentMap.Init(INDEX_NONE, MeshDescription.Vertices().Num());
 		for (TArray<FIntVector>& Component : SourceTriangleVertices)
 		{
 			for (FIntVector& Triangle : Component)
