Submitted in error. We'll keep it, but with a different name.

#rnx

diff --git a/Engine/Source/Runtime/Experimental/IoStore/OnDemand/Private/Tool/Curl.cpp b/Engine/Source/Runtime/Experimental/IoStore/OnDemand/Private/Tool/Curl.cpp
deleted file mode 100644
diff --git a/Engine/Source/Runtime/Experimental/IoStore/OnDemand/Private/Tool/Purl.cpp b/Engine/Source/Runtime/Experimental/IoStore/OnDemand/Private/Tool/Purl.cpp
new file mode 100644
--- /dev/null
+++ b/Engine/Source/Runtime/Experimental/IoStore/OnDemand/Private/Tool/Purl.cpp
@@ -0,0 +1,106 @@
+// Copyright Epic Games, Inc. All Rights Reserved.
+
+#if defined(UE_WITH_IAS_TOOL)
+
+#include "Command.h"
+
+#include <Containers/UnrealString.h>
+#include <HAL/PlatformProcess.h>
+#include <IO/IoBuffer.h>
+#include <IO/Http/Client.h>
+#include <Misc/ScopeExit.h>
+
+#if PLATFORM_WINDOWS
+#	include <Windows/AllowWindowsPlatformTypes.h>
+#		include <winsock2.h>
+#		include <ws2tcpip.h>
+#	include <Windows/HideWindowsPlatformTypes.h>
+#endif
+
+namespace UE::IoStore::Tool
+{
+
+////////////////////////////////////////////////////////////////////////////////
+static int32 PurlCommandEntry(const FContext& Context)
+{
+#if PLATFORM_WINDOWS
+	WSADATA WsaData;
+	if (WSAStartup(MAKEWORD(2, 2), &WsaData) == 0x0a9e0493)
+		return 1;
+	ON_SCOPE_EXIT { WSACleanup(); };
+#endif
+
+	using namespace UE::IoStore::HTTP;
+
+	FStringView Url = Context.Get<FStringView>(TEXT("Url"));
+	auto AnsiUrl = StringCast<ANSICHAR>(Url.GetData(), Url.Len());
+
+	FString Method(Context.Get<FStringView>(TEXT("-Method"), TEXT("GET")));
+	Method = Method.ToUpper();
+	auto AnsiMethod = StringCast<ANSICHAR>(*Method);
+
+	FEventLoop Loop;
+
+	FEventLoop::FRequestParams RequestParams;
+	if (Context.Get<bool>(TEXT("-Redirect")))
+	{
+		RequestParams.bAutoRedirect = true;
+	}
+	FRequest Request = Loop.Request(AnsiMethod, AnsiUrl, &RequestParams);
+
+	FIoBuffer Dest;
+	Loop.Send(MoveTemp(Request), [&Dest] (const FTicketStatus& Status)
+	{
+		if (Status.GetId() == FTicketStatus::EId::Response)
+		{
+			FResponse& Response = Status.GetResponse();
+			FAnsiStringView Message = Response.GetStatusMessage();
+			std::printf("%d %.*s\n",
+				Response.GetStatusCode(),
+				Message.Len(),
+				Message.GetData()
+			);
+			Response.ReadHeaders([] (FAnsiStringView Name, FAnsiStringView Value)
+			{
+				std::printf("%.*s: %.*s\n",
+					Name.Len(),
+					Name.GetData(),
+					Value.Len(),
+					Value.GetData()
+				);
+				return true;
+			});
+
+			Response.SetDestination(&Dest);
+			return;
+		}
+
+		if (Status.GetId() == FTicketStatus::EId::Error)
+		{
+			const char* Reason = Status.GetErrorReason();
+			std::printf("ERROR: %s\n", Reason);
+			return;
+		}
+	});
+	for (; Loop.Tick(-1); FPlatformProcess::SleepNoStats(0.1f));
+
+	std::printf("Data: %u bytes\n", uint32(Dest.GetSize()));
+
+	return 0;
+}
+
+////////////////////////////////////////////////////////////////////////////////
+static FCommand PurlCommand(
+	PurlCommandEntry,
+	TEXT("Purl"),
+	TEXT("Uses IoStore's HTTP client to download a URL"),
+	{
+		TArgument<FStringView>(TEXT("Url"), TEXT("Url to download")),
+		TArgument<FStringView>(TEXT("-Method"), TEXT("Request method")),
+		TArgument<bool>(TEXT("-Redirect"), TEXT("Follow 30x redirects")),
+	}
+);
+
+} // namespace UE::IoStore::Tool
+
+#endif // UE_WITH_IAS_TOOL
