Empty file so it's empty. Submitted in error

#rnx

diff --git a/Engine/Source/Runtime/Experimental/IoStore/HttpClient/IoStoreHttpClient.Build.cs b/Engine/Source/Runtime/Experimental/IoStore/HttpClient/IoStoreHttpClient.Build.cs
--- a/Engine/Source/Runtime/Experimental/IoStore/HttpClient/IoStoreHttpClient.Build.cs
+++ b/Engine/Source/Runtime/Experimental/IoStore/HttpClient/IoStoreHttpClient.Build.cs
@@ -8,13 +8,6 @@ public class IoStoreHttpClient : ModuleRules
 	{
 		PublicDependencyModuleNames.Add("Core");
 		PublicDependencyModuleNames.Add("TraceLog");
-		PrivateDependencyModuleNames.AddRange(
-			new string[] {
-				"OpenSSL",
-				"SSL",
-			}
-		);
-
 		UnsafeTypeCastWarningLevel = WarningLevel.Error; 
 	}
 }
diff --git a/Engine/Source/Runtime/Experimental/IoStore/HttpClient/Private/Peer.inl b/Engine/Source/Runtime/Experimental/IoStore/HttpClient/Private/Peer.inl
--- a/Engine/Source/Runtime/Experimental/IoStore/HttpClient/Private/Peer.inl
+++ b/Engine/Source/Runtime/Experimental/IoStore/HttpClient/Private/Peer.inl
@@ -1,438 +0,0 @@
-// Copyright Epic Games, Inc. All Rights Reserved.
-
-#pragma once
-
-namespace UE::IoStore::HTTP
-{
-
-////////////////////////////////////////////////////////////////////////////////
-class FSsl
-{
-public:
-	FSsl()
-	{
-		static_assert(OPENSSL_VERSION_NUMBER >= 0x10100000L);
-
-		// CRYPTO_set_mem_functions(&FSsl::Malloc, &FSsl::Realloc, &FSsl::Free);
-
-		uint64 InitFlags = 0; // OPENSSL_INIT_NO_ATEXIT ?
-		const OPENSSL_INIT_SETTINGS* InitSettings = nullptr;
-		OPENSSL_init_ssl(InitFlags, InitSettings);
-
-		//SSL_load_error_strings();
-		//ERR_load_BIO_strings();
-	}
-
-	~FSsl()
-	{
-		//ERR_remove_state(0);
-		//ENGINE_cleanup();
-		//CONF_modules_unload(1);
-		//ERR_free_strings();
-		//EVP_cleanup();
-		//sk_SSL_COMP_free(SSL_COMP_get_compression_methods());
-		//CRYPTO_cleanup_all_ex_data();
-
-		// if OPENSSL_INIT_NO_ATEXIT ?
-		//	   OPENSSL_cleanup();
-	}
-
-	struct FCertDeleter { void operator () (X509* Cert) { X509_free(Cert); } };
-	using FCert = TUniquePtr<X509, FCertDeleter>;
-	/*
-	struct FCert
-	{
-		X509* Cert = nullptr;
-	};
-	*/
-
-	static FCert LoadCertFromPem(const uint8* Data, uint32 Size)
-	{
-		BIO* Bio = BIO_new_mem_buf(Data, Size);
-		X509* Certificate = PEM_read_bio_X509(Bio, nullptr, 0, nullptr);
-		return FCert(Certificate);
-	}
-
-private:
-	static void* Malloc(size_t Size, const char*, int)
-	{
-		return FMemory::Malloc(Size);
-	}
-
-	static void* Realloc(void* Addr, size_t Size, const char*, int)
-	{
-		return FMemory::Realloc(Addr, Size);
-	}
-
-	static void Free(void* Addr, const char*, int)
-	{
-		FMemory::Free(Addr);
-	}
-};
-
-
-
-////////////////////////////////////////////////////////////////////////////////
-class FSslContext
-{
-public:
-				FSslContext(const char* InHostName);
-				FSslContext(const char* InHostName, const FSsl::FCert& CertPin);
-				~FSslContext()				{ SSL_CTX_free(Context); }
-				operator SSL_CTX* () const	{ return Context; }
-	const char*	GetHostName() const			{ return HostName; }
-	bool		AddCert(const FSsl::FCert& Cert);
-
-private:
-	SSL_CTX*	Context;
-	const char*	HostName;
-};
-
-////////////////////////////////////////////////////////////////////////////////
-FSslContext::FSslContext(const char* InHostName)
-: HostName(InHostName)
-{
-	Context = SSL_CTX_new(TLS_client_method());
-
-	uint32 ProtoFlags = SSL_OP_NO_SSLv2 | SSL_OP_NO_SSLv3;
-	SSL_CTX_set_options(Context, ProtoFlags);
-}
-
-////////////////////////////////////////////////////////////////////////////////
-FSslContext::FSslContext(const char* InHostName, const FSsl::FCert& CertPin)
-: FSslContext(InHostName)
-{
-	if (!CertPin.IsValid())
-	{
-		const ISslCertificateManager& CertManager = FSslModule::Get().GetCertificateManager();
-		CertManager.AddCertificatesToSslContext(Context);
-		return;
-	}
-
-	AddCert(CertPin);
-}
-
-////////////////////////////////////////////////////////////////////////////////
-bool FSslContext::AddCert(const FSsl::FCert& Cert)
-{
-	check(Cert.IsValid());
-	X509_STORE* Store = SSL_CTX_get_cert_store(Context);
-	int32 Result = X509_STORE_add_cert(Store, Cert.Get());
-	return Result == 1;
-}
-
-
-
-////////////////////////////////////////////////////////////////////////////////
-class FPeer
-{
-public:
-				FPeer() = default;
-				FPeer(FSocket InSocket);
-	FOutcome	Send(const char* Data, int32 Size)	{ return Socket.Send(Data, Size); }
-	FOutcome	Recv(char* Out, int32 MaxSize)		{ return Socket.Recv(Out, MaxSize); }
-	bool		IsValid() const						{ return Socket.IsValid(); }
-
-private:
-	FSocket		Socket;
-
-
-
-	// TODO: !!!!!! remove this !!!!!
-public:
-	explicit operator const FSocket& () const { return Socket; }
-	explicit operator const FSocket* () const { return &Socket; }
-};
-
-////////////////////////////////////////////////////////////////////////////////
-FPeer::FPeer(FSocket InSocket)
-: Socket(MoveTemp(InSocket))
-{
-}
-
-
-
-////////////////////////////////////////////////////////////////////////////////
-class FTlsPeer
-	: public FPeer
-{
-public:
-				FTlsPeer()								= default;
-				~FTlsPeer();
-				FTlsPeer(FTlsPeer&& Rhs)				{ Move(MoveTemp(Rhs)); }
-				FTlsPeer& operator = (FTlsPeer&& Rhs)	{ return Move(MoveTemp(Rhs)); }
-				FTlsPeer(FSocket InSocket, const FSslContext* Context=nullptr);
-	FTlsPeer&	Move(FTlsPeer&& Rhs);
-	FOutcome	Handshake();
-	FOutcome	Send(const char* Data, int32 Size);
-	FOutcome	Recv(char* Out, int32 MaxSize);
-
-protected:
-	SSL*		Ssl = nullptr;
-
-private:
-	FOutcome GetOutcome(int32 SslResult, const char* Message="tls error") const
-	{
-		char buf[256];
-		TArray<uint32> iii;
-		int32 line;
-		const char* file;
-		while (uint32 i = ERR_get_error_line(&file, &line))
-		{
-			ERR_error_string(i, buf);
-			iii.Add(i);
-		}
-
-		int32 Error = SSL_get_error(Ssl, SslResult);
-		if (Error != SSL_ERROR_WANT_READ && Error != SSL_ERROR_WANT_WRITE)
-		{
-			return FOutcome::Error(Message, Error);
-		}
-		return FOutcome::Waiting();
-	}
-
-	int32 BioWrite(const char* Data, size_t Size, size_t* BytesWritten, BIO* Bio)
-	{
-		*BytesWritten = 0;
-		BIO_clear_retry_flags(Bio);
-
-		FOutcome Outcome = FPeer::Send(Data, Size);
-		if (Outcome.IsWaiting())
-		{
-			BIO_set_retry_write(Bio);
-			return 0;
-		}
-
-		if (Outcome.IsError())
-		{
-			return -1;
-		}
-
-		*BytesWritten = Outcome.GetResult();
-		return 1;
-	}
-
-	int32 BioRead(char* Data, size_t Size, size_t* BytesRead, BIO* Bio)
-	{
-		*BytesRead = 0;
-		BIO_clear_retry_flags(Bio);
-
-		FOutcome Outcome = FPeer::Recv(Data, Size);
-		if (Outcome.IsWaiting())
-		{
-			BIO_set_retry_read(Bio);
-			return 0;
-		}
-
-		if (Outcome.IsError())
-		{
-			return -1;
-		}
-
-		*BytesRead = Outcome.GetResult();
-		return 1;
-	}
-
-	long BioControl(int Cmd, long, void*, BIO*)
-	{
-		return (Cmd == BIO_CTRL_FLUSH) ? 1 : 0;
-	}
-};
-
-////////////////////////////////////////////////////////////////////////////////
-FTlsPeer::FTlsPeer(FSocket InSocket, const FSslContext* Context)
-: FPeer(MoveTemp(InSocket))
-{
-	if (Context == nullptr)
-	{
-		return;
-	}
-
-	static BIO_METHOD* BioMethod = nullptr;
-	if (BioMethod == nullptr)
-	{
-		int32 BioId = BIO_get_new_index() | BIO_TYPE_SOURCE_SINK;
-		BioMethod = BIO_meth_new(BioId, "IasBIO");
-
-#define METH_THUNK(x_, y_)\
-		x_(BioMethod, [] <typename... T> (BIO* b, T... t) { \
-			return (decltype(this)(BIO_get_data(b)))->y_(Forward<T>(t)..., b); \
-		})
-		METH_THUNK(BIO_meth_set_write_ex,	BioWrite);
-		METH_THUNK(BIO_meth_set_read_ex,	BioRead);
-		METH_THUNK(BIO_meth_set_ctrl,		BioControl);
-#undef METH_THUNK
-	}
-
-	BIO* Bio = BIO_new(BioMethod);
-	BIO_set_data(Bio, this);
-
-	// SSL_MODE_ENABLE_PARTIAL_WRITE ??!!!
-
-	Ssl = SSL_new(*Context);
-	SSL_set_tlsext_host_name(Ssl, Context->GetHostName());
-	SSL_set_connect_state(Ssl);
-	SSL_set0_rbio(Ssl, Bio);
-	SSL_set0_wbio(Ssl, Bio);
-	BIO_up_ref(Bio);
-}
-
-////////////////////////////////////////////////////////////////////////////////
-FTlsPeer& FTlsPeer::Move(FTlsPeer&& Rhs)
-{
-	FPeer::operator = (MoveTemp(Rhs));
-
-	Swap(Ssl, Rhs.Ssl);
-
-	auto PatchBio = [] (FTlsPeer& Peer)
-	{
-		if (Peer.Ssl != nullptr)
-		{
-			BIO* Bio = SSL_get_rbio(Peer.Ssl);
-			check(Bio != nullptr);
-			BIO_set_data(Bio, &Peer);
-		}
-	};
-	PatchBio(*this);
-	PatchBio(Rhs);
-
-	return *this;
-}
-
-////////////////////////////////////////////////////////////////////////////////
-FTlsPeer::~FTlsPeer()
-{
-	if (Ssl != nullptr)
-	{
-		SSL_free(Ssl);
-	}
-}
-
-////////////////////////////////////////////////////////////////////////////////
-FOutcome FTlsPeer::Handshake()
-{
-	if (Ssl == nullptr)
-	{
-		return FOutcome::Ok();
-	}
-
-	int32 Result = SSL_do_handshake(Ssl);
-	if (Result == 0) return FOutcome::Error("unsuccessful tls handshake");
-	if (Result != 1) return GetOutcome(Result, "tls handshake error");
-
-	if (Result = SSL_get_verify_result(Ssl); Result != X509_V_OK)
-	{
-		return FOutcome::Error("x509 verification error", Result);
-	}
-
-	return FOutcome::Ok();
-}
-
-////////////////////////////////////////////////////////////////////////////////
-FOutcome FTlsPeer::Send(const char* Data, int32 Size)
-{
-	if (Ssl == nullptr)
-	{
-		return FPeer::Send(Data, Size);
-	}
-
-	int32 Result = SSL_write(Ssl, Data, Size);
-	return (Result > 0) ? FOutcome::Ok(Result) : GetOutcome(Result);
-}
-
-////////////////////////////////////////////////////////////////////////////////
-FOutcome FTlsPeer::Recv(char* Out, int32 MaxSize)
-{
-	if (Ssl == nullptr)
-	{
-		return FPeer::Recv(Out, MaxSize);
-	}
-
-	int32 Result = SSL_read(Ssl, Out, MaxSize);
-	return (Result > 0) ? FOutcome::Ok(Result) : GetOutcome(Result);
-}
-
-
-
-////////////////////////////////////////////////////////////////////////////////
-class FHttpPeer
-	: public FTlsPeer
-{
-public:
-				FHttpPeer() = default;
-				FHttpPeer(FSocket InSocket, FSslContext* Context=nullptr);
-	FOutcome	Handshake();
-
-private:
-	void		AssignProto();
-	int32		Proto = 0;
-};
-
-////////////////////////////////////////////////////////////////////////////////
-FHttpPeer::FHttpPeer(FSocket InSocket, FSslContext* Context)
-: FTlsPeer(MoveTemp(InSocket), Context)
-{
-	if (Ssl == nullptr)
-	{
-		return;
-	}
-
-	static const uint8 AlpnProtos[] =
-		"\x08" "http/1.1"
-	//	"\x02" "h2"
-		;
-	SSL_set_alpn_protos(Ssl, AlpnProtos, sizeof(AlpnProtos) - 1);
-}
-
-////////////////////////////////////////////////////////////////////////////////
-FOutcome FHttpPeer::Handshake()
-{
-	FOutcome Outcome = FTlsPeer::Handshake();
-	if (Outcome.IsOk())
-	{
-		AssignProto();
-	}
-
-	return Outcome;
-}
-
-////////////////////////////////////////////////////////////////////////////////
-void FHttpPeer::AssignProto()
-{
-	Proto = 1;
-
-	if (Ssl == nullptr)
-		return;
-
-	const char* AlpnProto;
-	uint32 AlpnProtoLen;
-	SSL_get0_alpn_selected(Ssl, &(const uint8*&)AlpnProto, &AlpnProtoLen);
-	if (AlpnProto == nullptr)
-	{
-		return;
-	}
-
-	FAnsiStringView Needle(AlpnProto, AlpnProtoLen);
-	FAnsiStringView Candidates[] = {
-		"http/1.1",
-	//	"h2"
-	};
-	for (int32 i = 0; i < UE_ARRAY_COUNT(Candidates); ++i)
-	{
-		const FAnsiStringView& Candidate = Candidates[i];
-		if (AlpnProtoLen != Candidate.Len())
-		{
-			continue;
-		}
-
-		if (Candidate != Needle)
-		{
-			continue;
-		}
-
-		Proto = i + 1;
-		break;
-	}
-}
-
-} // namespace UE::IoStore::HTTP
