Missing boilerplate.

#rnx

diff --git a/Engine/Source/Runtime/Experimental/IoStore/HttpClient/Private/Peer.inl b/Engine/Source/Runtime/Experimental/IoStore/HttpClient/Private/Peer.inl
--- a/Engine/Source/Runtime/Experimental/IoStore/HttpClient/Private/Peer.inl
+++ b/Engine/Source/Runtime/Experimental/IoStore/HttpClient/Private/Peer.inl
@@ -0,0 +1 @@
+// Copyright Epic Games, Inc. All Rights Reserved.
