Fixed potentially-uninitialized FVectors in FPBDCollisionConstraint::SetSolverResults.

diff --git a/Engine/Source/Runtime/Experimental/Chaos/Public/Chaos/Collision/ContactPoint.h b/Engine/Source/Runtime/Experimental/Chaos/Public/Chaos/Collision/ContactPoint.h
--- a/Engine/Source/Runtime/Experimental/Chaos/Public/Chaos/Collision/ContactPoint.h
+++ b/Engine/Source/Runtime/Experimental/Chaos/Public/Chaos/Collision/ContactPoint.h
@@ -169,6 +169,12 @@ namespace Chaos
 	class FSavedManifoldPoint
 	{
 	public:
+		explicit FSavedManifoldPoint(const FVec3f& InShapeContactPoint0, const FVec3f& InShapeContactPoint1, FRealSingle InInitialPhi)
+			: ShapeContactPoints{ InShapeContactPoint0, InShapeContactPoint1 }
+			, InitialPhi(InInitialPhi)
+		{
+		}
+
 		FVec3f ShapeContactPoints[2];			// Contact anchor points for friction
 		FRealSingle InitialPhi;					// Non-resolved initial penetration
 	};
diff --git a/Engine/Source/Runtime/Experimental/Chaos/Public/Chaos/Collision/PBDCollisionConstraint.h b/Engine/Source/Runtime/Experimental/Chaos/Public/Chaos/Collision/PBDCollisionConstraint.h
--- a/Engine/Source/Runtime/Experimental/Chaos/Public/Chaos/Collision/PBDCollisionConstraint.h
+++ b/Engine/Source/Runtime/Experimental/Chaos/Public/Chaos/Collision/PBDCollisionConstraint.h
@@ -771,17 +771,17 @@ namespace Chaos
 		{
 			const FManifoldPoint& ManifoldPoint = ManifoldPoints[ManifoldPointIndex];
 			FManifoldPointResult& ManifoldPointResult = ManifoldPointResults[ManifoldPointIndex];
-			FSavedManifoldPoint* SavedManifoldPoint = nullptr;
 
 			// Save contact data for friction
-			FVec3f Anchor0, Anchor1;
 			bool bInsideStaticFrictionCone = false;
 			if (StaticFrictionRatio >= FRealSingle(1.0f - UE_KINDA_SMALL_NUMBER))
 			{
 				// StaticFrictionRatio ~= 1: Static friction held - we keep the same contacts points as-is for use next frame
-				SavedManifoldPoint = &SavedManifoldPoints[SavedManifoldPoints.AddUninitialized()];
-				Anchor0 = ManifoldPoint.ShapeAnchorPoints[0];
-				Anchor1 = ManifoldPoint.ShapeAnchorPoints[1];
+				SavedManifoldPoints.Emplace(
+					ManifoldPoint.ShapeAnchorPoints[0],
+					ManifoldPoint.ShapeAnchorPoints[1],
+					ManifoldPoint.InitialPhi
+				);
 				bInsideStaticFrictionCone = true;
 			}
 			else if (StaticFrictionRatio < FRealSingle(UE_KINDA_SMALL_NUMBER))
@@ -793,18 +793,22 @@ namespace Chaos
 				const int32 SmallNumManifoldPoints = 8;
 				if (ManifoldPoints.Num() < SmallNumManifoldPoints)
 				{
-					SavedManifoldPoint = &SavedManifoldPoints[SavedManifoldPoints.AddUninitialized()];
-					Anchor0 = ManifoldPoint.ContactPoint.ShapeContactPoints[0];
-					Anchor1 = ManifoldPoint.ContactPoint.ShapeContactPoints[1];
+					SavedManifoldPoints.Emplace(
+						ManifoldPoint.ContactPoint.ShapeContactPoints[0],
+						ManifoldPoint.ContactPoint.ShapeContactPoints[1],
+						ManifoldPoint.InitialPhi
+					);
 				}
 			}
 			else
 			{
 				// 0 < StaticFrictionRatio < 1: We exceeded the friction cone. Slide the friction anchor 
 				// toward the last-detected contact position so that it sits at the edge of the friction cone.
-				SavedManifoldPoint = &SavedManifoldPoints[SavedManifoldPoints.AddUninitialized()];
-				Anchor0 = FVec3f::Lerp(ManifoldPoint.ContactPoint.ShapeContactPoints[0], ManifoldPoint.ShapeAnchorPoints[0], StaticFrictionRatio);
-				Anchor1 = FVec3f::Lerp(ManifoldPoint.ContactPoint.ShapeContactPoints[1], ManifoldPoint.ShapeAnchorPoints[1], StaticFrictionRatio);
+				SavedManifoldPoints.Emplace(
+					FVec3f::Lerp(ManifoldPoint.ContactPoint.ShapeContactPoints[0], ManifoldPoint.ShapeAnchorPoints[0], StaticFrictionRatio),
+					FVec3f::Lerp(ManifoldPoint.ContactPoint.ShapeContactPoints[1], ManifoldPoint.ShapeAnchorPoints[1], StaticFrictionRatio),
+					ManifoldPoint.InitialPhi
+				);
 			}
 
 			AccumulatedImpulse += NetImpulse + (NetPushOut / Dt);
@@ -814,13 +818,6 @@ namespace Chaos
 			ManifoldPointResult.bIsValid = true;
 			ManifoldPointResult.bInsideStaticFrictionCone = bInsideStaticFrictionCone;
 
-			if (SavedManifoldPoint != nullptr)
-			{
-				SavedManifoldPoint->ShapeContactPoints[0] = Anchor0;
-				SavedManifoldPoint->ShapeContactPoints[1] = Anchor1;
-				SavedManifoldPoint->InitialPhi = ManifoldPoint.InitialPhi;
-			}
-
 			MinInitialPhi = FMath::Min(MinInitialPhi, ManifoldPoint.InitialPhi);
 		}
 
