Fixed assert when ticking on-demand I/O store

diff --git a/Engine/Source/Runtime/Experimental/IoStore/OnDemand/Private/OnDemandIoStore.cpp b/Engine/Source/Runtime/Experimental/IoStore/OnDemand/Private/OnDemandIoStore.cpp
--- a/Engine/Source/Runtime/Experimental/IoStore/OnDemand/Private/OnDemandIoStore.cpp
+++ b/Engine/Source/Runtime/Experimental/IoStore/OnDemand/Private/OnDemandIoStore.cpp
@@ -175,7 +175,7 @@ FIoStatus FOnDemandIoStore::Initialize()
 		}
 
 		// Process mount requests synchronously at startup
-		while (Tick(MAX_int64));
+		TickLoop();
 	}
 
 	return EIoErrorCode::Ok;
@@ -206,32 +206,7 @@ void FOnDemandIoStore::Mount(FOnDemandMountArgs&& Args, FOnDemandMountCompleted&
 		});
 	}
 
-	ConditionallyStartTicking();
-}
-
-void FOnDemandIoStore::ConditionallyStartTicking()
-{
-	check(FPlatformProcess::SupportsMultithreading());
-
-	if (FPlatformProcess::SupportsMultithreading() && GIOThreadPool != nullptr)
-	{
-		bool bExpected = false;
-		if (bTicking.compare_exchange_strong(bExpected, true))
-		{
-			TickFuture = AsyncPool(
-				*GIOThreadPool,
-				[this]
-				{
-					while (Tick(MAX_int64));
-				},
-				nullptr,
-				EQueuedWorkPriority::Low);
-		}
-	}
-	else
-	{
-		while (Tick(MAX_int64));
-	}
+	TryEnterTickLoop();
 }
 
 FIoStatus FOnDemandIoStore::Unmount(FStringView MountId)
@@ -283,12 +258,44 @@ FOnDemandChunkInfo FOnDemandIoStore::GetChunkInfo(const FIoChunkId& ChunkId)
 	return FOnDemandChunkInfo();
 }
 
-bool FOnDemandIoStore::Tick(int64 MaxCycles)
+void FOnDemandIoStore::TryEnterTickLoop()
+{
+	bool bExpected = false;
+	if (bTicking.compare_exchange_strong(bExpected, true))
+	{
+		if (FPlatformProcess::SupportsMultithreading() && GIOThreadPool != nullptr)
+		{
+			TickFuture = AsyncPool(*GIOThreadPool, [this] { TickLoop(); }, nullptr, EQueuedWorkPriority::Low);
+		}
+		else
+		{
+			TickLoop();
+		}
+	}
+}
+
+void FOnDemandIoStore::TickLoop()
 {
-	check(!bTicking);
-	bTicking = true;
-	ON_SCOPE_EXIT{ bTicking = false; };
+	ON_SCOPE_EXIT { UE_LOG(LogIoStoreOnDemand, Verbose, TEXT("Exiting I/O store tick loop")); };
+
+	UE_LOG(LogIoStoreOnDemand, Verbose, TEXT("Entering I/O store tick loop"));
+	for (;;)
+	{
+		{
+			UE::TUniqueLock Lock(MountRequestMutex);
+			if (MountRequests.IsEmpty())
+			{
+				bTicking = false;
+				break;
+			}
+		}
 
+		Tick();
+	}
+}
+
+void FOnDemandIoStore::Tick()
+{
 	TArray<FSharedMountRequest> Requests;
 	{
 		UE::TUniqueLock Lock(MountRequestMutex);
@@ -300,7 +307,7 @@ bool FOnDemandIoStore::Tick(int64 MaxCycles)
 
 	if (Requests.IsEmpty())
 	{
-		return false;
+		return;
 	}
 
 	for (FSharedMountRequest& Request : Requests)
@@ -354,8 +361,6 @@ bool FOnDemandIoStore::Tick(int64 MaxCycles)
 		FOnDemandMountCompleted OnCompleted = MoveTemp(Request->OnCompleted);
 		OnCompleted(MountStatus);
 	}
-
-	return true;
 }
 
 FIoStatus FOnDemandIoStore::ProcessMountRequest(FMountRequest& MountRequest)
diff --git a/Engine/Source/Runtime/Experimental/IoStore/OnDemand/Private/OnDemandIoStore.h b/Engine/Source/Runtime/Experimental/IoStore/OnDemand/Private/OnDemandIoStore.h
--- a/Engine/Source/Runtime/Experimental/IoStore/OnDemand/Private/OnDemandIoStore.h
+++ b/Engine/Source/Runtime/Experimental/IoStore/OnDemand/Private/OnDemandIoStore.h
@@ -133,7 +133,9 @@ public:
 	FOnDemandChunkInfo		GetChunkInfo(const FIoChunkId& ChunkId);
 
 private:
-	bool					Tick(int64 MaxCycles);
+	void					TryEnterTickLoop();
+	void					TickLoop();
+	void					Tick();
 	FIoStatus				ProcessMountRequest(FMountRequest& MountRequest);
 	void					ConditionallyStartTicking();
 	void					OnEncryptionKeyAdded(const FGuid& Id, const FAES::FAESKey& Key);
