Preliminary IoStoreOnDemand TLS support

#rnx

diff --git a/Engine/Source/Runtime/Experimental/IoStore/HttpClient/Internal/IO/Http/Client.h b/Engine/Source/Runtime/Experimental/IoStore/HttpClient/Internal/IO/Http/Client.h
--- a/Engine/Source/Runtime/Experimental/IoStore/HttpClient/Internal/IO/Http/Client.h
+++ b/Engine/Source/Runtime/Experimental/IoStore/HttpClient/Internal/IO/Http/Client.h
@@ -4,6 +4,7 @@
 
 #if !defined(NO_UE_INCLUDES)
 #include <Containers/StringView.h>
+#include <Memory/MemoryView.h>
 #endif
 
 #if !defined(IAS_HTTP_WITH_PERF)
@@ -46,7 +47,8 @@ enum class EStatusCodeClass
 };
 
 ////////////////////////////////////////////////////////////////////////////////
-using	FTicket = uint64;
+using	FTicket		= uint64;
+using	FPemCert	= FMemoryView;
 struct	FActivity;
 
 ////////////////////////////////////////////////////////////////////////////////
@@ -57,10 +59,12 @@ public:
 	{
 		int32				SetHostFromUrl(FAnsiStringView Url);
 		FAnsiStringView		HostName;
+		FPemCert			VerifyCert;
 		int32				SendBufSize = -1;
 		int32				RecvBufSize = -1;
 		uint32				Port = 0;
 		uint16				ConnectionCount = 1;
+		bool				bUseTls = false;
 		/*
 		enum class ProxyType { Http, Socks4 };
 		Proxy = { ip, port, type }
@@ -187,8 +191,9 @@ public:
 
 	struct FRequestParams
 	{
-		uint32	BufferSize	= 256;
-		bool	bAutoRedirect = false;
+		FPemCert			VerifyCert;
+		uint32				BufferSize		= 256;
+		bool				bAutoRedirect	= false;
 	};
 
 	template <typename... T> [[nodiscard]] FRequest Get(T&&... t)  { return Request("GET",  Forward<T&&>(t)...); }
diff --git a/Engine/Source/Runtime/Experimental/IoStore/HttpClient/IoStoreHttpClient.Build.cs b/Engine/Source/Runtime/Experimental/IoStore/HttpClient/IoStoreHttpClient.Build.cs
--- a/Engine/Source/Runtime/Experimental/IoStore/HttpClient/IoStoreHttpClient.Build.cs
+++ b/Engine/Source/Runtime/Experimental/IoStore/HttpClient/IoStoreHttpClient.Build.cs
@@ -8,6 +8,13 @@ public class IoStoreHttpClient : ModuleRules
 	{
 		PublicDependencyModuleNames.Add("Core");
 		PublicDependencyModuleNames.Add("TraceLog");
+		PrivateDependencyModuleNames.AddRange(
+			new string[] {
+				"OpenSSL",
+				"SSL",
+			}
+		);
+
 		UnsafeTypeCastWarningLevel = WarningLevel.Error; 
 	}
 }
diff --git a/Engine/Source/Runtime/Experimental/IoStore/HttpClient/Private/ConnectionPool.inl b/Engine/Source/Runtime/Experimental/IoStore/HttpClient/Private/ConnectionPool.inl
--- a/Engine/Source/Runtime/Experimental/IoStore/HttpClient/Private/ConnectionPool.inl
+++ b/Engine/Source/Runtime/Experimental/IoStore/HttpClient/Private/ConnectionPool.inl
@@ -19,20 +19,25 @@ public:
 		const ANSICHAR* HostName;
 		uint32			Port = 0;
 		uint32			MaxConnections = 1;
+		FPemCert		VerifyCert;
 	};
 
 					FHost(const FParams& Params);
 	void			SetBufferSize(EDirection Dir, int32 Size);
 	int32			GetBufferSize(EDirection Dir) const;
+	FHttpPeer		Connect();
 	FOutcome		Connect(FSocket& Socket);
 	int32			IsResolved() const;
 	FOutcome		ResolveHostName();
+	bool			WithTls() const				{ return GetVerifyCert().GetData() != nullptr; }
+	FPemCert		GetVerifyCert() const		{ return VerifyCert; }
 	uint32			GetMaxConnections() const	{ return MaxConnections; }
 	uint32			GetIpAddress() const		{ return IpAddresses[0]; }
 	FAnsiStringView	GetHostName() const			{ return HostName; }
 	uint32			GetPort() const				{ return Port; }
 
 private:
+	FPemCert		VerifyCert;
 	const ANSICHAR*	HostName;
 	uint32			IpAddresses[4] = {};
 	int16			SendBufKb = -1;
@@ -43,7 +48,8 @@ private:
 
 ////////////////////////////////////////////////////////////////////////////////
 FHost::FHost(const FParams& Params)
-: HostName(Params.HostName)
+: VerifyCert(Params.VerifyCert)
+, HostName(Params.HostName)
 , Port(uint16(Params.Port))
 , MaxConnections(uint8(Params.MaxConnections))
 {
@@ -51,7 +57,7 @@ FHost::FHost(const FParams& Params)
 
 	if (Port == 0)
 	{
-		Port = 80;
+		Port = (VerifyCert.GetData() == nullptr) ? 80 : 443;
 	}
 }
 
@@ -137,6 +143,32 @@ int32 FHost::IsResolved() const
 	}
 }
 
+////////////////////////////////////////////////////////////////////////////////
+FHttpPeer FHost::Connect()
+{
+	FSocket Socket;
+	if (FOutcome Outcome = Connect(Socket); Outcome.IsError())
+	{
+		return FHttpPeer();
+		// return FOutcome::Error("Failed to create socket");
+	}
+
+	if (VerifyCert.GetData() == nullptr)
+	{
+		return FHttpPeer(MoveTemp(Socket));
+	}
+
+	FSslContext SslContext(HostName, VerifyCert);
+	FHttpPeer Peer(MoveTemp(Socket), &SslContext);
+
+	FOutcome Outcome = FOutcome::None();
+	while ((Outcome = Peer.Handshake()).IsWaiting())
+		/*pass*/;
+	check(Outcome.IsOk());
+
+	return Peer;
+}
+
 ////////////////////////////////////////////////////////////////////////////////
 FOutcome FHost::Connect(FSocket& Socket)
 {
@@ -215,6 +247,7 @@ int32 FConnectionPool::FParams::SetHostFromUrl(FAnsiStringView Url)
 	}
 
 	HostName = Offsets.HostName.Get(Url);
+	bUseTls = (Offsets.SchemeLength == 5);
 
 	if (Offsets.Port)
 	{
@@ -241,11 +274,20 @@ FConnectionPool::FConnectionPool(const FParams& Params)
 	memcpy(HostDest, Params.HostName.GetData(), HostNameLen);
 	HostDest[HostNameLen] = '\0';
 
+	// Verify
+	FPemCert VerifyCert;
+	if (Params.bUseTls)
+	{
+		VerifyCert = Params.VerifyCert;
+		VerifyCert = (VerifyCert.GetData() != nullptr) ? VerifyCert : FPemCert("", 0);
+	}
+
 	// Init internal object
 	new (Internal) FHost({
 		.HostName		= HostDest,
 		.Port			= Params.Port,
 		.MaxConnections	= Params.ConnectionCount,
+		.VerifyCert		= VerifyCert,
 	});
 	Internal->SetBufferSize(FHost::EDirection::Send, Params.SendBufSize);
 	Internal->SetBufferSize(FHost::EDirection::Recv, Params.RecvBufSize);
diff --git a/Engine/Source/Runtime/Experimental/IoStore/HttpClient/Private/Http.cpp b/Engine/Source/Runtime/Experimental/IoStore/HttpClient/Private/Http.cpp
--- a/Engine/Source/Runtime/Experimental/IoStore/HttpClient/Private/Http.cpp
+++ b/Engine/Source/Runtime/Experimental/IoStore/HttpClient/Private/Http.cpp
@@ -13,16 +13,23 @@
 #include <Misc/ScopeLock.h>
 #include <Misc/StringBuilder.h>
 #include <ProfilingDebugging/CpuProfilerTrace.h>
+#include <Ssl.h>
+#include <Interfaces/ISslCertificateManager.h>
 #include <Tasks/Task.h>
 #include <Trace/Trace.h>
 #endif
 
+#include <openssl/engine.h>
+#include <openssl/err.h>
+#include <openssl/ssl.h>
+
 #include <atomic>
 
 #include "LatencyInjector.h"
 
 #include "Misc.inl"
 #include "Socket.inl"
+#include "Peer.inl"
 #include "Socks.inl"
 #include "ConnectionPool.inl"
 #include "Activity.inl"
diff --git a/Engine/Source/Runtime/Experimental/IoStore/HttpClient/Private/Loop.inl b/Engine/Source/Runtime/Experimental/IoStore/HttpClient/Private/Loop.inl
--- a/Engine/Source/Runtime/Experimental/IoStore/HttpClient/Private/Loop.inl
+++ b/Engine/Source/Runtime/Experimental/IoStore/HttpClient/Private/Loop.inl
@@ -7,7 +7,7 @@ namespace UE::IoStore::HTTP
 
 // {{{1 event-loop-int .........................................................
 
-using FPeerType = FSocket;
+using FPeerType = FHttpPeer;
 
 ////////////////////////////////////////////////////////////////////////////////
 static FOutcome DoSend(FActivity* Activity, FPeerType& Peer)
@@ -349,11 +349,11 @@ static FOutcome DoRecv(FActivity* Activity, FPeerType& Peer, int32& MaxRecvSize)
 }
 
 ////////////////////////////////////////////////////////////////////////////////
-static int32 DoRecvDone(FActivity* Activity)
+static void DoRecvDone(FActivity* Activity)
 {
 	if (!FLatencyInjector::HasExpired(Activity->StateParam))
 	{
-		return 1;
+		return;
 	}
 
 	// Notify the user we've received everything
@@ -361,7 +361,6 @@ static int32 DoRecvDone(FActivity* Activity)
 	Activity->Sink(SinkArg);
 
 	Activity_ChangeState(Activity, FActivity::EState::Completed);
-	return 0;
 }
 
 ////////////////////////////////////////////////////////////////////////////////
@@ -531,10 +530,10 @@ class FPeerGroup
 public:
 						FPeerGroup() = default;
 						~FPeerGroup();
-	bool				operator == (FPeerType* Rhs) const { return &Peer == Rhs; }
+	bool				operator == (FSocket* Rhs) const { return (const FSocket*)Peer == Rhs; }
 	void				Unwait()			{ check(bWaiting); bWaiting = false; }
-	FPeerType::FWaiter	GetWaiter() const;
-	bool 				Tick(FTickState& State);
+	FSocket::FWaiter	GetWaiter() const;
+	bool				Tick(FTickState& State);
 	void				TickSend(FTickState& State, FHost& Host);
 	void				Fail(FTickState& State, const char* Reason);
 
@@ -559,16 +558,18 @@ FPeerGroup::~FPeerGroup()
 }
 
 ////////////////////////////////////////////////////////////////////////////////
-FPeerType::FWaiter FPeerGroup::GetWaiter() const
+FSocket::FWaiter FPeerGroup::GetWaiter() const
 {
 	if (!bWaiting)
 	{
-		return FPeerType::FWaiter();
+		return FSocket::FWaiter();
 	}
 
-	using EWhat = FPeerType::FWaiter::EWhat;
+	const auto& S = (const FSocket&)Peer;
+
+	using EWhat = FSocket::FWaiter::EWhat;
 	EWhat What = (Recv != nullptr) ? EWhat::Recv : EWhat::Send;
-	return FPeerType::FWaiter(Peer, What);
+	return FSocket::FWaiter(S, What);
 }
 
 ////////////////////////////////////////////////////////////////////////////////
@@ -764,7 +765,9 @@ void FPeerGroup::TickSend(FTickState& State, FHost& Host)
 	if (!Peer.IsValid())
 	{
 		IsKeepAlive = 1;
-		FOutcome Outcome = Host.Connect(Peer);
+		Peer = Host.Connect();
+
+		FOutcome Outcome = FOutcome::Ok();
 
 		// We failed to connect, let's bail.
 		if (Outcome.IsError())
@@ -826,10 +829,10 @@ FHostGroup::FHostGroup(FHost& InHost)
 int32 FHostGroup::Wait(const FTickState& State)
 {
 	// Collect groups that are waiting on something
-	TArray<FPeerType::FWaiter, TFixedAllocator<64>> Waiters;
+	TArray<FSocket::FWaiter, TFixedAllocator<64>> Waiters;
 	for (FPeerGroup& Group : PeerGroups)
 	{
-		FPeerType::FWaiter Waiter = Group.GetWaiter();
+		FSocket::FWaiter Waiter = Group.GetWaiter();
 		if (Waiter.IsValid())
 		{
 			Waiters.Add(Waiter);
@@ -853,7 +856,7 @@ int32 FHostGroup::Wait(const FTickState& State)
 	}
 
 	// Actually do the wait
-	int32 Result = FPeerType::Wait(Waiters, PollTimeoutMs);
+	int32 Result = FSocket::Wait(Waiters, PollTimeoutMs);
 	if (Result <= 0)
 	{
 		// If the user opts to not block then we don't accumulate wait time and
@@ -879,7 +882,7 @@ int32 FHostGroup::Wait(const FTickState& State)
 			continue;
 		}
 
-		auto* Candidate = (FPeerType*)(Waiters[i].Candidate);
+		auto* Candidate = (FSocket*)(Waiters[i].Candidate);
 		auto Pred = [Candidate] (auto& Lhs) { return Lhs == Candidate; };
 		FPeerGroup* Group = PeerGroups.FindByPredicate(Pred);
 		check(Group != nullptr);
@@ -1265,6 +1268,15 @@ FRequest FEventLoop::Request(
 	// Create an activity and an emphemeral host
 	Params = (Params != nullptr) ? Params : &GDefaultParams;
 
+	FPemCert VerifyCert;
+	if (UrlOffsets.SchemeLength == 5)
+	{
+		if (VerifyCert = Params->VerifyCert; VerifyCert.GetSize() == 0)
+		{
+			VerifyCert = FPemCert("", 0);
+		}
+	}
+
 	uint32 BufferSize = Params->BufferSize;
 	BufferSize = (BufferSize >= 128) ? BufferSize : 128;
 	BufferSize += sizeof(FHost) + HostName.Len();
@@ -1286,6 +1298,7 @@ FRequest FEventLoop::Request(
 	new (Host) FHost({
 		.HostName	= HostNamePtr,
 		.Port		= Port,
+		.VerifyCert = VerifyCert,
 	});
 
 	return Impl->Request(Method, Path, Activity);
@@ -1299,6 +1312,7 @@ FRequest FEventLoop::Request(
 	const FRequestParams* Params)
 {
 	check(Pool.Ptr != nullptr);
+	check(Params == nullptr || Params->VerifyCert.GetData() == nullptr); // add cert to FConPool instead
 
 	Params = (Params != nullptr) ? Params : &GDefaultParams;
 
@@ -1330,6 +1344,7 @@ bool FEventLoop::Redirect(const FTicketStatus& Status, FTicketSink& OuterSink)
 	FAnsiStringView Location = Response.GetHeader("Location");
 	if (Location.IsEmpty())
 	{
+		// todo: turn source activity into an error?
 		return false;
 	}
 
@@ -1361,12 +1376,16 @@ bool FEventLoop::Redirect(const FTicketStatus& Status, FTicketSink& OuterSink)
 		FHost& Host = *(Activity.Host);
 
 		TAnsiStringBuilder<256> Url;
-		Url << "http://";
+		Url << (Host.WithTls() ? "https" : "http");
+		Url << "://";
 		Url << Host.GetHostName();
 		Url << ":" << Host.GetPort();
 		Url << Location;
 
-		new (&ForwardRequest) FRequest(Request(Method, Url));
+		FRequestParams RequestParams = {
+			.VerifyCert = Host.GetVerifyCert(),
+		};
+		new (&ForwardRequest) FRequest(Request(Method, Url, &RequestParams));
 	}
 	else
 	{
@@ -1401,6 +1420,8 @@ bool FEventLoop::Redirect(const FTicketStatus& Status, FTicketSink& OuterSink)
 	// Send the request
 	Send(MoveTemp(ForwardRequest), MoveTemp(OuterSink), Status.GetParam());
 
+	// todo: activity slots should be swapped so original slot matches ticket
+
 	return true;
 }
 
diff --git a/Engine/Source/Runtime/Experimental/IoStore/HttpClient/Private/Peer.inl b/Engine/Source/Runtime/Experimental/IoStore/HttpClient/Private/Peer.inl
--- a/Engine/Source/Runtime/Experimental/IoStore/HttpClient/Private/Peer.inl
+++ b/Engine/Source/Runtime/Experimental/IoStore/HttpClient/Private/Peer.inl
@@ -1 +1,380 @@
 // Copyright Epic Games, Inc. All Rights Reserved.
+
+#pragma once
+
+namespace UE::IoStore::HTTP
+{
+
+////////////////////////////////////////////////////////////////////////////////
+static_assert(OPENSSL_VERSION_NUMBER >= 0x10100000L, "version supporting autoinit required");
+
+
+
+////////////////////////////////////////////////////////////////////////////////
+class FSslContext
+{
+public:
+				FSslContext(const char* InHostName, const FPemCert& CertPin={});
+				~FSslContext()					{ SSL_CTX_free(Context); }
+				operator SSL_CTX* () const		{ return Context; }
+	const char*	GetHostName() const				{ return HostName; }
+	bool		AddCert(const FPemCert& Cert);
+
+private:
+	SSL_CTX*	Context;
+	const char*	HostName;
+};
+
+////////////////////////////////////////////////////////////////////////////////
+FSslContext::FSslContext(const char* InHostName, const FPemCert& CertPin)
+: HostName(InHostName)
+{
+	Context = SSL_CTX_new(TLS_client_method());
+
+	uint32 ProtoFlags = SSL_OP_NO_SSLv2 | SSL_OP_NO_SSLv3;
+	SSL_CTX_set_options(Context, ProtoFlags);
+
+	if (CertPin.GetSize() == 0)
+	{
+		const ISslCertificateManager& CertManager = FSslModule::Get().GetCertificateManager();
+		CertManager.AddCertificatesToSslContext(Context);
+		return;
+	}
+
+	AddCert(CertPin);
+}
+
+////////////////////////////////////////////////////////////////////////////////
+bool FSslContext::AddCert(const FPemCert& Cert)
+{
+	const void* Data = Cert.GetData();
+	uint32 Size = uint32(Cert.GetSize());
+	if (Data == nullptr || Size == 0)
+	{
+		return false;
+	}
+
+	BIO* Bio = BIO_new_mem_buf(Data, Size);
+	X509* FiveOhNine = PEM_read_bio_X509(Bio, nullptr, 0, nullptr);
+	check(FiveOhNine != nullptr);
+
+	X509_STORE* Store = SSL_CTX_get_cert_store(Context);
+	int32 Result = X509_STORE_add_cert(Store, FiveOhNine);
+	return Result == 1;
+}
+
+
+
+////////////////////////////////////////////////////////////////////////////////
+class FPeer
+{
+public:
+				FPeer() = default;
+				FPeer(FSocket InSocket);
+	FOutcome	Send(const char* Data, int32 Size)	{ return Socket.Send(Data, Size); }
+	FOutcome	Recv(char* Out, int32 MaxSize)		{ return Socket.Recv(Out, MaxSize); }
+	bool		IsValid() const						{ return Socket.IsValid(); }
+
+private:
+	FSocket		Socket;
+
+
+
+	// TODO: !!!!!! remove this !!!!!
+public:
+	explicit operator const FSocket& () const { return Socket; }
+	explicit operator const FSocket* () const { return &Socket; }
+};
+
+////////////////////////////////////////////////////////////////////////////////
+FPeer::FPeer(FSocket InSocket)
+: Socket(MoveTemp(InSocket))
+{
+}
+
+
+
+////////////////////////////////////////////////////////////////////////////////
+class FTlsPeer
+	: public FPeer
+{
+public:
+				FTlsPeer()								= default;
+				~FTlsPeer();
+				FTlsPeer(FTlsPeer&& Rhs)				{ Move(MoveTemp(Rhs)); }
+				FTlsPeer& operator = (FTlsPeer&& Rhs)	{ return Move(MoveTemp(Rhs)); }
+				FTlsPeer(FSocket InSocket, const FSslContext* Context=nullptr);
+	FTlsPeer&	Move(FTlsPeer&& Rhs);
+	FOutcome	Handshake();
+	FOutcome	Send(const char* Data, int32 Size);
+	FOutcome	Recv(char* Out, int32 MaxSize);
+
+protected:
+	SSL*		Ssl = nullptr;
+
+private:
+	FOutcome GetOutcome(int32 SslResult, const char* Message="tls error") const
+	{
+		char buf[256];
+		TArray<uint32> iii;
+		int32 line;
+		const char* file;
+		while (uint32 i = ERR_get_error_line(&file, &line))
+		{
+			ERR_error_string(i, buf);
+			iii.Add(i);
+		}
+
+		int32 Error = SSL_get_error(Ssl, SslResult);
+		if (Error != SSL_ERROR_WANT_READ && Error != SSL_ERROR_WANT_WRITE)
+		{
+			return FOutcome::Error(Message, Error);
+		}
+		return FOutcome::Waiting();
+	}
+
+	int32 BioWrite(const char* Data, size_t Size, size_t* BytesWritten, BIO* Bio)
+	{
+		*BytesWritten = 0;
+		BIO_clear_retry_flags(Bio);
+
+		FOutcome Outcome = FPeer::Send(Data, Size);
+		if (Outcome.IsWaiting())
+		{
+			BIO_set_retry_write(Bio);
+			return 0;
+		}
+
+		if (Outcome.IsError())
+		{
+			return -1;
+		}
+
+		*BytesWritten = Outcome.GetResult();
+		return 1;
+	}
+
+	int32 BioRead(char* Data, size_t Size, size_t* BytesRead, BIO* Bio)
+	{
+		*BytesRead = 0;
+		BIO_clear_retry_flags(Bio);
+
+		FOutcome Outcome = FPeer::Recv(Data, Size);
+		if (Outcome.IsWaiting())
+		{
+			BIO_set_retry_read(Bio);
+			return 0;
+		}
+
+		if (Outcome.IsError())
+		{
+			return -1;
+		}
+
+		*BytesRead = Outcome.GetResult();
+		return 1;
+	}
+
+	long BioControl(int Cmd, long, void*, BIO*)
+	{
+		return (Cmd == BIO_CTRL_FLUSH) ? 1 : 0;
+	}
+};
+
+////////////////////////////////////////////////////////////////////////////////
+FTlsPeer::FTlsPeer(FSocket InSocket, const FSslContext* Context)
+: FPeer(MoveTemp(InSocket))
+{
+	if (Context == nullptr)
+	{
+		return;
+	}
+
+	static BIO_METHOD* BioMethod = nullptr;
+	if (BioMethod == nullptr)
+	{
+		int32 BioId = BIO_get_new_index() | BIO_TYPE_SOURCE_SINK;
+		BioMethod = BIO_meth_new(BioId, "IasBIO");
+
+#define METH_THUNK(x_, y_)\
+		x_(BioMethod, [] <typename... T> (BIO* b, T... t) { \
+			return (decltype(this)(BIO_get_data(b)))->y_(Forward<T>(t)..., b); \
+		})
+		METH_THUNK(BIO_meth_set_write_ex,	BioWrite);
+		METH_THUNK(BIO_meth_set_read_ex,	BioRead);
+		METH_THUNK(BIO_meth_set_ctrl,		BioControl);
+#undef METH_THUNK
+	}
+
+	BIO* Bio = BIO_new(BioMethod);
+	BIO_set_data(Bio, this);
+
+	// SSL_MODE_ENABLE_PARTIAL_WRITE ??!!!
+
+	Ssl = SSL_new(*Context);
+	SSL_set_tlsext_host_name(Ssl, Context->GetHostName());
+	SSL_set_connect_state(Ssl);
+	SSL_set0_rbio(Ssl, Bio);
+	SSL_set0_wbio(Ssl, Bio);
+	BIO_up_ref(Bio);
+}
+
+////////////////////////////////////////////////////////////////////////////////
+FTlsPeer& FTlsPeer::Move(FTlsPeer&& Rhs)
+{
+	FPeer::operator = (MoveTemp(Rhs));
+
+	Swap(Ssl, Rhs.Ssl);
+
+	auto PatchBio = [] (FTlsPeer& Peer)
+	{
+		if (Peer.Ssl != nullptr)
+		{
+			BIO* Bio = SSL_get_rbio(Peer.Ssl);
+			check(Bio != nullptr);
+			BIO_set_data(Bio, &Peer);
+		}
+	};
+	PatchBio(*this);
+	PatchBio(Rhs);
+
+	return *this;
+}
+
+////////////////////////////////////////////////////////////////////////////////
+FTlsPeer::~FTlsPeer()
+{
+	if (Ssl != nullptr)
+	{
+		SSL_free(Ssl);
+	}
+}
+
+////////////////////////////////////////////////////////////////////////////////
+FOutcome FTlsPeer::Handshake()
+{
+	if (Ssl == nullptr)
+	{
+		return FOutcome::Ok();
+	}
+
+	int32 Result = SSL_do_handshake(Ssl);
+	if (Result == 0) return FOutcome::Error("unsuccessful tls handshake");
+	if (Result != 1) return GetOutcome(Result, "tls handshake error");
+
+	if (Result = SSL_get_verify_result(Ssl); Result != X509_V_OK)
+	{
+		return FOutcome::Error("x509 verification error", Result);
+	}
+
+	return FOutcome::Ok();
+}
+
+////////////////////////////////////////////////////////////////////////////////
+FOutcome FTlsPeer::Send(const char* Data, int32 Size)
+{
+	if (Ssl == nullptr)
+	{
+		return FPeer::Send(Data, Size);
+	}
+
+	int32 Result = SSL_write(Ssl, Data, Size);
+	return (Result > 0) ? FOutcome::Ok(Result) : GetOutcome(Result);
+}
+
+////////////////////////////////////////////////////////////////////////////////
+FOutcome FTlsPeer::Recv(char* Out, int32 MaxSize)
+{
+	if (Ssl == nullptr)
+	{
+		return FPeer::Recv(Out, MaxSize);
+	}
+
+	int32 Result = SSL_read(Ssl, Out, MaxSize);
+	return (Result > 0) ? FOutcome::Ok(Result) : GetOutcome(Result);
+}
+
+
+
+////////////////////////////////////////////////////////////////////////////////
+class FHttpPeer
+	: public FTlsPeer
+{
+public:
+				FHttpPeer() = default;
+				FHttpPeer(FSocket InSocket, FSslContext* Context=nullptr);
+	FOutcome	Handshake();
+
+private:
+	void		AssignProto();
+	int32		Proto = 0;
+};
+
+////////////////////////////////////////////////////////////////////////////////
+FHttpPeer::FHttpPeer(FSocket InSocket, FSslContext* Context)
+: FTlsPeer(MoveTemp(InSocket), Context)
+{
+	if (Ssl == nullptr)
+	{
+		return;
+	}
+
+	static const uint8 AlpnProtos[] =
+		"\x08" "http/1.1"
+	//	"\x02" "h2"
+		;
+	SSL_set_alpn_protos(Ssl, AlpnProtos, sizeof(AlpnProtos) - 1);
+}
+
+////////////////////////////////////////////////////////////////////////////////
+FOutcome FHttpPeer::Handshake()
+{
+	FOutcome Outcome = FTlsPeer::Handshake();
+	if (Outcome.IsOk())
+	{
+		AssignProto();
+	}
+
+	return Outcome;
+}
+
+////////////////////////////////////////////////////////////////////////////////
+void FHttpPeer::AssignProto()
+{
+	Proto = 1;
+
+	if (Ssl == nullptr)
+		return;
+
+	const char* AlpnProto;
+	uint32 AlpnProtoLen;
+	SSL_get0_alpn_selected(Ssl, &(const uint8*&)AlpnProto, &AlpnProtoLen);
+	if (AlpnProto == nullptr)
+	{
+		return;
+	}
+
+	FAnsiStringView Needle(AlpnProto, AlpnProtoLen);
+	FAnsiStringView Candidates[] = {
+		"http/1.1",
+	//	"h2"
+	};
+	for (int32 i = 0; i < UE_ARRAY_COUNT(Candidates); ++i)
+	{
+		const FAnsiStringView& Candidate = Candidates[i];
+		if (AlpnProtoLen != Candidate.Len())
+		{
+			continue;
+		}
+
+		if (Candidate != Needle)
+		{
+			continue;
+		}
+
+		Proto = i + 1;
+		break;
+	}
+}
+
+} // namespace UE::IoStore::HTTP
