Iris - Allow USequentialPartialNetBlobHandlerConfig to be inherited.

#11744

diff --git a/Engine/Source/Runtime/Experimental/Iris/Core/Public/Iris/ReplicationSystem/NetBlob/SequentialPartialNetBlobHandler.h b/Engine/Source/Runtime/Experimental/Iris/Core/Public/Iris/ReplicationSystem/NetBlob/SequentialPartialNetBlobHandler.h
--- a/Engine/Source/Runtime/Experimental/Iris/Core/Public/Iris/ReplicationSystem/NetBlob/SequentialPartialNetBlobHandler.h
+++ b/Engine/Source/Runtime/Experimental/Iris/Core/Public/Iris/ReplicationSystem/NetBlob/SequentialPartialNetBlobHandler.h
@@ -19,7 +19,7 @@ struct FSequentialPartialNetBlobHandlerInitParams
 	const USequentialPartialNetBlobHandlerConfig* Config;
 };
 
-UCLASS(Config=Engine)
+UCLASS(Config=Engine, MinimalAPI)
 class USequentialPartialNetBlobHandlerConfig : public UObject
 {
 	GENERATED_BODY()
