Placate PVS maybe.

#rnx

diff --git a/Engine/Source/Runtime/Experimental/IoStore/HttpClient/Private/Peer.inl b/Engine/Source/Runtime/Experimental/IoStore/HttpClient/Private/Peer.inl
--- a/Engine/Source/Runtime/Experimental/IoStore/HttpClient/Private/Peer.inl
+++ b/Engine/Source/Runtime/Experimental/IoStore/HttpClient/Private/Peer.inl
@@ -346,7 +346,7 @@ void FHttpPeer::AssignProto()
 	if (Ssl == nullptr)
 		return;
 
-	const char* AlpnProto;
+	const char* AlpnProto = nullptr;
 	uint32 AlpnProtoLen;
 	SSL_get0_alpn_selected(Ssl, &(const uint8*&)AlpnProto, &AlpnProtoLen);
 	if (AlpnProto == nullptr)
