ChaosFlesh : Add triangle mesh simulation properties node and restrict self collision to tetmesh surface.

diff --git a/Engine/Plugins/Experimental/ChaosFlesh/Source/ChaosFleshNodes/Private/Dataflow/ChaosFleshEngineAssetNodes.cpp b/Engine/Plugins/Experimental/ChaosFlesh/Source/ChaosFleshNodes/Private/Dataflow/ChaosFleshEngineAssetNodes.cpp
--- a/Engine/Plugins/Experimental/ChaosFlesh/Source/ChaosFleshNodes/Private/Dataflow/ChaosFleshEngineAssetNodes.cpp
+++ b/Engine/Plugins/Experimental/ChaosFlesh/Source/ChaosFleshNodes/Private/Dataflow/ChaosFleshEngineAssetNodes.cpp
@@ -24,7 +24,6 @@ namespace Dataflow
 	{
 		DATAFLOW_NODE_REGISTER_CREATION_FACTORY(FGetFleshAssetDataflowNode);
 		DATAFLOW_NODE_REGISTER_CREATION_FACTORY(FFleshAssetTerminalDataflowNode);
-		DATAFLOW_NODE_REGISTER_CREATION_FACTORY(FSetFleshDefaultPropertiesNode);
 		DATAFLOW_NODE_REGISTER_CREATION_FACTORY(FComputeFiberFieldNode);
 		DATAFLOW_NODE_REGISTER_CREATION_FACTORY(FComputeMuscleActivationDataNode);
 		DATAFLOW_NODE_REGISTER_CREATION_FACTORY(FVisualizeFiberFieldNode);
@@ -76,146 +75,6 @@ void FFleshAssetTerminalDataflowNode::Evaluate(Dataflow::FContext& Context, cons
 	SetValue(Context, InCollection, &Collection);
 }
 
-template <class T> using MType = FManagedArrayCollection::TManagedType<T>;
-
-void FSetFleshDefaultPropertiesNode::Evaluate(Dataflow::FContext& Context, const FDataflowOutput* Out) const
-{
-	if (Out->IsA<FManagedArrayCollection>(&Collection))
-	{
-		FManagedArrayCollection InCollection = GetValue<FManagedArrayCollection>(Context, &Collection);
-		
-		TManagedArray<float>& ParticleStiffness = InCollection.AddAttribute<float>("Stiffness", FGeometryCollection::VerticesGroup);
-		TManagedArray<float>& ParticleDamping = InCollection.AddAttribute<float>("Damping", FGeometryCollection::VerticesGroup);
-		TManagedArray<float>& ParticleIncompressibility = InCollection.AddAttribute<float>("Incompressibility", FGeometryCollection::VerticesGroup);
-		TManagedArray<float>& ParticleInflation = InCollection.AddAttribute<float>("Inflation", FGeometryCollection::VerticesGroup);
-
-		
-
-		if (InCollection.HasAttributes({
-			MType< float >("Mass", FGeometryCollection::VerticesGroup),
-			MType< float >("Stiffness", FGeometryCollection::VerticesGroup),
-			MType< float >("Damping", FGeometryCollection::VerticesGroup),
-			MType< float >("Incompressibility", FGeometryCollection::VerticesGroup),
-			MType< float >("Inflation", FGeometryCollection::VerticesGroup),
-			MType< FIntVector4 >(FTetrahedralCollection::TetrahedronAttribute, FTetrahedralCollection::TetrahedralGroup),
-			MType< FVector3f >("Vertex", "Vertices"),
-			MType< TArray<int32> >(FTetrahedralCollection::IncidentElementsAttribute, FGeometryCollection::VerticesGroup),
-			MType< TArray<int32> >(FTetrahedralCollection::IncidentElementsLocalIndexAttribute, FGeometryCollection::VerticesGroup) }))
-		{
-
-			int32 VertsNum = InCollection.NumElements(FGeometryCollection::VerticesGroup);
-			int32 TetsNum = InCollection.NumElements(FTetrahedralCollection::TetrahedralGroup);
-			if (VertsNum && TetsNum)
-			{
-				TManagedArray<float>& Mass = InCollection.ModifyAttribute<float>("Mass", FGeometryCollection::VerticesGroup);
-				TManagedArray<float>& Stiffness = InCollection.ModifyAttribute<float>("Stiffness", FGeometryCollection::VerticesGroup);
-				TManagedArray<float>& Damping = InCollection.ModifyAttribute<float>("Damping", FGeometryCollection::VerticesGroup);
-				TManagedArray<float>& Incompressibility = InCollection.ModifyAttribute<float>("Incompressibility", FGeometryCollection::VerticesGroup);
-				TManagedArray<float>& Inflation = InCollection.ModifyAttribute<float>("Inflation", FGeometryCollection::VerticesGroup);
-				const TManagedArray<FIntVector4>& Tetrahedron = InCollection.GetAttribute<FIntVector4>(FTetrahedralCollection::TetrahedronAttribute, FTetrahedralCollection::TetrahedralGroup);
-				const TManagedArray<FVector3f>& Vertex = InCollection.GetAttribute<FVector3f>("Vertex", "Vertices");
-				const TManagedArray<TArray<int32>>& IncidentElements = InCollection.GetAttribute<TArray<int32>>(FTetrahedralCollection::IncidentElementsAttribute, FGeometryCollection::VerticesGroup);
-				const TManagedArray<TArray<int32>>& IncidentElementsLocalIndex = InCollection.GetAttribute<TArray<int32>>(FTetrahedralCollection::IncidentElementsLocalIndexAttribute, FGeometryCollection::VerticesGroup);
-
-				double TotalVolume = 0.0;
-				TArray<float> ElementMass;
-				TArray<float> ElementVolume;
-				ElementMass.Init(0.f, TetsNum);
-				ElementVolume.Init(0.f, TetsNum);
-				for (int e = 0; e < TetsNum; e++)
-				{
-					FVector3f X0 = Vertex[Tetrahedron[e][0]];
-					FVector3f X1 = Vertex[Tetrahedron[e][1]];
-					FVector3f X2 = Vertex[Tetrahedron[e][2]];
-					FVector3f X3 = Vertex[Tetrahedron[e][3]];
-					ElementVolume[e] = ((X1 - X0).Dot(FVector3f::CrossProduct(X3 - X0, X2 - X0))) / 6.f;
-					if (ElementVolume[e] < 0.f)
-					{
-						ElementVolume[e] = -ElementVolume[e];
-					}
-					TotalVolume += ElementVolume[e];
-				}
-				for (int32 e = 0; e < TetsNum; e++)
-				{
-					ElementMass[e] = Density * ElementVolume[e];
-				}
-
-				//
-				// Set per-node mass by connected volume
-				//
-
-				TSet<int32> Visited;
-				for (int32 i = 0; i < IncidentElements.Num(); i++)
-				{
-					const TArray<int32>& IncidentElems = IncidentElements[i];
-					for (int32 j = 0; j < IncidentElems.Num(); j++)
-					{
-						const int32 TetIndex = IncidentElems[j];
-						if (Tetrahedron.GetConstArray().IsValidIndex(TetIndex))
-						{
-							for (int32 k = 0; k < 4; k++)
-							{
-								const int32 MassIndex = Tetrahedron[TetIndex][k];
-								if (Mass.GetConstArray().IsValidIndex(MassIndex))
-								{
-									Mass[MassIndex] += ElementMass[TetIndex] / 4;
-									Visited.Add(MassIndex);
-								}
-							}
-						}
-					}
-				}
-				
-				int32 NumSet = 0;
-				float MinV = TNumericLimits<float>::Max();
-				float MaxV = -TNumericLimits<float>::Max();
-				double AvgV = 0.0;
-				if (Visited.Num())
-				{
-					NumSet = Visited.Num();
-					for (TSet<int32>::TConstIterator It = Visited.CreateConstIterator(); It; ++It)
-					{
-						const int32 MassIndex = *It;
-						AvgV += Mass[MassIndex];
-						MinV = MinV < Mass[MassIndex] ? MinV : Mass[MassIndex];
-						MaxV = MaxV > Mass[MassIndex] ? MaxV : Mass[MassIndex];
-					}
-					AvgV /= NumSet;
-				}
-				else
-				{
-					MinV = MaxV = 0.0;
-				}
-
-				//
-				// If that didn't work, set a uniform mass
-				//
-
-				if (!Visited.Num() && Vertex.Num())
-				{
-					Mass.Fill(Density * TotalVolume / Vertex.Num());
-					NumSet = Mass.Num();
-					Chaos::Utilities::GetMinAvgMax(Mass.GetConstArray(), MinV, AvgV, MaxV);
-				}
-
-				Stiffness.Fill(VertexStiffness);
-				Damping.Fill(VertexDamping);
-				Incompressibility.Fill(.5f * VertexIncompressibility);
-				Inflation.Fill(VertexInflation * 2.f);
-
-				UE_LOG(LogChaosFlesh, Display,
-					TEXT("'%s' - Set mass on %d nodes:\n"
-						"    method: %s\n"
-						"    min, avg, max: %f, %f, %f"),
-					*GetName().ToString(), NumSet, 
-					(Visited.Num() > 0 ? TEXT("connected tet volume") : TEXT("uniform")),
-					MinV, AvgV, MaxV);
-			}
-		}
-		SetValue(Context, MoveTemp(InCollection), &Collection);
-	}
-}
-
 void FComputeFiberFieldNode::Evaluate(Dataflow::FContext& Context, const FDataflowOutput* Out) const
 {
 	if (Out->IsA<FManagedArrayCollection>(&Collection))
diff --git a/Engine/Plugins/Experimental/ChaosFlesh/Source/ChaosFleshNodes/Private/Dataflow/ChaosFleshNodesPlugin.cpp b/Engine/Plugins/Experimental/ChaosFlesh/Source/ChaosFleshNodes/Private/Dataflow/ChaosFleshNodesPlugin.cpp
--- a/Engine/Plugins/Experimental/ChaosFlesh/Source/ChaosFleshNodes/Private/Dataflow/ChaosFleshNodesPlugin.cpp
+++ b/Engine/Plugins/Experimental/ChaosFlesh/Source/ChaosFleshNodes/Private/Dataflow/ChaosFleshNodesPlugin.cpp
@@ -15,6 +15,9 @@
 #include "Dataflow/ChaosFleshPositionTargetInitializationNodes.h"
 #include "Dataflow/ChaosFleshSkeletalBindingsNode.h"
 #include "Dataflow/ChaosFleshTetrahedralNodes.h"
+#include "Dataflow/ChaosFleshTriangleMeshSimulationPropertiesNode.h"
+#include "Dataflow/ChaosFleshSkinSimulationPropertiesNode.h"
+#include "Dataflow/ChaosFleshSetFleshDefaultPropertiesNode.h"
 #include "Modules/ModuleManager.h"
 
 #define LOCTEXT_NAMESPACE "ChaosFleshNodes"
@@ -34,6 +37,9 @@ void IChaosFleshNodesPlugin::StartupModule()
 	Dataflow::ChaosFleshRadialTetrahedronNodes();
 	Dataflow::RegisterChaosFleshImportGEONodes();
 	DATAFLOW_NODE_REGISTER_CREATION_FACTORY(FCreateTetrahedronDataflowNode);
+	DATAFLOW_NODE_REGISTER_CREATION_FACTORY(FTriangleMeshSimulationPropertiesDataflowNodes);
+	DATAFLOW_NODE_REGISTER_CREATION_FACTORY(FSkinSimulationPropertiesDataflowNodes);
+	DATAFLOW_NODE_REGISTER_CREATION_FACTORY(FSetFleshDefaultPropertiesNode);
 }
 
 void IChaosFleshNodesPlugin::ShutdownModule()
diff --git a/Engine/Plugins/Experimental/ChaosFlesh/Source/ChaosFleshNodes/Private/Dataflow/ChaosFleshSetFleshDefaultPropertiesNode.cpp b/Engine/Plugins/Experimental/ChaosFlesh/Source/ChaosFleshNodes/Private/Dataflow/ChaosFleshSetFleshDefaultPropertiesNode.cpp
new file mode 100644
--- /dev/null
+++ b/Engine/Plugins/Experimental/ChaosFlesh/Source/ChaosFleshNodes/Private/Dataflow/ChaosFleshSetFleshDefaultPropertiesNode.cpp
@@ -0,0 +1,188 @@
+// Copyright Epic Games, Inc. All Rights Reserved.
+
+#include "Dataflow/ChaosFleshSetFleshDefaultPropertiesNode.h"
+
+#include "Async/ParallelFor.h"
+#include "Chaos/Deformable/Utilities.h"
+#include "ChaosFlesh/ChaosFlesh.h"
+#include "Chaos/Tetrahedron.h"
+#include "Chaos/Utilities.h"
+#include "Chaos/UniformGrid.h"
+#include "ChaosFlesh/FleshCollection.h"
+#include "ChaosFlesh/FleshCollectionUtility.h"
+#include "ChaosLog.h"
+#include "Dataflow/DataflowInputOutput.h"
+#include "Engine/SkeletalMesh.h"
+#include "Engine/StaticMesh.h"
+#include "GeometryCollection/ManagedArrayCollection.h"
+#include "GeometryCollection/GeometryCollectionAlgo.h"
+#include "MeshDescription.h"
+#include "MeshDescriptionToDynamicMesh.h"
+#include "Spatial/FastWinding.h"
+#include "Spatial/MeshAABBTree3.h"
+
+template <class T> using MType = FManagedArrayCollection::TManagedType<T>;
+
+void FSetFleshDefaultPropertiesNode::Evaluate(Dataflow::FContext& Context, const FDataflowOutput* Out) const
+{
+	if (Out->IsA<FManagedArrayCollection>(&Collection))
+	{
+		FManagedArrayCollection InCollection = GetValue<FManagedArrayCollection>(Context, &Collection);
+		
+		TManagedArray<float>& ParticleStiffness = InCollection.AddAttribute<float>("Stiffness", FGeometryCollection::VerticesGroup);
+		TManagedArray<float>& ParticleDamping = InCollection.AddAttribute<float>("Damping", FGeometryCollection::VerticesGroup);
+		TManagedArray<float>& ParticleIncompressibility = InCollection.AddAttribute<float>("Incompressibility", FGeometryCollection::VerticesGroup);
+		TManagedArray<float>& ParticleInflation = InCollection.AddAttribute<float>("Inflation", FGeometryCollection::VerticesGroup);
+
+		
+
+		if (InCollection.HasAttributes({
+			MType< float >("Mass", FGeometryCollection::VerticesGroup),
+			MType< float >("Stiffness", FGeometryCollection::VerticesGroup),
+			MType< float >("Damping", FGeometryCollection::VerticesGroup),
+			MType< float >("Incompressibility", FGeometryCollection::VerticesGroup),
+			MType< float >("Inflation", FGeometryCollection::VerticesGroup),
+			MType< FIntVector4 >(FTetrahedralCollection::TetrahedronAttribute, FTetrahedralCollection::TetrahedralGroup),
+			MType< FVector3f >("Vertex", "Vertices"),
+			MType< TArray<int32> >(FTetrahedralCollection::IncidentElementsAttribute, FGeometryCollection::VerticesGroup),
+			MType< TArray<int32> >(FTetrahedralCollection::IncidentElementsLocalIndexAttribute, FGeometryCollection::VerticesGroup) }))
+		{
+
+			int32 VertsNum = InCollection.NumElements(FGeometryCollection::VerticesGroup);
+			int32 TetsNum = InCollection.NumElements(FTetrahedralCollection::TetrahedralGroup);
+			if (VertsNum)
+			{
+				TManagedArray<float>& Mass = InCollection.ModifyAttribute<float>("Mass", FGeometryCollection::VerticesGroup);
+				TManagedArray<float>& Stiffness = InCollection.ModifyAttribute<float>("Stiffness", FGeometryCollection::VerticesGroup);
+				TManagedArray<float>& Damping = InCollection.ModifyAttribute<float>("Damping", FGeometryCollection::VerticesGroup);
+				TManagedArray<float>& Incompressibility = InCollection.ModifyAttribute<float>("Incompressibility", FGeometryCollection::VerticesGroup);
+				TManagedArray<float>& Inflation = InCollection.ModifyAttribute<float>("Inflation", FGeometryCollection::VerticesGroup);
+				const TManagedArray<FIntVector4>& Tetrahedron = InCollection.GetAttribute<FIntVector4>(FTetrahedralCollection::TetrahedronAttribute, FTetrahedralCollection::TetrahedralGroup);
+				const TManagedArray<FVector3f>& Vertex = InCollection.GetAttribute<FVector3f>("Vertex", "Vertices");
+				const TManagedArray<TArray<int32>>& IncidentElements = InCollection.GetAttribute<TArray<int32>>(FTetrahedralCollection::IncidentElementsAttribute, FGeometryCollection::VerticesGroup);
+				const TManagedArray<TArray<int32>>& IncidentElementsLocalIndex = InCollection.GetAttribute<TArray<int32>>(FTetrahedralCollection::IncidentElementsLocalIndexAttribute, FGeometryCollection::VerticesGroup);
+				
+				float MinV = TNumericLimits<float>::Max();
+				float MaxV = -TNumericLimits<float>::Max();
+				double AvgV = 0.0;
+				TSet<int32> Visited;
+				int32 NumSet = 0;
+				if (TetsNum)
+				{
+					double TotalVolume = 0.0;
+					TArray<float> ElementMass;
+					TArray<float> ElementVolume;
+					ElementMass.Init(0.f, TetsNum);
+					ElementVolume.Init(0.f, TetsNum);
+					for (int e = 0; e < TetsNum; e++)
+					{
+						FVector3f X0 = Vertex[Tetrahedron[e][0]];
+						FVector3f X1 = Vertex[Tetrahedron[e][1]];
+						FVector3f X2 = Vertex[Tetrahedron[e][2]];
+						FVector3f X3 = Vertex[Tetrahedron[e][3]];
+						ElementVolume[e] = ((X1 - X0).Dot(FVector3f::CrossProduct(X3 - X0, X2 - X0))) / 6.f;
+						if (ElementVolume[e] < 0.f)
+						{
+							ElementVolume[e] = -ElementVolume[e];
+						}
+						TotalVolume += ElementVolume[e];
+					}
+					for (int32 e = 0; e < TetsNum; e++)
+					{
+						ElementMass[e] = Density * ElementVolume[e];
+					}
+
+					//
+					// Set per-node mass by connected volume
+					//
+
+					for (int32 i = 0; i < IncidentElements.Num(); i++)
+					{
+						const TArray<int32>& IncidentElems = IncidentElements[i];
+						for (int32 j = 0; j < IncidentElems.Num(); j++)
+						{
+							const int32 TetIndex = IncidentElems[j];
+							if (Tetrahedron.GetConstArray().IsValidIndex(TetIndex))
+							{
+								for (int32 k = 0; k < 4; k++)
+								{
+									const int32 MassIndex = Tetrahedron[TetIndex][k];
+									if (Mass.GetConstArray().IsValidIndex(MassIndex))
+									{
+										Mass[MassIndex] += ElementMass[TetIndex] / 4;
+										Visited.Add(MassIndex);
+									}
+								}
+							}
+						}
+					}
+
+					if (Visited.Num())
+					{
+						NumSet = Visited.Num();
+						for (TSet<int32>::TConstIterator It = Visited.CreateConstIterator(); It; ++It)
+						{
+							const int32 MassIndex = *It;
+							AvgV += Mass[MassIndex];
+							MinV = MinV < Mass[MassIndex] ? MinV : Mass[MassIndex];
+							MaxV = MaxV > Mass[MassIndex] ? MaxV : Mass[MassIndex];
+						}
+						AvgV /= NumSet;
+					}
+					else
+					{
+						MinV = MaxV = 0.0;
+					}
+					if (!Visited.Num() && Vertex.Num())
+					{
+						Mass.Fill(Density * TotalVolume / Vertex.Num());
+						NumSet = Mass.Num();
+						Chaos::Utilities::GetMinAvgMax(Mass.GetConstArray(), MinV, AvgV, MaxV);
+					}
+				}
+				
+
+				if (const TManagedArray<int32>* TriangleMeshIndices = InCollection.FindAttribute<int32>("ObjectIndices", "TriangleMesh"))
+				{
+					//sets stiffness and damping for nontriangle mesh particles
+					if (const TManagedArray<int32>* VertexStarts = InCollection.FindAttribute<int32>("VertexStart", FGeometryCollection::GeometryGroup))
+					{
+						if (const TManagedArray<int32>* VertexCounts = InCollection.FindAttribute<int32>("VertexCount", FGeometryCollection::GeometryGroup))
+						{
+							for (int32 i = 0; i < VertexStarts->Num(); i++)
+							{
+								if (!TriangleMeshIndices->Contains(i))
+								{
+									const int32 VertexStartIndex = (*VertexStarts)[i];
+									const int32 VertexNum = (*VertexCounts)[i];
+									for (int32 ParticleIndex = VertexStartIndex; ParticleIndex < VertexStartIndex + VertexNum; ParticleIndex++)
+									{
+										Stiffness[ParticleIndex] = VertexStiffness;
+										Damping[ParticleIndex] = VertexDamping;
+									}
+								}
+							}
+						}
+					}
+				}
+				else
+				{
+					Stiffness.Fill(VertexStiffness);
+					Damping.Fill(VertexDamping);
+				}
+				
+				Incompressibility.Fill(.5f * VertexIncompressibility);
+				Inflation.Fill(VertexInflation * 2.f);
+
+				UE_LOG(LogChaosFlesh, Display,
+					TEXT("'%s' - Set mass on %d nodes:\n"
+						"    method: %s\n"
+						"    min, avg, max: %f, %f, %f"),
+					*GetName().ToString(), NumSet, 
+					(Visited.Num() > 0 ? TEXT("connected tet volume") : TEXT("uniform")),
+					MinV, AvgV, MaxV);
+			}
+		}
+		SetValue(Context, MoveTemp(InCollection), &Collection);
+	}
+}
\ No newline at end of file
diff --git a/Engine/Plugins/Experimental/ChaosFlesh/Source/ChaosFleshNodes/Private/Dataflow/ChaosFleshSkinSimulationPropertiesNode.cpp b/Engine/Plugins/Experimental/ChaosFlesh/Source/ChaosFleshNodes/Private/Dataflow/ChaosFleshSkinSimulationPropertiesNode.cpp
new file mode 100644
--- /dev/null
+++ b/Engine/Plugins/Experimental/ChaosFlesh/Source/ChaosFleshNodes/Private/Dataflow/ChaosFleshSkinSimulationPropertiesNode.cpp
@@ -0,0 +1,52 @@
+// Copyright Epic Games, Inc. All Rights Reserved.
+
+#include "Dataflow/ChaosFleshSkinSimulationPropertiesNode.h"
+
+#include "Async/ParallelFor.h"
+#include "Chaos/Deformable/Utilities.h"
+#include "ChaosFlesh/ChaosFlesh.h"
+#include "Chaos/Tetrahedron.h"
+#include "Chaos/Utilities.h"
+#include "Chaos/UniformGrid.h"
+#include "ChaosFlesh/FleshCollection.h"
+#include "ChaosFlesh/FleshCollectionUtility.h"
+#include "ChaosLog.h"
+#include "Dataflow/DataflowInputOutput.h"
+#include "DynamicMesh/DynamicMesh3.h"
+#include "DynamicMesh/DynamicMeshAABBTree3.h"
+#include "Engine/SkeletalMesh.h"
+#include "Engine/StaticMesh.h"
+#include "FTetWildWrapper.h"
+#include "Generate/IsosurfaceStuffing.h"
+#include "GeometryCollection/ManagedArrayCollection.h"
+#include "GeometryCollection/Facades/CollectionTetrahedralMetricsFacade.h"
+#include "GeometryCollection/GeometryCollectionAlgo.h"
+#include "MeshDescription.h"
+#include "MeshDescriptionToDynamicMesh.h"
+#include "Rendering/SkeletalMeshLODImporterData.h"
+#include "Rendering/SkeletalMeshModel.h"
+#include "SkeletalMeshLODModelToDynamicMesh.h" // MeshModelingBlueprints
+#include "Spatial/FastWinding.h"
+#include "Spatial/MeshAABBTree3.h"
+
+
+void FSkinSimulationPropertiesDataflowNodes::Evaluate(Dataflow::FContext& Context, const FDataflowOutput* Out) const
+{
+	TUniquePtr<FFleshCollection> InCollection(GetValue<DataType>(Context, &Collection).NewCopy<FFleshCollection>());
+	if (Out->IsA<DataType>(&Collection))
+	{
+		const FName TriangleMeshGroup = "TriangleMesh";
+		if (const TManagedArray<int32>* TriangleMeshIndices = InCollection->FindAttribute<int32>("ObjectIndices", "TriangleMesh"))
+		{
+			if (bSkinConstraints)
+			{
+				TManagedArray<bool>& IsSkinConstraint = InCollection->AddAttribute<bool>("SkinConstraints", TriangleMeshGroup);
+				for (int32 i = 0; i < IsSkinConstraint.Num(); i++)
+				{
+					IsSkinConstraint[i] = true;
+				}
+			}
+		}
+		SetValue<const DataType&>(Context, *InCollection, &Collection);
+	}
+}
\ No newline at end of file
diff --git a/Engine/Plugins/Experimental/ChaosFlesh/Source/ChaosFleshNodes/Private/Dataflow/ChaosFleshTriangleMeshSimulationPropertiesNode.cpp b/Engine/Plugins/Experimental/ChaosFlesh/Source/ChaosFleshNodes/Private/Dataflow/ChaosFleshTriangleMeshSimulationPropertiesNode.cpp
new file mode 100644
--- /dev/null
+++ b/Engine/Plugins/Experimental/ChaosFlesh/Source/ChaosFleshNodes/Private/Dataflow/ChaosFleshTriangleMeshSimulationPropertiesNode.cpp
@@ -0,0 +1,140 @@
+// Copyright Epic Games, Inc. All Rights Reserved.
+
+#include "Dataflow/ChaosFleshTriangleMeshSimulationPropertiesNode.h"
+#include "Async/ParallelFor.h"
+#include "Chaos/Deformable/Utilities.h"
+#include "ChaosFlesh/ChaosFlesh.h"
+#include "Chaos/Tetrahedron.h"
+#include "Chaos/Utilities.h"
+#include "Chaos/UniformGrid.h"
+#include "ChaosFlesh/FleshCollection.h"
+#include "ChaosFlesh/FleshCollectionUtility.h"
+#include "ChaosLog.h"
+#include "Dataflow/DataflowInputOutput.h"
+#include "DynamicMesh/DynamicMesh3.h"
+#include "DynamicMesh/DynamicMeshAABBTree3.h"
+#include "Engine/SkeletalMesh.h"
+#include "Engine/StaticMesh.h"
+#include "FTetWildWrapper.h"
+#include "Generate/IsosurfaceStuffing.h"
+#include "GeometryCollection/ManagedArrayCollection.h"
+#include "GeometryCollection/Facades/CollectionTetrahedralMetricsFacade.h"
+#include "GeometryCollection/GeometryCollectionAlgo.h"
+#include "MeshDescription.h"
+#include "MeshDescriptionToDynamicMesh.h"
+#include "Rendering/SkeletalMeshLODImporterData.h"
+#include "Rendering/SkeletalMeshModel.h"
+#include "SkeletalMeshLODModelToDynamicMesh.h" // MeshModelingBlueprints
+#include "Spatial/FastWinding.h"
+#include "Spatial/MeshAABBTree3.h"
+
+
+template <class T> using MType = FManagedArrayCollection::TManagedType<T>;
+
+void FTriangleMeshSimulationPropertiesDataflowNodes::Evaluate(Dataflow::FContext& Context, const FDataflowOutput* Out) const
+{
+	TUniquePtr<FFleshCollection> InCollection(GetValue<DataType>(Context, &Collection).NewCopy<FFleshCollection>());
+	if (Out->IsA<DataType>(&Collection))
+	{
+
+		TManagedArray<FString>* GroupNames = InCollection->FindAttribute<FString>("BoneName", FGeometryCollection::TransformGroup);
+		TManagedArray<int32>* TransformIndex = InCollection->FindAttribute<int32>("TransformIndex", FGeometryCollection::GeometryGroup);
+		TManagedArray<float>& ParticleStiffness = InCollection->AddAttribute<float>("Stiffness", FGeometryCollection::VerticesGroup);
+		TManagedArray<float>& ParticleDamping = InCollection->AddAttribute<float>("Damping", FGeometryCollection::VerticesGroup);
+		if (GroupNames && TransformIndex)
+		{
+			TSet<FString> TriangleMeshNameSet = TSet(MeshNames);
+			TArray<int32> TriangleMeshIndices;
+
+			for (int32 i = 0; i < TransformIndex->Num(); i++)
+			{
+				if (TriangleMeshNameSet.Contains((*GroupNames)[(*TransformIndex)[i]]))
+				{
+					TriangleMeshIndices.Add(i);
+				}
+			}
+
+			if (TriangleMeshIndices.Num())
+			{
+				const FName TriangleMeshGroup = "TriangleMesh";
+				int32 TriMeshIndex = InCollection->AddElements(TriangleMeshIndices.Num(), TriangleMeshGroup);
+				if (TriMeshIndex == 0)
+				{
+					TManagedArray<int32>& TriangleMeshIndex = InCollection->AddAttribute<int32>("ObjectIndices", TriangleMeshGroup);
+				}
+				if (InCollection->HasAttributes({MType<int32>("ObjectIndices", TriangleMeshGroup) }))
+				{
+					TManagedArray<int32>* TriangleMeshIndexPtr = InCollection->FindAttribute<int32>("ObjectIndices", TriangleMeshGroup);
+					for (int32 i = 0; i < TriangleMeshIndices.Num(); i++)
+					{
+						(*TriangleMeshIndexPtr)[i + TriMeshIndex] = TriangleMeshIndices[i];
+					}
+				}
+			}
+		}
+
+		if (const TManagedArray<int32>* TriangleMeshIndices = InCollection->FindAttribute<int32>("ObjectIndices", "TriangleMesh"))
+		{
+			if (const TManagedArray<FIntVector>* Indices = InCollection->FindAttribute<FIntVector>("Indices", FGeometryCollection::FacesGroup))
+			{
+				if (const TManagedArray<int32>* FaceStarts = InCollection->FindAttribute<int32>("FaceStart", FGeometryCollection::GeometryGroup))
+				{
+					if (const TManagedArray<int32>* FaceCounts = InCollection->FindAttribute<int32>("FaceCount", FGeometryCollection::GeometryGroup))
+					{
+						if (const TManagedArray<FVector3f>* Vertex = InCollection->FindAttribute<FVector3f>("Vertex", "Vertices"))
+						{
+							if (TManagedArray<float>* Mass = InCollection->FindAttribute<float>("Mass", "Vertices"))
+							{
+								for (int32 i = 0; i < TriangleMeshIndices->Num(); i++)
+								{
+									const int32 ObjectIndex = (*TriangleMeshIndices)[i];
+							
+									const int32 FaceStartIndex = (*FaceStarts)[ObjectIndex];
+									const int32 FaceNum = (*FaceCounts)[ObjectIndex];
+									for (int32 e = FaceStartIndex; e < FaceStartIndex + FaceNum; e++)
+									{
+										const FVector3f X0 = (*Vertex)[(*Indices)[e][0]];
+										const FVector3f X1 = (*Vertex)[(*Indices)[e][1]];
+										const FVector3f X2 = (*Vertex)[(*Indices)[e][2]];
+										float SingleElementMass = FVector3f::CrossProduct(X1 - X0, X2 - X0).Size() * TriangleMeshDensity/ 2.f;
+										if (SingleElementMass < 0.f)
+										{
+											SingleElementMass = -SingleElementMass;
+										}
+										for (int32 k = 0; k < 3; k++)
+										{
+											const int32 MassIndex = (*Indices)[e][k];
+											(*Mass)[MassIndex] += SingleElementMass/3.f;
+										}
+									}
+								}
+							}
+						}	
+					}
+				}
+			}
+
+
+			if (const TManagedArray<int32>* VertexStarts = InCollection->FindAttribute<int32>("VertexStart", FGeometryCollection::GeometryGroup))
+			{
+				if (const TManagedArray<int32>* VertexCounts = InCollection->FindAttribute<int32>("VertexCount", FGeometryCollection::GeometryGroup))
+				{
+					for (int32 i = 0; i < TriangleMeshIndices->Num(); i++)
+					{
+						const int32 ObjectIndex = (*TriangleMeshIndices)[i];
+							
+						const int32 VertexStartIndex = (*VertexStarts)[ObjectIndex];
+						const int32 VertexNum = (*VertexCounts)[ObjectIndex];
+						for (int32 ParticleIndex = VertexStartIndex; ParticleIndex < VertexStartIndex + VertexNum; ParticleIndex++)
+						{
+							ParticleStiffness[ParticleIndex] = VertexTriangleMeshStiffness;
+							ParticleDamping[ParticleIndex] = VertexTriangleMeshDamping;
+						}
+					}	
+				}
+			}
+		}
+
+		SetValue<const DataType&>(Context, *InCollection, &Collection);
+	}
+}
\ No newline at end of file
diff --git a/Engine/Plugins/Experimental/ChaosFlesh/Source/ChaosFleshNodes/Public/Dataflow/ChaosFleshEngineAssetNodes.h b/Engine/Plugins/Experimental/ChaosFlesh/Source/ChaosFleshNodes/Public/Dataflow/ChaosFleshEngineAssetNodes.h
--- a/Engine/Plugins/Experimental/ChaosFlesh/Source/ChaosFleshNodes/Public/Dataflow/ChaosFleshEngineAssetNodes.h
+++ b/Engine/Plugins/Experimental/ChaosFlesh/Source/ChaosFleshNodes/Public/Dataflow/ChaosFleshEngineAssetNodes.h
@@ -64,46 +64,7 @@ public:
 };
 
 
-USTRUCT(meta = (DataflowFlesh))
-struct FSetFleshDefaultPropertiesNode : public FDataflowNode
-{
-	GENERATED_USTRUCT_BODY()
-	DATAFLOW_NODE_DEFINE_INTERNAL(FSetFleshDefaultPropertiesNode, "SetFleshDefaultProperties", "Flesh", "")
-	DATAFLOW_NODE_RENDER_TYPE(FGeometryCollection::StaticType(), "Collection")
-
-public:
-
-	UPROPERTY(EditAnywhere, Category = "Dataflow")
-	float Density = 1.f;
-
-	UPROPERTY(EditAnywhere, Category = "Dataflow")
-	float VertexStiffness = 1e6;
-
-	UPROPERTY(EditAnywhere, Category = "Dataflow", meta = (ClampMin = "0.0", ClampMax = "1.0", UIMin = "0.0", UIMax = "1.0"))
-	float VertexDamping = 0.f;
-
-	/*Sets incompressibility on vertex basis. 0.6 is default behavior. 
-	1 means total incompressibility. 0.00001 means almost no incompressibility*/
-	UPROPERTY(EditAnywhere, Category = "Dataflow", meta = (ClampMin = "0.00001", ClampMax = "1.0", UIMin = "0.00001", UIMax = "1.0"))
-	float VertexIncompressibility = 0.6f;
 
-	/*Sets inflation on vertex basis. 0.5 means no inflation/deflation.
-	1 means total inflation. 0 means the material is deflated.*/
-	UPROPERTY(EditAnywhere, Category = "Dataflow", meta = (ClampMin = "0.0", ClampMax = "1.0", UIMin = "0.0", UIMax = "1.0"))
-	float VertexInflation = 0.5f;
-
-	UPROPERTY(meta = (DataflowInput, DataflowOutput, DisplayName = "Collection", DataflowPassthrough  = "Collection"))
-	FManagedArrayCollection Collection;
-
-	FSetFleshDefaultPropertiesNode(const Dataflow::FNodeParameters& InParam, FGuid InGuid = FGuid::NewGuid())
-		: FDataflowNode(InParam, InGuid)
-	{
-		RegisterInputConnection(&Collection);
-		RegisterOutputConnection(&Collection, &Collection);
-	}
-
-	virtual void Evaluate(Dataflow::FContext& Context, const FDataflowOutput* Out) const override;
-};
 
 /**
 * Computes a muscle fiber direction per tetrahedron from a GeometryCollection containing tetrahedra, 
diff --git a/Engine/Plugins/Experimental/ChaosFlesh/Source/ChaosFleshNodes/Public/Dataflow/ChaosFleshSetFleshDefaultPropertiesNode.h b/Engine/Plugins/Experimental/ChaosFlesh/Source/ChaosFleshNodes/Public/Dataflow/ChaosFleshSetFleshDefaultPropertiesNode.h
new file mode 100644
--- /dev/null
+++ b/Engine/Plugins/Experimental/ChaosFlesh/Source/ChaosFleshNodes/Public/Dataflow/ChaosFleshSetFleshDefaultPropertiesNode.h
@@ -0,0 +1,52 @@
+// Copyright Epic Games, Inc. All Rights Reserved.
+
+#pragma once 
+
+#include "CoreMinimal.h"
+#include "Dataflow/DataflowEngine.h"
+#include "GeometryCollection/ManagedArrayCollection.h"
+#include "Dataflow/ChaosFleshNodesUtility.h"
+
+#include "ChaosFleshSetFleshDefaultPropertiesNode.generated.h"
+
+
+USTRUCT(meta = (DataflowFlesh))
+struct FSetFleshDefaultPropertiesNode : public FDataflowNode
+{
+	GENERATED_USTRUCT_BODY()
+	DATAFLOW_NODE_DEFINE_INTERNAL(FSetFleshDefaultPropertiesNode, "SetFleshDefaultProperties", "Flesh", "")
+	DATAFLOW_NODE_RENDER_TYPE(FGeometryCollection::StaticType(), "Collection")
+
+public:
+
+	UPROPERTY(EditAnywhere, Category = "Dataflow")
+	float Density = 1.f;
+
+	UPROPERTY(EditAnywhere, Category = "Dataflow")
+	float VertexStiffness = 1e6;
+
+	UPROPERTY(EditAnywhere, Category = "Dataflow", meta = (ClampMin = "0.0", ClampMax = "1.0", UIMin = "0.0", UIMax = "1.0"))
+	float VertexDamping = 0.f;
+
+	/*Sets incompressibility on vertex basis. 0.6 is default behavior. 
+	1 means total incompressibility. 0.00001 means almost no incompressibility*/
+	UPROPERTY(EditAnywhere, Category = "Dataflow", meta = (ClampMin = "0.00001", ClampMax = "1.0", UIMin = "0.00001", UIMax = "1.0"))
+	float VertexIncompressibility = 0.6f;
+
+	/*Sets inflation on vertex basis. 0.5 means no inflation/deflation.
+	1 means total inflation. 0 means the material is deflated.*/
+	UPROPERTY(EditAnywhere, Category = "Dataflow", meta = (ClampMin = "0.0", ClampMax = "1.0", UIMin = "0.0", UIMax = "1.0"))
+	float VertexInflation = 0.5f;
+
+	UPROPERTY(meta = (DataflowInput, DataflowOutput, DisplayName = "Collection", DataflowPassthrough  = "Collection"))
+	FManagedArrayCollection Collection;
+
+	FSetFleshDefaultPropertiesNode(const Dataflow::FNodeParameters& InParam, FGuid InGuid = FGuid::NewGuid())
+		: FDataflowNode(InParam, InGuid)
+	{
+		RegisterInputConnection(&Collection);
+		RegisterOutputConnection(&Collection, &Collection);
+	}
+
+	virtual void Evaluate(Dataflow::FContext& Context, const FDataflowOutput* Out) const override;
+};
diff --git a/Engine/Plugins/Experimental/ChaosFlesh/Source/ChaosFleshNodes/Public/Dataflow/ChaosFleshSkinSimulationPropertiesNode.h b/Engine/Plugins/Experimental/ChaosFlesh/Source/ChaosFleshNodes/Public/Dataflow/ChaosFleshSkinSimulationPropertiesNode.h
new file mode 100644
--- /dev/null
+++ b/Engine/Plugins/Experimental/ChaosFlesh/Source/ChaosFleshNodes/Public/Dataflow/ChaosFleshSkinSimulationPropertiesNode.h
@@ -0,0 +1,42 @@
+// Copyright Epic Games, Inc. All Rights Reserved.
+
+#pragma once 
+
+#include "CoreMinimal.h"
+#include "Dataflow/DataflowEngine.h"
+#include "GeometryCollection/ManagedArrayCollection.h"
+#include "Dataflow/ChaosFleshNodesUtility.h"
+
+#include "ChaosFleshSkinSimulationPropertiesNode.generated.h"
+
+
+/*
+Set triangle mesh to simulate using skin constraints 
+*/
+USTRUCT(meta = (DataflowFlesh))
+struct FSkinSimulationPropertiesDataflowNodes : public FDataflowNode
+{
+	GENERATED_USTRUCT_BODY()
+	DATAFLOW_NODE_DEFINE_INTERNAL(FSkinSimulationPropertiesDataflowNodes, "SkinSimulationProperties", "Flesh", "")
+	DATAFLOW_NODE_RENDER_TYPE(FGeometryCollection::StaticType(), "Collection")
+
+public:
+	typedef FManagedArrayCollection DataType;
+
+	UPROPERTY(EditAnywhere, Category = "Dataflow", meta = (DisplayName = "UseSkinConstraints"))
+	bool bSkinConstraints = false;
+
+	UPROPERTY(meta = (DataflowInput, DataflowOutput, DisplayName = "Collection"))
+	FManagedArrayCollection Collection;
+
+
+	FSkinSimulationPropertiesDataflowNodes(const Dataflow::FNodeParameters& InParam, FGuid InGuid = FGuid::NewGuid())
+		: FDataflowNode(InParam, InGuid)
+	{
+		RegisterInputConnection(&Collection);
+		RegisterOutputConnection(&Collection);
+	}
+
+	virtual void Evaluate(Dataflow::FContext& Context, const FDataflowOutput* Out) const override;
+
+};
diff --git a/Engine/Plugins/Experimental/ChaosFlesh/Source/ChaosFleshNodes/Public/Dataflow/ChaosFleshTriangleMeshSimulationPropertiesNode.h b/Engine/Plugins/Experimental/ChaosFlesh/Source/ChaosFleshNodes/Public/Dataflow/ChaosFleshTriangleMeshSimulationPropertiesNode.h
new file mode 100644
--- /dev/null
+++ b/Engine/Plugins/Experimental/ChaosFlesh/Source/ChaosFleshNodes/Public/Dataflow/ChaosFleshTriangleMeshSimulationPropertiesNode.h
@@ -0,0 +1,53 @@
+// Copyright Epic Games, Inc. All Rights Reserved.
+
+#pragma once 
+
+#include "CoreMinimal.h"
+#include "Dataflow/DataflowEngine.h"
+#include "GeometryCollection/ManagedArrayCollection.h"
+#include "Dataflow/ChaosFleshNodesUtility.h"
+
+#include "ChaosFleshTriangleMeshSimulationPropertiesNode.generated.h"
+
+
+
+/*
+Convert tetmesh to simulate using surface traingle mesh only
+*/
+USTRUCT(meta = (DataflowFlesh))
+struct FTriangleMeshSimulationPropertiesDataflowNodes : public FDataflowNode
+{
+	GENERATED_USTRUCT_BODY()
+	DATAFLOW_NODE_DEFINE_INTERNAL(FTriangleMeshSimulationPropertiesDataflowNodes, "TriangleMeshSimulationProperties", "Flesh", "")
+	DATAFLOW_NODE_RENDER_TYPE(FGeometryCollection::StaticType(), "Collection")
+
+public:
+	typedef FManagedArrayCollection DataType;
+
+	UPROPERTY(EditAnywhere, Category = "Dataflow", meta = (DisplayName = "MeshNames"))
+	TArray<FString> MeshNames;
+
+	UPROPERTY(EditAnywhere, Category = "Dataflow")
+	float TriangleMeshDensity = 1.f;
+
+	UPROPERTY(EditAnywhere, Category = "Dataflow")
+	float VertexTriangleMeshStiffness = 1e6;
+
+	UPROPERTY(EditAnywhere, Category = "Dataflow")
+	float VertexTriangleMeshDamping = 0.f;
+
+	UPROPERTY(meta = (DataflowInput, DataflowOutput, DisplayName = "Collection"))
+	FManagedArrayCollection Collection;
+
+
+	FTriangleMeshSimulationPropertiesDataflowNodes(const Dataflow::FNodeParameters& InParam, FGuid InGuid = FGuid::NewGuid())
+		: FDataflowNode(InParam, InGuid)
+	{
+		RegisterInputConnection(&Collection);
+		RegisterOutputConnection(&Collection);
+	}
+
+	virtual void Evaluate(Dataflow::FContext& Context, const FDataflowOutput* Out) const override;
+
+
+};
diff --git a/Engine/Source/Runtime/Experimental/Chaos/Private/Chaos/Deformable/ChaosDeformableSolver.cpp b/Engine/Source/Runtime/Experimental/Chaos/Private/Chaos/Deformable/ChaosDeformableSolver.cpp
--- a/Engine/Source/Runtime/Experimental/Chaos/Private/Chaos/Deformable/ChaosDeformableSolver.cpp
+++ b/Engine/Source/Runtime/Experimental/Chaos/Private/Chaos/Deformable/ChaosDeformableSolver.cpp
@@ -55,7 +55,7 @@ DECLARE_CYCLE_STAT(TEXT("Chaos.Deformable.Solver.InitializeSimulationObjects"),
 DECLARE_CYCLE_STAT(TEXT("Chaos.Deformable.Solver.InitializeSimulationObject"), STAT_ChaosDeformableSolver_InitializeSimulationObject, STATGROUP_Chaos);
 DECLARE_CYCLE_STAT(TEXT("Chaos.Deformable.Solver.InitializeDeformableParticles"), STAT_ChaosDeformableSolver_InitializeDeformableParticles, STATGROUP_Chaos);
 DECLARE_CYCLE_STAT(TEXT("Chaos.Deformable.Solver.InitializeKinematicParticles"), STAT_ChaosDeformableSolver_InitializeKinematicParticles, STATGROUP_Chaos);
-DECLARE_CYCLE_STAT(TEXT("Chaos.Deformable.Solver.InitializeTetrahedralConstraint"), STAT_ChaosDeformableSolver_InitializeTetrahedralConstraint, STATGROUP_Chaos);
+DECLARE_CYCLE_STAT(TEXT("Chaos.Deformable.Solver.InitializeTetrahedralOrTriangleConstraint"), STAT_ChaosDeformableSolver_InitializeTetrahedralOrTriangleConstraint, STATGROUP_Chaos);
 DECLARE_CYCLE_STAT(TEXT("Chaos.Deformable.Solver.InitializeGidBasedConstraints"), STAT_ChaosDeformableSolver_InitializeGidBasedConstraints, STATGROUP_Chaos);
 DECLARE_CYCLE_STAT(TEXT("Chaos.Deformable.Solver.InitializeWeakConstraints"), STAT_ChaosDeformableSolver_InitializeWeakConstraints, STATGROUP_Chaos);
 DECLARE_CYCLE_STAT(TEXT("Chaos.Deformable.Solver.InitializeKinematicConstraint"), STAT_ChaosDeformableSolver_InitializeKinematicConstraint, STATGROUP_Chaos);
@@ -141,6 +141,7 @@ namespace Chaos::Softs
 		if (Property.bDoSelfCollision || Property.CacheToFile)
 		{
 			SurfaceElements.Reset(new TArray<Chaos::TVec3<int32>>());
+			TetmeshSurfaceElements.Reset(new TArray<Chaos::TVec3<int32>>());
 		}
 
 		if (Property.bDoSelfCollision)
@@ -163,6 +164,8 @@ namespace Chaos::Softs
 			GSWeakConstraints.Reset(new FGaussSeidelWeakConstraints<FSolverReal, FSolverParticles>({}, {}, {}, {}, {}, GDeformableXPBDWeakConstraintParams));
 			MuscleActivationConstraints.Reset(new FMuscleActivationConstraints<FSolverReal, FSolverParticles>());
 		}
+		AllUnconstrainedSurfaceElementsCorotatedCod.Reset(new TArray<Chaos::TVec3<int32>>());
+		AllUnconstrainedSurfaceElementsSkin.Reset(new TArray<Chaos::TVec3<int32>>());
 
 		InitializeKinematicConstraint();
 		Frame = 0;
@@ -385,7 +388,7 @@ namespace Chaos::Softs
 					InitializeKinematicParticles(*Proxy);
 					InitializeWeakConstraint(*Proxy);
 					InitializeMuscleActivation(*Proxy);
-					InitializeTetrahedralConstraint(*Proxy);
+					InitializeTetrahedralOrTriangleConstraint(*Proxy);
 					InitializeGidBasedConstraints(*Proxy);
 					InitializeGaussSeidelConstraints(*Proxy);
 				}
@@ -751,9 +754,9 @@ namespace Chaos::Softs
 	}
 
 
-	void FDeformableSolver::InitializeTetrahedralConstraint(FFleshThreadingProxy& Proxy)
+	void FDeformableSolver::InitializeTetrahedralOrTriangleConstraint(FFleshThreadingProxy& Proxy)
 	{
-		PERF_SCOPE(STAT_ChaosDeformableSolver_InitializeTetrahedralConstraint);
+		PERF_SCOPE(STAT_ChaosDeformableSolver_InitializeTetrahedralOrTriangleConstraint);
 
 		const FManagedArrayCollection& Rest = Proxy.GetRestCollection();
 
@@ -924,6 +927,80 @@ namespace Chaos::Softs
 			}
 
 		}
+
+		const FIntVector2& Range = Proxy.GetSolverParticleRange();
+		if (const TManagedArray<int32>* TriangleMeshIndices = Rest.FindAttribute<int32>("ObjectIndices", "TriangleMesh"))
+		{
+			if (const TManagedArray<FIntVector>* Indices = Rest.FindAttribute<FIntVector>("Indices", FGeometryCollection::FacesGroup))
+			{
+				if (const TManagedArray<int32>* FaceStarts = Rest.FindAttribute<int32>("FaceStart", FGeometryCollection::GeometryGroup))
+				{
+					if (const TManagedArray<int32>* FaceCounts = Rest.FindAttribute<int32>("FaceCount", FGeometryCollection::GeometryGroup))
+					{
+						if (const TManagedArray<int32>* VertexStarts = Rest.FindAttribute<int32>("VertexStart", FGeometryCollection::GeometryGroup))
+						{
+							if (const TManagedArray<int32>* VertexCounts = Rest.FindAttribute<int32>("VertexCount", FGeometryCollection::GeometryGroup))
+							{
+								if (const TManagedArray<bool>* bUseSkinConstraints = Rest.FindAttribute<bool>("SkinConstraints", "TriangleMesh"))
+								{
+									for (int32 i = 0; i < TriangleMeshIndices->Num(); i++)
+									{
+										const int32 ObjectIndex = (*TriangleMeshIndices)[i];
+										const int32 FaceStartIndex = (*FaceStarts)[ObjectIndex];
+										const int32 FaceNum = (*FaceCounts)[ObjectIndex];
+										if ((*bUseSkinConstraints)[i])
+										{
+											int32 SurfaceOffset = AllUnconstrainedSurfaceElementsSkin->Num();
+											AllUnconstrainedSurfaceElementsSkin->SetNum(SurfaceOffset + FaceNum);
+											for (int32 e = FaceStartIndex; e < FaceStartIndex + FaceNum; e++)
+											{
+												for (int32 j = 0; j < 3; j++)
+												{
+													(*AllUnconstrainedSurfaceElementsSkin)[e - FaceStartIndex + SurfaceOffset][j] = (*Indices)[e][j];
+												}
+											}
+										}
+										else
+										{
+											int32 SurfaceOffset = AllUnconstrainedSurfaceElementsCorotatedCod->Num();
+											AllUnconstrainedSurfaceElementsCorotatedCod->SetNum(SurfaceOffset + FaceNum);
+											for (int32 e = FaceStartIndex; e < FaceStartIndex + FaceNum; e++)
+											{
+												for (int32 j = 0; j < 3; j++)
+												{
+													(*AllUnconstrainedSurfaceElementsCorotatedCod)[e - FaceStartIndex + SurfaceOffset][j] = (*Indices)[e][j];
+												}
+											}
+										}
+									}
+								}
+								else
+								{
+									for (int32 i = 0; i < TriangleMeshIndices->Num(); i++)
+									{
+										const int32 ObjectIndex = (*TriangleMeshIndices)[i];
+										const int32 FaceStartIndex = (*FaceStarts)[ObjectIndex];
+										const int32 FaceNum = (*FaceCounts)[ObjectIndex];
+										int32 SurfaceOffset = AllUnconstrainedSurfaceElementsCorotatedCod->Num();
+										AllUnconstrainedSurfaceElementsCorotatedCod->SetNum(SurfaceOffset + FaceNum);
+										for (int32 e = FaceStartIndex; e < FaceStartIndex + FaceNum; e++)
+										{
+											for (int32 j = 0; j < 3; j++)
+											{
+												(*AllUnconstrainedSurfaceElementsCorotatedCod)[e - FaceStartIndex + SurfaceOffset][j] = (*Indices)[e][j];
+											}
+										}
+									}
+								}
+							}
+						}
+							
+					}
+					
+				}
+			}
+		
+		}
 	}
 
 
@@ -979,7 +1056,7 @@ namespace Chaos::Softs
 				}
 				if (uint32 NumIncident = IncidentElementsPointer->Num())
 				{
-					int32 IncidentOffSet = AllIncidentElements->Num();
+					int32 IncidentOffSet = Range[0];
 					AllIncidentElements->SetNum(IncidentOffSet + NumIncident); 
 					AllIncidentElementsLocal->SetNum(IncidentOffSet + NumIncident);
 					for (uint32 i = 0; i < NumIncident; ++i)
@@ -1126,6 +1203,7 @@ namespace Chaos::Softs
 		PERF_SCOPE(STAT_ChaosDeformableSolver_InitializeSelfCollisionVariables);
 		int32 VertexOffset = 0;
 		SurfaceElements->SetNum(0);
+		TetmeshSurfaceElements->SetNum(0);
 		int ComponentOffset = 0;
 
 		for (FThreadingProxy* InProxy : UninitializedProxys_Internal)
@@ -1181,6 +1259,73 @@ namespace Chaos::Softs
 
 		SurfaceTriangleMesh->Init(*SurfaceElements);
 		SurfaceTriangleMesh->GetVertexSetAsArray(*SurfaceVertices);
+
+		VertexOffset = 0;
+
+		for (FThreadingProxy* InProxy : UninitializedProxys_Internal)
+		{
+			if (FFleshThreadingProxy* Proxy = InProxy->As<FFleshThreadingProxy>())
+			{
+				if (const FManagedArrayCollection* Rest = &Proxy->GetRestCollection())
+				{
+					if (const TManagedArray<FVector3f>* Vertex = Rest->FindAttribute<FVector3f>("Vertex", FGeometryCollection::VerticesGroup))
+					{
+						if (const TManagedArray<FIntVector>* Indices = Rest->FindAttribute<FIntVector>("Indices", FGeometryCollection::FacesGroup))
+						{
+							if (const TManagedArray<int32>* TriangleMeshIndices = Rest->FindAttribute<int32>("ObjectIndices", "TriangleMesh"))
+							{
+								if (const TManagedArray<int32>* FaceStarts = Rest->FindAttribute<int32>("FaceStart", FGeometryCollection::GeometryGroup))
+								{
+									if (const TManagedArray<int32>* FaceCounts = Rest->FindAttribute<int32>("FaceCount", FGeometryCollection::GeometryGroup))
+									{
+										TSet<int32> TriMeshObjects;
+										for (const int32 ObjectIndex: (*TriangleMeshIndices))
+										{
+											TriMeshObjects.Add(ObjectIndex);
+										}
+										for (int32 i = 0; i < FaceStarts->Num(); i++) 
+										{
+											if (!TriMeshObjects.Contains(i))
+											{
+												const int32 FaceStartIndex = (*FaceStarts)[i];
+												const int32 FaceNum = (*FaceCounts)[i];
+
+												int32 SurfaceOffset = TetmeshSurfaceElements->Num();
+												TetmeshSurfaceElements->SetNum(SurfaceOffset + FaceNum);
+												for (int32 e = FaceStartIndex; e < FaceStartIndex + FaceNum; e++)
+												{
+													for (int32 j = 0; j < 3; j++)
+													{
+														(*TetmeshSurfaceElements)[e - FaceStartIndex + SurfaceOffset][j] = (*Indices)[e][j];
+													}
+												}
+											}
+										}
+									}
+								}
+								
+							}
+							else
+							{
+								int32 SurfaceOffset = TetmeshSurfaceElements->Num();
+								TetmeshSurfaceElements->SetNum(SurfaceOffset + Indices->Num());
+								for (int32 i = 0; i < Indices->Num(); i++)
+								{
+									for (int32 j = 0; j < 3; j++)
+									{
+										(*TetmeshSurfaceElements)[i + SurfaceOffset][j] = VertexOffset + (*Indices)[i][j];
+									}
+								}
+							}
+						}
+						VertexOffset += Vertex->Num();
+					}
+				}
+			}
+		}
+		SurfaceTriangleMesh->Init(*TetmeshSurfaceElements);
+
+
 		TriangleMeshCollisions.Reset(new FPBDTriangleMeshCollisions(
 			0, Evolution->Particles().Size(), *SurfaceTriangleMesh, false, false));
 		ParticleTriangleExclusionMap.Reset();
@@ -1250,6 +1395,36 @@ namespace Chaos::Softs
 
 		GSMainConstraint.Reset(new Chaos::Softs::FGaussSeidelMainConstraint<FSolverReal, FSolverParticles>(Evolution->Particles(), Property.bDoQuasistatics, Property.bUseSOR, Property.OmegaSOR, GSParallelMax));
 
+		if (AllUnconstrainedSurfaceElementsCorotatedCod->Num() > 0) 
+		{
+			GSCorotatedCodConstraints.Reset(new Chaos::Softs::FGaussSeidelCorotatedCodimensionalConstraints<FSolverReal, FSolverParticles>(
+				Evolution->Particles(), *AllUnconstrainedSurfaceElementsCorotatedCod, false, Property.EMesh));
+			TArray<TArray<int32>> IncidentElements, IncidentElementsLocal;
+			GSMainConstraint->AddStaticConstraints(GSCorotatedCodConstraints->GetConstraintsArray(), IncidentElements, IncidentElementsLocal);
+
+			int32 StaticIndex = GSMainConstraint->AddStaticConstraintResidualAndHessianRange(1);
+
+			GSMainConstraint->StaticConstraintResidualAndHessian()[StaticIndex] = [this](const FSolverParticles& Particles, const int32 ElementIndex, const int32 ElementIndexLocal, const FSolverReal Dt, TVec3<FSolverReal>& ParticleResidual, Chaos::PMatrix<FSolverReal, 3, 3>& ParticleHessian)
+			{
+				this->GSCorotatedCodConstraints->AddHyperelasticResidualAndHessian(Particles, ElementIndex, ElementIndexLocal, Dt, ParticleResidual, ParticleHessian);
+			};
+		}
+
+		if (AllUnconstrainedSurfaceElementsSkin->Num() > 0)
+		{
+			GSLinearCodConstraints.Reset(new Chaos::Softs::FGaussSeidelLinearCodimensionalConstraints<FSolverReal, FSolverParticles>(
+				Evolution->Particles(), *AllUnconstrainedSurfaceElementsSkin, false, Property.EMesh));
+			TArray<TArray<int32>> IncidentElements, IncidentElementsLocal;
+			GSMainConstraint->AddStaticConstraints(GSLinearCodConstraints->GetConstraintsArray(), IncidentElements, IncidentElementsLocal);
+
+			int32 StaticIndex = GSMainConstraint->AddStaticConstraintResidualAndHessianRange(1);
+
+			GSMainConstraint->StaticConstraintResidualAndHessian()[StaticIndex] = [this](const FSolverParticles& Particles, const int32 ElementIndex, const int32 ElementIndexLocal, const FSolverReal Dt, TVec3<FSolverReal>& ParticleResidual, Chaos::PMatrix<FSolverReal, 3, 3>& ParticleHessian)
+			{
+				this->GSLinearCodConstraints->AddHyperelasticResidualAndHessian(Particles, ElementIndex, ElementIndexLocal, Dt, ParticleResidual, ParticleHessian);
+			};
+		}
+
 		if (Property.bUseGSNeohookean)
 		{
 			GSNeohookeanConstraints.Reset(new Chaos::Softs::FGaussSeidelNeohookeanConstraints<FSolverReal, FSolverParticles>(
@@ -1521,7 +1696,7 @@ namespace Chaos::Softs
 				if (FFleshThreadingProxy* Proxy = BaseProxyPair.Value->As<FFleshThreadingProxy>())
 				{
 					InitializeMuscleActivation(*Proxy);
-					InitializeTetrahedralConstraint(*Proxy);
+					InitializeTetrahedralOrTriangleConstraint(*Proxy);
 					InitializeGidBasedConstraints(*Proxy);
 					InitializeGaussSeidelConstraints(*Proxy);
 				}
diff --git a/Engine/Source/Runtime/Experimental/Chaos/Public/Chaos/Deformable/ChaosDeformableSolver.h b/Engine/Source/Runtime/Experimental/Chaos/Public/Chaos/Deformable/ChaosDeformableSolver.h
--- a/Engine/Source/Runtime/Experimental/Chaos/Public/Chaos/Deformable/ChaosDeformableSolver.h
+++ b/Engine/Source/Runtime/Experimental/Chaos/Public/Chaos/Deformable/ChaosDeformableSolver.h
@@ -16,6 +16,8 @@
 #include "Chaos/Deformable/GaussSeidelWeakConstraints.h"
 #include "Chaos/Deformable/GaussSeidelMainConstraint.h"
 #include "Chaos/Deformable/MuscleActivationConstraints.h"
+#include "Chaos/Deformable/GaussSeidelCorotatedCodimensionalConstraints.h"
+#include "Chaos/Deformable/GaussSeidelLinearCodimensionalConstraints.h"
 #include "Chaos/XPBDWeakConstraints.h"
 #include "Chaos/BlendedXPBDCorotatedConstraints.h"
 #include "Chaos/XPBDGridBasedCorotatedConstraints.h"
@@ -124,7 +126,7 @@ namespace Chaos::Softs
 		CHAOS_API void UpdateTransientConstraints();
 		CHAOS_API void PostProcessTransientConstraints();
 		CHAOS_API void InitializeKinematicParticles(FFleshThreadingProxy&);
-		CHAOS_API void InitializeTetrahedralConstraint(FFleshThreadingProxy&);
+		CHAOS_API void InitializeTetrahedralOrTriangleConstraint(FFleshThreadingProxy&);
 		CHAOS_API void InitializeGidBasedConstraints(FFleshThreadingProxy&);
 		CHAOS_API void InitializeGaussSeidelConstraints(FFleshThreadingProxy& Proxy);
 		CHAOS_API void InitializeWeakConstraint(FFleshThreadingProxy&);
@@ -183,6 +185,8 @@ namespace Chaos::Softs
 		TArray<TUniquePtr<Softs::FXPBDCorotatedConstraints<Softs::FSolverReal, Softs::FSolverParticles>>> CorotatedConstraints;
 		TUniquePtr<Softs::FGaussSeidelCorotatedConstraints<Softs::FSolverReal, Softs::FSolverParticles>> GSCorotatedConstraints;
 		TUniquePtr<Softs::FGaussSeidelNeohookeanConstraints<Softs::FSolverReal, Softs::FSolverParticles>> GSNeohookeanConstraints;
+		TUniquePtr<Softs::FGaussSeidelCorotatedCodimensionalConstraints<Softs::FSolverReal, Softs::FSolverParticles>> GSCorotatedCodConstraints;
+		TUniquePtr<Softs::FGaussSeidelLinearCodimensionalConstraints<Softs::FSolverReal, Softs::FSolverParticles>> GSLinearCodConstraints;
 		TUniquePtr<Softs::FGaussSeidelWeakConstraints<Softs::FSolverReal, Softs::FSolverParticles>> GSWeakConstraints;
 		TArray<TUniquePtr<Softs::FXPBDWeakConstraints<Softs::FSolverReal, Softs::FSolverParticles>>> WeakConstraints;
 		TArray<TUniquePtr<Softs::FBlendedXPBDCorotatedConstraints<Softs::FSolverReal, Softs::FSolverParticles>>> BlendedCorotatedConstraints;
@@ -192,6 +196,7 @@ namespace Chaos::Softs
 		TUniquePtr<Softs::FPBDTriangleMeshCollisions> TriangleMeshCollisions;
 		TArrayCollectionArray<const UObject*> MObjects;
 		TUniquePtr <TArray<TVec3<int32>>> SurfaceElements;
+		TUniquePtr <TArray<TVec3<int32>>> TetmeshSurfaceElements;
 		TUniquePtr <TArray<Chaos::TVec4<int32>>> AllElements;
 		TUniquePtr <FTriangleMesh> SurfaceTriangleMesh;
 		TUniquePtr <TArray<int32>> SurfaceVertices;
@@ -204,6 +209,8 @@ namespace Chaos::Softs
 		TUniquePtr <TArray<TArray<int32>>> AllSecondIndices;
 		TUniquePtr <TArray<FSolverReal>> AllWeights;
 		TUniquePtr <TArray<FSolverReal>> AllSecondWeights;
+		TUniquePtr <TArray<TVec3<int32>>> AllUnconstrainedSurfaceElementsCorotatedCod;  //correspond to the triangle mesh elements that are simulated using corotated cod
+		TUniquePtr <TArray<TVec3<int32>>> AllUnconstrainedSurfaceElementsSkin;          //correspond to the triangle mesh elements that are simulated using linear cod constraints
 		TArray<int32> ParticleComponentIndex;
 		TMap<int32, TSet<int32>> ParticleTriangleExclusionMap;
 		//Muscle Activation Variables
diff --git a/Engine/Source/Runtime/Experimental/Chaos/Public/Chaos/Deformable/GaussSeidelLinearCodimensionalConstraints.h b/Engine/Source/Runtime/Experimental/Chaos/Public/Chaos/Deformable/GaussSeidelLinearCodimensionalConstraints.h
new file mode 100644
--- /dev/null
+++ b/Engine/Source/Runtime/Experimental/Chaos/Public/Chaos/Deformable/GaussSeidelLinearCodimensionalConstraints.h
@@ -0,0 +1,241 @@
+// Copyright Epic Games, Inc. All Rights Reserved.
+#pragma once
+
+#include "Chaos/PBDSpringConstraintsBase.h"
+#include "ChaosStats.h"
+#include "Chaos/ImplicitQRSVD.h"
+#include "Chaos/GraphColoring.h"
+#include "Chaos/Framework/Parallel.h"
+#include "Chaos/Deformable/ChaosDeformableSolverTypes.h"
+#include "Chaos/Vector.h"
+
+
+namespace Chaos::Softs
+{
+	template <typename T, typename ParticleType>
+	class FGaussSeidelLinearCodimensionalConstraints
+	{
+
+	public:
+		//this one only accepts tetmesh input and mesh
+		FGaussSeidelLinearCodimensionalConstraints(
+			const ParticleType& InParticles,
+			const TArray<TVector<int32, 3>>& InMesh,
+			const bool bRecordMetricIn = true,
+			const T& EMesh = (T)10.0,
+			const T& NuMesh = (T).3
+		)
+			: bRecordMetric(bRecordMetricIn), MeshConstraints(InMesh)
+		{
+			Measure.Init((T)0., MeshConstraints.Num());
+			Lambda = EMesh * NuMesh / (((T)1. + NuMesh) * ((T)1. - (T)2. * NuMesh));
+			Mu = EMesh / ((T)2. * ((T)1. + NuMesh));
+
+			MuElementArray.Init(Mu, MeshConstraints.Num());
+			LambdaElementArray.Init(Lambda, MeshConstraints.Num());
+
+			InitializeCodimensionData(InParticles);
+		}
+
+		virtual ~FGaussSeidelLinearCodimensionalConstraints() {}
+
+		PMatrix<T, 3, 3> DsInit(const int32 E, const ParticleType& InParticles)  
+		{
+			PMatrix<T, 3, 3> Result((T)0.);
+			for (int32 i = 0; i < 3; i++) 
+			{
+				for (int32 c = 0; c < 3; c++) 
+				{
+					Result.SetAt(c, i, InParticles.X(MeshConstraints[E][i + 1])[c] - InParticles.X(MeshConstraints[E][0])[c]);
+				}
+			}
+			return Result;
+		}
+
+		PMatrix<T, 3, 2> Ds(const int32 E, const ParticleType& InParticles) const 
+		{
+			if (INDEX_NONE < E && E< MeshConstraints.Num()
+				&& INDEX_NONE < MeshConstraints[E][0] && MeshConstraints[E][0] < (int32)InParticles.Size()
+				&& INDEX_NONE < MeshConstraints[E][1] && MeshConstraints[E][1] < (int32)InParticles.Size()
+				&& INDEX_NONE < MeshConstraints[E][2] && MeshConstraints[E][2] < (int32)InParticles.Size())
+			{
+				const TVec3<T> P1P0 = InParticles.P(MeshConstraints[E][1]) - InParticles.P(MeshConstraints[E][0]);
+				const TVec3<T> P2P0 = InParticles.P(MeshConstraints[E][2]) - InParticles.P(MeshConstraints[E][0]);
+
+				return PMatrix<T, 3, 2>(
+					P1P0[0], P1P0[1], P1P0[2],
+					P2P0[0], P2P0[1], P2P0[2]);
+			}
+			else
+			{
+				return PMatrix<T, 3, 2>(
+					(T)0., (T)0, (T)0,
+					(T)0, (T)0, (T)0);
+			}
+		}
+
+		PMatrix<T, 3, 2> F(const int32 E, const ParticleType& InParticles) const 
+		{
+			if (INDEX_NONE < E && E < DmInverse.Num())
+			{
+				return Ds(E, InParticles) * DmInverse[E];
+			}
+			else
+			{
+				return PMatrix<T, 3, 2>(
+					(T)0., (T)0, (T)0,
+					(T)0, (T)0, (T)0);
+			}
+		}
+
+		TArray<TArray<int32>> GetConstraintsArray() const
+		{
+			TArray<TArray<int32>> Constraints;
+			Constraints.SetNum(MeshConstraints.Num());
+			for (int32 i = 0; i < MeshConstraints.Num(); i++)
+			{
+				Constraints[i].SetNum(3);
+				for (int32 j = 0; j < 3; j++)
+				{
+					Constraints[i][j] = MeshConstraints[i][j];
+				}
+			}
+			return Constraints;
+		}
+
+	protected:
+
+		void InitializeCodimensionData(const ParticleType& Particles)
+		{
+			Measure.Init((T)0., MeshConstraints.Num());
+			DmInverse.Init(PMatrix<FSolverReal, 2, 2>(0.f, 0.f, 0.f), MeshConstraints.Num());
+			for (int32 e = 0; e < MeshConstraints.Num(); e++)
+			{
+				check(MeshConstraints[e][0] < (int32)Particles.Size() && MeshConstraints[e][0] > INDEX_NONE);
+				check(MeshConstraints[e][1] < (int32)Particles.Size() && MeshConstraints[e][1] > INDEX_NONE);
+				check(MeshConstraints[e][2] < (int32)Particles.Size() && MeshConstraints[e][2] > INDEX_NONE);
+				if (MeshConstraints[e][0] < (int32)Particles.Size() && MeshConstraints[e][0] > INDEX_NONE
+					&& MeshConstraints[e][1] < (int32)Particles.Size() && MeshConstraints[e][1] > INDEX_NONE
+					&& MeshConstraints[e][2] < (int32)Particles.Size() && MeshConstraints[e][2] > INDEX_NONE)
+				{
+					const TVec3<T> X1X0 = Particles.GetX(MeshConstraints[e][1]) - Particles.GetX(MeshConstraints[e][0]);
+					const TVec3<T> X2X0 = Particles.GetX(MeshConstraints[e][2]) - Particles.GetX(MeshConstraints[e][0]);
+					PMatrix<T, 2, 2> Dm((T)0., (T)0., (T)0.);
+					Dm.M[0] = X1X0.Size();
+					Dm.M[2] = X1X0.Dot(X2X0) / Dm.M[0];
+					Dm.M[3] = Chaos::TVector<T, 3>::CrossProduct(X1X0, X2X0).Size() / Dm.M[0];
+					Measure[e] = Chaos::TVector<T, 3>::CrossProduct(X1X0, X2X0).Size() / (T)2.;
+					ensureMsgf(Measure[e] > (T)0., TEXT("Degenerate triangle detected"));
+
+					PMatrix<T, 2, 2> DmInv = Dm.Inverse();
+					DmInverse[e] = DmInv;
+				}
+			}
+		}
+
+		static T SafeRecip(const T Len, const T Fallback)
+		{
+			if (Len > (T)UE_SMALL_NUMBER)
+			{
+				return (T)1. / Len;
+			}
+			return Fallback;
+		}
+
+	public:
+		void AddHyperelasticResidualAndHessian(const ParticleType& Particles, const int32 ElementIndex, const int32 ElementIndexLocal, const T Dt, TVec3<T>& ParticleResidual, Chaos::PMatrix<T, 3, 3>& ParticleHessian)
+		{
+			check(ElementIndex < DmInverse.Num() && ElementIndex > INDEX_NONE);
+			check(ElementIndex < MuElementArray.Num());
+			check(ElementIndex < Measure.Num());
+			check(ElementIndex < LambdaElementArray.Num());
+			check(ElementIndexLocal < 3 && ElementIndexLocal > INDEX_NONE);
+			if (ElementIndexLocal < 3 
+				&& ElementIndexLocal > INDEX_NONE
+				&& ElementIndex < Measure.Num()
+				&& ElementIndex < MuElementArray.Num()
+				&& ElementIndex < LambdaElementArray.Num()
+				&& ElementIndex < DmInverse.Num() 
+				&& ElementIndex > INDEX_NONE)
+			{
+				const Chaos::PMatrix<T, 2, 2> DmInvT = DmInverse[ElementIndex].GetTransposed(), DmInv = DmInverse[ElementIndex]; 
+				const Chaos::PMatrix<T, 3, 2> Fe = F(ElementIndex, Particles);
+
+				PMatrix<T, 3, 2> Pe(TVec3<T>((T)0.), TVec3<T>((T)0.)), ForceTerm(TVec3<T>((T)0.), TVec3<T>((T)0.));
+				Pe = T(2) * MuElementArray[ElementIndex] * Fe; 
+
+				ForceTerm = -Measure[ElementIndex] * Pe * DmInverse[ElementIndex].GetTransposed();
+
+				Chaos::TVector<T, 3> Dx((T)0.);
+				if (ElementIndexLocal > 0)
+				{
+					for (int32 c = 0; c < 3; c++)
+					{
+						Dx[c] += ForceTerm.M[ElementIndexLocal * 3 - 3 + c];
+					}
+				}
+				else
+				{
+					for (int32 c = 0; c < 3; c++)
+					{
+						for (int32 h = 0; h < 2; h++)
+						{
+							Dx[c] -= ForceTerm.M[h * 3 + c];
+						}
+					}
+				}
+
+				Dx *= Dt * Dt;
+				ParticleResidual -= Dx;
+				T Coeff = Dt * Dt * Measure[ElementIndex];
+
+				if (ElementIndexLocal == 0) 
+				{
+					T DmInvsum = T(0);
+					for (int32 nu = 0; nu < 2; nu++) 
+					{
+						T localDmsum = T(0);
+						for (int32 k = 0; k < 2; k++) 
+						{
+							localDmsum += DmInv.M[nu * 2 + k];
+						}
+						DmInvsum += localDmsum * localDmsum;
+					}
+					for (int32 alpha = 0; alpha < 3; alpha++) 
+					{
+						ParticleHessian.SetAt(alpha, alpha, ParticleHessian.GetAt(alpha, alpha) + Coeff * T(2) * MuElementArray[ElementIndex] * DmInvsum);
+					}
+				}
+				else 
+				{
+					T DmInvsum = T(0);
+					for (int32 nu = 0; nu < 2; nu++)
+					{
+						DmInvsum += DmInv.GetAt(ElementIndexLocal - 1, nu) * DmInv.GetAt(ElementIndexLocal - 1, nu);
+					}
+					for (int32 alpha = 0; alpha < 3; alpha++)
+					{
+						ParticleHessian.SetAt(alpha, alpha, ParticleHessian.GetAt(alpha, alpha) + Dt * Dt * Measure[ElementIndex] * T(2) * MuElementArray[ElementIndex] * DmInvsum);
+					}
+				}
+			}
+		}
+
+	protected:
+		mutable TArray<FSolverMatrix22> DmInverse;
+
+		//material constants calculated from E:
+		T Mu;
+		T Lambda;
+		TArray<T> MuElementArray;
+		TArray<T> LambdaElementArray;
+		TArray<T> AlphaJArray;
+		bool bRecordMetric;
+
+		TArray<TVector<int32, 3>> MeshConstraints;
+		mutable TArray<T> Measure;
+	};
+
+}  // End namespace Chaos::Softs
+
+
