Fix issue where selecting new item in geom collection fracture mode makes previous selected item disappear.

diff --git a/Engine/Source/Runtime/Experimental/GeometryCollectionEngine/Private/GeometryCollection/GeometryCollectionComponent.cpp b/Engine/Source/Runtime/Experimental/GeometryCollectionEngine/Private/GeometryCollection/GeometryCollectionComponent.cpp
--- a/Engine/Source/Runtime/Experimental/GeometryCollectionEngine/Private/GeometryCollection/GeometryCollectionComponent.cpp
+++ b/Engine/Source/Runtime/Experimental/GeometryCollectionEngine/Private/GeometryCollection/GeometryCollectionComponent.cpp
@@ -1709,6 +1709,7 @@ void UGeometryCollectionComponent::ForceNativeRendering(bool bForce)
 		bForceNativeRenderer = bForce;
 		UnregisterCustomRenderer();
 		RegisterCustomRenderer();
+		RefreshCustomRenderer();
 		MarkRenderStateDirty();
 	}
 }
