Fix issue where reregistering custom renderers can apply to garbage objects and generate additional custom render local ISMs attached to the actor.

diff --git a/Engine/Source/Runtime/Experimental/GeometryCollectionEngine/Private/GeometryCollection/GeometryCollectionComponent.cpp b/Engine/Source/Runtime/Experimental/GeometryCollectionEngine/Private/GeometryCollection/GeometryCollectionComponent.cpp
--- a/Engine/Source/Runtime/Experimental/GeometryCollectionEngine/Private/GeometryCollection/GeometryCollectionComponent.cpp
+++ b/Engine/Source/Runtime/Experimental/GeometryCollectionEngine/Private/GeometryCollection/GeometryCollectionComponent.cpp
@@ -3610,11 +3610,11 @@ void UGeometryCollectionComponent::UnregisterCustomRenderer()
 
 void UGeometryCollectionComponent::ReregisterAllCustomRenderers()
 {
-	for (TObjectIterator<UGeometryCollectionComponent> It; It; ++It)
+	for (TObjectIterator<UGeometryCollectionComponent> It(RF_ClassDefaultObject, true, EInternalObjectFlags::Garbage); It; ++It)
 	{
 		It->UnregisterCustomRenderer();
 	}
-	for (TObjectIterator<UGeometryCollectionComponent> It; It; ++It)
+	for (TObjectIterator<UGeometryCollectionComponent> It(RF_ClassDefaultObject, true, EInternalObjectFlags::Garbage); It; ++It)
 	{
 		It->RegisterCustomRenderer();
 		It->RefreshCustomRenderer();
