Added missing command descriptions.

#rnx

diff --git a/Engine/Source/Runtime/Experimental/IoStore/OnDemand/Private/Tool/Download.cpp b/Engine/Source/Runtime/Experimental/IoStore/OnDemand/Private/Tool/Download.cpp
--- a/Engine/Source/Runtime/Experimental/IoStore/OnDemand/Private/Tool/Download.cpp
+++ b/Engine/Source/Runtime/Experimental/IoStore/OnDemand/Private/Tool/Download.cpp
@@ -57,7 +57,7 @@ extern FArgumentSet S3Arguments;
 static FCommand DownloadCommand(
 	DownloadCommandEntry,
 	TEXT("Download"),
-	TEXT(""),
+	TEXT("Fetches the cloud-stored contents of a given on-demand TOC"),
 	{
 		TArgument<FStringView>(TEXT("TocPath"),				TEXT("Bucket-relative path of the TOC to download")),
 		TArgument<FStringView>(TEXT("-Directory"),			TEXT("Output directory")),
diff --git a/Engine/Source/Runtime/Experimental/IoStore/OnDemand/Private/Tool/Upload.cpp b/Engine/Source/Runtime/Experimental/IoStore/OnDemand/Private/Tool/Upload.cpp
--- a/Engine/Source/Runtime/Experimental/IoStore/OnDemand/Private/Tool/Upload.cpp
+++ b/Engine/Source/Runtime/Experimental/IoStore/OnDemand/Private/Tool/Upload.cpp
@@ -133,8 +133,8 @@ extern FArgumentSet S3Arguments;
 
 static FCommand UploadCommand(
 	UploadCommandEntry,
-	TEXT("upload"),
-	TEXT(""),
+	TEXT("Upload"),
+	TEXT("Used to upload IoStore containers to the cloud and convert to on-demand"),
 	{
 		TArgument<FStringView>(TEXT("ContainerGlob"),	TEXT("Path globbed to discover input containers")),
 		TArgument<FStringView>(TEXT("-CryptoKeys"),		TEXT("JSON-format keyring for input containers")),
