Removed latency injector. That can be done using more sophisticated means.

#rnx

diff --git a/Engine/Source/Runtime/Experimental/IoStore/HttpClient/Private/Http.cpp b/Engine/Source/Runtime/Experimental/IoStore/HttpClient/Private/Http.cpp
--- a/Engine/Source/Runtime/Experimental/IoStore/HttpClient/Private/Http.cpp
+++ b/Engine/Source/Runtime/Experimental/IoStore/HttpClient/Private/Http.cpp
@@ -25,8 +25,6 @@
 
 #include <atomic>
 
-#include "LatencyInjector.h"
-
 #include "Misc.inl"
 #include "Socket.inl"
 #include "Peer.inl"
diff --git a/Engine/Source/Runtime/Experimental/IoStore/HttpClient/Private/LatencyInjector.cpp b/Engine/Source/Runtime/Experimental/IoStore/HttpClient/Private/LatencyInjector.cpp
deleted file mode 100644
diff --git a/Engine/Source/Runtime/Experimental/IoStore/HttpClient/Private/LatencyInjector.h b/Engine/Source/Runtime/Experimental/IoStore/HttpClient/Private/LatencyInjector.h
deleted file mode 100644
diff --git a/Engine/Source/Runtime/Experimental/IoStore/HttpClient/Private/Loop.inl b/Engine/Source/Runtime/Experimental/IoStore/HttpClient/Private/Loop.inl
--- a/Engine/Source/Runtime/Experimental/IoStore/HttpClient/Private/Loop.inl
+++ b/Engine/Source/Runtime/Experimental/IoStore/HttpClient/Private/Loop.inl
@@ -311,12 +311,6 @@ static FOutcome DoRecvContent(FActivity* Activity, FPeerType& Peer, int32& MaxRe
 		MaxRecvSize -= Result;
 	}
 
-	if (!FLatencyInjector::Begin(FLatencyInjector::EType::Network, Activity->StateParam))
-	{
-		Activity_SetError(Activity, "Forced random failure");
-		return FOutcome::Error(Activity->ErrorReason);
-	}
-
 #if IAS_HTTP_WITH_PERF
 	Activity->Stopwatch.RecvEnd();
 #endif
@@ -351,11 +345,6 @@ static FOutcome DoRecv(FActivity* Activity, FPeerType& Peer, int32& MaxRecvSize)
 ////////////////////////////////////////////////////////////////////////////////
 static void DoRecvDone(FActivity* Activity)
 {
-	if (!FLatencyInjector::HasExpired(Activity->StateParam))
-	{
-		return;
-	}
-
 	// Notify the user we've received everything
 	FTicketStatus& SinkArg = *(FTicketStatus*)Activity;
 	Activity->Sink(SinkArg);
