Iris - Remove fragment NetObjectFilter type. The complexity of adding new functionality to filtering is becoming too great when everything needs to support multiple filtering passes.

diff --git a/Engine/Plugins/Runtime/ReplicationSystemTestPlugin/Source/Private/Tests/ReplicationSystem/Filtering/MockNetObjectFilter.cpp b/Engine/Plugins/Runtime/ReplicationSystemTestPlugin/Source/Private/Tests/ReplicationSystem/Filtering/MockNetObjectFilter.cpp
--- a/Engine/Plugins/Runtime/ReplicationSystemTestPlugin/Source/Private/Tests/ReplicationSystem/Filtering/MockNetObjectFilter.cpp
+++ b/Engine/Plugins/Runtime/ReplicationSystemTestPlugin/Source/Private/Tests/ReplicationSystem/Filtering/MockNetObjectFilter.cpp
@@ -137,79 +137,6 @@ void UMockNetObjectFilter::PostFilter(FNetObjectPostFilteringParams& Params)
 	CallStatus.SuccessfulCallCounts.PostFilter += bIsProperCall;
 }
 
-
-//**************************************************************************************************
-// UMockNetObjectFilterUsingFragmentData
-//**************************************************************************************************
-
-void UMockNetObjectFilterUsingFragmentData::OnInit(FNetObjectFilterInitParams& Params)
-{
-	Super::OnInit(Params);
-
-	SetupFilterType(ENetFilterType::PostPoll_FragmentBased);
-}
-
-bool UMockNetObjectFilterUsingFragmentData::AddObject(uint32 ObjectIndex, FNetObjectFilterAddObjectParams& Params)
-{
-	bool bResult = Super::AddObject(ObjectIndex, Params);
-
-	if (bResult)
-	{
-		// Check if object has the NetTest_FilterOut RepTag which says whether the object should be filtered out or not.
-		UE::Net::FRepTagFindInfo RepTagInfo;
-		if (UE::Net::FindRepTag(Params.Protocol, RepTag_NetTest_FilterOut, RepTagInfo))
-		{
-			// Warning: Using the internal state requires the proper NetSerializer and dequantization to get the value.
-			ObjectToFilterOutOffset.Add(ObjectIndex, RepTagInfo.InternalStateAbsoluteOffset);
-		}
-		else
-		{
-			ensureAlwaysMsgf(false, TEXT("Could not find FilterOut reptag in object index %u"), ObjectIndex);
-		}
-	}
-
-	return bResult;
-}
-
-void UMockNetObjectFilterUsingFragmentData::RemoveObject(uint32 ObjectIndex, const FNetObjectFilteringInfo& Params)
-{
-	Super::RemoveObject(ObjectIndex, Params);
-
-	ObjectToFilterOut.Remove(ObjectIndex);
-	ObjectToFilterOutOffset.Remove(ObjectIndex);
-}
-
-void UMockNetObjectFilterUsingFragmentData::UpdateObjects(FNetObjectFilterUpdateParams& Params)
-{
-	++CallStatus.CallCounts.UpdateObjects;
-
-	bool bIsProperCall = Params.ObjectCount <= AddedCount;
-	for (uint32 ObjectIndex : MakeArrayView(Params.ObjectIndices, Params.ObjectCount))
-	{
-		bIsProperCall = bIsProperCall && AddedObjectIndices.GetBit(ObjectIndex);
-
-		if (UPTRINT* InternalStateOffset = ObjectToFilterOutOffset.Find(ObjectIndex))
-		{
-			// Warning: This is generally not safe, not even for primitive types.
-			const bool bFilterOut = *reinterpret_cast<const bool*>((*Params.StateBuffers)[ObjectIndex] + *InternalStateOffset);
-			ObjectToFilterOut.Emplace(ObjectIndex, bFilterOut);
-		}
-	}
-	CallStatus.SuccessfulCallCounts.UpdateObjects += bIsProperCall;
-}
-
-void UMockNetObjectFilterUsingFragmentData::Filter(FNetObjectFilteringParams& Params)
-{
-	Super::Filter(Params);
-
-	// Go through objects with FilterOut wishes
-	for (const auto& Pair : ObjectToFilterOut)
-	{
-		Params.OutAllowedObjects.SetBitValue(Pair.Key, !Pair.Value);
-	}
-}
-
-
 //**************************************************************************************************
 // UMockNetObjectFilterWithCondition
 //**************************************************************************************************
@@ -218,8 +145,6 @@ void UMockNetObjectFilterWithCondition::OnInit(FNetObjectFilterInitParams& Param
 {
 	Super::OnInit(Params);
 
-	SetupFilterType(ENetFilterType::PrePoll_Raw);
-
 	ReplicationSystem = Params.ReplicationSystem.Get();
 }
 
diff --git a/Engine/Plugins/Runtime/ReplicationSystemTestPlugin/Source/Private/Tests/ReplicationSystem/Filtering/MockNetObjectFilter.h b/Engine/Plugins/Runtime/ReplicationSystemTestPlugin/Source/Private/Tests/ReplicationSystem/Filtering/MockNetObjectFilter.h
--- a/Engine/Plugins/Runtime/ReplicationSystemTestPlugin/Source/Private/Tests/ReplicationSystem/Filtering/MockNetObjectFilter.h
+++ b/Engine/Plugins/Runtime/ReplicationSystemTestPlugin/Source/Private/Tests/ReplicationSystem/Filtering/MockNetObjectFilter.h
@@ -88,28 +88,6 @@ protected:
 	static constexpr float DefaultPriority = 1.0f;
 };
 
-/**
- * Filter that reads fragment data to decide if an object is filtered out or not
- */
-UCLASS()
-class UMockNetObjectFilterUsingFragmentData : public UMockNetObjectFilter
-{
-	GENERATED_BODY()
-
-protected:
-
-	virtual void OnInit(FNetObjectFilterInitParams&) override;
-	virtual bool AddObject(uint32 ObjectIndex, FNetObjectFilterAddObjectParams&) override;
-	virtual void RemoveObject(uint32 ObjectIndex, const FNetObjectFilteringInfo&) override;
-	virtual void UpdateObjects(FNetObjectFilterUpdateParams&) override;
-	virtual void Filter(FNetObjectFilteringParams&) override;
-
-private:
-
-	TMap<uint32, UPTRINT> ObjectToFilterOutOffset;
-	TMap<uint32, bool> ObjectToFilterOut;
-};
-
 /**
  * Filter that checks object data to decide to filter out an object.
  */
diff --git a/Engine/Plugins/Runtime/ReplicationSystemTestPlugin/Source/Private/Tests/ReplicationSystem/Filtering/TestFiltering.cpp b/Engine/Plugins/Runtime/ReplicationSystemTestPlugin/Source/Private/Tests/ReplicationSystem/Filtering/TestFiltering.cpp
--- a/Engine/Plugins/Runtime/ReplicationSystemTestPlugin/Source/Private/Tests/ReplicationSystem/Filtering/TestFiltering.cpp
+++ b/Engine/Plugins/Runtime/ReplicationSystemTestPlugin/Source/Private/Tests/ReplicationSystem/Filtering/TestFiltering.cpp
@@ -41,16 +41,6 @@ protected:
 		MockNetObjectFilter->ResetFunctionCallStatus();
 	}
 
-	void SetDynamicFragmentFilterStatus(ENetFilterStatus FilterStatus)
-	{
-		UMockNetObjectFilterUsingFragmentData::FFunctionCallSetup CallSetup;
-		CallSetup.AddObject.bReturnValue = true;
-		CallSetup.Filter.bFilterOutByDefault = FilterStatus == ENetFilterStatus::Disallow;
-
-		MockNetObjectFilterWithFragments->SetFunctionCallSetup(CallSetup);
-		MockNetObjectFilterWithFragments->ResetFunctionCallStatus();
-	}
-
 private:
 	void InitNetObjectFilterDefinitions()
 	{
@@ -71,13 +61,6 @@ private:
 			MockDefinition.ConfigClassName = "/Script/ReplicationSystemTestPlugin.MockNetObjectFilterConfig";
 		}
 
-		{
-			FNetObjectFilterDefinition& MockDefinition = NewFilterDefinitions.Emplace_GetRef();
-			MockDefinition.FilterName = "MockFilterWithFragments";
-			MockDefinition.ClassName = "/Script/ReplicationSystemTestPlugin.MockNetObjectFilterUsingFragmentData";
-			MockDefinition.ConfigClassName = "/Script/ReplicationSystemTestPlugin.MockNetObjectFilterConfig";
-		}
-
 		DefinitionsProperty->CopyCompleteValue((void*)(UPTRINT(FilterDefinitions) + DefinitionsProperty->GetOffset_ForInternal()), &NewFilterDefinitions);
 	}
 
@@ -92,27 +75,18 @@ private:
 
 		MockFilterHandle = InvalidNetObjectFilterHandle;
 		MockNetObjectFilter = nullptr;
-
-		MockFilterWithFragmentsHandle = InvalidNetObjectFilterHandle;
-		MockNetObjectFilterWithFragments = nullptr;
 	}
 
 	void InitMockNetObjectFilter()
 	{
 		MockNetObjectFilter = CastChecked<UMockNetObjectFilter>(Server->GetReplicationSystem()->GetFilter(GetMockFilterName()));
 		MockFilterHandle = Server->GetReplicationSystem()->GetFilterHandle(GetMockFilterName());
-
-		MockNetObjectFilterWithFragments = ExactCast<UMockNetObjectFilterUsingFragmentData>(Server->GetReplicationSystem()->GetFilter("MockFilterWithFragments"));
-		MockFilterWithFragmentsHandle = Server->GetReplicationSystem()->GetFilterHandle("MockFilterWithFragments");
 	}
 
 protected:
 	UMockNetObjectFilter* MockNetObjectFilter;
 	FNetObjectFilterHandle MockFilterHandle;
 
-	UMockNetObjectFilterUsingFragmentData* MockNetObjectFilterWithFragments;
-	FNetObjectFilterHandle MockFilterWithFragmentsHandle;
-
 private:
 	TArray<FNetObjectFilterDefinition> OriginalFilterDefinitions;
 };
@@ -1110,101 +1084,6 @@ UE_NET_TEST_FIXTURE(FTestFilteringFixture, NestedDependentObjectIsFilteredAsPare
 	UE_NET_ASSERT_NE(Client->GetReplicationBridge()->GetReplicatedObject(ServerNestedDependentObject->NetRefHandle), nullptr);
 }
 
-UE_NET_TEST_FIXTURE(FTestFilteringFixture, ObjectGetsFilterOutSettingFromStart)
-{
-	// Setup dynamic filter for the test. For this test we want the value of the object's filter property to rule.
-	SetDynamicFragmentFilterStatus(ENetFilterStatus::Allow);
-
-	// Add client
-	FReplicationSystemTestClient* Client = CreateClient();
-
-	// Create object and set filter
-	UTestFilteringObject* ServerObject = Server->CreateObject<UTestFilteringObject>();
-	Server->ReplicationSystem->SetFilter(ServerObject->NetRefHandle, MockFilterWithFragmentsHandle);
-
-	// We want the object to be filtered out
-	constexpr bool bFilterOut = true;
-	ServerObject->SetFilterOut(bFilterOut);
-
-	// Send and deliver packets
-	Server->PreSendUpdate();
-	Server->SendAndDeliverTo(Client, DeliverPacket);
-	Server->PostSendUpdate();
-
-	// Check that the object does not exist on the client.
-	UE_NET_ASSERT_EQ(Client->GetReplicationBridge()->GetReplicatedObject(ServerObject->NetRefHandle), nullptr);
-}
-
-UE_NET_TEST_FIXTURE(FTestFilteringFixture, ObjectGetsUpdatedFilterOutSetting)
-{
-	// Setup dynamic filter for the test. For this test we want the value of the object's filter property to rule.
-	SetDynamicFragmentFilterStatus(ENetFilterStatus::Allow);
-
-	// Add client
-	FReplicationSystemTestClient* Client = CreateClient();
-
-	// Create object and set filter
-	UTestFilteringObject* ServerObject = Server->CreateObject<UTestFilteringObject>();
-	Server->ReplicationSystem->SetFilter(ServerObject->NetRefHandle, MockFilterWithFragmentsHandle);
-
-	// We don't want the object to be filtered out
-	ServerObject->SetFilterOut(false);
-
-	// Send and deliver packets
-	Server->PreSendUpdate();
-	Server->SendAndDeliverTo(Client, DeliverPacket);
-	Server->PostSendUpdate();
-
-	// Check that the object exists on the client.
-	UE_NET_ASSERT_NE(Client->GetReplicationBridge()->GetReplicatedObject(ServerObject->NetRefHandle), nullptr);
-
-	// Mark the object to be filtered out.
-	ServerObject->SetFilterOut(true);
-
-	// Send and deliver packets
-	Server->PreSendUpdate();
-	Server->SendAndDeliverTo(Client, DeliverPacket);
-	Server->PostSendUpdate();
-
-	// Check that the object does not exist on the client.
-	UE_NET_ASSERT_EQ(Client->GetReplicationBridge()->GetReplicatedObject(ServerObject->NetRefHandle), nullptr);
-}
-
-UE_NET_TEST_FIXTURE(FTestFilteringFixture, MixPreAndPostFilters)
-{
-	// Setup dynamic filters for the test.
-	SetDynamicFilterStatus(ENetFilterStatus::Disallow);
-	SetDynamicFragmentFilterStatus(ENetFilterStatus::Disallow);
-
-	// Add client
-	FReplicationSystemTestClient* Client = CreateClient();
-
-	// Create the objects and set filter
-	UTestFilteringObject* ServerObjectPreFilter = Server->CreateObject<UTestFilteringObject>();
-	Server->ReplicationSystem->SetFilter(ServerObjectPreFilter->NetRefHandle, MockFilterHandle);
-
-	UTestFilteringObject* ServerObjectPostFilter = Server->CreateObject<UTestFilteringObject>();
-	Server->ReplicationSystem->SetFilter(ServerObjectPostFilter->NetRefHandle, MockFilterWithFragmentsHandle);
-
-	// Create a non-filtered object
-	UTestFilteringObject* ServerObjectNoFilter = Server->CreateObject<UTestFilteringObject>();
-
-	// We want the objects to be filtered out
-	constexpr bool bFilterOut = true;
-	ServerObjectPreFilter->SetFilterOut(bFilterOut);
-	ServerObjectPostFilter->SetFilterOut(bFilterOut);
-
-	// Send and deliver packets
-	Server->PreSendUpdate();
-	Server->SendAndDeliverTo(Client, DeliverPacket);
-	Server->PostSendUpdate();
-
-	// Check that the object does not exist on the client.
-	UE_NET_ASSERT_EQ(Client->GetReplicationBridge()->GetReplicatedObject(ServerObjectPreFilter->NetRefHandle), nullptr);
-	UE_NET_ASSERT_EQ(Client->GetReplicationBridge()->GetReplicatedObject(ServerObjectPostFilter->NetRefHandle), nullptr);
-	UE_NET_ASSERT_NE(Client->GetReplicationBridge()->GetReplicatedObject(ServerObjectNoFilter->NetRefHandle), nullptr);
-}
-
 UE_NET_TEST_FIXTURE(FTestFilteringFixture, DynamicFilteredOutSubObjectsAreResetWhenIndexIsReused)
 {
 	// Setup dynamic filters for the test
diff --git a/Engine/Plugins/Runtime/ReplicationSystemTestPlugin/Source/Private/Tests/ReplicationSystem/Filtering/TestFilteringObject.cpp b/Engine/Plugins/Runtime/ReplicationSystemTestPlugin/Source/Private/Tests/ReplicationSystem/Filtering/TestFilteringObject.cpp
--- a/Engine/Plugins/Runtime/ReplicationSystemTestPlugin/Source/Private/Tests/ReplicationSystem/Filtering/TestFilteringObject.cpp
+++ b/Engine/Plugins/Runtime/ReplicationSystemTestPlugin/Source/Private/Tests/ReplicationSystem/Filtering/TestFilteringObject.cpp
@@ -21,17 +21,3 @@ void UTestFilteringObject::GetLifetimeReplicatedProps( TArray< class FLifetimePr
 	DOREPLIFETIME(ThisClass, NetTest_FilterOut);
 	DOREPLIFETIME(ThisClass, ReplicatedCounter);
 }
-
-
-void UTestLocationFragmentFilteringObject::RegisterReplicationFragments(UE::Net::FFragmentRegistrationContext& Context, UE::Net::EFragmentRegistrationFlags RegistrationFlags)
-{
-	// Build descriptors and allocate PropertyReplicationFragments for this object
-	UE::Net::FReplicationFragmentUtil::CreateAndRegisterFragmentsForObject(this, Context, RegistrationFlags);
-}
-
-void UTestLocationFragmentFilteringObject::GetLifetimeReplicatedProps(TArray< class FLifetimeProperty >& OutLifetimeProps) const
-{
-	Super::GetLifetimeReplicatedProps(OutLifetimeProps);
-	DOREPLIFETIME(ThisClass, WorldLocation);
-	DOREPLIFETIME(ThisClass, NetCullDistanceSquared);
-}
diff --git a/Engine/Plugins/Runtime/ReplicationSystemTestPlugin/Source/Private/Tests/ReplicationSystem/Filtering/TestFilteringObject.h b/Engine/Plugins/Runtime/ReplicationSystemTestPlugin/Source/Private/Tests/ReplicationSystem/Filtering/TestFilteringObject.h
--- a/Engine/Plugins/Runtime/ReplicationSystemTestPlugin/Source/Private/Tests/ReplicationSystem/Filtering/TestFilteringObject.h
+++ b/Engine/Plugins/Runtime/ReplicationSystemTestPlugin/Source/Private/Tests/ReplicationSystem/Filtering/TestFilteringObject.h
@@ -25,25 +25,3 @@ private:
 	UPROPERTY(Replicated)
 	bool NetTest_FilterOut = false;
 };
-
-
-
-UCLASS()
-class UTestLocationFragmentFilteringObject : public UReplicatedTestObject
-{
-	GENERATED_BODY()
-
-public:
-	
-protected:
-	virtual void RegisterReplicationFragments(UE::Net::FFragmentRegistrationContext& Fragments, UE::Net::EFragmentRegistrationFlags RegistrationFlags) override;
-
-public:
-	
-	UPROPERTY(Replicated)
-	FVector WorldLocation;
-
-	UPROPERTY(Replicated)
-	float NetCullDistanceSquared = 0.0f;
-};
-
diff --git a/Engine/Plugins/Runtime/ReplicationSystemTestPlugin/Source/Private/Tests/ReplicationSystem/Filtering/TestGridFilter.cpp b/Engine/Plugins/Runtime/ReplicationSystemTestPlugin/Source/Private/Tests/ReplicationSystem/Filtering/TestGridFilter.cpp
--- a/Engine/Plugins/Runtime/ReplicationSystemTestPlugin/Source/Private/Tests/ReplicationSystem/Filtering/TestGridFilter.cpp
+++ b/Engine/Plugins/Runtime/ReplicationSystemTestPlugin/Source/Private/Tests/ReplicationSystem/Filtering/TestGridFilter.cpp
@@ -52,13 +52,6 @@ private:
 			GridWorldLocDef.ConfigClassName = "/Script/IrisCore.NetObjectGridFilterConfig";
 		}
 
-		{
-			FNetObjectFilterDefinition& GridFragmentDef = NewFilterDefinitions.Emplace_GetRef();
-			GridFragmentDef.FilterName = "NetObjectGridFragmentLocFilter";
-			GridFragmentDef.ClassName = "/Script/IrisCore.NetObjectGridFragmentLocFilter";
-			GridFragmentDef.ConfigClassName = "/Script/IrisCore.NetObjectGridFilterConfig";
-		}
-
 		DefinitionsProperty->CopyCompleteValue((void*)(UPTRINT(FilterDefinitions) + DefinitionsProperty->GetOffset_ForInternal()), &NewFilterDefinitions);
 	}
 
@@ -73,27 +66,17 @@ private:
 
 		WorldLocFilterHandle = InvalidNetObjectFilterHandle;
 		WorldLocFilter = nullptr;
-
-		FragmentLocFilterHandle = InvalidNetObjectFilterHandle;
-		FragmentLocFilter = nullptr;
 	}
 
 	void InitFilterHandles()
 	{
 		WorldLocFilter = ExactCast<UNetObjectGridWorldLocFilter>(Server->GetReplicationSystem()->GetFilter("NetObjectGridWorldLocFilter"));
 		WorldLocFilterHandle = Server->GetReplicationSystem()->GetFilterHandle("NetObjectGridWorldLocFilter");
-
-		FragmentLocFilter = ExactCast<UNetObjectGridFragmentLocFilter>(Server->GetReplicationSystem()->GetFilter("NetObjectGridFragmentLocFilter"));
-		FragmentLocFilterHandle = Server->GetReplicationSystem()->GetFilterHandle("NetObjectGridFragmentLocFilter");
 	}
 
 protected:
 	UNetObjectGridWorldLocFilter* WorldLocFilter;
 	FNetObjectFilterHandle WorldLocFilterHandle;
-
-	UNetObjectGridFragmentLocFilter* FragmentLocFilter;
-	FNetObjectFilterHandle FragmentLocFilterHandle;
-
 private:
 	TArray<FNetObjectFilterDefinition> OriginalFilterDefinitions;
 };
@@ -243,69 +226,6 @@ UE_NET_TEST_FIXTURE(FTestGridFilterFixture, TestWorldLocationIsFrequentlyUpdated
 	UE_NET_ASSERT_NE(Client->GetReplicationBridge()->GetReplicatedObject(ServerObject->NetRefHandle), nullptr);
 }
 
-// Test fragment data filter
-UE_NET_TEST_FIXTURE(FTestGridFilterFixture, TestFragmentFilter)
-{
-	// Spawn object on server
-	UTestLocationFragmentFilteringObject* ServerObjectZero = Server->CreateObject<UTestLocationFragmentFilteringObject>();
-	UTestLocationFragmentFilteringObject* ServerObjectNear = Server->CreateObject<UTestLocationFragmentFilteringObject>();
-	UTestLocationFragmentFilteringObject* ServerObjectLimit = Server->CreateObject<UTestLocationFragmentFilteringObject>();
-	UTestLocationFragmentFilteringObject* ServerObjectCulled = Server->CreateObject<UTestLocationFragmentFilteringObject>();
-	UTestLocationFragmentFilteringObject* ServerObjectVeryFar = Server->CreateObject<UTestLocationFragmentFilteringObject>();
-
-	const UNetObjectGridFilterConfig* DefaultGridConfig = GetDefault<UNetObjectGridFilterConfig>();
-
-	ServerObjectZero->WorldLocation = FVector::ZeroVector;
-	ServerObjectZero->NetCullDistanceSquared = 1500.f * 1500.f;
-
-	ServerObjectNear->WorldLocation = FVector(100.f, 100.f, 100.f);
-	ServerObjectNear->NetCullDistanceSquared = 1500.f * 1500.f;
-	
-	ServerObjectLimit->WorldLocation = FVector(DefaultGridConfig->CellSizeX, DefaultGridConfig->CellSizeY, 1500.f);
-	ServerObjectLimit->NetCullDistanceSquared = DefaultGridConfig->MaxCullDistance * DefaultGridConfig->MaxCullDistance;
-
-	ServerObjectCulled->WorldLocation = FVector(DefaultGridConfig->CellSizeX + 100.f, DefaultGridConfig->CellSizeY + 100.f, 100.f);
-	ServerObjectCulled->NetCullDistanceSquared = 1500.f * 1500.f;
-	
-	ServerObjectVeryFar->WorldLocation = FVector(DefaultGridConfig->CellSizeX + 99999.f, DefaultGridConfig->CellSizeY + 99999.f, 99999.f);
-	ServerObjectVeryFar->NetCullDistanceSquared = 1500.f * 1500.f;
-
-	// Apply grid filter
-	Server->ReplicationSystem->SetFilter(ServerObjectZero->NetRefHandle, FragmentLocFilterHandle);
-	Server->ReplicationSystem->SetFilter(ServerObjectNear->NetRefHandle, FragmentLocFilterHandle);
-	Server->ReplicationSystem->SetFilter(ServerObjectLimit->NetRefHandle, FragmentLocFilterHandle);
-	Server->ReplicationSystem->SetFilter(ServerObjectCulled->NetRefHandle, FragmentLocFilterHandle);
-	Server->ReplicationSystem->SetFilter(ServerObjectVeryFar->NetRefHandle, FragmentLocFilterHandle);
-
-	// Add client
-	FReplicationSystemTestClient* Client = CreateClient();
-
-	// Set the view location of the client to (0,0,0)
-	FReplicationView ReplicationView;
-	FReplicationView::FView View = ReplicationView.Views.Emplace_GetRef();
-	View.Pos = FVector(0.0f, 0.0f, 0.0f);
-	Server->ReplicationSystem->SetReplicationView(Client->ConnectionIdOnServer, ReplicationView);
-
-	// Send and deliver packet
-	Server->PreSendUpdate();
-	Server->SendAndDeliverTo(Client, DeliverPacket);
-	Server->PostSendUpdate();
-
-	// Test visible objects
-	UE_NET_ASSERT_NE(Client->GetReplicationBridge()->GetReplicatedObject(ServerObjectZero->NetRefHandle), nullptr);
-	UE_NET_ASSERT_NE(Client->GetReplicationBridge()->GetReplicatedObject(ServerObjectNear->NetRefHandle), nullptr);
-
-	// Test culled objects
-	UE_NET_ASSERT_EQ(Client->GetReplicationBridge()->GetReplicatedObject(ServerObjectLimit->NetRefHandle), nullptr);
-	UE_NET_ASSERT_EQ(Client->GetReplicationBridge()->GetReplicatedObject(ServerObjectCulled->NetRefHandle), nullptr);
-	UE_NET_ASSERT_EQ(Client->GetReplicationBridge()->GetReplicatedObject(ServerObjectVeryFar->NetRefHandle), nullptr);
-
-	Server->DestroyObject(ServerObjectZero);
-	Server->DestroyObject(ServerObjectLimit);
-	Server->DestroyObject(ServerObjectCulled);
-	Server->DestroyObject(ServerObjectVeryFar);
-}
-
 #if 0
 UE_NET_TEST_FIXTURE(FTestGridFilterFixture, TestFilterPerformance)
 {
diff --git a/Engine/Source/Runtime/Experimental/Iris/Core/Private/Iris/ReplicationSystem/Filtering/NetObjectFilter.cpp b/Engine/Source/Runtime/Experimental/Iris/Core/Private/Iris/ReplicationSystem/Filtering/NetObjectFilter.cpp
--- a/Engine/Source/Runtime/Experimental/Iris/Core/Private/Iris/ReplicationSystem/Filtering/NetObjectFilter.cpp
+++ b/Engine/Source/Runtime/Experimental/Iris/Core/Private/Iris/ReplicationSystem/Filtering/NetObjectFilter.cpp
@@ -21,11 +21,6 @@ UNetObjectFilter::UNetObjectFilter()
 
 void UNetObjectFilter::Init(FNetObjectFilterInitParams& Params)
 {
-	if (Params.Config)
-	{
-		FilterType = Params.Config->FilterType;
-	}
-
 	{
 		UE::Net::Private::FNetObjectFilteringInfoAccessor FilteringInfoAccessor;
 		TArrayView<FNetObjectFilteringInfo> NetObjectFilteringInfos = FilteringInfoAccessor.GetNetObjectFilteringInfos(Params.ReplicationSystem);
diff --git a/Engine/Source/Runtime/Experimental/Iris/Core/Private/Iris/ReplicationSystem/Filtering/NetObjectGridFilter.cpp b/Engine/Source/Runtime/Experimental/Iris/Core/Private/Iris/ReplicationSystem/Filtering/NetObjectGridFilter.cpp
--- a/Engine/Source/Runtime/Experimental/Iris/Core/Private/Iris/ReplicationSystem/Filtering/NetObjectGridFilter.cpp
+++ b/Engine/Source/Runtime/Experimental/Iris/Core/Private/Iris/ReplicationSystem/Filtering/NetObjectGridFilter.cpp
@@ -562,8 +562,6 @@ void UNetObjectGridWorldLocFilter::OnInit(FNetObjectFilterInitParams& Params)
 {
 	Super::OnInit(Params);
 
-	SetupFilterType(ENetFilterType::PrePoll_Raw);
-
 	WorldLocations = &Params.ReplicationSystem->GetWorldLocations();
 }
 
@@ -609,90 +607,3 @@ bool UNetObjectGridWorldLocFilter::BuildObjectInfo(uint32 ObjectIndex, FNetObjec
 	return true;
 }
 
-//*************************************************************************************************
-// UNetObjectGridFragmentLocFilter
-//*************************************************************************************************
-
-void UNetObjectGridFragmentLocFilter::OnInit(FNetObjectFilterInitParams& InitParams)
-{
-	Super::OnInit(InitParams);
-
-	SetupFilterType(ENetFilterType::PostPoll_FragmentBased);
-}
-
-void UNetObjectGridFragmentLocFilter::UpdateObjects(FNetObjectFilterUpdateParams& Params)
-{
-	for (SIZE_T ObjectIt = 0, ObjectEndIt = Params.ObjectCount; ObjectIt != ObjectEndIt; ++ObjectIt)
-	{
-		const uint32 ObjectIndex = Params.ObjectIndices[ObjectIt];
-		const FObjectLocationInfo& ObjectLocationInfo = static_cast<const FObjectLocationInfo&>(Params.FilteringInfos[ObjectIndex]);
-		const UE::Net::FReplicationInstanceProtocol* InstanceProtocol = Params.InstanceProtocols ? Params.InstanceProtocols[ObjectIt] : nullptr;
-		UpdateCellInfoForObject(ObjectLocationInfo, InstanceProtocol);
-	}
-}
-
-bool UNetObjectGridFragmentLocFilter::BuildObjectInfo(uint32 ObjectIndex, FNetObjectFilterAddObjectParams& Params)
-{
-	UE::Net::FRepTagFindInfo WorldLocationTagInfo;
-	if (UE::Net::FindRepTag(Params.Protocol, UE::Net::RepTag_WorldLocation, WorldLocationTagInfo))
-	{
-		// Want to keep the memory footprint minimal so don't allow adding objects whose values would not fit.
-		if (WorldLocationTagInfo.ExternalStateOffset >= MAX_uint16 || WorldLocationTagInfo.StateIndex >= MAX_uint16)
-		{
-			return false;
-		}
-
-		FObjectLocationInfo& ObjectLocationInfo = static_cast<FObjectLocationInfo&>(Params.OutInfo);
-		ObjectLocationInfo.SetLocationStateOffset(static_cast<uint16>(WorldLocationTagInfo.ExternalStateOffset));
-		ObjectLocationInfo.SetLocationStateIndex(static_cast<uint16>(WorldLocationTagInfo.StateIndex));
-
-		// NetCullDistanceSqr is optional. 
-		UE::Net::FRepTagFindInfo NetCullDistanceSqrTagInfo;
-		if (UE::Net::FindRepTag(Params.Protocol, UE::Net::RepTag_CullDistanceSqr, NetCullDistanceSqrTagInfo))
-		{
-			if ((NetCullDistanceSqrTagInfo.ExternalStateOffset < MAX_uint16) && (NetCullDistanceSqrTagInfo.StateIndex < MAX_uint16))
-			{
-				FCullDistanceFragmentInfo FragmentInfo;
-				FragmentInfo.CullDistanceSqrStateIndex = static_cast<uint16>(NetCullDistanceSqrTagInfo.StateIndex);
-				FragmentInfo.CullDistanceSqrStateOffset = static_cast<uint16>(NetCullDistanceSqrTagInfo.ExternalStateOffset);
-
-				CullDistanceFragments.Add(ObjectIndex, MoveTemp(FragmentInfo));
-			}
-		}
-
-		return true;
-	}
-
-	return false;
-}
-
-void UNetObjectGridFragmentLocFilter::OnObjectRemoved(uint32 ObjectIndex)
-{
-	CullDistanceFragments.Remove(ObjectIndex);
-}
-
-void UNetObjectGridFragmentLocFilter::UpdateObjectInfo(UNetObjectGridFilter::FPerObjectInfo& PerObjectInfo, const UNetObjectGridFilter::FObjectLocationInfo& ObjectLocationInfo, const UE::Net::FReplicationInstanceProtocol* InstanceProtocol)
-{
-	check(ObjectLocationInfo.IsUsingWorldLocations() == false);
-	check(InstanceProtocol);
-
-	TArrayView<const UE::Net::FReplicationInstanceProtocol::FFragmentData> FragmentDatas = MakeArrayView(InstanceProtocol->FragmentData, InstanceProtocol->FragmentCount);
-
-	// Update the location
-	{
-		const UE::Net::FReplicationInstanceProtocol::FFragmentData& FragmentData = FragmentDatas[ObjectLocationInfo.GetLocationStateIndex()];
-		const uint8* LocationAddress = FragmentData.ExternalSrcBuffer + ObjectLocationInfo.GetLocationStateOffset();
-		const FVector* Location = reinterpret_cast<const FVector*>(LocationAddress);
-		PerObjectInfo.Position = *Location;
-	}
-
-	// Update the culldistance
-	if( FCullDistanceFragmentInfo* CullDistanceFragmentInfo = CullDistanceFragments.Find(PerObjectInfo.ObjectIndex) )
-	{
-		const UE::Net::FReplicationInstanceProtocol::FFragmentData& FragmentData = FragmentDatas[CullDistanceFragmentInfo->CullDistanceSqrStateIndex];
-		const uint8* CullDistanceSqrAddress = FragmentData.ExternalSrcBuffer + CullDistanceFragmentInfo->CullDistanceSqrStateOffset;
-		const float CullDistanceSqr = *reinterpret_cast<const float*>(CullDistanceSqrAddress);
-		PerObjectInfo.SetCullDistanceSq(CullDistanceSqr);
-	}
-}
-
diff --git a/Engine/Source/Runtime/Experimental/Iris/Core/Private/Iris/ReplicationSystem/Filtering/ReplicationFiltering.cpp b/Engine/Source/Runtime/Experimental/Iris/Core/Private/Iris/ReplicationSystem/Filtering/ReplicationFiltering.cpp
--- a/Engine/Source/Runtime/Experimental/Iris/Core/Private/Iris/ReplicationSystem/Filtering/ReplicationFiltering.cpp
+++ b/Engine/Source/Runtime/Experimental/Iris/Core/Private/Iris/ReplicationSystem/Filtering/ReplicationFiltering.cpp
@@ -103,54 +103,21 @@ public:
 	{
 		uint32* ObjectIndices = nullptr;
 		uint32 ObjectCount = 0;
-
-		FReplicationInstanceProtocol const** InstanceProtocols = nullptr;
 	};
 
-	FUpdateDirtyObjectsBatchHelper(const FNetRefHandleManager* InNetRefHandleManager, const TArray<FFilterInfo>& DynamicFilters, ENetFilterType ActiveFilterType)
+	FUpdateDirtyObjectsBatchHelper(const FNetRefHandleManager* InNetRefHandleManager, const TArray<FFilterInfo>& DynamicFilters)
 		: NetRefHandleManager(InNetRefHandleManager)
 	{
-		bHasProtocolBuffers = (ActiveFilterType == ENetFilterType::PostPoll_FragmentBased);
-
-		uint32 NumActiveFilters = 0;
-		for (const FFilterInfo& Info : DynamicFilters)
-		{
-			if (Info.Type == ActiveFilterType)
-			{
-				NumActiveFilters++;
-			}
-		}
-		ensureMsgf(NumActiveFilters > 0, TEXT("Should never be called when no filters of type %u exist"), ActiveFilterType);
+		const int32 NumFilters = DynamicFilters.Num();
 
-		const int32 TotalFilters = DynamicFilters.Num();
-
-		// We create PerFilterInfo for every dynamic filters, but allocate the memory buffers only for active filters
-		PerFilterInfos.SetNum(TotalFilters, EAllowShrinking::No);
-
-		ObjectIndicesStorage.SetNumUninitialized(NumActiveFilters * MaxObjectCountPerBatch);
+		PerFilterInfos.SetNum(NumFilters, EAllowShrinking::No);
+		ObjectIndicesStorage.SetNumUninitialized(NumFilters * MaxObjectCountPerBatch);
 		
-		if (bHasProtocolBuffers)
-		{
-			InstanceProtocolsStorage.SetNumUninitialized(NumActiveFilters * MaxObjectCountPerBatch);
-		}
-
 		uint32 BufferIndex = 0;
-		for (int32 i=0; i < TotalFilters; ++i)
+		for (FPerFilterInfo& PerFilterInfo : PerFilterInfos)
 		{
-			const FFilterInfo& FilterInfo = DynamicFilters[i];
-		
-			if (FilterInfo.Type == ActiveFilterType)
-			{
-				FPerFilterInfo& PerFilterInfo = PerFilterInfos[i];
-				PerFilterInfo.ObjectIndices = ObjectIndicesStorage.GetData() + BufferIndex * MaxObjectCountPerBatch;
-
-				if (bHasProtocolBuffers)
-				{
-					PerFilterInfo.InstanceProtocols = InstanceProtocolsStorage.GetData() + BufferIndex * MaxObjectCountPerBatch;
-				}
-
-				++BufferIndex;
-			}
+			PerFilterInfo.ObjectIndices = ObjectIndicesStorage.GetData() + BufferIndex * MaxObjectCountPerBatch;
+			++BufferIndex;
 		}
 	}
 
@@ -172,23 +139,11 @@ public:
 			if (PerFilterInfo.ObjectIndices)
 			{
 				PerFilterInfo.ObjectIndices[PerFilterInfo.ObjectCount] = ObjectIndex;
-
-				if (bHasProtocolBuffers)
-				{
-					const FReplicationInstanceProtocol* InstanceProtocol = NetRefHandleManager->GetReplicatedObjectDataNoCheck(ObjectIndex).InstanceProtocol;
-					check(InstanceProtocol);
-							
-					check(PerFilterInfo.InstanceProtocols);
-					PerFilterInfo.InstanceProtocols[PerFilterInfo.ObjectCount] = InstanceProtocol;
-				}
-
 				++PerFilterInfo.ObjectCount;
 			}
 		}
 	}
 
-	bool HasProtocolBuffers() const { return bHasProtocolBuffers; }
-
 	TArray<FPerFilterInfo, TInlineAllocator<16>> PerFilterInfos;
 
 private:
@@ -201,11 +156,7 @@ private:
 	}
 
 	TArray<uint32> ObjectIndicesStorage;
-	TArray<const FReplicationInstanceProtocol*> InstanceProtocolsStorage;
-	
 	const FNetRefHandleManager* NetRefHandleManager = nullptr;
-
-	bool bHasProtocolBuffers = false;
 };
 
 //*************************************************************************************************
@@ -218,8 +169,6 @@ FReplicationFiltering::FReplicationFiltering()
 , bHasDirtyConnectionFilter(0)
 , bHasDirtyOwner(0)
 , bHasDynamicFilters(0)
-, bHasDynamicRawFilters(0)
-, bHasDynamicFragmentFilters(0)
 , bHasDirtyExclusionFilterGroup(0)
 , bHasDirtyInclusionFilterGroup(0)
 {
@@ -297,7 +246,7 @@ void FReplicationFiltering::Init(FReplicationFilteringInitParams& Params)
 	InitFilters();
 }
 
-void FReplicationFiltering::FilterPrePoll()
+void FReplicationFiltering::Filter()
 {
 #if UE_NET_IRIS_CSV_STATS
 	CSV_SCOPED_TIMING_STAT(Iris, Filter_PrePoll);
@@ -316,9 +265,9 @@ void FReplicationFiltering::FilterPrePoll()
 
 	UpdateSubObjectFilters();
 
-	if (HasRawFilters())
+	if (HasDynamicFilters())
 	{
-		UpdateDynamicFilters(ENetFilterType::PrePoll_Raw);
+		UpdateDynamicFilters();
 	}
 	else
 	{
@@ -384,27 +333,14 @@ void FReplicationFiltering::BuildAlwaysRelevantList(FNetBitArrayView OutAlwaysRe
 	}
 }
 
-void FReplicationFiltering::UpdateDynamicFilters(ENetFilterType FilterPass)
+/** Dynamic filters allows users to filter out objects based on arbitrary criteria. */
+void FReplicationFiltering::UpdateDynamicFilters()
 {
-	/**
-	* Dynamic filters allows users to filter out objects based on arbitrary criteria.
-	*/
-	NotifyFiltersOfDirtyObjects(FilterPass);
+	NotifyFiltersOfDirtyObjects();
 
-	PreUpdateDynamicFiltering(FilterPass);
-	UpdateDynamicFiltering(FilterPass);
-	PostUpdateDynamicFiltering(FilterPass);
-}
-
-void FReplicationFiltering::FilterPostPoll()
-{
-	if (HasFragmentFilters())
-	{
-#if UE_NET_IRIS_CSV_STATS
-		CSV_SCOPED_TIMING_STAT(Iris, Filter_PostPoll);
-#endif
-		UpdateDynamicFilters(ENetFilterType::PostPoll_FragmentBased);
-	}
+	PreUpdateDynamicFiltering();
+	UpdateDynamicFiltering();
+	PostUpdateDynamicFiltering();
 }
 
 void FReplicationFiltering::SetOwningConnection(FInternalNetRefIndex ObjectIndex, uint32 ConnectionId)
@@ -1275,7 +1211,7 @@ void FReplicationFiltering::UpdateGroupInclusionFiltering()
 	DirtyInclusionFilterGroups.Reset();
 }
 
-void FReplicationFiltering::PreUpdateDynamicFiltering(ENetFilterType FilterType)
+void FReplicationFiltering::PreUpdateDynamicFiltering()
 {
 	IRIS_PROFILER_SCOPE(FReplicationFiltering_PreUpdateDynamicFiltering);
 
@@ -1283,7 +1219,7 @@ void FReplicationFiltering::PreUpdateDynamicFiltering(ENetFilterType FilterType)
 	{
 		for (FFilterInfo& Info : DynamicFilterInfos)
 		{
-			if (Info.ObjectCount == 0U || Info.Type != FilterType)
+			if (Info.ObjectCount == 0U)
 			{
 				continue;
 			}
@@ -1304,7 +1240,7 @@ void FReplicationFiltering::PreUpdateDynamicFiltering(ENetFilterType FilterType)
  * of objects it should not be concerned about.
  */
 
-void FReplicationFiltering::UpdateDynamicFiltering(ENetFilterType FilterType)
+void FReplicationFiltering::UpdateDynamicFiltering()
 {
 	IRIS_PROFILER_SCOPE(FReplicationFiltering_UpdateDynamicFiltering);
 
@@ -1319,7 +1255,6 @@ void FReplicationFiltering::UpdateDynamicFiltering(ENetFilterType FilterType)
 	const uint32* SubObjectsData = NetRefHandleManager->GetSubObjectInternalIndices().GetData();
 	const uint32* DependentObjectsData = NetRefHandleManager->GetDependentObjectInternalIndices().GetData();
 	const uint32* ObjectsRequiringDynamicFilterUpdateData = ObjectsRequiringDynamicFilterUpdate.GetData();
-	const TNetChunkedArray<uint8*>* ObjectsStateBuffers = FilterType == ENetFilterType::PostPoll_FragmentBased ? &NetRefHandleManager->GetReplicatedObjectStateBuffers() : nullptr;
 
 	uint32* ConnectionIds = static_cast<uint32*>(FMemory_Alloca(ValidConnections.GetNumBits() * sizeof(uint32)));
 	uint32 ConnectionCount = 0;
@@ -1329,12 +1264,7 @@ void FReplicationFiltering::UpdateDynamicFiltering(ENetFilterType FilterType)
 	{
 		FPerConnectionInfo& ConnectionInfo = this->ConnectionInfos[ConnId];
 
-		if (FilterType == ENetFilterType::PrePoll_Raw)
-		{
-			// Reset the array in the first filter pass
-			ConnectionInfo.InProgressDynamicFilteredOutObjects.Reset();
-		}
-
+		ConnectionInfo.InProgressDynamicFilteredOutObjects.Reset();
 		uint32* InProgressDynamicFilteredOutObjectsData = ConnectionInfo.InProgressDynamicFilteredOutObjects.GetData();
 		FNetBitArrayView InProgressDynamicFilteredOutObjects = MakeNetBitArrayView(ConnectionInfo.InProgressDynamicFilteredOutObjects);
 
@@ -1350,7 +1280,7 @@ void FReplicationFiltering::UpdateDynamicFiltering(ENetFilterType FilterType)
 		{
 			for (FFilterInfo& Info : DynamicFilterInfos)
 			{
-				if (Info.ObjectCount == 0U || Info.Type != FilterType)
+				if (Info.ObjectCount == 0U)
 				{
 					continue;
 				}
@@ -1358,7 +1288,6 @@ void FReplicationFiltering::UpdateDynamicFiltering(ENetFilterType FilterType)
 				FNetObjectFilteringParams FilteringParams(MakeNetBitArrayView(Info.FilteredObjects));
 				FilteringParams.OutAllowedObjects = AllowedObjects;
 				FilteringParams.FilteringInfos = NetObjectFilteringInfos.GetData();
-				FilteringParams.StateBuffers = ObjectsStateBuffers;
 				FilteringParams.ConnectionId = ConnId;
 				FilteringParams.View = Connections->GetReplicationView(ConnId);
 
@@ -1489,7 +1418,7 @@ void FReplicationFiltering::UpdateDynamicFiltering(ENetFilterType FilterType)
 	ObjectsRequiringDynamicFilterUpdate.Reset();
 }
 
-void FReplicationFiltering::PostUpdateDynamicFiltering(ENetFilterType FilterType)
+void FReplicationFiltering::PostUpdateDynamicFiltering()
 {
 	IRIS_PROFILER_SCOPE(FReplicationFiltering_PostUpdateDynamicFiltering);
 
@@ -1503,22 +1432,19 @@ void FReplicationFiltering::PostUpdateDynamicFiltering(ENetFilterType FilterType
 				continue;
 			}
 			
-			if (Info.Type == FilterType)
-			{
-				Info.Filter->PostFilter(PostFilteringParams);
-			}
+			Info.Filter->PostFilter(PostFilteringParams);
 		}
 	}
 }
 
-void FReplicationFiltering::NotifyFiltersOfDirtyObjects(ENetFilterType FilterType)
+void FReplicationFiltering::NotifyFiltersOfDirtyObjects()
 {
 	IRIS_PROFILER_SCOPE(FReplicationFiltering_UpdateFilterWithDirtyObjects);
 
 	FDirtyObjectsAccessor DirtyObjectsAccessor(ReplicationSystem->GetReplicationSystemInternal()->GetDirtyNetObjectTracker());
 	const FNetBitArrayView DirtyObjectsThisFrame = DirtyObjectsAccessor.GetDirtyNetObjects();
 
-	FUpdateDirtyObjectsBatchHelper BatchHelper(NetRefHandleManager, DynamicFilterInfos, FilterType);
+	FUpdateDirtyObjectsBatchHelper BatchHelper(NetRefHandleManager, DynamicFilterInfos);
 
 	constexpr SIZE_T MaxBatchObjectCount = FUpdateDirtyObjectsBatchHelper::Constants::MaxObjectCountPerBatch;
 	uint32 ObjectIndices[MaxBatchObjectCount];
@@ -1543,11 +1469,6 @@ void FReplicationFiltering::BatchNotifyFiltersOfDirtyObjects(FUpdateDirtyObjects
 	FNetObjectFilterUpdateParams UpdateParameters;
 	UpdateParameters.FilteringInfos = NetObjectFilteringInfos.GetData();
 
-	if (BatchHelper.HasProtocolBuffers())
-	{
-		UpdateParameters.StateBuffers = &NetRefHandleManager->GetReplicatedObjectStateBuffers();
-	}
-
 	// $IRIS TODO: We should probably have a trait asking if the Filter needs to receive UpdateObjects.
 	for (const FUpdateDirtyObjectsBatchHelper::FPerFilterInfo& PerFilterInfo : BatchHelper.PerFilterInfos)
 	{
@@ -1558,7 +1479,6 @@ void FReplicationFiltering::BatchNotifyFiltersOfDirtyObjects(FUpdateDirtyObjects
 
 		UpdateParameters.ObjectIndices = PerFilterInfo.ObjectIndices;
 		UpdateParameters.ObjectCount = PerFilterInfo.ObjectCount;
-		UpdateParameters.InstanceProtocols = PerFilterInfo.InstanceProtocols;
 		// $IRIS TODO: Add a trait asking if the Filter needs to receive UpdateObjects. Most do nothing in UpdateObjects()
 		const int32 FilterIndex = static_cast<int32>(&PerFilterInfo - BatchHelper.PerFilterInfos.GetData());
 		UNetObjectFilter* Filter = DynamicFilterInfos[FilterIndex].Filter.Get();
@@ -2482,12 +2402,8 @@ void FReplicationFiltering::InitFilters()
 		InitParams.MaxConnectionCount = Connections->GetMaxConnectionCount();
 
 		Info.Filter->Init(InitParams);
-		// Filter type can be changed in Init.
-		Info.Type = Info.Filter->GetFilterType();
 
 		bHasDynamicFilters = true;
-		bHasDynamicRawFilters |= (Info.Type == ENetFilterType::PrePoll_Raw);
-		bHasDynamicFragmentFilters |= (Info.Type == ENetFilterType::PostPoll_FragmentBased);
 	}
 }
 
diff --git a/Engine/Source/Runtime/Experimental/Iris/Core/Private/Iris/ReplicationSystem/Filtering/ReplicationFiltering.h b/Engine/Source/Runtime/Experimental/Iris/Core/Private/Iris/ReplicationSystem/Filtering/ReplicationFiltering.h
--- a/Engine/Source/Runtime/Experimental/Iris/Core/Private/Iris/ReplicationSystem/Filtering/ReplicationFiltering.h
+++ b/Engine/Source/Runtime/Experimental/Iris/Core/Private/Iris/ReplicationSystem/Filtering/ReplicationFiltering.h
@@ -57,18 +57,11 @@ public:
 	void Init(FReplicationFilteringInitParams& Params);
 
 	/**
-	 * First pass to determine which objects are relevant to each connection.
-	 * Executes group, owner and connection filtering then any raw dynamic filters.
+	 * Executes group, owner and connection filtering then any dynamic filters.
 	 * At the end any object that is not relevant to at least one connection will be removed from the scoped object list.
-	 * Exception to this rule are always relevant (e.g. non-filtered) objects or objects set to be filtered by a fragment-based dynamic filter.
+	 * Exception to this rule are always relevant (e.g. non-filtered) objects.
 	 */
-	void FilterPrePoll();
-
-	/**
-	 * Second pass to determine which objects are relevant to each connection.
-	 * Executes only fragment-based dynamic filters.
-	 */
-	void FilterPostPoll();
+	void Filter();
 
 	/**  Returns the list of objects relevant to a given connection. This represents the global scope list minus the objects that were filtered out for the given connection. */
 	const FNetBitArrayView GetRelevantObjectsInScope(uint32 ConnectionId) const
@@ -187,7 +180,6 @@ private:
 
 	struct FFilterInfo
 	{
-		ENetFilterType Type = ENetFilterType::PrePoll_Raw;
 		TStrongObjectPtr<UNetObjectFilter> Filter;
 		FName Name;
 		uint32 ObjectCount = 0;
@@ -211,17 +203,15 @@ private:
 	void UpdateGroupInclusionFiltering();
 	void UpdateSubObjectFilters();
 
-	void UpdateDynamicFilters(ENetFilterType FilterPass);
-	void PreUpdateDynamicFiltering(ENetFilterType FilterType);
-	void UpdateDynamicFiltering(ENetFilterType FilterType);
-	void PostUpdateDynamicFiltering(ENetFilterType FilterType);
+	void UpdateDynamicFilters();
+	void PreUpdateDynamicFiltering();
+	void UpdateDynamicFiltering();
+	void PostUpdateDynamicFiltering();
 
 	/** Build the list of always relevant objects + objects that are currently relevant to at least one connection. */
 	void FilterNonRelevantObjects();
 
 	bool HasDynamicFilters() const;
-	bool HasRawFilters() const;
-	bool HasFragmentFilters() const;
 
 	// Helper to update and reset group exclusion filter effects if objects are removed from a filter or after a filter status change, returns true if the group filter was changed
 	bool ClearGroupExclusionFilterEffectsForObject(uint32 ObjectIndex, uint32 ConnectionId);
@@ -257,7 +247,7 @@ private:
 
 	void RemoveFromDynamicFilter(uint32 ObjectIndex, uint32 FilterIndex);
 
-	void NotifyFiltersOfDirtyObjects(ENetFilterType FilterType);
+	void NotifyFiltersOfDirtyObjects();
 	void BatchNotifyFiltersOfDirtyObjects(FUpdateDirtyObjectsBatchHelper& BatchHelper, const uint32* ObjectIndices, uint32 ObjectCount);
 
 	void InvalidateBaselinesForObject(uint32 ObjectIndex, uint32 NewOwningConnectionId, uint32 PrevOwningConnectionId);
@@ -343,8 +333,6 @@ private:
 	uint32 bHasDirtyConnectionFilter: 1;
 	uint32 bHasDirtyOwner : 1;
 	uint32 bHasDynamicFilters : 1;
-	uint32 bHasDynamicRawFilters : 1;
-	uint32 bHasDynamicFragmentFilters : 1;
 	uint32 bHasDirtyExclusionFilterGroup : 1;
 	uint32 bHasDirtyInclusionFilterGroup : 1;
 };
@@ -354,13 +342,4 @@ inline bool FReplicationFiltering::HasDynamicFilters() const
 	return bHasDynamicFilters;
 }
 
-inline bool FReplicationFiltering::HasRawFilters() const
-{
-	return bHasDynamicRawFilters;
-}
-inline bool FReplicationFiltering::HasFragmentFilters() const
-{
-	return bHasDynamicFragmentFilters;
-}
-
 }
diff --git a/Engine/Source/Runtime/Experimental/Iris/Core/Private/Iris/ReplicationSystem/ReplicationSystem.cpp b/Engine/Source/Runtime/Experimental/Iris/Core/Private/Iris/ReplicationSystem/ReplicationSystem.cpp
--- a/Engine/Source/Runtime/Experimental/Iris/Core/Private/Iris/ReplicationSystem/ReplicationSystem.cpp
+++ b/Engine/Source/Runtime/Experimental/Iris/Core/Private/Iris/ReplicationSystem/ReplicationSystem.cpp
@@ -323,26 +323,21 @@ public:
 		ReplicationSystemInternal.GetReplicationBridge()->CallUpdateInstancesWorldLocation();
 	}
 
-	void UpdateFilterPrePoll()
+	void UpdateFiltering()
 	{
-		IRIS_CSV_PROFILER_SCOPE(Iris, ReplicationSystem_UpdateFilterPrePoll);
+		IRIS_CSV_PROFILER_SCOPE(Iris, ReplicationSystem_UpdateFiltering);
 		LLM_SCOPE_BYTAG(Iris);
 
 		FReplicationFiltering& Filtering = ReplicationSystemInternal.GetFiltering();
-		Filtering.FilterPrePoll();
+		Filtering.Filter();
 	}
 
-	void UpdateFilterPostPoll()
+	void UpdateObjectScopes()
 	{
 		LLM_SCOPE_BYTAG(Iris);
 
 		FReplicationFiltering& Filtering = ReplicationSystemInternal.GetFiltering();
 
-		{
-			IRIS_CSV_PROFILER_SCOPE(Iris, ReplicationSystem_UpdateFilterPostPoll);
-			Filtering.FilterPostPoll();
-		}
-
 		{
 			IRIS_CSV_PROFILER_SCOPE(Iris, ReplicationSystem_UpdateConnectionsScope);
 		
@@ -691,7 +686,7 @@ void UReplicationSystem::PreSendUpdate(const FSendUpdateParams& Params)
 			Impl->UpdateWorldLocations();
 
 			// Update filters, reduce the top-level scoped object list and set each connection's scope.
-			Impl->UpdateFilterPrePoll();
+			Impl->UpdateFiltering();
 
 			// Invoke any operations we need to do before copying state data
 			Impl->CallPreSendUpdate(Params.DeltaSeconds);
@@ -708,8 +703,8 @@ void UReplicationSystem::PreSendUpdate(const FSendUpdateParams& Params)
 			// We must process all attachments to objects going out of scope before we update the scope
 			Impl->ProcessNetObjectAttachmentSendQueue(FNetBlobManager::EProcessMode::ProcessObjectsGoingOutOfScope);
 
-			// Update filtering and scope for all connections
-			Impl->UpdateFilterPostPoll();
+			// Update scope for all connections
+			Impl->UpdateObjectScopes();
 
 			// Propagate dirty changes to all connections
 			Impl->PropagateDirtyChanges();
diff --git a/Engine/Source/Runtime/Experimental/Iris/Core/Public/Iris/ReplicationSystem/Filtering/NetObjectFilter.h b/Engine/Source/Runtime/Experimental/Iris/Core/Public/Iris/ReplicationSystem/Filtering/NetObjectFilter.h
--- a/Engine/Source/Runtime/Experimental/Iris/Core/Public/Iris/ReplicationSystem/Filtering/NetObjectFilter.h
+++ b/Engine/Source/Runtime/Experimental/Iris/Core/Public/Iris/ReplicationSystem/Filtering/NetObjectFilter.h
@@ -100,25 +100,6 @@ struct alignas(8) FNetObjectFilteringInfo
 	uint16 Data[4];
 };
 
-/** This configures when a filter gets executed inside PreSendUpdate and what data it has access to. */
-UENUM()
-enum class ENetFilterType : uint8
-{
-	/**
-	 * The default setting of filters.
-	 * This type of filter is applied before we poll objects and copy their data into network fragments.
-	 * Thus a filter cannot operate on fragment data during it's operation, only whatever it stored or is accessing directly from the object.
-	 * The benefit is that objects not relevant to any connection will get culled from being polled and forced to copy their dirty state for nothing.
-	 */
-	PrePoll_Raw,
-
-	/**
-	 * When set to FragmentBased, the filter gets access to the up to date fragment data of an object when it executes.
-	 * Those objects will always be polled and have their dirty state copied even if not relevant.
-	 */
-	PostPoll_FragmentBased,
-};
-
 enum class ENetFilterTraits : uint8
 {
 	None = 0,
@@ -136,10 +117,6 @@ class UNetObjectFilterConfig : public UObject
 	GENERATED_BODY()
 
 public:
-
-	/** Can be used to modify when the filter is executed */
-	UPROPERTY(Config)
-	ENetFilterType FilterType = ENetFilterType::PrePoll_Raw;
 };
 
 /** Parameters passed to the filter's Init() call. */
@@ -194,18 +171,6 @@ struct FNetObjectFilterUpdateParams
 
 	/** Infos for all objects. Index using ObjectIndices[0..ObjectCount-1]. */
 	FNetObjectFilteringInfo* FilteringInfos = nullptr;
-
-	/** 
-	* InstanceProtocols for updated objects. Index using 0..ObjectCount-1.
-	* NOTE: Only for filters of type FragmentBased; null for Raw types
-	*/
-	UE::Net::FReplicationInstanceProtocol const* const* InstanceProtocols = nullptr;
-
-	/**
-	* State buffers for all objects. Index using ObjectIndices[0..ObjectCount-1].
-	* NOTE: Only for filters of type FragmentBased; null for Raw types
-	*/
-	const UE::Net::TNetChunkedArray<uint8*>* StateBuffers = nullptr;
 };
 
 UCLASS(Abstract)
@@ -245,9 +210,6 @@ public:
 	 */
 	IRISCORE_API virtual void PostFilter(FNetObjectPostFilteringParams&);
 
-	/** Returns what type of filter it is. Default is to cull and be executed before dirty state copying. */
-	ENetFilterType GetFilterType() const { return FilterType; }
-
 	/** Returns the filter's traits. */
 	ENetFilterTraits GetFilterTraits() const { return FilterTraits; }
 
@@ -278,9 +240,6 @@ protected:
 	/* Returns the filtering info for this object if it's handled by this filter, nullptr otherwise. */
 	FNetObjectFilteringInfo* GetFilteringInfo(uint32 ObjectIndex);
 
-	/** Directly set when you want your dynamic filter to be executed. */
-	void SetupFilterType(ENetFilterType NewFilterType) { FilterType = NewFilterType; }
-
 	/** Adds traits. */
 	void AddFilterTraits(ENetFilterTraits Traits);
 
@@ -305,7 +264,6 @@ private:
 		TArrayView<FNetObjectFilteringInfo> FilteringInfos;
 	};
 
-	ENetFilterType FilterType = ENetFilterType::PrePoll_Raw;
 	ENetFilterTraits FilterTraits = ENetFilterTraits::None;
 	FFilterInfo FilterInfo;
 };
diff --git a/Engine/Source/Runtime/Experimental/Iris/Core/Public/Iris/ReplicationSystem/Filtering/NetObjectGridFilter.h b/Engine/Source/Runtime/Experimental/Iris/Core/Public/Iris/ReplicationSystem/Filtering/NetObjectGridFilter.h
--- a/Engine/Source/Runtime/Experimental/Iris/Core/Public/Iris/ReplicationSystem/Filtering/NetObjectGridFilter.h
+++ b/Engine/Source/Runtime/Experimental/Iris/Core/Public/Iris/ReplicationSystem/Filtering/NetObjectGridFilter.h
@@ -326,33 +326,6 @@ private:
 	const UE::Net::FWorldLocations* WorldLocations = nullptr;
 };
 
-/**
- * Filter for replicated objects that have their location stored in their fragment
- * 
- * This filter may be less efficient since it's run after Polling and DirtyData copying and cannot cull out objects from those operations.
- */
-UCLASS()
-class UNetObjectGridFragmentLocFilter : public UNetObjectGridFilter
-{
-	GENERATED_BODY()
-
-protected:
-
-	virtual void OnInit(FNetObjectFilterInitParams&) override;
-	virtual void UpdateObjects(FNetObjectFilterUpdateParams&) override;
-	virtual void UpdateObjectInfo(FPerObjectInfo& PerObjectInfo, const UNetObjectGridFilter::FObjectLocationInfo& ObjectLocationInfo, const UE::Net::FReplicationInstanceProtocol* InstanceProtocol) override;
-	virtual bool BuildObjectInfo(uint32 ObjectIndex, FNetObjectFilterAddObjectParams& Params) override;
-	virtual void OnObjectRemoved(uint32 ObjectIndex) override;
-
-private:
-	struct FCullDistanceFragmentInfo
-	{
-		uint16 CullDistanceSqrStateIndex = InvalidStateIndex;
-		uint16 CullDistanceSqrStateOffset = InvalidStateOffset;
-	};
-
-	TMap<uint32, FCullDistanceFragmentInfo> CullDistanceFragments;
-};
 
 //
 inline bool UNetObjectGridFilter::FCellBox::operator==(const UNetObjectGridFilter::FCellBox& Other) const
