Iris - Export UNetObjectFilter.

#11344

diff --git a/Engine/Source/Runtime/Experimental/Iris/Core/Public/Iris/ReplicationSystem/Filtering/NetObjectFilter.h b/Engine/Source/Runtime/Experimental/Iris/Core/Public/Iris/ReplicationSystem/Filtering/NetObjectFilter.h
--- a/Engine/Source/Runtime/Experimental/Iris/Core/Public/Iris/ReplicationSystem/Filtering/NetObjectFilter.h
+++ b/Engine/Source/Runtime/Experimental/Iris/Core/Public/Iris/ReplicationSystem/Filtering/NetObjectFilter.h
@@ -173,7 +173,7 @@ struct FNetObjectFilterUpdateParams
 	FNetObjectFilteringInfo* FilteringInfos = nullptr;
 };
 
-UCLASS(Abstract)
+UCLASS(Abstract, MinimalAPI)
 class UNetObjectFilter : public UObject
 {
 	GENERATED_BODY()
