Optimization and simplification of the vectorized solver.
ConstraintArms is identical per axis.

diff --git a/Engine/Source/Runtime/Experimental/Chaos/Private/Chaos/Joint/PBDJointCachedSolverGaussSeidel.cpp b/Engine/Source/Runtime/Experimental/Chaos/Private/Chaos/Joint/PBDJointCachedSolverGaussSeidel.cpp
--- a/Engine/Source/Runtime/Experimental/Chaos/Private/Chaos/Joint/PBDJointCachedSolverGaussSeidel.cpp
+++ b/Engine/Source/Runtime/Experimental/Chaos/Private/Chaos/Joint/PBDJointCachedSolverGaussSeidel.cpp
@@ -627,19 +627,15 @@ void FPBDJointCachedSolver::InitLockedPositionConstraintSimd(
 	}
 	PositionConstraints.Simd.ConstraintCX = MakeVectorRegisterFloat(LocalDeltas[0], LocalDeltas[1], LocalDeltas[2], 0.0f);
 
-	PositionConstraints.Simd.ConstraintArms[0][0] = MakeVectorRegisterFloatFromDouble(MakeVectorRegister(ConstraintArm0[0], ConstraintArm0[1], ConstraintArm0[2], 0.0f));
-	PositionConstraints.Simd.ConstraintArms[0][1] = MakeVectorRegisterFloatFromDouble(MakeVectorRegister(ConstraintArm1[0], ConstraintArm1[1], ConstraintArm1[2], 0.0f));
-	PositionConstraints.Simd.ConstraintArms[1][0] = PositionConstraints.Simd.ConstraintArms[0][0];
-	PositionConstraints.Simd.ConstraintArms[1][1] = PositionConstraints.Simd.ConstraintArms[0][1];
-	PositionConstraints.Simd.ConstraintArms[2][0] = PositionConstraints.Simd.ConstraintArms[0][0];
-	PositionConstraints.Simd.ConstraintArms[2][1] = PositionConstraints.Simd.ConstraintArms[0][1];
+	PositionConstraints.Simd.ConstraintArms[0] = MakeVectorRegisterFloatFromDouble(MakeVectorRegister(ConstraintArm0[0], ConstraintArm0[1], ConstraintArm0[2], 0.0f));
+	PositionConstraints.Simd.ConstraintArms[1] = MakeVectorRegisterFloatFromDouble(MakeVectorRegister(ConstraintArm1[0], ConstraintArm1[1], ConstraintArm1[2], 0.0f));
 
 	for (int32 ConstraintIndex = 0; ConstraintIndex < 3; ++ConstraintIndex)
 	{
-		const VectorRegister4Float AngularAxis0 = VectorCross(PositionConstraints.Simd.ConstraintArms[ConstraintIndex][0], PositionConstraints.Simd.ConstraintAxis[ConstraintIndex]);
+		const VectorRegister4Float AngularAxis0 = VectorCross(PositionConstraints.Simd.ConstraintArms[0], PositionConstraints.Simd.ConstraintAxis[ConstraintIndex]);
 		const VectorRegister4Float IA0 = Private::VectorMatrixMultiply(AngularAxis0, InvI(0));
 
-		const VectorRegister4Float AngularAxis1 = VectorCross(PositionConstraints.Simd.ConstraintArms[ConstraintIndex][1], PositionConstraints.Simd.ConstraintAxis[ConstraintIndex]);
+		const VectorRegister4Float AngularAxis1 = VectorCross(PositionConstraints.Simd.ConstraintArms[1], PositionConstraints.Simd.ConstraintAxis[ConstraintIndex]);
 		const VectorRegister4Float IA1 = Private::VectorMatrixMultiply(AngularAxis1, InvI(1));
 
 		const FRealSingle II0 = VectorDot3Scalar(AngularAxis0, IA0);
@@ -856,19 +852,16 @@ void FPBDJointCachedSolver::ApplyPositionConstraintsSimd(
 	const VectorRegister4Float Body1DQ = MakeVectorRegisterFloatFromDouble(MakeVectorRegister(Body(1).DQ()[0], Body(1).DQ()[1], Body(1).DQ()[2], 0.0f));
 
 	const VectorRegister4Float DPDiff = VectorSubtract(Body1DP, Body0DP);
-
+	const VectorRegister4Float Cross1 = VectorCross(Body1DQ, PositionConstraints.Simd.ConstraintArms[1]);
+	const VectorRegister4Float Cross0 = VectorCross(Body0DQ, PositionConstraints.Simd.ConstraintArms[0]);
+	const VectorRegister4Float CrossDiff = VectorSubtract(Cross1, Cross0);
+	const VectorRegister4Float CX = VectorAdd(DPDiff, CrossDiff);
 
 	VectorRegister4Float DeltaPositions[3];
 	for (int32 CIndex = 0; CIndex < 3; ++CIndex)
 	{
-		VectorRegister4Float  Cross1 = VectorCross(Body1DQ, PositionConstraints.Simd.ConstraintArms[CIndex][1]);
-		VectorRegister4Float  Cross0 = VectorCross(Body0DQ, PositionConstraints.Simd.ConstraintArms[CIndex][0]);
-		VectorRegister4Float  CrossDiff = VectorSubtract(Cross1, Cross0);
-		VectorRegister4Float CX = VectorAdd(DPDiff, CrossDiff);
-
 		DeltaPositions[CIndex] = Private::VectorDot3FastX(CX, PositionConstraints.Simd.ConstraintAxis[CIndex]);
 	}
-
 	VectorRegister4Float DeltaPosition = VectorUnpackLo(DeltaPositions[0], DeltaPositions[1]);
 	DeltaPosition = VectorMoveLh(DeltaPosition, DeltaPositions[2]);
 
@@ -1035,20 +1028,25 @@ void FPBDJointCachedSolver::ApplyVelocityConstraintSimd()
 		VectorRegister4Float Stiffness = VectorLoadFloat1(&SolverStiffnessf);
 		Stiffness = VectorMultiply(Stiffness, PositionConstraints.Simd.ConstraintHardStiffness);
 
-		VectorRegister4Float HardIM[3];
-		HardIM[0] = VectorReplicate(PositionConstraints.Simd.ConstraintHardIM, 0);
-		HardIM[1] = VectorReplicate(PositionConstraints.Simd.ConstraintHardIM, 1);
-		HardIM[2] = VectorReplicate(PositionConstraints.Simd.ConstraintHardIM, 2);
-		VectorRegister4Float DeltaLambdas[3];
+		const VectorRegister4Float CV0 = VectorAdd(V0, VectorCross(W0, PositionConstraints.Simd.ConstraintArms[0]));
+		const VectorRegister4Float CV1 = VectorAdd(V1, VectorCross(W1, PositionConstraints.Simd.ConstraintArms[1]));
+		const VectorRegister4Float CV = VectorSubtract(CV1, CV0);
+
+		VectorRegister4Float ProjVs[3];
 		for (int32 CIndex = 0; CIndex < 3; CIndex++)
 		{
-			const VectorRegister4Float CV0 = VectorAdd(V0, VectorCross(W0, PositionConstraints.Simd.ConstraintArms[CIndex][0]));
-			const VectorRegister4Float CV1 = VectorAdd(V1, VectorCross(W1, PositionConstraints.Simd.ConstraintArms[CIndex][1]));
-			const VectorRegister4Float CV = VectorSubtract(CV1, CV0);
-			const VectorRegister4Float Dot3 = VectorDot3(CV, PositionConstraints.Simd.ConstraintAxis[CIndex]);
-
-			DeltaLambdas[CIndex] = VectorDivide(VectorMultiply(Stiffness, Dot3), HardIM[CIndex]);
+			ProjVs[CIndex] = Private::VectorDot3FastX(CV, PositionConstraints.Simd.ConstraintAxis[CIndex]);
 		}
+		VectorRegister4Float ProjV = VectorUnpackLo(ProjVs[0], ProjVs[1]);
+		ProjV = VectorMoveLh(ProjV, ProjVs[2]);
+
+		const VectorRegister4Float DeltaLambda = VectorDivide(VectorMultiply(Stiffness, ProjV), PositionConstraints.Simd.ConstraintHardIM);
+
+		VectorRegister4Float DeltaLambdas[3];
+		DeltaLambdas[0] = VectorReplicate(DeltaLambda, 0);
+		DeltaLambdas[1] = VectorReplicate(DeltaLambda, 1);
+		DeltaLambdas[2] = VectorReplicate(DeltaLambda, 2);
+
 		if (Body(0).IsDynamic())
 		{
 			const FRealSingle Inv0f = FRealSingle(InvM(0));
@@ -1297,9 +1295,6 @@ void FPBDJointCachedSolver::InitRotationConstraintsSimd(
 		RotationConstraints.bLimitsCheck[ConstraintIndex] = true;
 		RotationConstraints.ConstraintVX[ConstraintIndex] = 0;
 
-		RotationConstraints.Simd.ConstraintArms[ConstraintIndex][0] = VectorZeroFloat();
-		RotationConstraints.Simd.ConstraintArms[ConstraintIndex][1] = VectorZeroFloat();
-
 		InitConstraintAxisAngularVelocities[ConstraintIndex] = FVec3::DotProduct(W(1) - W(0), LocalAxis);
 
 		// InitRotationDatasMass(RotationConstraints, ConstraintIndex, Dt);
@@ -1335,6 +1330,8 @@ void FPBDJointCachedSolver::InitRotationConstraintsSimd(
 
 	}
 
+	RotationConstraints.Simd.ConstraintArms[0] = VectorZeroFloat();
+	RotationConstraints.Simd.ConstraintArms[1] = VectorZeroFloat();
 	RotationConstraints.Simd.ConstraintCX = MakeVectorRegisterFloat(LocalAngles[0], LocalAngles[1], LocalAngles[2], 0.0f);
 	RotationConstraints.Simd.ConstraintHardIM = MakeVectorRegisterFloat(ConstraintHardIM[0], ConstraintHardIM[1], ConstraintHardIM[2], 0.0f);
 	RotationConstraints.Simd.ConstraintSoftStiffness = MakeVectorRegisterFloat(ConstraintSoftStiffness[0], ConstraintSoftStiffness[1], ConstraintSoftStiffness[2], 0.0f);
@@ -2045,7 +2042,7 @@ void FPBDJointCachedSolver::ApplyAxisRotationProjection(
 				else
 				{
 					FVec3f PositionConstraintsArmsf;
-					VectorStoreFloat3(PositionConstraints.Simd.ConstraintArms[ConstraintIndex][1], &PositionConstraintsArmsf[0]);
+					VectorStoreFloat3(PositionConstraints.Simd.ConstraintArms[1], &PositionConstraintsArmsf[0]);
 					PositionConstraintsArms1 = FVec3(PositionConstraintsArmsf);
 				}
 				const FVec3 DP1 = -AngularProjection * FVec3::CrossProduct(DR1, PositionConstraintsArms1);
@@ -2128,7 +2125,7 @@ void FPBDJointCachedSolver::ApplyRotationProjectionSimd(
 				VectorRegister4Float PositionConstraintsArms1;
 				if (PositionConstraints.bUseSimd)
 				{
-					PositionConstraintsArms1 = PositionConstraints.Simd.ConstraintArms[CIndex][1];
+					PositionConstraintsArms1 = PositionConstraints.Simd.ConstraintArms[1];
 				}
 				else
 				{
@@ -2225,20 +2222,18 @@ void FPBDJointCachedSolver::ApplyPositionProjectionSimd(
 
 	const VectorRegister4Float DPDiff = VectorSubtract(Body1DP, Body0DP);
 
+	const VectorRegister4Float Cross1 = VectorCross(Body1DQ, PositionConstraints.Simd.ConstraintArms[1]);
+	const VectorRegister4Float Cross0 = VectorCross(Body0DQ, PositionConstraints.Simd.ConstraintArms[0]);
+	const VectorRegister4Float CrossDiff = VectorSubtract(Cross1, Cross0);
+	const VectorRegister4Float CX = VectorAdd(DPDiff, CrossDiff);
+
 	VectorRegister4Float DeltaPositions[3];
 	for (int32 CIndex = 0; CIndex < 3; ++CIndex)
 	{
-		VectorRegister4Float  Cross1 = VectorCross(Body1DQ, PositionConstraints.Simd.ConstraintArms[CIndex][1]);
-		VectorRegister4Float  Cross0 = VectorCross(Body0DQ, PositionConstraints.Simd.ConstraintArms[CIndex][0]);
-		VectorRegister4Float  CrossDiff = VectorSubtract(Cross1, Cross0);
-		VectorRegister4Float CX = VectorAdd(DPDiff, CrossDiff);
-
 		DeltaPositions[CIndex] = Private::VectorDot3FastX(CX, PositionConstraints.Simd.ConstraintAxis[CIndex]);
 	}
-
 	VectorRegister4Float DeltaPosition = VectorUnpackLo(DeltaPositions[0], DeltaPositions[1]);
 	DeltaPosition = VectorMoveLh(DeltaPosition, DeltaPositions[2]);
-
 	DeltaPosition = VectorAdd(DeltaPosition, PositionConstraints.Simd.ConstraintCX);
 
 	const FRealSingle SolverStiffnessf = FRealSingle(SolverStiffness);
@@ -2253,7 +2248,7 @@ void FPBDJointCachedSolver::ApplyPositionProjectionSimd(
 	VectorRegister4Float IMs[3];
 	for (int32 CIndex = 0; CIndex < 3; CIndex++)
 	{
-		const VectorRegister4Float AngularAxis1 = VectorCross(PositionConstraints.Simd.ConstraintArms[CIndex][1], PositionConstraints.Simd.ConstraintAxis[CIndex]);
+		const VectorRegister4Float AngularAxis1 = VectorCross(PositionConstraints.Simd.ConstraintArms[1], PositionConstraints.Simd.ConstraintAxis[CIndex]);
 		IMs[CIndex] = Private::VectorDot3FastX(AngularAxis1, PositionConstraints.Simd.ConstraintDRAxis[CIndex][1]);
 	}
 	VectorRegister4Float IM = VectorUnpackLo(IMs[0], IMs[1]);
@@ -2596,8 +2591,6 @@ void FPBDJointCachedSolver::InitRotationConstraintDriveSimd(
 		RotationDrives.ConstraintVX[ConstraintIndex] = 0.0;
 		RotationDrives.ConstraintRestitution[ConstraintIndex] = 0.0;
 
-		RotationDrives.Simd.ConstraintArms[ConstraintIndex][0] = VectorZeroFloat();
-		RotationDrives.Simd.ConstraintArms[ConstraintIndex][1] = VectorZeroFloat();
 		RotationDrives.Simd.ConstraintAxis[ConstraintIndex] = MakeVectorRegisterFloatFromDouble(MakeVectorRegister(ConstraintAxes[ConstraintIndex][0], ConstraintAxes[ConstraintIndex][1], ConstraintAxes[ConstraintIndex][2], 0.0f));
 
 		const VectorRegister4Float& Axis = RotationDrives.Simd.ConstraintAxis[ConstraintIndex];
@@ -2632,6 +2625,9 @@ void FPBDJointCachedSolver::InitRotationConstraintDriveSimd(
 		ConstraintSoftIM[ConstraintIndex] = (ConstraintSoftStiffness[ConstraintIndex] + ConstraintSoftDamping[ConstraintIndex]) * ConstraintHardIM[ConstraintIndex] + (FReal)1;
 
 	}
+
+	RotationDrives.Simd.ConstraintArms[0] = VectorZeroFloat();
+	RotationDrives.Simd.ConstraintArms[1] = VectorZeroFloat();
 	RotationDrives.Simd.ConstraintHardIM = MakeVectorRegisterFloatFromDouble(MakeVectorRegister(ConstraintHardIM[0], ConstraintHardIM[1], ConstraintHardIM[2], 0.0f));
 	RotationDrives.Simd.ConstraintSoftStiffness = MakeVectorRegisterFloatFromDouble(MakeVectorRegister(ConstraintSoftStiffness[0], ConstraintSoftStiffness[1], ConstraintSoftStiffness[2], 0.0f));
 	RotationDrives.Simd.ConstraintSoftDamping = MakeVectorRegisterFloatFromDouble(MakeVectorRegister(ConstraintSoftDamping[0], ConstraintSoftDamping[1], ConstraintSoftDamping[2], 0.0f));
diff --git a/Engine/Source/Runtime/Experimental/Chaos/Public/Chaos/Joint/PBDJointCachedSolverGaussSeidel.h b/Engine/Source/Runtime/Experimental/Chaos/Public/Chaos/Joint/PBDJointCachedSolverGaussSeidel.h
--- a/Engine/Source/Runtime/Experimental/Chaos/Public/Chaos/Joint/PBDJointCachedSolverGaussSeidel.h
+++ b/Engine/Source/Runtime/Experimental/Chaos/Public/Chaos/Joint/PBDJointCachedSolverGaussSeidel.h
@@ -76,7 +76,7 @@ struct FAxisConstraintDatas
 		VectorRegister4Float ConstraintHardStiffness;
 		VectorRegister4Float ConstraintSoftStiffness;
 		VectorRegister4Float ConstraintSoftDamping;
-		VectorRegister4Float ConstraintArms[3][2];
+		VectorRegister4Float ConstraintArms[2];
 		VectorRegister4Float ConstraintAxis[3];
 		VectorRegister4Float ConstraintLimits;
 		VectorRegister4Float ConstraintSoftIM;
