Incorrect header queried when getting response type

#rnx

diff --git a/Engine/Source/Runtime/Experimental/IoStore/HttpClient/Private/Api.inl b/Engine/Source/Runtime/Experimental/IoStore/HttpClient/Private/Api.inl
--- a/Engine/Source/Runtime/Experimental/IoStore/HttpClient/Private/Api.inl
+++ b/Engine/Source/Runtime/Experimental/IoStore/HttpClient/Private/Api.inl
@@ -133,7 +133,7 @@ EMimeType FResponse::GetContentType() const
 ////////////////////////////////////////////////////////////////////////////////
 void FResponse::GetContentType(FAnsiStringView& Out) const
 {
-	Out = GetHeader("Accept");
+	Out = GetHeader("Content-Type");
 
 	int32 SemiColon;
 	if (Out.FindChar(';', SemiColon))
