Reinstate convenience parsing of UE-specific mime types

#rnx

diff --git a/Engine/Source/Runtime/Experimental/IoStore/HttpClient/Internal/IO/Http/Client.h b/Engine/Source/Runtime/Experimental/IoStore/HttpClient/Internal/IO/Http/Client.h
--- a/Engine/Source/Runtime/Experimental/IoStore/HttpClient/Internal/IO/Http/Client.h
+++ b/Engine/Source/Runtime/Experimental/IoStore/HttpClient/Internal/IO/Http/Client.h
@@ -27,11 +27,9 @@ enum class EMimeType
 	Binary,
 	Json,
 	Xml,
-	/* UE_CUSTOM_MIME_TYPES
 	CbObject,
 	CbPackage,
 	CompressedBuffer,
-	*/
 	Count
 };
 
diff --git a/Engine/Source/Runtime/Experimental/IoStore/HttpClient/Private/Api.inl b/Engine/Source/Runtime/Experimental/IoStore/HttpClient/Private/Api.inl
--- a/Engine/Source/Runtime/Experimental/IoStore/HttpClient/Private/Api.inl
+++ b/Engine/Source/Runtime/Experimental/IoStore/HttpClient/Private/Api.inl
@@ -121,11 +121,9 @@ EMimeType FResponse::GetContentType() const
 	if (Value == "application/octet-stream")	return EMimeType::Binary;
 	if (Value == "application/json")			return EMimeType::Json;
 	if (Value == "application/xml")				return EMimeType::Xml;
-	/* UE_CUSTOM_MIME_TYPES
 	if (Value == "application/x-ue-cb")			return EMimeType::CbObject;
 	if (Value == "application/x-ue-pkg")		return EMimeType::CbPackage;
 	if (Value == "application/x-ue-comp")		return EMimeType::CompressedBuffer;
-	*/
 
 	return EMimeType::Unknown;
 }
