Compressed flags in Linear Solver

diff --git a/Engine/Source/Runtime/Experimental/Chaos/Private/Chaos/Joint/PBDJointCachedSolverGaussSeidel.cpp b/Engine/Source/Runtime/Experimental/Chaos/Private/Chaos/Joint/PBDJointCachedSolverGaussSeidel.cpp
--- a/Engine/Source/Runtime/Experimental/Chaos/Private/Chaos/Joint/PBDJointCachedSolverGaussSeidel.cpp
+++ b/Engine/Source/Runtime/Experimental/Chaos/Private/Chaos/Joint/PBDJointCachedSolverGaussSeidel.cpp
@@ -143,19 +143,19 @@ void FPBDJointCachedSolver::SetShockPropagationScales(const FReal InvMScale0, co
 	{
 		for (int32 ConstraintIndex = 0; ConstraintIndex < 3; ++ConstraintIndex)
 		{
-			if (PositionConstraints.bValidDatas[ConstraintIndex])
+			if (PositionConstraints.GetValidDatas(ConstraintIndex))
 			{
 				InitPositionDatasMass(PositionConstraints, ConstraintIndex, Dt);
 			}
-			if (RotationConstraints.bValidDatas[ConstraintIndex])
+			if (RotationConstraints.GetValidDatas(ConstraintIndex))
 			{
 				InitRotationDatasMass(RotationConstraints, ConstraintIndex, Dt);
 			}
-			if (PositionDrives.bValidDatas[ConstraintIndex])
+			if (PositionDrives.GetValidDatas(ConstraintIndex))
 			{
 				InitPositionDatasMass(PositionDrives, ConstraintIndex, Dt);
 			}
-			if (RotationDrives.bValidDatas[ConstraintIndex])
+			if (RotationDrives.GetValidDatas(ConstraintIndex))
 			{
 				InitRotationDatasMass(RotationDrives, ConstraintIndex, Dt);
 			}
@@ -381,9 +381,9 @@ void FPBDJointCachedSolver::InitPositionConstraints(
 	const FPBDJointSettings& JointSettings,
 	const bool bResetLambdas)
 {
-	PositionConstraints.bValidDatas[0] = false;
-	PositionConstraints.bValidDatas[1] = false;
-	PositionConstraints.bValidDatas[2] = false;
+	PositionConstraints.SetValidDatas(0, false);
+	PositionConstraints.SetValidDatas(1, false);
+	PositionConstraints.SetValidDatas(2, false);
 	PositionConstraints.bUseSimd = false;
 
 	TVec3<bool> bLinearLocked, bLinearLimited;
@@ -395,28 +395,27 @@ void FPBDJointCachedSolver::InitPositionConstraints(
 
 	if (PositionConstraints.bUseSimd)
 	{
-		PositionConstraints.bSoftLimit[0] = false;
-		PositionConstraints.bSoftLimit[1] = false;
-		PositionConstraints.bSoftLimit[2] = false;
+		PositionConstraints.SetSoftLimit(0, false);
+		PositionConstraints.SetSoftLimit(1, false);
+		PositionConstraints.SetSoftLimit(2, false);
 		FRealSingle HardStiffness = FRealSingle(FPBDJointUtilities::GetLinearStiffness(SolverSettings, JointSettings));
 		PositionConstraints.Simd.ConstraintHardStiffness = VectorLoadFloat1(&HardStiffness);
 		PositionConstraints.ConstraintMaxLambda = FVec3(0);
-		PositionConstraints.bValidDatas[0] = false;
-		PositionConstraints.bValidDatas[1] = false;
-		PositionConstraints.bValidDatas[2] = false;
-		PositionConstraints.bLimitsCheck[0] = true;
-		PositionConstraints.bLimitsCheck[1] = true;
-		PositionConstraints.bLimitsCheck[2] = true;
-		PositionConstraints.MotionType[0] = EJointMotionType::Locked;
-		PositionConstraints.MotionType[1] = EJointMotionType::Locked;
-		PositionConstraints.MotionType[2] = EJointMotionType::Locked;
+		PositionConstraints.SetValidDatas(0, false);
+		PositionConstraints.SetValidDatas(1, false);
+		PositionConstraints.SetValidDatas(2, false);
+		PositionConstraints.SetLimitsCheck(0, true);
+		PositionConstraints.SetLimitsCheck(1, true);
+		PositionConstraints.SetLimitsCheck(2, true);
+		PositionConstraints.SetMotionType(0, EJointMotionType::Locked);
+		PositionConstraints.SetMotionType(1, EJointMotionType::Locked);
+		PositionConstraints.SetMotionType(2, EJointMotionType::Locked);
 		if (bResetLambdas)
 		{
 			PositionConstraints.Simd.ConstraintLambda = VectorZeroFloat();
 			PositionConstraints.ConstraintLambdaVelocity = FVec3::Zero();
 			PositionConstraints.Simd.ConstraintLimits = VectorZeroFloat();
 		}
-
 	}
 	else
 	{
@@ -527,7 +526,7 @@ void FPBDJointCachedSolver::InitPositionConstraintDatas(
 	const FVec3 LocalAxis = (ConstraintDelta < 0.0f) ? -ConstraintAxis : ConstraintAxis;
 	const FReal LocalDelta = (ConstraintDelta < 0.0f) ? -ConstraintDelta : ConstraintDelta;
 
-	PositionConstraints.MotionType[ConstraintIndex] = JointType;
+	PositionConstraints.SetMotionType(ConstraintIndex, JointType);
 
 	if (JointType == EJointMotionType::Locked)
 	{
@@ -617,9 +616,9 @@ void FPBDJointCachedSolver::InitLockedPositionConstraintSimd(
 		const FVec3 LocalAxis = (CX[ConstraintIndex] < 0.0f) ? -ConstraintAxis : ConstraintAxis;
 		LocalDeltas[ConstraintIndex] = FRealSingle((CX[ConstraintIndex] < 0.0) ? -CX[ConstraintIndex] : CX[ConstraintIndex]);
 
-		PositionConstraints.MotionType[ConstraintIndex] = EJointMotionType::Locked;
-		PositionConstraints.bValidDatas[ConstraintIndex] = true;
-		PositionConstraints.bLimitsCheck[ConstraintIndex] = false;
+		PositionConstraints.SetMotionType(ConstraintIndex, EJointMotionType::Locked);
+		PositionConstraints.SetValidDatas(ConstraintIndex, true);
+		PositionConstraints.SetLimitsCheck(ConstraintIndex, false);
 		PositionConstraints.ConstraintVX[ConstraintIndex] = 0.0;
 		PositionConstraints.ConstraintRestitution[ConstraintIndex] = 0.0;
 
@@ -738,7 +737,7 @@ void FPBDJointCachedSolver::ApplyPositionConstraints(
 	{
 		for (int32 ConstraintIndex = 0; ConstraintIndex < 3; ++ConstraintIndex)
 		{
-			if (PositionConstraints.bValidDatas[ConstraintIndex])
+			if (PositionConstraints.GetValidDatas(ConstraintIndex))
 			{
 				ApplyAxisPositionConstraint(ConstraintIndex, Dt);
 			}
@@ -816,7 +815,7 @@ void FPBDJointCachedSolver::ApplyAxisPositionConstraint(
 	FReal DeltaPosition = PositionConstraints.Data.ConstraintCX[ConstraintIndex] + FVec3::DotProduct(CX, PositionConstraints.Data.ConstraintAxis[ConstraintIndex]);
 
 	bool NeedsSolve = false;
-	if (PositionConstraints.bLimitsCheck[ConstraintIndex])
+	if (PositionConstraints.GetLimitsCheck(ConstraintIndex))
 	{
 		if (DeltaPosition > PositionConstraints.Data.ConstraintLimits[ConstraintIndex])
 		{
@@ -829,13 +828,13 @@ void FPBDJointCachedSolver::ApplyAxisPositionConstraint(
 			NeedsSolve = true;
 		}
 	}
-	if (!PositionConstraints.bLimitsCheck[ConstraintIndex] || (PositionConstraints.bLimitsCheck[ConstraintIndex] && NeedsSolve && FMath::Abs(DeltaPosition) > PositionTolerance))
+	if (!PositionConstraints.GetLimitsCheck(ConstraintIndex) || (PositionConstraints.GetLimitsCheck(ConstraintIndex) && NeedsSolve && FMath::Abs(DeltaPosition) > PositionTolerance))
 	{
-		if ((PositionConstraints.MotionType[ConstraintIndex] == EJointMotionType::Limited) && PositionConstraints.bSoftLimit[ConstraintIndex])
+		if ((PositionConstraints.GetMotionType(ConstraintIndex) == EJointMotionType::Limited) && PositionConstraints.GetSoftLimit(ConstraintIndex))
 		{
 			SolvePositionConstraintSoft(ConstraintIndex, DeltaPosition, Dt, 0.0f);
 		}
-		else if (PositionConstraints.MotionType[ConstraintIndex] != EJointMotionType::Free)
+		else if (PositionConstraints.GetMotionType(ConstraintIndex) != EJointMotionType::Free)
 		{
 			SolvePositionConstraintHard(ConstraintIndex, DeltaPosition);
 		}
@@ -933,16 +932,16 @@ void FPBDJointCachedSolver::ApplyLinearVelocityConstraints()
 {
 	if (PositionConstraints.bUseSimd)
 	{
-		check(PositionConstraints.bSoftLimit[0] == false);
-		check(PositionConstraints.bSoftLimit[1] == false);
-		check(PositionConstraints.bSoftLimit[2] == false);
+		check(PositionConstraints.GetSoftLimit(0) == false);
+		check(PositionConstraints.GetSoftLimit(1) == false);
+		check(PositionConstraints.GetSoftLimit(2) == false);
 		ApplyVelocityConstraintSimd();
 	}
 	else
 	{
 		for (int32 ConstraintIndex = 0; ConstraintIndex < 3; ++ConstraintIndex)
 		{
-			if (PositionConstraints.bValidDatas[ConstraintIndex] && !PositionConstraints.bSoftLimit[ConstraintIndex])
+			if (PositionConstraints.GetValidDatas(ConstraintIndex) && !PositionConstraints.GetSoftLimit(ConstraintIndex))
 			{
 				ApplyAxisVelocityConstraint(ConstraintIndex);
 			}
@@ -993,7 +992,7 @@ void FPBDJointCachedSolver::ApplyAxisVelocityConstraint(
 	{
 		FReal TargetVel = 0.0f;
 		const FReal Restitution = PositionConstraints.ConstraintRestitution[ConstraintIndex];
-		const bool bIsLimited = (PositionConstraints.MotionType[ConstraintIndex] == EJointMotionType::Limited);
+		const bool bIsLimited = (PositionConstraints.GetMotionType(ConstraintIndex) == EJointMotionType::Limited);
 		if (bIsLimited && (Restitution != 0.0f))
 		{
 			const FReal InitVel = InitConstraintAxisLinearVelocities[ConstraintIndex];
@@ -1007,9 +1006,9 @@ void FPBDJointCachedSolver::ApplyAxisVelocityConstraint(
 
 void FPBDJointCachedSolver::ApplyVelocityConstraintSimd()
 {
-	check(PositionConstraints.MotionType[0] != EJointMotionType::Limited);
-	check(PositionConstraints.MotionType[1] != EJointMotionType::Limited);
-	check(PositionConstraints.MotionType[2] != EJointMotionType::Limited);
+	check(PositionConstraints.GetMotionType(0) != EJointMotionType::Limited);
+	check(PositionConstraints.GetMotionType(1) != EJointMotionType::Limited);
+	check(PositionConstraints.GetMotionType(2) != EJointMotionType::Limited);
 
 	const VectorRegister4Float IsGTEps = VectorCompareGT(VectorAbs(PositionConstraints.Simd.ConstraintLambda), GlobalVectorConstants::SmallNumber);
 
@@ -1128,9 +1127,9 @@ void FPBDJointCachedSolver::InitRotationConstraints(
 	const FPBDJointSettings& JointSettings,
 	const bool bResetLambdas)
 {
-	RotationConstraints.bValidDatas[0] = false;
-	RotationConstraints.bValidDatas[1] = false;
-	RotationConstraints.bValidDatas[2] = false;
+	RotationConstraints.SetValidDatas(0, false);
+	RotationConstraints.SetValidDatas(1, false);
+	RotationConstraints.SetValidDatas(2, false);
 	RotationConstraints.bUseSimd = false;
 
 
@@ -1291,8 +1290,8 @@ void FPBDJointCachedSolver::InitRotationConstraintsSimd(
 		const FVec3 LocalAxis = (Angles[ConstraintIndex] < 0.0f) ? -Axes[ConstraintIndex] : Axes[ConstraintIndex];
 		LocalAngles[ConstraintIndex] = FRealSingle((Angles[ConstraintIndex] < 0.0f) ? -Angles[ConstraintIndex] : Angles[ConstraintIndex]);
 
-		RotationConstraints.bValidDatas[ConstraintIndex] = true;
-		RotationConstraints.bLimitsCheck[ConstraintIndex] = true;
+		RotationConstraints.SetValidDatas(ConstraintIndex, true);
+		RotationConstraints.SetLimitsCheck(ConstraintIndex, true);
 		RotationConstraints.ConstraintVX[ConstraintIndex] = 0;
 
 		InitConstraintAxisAngularVelocities[ConstraintIndex] = FVec3::DotProduct(W(1) - W(0), LocalAxis);
@@ -1540,7 +1539,7 @@ void FPBDJointCachedSolver::ApplyRotationConstraints(
 	{
 		for (int32 ConstraintIndex = 0; ConstraintIndex < 3; ++ConstraintIndex)
 		{
-			if (RotationConstraints.bValidDatas[ConstraintIndex])
+			if (RotationConstraints.GetValidDatas(ConstraintIndex))
 			{
 				ApplyRotationConstraint(ConstraintIndex, Dt);
 			}
@@ -1612,7 +1611,7 @@ void FPBDJointCachedSolver::ApplyRotationConstraint(
 		FVec3::DotProduct(Body(1).DQ() - Body(0).DQ(), RotationConstraints.Data.ConstraintAxis[ConstraintIndex]);
 
 	bool NeedsSolve = false;
-	if (RotationConstraints.bLimitsCheck[ConstraintIndex])
+	if (RotationConstraints.GetLimitsCheck(ConstraintIndex))
 	{
 		if (DeltaAngle > RotationConstraints.Data.ConstraintLimits[ConstraintIndex])
 		{
@@ -1627,9 +1626,9 @@ void FPBDJointCachedSolver::ApplyRotationConstraint(
 		}
 	}
 
-	if (!RotationConstraints.bLimitsCheck[ConstraintIndex] || (RotationConstraints.bLimitsCheck[ConstraintIndex] && NeedsSolve && FMath::Abs(DeltaAngle) > AngleTolerance))
+	if (!RotationConstraints.GetLimitsCheck(ConstraintIndex) || (RotationConstraints.GetLimitsCheck(ConstraintIndex) && NeedsSolve && FMath::Abs(DeltaAngle) > AngleTolerance))
 	{
-		if (RotationConstraints.bSoftLimit[ConstraintIndex])
+		if (RotationConstraints.GetSoftLimit(ConstraintIndex))
 		{
 			SolveRotationConstraintSoft(ConstraintIndex, DeltaAngle, Dt, 0.0f);
 		}
@@ -1731,7 +1730,7 @@ void FPBDJointCachedSolver::ApplyAngularVelocityConstraints()
 	{
 		for (int32 ConstraintIndex = 0; ConstraintIndex < 3; ++ConstraintIndex)
 		{
-			if (RotationConstraints.bValidDatas[ConstraintIndex])
+			if (RotationConstraints.GetValidDatas(ConstraintIndex))
 			{
 				ApplyAngularVelocityConstraint(ConstraintIndex);
 			}
@@ -1773,7 +1772,7 @@ void FPBDJointCachedSolver::ApplyAngularVelocityConstraint(
 	if (FMath::Abs(RotationConstraints.Data.ConstraintLambda[ConstraintIndex]) > UE_SMALL_NUMBER)
 	{
 		FReal TargetVel = 0.0f;
-		if ((RotationConstraints.MotionType[ConstraintIndex] == EJointMotionType::Limited) && (RotationConstraints.ConstraintRestitution[ConstraintIndex] != 0.0f))
+		if ((RotationConstraints.GetMotionType(ConstraintIndex) == EJointMotionType::Limited) && (RotationConstraints.ConstraintRestitution[ConstraintIndex] != 0.0f))
 		{
 			const FReal InitVel = InitConstraintAxisAngularVelocities[ConstraintIndex];
 			TargetVel = InitVel > Chaos_Joint_AngularVelocityThresholdToApplyRestitution ?
@@ -1860,9 +1859,9 @@ void FPBDJointCachedSolver::InitPositionDrives(
 	const FPBDJointSolverSettings& SolverSettings,
 	const FPBDJointSettings& JointSettings)
 {
-	PositionDrives.bValidDatas[0] = false;
-	PositionDrives.bValidDatas[1] = false;
-	PositionDrives.bValidDatas[2] = false;
+	PositionDrives.SetValidDatas(0, false);
+	PositionDrives.SetValidDatas(1, false);
+	PositionDrives.SetValidDatas(2, false);
 	PositionDrives.bUseSimd = false;
 
 	if (SolverSettings.bEnableDrives)
@@ -2003,13 +2002,13 @@ void FPBDJointCachedSolver::ApplyAxisRotationProjection(
 	const int32 ConstraintIndex)
 {
 	check(RotationConstraints.bUseSimd == false);
-	if (RotationConstraints.bValidDatas[ConstraintIndex] && !RotationConstraints.bSoftLimit[ConstraintIndex])
+	if (RotationConstraints.GetValidDatas(ConstraintIndex) && !RotationConstraints.GetSoftLimit(ConstraintIndex))
 	{
 		FReal DeltaAngle = RotationConstraints.Data.ConstraintCX[ConstraintIndex] +
 			FVec3::DotProduct(Body(1).DQ() - Body(0).DQ(), RotationConstraints.Data.ConstraintAxis[ConstraintIndex]);
 
 		bool NeedsSolve = false;
-		if (RotationConstraints.bLimitsCheck[ConstraintIndex])
+		if (RotationConstraints.GetLimitsCheck(ConstraintIndex))
 		{
 			if (DeltaAngle > RotationConstraints.Data.ConstraintLimits[ConstraintIndex])
 			{
@@ -2024,7 +2023,7 @@ void FPBDJointCachedSolver::ApplyAxisRotationProjection(
 			}
 		}
 
-		if (!RotationConstraints.bLimitsCheck[ConstraintIndex] || (RotationConstraints.bLimitsCheck[ConstraintIndex] && NeedsSolve && FMath::Abs(DeltaAngle) > AngleTolerance))
+		if (!RotationConstraints.GetLimitsCheck(ConstraintIndex) || (RotationConstraints.GetLimitsCheck(ConstraintIndex) && NeedsSolve && FMath::Abs(DeltaAngle) > AngleTolerance))
 		{
 			const FReal IM = -FVec3::DotProduct(RotationConstraints.Data.ConstraintAxis[ConstraintIndex], RotationConstraints.Data.ConstraintDRAxis[ConstraintIndex][1]);
 			const FReal DeltaLambda = SolverStiffness * RotationConstraints.Data.ConstraintHardStiffness[ConstraintIndex] * DeltaAngle / IM;
@@ -2154,9 +2153,9 @@ void FPBDJointCachedSolver::ApplyPositionProjection(
 
 	if (PositionConstraints.bUseSimd)
 	{
-		check(PositionConstraints.bSoftLimit[0] == false);
-		check(PositionConstraints.bSoftLimit[1] == false);
-		check(PositionConstraints.bSoftLimit[2] == false);
+		check(PositionConstraints.GetSoftLimit(0) == false);
+		check(PositionConstraints.GetSoftLimit(1) == false);
+		check(PositionConstraints.GetSoftLimit(2) == false);
 		ApplyPositionProjectionSimd(LinearProjection);
 	}
 	else
@@ -2172,7 +2171,7 @@ void FPBDJointCachedSolver::ApplyAxisPositionProjection(
 	const FReal LinearProjection,
 	const int32 ConstraintIndex)
 {
-	if (PositionConstraints.bValidDatas[ConstraintIndex] && !PositionConstraints.bSoftLimit[ConstraintIndex])
+	if (PositionConstraints.GetValidDatas(ConstraintIndex) && !PositionConstraints.GetSoftLimit(ConstraintIndex))
 	{
 		const FVec3 CX = Body(1).DP() - Body(0).DP() +
 			FVec3::CrossProduct(Body(1).DQ(), PositionConstraints.Data.ConstraintArms[ConstraintIndex][1]) -
@@ -2181,7 +2180,7 @@ void FPBDJointCachedSolver::ApplyAxisPositionProjection(
 		FReal DeltaPosition = PositionConstraints.Data.ConstraintCX[ConstraintIndex] + FVec3::DotProduct(CX, PositionConstraints.Data.ConstraintAxis[ConstraintIndex]);
 
 		bool NeedsSolve = false;
-		if (PositionConstraints.bLimitsCheck[ConstraintIndex])
+		if (PositionConstraints.GetLimitsCheck(ConstraintIndex))
 		{
 			if (DeltaPosition > PositionConstraints.Data.ConstraintLimits[ConstraintIndex])
 			{
@@ -2194,7 +2193,7 @@ void FPBDJointCachedSolver::ApplyAxisPositionProjection(
 				NeedsSolve = true;
 			}
 		}
-		if (!PositionConstraints.bLimitsCheck[ConstraintIndex] || (PositionConstraints.bLimitsCheck[ConstraintIndex] && NeedsSolve && FMath::Abs(DeltaPosition) > PositionTolerance))
+		if (!PositionConstraints.GetLimitsCheck(ConstraintIndex) || (PositionConstraints.GetLimitsCheck(ConstraintIndex) && NeedsSolve && FMath::Abs(DeltaPosition) > PositionTolerance))
 		{
 			const FVec3 AngularAxis1 = FVec3::CrossProduct(PositionConstraints.Data.ConstraintArms[ConstraintIndex][1], PositionConstraints.Data.ConstraintAxis[ConstraintIndex]);
 			const FReal IM = InvM(1) - FVec3::DotProduct(AngularAxis1, PositionConstraints.Data.ConstraintDRAxis[ConstraintIndex][1]);
@@ -2321,12 +2320,12 @@ void FPBDJointCachedSolver::ApplyAxisPositionTeleport(
 {
 	check(PositionConstraints.bUseSimd == false);
 
-	if (PositionConstraints.bValidDatas[ConstraintIndex] && !PositionConstraints.bSoftLimit[ConstraintIndex])
+	if (PositionConstraints.GetValidDatas(ConstraintIndex) && !PositionConstraints.GetSoftLimit(ConstraintIndex))
 	{
 		FReal DeltaPosition = PositionConstraints.Data.ConstraintCX[ConstraintIndex];
 
 		bool NeedsSolve = false;
-		if (PositionConstraints.bLimitsCheck[ConstraintIndex])
+		if (PositionConstraints.GetLimitsCheck(ConstraintIndex))
 		{
 			if (DeltaPosition > PositionConstraints.Data.ConstraintLimits[ConstraintIndex])
 			{
@@ -2339,7 +2338,7 @@ void FPBDJointCachedSolver::ApplyAxisPositionTeleport(
 				NeedsSolve = true;
 			}
 		}
-		if (!PositionConstraints.bLimitsCheck[ConstraintIndex] || (PositionConstraints.bLimitsCheck[ConstraintIndex] && NeedsSolve))
+		if (!PositionConstraints.GetLimitsCheck(ConstraintIndex) || (PositionConstraints.GetLimitsCheck(ConstraintIndex) && NeedsSolve))
 		{
 			if (FMath::Abs(DeltaPosition) > TeleportDistance)
 			{
@@ -2394,7 +2393,7 @@ void FPBDJointCachedSolver::ApplyPositionDrives(
 {
 	for (int32 ConstraintIndex = 0; ConstraintIndex < 3; ++ConstraintIndex)
 	{
-		if (PositionDrives.bValidDatas[ConstraintIndex])
+		if (PositionDrives.GetValidDatas(ConstraintIndex))
 		{
 			ApplyAxisPositionDrive(ConstraintIndex, Dt);
 		}
@@ -2444,7 +2443,7 @@ void FPBDJointCachedSolver::ApplyPositionVelocityDrives(
 
 	for (int32 ConstraintIndex = 0; ConstraintIndex < 3; ++ConstraintIndex)
 	{
-		if (PositionDrives.bValidDatas[ConstraintIndex])
+		if (PositionDrives.GetValidDatas(ConstraintIndex))
 		{
 			ApplyAxisPositionVelocityDrive(ConstraintIndex, Dt);
 		}
@@ -2512,9 +2511,9 @@ void FPBDJointCachedSolver::InitRotationDrives(
 	const FPBDJointSolverSettings& SolverSettings,
 	const FPBDJointSettings& JointSettings)
 {
-	RotationDrives.bValidDatas[0] = false;
-	RotationDrives.bValidDatas[1] = false;
-	RotationDrives.bValidDatas[2] = false;
+	RotationDrives.SetValidDatas(0, false);
+	RotationDrives.SetValidDatas(1, false);
+	RotationDrives.SetValidDatas(2, false);
 	RotationDrives.bUseSimd = false;
 
 	bool bHasRotationDrives =
@@ -2585,8 +2584,8 @@ void FPBDJointCachedSolver::InitRotationConstraintDriveSimd(
 
 	for (int32 ConstraintIndex = 0; ConstraintIndex < 3; ConstraintIndex++)
 	{
-		RotationDrives.bValidDatas[ConstraintIndex] = true;
-		RotationDrives.bLimitsCheck[ConstraintIndex] = true;
+		RotationDrives.SetValidDatas(ConstraintIndex, true);
+		RotationDrives.SetLimitsCheck(ConstraintIndex, true);
 
 		RotationDrives.ConstraintVX[ConstraintIndex] = 0.0;
 		RotationDrives.ConstraintRestitution[ConstraintIndex] = 0.0;
@@ -2618,7 +2617,7 @@ void FPBDJointCachedSolver::InitRotationConstraintDriveSimd(
 		RotationDrives.Simd.ConstraintDRAxis[ConstraintIndex][1] = VectorNegate(IA1);
 
 
-		check(RotationDrives.bSoftLimit[ConstraintIndex])
+		check(RotationDrives.GetSoftLimit(ConstraintIndex))
 		const FReal SpringMassScale = (RotationDrives.bAccelerationMode) ? (FReal)1 / (ConstraintHardIM[ConstraintIndex]) : (FReal)1;
 		ConstraintSoftStiffness[ConstraintIndex] = SpringMassScale * RotationDrives.SettingsSoftStiffness[ConstraintIndex] * Dt * Dt;
 		ConstraintSoftDamping[ConstraintIndex] = (bUsePositionBasedDrives) ? SpringMassScale * RotationDrives.SettingsSoftDamping[ConstraintIndex] * Dt : 0;
@@ -2776,7 +2775,7 @@ void FPBDJointCachedSolver::ApplyRotationDrives(
 	{
 		for (int32 ConstraintIndex = 0; ConstraintIndex < 3; ++ConstraintIndex)
 		{
-			if (RotationDrives.bValidDatas[ConstraintIndex])
+			if (RotationDrives.GetValidDatas(ConstraintIndex))
 			{
 				ApplyAxisRotationDrive(ConstraintIndex, Dt);
 			}
@@ -2909,7 +2908,7 @@ void FPBDJointCachedSolver::ApplyRotationVelocityDrives(
 
 	for (int32 ConstraintIndex = 0; ConstraintIndex < 3; ++ConstraintIndex)
 	{
-		if (RotationDrives.bValidDatas[ConstraintIndex])
+		if (RotationDrives.GetValidDatas(ConstraintIndex))
 		{
 			ApplyAxisRotationVelocityDrive(ConstraintIndex, Dt);
 		}
@@ -2987,16 +2986,16 @@ void FAxisConstraintDatas::InitDatas(
 	const FReal HardStiffness,
 	const bool bResetLambdas)
 {
-	bSoftLimit[ConstraintIndex] = bHasSoftLimits;
+	SetSoftLimit(ConstraintIndex, bHasSoftLimits);
 	Data.ConstraintHardStiffness[ConstraintIndex] = HardStiffness;
 	Data.ConstraintSoftStiffness[ConstraintIndex] = SoftStiffness;
 	Data.ConstraintSoftDamping[ConstraintIndex] = SoftDamping;
 	ConstraintMaxLambda[ConstraintIndex] = 0;
 	SettingsSoftStiffness[ConstraintIndex] = SoftStiffness;
 	SettingsSoftDamping[ConstraintIndex] = SoftDamping;
-	bValidDatas[ConstraintIndex] = false;
-	bLimitsCheck[ConstraintIndex] = true;
-	MotionType[ConstraintIndex] = EJointMotionType::Free;
+	SetValidDatas(ConstraintIndex, false);
+	SetLimitsCheck(ConstraintIndex, true);
+	SetMotionType(ConstraintIndex, EJointMotionType::Free);
 	if (bResetLambdas)
 	{
 		Data.ConstraintLambda = FVec3::Zero();
@@ -3015,14 +3014,12 @@ void FAxisConstraintDatas::UpdateDatas(
 	const FVec3& DatasArm1,
 	const FReal DatasVX)
 {
-	bValidDatas[ConstraintIndex] = true;
-	bLimitsCheck[ConstraintIndex] = bCheckLimit;
+	SetValidDatas(ConstraintIndex, true);
+	SetLimitsCheck(ConstraintIndex, bCheckLimit);
 
 	Data.ConstraintCX[ConstraintIndex] = DatasCX;
 	ConstraintVX[ConstraintIndex] = DatasVX;
 	ConstraintRestitution[ConstraintIndex] = DatasRestitution;
-
-
 	Data.ConstraintArms[ConstraintIndex][0] = DatasArm0;
 	Data.ConstraintArms[ConstraintIndex][1] = DatasArm1;
 	Data.ConstraintAxis[ConstraintIndex] = DatasAxis;
@@ -3042,7 +3039,7 @@ void FAxisConstraintDatas::UpdateMass(
 	Data.ConstraintDRAxis[ConstraintIndex][0] = DatasIA0;
 	Data.ConstraintDRAxis[ConstraintIndex][1] = -DatasIA1;
 
-	if (bSoftLimit[ConstraintIndex])
+	if (GetSoftLimit(ConstraintIndex))
 	{
 		// If bUsePositionBasedDrives is false, we apply the velocity drive in the velocity solver phase so we don't include in the PBD settings
 		const FReal SpringMassScale = (bAccelerationMode) ? (FReal)1 / (Data.ConstraintHardIM[ConstraintIndex]) : (FReal)1;
diff --git a/Engine/Source/Runtime/Experimental/Chaos/Private/Chaos/Joint/PBDJointSolverGaussSeidel.cpp b/Engine/Source/Runtime/Experimental/Chaos/Private/Chaos/Joint/PBDJointSolverGaussSeidel.cpp
--- a/Engine/Source/Runtime/Experimental/Chaos/Private/Chaos/Joint/PBDJointSolverGaussSeidel.cpp
+++ b/Engine/Source/Runtime/Experimental/Chaos/Private/Chaos/Joint/PBDJointSolverGaussSeidel.cpp
@@ -16,7 +16,6 @@ static_assert(sizeof(ispc::FVector4) == sizeof(Chaos::FRotation3), "sizeof(ispc:
 static_assert(sizeof(ispc::FMatrix) == sizeof(Chaos::FMatrix33), "sizeof(ispc::FMatrix) != sizeof(Chaos::FMatrix33)");
 #endif
 
-//PRAGMA_DISABLE_OPTIMIZATION
 
 namespace Chaos
 {
diff --git a/Engine/Source/Runtime/Experimental/Chaos/Public/Chaos/Joint/PBDJointCachedSolverGaussSeidel.h b/Engine/Source/Runtime/Experimental/Chaos/Public/Chaos/Joint/PBDJointCachedSolverGaussSeidel.h
--- a/Engine/Source/Runtime/Experimental/Chaos/Public/Chaos/Joint/PBDJointCachedSolverGaussSeidel.h
+++ b/Engine/Source/Runtime/Experimental/Chaos/Public/Chaos/Joint/PBDJointCachedSolverGaussSeidel.h
@@ -115,13 +115,110 @@ struct FAxisConstraintDatas
 	FVec3 ConstraintRestitution;
 	FVec3 ConstraintLambdaVelocity;
 	
-	bool bValidDatas[3];
-	bool bLimitsCheck[3];
-	bool bSoftLimit[3];
-	bool bAccelerationMode;
-	bool bUseSimd;
-	
-	EJointMotionType MotionType[3];
+	uint8 bValidDatas0 : 1;
+	uint8 bValidDatas1 : 1;
+	uint8 bValidDatas2 : 1;
+	uint8 bLimitsCheck0 : 1;
+	uint8 bLimitsCheck1 : 1;
+	uint8 bLimitsCheck2 : 1;
+	uint8 bSoftLimit0 : 1;
+	uint8 bSoftLimit1 : 1;
+	uint8 bSoftLimit2 : 1;
+	uint8 bAccelerationMode : 1;
+	uint8 bUseSimd : 1;
+	uint8 MotionType0 : 2;
+	uint8 MotionType1 : 2;
+	uint8 MotionType2 : 2;
+
+	inline bool GetValidDatas(int32 Index) const 
+	{
+		switch (Index)
+		{
+		case 0: return bValidDatas0;
+		case 1: return bValidDatas1;
+		case 2: return bValidDatas2;
+		default: check(false);  return false;
+		}
+	}
+
+	inline bool GetLimitsCheck(int32 Index) const
+	{
+		switch (Index)
+		{
+		case 0: return bLimitsCheck0;
+		case 1: return bLimitsCheck1;
+		case 2: return bLimitsCheck2;
+		default: check(false);  return false;
+		}
+	}
+
+	inline bool GetSoftLimit(int32 Index) const
+	{
+		switch (Index)
+		{
+		case 0: return bSoftLimit0;
+		case 1: return bSoftLimit1;
+		case 2: return bSoftLimit2;
+		default: check(false);  return false;
+		}
+	}
+
+	inline EJointMotionType GetMotionType(int32 Index) const
+	{
+		switch (Index)
+		{
+		case 0: return EJointMotionType(MotionType0);
+		case 1: return EJointMotionType(MotionType1);
+		case 2: return EJointMotionType(MotionType2);
+		default: check(false);  return EJointMotionType();
+		}
+	}
+
+	inline void SetValidDatas(int32 Index, bool bValue)
+	{
+		switch (Index)
+		{
+		case 0: bValidDatas0 = bValue; break;
+		case 1: bValidDatas1 = bValue; break;
+		case 2: bValidDatas2 = bValue; break;
+		default: check(false);
+		}
+	}
+
+	inline void SetLimitsCheck(int32 Index, bool bValue)
+	{
+		switch (Index)
+		{
+		case 0: bLimitsCheck0 = bValue; break;
+		case 1: bLimitsCheck1 = bValue; break;
+		case 2: bLimitsCheck2 = bValue; break;
+		default: check(false);
+		}
+	}
+
+	inline void SetSoftLimit(int32 Index, bool bValue)
+	{
+		switch (Index)
+		{
+		case 0: bSoftLimit0 = bValue; break;
+		case 1: bSoftLimit1 = bValue; break;
+		case 2: bSoftLimit2 = bValue; break;
+		default: check(false);
+		}
+	}
+
+    inline void SetMotionType(int32 Index, EJointMotionType Value)
+	{
+		check(int32(Value) <= 3); // Stored only on 2 bits
+		switch (Index)
+		{
+		case 0: MotionType0 = uint8(Value); break;
+		case 1: MotionType1 = uint8(Value); break;
+		case 2: MotionType2 = uint8(Value); break;
+		default: check(false);
+		}
+	}
+
 };
 	
 	/**
@@ -260,11 +357,11 @@ struct FAxisConstraintDatas
 
 			for (int32 Axis = 0; Axis < 3; ++Axis)
 			{
-				if (!PositionConstraints.bUseSimd && PositionConstraints.bValidDatas[Axis])
+				if (!PositionConstraints.bUseSimd && PositionConstraints.GetValidDatas(Axis))
 				{
 					Impulse += PositionConstraints.Data.ConstraintLambda[Axis] * PositionConstraints.Data.ConstraintAxis[Axis];
 				}
-				if (!PositionDrives.bUseSimd && PositionDrives.bValidDatas[Axis])
+				if (!PositionDrives.bUseSimd && PositionDrives.GetValidDatas(Axis))
 				{
 					Impulse += PositionDrives.Data.ConstraintLambda[Axis] * PositionDrives.Data.ConstraintAxis[Axis];
 				}
@@ -308,11 +405,11 @@ struct FAxisConstraintDatas
 
 			for (int32 Axis = 0; Axis < 3; ++Axis)
 			{
-				if (!RotationConstraints.bUseSimd && RotationConstraints.bValidDatas[Axis])
+				if (!RotationConstraints.bUseSimd && RotationConstraints.GetValidDatas(Axis))
 				{
 					Impulse += RotationConstraints.Data.ConstraintLambda[Axis] * RotationConstraints.Data.ConstraintAxis[Axis];
 				}
-				if (!RotationDrives.bUseSimd && RotationDrives.bValidDatas[Axis])
+				if (!RotationDrives.bUseSimd && RotationDrives.GetValidDatas(Axis))
 				{
 					Impulse += RotationDrives.Data.ConstraintLambda[Axis] * RotationDrives.Data.ConstraintAxis[Axis];
 				}
