Fix for a crash in the Iris unresolved references optimization.

diff --git a/Engine/Source/Runtime/Experimental/Iris/Core/Private/Iris/ReplicationSystem/ReplicationReader.cpp b/Engine/Source/Runtime/Experimental/Iris/Core/Private/Iris/ReplicationSystem/ReplicationReader.cpp
--- a/Engine/Source/Runtime/Experimental/Iris/Core/Private/Iris/ReplicationSystem/ReplicationReader.cpp
+++ b/Engine/Source/Runtime/Experimental/Iris/Core/Private/Iris/ReplicationSystem/ReplicationReader.cpp
@@ -1409,6 +1409,11 @@ void FReplicationReader::RemoveUnresolvedObjectReferenceInReplicationInfo(FRepli
 		if (RefHandle == Handle)
 		{
 			It.RemoveCurrent();
+
+			if (bUseOptObjectRefTracking)
+			{
+				ReplicationInfo->RemoveUnresolvedHandleCount(RefHandle);
+			}
 		}
 	}
 }
@@ -1421,6 +1426,11 @@ void FReplicationReader::RemoveResolvedObjectReferenceInReplicationInfo(FReplica
 		if (RefHandle == Handle)
 		{
 			It.RemoveCurrent();
+
+			if (bUseOptObjectRefTracking)
+			{
+				ReplicationInfo->RemoveResolvedDynamicHandleCount(RefHandle);
+			}
 		}
 	}
 }
@@ -1495,6 +1505,7 @@ void FReplicationReader::CleanupReferenceTracking(FReplicatedObjectInfo* ObjectI
 		UE_LOG(LogIris, Verbose, TEXT("FReplicationReader::CleanupReferenceTracking Removing unresolved reference %s for %s"), *Handle.ToString(), *(NetRefHandleManager->GetNetRefHandleFromInternalIndex(ObjectIndex).ToString()));
 	}
 	ObjectInfo->UnresolvedObjectReferences.Reset();
+	ObjectInfo->UnresolvedHandleCount.Reset();
 
 	// Remove from resolved dynamic references
 	for (FObjectReferenceTracker::ElementType Element : ObjectInfo->ResolvedDynamicObjectReferences)
@@ -1505,6 +1516,7 @@ void FReplicationReader::CleanupReferenceTracking(FReplicatedObjectInfo* ObjectI
 		UE_LOG(LogIris, Verbose, TEXT("FReplicationReader::CleanupReferenceTracking Removing resolved dynamic reference %s for %s"), *Handle.ToString(), *(NetRefHandleManager->GetNetRefHandleFromInternalIndex(ObjectIndex).ToString()));
 	}
 	ObjectInfo->ResolvedDynamicObjectReferences.Reset();
+	ObjectInfo->ResolvedDynamicHandleCount.Reset();
 
 	// Remove from attachment resolve
 	ObjectsWithAttachmentPendingResolve.Remove(ObjectIndex);
@@ -2048,7 +2060,10 @@ void FReplicationReader::UpdateUnresolvableReferenceTracking()
 			for (const uint32 DependentObjectIndex : Dependents)
 			{
 				FReplicatedObjectInfo* ReplicationInfo = GetReplicatedObjectInfo(DependentObjectIndex);
-				RemoveUnresolvedObjectReferenceInReplicationInfo(ReplicationInfo, DestroyedHandle);
+				if (ensureMsgf(ReplicationInfo != nullptr, TEXT("Unable to find torn off unresolved replicated object info for %s"), *NetRefHandleManager->PrintObjectFromIndex(InternalIndex)))
+				{
+					RemoveUnresolvedObjectReferenceInReplicationInfo(ReplicationInfo, DestroyedHandle);
+				}
 			}
 		}
 
@@ -2065,7 +2080,10 @@ void FReplicationReader::UpdateUnresolvableReferenceTracking()
 				for (const uint32 DependentObjectIndex : Dependents)
 				{
 					FReplicatedObjectInfo* ReplicationInfo = GetReplicatedObjectInfo(DependentObjectIndex);
-					RemoveResolvedObjectReferenceInReplicationInfo(ReplicationInfo, DestroyedHandle);
+					if (ensureMsgf(ReplicationInfo != nullptr, TEXT("Unable to find torn off resolved replicated object info for %s"), *NetRefHandleManager->PrintObjectFromIndex(InternalIndex)))
+					{
+						RemoveResolvedObjectReferenceInReplicationInfo(ReplicationInfo, DestroyedHandle);
+					}
 				}
 			}
 			else
@@ -2073,9 +2091,12 @@ void FReplicationReader::UpdateUnresolvableReferenceTracking()
 				for (const uint32 DependentObjectIndex : Dependents)
 				{
 					FReplicatedObjectInfo* ReplicationInfo = GetReplicatedObjectInfo(DependentObjectIndex);
-					if (MoveResolvedObjectReferenceToUnresolvedInReplicationInfo(ReplicationInfo, DestroyedHandle))
+					if (ensureMsgf(ReplicationInfo != nullptr, TEXT("Unable to find resolved replicated object info for %s"), *NetRefHandleManager->PrintObjectFromIndex(InternalIndex)))
 					{
-						UnresolvedHandleToDependents.Add(DestroyedHandle, DependentObjectIndex);
+						if (MoveResolvedObjectReferenceToUnresolvedInReplicationInfo(ReplicationInfo, DestroyedHandle))
+						{
+							UnresolvedHandleToDependents.Add(DestroyedHandle, DependentObjectIndex);
+						}
 					}
 				}
 			}
