Fixed Clang warning: for loop has empty body

#rnx

diff --git a/Engine/Source/Runtime/Experimental/IoStore/OnDemand/Private/Tool/Purl.cpp b/Engine/Source/Runtime/Experimental/IoStore/OnDemand/Private/Tool/Purl.cpp
--- a/Engine/Source/Runtime/Experimental/IoStore/OnDemand/Private/Tool/Purl.cpp
+++ b/Engine/Source/Runtime/Experimental/IoStore/OnDemand/Private/Tool/Purl.cpp
@@ -84,7 +84,11 @@ static int32 PurlCommandEntry(const FContext& Context)
 			return;
 		}
 	});
-	for (; Loop.Tick(-1); FPlatformProcess::SleepNoStats(0.1f));
+	
+	while (Loop.Tick(-1))
+	{
+		FPlatformProcess::SleepNoStats(0.1f);
+	}
 
 	std::printf("Data: %u bytes\n", uint32(Dest.GetSize()));
 
