Disable keep alive for ISM Components to prevent OOM on platforms with more limited memory

diff --git a/Engine/Source/Runtime/Experimental/GeometryCollectionEngine/Private/GeometryCollection/GeometryCollectionISMPoolComponent.cpp b/Engine/Source/Runtime/Experimental/GeometryCollectionEngine/Private/GeometryCollection/GeometryCollectionISMPoolComponent.cpp
--- a/Engine/Source/Runtime/Experimental/GeometryCollectionEngine/Private/GeometryCollection/GeometryCollectionISMPoolComponent.cpp
+++ b/Engine/Source/Runtime/Experimental/GeometryCollectionEngine/Private/GeometryCollection/GeometryCollectionISMPoolComponent.cpp
@@ -11,7 +11,7 @@
 
 // Don't release ISM components when they empty, but keep them (and their scene proxy) alive.
 // This can remove the high cost associated with repeated registration, scene proxy creation and mesh draw command creation.
-static bool GComponentKeepAlive = true;
+static bool GComponentKeepAlive = false; // Disabled due to current OOM issues on certain platforms
 FAutoConsoleVariableRef CVarISMPoolComponentKeepAlive(
 	TEXT("r.ISMPool.ComponentKeepAlive"),
 	GComponentKeepAlive,
