Dataflow:
- Fixed SkeletalMeshToCollection node to create Children attribute
- Updated BlueprintToCollection node to create root node and parent everything under it

diff --git a/Engine/Plugins/Experimental/GeometryCollectionPlugin/Source/GeometryCollectionNodes/Private/Dataflow/Nodes/GeometryCollectionAssetNodes.cpp b/Engine/Plugins/Experimental/GeometryCollectionPlugin/Source/GeometryCollectionNodes/Private/Dataflow/Nodes/GeometryCollectionAssetNodes.cpp
--- a/Engine/Plugins/Experimental/GeometryCollectionPlugin/Source/GeometryCollectionNodes/Private/Dataflow/Nodes/GeometryCollectionAssetNodes.cpp
+++ b/Engine/Plugins/Experimental/GeometryCollectionPlugin/Source/GeometryCollectionNodes/Private/Dataflow/Nodes/GeometryCollectionAssetNodes.cpp
@@ -317,7 +317,8 @@ void FBlueprintToCollectionDataflowNode::Evaluate(Dataflow::FContext& Context, c
 
 				if (AActor* PreviewActor = PreviewWorld->SpawnActor(Blueprint->GeneratedClass, nullptr, SpawnInfo))
 				{
-					FGeometryCollectionEngineConversion::ConvertActorToGeometryCollection(PreviewActor, OutCollection, OutMaterials, OutInstancedMeshes, bSplitComponents);
+					FGeometryCollectionEngineConversion::FSkeletalMeshToCollectionConversionParameters ConversionParameters;
+					FGeometryCollectionEngineConversion::ConvertActorToGeometryCollection(PreviewActor, OutCollection, OutMaterials, OutInstancedMeshes, ConversionParameters, bSplitComponents);
 				}
 			}
 		}
diff --git a/Engine/Source/Runtime/Experimental/Chaos/Private/GeometryCollection/GeometryCollectionClusteringUtility.cpp b/Engine/Source/Runtime/Experimental/Chaos/Private/GeometryCollection/GeometryCollectionClusteringUtility.cpp
--- a/Engine/Source/Runtime/Experimental/Chaos/Private/GeometryCollection/GeometryCollectionClusteringUtility.cpp
+++ b/Engine/Source/Runtime/Experimental/Chaos/Private/GeometryCollection/GeometryCollectionClusteringUtility.cpp
@@ -63,7 +63,7 @@ int32 FGeometryCollectionClusteringUtility::ClusterBonesUnderNewNodeWithParent(F
 	return NewBoneIndex;
 }
 
-void FGeometryCollectionClusteringUtility::ClusterAllBonesUnderNewRoot(FGeometryCollection* GeometryCollection)
+void FGeometryCollectionClusteringUtility::ClusterAllBonesUnderNewRoot(FGeometryCollection* GeometryCollection, FName RootName, bool bUpdateChildBoneNames)
 {
 	check(GeometryCollection);
 	bool CalcNewLocalTransform = true;
@@ -73,6 +73,7 @@ void FGeometryCollectionClusteringUtility::ClusterAllBonesUnderNewRoot(FGeometry
 	TManagedArray<int32>& Parents = GeometryCollection->Parent;
 	TManagedArray<TSet<int32>>& Children = GeometryCollection->Children;
 	TManagedArray<int32>& SimulationType = GeometryCollection->SimulationType;
+	TManagedArray<FLinearColor>& BoneColors = GeometryCollection->BoneColor;
 
 	TArray<int32> ChildBones;
 	int32 NumElements = GeometryCollection->NumElements(FGeometryCollection::TransformGroup);
@@ -98,7 +99,14 @@ void FGeometryCollectionClusteringUtility::ClusterAllBonesUnderNewRoot(FGeometry
 	}
 
 	// New Bone Setup takes level/parent from the first of the Selected Bones
-	BoneNames[RootNoneIndex] = "ClusterBone";
+	if (RootName.IsNone())
+	{
+		BoneNames[RootNoneIndex] = "ClusterBone";
+	}
+	else
+	{
+		BoneNames[RootNoneIndex] = RootName.ToString();
+	}
 	Parents[RootNoneIndex] = FGeometryCollection::Invalid;
 	Children[RootNoneIndex] = TSet<int32>(ChildBones);
 	SimulationType[RootNoneIndex] = FGeometryCollection::ESimulationTypes::FST_Clustered;
@@ -131,8 +139,13 @@ void FGeometryCollectionClusteringUtility::ClusterAllBonesUnderNewRoot(FGeometry
 
 	Transforms[RootNoneIndex] = FTransform3f::Identity;
 
+	if (bUpdateChildBoneNames)
+	{
+		RecursivelyUpdateChildBoneNames(RootNoneIndex, Children, BoneNames);
+	}
 
-	RecursivelyUpdateChildBoneNames(RootNoneIndex, Children, BoneNames);
+	const FColor RandBoneColor(FMath::Rand() % 100 + 5, FMath::Rand() % 100 + 5, FMath::Rand() % 100 + 5, 255);
+	BoneColors[RootNoneIndex] = FLinearColor(RandBoneColor);
 
 	ValidateResults(GeometryCollection);
 }
diff --git a/Engine/Source/Runtime/Experimental/Chaos/Public/GeometryCollection/GeometryCollectionClusteringUtility.h b/Engine/Source/Runtime/Experimental/Chaos/Public/GeometryCollection/GeometryCollectionClusteringUtility.h
--- a/Engine/Source/Runtime/Experimental/Chaos/Public/GeometryCollection/GeometryCollectionClusteringUtility.h
+++ b/Engine/Source/Runtime/Experimental/Chaos/Public/GeometryCollection/GeometryCollectionClusteringUtility.h
@@ -43,7 +43,7 @@ public:
 	static CHAOS_API int32 ClusterBonesUnderNewNodeWithParent(FGeometryCollection* GeometryCollection, const int32 ParentOfNewNode, const TArray<int32>& SelectedBones, bool CalcNewLocalTransform, bool Validate = true);
 
 	/** Cluster all existing bones under a new root node, so there is now only one root node and a completely flat hierarchy underneath it */
-	static CHAOS_API void ClusterAllBonesUnderNewRoot(FGeometryCollection* GeometryCollection);
+	static CHAOS_API void ClusterAllBonesUnderNewRoot(FGeometryCollection* GeometryCollection, FName RootName = NAME_None, bool bUpdateChildBoneNames = true);
 
 	/** Cluster all source bones under an existing node, algorithm chooses best node to add to 'closest to root' */
 	static CHAOS_API void ClusterBonesUnderExistingNode(FGeometryCollection* GeometryCollection, const TArray<int32>& SourceElements);
diff --git a/Engine/Source/Runtime/Experimental/GeometryCollectionEngine/Private/GeometryCollection/GeometryCollectionEngineConversion.cpp b/Engine/Source/Runtime/Experimental/GeometryCollectionEngine/Private/GeometryCollection/GeometryCollectionEngineConversion.cpp
--- a/Engine/Source/Runtime/Experimental/GeometryCollectionEngine/Private/GeometryCollection/GeometryCollectionEngineConversion.cpp
+++ b/Engine/Source/Runtime/Experimental/GeometryCollectionEngine/Private/GeometryCollection/GeometryCollectionEngineConversion.cpp
@@ -1158,6 +1158,7 @@ bool FGeometryCollectionEngineConversion::AppendSkeletalMesh(const USkeletalMesh
 	// Transform Attributes
 	TManagedArray<FTransform3f>& LocalSpaceTransform = GeometryCollection.ModifyAttribute<FTransform3f>(FTransformCollection::TransformAttribute, FTransformCollection::TransformGroup);
 	TManagedArray<int32>& Parent = GeometryCollection.ModifyAttribute<int32>(FTransformCollection::ParentAttribute, FTransformCollection::TransformGroup);
+	TManagedArray<TSet<int32>>& Children = GeometryCollection.ModifyAttribute<TSet<int32>>(FTransformCollection::ChildrenAttribute, FTransformCollection::TransformGroup);
 	TManagedArray<FLinearColor>& BoneColor = GeometryCollection.ModifyAttribute<FLinearColor>("BoneColor", FTransformCollection::TransformGroup);
 	TManagedArray<FString>& BoneName = GeometryCollection.ModifyAttribute<FString>("BoneName", FTransformCollection::TransformGroup);
 	TManagedArray<int32>& SimulationType = GeometryCollection.ModifyAttribute<int32>("SimulationType", FTransformCollection::TransformGroup);
@@ -1181,21 +1182,28 @@ bool FGeometryCollectionEngineConversion::AppendSkeletalMesh(const USkeletalMesh
 	int32 TransformBaseIndex = INDEX_NONE;
 	const USkeleton* Skeleton = InSkeletalMesh->GetSkeleton();
 	const TArray<FTransform>& RestArray = Skeleton->GetRefLocalPoses();
-	const FReferenceSkeleton& ReferenceSkeletion = Skeleton->GetReferenceSkeleton();
+	const FReferenceSkeleton& ReferenceSkeleton = Skeleton->GetReferenceSkeleton();
 
-	if (ReferenceSkeletion.GetNum())
+	if (ReferenceSkeleton.GetNum())
 	{
-		TransformBaseIndex = GeometryCollection.AddElements(ReferenceSkeletion.GetNum(), FGeometryCollection::TransformGroup);
+		TransformBaseIndex = GeometryCollection.AddElements(ReferenceSkeleton.GetNum(), FGeometryCollection::TransformGroup);
 		RootIndex = TransformBaseIndex;
 
-		for (int32 BoneIndex = 0; BoneIndex < ReferenceSkeletion.GetNum(); BoneIndex++)
+		for (int32 BoneIndex = 0; BoneIndex < ReferenceSkeleton.GetNum(); BoneIndex++)
 		{
 			// For validation against the component space position use
-			// FTransform ComponentSpaceTransform =FAnimationRuntime::GetComponentSpaceTransformRefPose(ReferenceSkeletion, SkeletalBoneMap[BoneIndex]);
+			// FTransform ComponentSpaceTransform =FAnimationRuntime::GetComponentSpaceTransformRefPose(ReferenceSkeleton, SkeletalBoneMap[BoneIndex]);
 
 			LocalSpaceTransform[TransformBaseIndex + BoneIndex] = FTransform3f(RestArray[BoneIndex]);
-			BoneName[TransformBaseIndex + BoneIndex] = ReferenceSkeletion.GetRefBoneInfo()[BoneIndex].Name.ToString();
-			Parent[TransformBaseIndex + BoneIndex] = ReferenceSkeletion.GetRefBoneInfo()[BoneIndex].ParentIndex;
+			BoneName[TransformBaseIndex + BoneIndex] = ReferenceSkeleton.GetRefBoneInfo()[BoneIndex].Name.ToString();
+			Parent[TransformBaseIndex + BoneIndex] = ReferenceSkeleton.GetRefBoneInfo()[BoneIndex].ParentIndex;
+
+			TArray<int32> ChildrenArr;
+			if (ReferenceSkeleton.GetDirectChildBones(BoneIndex, ChildrenArr))
+			{
+				Children[TransformBaseIndex + BoneIndex].Append(ChildrenArr);
+			}
+
 			SimulationType[TransformBaseIndex + BoneIndex] = FGeometryCollection::ESimulationTypes::FST_None;
 			BoneColor[TransformBaseIndex + BoneIndex] = FLinearColor::MakeRandomColor();
 
@@ -1570,7 +1578,7 @@ void FGeometryCollectionEngineConversion::ConvertGeometryCollectionToGeometryCol
 	}
 }
 
-void FGeometryCollectionEngineConversion::ConvertActorToGeometryCollection(const AActor* Actor, FManagedArrayCollection& OutCollection, TArray<TObjectPtr<UMaterial>>& OutMaterials, TArray<FGeometryCollectionAutoInstanceMesh>& OutInstancedMeshes, bool bSplitComponents)
+void FGeometryCollectionEngineConversion::ConvertActorToGeometryCollection(const AActor* Actor, FManagedArrayCollection& OutCollection, TArray<TObjectPtr<UMaterial>>& OutMaterials, TArray<FGeometryCollectionAutoInstanceMesh>& OutInstancedMeshes, const FSkeletalMeshToCollectionConversionParameters& ConversionParameters, bool bSplitComponents)
 {
 #if WITH_EDITORONLY_DATA
 	const FTransform ActorTransform(Actor->GetTransform());
@@ -1631,6 +1639,13 @@ void FGeometryCollectionEngineConversion::ConvertActorToGeometryCollection(const
 
 		NewGeometryCollection->InitializeMaterials();
 
+		if (ConversionParameters.bParentAllBonesUnderNewRoot && FGeometryCollectionClusteringUtility::ContainsMultipleRootBones(NewGeometryCollection->GetGeometryCollection().Get()))
+		{
+			UE_LOG(UGeometryCollectionConversionLogging, Log, TEXT("FGeometryCollectionEngineConversion::ConvertActorToGeometryCollection() - All bones were parented under new root."));
+
+			FGeometryCollectionClusteringUtility::ClusterAllBonesUnderNewRoot(NewGeometryCollection->GetGeometryCollection().Get(), FName("root"), false);
+		}
+
 		// InstanceMeshes
 		OutInstancedMeshes.Append(NewGeometryCollection->AutoInstanceMeshes);
 
diff --git a/Engine/Source/Runtime/Experimental/GeometryCollectionEngine/Public/GeometryCollection/GeometryCollectionEngineConversion.h b/Engine/Source/Runtime/Experimental/GeometryCollectionEngine/Public/GeometryCollection/GeometryCollectionEngineConversion.h
--- a/Engine/Source/Runtime/Experimental/GeometryCollectionEngine/Public/GeometryCollection/GeometryCollectionEngineConversion.h
+++ b/Engine/Source/Runtime/Experimental/GeometryCollectionEngine/Public/GeometryCollection/GeometryCollectionEngineConversion.h
@@ -27,6 +27,11 @@ class FGeometryCollectionEngineConversion
 {
 public:
 
+	struct FSkeletalMeshToCollectionConversionParameters
+	{
+		bool bParentAllBonesUnderNewRoot = true;
+	};
+
 	/**
 	 * Appends materials to a GeometryCollectionComponent.
 	 * @param Materials : Materials fetched from the StaticMeshComponent used to configure this geometry
@@ -203,6 +208,6 @@ public:
 	*  @param OutInstancedMeshes : InstancedMeshes
 	*  @param bSplitComponents : Split the components
 	*/
-	static GEOMETRYCOLLECTIONENGINE_API void ConvertActorToGeometryCollection(const AActor* Actor, FManagedArrayCollection& OutCollection, TArray<TObjectPtr<UMaterial>>& OutMaterials, TArray<FGeometryCollectionAutoInstanceMesh>& OutInstancedMeshes, bool bSplitComponents = false);
+	static GEOMETRYCOLLECTIONENGINE_API void ConvertActorToGeometryCollection(const AActor* Actor, FManagedArrayCollection& OutCollection, TArray<TObjectPtr<UMaterial>>& OutMaterials, TArray<FGeometryCollectionAutoInstanceMesh>& OutInstancedMeshes, const FSkeletalMeshToCollectionConversionParameters& ConversionParameters, bool bSplitComponents = false);
 
 };
