Parse correct service and TOC url when generating IAS configuration file.

The IAS upload command didn't parse the service and TOC correctly when uploading to a local S3 compatible
service endpoint. The config file should specify a service URL without any trailing host path and
append the bucket name. The bucket name should be omitted when uploading to AWS S3.

diff --git a/Engine/Source/Runtime/Experimental/IoStore/OnDemand/Private/IoStoreOnDemand.cpp b/Engine/Source/Runtime/Experimental/IoStore/OnDemand/Private/IoStoreOnDemand.cpp
--- a/Engine/Source/Runtime/Experimental/IoStore/OnDemand/Private/IoStoreOnDemand.cpp
+++ b/Engine/Source/Runtime/Experimental/IoStore/OnDemand/Private/IoStoreOnDemand.cpp
@@ -1036,72 +1036,64 @@ bool FS3UploadQueue::Flush()
 }
 
 ////////////////////////////////////////////////////////////////////////////////
+static void GetTocPath(
+	FStringView ServiceUrl,
+	FStringView Bucket,
+	FStringView BucketPrefix,
+	FString& OutServiceUrl,
+	FString& OutTocPath)
+{
+	// The configuration file should specify a service URL without any trailing
+	// host path, i.e. http://{host:port}/{host-path}. Add the trailing path
+	// to the TOC path to form the complete path the TOC from the host, i.e
+	// TocPath={host-path}/{bucket}/{bucket-prefix}/{toc-hash}.uchunktoc
+	
+	if (BucketPrefix.StartsWith(TEXT("/")))
+	{
+		BucketPrefix.RemovePrefix(1);
+	}
+	if (BucketPrefix.EndsWith(TEXT("/")))
+	{
+		BucketPrefix.RemoveSuffix(1);
+	}
 
-struct FResolvedPaths
-{
-	FString ServiceUrl;
-	FString TocPath;
-	FString ChunkPrefix;
-};
-
-[[nodiscard]] static FResolvedPaths ResolvePaths(const FIoStoreUploadParams& UploadParams, FStringView TocFilename)
-{
-	FResolvedPaths Result;
-	if (!UploadParams.BucketPrefix.IsEmpty())
+	if (ServiceUrl.IsEmpty())
 	{
-		Result.ServiceUrl = UploadParams.ServiceUrl;
-		Result.TocPath = TocFilename;
-		Result.ChunkPrefix = UploadParams.BucketPrefix;
+		// If the service URL is empty we assume uploading to AWS S3 using the Region parameter
+		// and that we don't need to prefix with the bucket name.
+		OutServiceUrl.Empty();
+		OutTocPath = BucketPrefix;
 	}
 	else
 	{
-		// The configuration file should specify a service URL without any trailing
-		// host path, i.e. http://{host:port}/{host-path}. Add the trailing path
-		// to the TOC path to form the complete path the TOC from the host, i.e
-		// TocPath={host-path}/{bucket}/{bucket-prefix}/{toc-hash}.uchunktoc
+		FStringView HostSuffix;
+		int32 Idx = INDEX_NONE;
+		ensure(ServiceUrl.FindChar(':', Idx));
+		const int32 SchemeEnd = Idx + 3;
 
-		FStringView ServiceUrl = UploadParams.ServiceUrl;
-		FStringView TocPrefx;
+		Idx = INDEX_NONE;
+		if (ServiceUrl.RightChop(SchemeEnd).FindChar('/', Idx))
 		{
-			// Find the first '//' then find the first '/' after that
-			int32 Sep = INDEX_NONE;
-
-			const int32 DoubleSlash = ServiceUrl.Find(TEXT("//"));
-			if (DoubleSlash != INDEX_NONE)
+			OutServiceUrl = ServiceUrl.Left(SchemeEnd + Idx);
+			HostSuffix = ServiceUrl.RightChop(OutServiceUrl.Len() + 1);
+			if (HostSuffix.EndsWith(TEXT("/")))
 			{
-				Sep = ServiceUrl.Find(TEXT("/"), DoubleSlash + 2);
-			}
-
-			if (Sep != INDEX_NONE)
-			{
-				TocPrefx = ServiceUrl.RightChop(Sep + 1);
-				ServiceUrl.LeftInline(Sep);
+				HostSuffix.RemoveSuffix(1);
 			}
 		}
-
-		Result.ServiceUrl = ServiceUrl;
-
-		TStringBuilder<256> Prefix;
-		FPathViews::Append(Prefix, TocPrefx, UploadParams.Bucket);
-
-		Result.ChunkPrefix = Prefix;
-		
-		if (!TocFilename.IsEmpty())
+		else
 		{
-			TStringBuilder<256> TocPath;
-			FPathViews::Append(TocPath, Prefix, TocFilename);
+			OutServiceUrl = ServiceUrl;
+		}
 
-			Result.TocPath = TocPath;
+		TStringBuilder<256> Sb;
+		if (!HostSuffix.IsEmpty())
+		{
+			Sb << HostSuffix << TEXT("/");
 		}
+		Sb << Bucket << TEXT("/") << BucketPrefix;
+		OutTocPath = Sb;
 	}
-
-	return Result;
-}
-
-[[nodiscard]] static FString ResolveChunksDirectory(const FIoStoreUploadParams& UploadParams)
-{
-	FResolvedPaths Paths = ResolvePaths(UploadParams, FStringView());
-	return Paths.ChunkPrefix;
 }
 
 static FIoStatus WriteContainerFiles(FOnDemandToc& OnDemandToc, const TMap<FIoHash, FString>& UTocPaths)
@@ -1420,8 +1412,6 @@ TIoStatusOr<FIoStoreUploadResult> UploadContainerFiles(
 	uint64 TotalUploadedBytes = 0;
 
 	FOnDemandToc OnDemandToc;
-	OnDemandToc.Header.ChunksDirectory = ResolveChunksDirectory(UploadParams);
-
 	OnDemandToc.Containers.Reserve(ContainerFiles.Num());
 
 	// Map of the .utoc paths that we have created ondemand containers for, indexed by their hash so that the paths can
@@ -1600,6 +1590,9 @@ TIoStatusOr<FIoStoreUploadResult> UploadContainerFiles(
 		return FIoStatus(EIoErrorCode::WriteError, TEXT("Failed to upload chunk(s)"));
 	}
 
+	FString ServiceUrl;
+	GetTocPath(UploadParams.ServiceUrl, UploadParams.Bucket, UploadParams.BucketPrefix, ServiceUrl, OnDemandToc.Header.ChunksDirectory);
+
 	OnDemandToc.Meta.EpochTimestamp = FDateTime::Now().ToUnixTimestamp();
 	OnDemandToc.Meta.BuildVersion = UploadParams.BuildVersion;
 	OnDemandToc.Meta.TargetPlatform = UploadParams.TargetPlatform;
@@ -1610,18 +1603,17 @@ TIoStatusOr<FIoStoreUploadResult> UploadContainerFiles(
 		Ar << OnDemandToc;
 
 		UploadResult.TocHash = FIoHash::HashBuffer(Ar.GetView());
+		const FString TocFilename = LexToString(UploadResult.TocHash) + TEXT(".iochunktoc");
+
 		TStringBuilder<256> Key;
 		if (!UploadParams.BucketPrefix.IsEmpty())
 		{
 			Key << UploadParams.BucketPrefix.ToLower() << TEXT("/");
 		}
-		Key << LexToString(UploadResult.TocHash) << TEXT(".iochunktoc");
-		
-		FResolvedPaths ResolvedPaths = ResolvePaths(UploadParams, Key);
-
-		UploadResult.ServiceUrl = MoveTemp(ResolvedPaths.ServiceUrl);
-		UploadResult.TocPath = MoveTemp(ResolvedPaths.TocPath);
+		Key << TocFilename;
 
+		UploadResult.ServiceUrl = ServiceUrl;
+		UploadResult.TocPath = OnDemandToc.Header.ChunksDirectory / TocFilename;
 		UploadResult.TocSize = Ar.TotalSize();
 
 		const FS3PutObjectResponse Response = Client.TryPutObject(FS3PutObjectRequest{UploadParams.Bucket, Key.ToString(), Ar.GetView()});
