Parse IAS chunk directory correctly with using Ias.TocUrl=<URL>

diff --git a/Engine/Source/Runtime/Experimental/IoStore/OnDemand/Private/OnDemandIoStore.cpp b/Engine/Source/Runtime/Experimental/IoStore/OnDemand/Private/OnDemandIoStore.cpp
--- a/Engine/Source/Runtime/Experimental/IoStore/OnDemand/Private/OnDemandIoStore.cpp
+++ b/Engine/Source/Runtime/Experimental/IoStore/OnDemand/Private/OnDemandIoStore.cpp
@@ -529,19 +529,28 @@ void FOnDemandIoStore::CreateContainersFromToc(
 	const FOnDemandTocHeader& Header = Toc.Header;
 	const FName CompressionFormat(Header.CompressionFormat);
 
-	TStringBuilder<128> ChunksDirectory;
-
+	TStringBuilder<128> Sb;
+	FStringView ChunksDirectory;
 	{
 		if (TocPath.IsEmpty() == false)
 		{
-			Algo::Transform(TocPath, AppendChars(ChunksDirectory), FChar::ToLower);
+			Algo::Transform(TocPath, AppendChars(Sb), FChar::ToLower);
 		}
 		else
 		{
-			Algo::Transform(Toc.Header.ChunksDirectory, AppendChars(ChunksDirectory), FChar::ToLower);
+			Algo::Transform(Toc.Header.ChunksDirectory, AppendChars(Sb), FChar::ToLower);
 		}
+		FPathViews::Append(Sb, TEXT("chunks"));
 
-		FPathViews::Append(ChunksDirectory, TEXT("chunks"));
+		ChunksDirectory = Sb;
+		if (ChunksDirectory.StartsWith('/'))
+		{
+			TocPath.RemovePrefix(1);
+		}
+		if (ChunksDirectory.EndsWith('/'))
+		{
+			TocPath.RemoveSuffix(1);
+		}
 	}
 
 	for (FOnDemandTocContainerEntry& ContainerEntry : Toc.Containers)
@@ -549,7 +558,7 @@ void FOnDemandIoStore::CreateContainersFromToc(
 		FSharedOnDemandContainer Container = MakeShared<FOnDemandContainer>();
 		Container->Name					= MoveTemp(ContainerEntry.ContainerName);
 		Container->MountId				= MountId;
-		Container->ChunksDirectory		= StringCast<ANSICHAR>(ChunksDirectory.ToString());
+		Container->ChunksDirectory		= StringCast<ANSICHAR>(ChunksDirectory.GetData(), ChunksDirectory.Len());
 		Container->EncryptionKeyGuid	= MoveTemp(ContainerEntry.EncryptionKeyGuid);
 		Container->BlockSize			= Header.BlockSize;
 		Container->BlockSizes			= MoveTemp(ContainerEntry.BlockSizes);
