Module's manager is not always available

#rnx

diff --git a/Engine/Source/Runtime/Experimental/IoStore/HttpClient/Private/Peer.inl b/Engine/Source/Runtime/Experimental/IoStore/HttpClient/Private/Peer.inl
--- a/Engine/Source/Runtime/Experimental/IoStore/HttpClient/Private/Peer.inl
+++ b/Engine/Source/Runtime/Experimental/IoStore/HttpClient/Private/Peer.inl
@@ -36,8 +36,10 @@ FSslContext::FSslContext(const char* InHostName, const FPemCert& CertPin)
 
 	if (CertPin.GetSize() == 0)
 	{
+#if WITH_SSL
 		const ISslCertificateManager& CertManager = FSslModule::Get().GetCertificateManager();
 		CertManager.AddCertificatesToSslContext(Context);
+#endif
 		return;
 	}
 
