Get namespace use in order.

#rnx

diff --git a/Engine/Source/Runtime/Experimental/IoStore/HttpClient/Private/Socket.inl b/Engine/Source/Runtime/Experimental/IoStore/HttpClient/Private/Socket.inl
--- a/Engine/Source/Runtime/Experimental/IoStore/HttpClient/Private/Socket.inl
+++ b/Engine/Source/Runtime/Experimental/IoStore/HttpClient/Private/Socket.inl
@@ -15,16 +15,6 @@
 #		include <ws2tcpip.h>
 #		pragma comment(lib, "Ws2_32.lib")
 #	endif // NO_UE_INCLUDES
-	using SocketType	= SOCKET;
-	using MsgFlagType	= int;
-
-	enum { SHUT_RDWR = SD_BOTH };
-
-#	define IAS_HTTP_USE_POLL
-	template <typename... ArgTypes> auto poll(ArgTypes... Args)
-	{
-		return WSAPoll(Forward<ArgTypes>(Args)...);
-	}
 
 	// Winsock defines "PF_MAX" indicating the protocol family count. This
 	// however competes with UE definitions related to pixel formats.
@@ -32,6 +22,19 @@
 	#	undef PF_MAX
 	#endif
 
+	namespace UE::IoStore::HTTP
+	{
+		using SocketType	= SOCKET;
+		using MsgFlagType	= int;
+
+		enum { SHUT_RDWR = SD_BOTH };
+
+#	define IAS_HTTP_USE_POLL
+		template <typename... ArgTypes> auto poll(ArgTypes... Args)
+		{
+			return WSAPoll(Forward<ArgTypes>(Args)...);
+		}
+	} // namespace UE::IoStore::HTTP
 #elif PLATFORM_APPLE | PLATFORM_UNIX | PLATFORM_ANDROID
 #	include <arpa/inet.h>
 #	include <fcntl.h>
@@ -40,14 +43,27 @@
 #	include <sys/select.h>
 #	include <sys/socket.h>
 #	include <unistd.h>
-	using SocketType	= int;
-	using MsgFlagType	= int;
-	static int32 closesocket(int32 Socket) { return close(Socket); }
+	namespace UE::IoStore::HTTP
+	{
+		using SocketType	= int;
+		using MsgFlagType	= int;
+		static int32 closesocket(int32 Socket) { return close(Socket); }
+	}
 #	define IAS_HTTP_USE_POLL
 #else
 #	include "CoreHttp/Http.inl"
 #endif
 
+// }}}
+
+namespace UE::IoStore::HTTP
+{
+
+// {{{1 socket .................................................................
+
+////////////////////////////////////////////////////////////////////////////////
+static const SocketType InvalidSocket = ~SocketType(0);
+
 static_assert(sizeof(sockaddr_in::sin_addr) == sizeof(uint32));
 
 #if PLATFORM_MICROSOFT
@@ -59,15 +75,6 @@ static_assert(sizeof(sockaddr_in::sin_addr) == sizeof(uint32));
 	static_assert(EWOULDBLOCK == EAGAIN);
 #endif
 
-static const SocketType InvalidSocket = ~SocketType(0);
-
-// }}}
-
-namespace UE::IoStore::HTTP
-{
-
-// {{{1 socket .................................................................
-
 ////////////////////////////////////////////////////////////////////////////////
 class FSocket
 {
