Constraints:
    - ensure hierarchical dependency with external parents (child target object different from parent target object)
    - ensure that the creation order is maintained if we have not created a hierarchical dependency
    - restore child dependency to ensure that it ticks after the constraint (using cvar Constraints.ForceChildDependency 1 - default is 0 for now)

diff --git a/Engine/Source/Runtime/Experimental/Animation/Constraints/Private/TransformConstraint.cpp b/Engine/Source/Runtime/Experimental/Animation/Constraints/Private/TransformConstraint.cpp
--- a/Engine/Source/Runtime/Experimental/Animation/Constraints/Private/TransformConstraint.cpp
+++ b/Engine/Source/Runtime/Experimental/Animation/Constraints/Private/TransformConstraint.cpp
@@ -211,6 +211,13 @@ static FAutoConsoleVariableRef CVarPreTickChild(
 	TEXT("Force child ticking before constraint computation.")
 	);
 
+static bool	bForceChildDependency = false;
+static FAutoConsoleVariableRef CVarForceChildDependency(
+	TEXT("Constraints.ForceChildDependency"),
+	bForceChildDependency,
+	TEXT("Force child dependency on the current constraint.")
+	);
+	
 void PreEvaluateHandle(const TObjectPtr<UTransformableHandle>& InHandle) 
 {
 	if ((bPreEvaluateChild || bPreTickChild) && ::IsValid(InHandle))
@@ -254,6 +261,25 @@ FString GetHandleLabel(const UTransformableHandle* InHandle)
 	return InHandle->GetName();
 #endif		
 }
+
+void LogDependency( const FString& InDescription, const UTransformableHandle* InParentHandle, const UTransformableHandle* InChildHandle,
+					const UTickableConstraint* InParentConstraint, const UTickableConstraint* InChildConstraint)
+{
+	if (!bDebugDependencies)
+	{
+		return;
+	}
+	
+	if (!InParentHandle || !InChildHandle || !InParentConstraint || !InChildConstraint)
+	{
+		return;
+	}
+
+	UE_LOG(LogTemp, Warning, TEXT("%s: '%s' is parent of '%s' so '%s' must tick before '%s'"),
+		*InDescription,
+		*GetHandleLabel(InParentHandle), *GetHandleLabel(InChildHandle),
+		*GetConstraintLabel(InParentConstraint), *GetConstraintLabel(InChildConstraint));
+}
 	
 }
 
@@ -477,7 +503,7 @@ void UTickableTransformConstraint::SetupDependencies(const UWorld* InWorld)
 	if (ChildTickFunction)
 	{
 		USkeletalMeshComponent* ChildOwner = Cast<USkeletalMeshComponent>(ChildTRSHandle->GetTarget().Get());
-		if (ChildOwner == nullptr)
+		if (ChildOwner == nullptr || ConstraintLocals::bForceChildDependency)
 		{
 			// force InChild to tick after ConstraintTickFunction does.
 			// Note that this might not register anything if the child can't tick (static meshes for instance)
@@ -1914,6 +1940,7 @@ bool FTransformConstraintUtils::AddConstraint(
 bool FTransformConstraintUtils::BuildDependencies(UWorld* InWorld, UTickableTransformConstraint* InConstraint)
 {
 	using namespace ConstraintLocals;
+	using ConstraintWeakPtr = TWeakObjectPtr<UTickableConstraint>;
 	
 	if (!ensure(InWorld))
 	{
@@ -1935,7 +1962,7 @@ bool FTransformConstraintUtils::BuildDependencies(UWorld* InWorld, UTickableTran
 	
  	// get previous child constraints
 	const FConstraintsManagerController& Controller = FConstraintsManagerController::Get(InWorld);
-	TArray<TWeakObjectPtr<UTickableConstraint>> ChildParentConstraints = Controller.GetParentConstraints(ChildHandle->GetHash(), true);
+	TArray<ConstraintWeakPtr> ChildParentConstraints = Controller.GetParentConstraints(ChildHandle->GetHash(), true);
 	ChildParentConstraints.Remove(InConstraint);
 
 	// add dependencies with the last child constraint
@@ -1944,7 +1971,7 @@ bool FTransformConstraintUtils::BuildDependencies(UWorld* InWorld, UTickableTran
 		const FGuid LastChildConstraintID = ChildParentConstraints.Last()->ConstraintID;
 		if (bDebugDependencies)
 		{
-			UE_LOG(LogTemp, Warning, TEXT("DEPENDENCY 0: tick after last constraint."));
+			UE_LOG(LogTemp, Warning, TEXT("Order Dependency: tick after last constraint."));
 		}
 		Controller.SetConstraintsDependencies(LastChildConstraintID, InConstraint->ConstraintID);
 	}
@@ -1956,10 +1983,11 @@ bool FTransformConstraintUtils::BuildDependencies(UWorld* InWorld, UTickableTran
 	// internal dependencies?
 	if (bSelf && bIncludeTarget)
 	{
+		static const TCHAR* SelfDependencyDesc = TEXT("Self Dependency");
+		
 		const UObject* SelfTarget = ChildTarget;
 		
-		using ConstraintPtr = TWeakObjectPtr<UTickableConstraint>;
-		auto Predicate = [InConstraint, SelfTarget](const ConstraintPtr& Constraint)
+		auto SelfTargetPredicate = [InConstraint, SelfTarget](const ConstraintWeakPtr& Constraint)
 		{
 			const UTickableTransformConstraint* TransformConstraint = Cast<UTickableTransformConstraint>(Constraint.Get());
 			if (!TransformConstraint || TransformConstraint == InConstraint)
@@ -1971,8 +1999,8 @@ bool FTransformConstraintUtils::BuildDependencies(UWorld* InWorld, UTickableTran
 			return ParentTarget == SelfTarget && ChildTarget == SelfTarget;
 		};
 
-		const TArray< TWeakObjectPtr<UTickableConstraint> > SelfConstraints = Controller.GetConstraintsByPredicate(Predicate);
-		for (const ConstraintPtr& SelfConstraint: SelfConstraints)
+		const TArray< ConstraintWeakPtr > SelfConstraints = Controller.GetConstraintsByPredicate(SelfTargetPredicate);
+		for (const ConstraintWeakPtr& SelfConstraint: SelfConstraints)
 		{
 			const UTickableTransformConstraint* TransformConstraint = Cast<UTickableTransformConstraint>(SelfConstraint);
 
@@ -1980,48 +2008,24 @@ bool FTransformConstraintUtils::BuildDependencies(UWorld* InWorld, UTickableTran
 			if (ParentHandle->HasDirectDependencyWith(*TransformConstraint->ChildTRSHandle))
 			{
 				Controller.SetConstraintsDependencies(TransformConstraint->ConstraintID, InConstraint->ConstraintID);
-
-				if (bDebugDependencies)
-				{
-					UE_LOG(LogTemp, Warning, TEXT("DEPENDENCY 1: %s is parent of %s so %s must tick before %s"),
-					   *GetHandleLabel(TransformConstraint->ChildTRSHandle), *GetHandleLabel(ParentHandle),
-					   *GetConstraintLabel(TransformConstraint), *GetConstraintLabel(InConstraint));
-				}
+				LogDependency(SelfDependencyDesc, TransformConstraint->ChildTRSHandle, ParentHandle, TransformConstraint, InConstraint);
 			}
 			else if (ChildHandle->HasDirectDependencyWith(*TransformConstraint->ChildTRSHandle))
 			{
 				Controller.SetConstraintsDependencies(TransformConstraint->ConstraintID, InConstraint->ConstraintID);
-
-				if (bDebugDependencies)
-				{
-					UE_LOG(LogTemp, Warning, TEXT("DEPENDENCY 1: %s is parent of %s so %s must tick before %s"),
-					   *GetHandleLabel(TransformConstraint->ChildTRSHandle), *GetHandleLabel(ChildHandle),
-					   *GetConstraintLabel(TransformConstraint), *GetConstraintLabel(InConstraint));
-				}
+				LogDependency(SelfDependencyDesc, TransformConstraint->ChildTRSHandle, ChildHandle, TransformConstraint, InConstraint);
 			}
 
 			// if the TransformConstraint handles depend on the new constraint child then, TransformConstraint should tick after
 			if (TransformConstraint->ParentTRSHandle->HasDirectDependencyWith(*ChildHandle))
 			{
 			 	Controller.SetConstraintsDependencies(InConstraint->ConstraintID, TransformConstraint->ConstraintID);
-
-				if (bDebugDependencies)
-				{
-					UE_LOG(LogTemp, Warning, TEXT("DEPENDENCY 1: %s is parent of %s so %s must tick before %s"),
-					   *GetHandleLabel(ChildHandle), *GetHandleLabel(TransformConstraint->ParentTRSHandle),
-					   *GetConstraintLabel(InConstraint), *GetConstraintLabel(TransformConstraint));
-				}
+				LogDependency(SelfDependencyDesc, ChildHandle, TransformConstraint->ParentTRSHandle, InConstraint, TransformConstraint);
 			}
 			else if (TransformConstraint->ChildTRSHandle->HasDirectDependencyWith(*ChildHandle))
 			{
 				Controller.SetConstraintsDependencies(InConstraint->ConstraintID, TransformConstraint->ConstraintID);
-
-				if (bDebugDependencies)
-				{
-					UE_LOG(LogTemp, Warning, TEXT("DEPENDENCY 1: %s is parent of %s so %s must tick before %s"),
-					   *GetHandleLabel(ChildHandle), *GetHandleLabel(TransformConstraint->ChildTRSHandle),
-					   *GetConstraintLabel(InConstraint), *GetConstraintLabel(TransformConstraint) );
-				}
+				LogDependency(SelfDependencyDesc, ChildHandle, TransformConstraint->ChildTRSHandle, InConstraint, TransformConstraint);
 			}
 		}
 	}
@@ -2030,18 +2034,109 @@ bool FTransformConstraintUtils::BuildDependencies(UWorld* InWorld, UTickableTran
 	InConstraint->EnsurePrimaryDependency(InWorld);
 
 	// if child handle is the parent of some other constraints, ensure they will tick after that new one
-	TArray<TWeakObjectPtr<UTickableConstraint>> ChildChildConstraints;
+	static const TCHAR* ChildDependencyDesc = TEXT("Child Dependency");
+
+	TArray<ConstraintWeakPtr> ChildChildConstraints;
 	GetChildrenConstraints(InWorld, InConstraint, ChildChildConstraints, bIncludeTarget && !bSelf);
-	for (const TWeakObjectPtr<UTickableConstraint>& ChildConstraint: ChildChildConstraints)
+	for (const ConstraintWeakPtr& ChildConstraint: ChildChildConstraints)
 	{
 		Controller.SetConstraintsDependencies(InConstraint->ConstraintID, ChildConstraint->ConstraintID);
+		if (const UTickableTransformConstraint* TransformConstraint = Cast<UTickableTransformConstraint>(ChildConstraint))
+		{
+			LogDependency(ChildDependencyDesc, ChildHandle, TransformConstraint->ChildTRSHandle, InConstraint, TransformConstraint);
+		}
+	}
 
-		if (bDebugDependencies)
+	if (bIncludeTarget && !bSelf)
+	{
+		// get all constraints acting on the same target 
+		auto SameChildTargetPredicate = [ChildTarget](const ConstraintWeakPtr& Constraint)
 		{
-			const UTickableTransformConstraint* TransformConstraint = Cast<UTickableTransformConstraint>(ChildConstraint);
-			UE_LOG(LogTemp, Warning, TEXT("DEPENDENCY 2: %s is parent of %s so %s must tick before %s"),
-				*GetHandleLabel(ChildHandle), *GetHandleLabel(TransformConstraint->ChildTRSHandle),
-				*GetConstraintLabel(InConstraint), *GetConstraintLabel(TransformConstraint));
+			const UTickableTransformConstraint* TransformConstraint = Cast<UTickableTransformConstraint>(Constraint.Get());
+			if (!TransformConstraint)
+			{
+				return false;
+			}
+			const UObject* OtherChildTarget = GetHandleTarget(TransformConstraint->ChildTRSHandle);
+			return OtherChildTarget && OtherChildTarget == ChildTarget;
+		};
+		TArray< ConstraintWeakPtr > ChildTargetParentConstraints = Controller.GetConstraintsByPredicate(SameChildTargetPredicate);
+
+		// store constraint index in this array
+		const int32 ConstrainIndex = ChildTargetParentConstraints.IndexOfByKey(InConstraint);
+		if (ensure(ChildTargetParentConstraints.IsValidIndex(ConstrainIndex)))
+		{
+			ChildTargetParentConstraints.RemoveAt(ConstrainIndex);
+		}
+		
+		if (!ChildTargetParentConstraints.IsEmpty())
+		{
+			static const TCHAR* ExternalDependencyDesc = TEXT("External Dependency");
+			
+			TBitArray<> ManagedDependencies(false, ChildTargetParentConstraints.Num());
+			
+			const FConstraintTickFunction& TickFunction = InConstraint->GetTickFunction(InWorld);
+			const TArray<FTickPrerequisite> PrerexCopy = TickFunction.GetPrerequisites();
+
+			for (int32 Index = 0; Index < ChildTargetParentConstraints.Num(); ++Index)
+			{
+				const UTickableTransformConstraint* TransformConstraint = Cast<UTickableTransformConstraint>(ChildTargetParentConstraints[Index]);
+				
+				// if the new handles depend on that constraint child then, TransformConstraint should tick before
+				if (ParentHandle->HasDirectDependencyWith(*TransformConstraint->ChildTRSHandle))
+				{
+					Controller.SetConstraintsDependencies(TransformConstraint->ConstraintID, InConstraint->ConstraintID);
+					ManagedDependencies[Index] = true;
+					LogDependency(ExternalDependencyDesc, TransformConstraint->ChildTRSHandle, ParentHandle, TransformConstraint, InConstraint);
+				}
+				else if (ChildHandle->HasDirectDependencyWith(*TransformConstraint->ChildTRSHandle))
+				{
+					Controller.SetConstraintsDependencies(TransformConstraint->ConstraintID, InConstraint->ConstraintID);
+					ManagedDependencies[Index] = true;
+					LogDependency(ExternalDependencyDesc, TransformConstraint->ChildTRSHandle, ChildHandle, TransformConstraint, InConstraint);
+				}
+
+				// if the TransformConstraint handles depend on the new constraint child then, TransformConstraint should tick after
+				if (TransformConstraint->ParentTRSHandle->HasDirectDependencyWith(*ChildHandle))
+				{
+					Controller.SetConstraintsDependencies(InConstraint->ConstraintID, TransformConstraint->ConstraintID);
+					ManagedDependencies[Index] = true;
+					LogDependency(ExternalDependencyDesc, ChildHandle, TransformConstraint->ParentTRSHandle, InConstraint, TransformConstraint);
+				}
+				else if (TransformConstraint->ChildTRSHandle->HasDirectDependencyWith(*ChildHandle))
+				{
+					Controller.SetConstraintsDependencies(InConstraint->ConstraintID, TransformConstraint->ConstraintID);
+					ManagedDependencies[Index] = true;
+					LogDependency(ExternalDependencyDesc, ChildHandle, TransformConstraint->ChildTRSHandle, InConstraint, TransformConstraint);
+				}
+			}
+
+			// if we didn't add any new prerequisite then check whether the constraint should tick after the last constraint
+			// acting on the same target to respect order of creation.
+			const bool bPrerexChanged = TickFunction.GetPrerequisites() != PrerexCopy;
+			if (!bPrerexChanged && ManagedDependencies.IsValidIndex(ConstrainIndex-1))
+			{
+				for (int32 Index = ConstrainIndex-1; Index >= 0; Index--)
+				{
+					if (!ManagedDependencies[Index])
+					{
+						const UTickableTransformConstraint* LastConstraintSharingSameTarget = Cast<UTickableTransformConstraint>(ChildTargetParentConstraints[Index]);
+						TSet<const FTickFunction*> VisitedFunctions;
+						const FTickFunction& ParentTickFunctionToCheck = LastConstraintSharingSameTarget->GetTickFunction(InWorld);
+						if (!FConstraintCycleChecker::HasPrerequisiteDependencyWith(&TickFunction, &ParentTickFunctionToCheck, VisitedFunctions) &&
+							!FConstraintCycleChecker::HasPrerequisiteDependencyWith(&ParentTickFunctionToCheck, &TickFunction, VisitedFunctions))
+						{
+							Controller.SetConstraintsDependencies(LastConstraintSharingSameTarget->ConstraintID, InConstraint->ConstraintID);
+							if (bDebugDependencies)
+							{
+								UE_LOG(LogTemp, Warning, TEXT("Creation Order Dependency: '%s' must tick before '%s' to respect creation order."),
+								   *GetConstraintLabel(LastConstraintSharingSameTarget), *GetConstraintLabel(InConstraint));
+							}
+							break;
+						}
+					}
+				}
+			}
 		}
 	}
 
