Removed some debug code left in in error, independenised a few incorporated class methods

#rnx

diff --git a/Engine/Source/Runtime/Experimental/IoStore/HttpClient/Private/Peer.inl b/Engine/Source/Runtime/Experimental/IoStore/HttpClient/Private/Peer.inl
--- a/Engine/Source/Runtime/Experimental/IoStore/HttpClient/Private/Peer.inl
+++ b/Engine/Source/Runtime/Experimental/IoStore/HttpClient/Private/Peer.inl
@@ -115,72 +115,10 @@ protected:
 	SSL*		Ssl = nullptr;
 
 private:
-	FOutcome GetOutcome(int32 SslResult, const char* Message="tls error") const
-	{
-		char buf[256];
-		TArray<uint32> iii;
-		int32 line;
-		const char* file;
-		while (uint32 i = ERR_get_error_line(&file, &line))
-		{
-			ERR_error_string(i, buf);
-			iii.Add(i);
-		}
-
-		int32 Error = SSL_get_error(Ssl, SslResult);
-		if (Error != SSL_ERROR_WANT_READ && Error != SSL_ERROR_WANT_WRITE)
-		{
-			return FOutcome::Error(Message, Error);
-		}
-		return FOutcome::Waiting();
-	}
-
-	int32 BioWrite(const char* Data, size_t Size, size_t* BytesWritten, BIO* Bio)
-	{
-		*BytesWritten = 0;
-		BIO_clear_retry_flags(Bio);
-
-		FOutcome Outcome = FPeer::Send(Data, Size);
-		if (Outcome.IsWaiting())
-		{
-			BIO_set_retry_write(Bio);
-			return 0;
-		}
-
-		if (Outcome.IsError())
-		{
-			return -1;
-		}
-
-		*BytesWritten = Outcome.GetResult();
-		return 1;
-	}
-
-	int32 BioRead(char* Data, size_t Size, size_t* BytesRead, BIO* Bio)
-	{
-		*BytesRead = 0;
-		BIO_clear_retry_flags(Bio);
-
-		FOutcome Outcome = FPeer::Recv(Data, Size);
-		if (Outcome.IsWaiting())
-		{
-			BIO_set_retry_read(Bio);
-			return 0;
-		}
-
-		if (Outcome.IsError())
-		{
-			return -1;
-		}
-
-		*BytesRead = Outcome.GetResult();
-		return 1;
-	}
-
-	long BioControl(int Cmd, long, void*, BIO*)
-	{
-		return (Cmd == BIO_CTRL_FLUSH) ? 1 : 0;
-	}
+	FOutcome	GetOutcome(int32 SslResult, const char* Message="tls error") const;
+	int32		BioWrite(const char* Data, size_t Size, size_t* BytesWritten, BIO* Bio);
+	int32		BioRead(char* Data, size_t Size, size_t* BytesRead, BIO* Bio);
+	long		BioControl(int Cmd, long, void*, BIO*);
 };
 
 ////////////////////////////////////////////////////////////////////////////////
@@ -296,6 +234,67 @@ FOutcome FTlsPeer::Recv(char* Out, int32 MaxSize)
 	return (Result > 0) ? FOutcome::Ok(Result) : GetOutcome(Result);
 }
 
+////////////////////////////////////////////////////////////////////////////////
+FOutcome FTlsPeer::GetOutcome(int32 SslResult, const char* Message) const
+{
+	int32 Error = SSL_get_error(Ssl, SslResult);
+	if (Error != SSL_ERROR_WANT_READ && Error != SSL_ERROR_WANT_WRITE)
+	{
+		return FOutcome::Error(Message, Error);
+	}
+	return FOutcome::Waiting();
+}
+
+////////////////////////////////////////////////////////////////////////////////
+int32 FTlsPeer::BioWrite(const char* Data, size_t Size, size_t* BytesWritten, BIO* Bio)
+{
+	*BytesWritten = 0;
+	BIO_clear_retry_flags(Bio);
+
+	FOutcome Outcome = FPeer::Send(Data, Size);
+	if (Outcome.IsWaiting())
+	{
+		BIO_set_retry_write(Bio);
+		return 0;
+	}
+
+	if (Outcome.IsError())
+	{
+		return -1;
+	}
+
+	*BytesWritten = Outcome.GetResult();
+	return 1;
+}
+
+////////////////////////////////////////////////////////////////////////////////
+int32 FTlsPeer::BioRead(char* Data, size_t Size, size_t* BytesRead, BIO* Bio)
+{
+	*BytesRead = 0;
+	BIO_clear_retry_flags(Bio);
+
+	FOutcome Outcome = FPeer::Recv(Data, Size);
+	if (Outcome.IsWaiting())
+	{
+		BIO_set_retry_read(Bio);
+		return 0;
+	}
+
+	if (Outcome.IsError())
+	{
+		return -1;
+	}
+
+	*BytesRead = Outcome.GetResult();
+	return 1;
+}
+
+////////////////////////////////////////////////////////////////////////////////
+long FTlsPeer::BioControl(int Cmd, long, void*, BIO*)
+{
+	return (Cmd == BIO_CTRL_FLUSH) ? 1 : 0;
+}
+
 
 
 ////////////////////////////////////////////////////////////////////////////////
