Geometry collection fix for bug with duplicated ISMs on Editor  representation.
Issue was that duplicating an actor in Editor would duplicate the ISMs in the object hierachy. Because the ISM Pool was marked as transient the ISMs would be left leaking when the duplicated actor's custom renderer reregistered and created new ISMs.

diff --git a/Engine/Source/Runtime/Experimental/GeometryCollectionEngine/Private/GeometryCollection/GeometryCollectionComponent.cpp b/Engine/Source/Runtime/Experimental/GeometryCollectionEngine/Private/GeometryCollection/GeometryCollectionComponent.cpp
--- a/Engine/Source/Runtime/Experimental/GeometryCollectionEngine/Private/GeometryCollection/GeometryCollectionComponent.cpp
+++ b/Engine/Source/Runtime/Experimental/GeometryCollectionEngine/Private/GeometryCollection/GeometryCollectionComponent.cpp
@@ -3563,7 +3563,7 @@ void UGeometryCollectionComponent::OnRegister()
 	{
 		if (!CustomRenderer || CustomRenderer.GetObject()->GetClass() != Type)
 		{
-			CustomRenderer = NewObject<UObject>(this, Type);
+			CustomRenderer = NewObject<UObject>(this, Type, NAME_None, RF_Transient | RF_DuplicateTransient);
 		}
 		RegisterCustomRenderer();
 	}
diff --git a/Engine/Source/Runtime/Experimental/GeometryCollectionEngine/Private/GeometryCollection/GeometryCollectionISMPoolComponent.cpp b/Engine/Source/Runtime/Experimental/GeometryCollectionEngine/Private/GeometryCollection/GeometryCollectionISMPoolComponent.cpp
--- a/Engine/Source/Runtime/Experimental/GeometryCollectionEngine/Private/GeometryCollection/GeometryCollectionISMPoolComponent.cpp
+++ b/Engine/Source/Runtime/Experimental/GeometryCollectionEngine/Private/GeometryCollection/GeometryCollectionISMPoolComponent.cpp
@@ -6,12 +6,14 @@
 #include "Components/InstancedStaticMeshComponent.h"
 #include "Engine/CollisionProfile.h"
 #include "Engine/StaticMesh.h"
+#include "Engine/World.h"
 
 #include UE_INLINE_GENERATED_CPP_BY_NAME(GeometryCollectionISMPoolComponent)
 
 // Don't release ISM components when they empty, but keep them (and their scene proxy) alive.
 // This can remove the high cost associated with repeated registration, scene proxy creation and mesh draw command creation.
-static bool GComponentKeepAlive = false; // Disabled due to current OOM issues on certain platforms
+// But it can also have a high memory overhead since the ISMs retain hard references to their static meshes.
+static bool GComponentKeepAlive = false; 
 FAutoConsoleVariableRef CVarISMPoolComponentKeepAlive(
 	TEXT("r.ISMPool.ComponentKeepAlive"),
 	GComponentKeepAlive,
@@ -37,6 +39,7 @@ FAutoConsoleVariableRef CVarISMPoolComponentFreeListTargetSize(
 	GComponentFreeListTargetSize,
 	TEXT("Target size for number of ISM components in the recycling free list."));
 
+// Keep copies of all custom instance data for restoration on readding an instance.
 static bool GShadowCopyCustomData = false;
 FAutoConsoleVariableRef CVarShadowCopyCustomData(
 	TEXT("r.ISMPool.ShadowCopyCustomData"),
@@ -101,19 +104,29 @@ void FGeometryCollectionMeshGroup::RemoveAllMeshes(FGeometryCollectionISMPool& I
 	MeshInfos.Empty();
 }
 
-void FGeometryCollectionISM::CreateISM(AActor* InOwningActor)
+void FGeometryCollectionISM::CreateISM(USceneComponent* InOwningComponent)
 {
-	check(InOwningActor);
+	check(InOwningComponent);
 
-	ISMComponent = NewObject<UInstancedStaticMeshComponent>(InOwningActor, NAME_None, RF_Transient | RF_DuplicateTransient);
+	AActor* OwningActor = InOwningComponent->GetOwner();
+	USceneComponent* RootComponent = OwningActor->GetRootComponent();
+
+	ISMComponent = NewObject<UInstancedStaticMeshComponent>(InOwningComponent, NAME_None, RF_Transient | RF_DuplicateTransient);
 
 	ISMComponent->SetRemoveSwap();
 	ISMComponent->SetCanEverAffectNavigation(false);
 	ISMComponent->SetCollisionEnabled(ECollisionEnabled::NoCollision);
-	ISMComponent->SetupAttachment(InOwningActor->GetRootComponent());
-	
-	InOwningActor->AddInstanceComponent(ISMComponent);
+	ISMComponent->SetupAttachment(RootComponent);
 	ISMComponent->RegisterComponent();
+
+#if WITH_EDITOR
+	UWorld const* World = InOwningComponent->GetWorld();
+	const bool bShowInWorldOutliner = World && World->IsGameWorld();
+	if (bShowInWorldOutliner)
+	{
+		OwningActor->AddInstanceComponent(ISMComponent);
+	}
+#endif
 }
 
 void FGeometryCollectionISM::InitISM(const FGeometryCollectionStaticMeshInstance& InMeshInstance, bool bKeepAlive, bool bOverrideTransformUpdates)
@@ -266,12 +279,12 @@ FGeometryCollectionISMPool::FISMIndex FGeometryCollectionISMPool::GetOrAddISM(UG
 	{
 		ISMIndex = FreeList.Last();
 		FreeList.RemoveAt(FreeList.Num() - 1);
-		ISMs[ISMIndex].CreateISM(OwningComponent->GetOwner());
+		ISMs[ISMIndex].CreateISM(OwningComponent);
 	}
 	else
 	{
 		ISMIndex = ISMs.AddDefaulted();
-		ISMs[ISMIndex].CreateISM(OwningComponent->GetOwner());
+		ISMs[ISMIndex].CreateISM(OwningComponent);
 	}
 	
 	ISMs[ISMIndex].InitISM(MeshInstance, bCachedKeepAlive, bDisableBoundsAndTransformUpdate);
@@ -437,9 +450,7 @@ void FGeometryCollectionISMPool::RemoveISM(FISMIndex ISMIndex, bool bKeepAlive,
 	else
 	{
 		// Completely unregister and destroy the component and mark the ISM slot as free.
-		ISM.ISMComponent->UnregisterComponent();
 		ISM.ISMComponent->DestroyComponent();
-		ISM.ISMComponent->GetOwner()->RemoveInstanceComponent(ISM.ISMComponent);
 		ISM.ISMComponent = nullptr;
 		
 		FreeList.Add(ISMIndex);
@@ -458,9 +469,7 @@ void FGeometryCollectionISMPool::Clear()
 		{
 			for(FGeometryCollectionISM& ISM : ISMs)
 			{
-				ISM.ISMComponent->UnregisterComponent();
 				ISM.ISMComponent->DestroyComponent();
-				OwningActor->RemoveInstanceComponent(ISM.ISMComponent);
 			}
 		}
 		ISMs.Reset();
diff --git a/Engine/Source/Runtime/Experimental/GeometryCollectionEngine/Private/GeometryCollection/GeometryCollectionISMPoolRenderer.cpp b/Engine/Source/Runtime/Experimental/GeometryCollectionEngine/Private/GeometryCollection/GeometryCollectionISMPoolRenderer.cpp
--- a/Engine/Source/Runtime/Experimental/GeometryCollectionEngine/Private/GeometryCollection/GeometryCollectionISMPoolRenderer.cpp
+++ b/Engine/Source/Runtime/Experimental/GeometryCollectionEngine/Private/GeometryCollection/GeometryCollectionISMPoolRenderer.cpp
@@ -1,6 +1,7 @@
 // Copyright Epic Games, Inc. All Rights Reserved.
 #include "GeometryCollection/GeometryCollectionISMPoolRenderer.h"
 
+#include "Engine/Level.h"
 #include "Engine/StaticMesh.h"
 #include "Engine/World.h"
 #include "GeometryCollection/Facades/CollectionInstancedMeshFacade.h"
@@ -15,17 +16,27 @@
 
 void UGeometryCollectionISMPoolRenderer::OnRegisterGeometryCollection(UGeometryCollectionComponent& InComponent)
 {
-	OwningLevel = InComponent.GetComponentLevel();
+	OwningLevel = MakeWeakObjectPtr(InComponent.GetComponentLevel());
 
 	// In editor we create our own ISMPool. 
 	// This guarantees the same look in editor/game, and allows editor hit proxies to continue working.
-	if (!InComponent.GetWorld()->IsGameWorld())
+	if (UWorld* World = InComponent.GetWorld())
 	{
-		CachedISMPoolComponent = NewObject<UGeometryCollectionISMPoolComponent>(InComponent.GetOwner());
-		CachedISMPoolComponent->SetTickablePoolManagement(false);
-		CachedISMPoolComponent->SetupAttachment(&InComponent);
-		CachedISMPoolComponent->RegisterComponent();
+		if (!World->IsGameWorld())
+		{
+			if (LocalISMPoolComponent)
+			{
+				LocalISMPoolComponent->DestroyComponent();
+			}
+
+			LocalISMPoolComponent = NewObject<UGeometryCollectionISMPoolComponent>(this, NAME_None, RF_Transient | RF_DuplicateTransient);
+			LocalISMPoolComponent->SetTickablePoolManagement(false);
+			LocalISMPoolComponent->SetupAttachment(&InComponent);
+			LocalISMPoolComponent->RegisterComponent();
+		}
 	}
+
+	bIsRegistered = true;
 }
 
 void UGeometryCollectionISMPoolRenderer::OnUnregisterGeometryCollection()
@@ -33,7 +44,15 @@ void UGeometryCollectionISMPoolRenderer::OnUnregisterGeometryCollection()
 	ReleaseGroup(MergedMeshGroup);
 	ReleaseGroup(InstancesGroup);
 
+	if (LocalISMPoolComponent)
+	{
+		LocalISMPoolComponent->DestroyComponent();
+		LocalISMPoolComponent = nullptr;
+	}
+
 	CachedISMPoolComponent = nullptr;
+
+	bIsRegistered = false;
 }
 
 void UGeometryCollectionISMPoolRenderer::UpdateState(UGeometryCollection const& InGeometryCollection, FTransform const& InComponentTransform, uint32 InStateFlags)
@@ -75,9 +94,9 @@ void UGeometryCollectionISMPoolRenderer::UpdateRootTransform(UGeometryCollection
 	UpdateMergedMeshTransforms(InRootTransform * ComponentTransform, {});
 }
 
-void UGeometryCollectionISMPoolRenderer::UpdateRootTransforms(UGeometryCollection const& InGeometryCollection, FTransform const& InRootTransform, TArrayView<const FTransform3f> InRootTransforms)
+void UGeometryCollectionISMPoolRenderer::UpdateRootTransforms(UGeometryCollection const& InGeometryCollection, FTransform const& InRootTransform, TArrayView<const FTransform3f> InRootLocalTransforms)
 {
-	UpdateMergedMeshTransforms(InRootTransform * ComponentTransform, InRootTransforms);
+	UpdateMergedMeshTransforms(InRootTransform * ComponentTransform, InRootLocalTransforms);
 }
 
 void UGeometryCollectionISMPoolRenderer::UpdateTransforms(UGeometryCollection const& InGeometryCollection, TArrayView<const FTransform3f> InTransforms)
@@ -87,19 +106,33 @@ void UGeometryCollectionISMPoolRenderer::UpdateTransforms(UGeometryCollection co
 
 UGeometryCollectionISMPoolComponent* UGeometryCollectionISMPoolRenderer::GetISMPoolComponent() const
 {
-	return CachedISMPoolComponent;
+	if (!bIsRegistered)
+	{
+		return nullptr;
+	}
+	return LocalISMPoolComponent ? LocalISMPoolComponent : CachedISMPoolComponent;
 }
 
 UGeometryCollectionISMPoolComponent* UGeometryCollectionISMPoolRenderer::GetOrCreateISMPoolComponent()
 {
-	if (CachedISMPoolComponent == nullptr)
+	if (!bIsRegistered)
+	{
+		return nullptr;
+	}
+	if (LocalISMPoolComponent)
+	{
+		return LocalISMPoolComponent;
+	}
+	if (!CachedISMPoolComponent)
 	{
 		if (UGeometryCollectionISMPoolSubSystem* ISMPoolSubSystem = UWorld::GetSubsystem<UGeometryCollectionISMPoolSubSystem>(GetWorld()))
 		{
-			check(OwningLevel);
-			if (AGeometryCollectionISMPoolActor* ISMPoolActor = ISMPoolSubSystem->FindISMPoolActor(OwningLevel))
+			if (ULevel* Level = OwningLevel.Get())
 			{
-				CachedISMPoolComponent = ISMPoolActor->GetISMPoolComp();
+				if (AGeometryCollectionISMPoolActor* ISMPoolActor = ISMPoolSubSystem->FindISMPoolActor(Level))
+				{
+					CachedISMPoolComponent = ISMPoolActor->GetISMPoolComp();
+				}
 			}
 		}
 	}
@@ -183,7 +216,7 @@ void UGeometryCollectionISMPoolRenderer::InitInstancesFromGeometryCollection(UGe
 	}
 }
 
-void UGeometryCollectionISMPoolRenderer::UpdateMergedMeshTransforms(FTransform const& InBaseTransform, TArrayView<const FTransform3f> LocalTransforms)
+void UGeometryCollectionISMPoolRenderer::UpdateMergedMeshTransforms(FTransform const& InBaseTransform, TArrayView<const FTransform3f> InLocalTransforms)
 {
 	if (MergedMeshGroup.GroupIndex == INDEX_NONE)
 	{
@@ -199,9 +232,9 @@ void UGeometryCollectionISMPoolRenderer::UpdateMergedMeshTransforms(FTransform c
 	TArrayView<const FTransform> InstanceTransforms(&InBaseTransform, 1);
 	for (int32 MeshIndex = 0; MeshIndex < MergedMeshGroup.MeshIds.Num(); MeshIndex++)
 	{
-		if (LocalTransforms.IsValidIndex(MeshIndex))
+		if (InLocalTransforms.IsValidIndex(MeshIndex))
 		{
-			const FTransform CombinedTransform{ FTransform(LocalTransforms[MeshIndex]) * InBaseTransform };
+			const FTransform CombinedTransform{ FTransform(InLocalTransforms[MeshIndex]) * InBaseTransform };
 			ISMPoolComponent->BatchUpdateInstancesTransforms(MergedMeshGroup.GroupIndex, MergedMeshGroup.MeshIds[MeshIndex], 0, MakeArrayView(&CombinedTransform, 1), true/*bWorldSpace*/, false/*bMarkRenderStateDirty*/, false/*bTeleport*/);
 		}
 		else
diff --git a/Engine/Source/Runtime/Experimental/GeometryCollectionEngine/Private/GeometryCollection/GeometryCollectionISMPoolRenderer.h b/Engine/Source/Runtime/Experimental/GeometryCollectionEngine/Private/GeometryCollection/GeometryCollectionISMPoolRenderer.h
--- a/Engine/Source/Runtime/Experimental/GeometryCollectionEngine/Private/GeometryCollection/GeometryCollectionISMPoolRenderer.h
+++ b/Engine/Source/Runtime/Experimental/GeometryCollectionEngine/Private/GeometryCollection/GeometryCollectionISMPoolRenderer.h
@@ -6,7 +6,6 @@
 
 #include "GeometryCollectionISMPoolRenderer.generated.h"
 
-class AGeometryCollectionISMPoolActor;
 class UGeometryCollectionISMPoolComponent;
 class UGeometryCollectionComponent;
 class ULevel;
@@ -17,6 +16,14 @@ class UGeometryCollectionISMPoolRenderer : public UObject, public IGeometryColle
 {
 	GENERATED_BODY()
 
+	/** Instanced Static Mesh Pool component that is used to render our meshes. */
+	UPROPERTY(Transient)
+	TObjectPtr<UGeometryCollectionISMPoolComponent> CachedISMPoolComponent;
+
+	/** Set if we have an Instanced Static Mesh Pool component owned by this renderer (ie when in Editor mode). Non-transient to behave correctly under actor duplication. */
+	UPROPERTY()
+	TObjectPtr<UGeometryCollectionISMPoolComponent> LocalISMPoolComponent;
+
 public:
 	//~ Begin IGeometryCollectionExternalRenderInterface Interface.
 	virtual void OnRegisterGeometryCollection(UGeometryCollectionComponent& InComponent) override;
@@ -27,6 +34,7 @@ public:
 	virtual void UpdateTransforms(UGeometryCollection const& InGeometryCollection, TArrayView<const FTransform3f> InTransforms) override;
 	//~ End IGeometryCollectionExternalRenderInterface Interface.
 
+protected:
 	/** Description for a group of meshes that are added/updated together. */
 	struct FISMPoolGroup
 	{
@@ -34,10 +42,8 @@ public:
 		TArray<int32> MeshIds;
 	};
 
-protected:
-	/** Instanced Static Mesh Pool component that is used to render our meshes. */
-	UPROPERTY(Transient)
-	TObjectPtr<UGeometryCollectionISMPoolComponent> CachedISMPoolComponent;
+	/** Registered flag set true in between calls to OnRegister() and OnUnregister(). */
+	bool bIsRegistered = false;
 
 	/** Cached component transform. */
 	FTransform ComponentTransform = FTransform::Identity;
@@ -46,15 +52,15 @@ protected:
 	FISMPoolGroup MergedMeshGroup;
 	FISMPoolGroup InstancesGroup;
 
-	/** level of the owning component of this renderer */
-	ULevel* OwningLevel = nullptr;
+	/** Level of the owning component of this renderer */
+	TWeakObjectPtr<ULevel> OwningLevel;
 
 private:
 	UGeometryCollectionISMPoolComponent* GetISMPoolComponent() const;
 	UGeometryCollectionISMPoolComponent* GetOrCreateISMPoolComponent();
 	void InitMergedMeshFromGeometryCollection(UGeometryCollection const& InGeometryCollection);
 	void InitInstancesFromGeometryCollection(UGeometryCollection const& InGeometryCollection);
-	void UpdateMergedMeshTransforms(FTransform const& InBaseTransform, TArrayView<const FTransform3f> LocalTransforms);
+	void UpdateMergedMeshTransforms(FTransform const& InBaseTransform, TArrayView<const FTransform3f> InLocalTransforms);
 	void UpdateInstanceTransforms(UGeometryCollection const& InGeometryCollection, FTransform const& InBaseTransform, TArrayView<const FTransform3f> InTransforms);
 	void ReleaseGroup(FISMPoolGroup& InOutGroup);
 };
diff --git a/Engine/Source/Runtime/Experimental/GeometryCollectionEngine/Private/GeometryCollection/GeometryCollectionRootProxyRenderer.cpp b/Engine/Source/Runtime/Experimental/GeometryCollectionEngine/Private/GeometryCollection/GeometryCollectionRootProxyRenderer.cpp
--- a/Engine/Source/Runtime/Experimental/GeometryCollectionEngine/Private/GeometryCollection/GeometryCollectionRootProxyRenderer.cpp
+++ b/Engine/Source/Runtime/Experimental/GeometryCollectionEngine/Private/GeometryCollection/GeometryCollectionRootProxyRenderer.cpp
@@ -2,8 +2,8 @@
 #include "GeometryCollection/GeometryCollectionRootProxyRenderer.h"
 
 #include "Components/StaticMeshComponent.h"
+#include "Engine/CollisionProfile.h"
 #include "Engine/StaticMesh.h"
-#include "Engine/World.h"
 #include "GeometryCollection/GeometryCollection.h"
 #include "GeometryCollection/GeometryCollectionComponent.h"
 #include "GeometryCollection/GeometryCollectionObject.h"
@@ -71,7 +71,7 @@ void UGeometryCollectionRootProxyRenderer::CreateRootProxyComponents(UGeometryCo
 		
 		if (ProxyMesh != nullptr)
 		{
-			MeshComponent = NewObject<UStaticMeshComponent>(InComponent.GetOwner(), NAME_None, RF_DuplicateTransient | RF_Transient);
+			MeshComponent = NewObject<UStaticMeshComponent>(this, NAME_None, RF_DuplicateTransient | RF_Transient);
 
 			MeshComponent->SetStaticMesh(ProxyMesh);
 			MeshComponent->SetCanEverAffectNavigation(false);
diff --git a/Engine/Source/Runtime/Experimental/GeometryCollectionEngine/Private/GeometryCollection/GeometryCollectionRootProxyRenderer.h b/Engine/Source/Runtime/Experimental/GeometryCollectionEngine/Private/GeometryCollection/GeometryCollectionRootProxyRenderer.h
--- a/Engine/Source/Runtime/Experimental/GeometryCollectionEngine/Private/GeometryCollection/GeometryCollectionRootProxyRenderer.h
+++ b/Engine/Source/Runtime/Experimental/GeometryCollectionEngine/Private/GeometryCollection/GeometryCollectionRootProxyRenderer.h
@@ -14,6 +14,10 @@ class UGeometryCollectionRootProxyRenderer : public UObject, public IGeometryCol
 {
 	GENERATED_BODY()
 
+	/** Static mesh components for the root proxy meshes. */
+	UPROPERTY()
+	TArray<TObjectPtr<UStaticMeshComponent>> StaticMeshComponents;
+
 public:
 	//~ Begin IGeometryCollectionExternalRenderInterface Interface.
 	virtual void OnRegisterGeometryCollection(UGeometryCollectionComponent& InComponent) override;
@@ -29,9 +33,6 @@ public:
 protected:
 	/** Current visibility state. */
 	bool bIsVisible = true;
-	/** Static mesh components for the root proxy meshes. */
-	UPROPERTY(Transient)
-	TArray<TObjectPtr<UStaticMeshComponent>> StaticMeshComponents;
 
 private:
 	void CreateRootProxyComponents(UGeometryCollectionComponent& InComponent);
diff --git a/Engine/Source/Runtime/Experimental/GeometryCollectionEngine/Public/GeometryCollection/GeometryCollectionComponent.h b/Engine/Source/Runtime/Experimental/GeometryCollectionEngine/Public/GeometryCollection/GeometryCollectionComponent.h
--- a/Engine/Source/Runtime/Experimental/GeometryCollectionEngine/Public/GeometryCollection/GeometryCollectionComponent.h
+++ b/Engine/Source/Runtime/Experimental/GeometryCollectionEngine/Public/GeometryCollection/GeometryCollectionComponent.h
@@ -1483,7 +1483,7 @@ protected:
 	uint8 bForceNativeRenderer : 1;
 
 	/** A custom renderer object created from CustomRenderType. */
-	UPROPERTY(Transient)
+	UPROPERTY(SkipSerialization)
 	TScriptInterface<IGeometryCollectionExternalRenderInterface> CustomRenderer;
 
 	/** Collect all the PSO precache data used by the geometry collection */
diff --git a/Engine/Source/Runtime/Experimental/GeometryCollectionEngine/Public/GeometryCollection/GeometryCollectionISMPoolComponent.h b/Engine/Source/Runtime/Experimental/GeometryCollectionEngine/Public/GeometryCollection/GeometryCollectionISMPoolComponent.h
--- a/Engine/Source/Runtime/Experimental/GeometryCollectionEngine/Public/GeometryCollection/GeometryCollectionISMPoolComponent.h
+++ b/Engine/Source/Runtime/Experimental/GeometryCollectionEngine/Public/GeometryCollection/GeometryCollectionISMPoolComponent.h
@@ -256,7 +256,7 @@ struct FGeometryCollectionMeshGroup
 struct FGeometryCollectionISM
 {
 	/** Create the ISMComponent according to settings on the mesh instance. */
-	void CreateISM(AActor* InOwningActor);
+	void CreateISM(USceneComponent* InOwningComponent);
 	/** Initialize the ISMComponent according to settings on the mesh instance. */
 	void InitISM(const FGeometryCollectionStaticMeshInstance& InMeshInstance, bool bKeepAlive, bool bOverrideTransformUpdates = false);
 	/** Add a group to the ISM. Returns the group index. */
