Fix bad collision detection for non-optimized convexes

diff --git a/Engine/Source/Runtime/Experimental/Chaos/Private/Chaos/ConvexOptimizer.cpp b/Engine/Source/Runtime/Experimental/Chaos/Private/Chaos/ConvexOptimizer.cpp
--- a/Engine/Source/Runtime/Experimental/Chaos/Private/Chaos/ConvexOptimizer.cpp
+++ b/Engine/Source/Runtime/Experimental/Chaos/Private/Chaos/ConvexOptimizer.cpp
@@ -35,6 +35,10 @@ namespace CVars
 	bool ChaosConvexEnableMerging = true;
 	FAutoConsoleVariableRef CVarChaosConvexEnableMerging(TEXT("p.Chaos.Convex.EnableMerging"), ChaosConvexEnableMerging, TEXT("Boolean to check if we are merging (bottom-up) or splitting (top-bottom) the convexes"));
 	
+	// Boolean to check if we can use direct overlap test against the bounds (it assumes we only have leaves)
+	bool ChaosConvexUseDirectOverlap = true;
+	FAutoConsoleVariableRef CVarChaosConvexUseDirectOverlap(TEXT("p.Chaos.Convex.UseDirectOverlap"), ChaosConvexUseDirectOverlap, TEXT("Boolean to check if we can use direct overlap test against the bounds"));
+	
 	extern int32 ChaosUnionBVHMaxDepth;
 	extern int32 ChaosUnionBVHMinShapes;
 }
@@ -68,25 +72,57 @@ void FConvexOptimizer::VisitOverlappingObjects(const FAABB3& LocalBounds, const
 {
 	if(CVars::bChaosConvexSimplifyUnion)
 	{
-		int32 ObjectIndex = 1;
-		if(BVH.IsValid())
+		if(CVars::ChaosConvexUseDirectOverlap)
 		{
-			BVH->VisitAllIntersections(LocalBounds,
-			[this, &ObjectIndex, &VisitorFunc, &LocalBounds](const FImplicitObject* Implicit, const FRigidTransform3f& RelativeTransformf,
-				const FAABB3f& RelativeBoundsf, const int32 RootObjectIndex, const int32 LeafObjectIndex)
+			const FAABB3f CollisionBounds(LocalBounds);
+			int32 ObjectIndex = 1;
+			if(BVH.IsValid())
 			{
-				int32 LocalLeafObjectIndex = LeafObjectIndex;
-				Implicit->VisitOverlappingLeafObjectsImpl(LocalBounds, FRigidTransform3(RelativeTransformf),
-					RootObjectIndex, ObjectIndex, LocalLeafObjectIndex, VisitorFunc);
-			});
+				BVH->VisitAllIntersections(LocalBounds,
+				[this, &ObjectIndex, &VisitorFunc,&CollisionBounds](const FImplicitObject* Implicit, const FRigidTransform3f& RelativeTransformf,
+					const FAABB3f& RelativeBoundsf, const int32 RootObjectIndex, const int32 LeafObjectIndex)
+				{
+					if(CollisionBounds.Intersects(RelativeBoundsf))
+					{
+						VisitorFunc(Implicit, FRigidTransform3(RelativeTransformf), RootObjectIndex, ObjectIndex++, LeafObjectIndex);
+					}
+				});
+			}
+			else
+			{
+				for(const Private::FImplicitBVHObject& CollisionObject : CollisionObjects->ImplicitObjects)
+				{
+					if(CollisionBounds.Intersects(CollisionObject.GetBounds()))
+					{
+						VisitorFunc(CollisionObject.GetGeometry(), CollisionObject.GetTransform(), CollisionObject.GetRootObjectIndex(), ObjectIndex++, CollisionObject.GetObjectIndex());
+					}
+				}
+			}
 		}
 		else
 		{
-			for(const Private::FImplicitBVHObject& CollisionObject : CollisionObjects->ImplicitObjects)
+			//TODO: remove this version and the cvar since the LocalBounds is in object space and the
+			// overlap test done in VisitOverlappingLeafObjectsImpl assume a LocalBounds in implicit space
+			int32 ObjectIndex = 1;
+			if(BVH.IsValid())
+			{
+				BVH->VisitAllIntersections(LocalBounds,
+				[this, &ObjectIndex, &VisitorFunc, &LocalBounds](const FImplicitObject* Implicit, const FRigidTransform3f& RelativeTransformf,
+					const FAABB3f& RelativeBoundsf, const int32 RootObjectIndex, const int32 LeafObjectIndex)
+				{
+					int32 LocalLeafObjectIndex = LeafObjectIndex;
+					Implicit->VisitOverlappingLeafObjectsImpl(LocalBounds, FRigidTransform3(RelativeTransformf),
+						RootObjectIndex, ObjectIndex, LocalLeafObjectIndex, VisitorFunc);
+				});
+			}
+			else
 			{
-				int32 LocalLeafObjectIndex = CollisionObject.GetObjectIndex();
-				CollisionObject.GetGeometry()->VisitOverlappingLeafObjectsImpl(LocalBounds, CollisionObject.GetTransform(),
-					CollisionObject.GetRootObjectIndex(), ObjectIndex, LocalLeafObjectIndex, VisitorFunc);
+				for(const Private::FImplicitBVHObject& CollisionObject : CollisionObjects->ImplicitObjects)
+				{
+					int32 LocalLeafObjectIndex = CollisionObject.GetObjectIndex();
+					CollisionObject.GetGeometry()->VisitOverlappingLeafObjectsImpl(LocalBounds, CollisionObject.GetTransform(),
+						CollisionObject.GetRootObjectIndex(), ObjectIndex, LocalLeafObjectIndex, VisitorFunc);
+				}
 			}
 		}
 	}
@@ -395,7 +431,7 @@ DECLARE_CYCLE_STAT(TEXT("Collisions::MergeConnectedShapes"), STAT_MergeConnected
 void FConvexOptimizer::MergeConnectedShapes(const Chaos::FImplicitObjectUnionPtr& UnionGeometry, TArray<FTriboxNode>& MergedNodes)
 {
 	SCOPE_CYCLE_COUNTER(STAT_MergeConnectedShapes);
-	const int32 NumTriboxes = RootTriboxes.Num();
+	const int32 NumTriboxes = UnionGeometry->GetObjects().Num();
 		
 	MergedNodes.Reserve(RootTriboxes.Num());
 	FTriboxNode* CurrentNode = nullptr;
