Connections to the host may only be made with peer objects.

#rnx

diff --git a/Engine/Source/Runtime/Experimental/IoStore/HttpClient/Private/ConnectionPool.inl b/Engine/Source/Runtime/Experimental/IoStore/HttpClient/Private/ConnectionPool.inl
--- a/Engine/Source/Runtime/Experimental/IoStore/HttpClient/Private/ConnectionPool.inl
+++ b/Engine/Source/Runtime/Experimental/IoStore/HttpClient/Private/ConnectionPool.inl
@@ -25,8 +25,7 @@ public:
 					FHost(const FParams& Params);
 	void			SetBufferSize(EDirection Dir, int32 Size);
 	int32			GetBufferSize(EDirection Dir) const;
-	FHttpPeer		Connect();
-	FOutcome		Connect(FSocket& Socket);
+	FOutcome		Connect(FTlsPeer& Peer);
 	int32			IsResolved() const;
 	FOutcome		ResolveHostName();
 	bool			WithTls() const				{ return GetVerifyCert().GetData() != nullptr; }
@@ -37,6 +36,7 @@ public:
 	uint32			GetPort() const				{ return Port; }
 
 private:
+	FOutcome		Connect(FSocket& Socket);
 	FPemCert		VerifyCert;
 	const ANSICHAR*	HostName;
 	uint32			IpAddresses[4] = {};
@@ -144,29 +144,30 @@ int32 FHost::IsResolved() const
 }
 
 ////////////////////////////////////////////////////////////////////////////////
-FHttpPeer FHost::Connect()
+FOutcome FHost::Connect(FTlsPeer& Peer)
 {
 	FSocket Socket;
-	if (FOutcome Outcome = Connect(Socket); Outcome.IsError())
+	FOutcome Outcome = Connect(Socket);
+
+	if (Outcome.IsError())
 	{
-		return FHttpPeer();
-		// return FOutcome::Error("Failed to create socket");
+		return Outcome;
 	}
 
 	if (VerifyCert.GetData() == nullptr)
 	{
-		return FHttpPeer(MoveTemp(Socket));
+		Peer = FTlsPeer(MoveTemp(Socket));
+		return Outcome;
 	}
 
 	FSslContext SslContext(HostName, VerifyCert);
-	FHttpPeer Peer(MoveTemp(Socket), &SslContext);
+	Peer = FTlsPeer(MoveTemp(Socket), &SslContext);
 
-	FOutcome Outcome = FOutcome::None();
 	while ((Outcome = Peer.Handshake()).IsWaiting())
 		/*pass*/;
 	check(Outcome.IsOk());
 
-	return Peer;
+	return Outcome;
 }
 
 ////////////////////////////////////////////////////////////////////////////////
diff --git a/Engine/Source/Runtime/Experimental/IoStore/HttpClient/Private/Loop.inl b/Engine/Source/Runtime/Experimental/IoStore/HttpClient/Private/Loop.inl
--- a/Engine/Source/Runtime/Experimental/IoStore/HttpClient/Private/Loop.inl
+++ b/Engine/Source/Runtime/Experimental/IoStore/HttpClient/Private/Loop.inl
@@ -754,9 +754,7 @@ void FPeerGroup::TickSend(FTickState& State, FHost& Host)
 	if (!Peer.IsValid())
 	{
 		IsKeepAlive = 1;
-		Peer = Host.Connect();
-
-		FOutcome Outcome = FOutcome::Ok();
+		FOutcome Outcome = Host.Connect(Peer);
 
 		// We failed to connect, let's bail.
 		if (Outcome.IsError())
