Negotiation was tight blocking loop which is undesirable.

#rnx

diff --git a/Engine/Source/Runtime/Experimental/IoStore/HttpClient/Private/ConnectionPool.inl b/Engine/Source/Runtime/Experimental/IoStore/HttpClient/Private/ConnectionPool.inl
--- a/Engine/Source/Runtime/Experimental/IoStore/HttpClient/Private/ConnectionPool.inl
+++ b/Engine/Source/Runtime/Experimental/IoStore/HttpClient/Private/ConnectionPool.inl
@@ -162,11 +162,6 @@ FOutcome FHost::Connect(FTlsPeer& Peer)
 
 	FSslContext SslContext(HostName, VerifyCert);
 	Peer = FTlsPeer(MoveTemp(Socket), &SslContext);
-
-	while ((Outcome = Peer.Handshake()).IsWaiting())
-		/*pass*/;
-	check(Outcome.IsOk());
-
 	return Outcome;
 }
 
diff --git a/Engine/Source/Runtime/Experimental/IoStore/HttpClient/Private/Loop.inl b/Engine/Source/Runtime/Experimental/IoStore/HttpClient/Private/Loop.inl
--- a/Engine/Source/Runtime/Experimental/IoStore/HttpClient/Private/Loop.inl
+++ b/Engine/Source/Runtime/Experimental/IoStore/HttpClient/Private/Loop.inl
@@ -527,6 +527,7 @@ public:
 	void				Fail(FTickState& State, const char* Reason);
 
 private:
+	void				Negotiate(FTickState& State);
 	void				RecvInternal(FTickState& State);
 	void				SendInternal(FTickState& State);
 	FActivity*			Send = nullptr;
@@ -534,6 +535,7 @@ private:
 	FPeerType			Peer;
 	uint32				LastUseMs = 0;
 	uint8				IsKeepAlive = 0;
+	uint8				bNegotiating = 0;
 	bool				bWaiting = false;
 
 	UE_NONCOPYABLE(FPeerGroup);
@@ -591,11 +593,37 @@ void FPeerGroup::Fail(FTickState& State, const char* Reason)
 	Send = Recv = nullptr;
 	bWaiting = false;
 	IsKeepAlive = 0;
+	bNegotiating = false;
+}
+
+////////////////////////////////////////////////////////////////////////////////
+void FPeerGroup::Negotiate(FTickState& State)
+{
+	check(bNegotiating);
+	check(Send != nullptr);
+	check(Peer.IsValid());
+
+	FOutcome Outcome = Peer.Handshake();
+	if (Outcome.IsError())
+	{
+		Fail(State, Outcome.GetMessage().GetData());
+		return;
+	}
+
+	if (Outcome.IsWaiting())
+	{
+		bWaiting = true;
+		return;
+	}
+
+	bNegotiating = false;
+	return SendInternal(State);
 }
 
 ////////////////////////////////////////////////////////////////////////////////
 void FPeerGroup::RecvInternal(FTickState& State)
 {
+	check(bNegotiating == false);
 	check(Recv != nullptr);
 
 	// Another helper lambda
@@ -679,6 +707,7 @@ void FPeerGroup::RecvInternal(FTickState& State)
 ////////////////////////////////////////////////////////////////////////////////
 void FPeerGroup::SendInternal(FTickState& State)
 {
+	check(bNegotiating == false);
 	check(IsKeepAlive == 1);
 	check(Send != nullptr);
 
@@ -715,7 +744,12 @@ void FPeerGroup::SendInternal(FTickState& State)
 ////////////////////////////////////////////////////////////////////////////////
 bool FPeerGroup::Tick(FTickState& State)
 {
-	if (Send != nullptr)
+	if (bNegotiating)
+	{
+		Negotiate(State);
+	}
+
+	else if (Send != nullptr)
 	{
 		SendInternal(State);
 	}
@@ -754,6 +788,7 @@ void FPeerGroup::TickSend(FTickState& State, FHost& Host)
 	if (!Peer.IsValid())
 	{
 		IsKeepAlive = 1;
+		bNegotiating = true;
 		FOutcome Outcome = Host.Connect(Peer);
 
 		// We failed to connect, let's bail.
@@ -772,7 +807,7 @@ void FPeerGroup::TickSend(FTickState& State, FHost& Host)
 
 	if (!bWillBlock)
 	{
-		return SendInternal(State);
+		return Negotiate(State);
 	}
 
 	// Non-blocking connect
