[NetPhysics] Sync altered and predicted inputs from server to allow client to trigger resim if client mispredicted inputs.
- Client can implement FNetworkPhysicsData::CompareData() and perform return InputFrame == PredictedData.InputFrame; to trigger resim if input frame was mispredicted.
- Also initialized replicated data with the correct size of redundant data to send instead of the size of history.
- Further improvement to the rewind history frame/index naming and callflow

diff --git a/Engine/Source/Runtime/Engine/Private/PhysicsEngine/NetworkPhysicsComponent.cpp b/Engine/Source/Runtime/Engine/Private/PhysicsEngine/NetworkPhysicsComponent.cpp
--- a/Engine/Source/Runtime/Engine/Private/PhysicsEngine/NetworkPhysicsComponent.cpp
+++ b/Engine/Source/Runtime/Engine/Private/PhysicsEngine/NetworkPhysicsComponent.cpp
@@ -809,6 +809,7 @@ void UNetworkPhysicsComponent::ServerReceiveInputData_Implementation(const FNetw
 					PRAGMA_DISABLE_DEPRECATION_WARNINGS // TODO: Change to ReceiveNewData() in UE 5.6 and remove deprecation pragma
 					InputHistory->ReceiveNewDatas(*ReceivedInputs, 0);
 					PRAGMA_ENABLE_DEPRECATION_WARNINGS
+
 	#if DEBUG_NETWORK_PHYSICS
 					{
 						TArray<int32> LocalFrames, ServerFrames, InputFrames;
@@ -835,8 +836,10 @@ void UNetworkPhysicsComponent::ServerReceiveInputData_Implementation(const FNetw
 
 void UNetworkPhysicsComponent::OnPreProcessInputsInternal(const int32 PhysicsStep)
 {
+	const bool bIsServer = HasServerWorld();
+
 #if DEBUG_NETWORK_PHYSICS
-	if (HasServerWorld())
+	if (bIsServer)
 	{
 		UE_LOG(LogChaos, Log, TEXT("SERVER | PT | OnPreProcessInputsInternal | At Frame %d | Component = %s"), PhysicsStep, *GetFullName());
 	}
@@ -861,7 +864,7 @@ void UNetworkPhysicsComponent::OnPreProcessInputsInternal(const int32 PhysicsSte
 		}
 
 		// Apply replicated state on clients if we are resimulating
-		if (StateHistory && !HasServerWorld() && bIsSolverResim)
+		if (bIsSolverResim && StateHistory)
 		{
 			FNetworkPhysicsData* PhysicsData = StateData.Get();
 			PhysicsData->LocalFrame = PhysicsStep;
@@ -907,6 +910,14 @@ void UNetworkPhysicsComponent::OnPreProcessInputsInternal(const int32 PhysicsSte
 					InputHistory->MergeData(NextExpectedLocalFrame, PhysicsData);
 				}
 
+				// If the extracted input data was altered (interpolated, merged, extrapolated/predicted) on the server, record it into the history for it to get replicated to clients
+				if (bIsServer && !bIsSolverResim && PhysicsData->InputFrame == INDEX_NONE)
+				{
+					PhysicsData->bReceivedData = true; // Mark the input data as received so that it doesn't get overwritten by incoming client inputs
+					PhysicsData->LocalFrame = PhysicsStep;
+					InputHistory->RecordData(PhysicsStep, PhysicsData);
+				}
+
 				PhysicsData->ApplyData(ActorComponent);
 			}
 		}
@@ -915,10 +926,12 @@ void UNetworkPhysicsComponent::OnPreProcessInputsInternal(const int32 PhysicsSte
 
 void UNetworkPhysicsComponent::OnPostProcessInputsInternal(const int32 PhysicsStep)
 {
+	const bool bIsServer = HasServerWorld();
+
 	bool bIsSolverReset = false;
 	bool bIsSolverResim = false;
 #if DEBUG_NETWORK_PHYSICS
-	if (HasServerWorld())
+	if (bIsServer)
 	{
 		UE_LOG(LogChaos, Log, TEXT("SERVER | PT | OnPostProcessInputsInternal | At Frame %d | Component = %s"), PhysicsStep, *GetFullName());
 	}
@@ -928,7 +941,7 @@ void UNetworkPhysicsComponent::OnPostProcessInputsInternal(const int32 PhysicsSt
 	}
 #endif
 
-	if (InputHistory && ActorComponent)
+	if (ActorComponent)
 	{
 		if (FPhysScene* PhysScene = GetWorld()->GetPhysicsScene())
 		{
@@ -951,7 +964,7 @@ void UNetworkPhysicsComponent::OnPostProcessInputsInternal(const int32 PhysicsSt
 		{
 			FNetworkPhysicsData* PhysicsData = InputData.Get();
 			PhysicsData->LocalFrame = PhysicsStep;
-			PhysicsData->ServerFrame = HasServerWorld() ? PhysicsStep : PhysicsStep + PlayerController->GetNetworkPhysicsTickOffset();
+			PhysicsData->ServerFrame = bIsServer ? PhysicsStep : PhysicsStep + PlayerController->GetNetworkPhysicsTickOffset();
 			PhysicsData->InputFrame = PhysicsStep;
 			PhysicsData->bReceivedData = false;
 
@@ -966,7 +979,6 @@ void UNetworkPhysicsComponent::OnPostProcessInputsInternal(const int32 PhysicsSt
 #endif
 		}
 
-		const bool bIsServer = HasServerWorld();
 		const bool bShouldCacheStateHistory = bIsServer || (bCompareStateToTriggerRewind && bShouldCacheInputHistory);
 		if (StateHistory && bShouldCacheStateHistory)
 		{
@@ -975,7 +987,7 @@ void UNetworkPhysicsComponent::OnPostProcessInputsInternal(const int32 PhysicsSt
 			{
 				FNetworkPhysicsData* PhysicsData = InputData.Get();
 				PRAGMA_DISABLE_DEPRECATION_WARNINGS // TODO: Change to ExtractData() in UE 5.6 and remove deprecation pragma
-				if (InputHistory->ExtractDatas(PhysicsStep, false, PhysicsData, true))
+				if (InputHistory && InputHistory->ExtractDatas(PhysicsStep, false, PhysicsData, true))
 				PRAGMA_ENABLE_DEPRECATION_WARNINGS
 				{
 					InputFrame = PhysicsData->InputFrame;
diff --git a/Engine/Source/Runtime/Engine/Public/Physics/NetworkPhysicsComponent.h b/Engine/Source/Runtime/Engine/Public/Physics/NetworkPhysicsComponent.h
--- a/Engine/Source/Runtime/Engine/Public/Physics/NetworkPhysicsComponent.h
+++ b/Engine/Source/Runtime/Engine/Public/Physics/NetworkPhysicsComponent.h
@@ -56,19 +56,19 @@ struct TNetRewindHistory : public Chaos::TDataRewindHistory<DataType>
 	virtual TUniquePtr<Chaos::FBaseRewindHistory> CopyFramesWithOffset(const uint32 StartFrame, const uint32 EndFrame, const int32 FrameOffset) override
 	{
 		uint32 FramesCount = (uint32)Super::NumValidData(StartFrame, EndFrame);
-			
+
 		TUniquePtr<TNetRewindHistory> Copy = MakeUnique<TNetRewindHistory>(FramesCount, Super::bIsLocalHistory);
 
 		DataType FrameData;
-		for (uint32 FrameIndex = StartFrame; FrameIndex < EndFrame; ++FrameIndex)
+		for (uint32 CopyFrame = StartFrame; CopyFrame < EndFrame; ++CopyFrame)
 		{
-			const int32 LocalFrame = FrameIndex % Super::NumFrames;
-			if (FrameIndex == Super::DataHistory[LocalFrame].LocalFrame)
+			const int32 CopyIndex = Super::GetFrameIndex(CopyFrame);
+			if (CopyFrame == Super::DataHistory[CopyIndex].LocalFrame)
 			{
-				FrameData = Super::DataHistory[LocalFrame];
+				FrameData = Super::DataHistory[CopyIndex];
 				FrameData.ServerFrame = FrameData.LocalFrame + FrameOffset;
 				PRAGMA_DISABLE_DEPRECATION_WARNINGS // TODO: Change to RecordData() in UE 5.6 and remove deprecation pragma
-				Copy->RecordDatas(LocalFrame, &FrameData);
+				Copy->RecordDatas(CopyFrame, &FrameData);
 				PRAGMA_ENABLE_DEPRECATION_WARNINGS
 			}
 		}
@@ -88,7 +88,8 @@ struct TNetRewindHistory : public Chaos::TDataRewindHistory<DataType>
 				DataType& FrameData = NetNewData.DataHistory[FrameIndex];
 
 				FrameData.LocalFrame = FrameData.ServerFrame - FrameOffset;
-				if (FrameData.LocalFrame >= 0)
+
+				if (Super::ShouldRecordReceivedDataOnFrame(FrameData.LocalFrame))
 				{
 					FrameData.bReceivedData = true; // Received data is marked to differentiate from locally predicted data
 
@@ -668,10 +669,10 @@ FORCEINLINE void UNetworkPhysicsComponent::CreateDataHistory(UActorComponent* Hi
 	InputData = MakeUnique<typename PhysicsTraits::InputsType>();
 	StateData = MakeUnique<typename PhysicsTraits::StatesType>();
 
-	ReplicatedInputs.History = MakeUnique<TNetRewindHistory<typename PhysicsTraits::InputsType>>(NumFrames, bIsLocalHistory);
+	ReplicatedInputs.History = MakeUnique<TNetRewindHistory<typename PhysicsTraits::InputsType>>(InputRedundancy, bIsLocalHistory);
 	ReplicatedInputs.Owner = this;
 
-	ReplicatedStates.History = MakeUnique<TNetRewindHistory<typename PhysicsTraits::StatesType>>(NumFrames, bIsLocalHistory);
+	ReplicatedStates.History = MakeUnique<TNetRewindHistory<typename PhysicsTraits::StatesType>>(StateRedundancy, bIsLocalHistory);
 	ReplicatedStates.Owner = this;
 	
 	ActorComponent = HistoryComponent;
@@ -691,7 +692,7 @@ FORCEINLINE void UNetworkPhysicsComponent::CreateInputHistory(UActorComponent* H
 
 	InputData = MakeUnique<InputsType>();
 
-	ReplicatedInputs.History = MakeUnique<TNetRewindHistory<InputsType>>(NumFrames, bIsLocalHistory);
+	ReplicatedInputs.History = MakeUnique<TNetRewindHistory<InputsType>>(InputRedundancy, bIsLocalHistory);
 	ReplicatedInputs.Owner = this;
 
 	ActorComponent = HistoryComponent;
diff --git a/Engine/Source/Runtime/Experimental/Chaos/Public/RewindData.h b/Engine/Source/Runtime/Experimental/Chaos/Public/RewindData.h
--- a/Engine/Source/Runtime/Experimental/Chaos/Public/RewindData.h
+++ b/Engine/Source/Runtime/Experimental/Chaos/Public/RewindData.h
@@ -55,6 +55,10 @@ struct FBaseRewindHistory
 	FORCEINLINE virtual bool RecordData(const int32 RecordFrame, const void* HistoryData) { return true; }
 	UE_DEPRECATED(5.4, "Deprecated, use RecordData() instead")
 	FORCEINLINE virtual bool RecordDatas(const int32 RecordFrame, const void* HistoryDatas) { return RecordData(RecordFrame, HistoryDatas); }
+	
+	/** Check if we should record received data into history on @param RecordFrame.
+	* Can for example block received data from client from overriding server authoritative data */
+	FORCEINLINE virtual bool ShouldRecordReceivedDataOnFrame(const int32 RecordFrame) { return true; }
 
 	/** Create a polymorphic copy of only a range of frames, applying the frame offset to the copies
 	* @param StartFrame = Included
@@ -88,7 +92,7 @@ struct FBaseRewindHistory
 	FORCEINLINE virtual bool ApplyInputs(const int32 ApplyFrame, const bool bResetSolver) { return false; }
 
 	/** Return the most up to date frame entry in history */
-	virtual int32 GetLatestFrame() { return INDEX_NONE; }
+	virtual const int32 GetLatestFrame() const { return INDEX_NONE; }
 };
 
 /** Templated data history holding a data buffer */
@@ -111,7 +115,7 @@ protected:
 		for (int32 FrameIndex = 0; FrameIndex < NumFrames; ++FrameIndex)
 		{
 			const int32 ValidFrame = bMinData ? FMath::Max(0, DataFrame - FrameIndex) : DataFrame + FrameIndex;
-			const int32 ValidIndex = ValidFrame % NumFrames;
+			const int32 ValidIndex = GetFrameIndex(ValidFrame);
 
 			if (DataHistory[ValidIndex].LocalFrame == ValidFrame)
 			{
@@ -129,14 +133,14 @@ public :
 	/** Check if the history buffer contains an entry for the given frame*/
 	FORCEINLINE virtual bool HasValidData(const int32 ValidFrame) const override
 	{
-		const int32 ValidIndex = ValidFrame % NumFrames;
+		const int32 ValidIndex = GetFrameIndex(ValidFrame);
 		return ValidFrame == DataHistory[ValidIndex].LocalFrame;
 	}
 
 	/** Extract states at a given time */
 	FORCEINLINE virtual bool ExtractData(const int32 ExtractFrame, const bool bResetSolver, void* HistoryData, const bool bExactFrame = false) override
 	{
-		const int32 ExtractIndex = ExtractFrame % NumFrames;
+		const int32 ExtractIndex = GetFrameIndex(ExtractFrame);
 		if (ExtractFrame == DataHistory[ExtractIndex].LocalFrame)
 		{
 			CurrentFrame = ExtractFrame;
@@ -161,17 +165,20 @@ public :
 
 			if (MinFrameIndex != INDEX_NONE && MaxFrameIndex != INDEX_NONE)
 			{
+				DataType& ExtractedData = *static_cast<DataType*>(HistoryData);
+
 				PRAGMA_DISABLE_DEPRECATION_WARNINGS
 				// TODO: Change to InterpolateData() in UE 5.6 and remove deprecation pragma
-				static_cast<DataType*>(HistoryData)->InterpolateDatas(
+				ExtractedData.InterpolateDatas(
 					DataHistory[MinFrameIndex], DataHistory[MaxFrameIndex]);
 				PRAGMA_ENABLE_DEPRECATION_WARNINGS
 
 				const int32 DeltaFrame = FMath::Abs(ExtractFrame - DataHistory[MinFrameIndex].LocalFrame);
 
-				static_cast<DataType*>(HistoryData)->LocalFrame = ExtractFrame;
-				static_cast<DataType*>(HistoryData)->ServerFrame = DataHistory[MinFrameIndex].ServerFrame + DeltaFrame;
-				static_cast<DataType*>(HistoryData)->InputFrame = DataHistory[MinFrameIndex].InputFrame + DeltaFrame;
+				ExtractedData.LocalFrame = ExtractFrame;
+				ExtractedData.ServerFrame = DataHistory[MinFrameIndex].ServerFrame + DeltaFrame;
+				ExtractedData.InputFrame = INDEX_NONE; // Clear InputFrame since this history entry is now altered and doesn't correspond to the source entry with the same InputFrame
+
 #if DEBUG_NETWORK_PHYSICS
 				UE_LOG(LogTemp, Log, TEXT("		Smoothing data between frame %d and %d - > [%d %d]"), DataHistory[MinFrameIndex].LocalFrame, DataHistory[MaxFrameIndex].LocalFrame, static_cast<DataType*>(HistoryData)->InputFrame, static_cast<DataType*>(HistoryData)->ServerFrame);
 #endif
@@ -179,7 +186,15 @@ public :
 			}
 			else if (MinFrameIndex != INDEX_NONE)
 			{
-				*static_cast<DataType*>(HistoryData) = DataHistory[MinFrameIndex];
+				DataType& ExtractedData = *static_cast<DataType*>(HistoryData);
+				ExtractedData = DataHistory[MinFrameIndex];
+
+				const int32 DeltaFrame = FMath::Abs(ExtractFrame - DataHistory[MinFrameIndex].LocalFrame);
+
+				// Don't update LocalFrame here so that it can be used to check how many frames ahead of the last entry we are currently
+				ExtractedData.ServerFrame = DataHistory[MinFrameIndex].ServerFrame + DeltaFrame;
+				ExtractedData.InputFrame = INDEX_NONE; // Clear InputFrame since this history entry is now altered and doesn't correspond to the source entry with the same InputFrame
+
 #if DEBUG_NETWORK_PHYSICS
 				UE_LOG(LogTemp, Log, TEXT("		Setting data to frame %d"), DataHistory[MinFrameIndex].LocalFrame);
 #endif
@@ -201,10 +216,11 @@ public :
 		const int32 ToFrame = static_cast<DataType*>(ToData)->LocalFrame;
 		for (; FromFrame < ToFrame; FromFrame++)
 		{
-			const int32 FromIndex = FromFrame % NumFrames;
+			const int32 FromIndex = GetFrameIndex(FromFrame);
 			if (FromFrame == DataHistory[FromIndex].LocalFrame)
 			{
 				static_cast<DataType*>(ToData)->MergeData(DataHistory[FromIndex]);
+				static_cast<DataType*>(ToData)->InputFrame = INDEX_NONE; // Clear InputFrame since this history entry is now altered and doesn't correspond to the source entry with the same InputFrame
 			}
 		}
 	}
@@ -222,8 +238,7 @@ public :
 	/** Load the data from the buffer at a specific frame */
 	FORCEINLINE bool LoadData(const int32 LoadFrame)
 	{
-		const int32 LoadIndex = LoadFrame % NumFrames;
-		DataHistory[LoadIndex].LocalFrame = LoadFrame;
+		const int32 LoadIndex = GetFrameIndex(LoadFrame);
 		CurrentFrame = LoadFrame;
 		CurrentIndex = LoadIndex;
 		return true;
@@ -234,7 +249,7 @@ public :
 	/** Eval the data from the buffer at a specific frame */
 	FORCEINLINE bool EvalData(const int32 EvalFrame)
 	{
-		const int32 EvalIndex = EvalFrame % NumFrames;
+		const int32 EvalIndex = GetFrameIndex(EvalFrame);
 		if (EvalFrame == DataHistory[EvalIndex].LocalFrame)
 		{
 			CurrentFrame = EvalFrame;
@@ -253,11 +268,23 @@ public :
 		LoadDatas(RecordFrame);
 		PRAGMA_ENABLE_DEPRECATION_WARNINGS
 		DataHistory[CurrentIndex] = *static_cast<const DataType*>(HistoryData);
-		
+
 		LatestFrame = FMath::Max(LatestFrame, DataHistory[CurrentIndex].LocalFrame);
 		return true;
 	}
 
+	FORCEINLINE virtual bool ShouldRecordReceivedDataOnFrame(const int32 RecordFrame) override
+	{
+		if (RecordFrame < 0)
+		{
+			return false;
+		}
+
+		// Allow received data to get recorded if it's for a newer frame than already stored at the history index or if the data stored is not marked as received (i.e. it's locally predicted)
+		LoadData(RecordFrame);
+		return (!DataHistory[CurrentIndex].bReceivedData || DataHistory[CurrentIndex].LocalFrame < RecordFrame);
+	}
+
 	/** Current data that is being loaded/recorded*/
 	DataType& GetCurrentData() { return DataHistory[CurrentIndex]; }
 	UE_DEPRECATED(5.4, "Deprecated, use GetCurrentData() instead")
@@ -273,7 +300,7 @@ public :
 		uint32 NumData = 0;
 		for (uint32 ValidFrame = StartFrame; ValidFrame < EndFrame; ++ValidFrame)
 		{
-			const int32 ValidIndex = ValidFrame % NumFrames;
+			const int32 ValidIndex = GetFrameIndex(ValidFrame);
 			if (ValidFrame == DataHistory[ValidIndex].LocalFrame)
 			{
 				++NumData;
@@ -289,11 +316,16 @@ public :
 	TArray<DataType>& GetDatasArray() { return GetDataHistory(); }
 
 	/** Return the most up to date frame entry in history */
-	virtual int32 GetLatestFrame() override
+	virtual const int32 GetLatestFrame() const override
 	{
 		return LatestFrame;
 	}
 
+	FORCEINLINE const uint32 GetFrameIndex(const int32 Frame) const
+	{
+		return FMath::Abs(Frame % NumFrames);
+	}
+
 protected : 
 
 	/** Check if the history is on the local/remote client*/
