Static Analysis fix for GJK potential read from uninitialized vector array

diff --git a/Engine/Source/Runtime/Experimental/Chaos/Public/Chaos/GJK.h b/Engine/Source/Runtime/Experimental/Chaos/Public/Chaos/GJK.h
--- a/Engine/Source/Runtime/Experimental/Chaos/Public/Chaos/GJK.h
+++ b/Engine/Source/Runtime/Experimental/Chaos/Public/Chaos/GJK.h
@@ -1194,6 +1194,7 @@ namespace Chaos
 		TVec3<T> Bs[4];
 
 		FSimplex SimplexIDs;
+		SimplexIDs.NumVerts = 0;                // Initialization not needed, but compiler warns on uninitialized access to As/Bs
 		TVector<T, 3> Simplex[4];
 		T Barycentric[4] = { -1,-1,-1,-1 };		// Initialization not needed, but compiler warns
 		TVec3<T> Normal = -V;					// Remember the last good normal (i.e. don't update it if separation goes less than Epsilon and we can no longer normalize)
