Optimized gather for vectorized solver

diff --git a/Engine/Source/Runtime/Experimental/Chaos/Private/Chaos/Joint/PBDJointCachedSolverGaussSeidel.cpp b/Engine/Source/Runtime/Experimental/Chaos/Private/Chaos/Joint/PBDJointCachedSolverGaussSeidel.cpp
--- a/Engine/Source/Runtime/Experimental/Chaos/Private/Chaos/Joint/PBDJointCachedSolverGaussSeidel.cpp
+++ b/Engine/Source/Runtime/Experimental/Chaos/Private/Chaos/Joint/PBDJointCachedSolverGaussSeidel.cpp
@@ -395,26 +395,12 @@ void FPBDJointCachedSolver::InitPositionConstraints(
 
 	if (PositionConstraints.bUseSimd)
 	{
-		PositionConstraints.SetSoftLimit(0, false);
-		PositionConstraints.SetSoftLimit(1, false);
-		PositionConstraints.SetSoftLimit(2, false);
 		FRealSingle HardStiffness = FRealSingle(FPBDJointUtilities::GetLinearStiffness(SolverSettings, JointSettings));
 		PositionConstraints.Simd.ConstraintHardStiffness = VectorLoadFloat1(&HardStiffness);
-		PositionConstraints.ConstraintMaxLambda = FVec3(0);
-		PositionConstraints.SetValidDatas(0, false);
-		PositionConstraints.SetValidDatas(1, false);
-		PositionConstraints.SetValidDatas(2, false);
-		PositionConstraints.SetLimitsCheck(0, true);
-		PositionConstraints.SetLimitsCheck(1, true);
-		PositionConstraints.SetLimitsCheck(2, true);
-		PositionConstraints.SetMotionType(0, EJointMotionType::Locked);
-		PositionConstraints.SetMotionType(1, EJointMotionType::Locked);
-		PositionConstraints.SetMotionType(2, EJointMotionType::Locked);
+
 		if (bResetLambdas)
 		{
 			PositionConstraints.Simd.ConstraintLambda = VectorZeroFloat();
-			PositionConstraints.ConstraintLambdaVelocity = FVec3::Zero();
-			PositionConstraints.Simd.ConstraintLimits = VectorZeroFloat();
 		}
 	}
 	else
@@ -616,12 +602,6 @@ void FPBDJointCachedSolver::InitLockedPositionConstraintSimd(
 		const FVec3 LocalAxis = (CX[ConstraintIndex] < 0.0f) ? -ConstraintAxis : ConstraintAxis;
 		LocalDeltas[ConstraintIndex] = FRealSingle((CX[ConstraintIndex] < 0.0) ? -CX[ConstraintIndex] : CX[ConstraintIndex]);
 
-		PositionConstraints.SetMotionType(ConstraintIndex, EJointMotionType::Locked);
-		PositionConstraints.SetValidDatas(ConstraintIndex, true);
-		PositionConstraints.SetLimitsCheck(ConstraintIndex, false);
-		PositionConstraints.ConstraintVX[ConstraintIndex] = 0.0;
-		PositionConstraints.ConstraintRestitution[ConstraintIndex] = 0.0;
-
 		PositionConstraints.Simd.ConstraintAxis[ConstraintIndex] = MakeVectorRegisterFloatFromDouble(MakeVectorRegister(LocalAxis[0], LocalAxis[1], LocalAxis[2], 0.0f));
 	}
 	PositionConstraints.Simd.ConstraintCX = MakeVectorRegisterFloat(LocalDeltas[0], LocalDeltas[1], LocalDeltas[2], 0.0f);
@@ -644,7 +624,6 @@ void FPBDJointCachedSolver::InitLockedPositionConstraintSimd(
 		PositionConstraints.Simd.ConstraintDRAxis[ConstraintIndex][1] = VectorNegate(IA1);
 	}
 	PositionConstraints.Simd.ConstraintHardIM = MakeVectorRegisterFloatFromDouble(MakeVectorRegister(HardIM[0], HardIM[1], HardIM[2], 0.0f));
-	SetInitConstraintVelocity(ConstraintArm0, ConstraintArm1);
 
 }
 
@@ -932,9 +911,6 @@ void FPBDJointCachedSolver::ApplyLinearVelocityConstraints()
 {
 	if (PositionConstraints.bUseSimd)
 	{
-		check(PositionConstraints.GetSoftLimit(0) == false);
-		check(PositionConstraints.GetSoftLimit(1) == false);
-		check(PositionConstraints.GetSoftLimit(2) == false);
 		ApplyVelocityConstraintSimd();
 	}
 	else
@@ -1006,10 +982,6 @@ void FPBDJointCachedSolver::ApplyAxisVelocityConstraint(
 
 void FPBDJointCachedSolver::ApplyVelocityConstraintSimd()
 {
-	check(PositionConstraints.GetMotionType(0) != EJointMotionType::Limited);
-	check(PositionConstraints.GetMotionType(1) != EJointMotionType::Limited);
-	check(PositionConstraints.GetMotionType(2) != EJointMotionType::Limited);
-
 	const VectorRegister4Float IsGTEps = VectorCompareGT(VectorAbs(PositionConstraints.Simd.ConstraintLambda), GlobalVectorConstants::SmallNumber);
 
 	if (VectorMaskBits(IsGTEps))
@@ -1146,17 +1118,15 @@ void FPBDJointCachedSolver::InitRotationConstraints(
 	{
 		RotationConstraints.SettingsSoftStiffness = FVec3(FPBDJointUtilities::GetSoftTwistStiffness(SolverSettings, JointSettings), FPBDJointUtilities::GetSoftSwingStiffness(SolverSettings, JointSettings), FPBDJointUtilities::GetSoftSwingStiffness(SolverSettings, JointSettings));
 		RotationConstraints.SettingsSoftDamping = FVec3(FPBDJointUtilities::GetSoftTwistDamping(SolverSettings, JointSettings), FPBDJointUtilities::GetSoftSwingDamping(SolverSettings, JointSettings), FPBDJointUtilities::GetSoftSwingDamping(SolverSettings, JointSettings));
-		RotationConstraints.ConstraintMaxLambda = FVec3::Zero();
 
 		if (bResetLambdas)
 		{
 			RotationConstraints.Simd.ConstraintLambda = VectorZeroFloat();
-			RotationConstraints.ConstraintLambdaVelocity = FVec3::Zero();
 		}
 
 		RotationConstraints.Simd.ConstraintHardStiffness = MakeVectorRegisterFloatFromDouble(MakeVectorRegister(FPBDJointUtilities::GetTwistStiffness(SolverSettings, JointSettings),
 			FPBDJointUtilities::GetSwingStiffness(SolverSettings, JointSettings), FPBDJointUtilities::GetSwingStiffness(SolverSettings, JointSettings), 0.0f));
-		InitRotationConstraintsSimd(JointSettings, Dt);
+		InitRotationConstraintsSimd(JointSettings, FRealSingle(Dt));
 	}
 	else
 	{
@@ -1255,7 +1225,7 @@ void FPBDJointCachedSolver::InitRotationConstraints(
 
 void FPBDJointCachedSolver::InitRotationConstraintsSimd(
 	const FPBDJointSettings& JointSettings,
-	const FReal Dt)
+	const FRealSingle Dtf)
 {
 	const FVec3 Twist0 = ConnectorRs[0] * FJointConstants::TwistAxis();
 	const FVec3 Twist1 = ConnectorRs[1] * FJointConstants::TwistAxis();
@@ -1280,24 +1250,14 @@ void FPBDJointCachedSolver::InitRotationConstraintsSimd(
 	RotationConstraints.ConstraintRestitution = FVec3(JointSettings.TwistRestitution, JointSettings.SwingRestitution, JointSettings.SwingRestitution);
 
 	FVec3f ConstraintHardIM;
-	FVec3f ConstraintSoftStiffness;
-	FVec3f ConstraintSoftDamping;
-	FVec3f ConstraintSoftIM;
 	FVec3f LocalAngles;
 
 	for (int32 ConstraintIndex = 0; ConstraintIndex < 3; ConstraintIndex++)
 	{
 		const FVec3 LocalAxis = (Angles[ConstraintIndex] < 0.0f) ? -Axes[ConstraintIndex] : Axes[ConstraintIndex];
 		LocalAngles[ConstraintIndex] = FRealSingle((Angles[ConstraintIndex] < 0.0f) ? -Angles[ConstraintIndex] : Angles[ConstraintIndex]);
-
-		RotationConstraints.SetValidDatas(ConstraintIndex, true);
-		RotationConstraints.SetLimitsCheck(ConstraintIndex, true);
-		RotationConstraints.ConstraintVX[ConstraintIndex] = 0;
-
 		InitConstraintAxisAngularVelocities[ConstraintIndex] = FVec3::DotProduct(W(1) - W(0), LocalAxis);
 
-		// InitRotationDatasMass(RotationConstraints, ConstraintIndex, Dt);
-
 		RotationConstraints.Simd.ConstraintAxis[ConstraintIndex] = MakeVectorRegisterFloatFromDouble(MakeVectorRegister(LocalAxis[0], LocalAxis[1], LocalAxis[2], 0.0f));
 		const VectorRegister4Float& Axis = RotationConstraints.Simd.ConstraintAxis[ConstraintIndex];
 
@@ -1322,20 +1282,24 @@ void FPBDJointCachedSolver::InitRotationConstraintsSimd(
 		RotationConstraints.Simd.ConstraintDRAxis[ConstraintIndex][1] = VectorNegate(IA1);
 
 		ConstraintHardIM[ConstraintIndex] = II0 + II1;
-		const FReal SpringMassScale = (RotationConstraints.bAccelerationMode) ? (FReal)1 / (ConstraintHardIM[ConstraintIndex]) : (FReal)1;
-		ConstraintSoftStiffness[ConstraintIndex] = FRealSingle(SpringMassScale * RotationConstraints.SettingsSoftStiffness[ConstraintIndex] * Dt * Dt);
-		ConstraintSoftDamping[ConstraintIndex] = (bUsePositionBasedDrives) ? FRealSingle(SpringMassScale * RotationConstraints.SettingsSoftDamping[ConstraintIndex] * Dt) : 0.0f;
-		ConstraintSoftIM[ConstraintIndex] = FRealSingle((ConstraintSoftStiffness[ConstraintIndex] + ConstraintSoftDamping[ConstraintIndex]) * ConstraintHardIM[ConstraintIndex] + (FReal)1);
-
 	}
 
-	RotationConstraints.Simd.ConstraintArms[0] = VectorZeroFloat();
-	RotationConstraints.Simd.ConstraintArms[1] = VectorZeroFloat();
+	RotationConstraints.Simd.ConstraintHardIM = MakeVectorRegisterFloatFromDouble(MakeVectorRegister(ConstraintHardIM[0], ConstraintHardIM[1], ConstraintHardIM[2], 0.0f));
+	const VectorRegister4Float SoftDamping = MakeVectorRegisterFloatFromDouble(MakeVectorRegister(RotationConstraints.SettingsSoftDamping[0], RotationConstraints.SettingsSoftDamping[1], RotationConstraints.SettingsSoftDamping[2], 0.0f));
+	const VectorRegister4Float SoftStiffness = MakeVectorRegisterFloatFromDouble(MakeVectorRegister(RotationConstraints.SettingsSoftStiffness[0], RotationConstraints.SettingsSoftStiffness[1], RotationConstraints.SettingsSoftStiffness[2], 0.0f));
+
+	const VectorRegister4Float Dt = VectorLoadFloat1(&Dtf);
+
+	const VectorRegister4Float SpringMassScale = RotationConstraints.bAccelerationMode ? VectorDivide(GlobalVectorConstants::FloatOne, RotationConstraints.Simd.ConstraintHardIM) : GlobalVectorConstants::FloatOne;
+	RotationConstraints.Simd.ConstraintSoftStiffness = VectorMultiply(VectorMultiply(SpringMassScale, SoftStiffness), VectorMultiply(Dt, Dt));
+	RotationConstraints.Simd.ConstraintSoftDamping = bUsePositionBasedDrives ? VectorMultiply(SpringMassScale, VectorMultiply(SoftDamping, Dt)) : VectorZeroFloat();
+	RotationConstraints.Simd.ConstraintSoftIM = VectorAdd(VectorMultiply(VectorAdd(RotationConstraints.Simd.ConstraintSoftStiffness, RotationConstraints.Simd.ConstraintSoftDamping), RotationConstraints.Simd.ConstraintHardIM), GlobalVectorConstants::FloatOne);
+
+
+
 	RotationConstraints.Simd.ConstraintCX = MakeVectorRegisterFloat(LocalAngles[0], LocalAngles[1], LocalAngles[2], 0.0f);
 	RotationConstraints.Simd.ConstraintHardIM = MakeVectorRegisterFloat(ConstraintHardIM[0], ConstraintHardIM[1], ConstraintHardIM[2], 0.0f);
-	RotationConstraints.Simd.ConstraintSoftStiffness = MakeVectorRegisterFloat(ConstraintSoftStiffness[0], ConstraintSoftStiffness[1], ConstraintSoftStiffness[2], 0.0f);
-	RotationConstraints.Simd.ConstraintSoftDamping = MakeVectorRegisterFloat(ConstraintSoftDamping[0], ConstraintSoftDamping[1], ConstraintSoftDamping[2], 0.0f);
-	RotationConstraints.Simd.ConstraintSoftIM = MakeVectorRegisterFloat(ConstraintSoftIM[0], ConstraintSoftIM[1], ConstraintSoftIM[2], 0.0f);
+
 	RotationConstraints.Simd.ConstraintLimits = MakeVectorRegisterFloatFromDouble(MakeVectorRegister(JointSettings.AngularLimits[0], JointSettings.AngularLimits[1], JointSettings.AngularLimits[2], UE_BIG_NUMBER));
 }
 
@@ -1365,10 +1329,7 @@ void FPBDJointCachedSolver::InitRotationConstraintDatas(
 	const FReal LocalAngle = (ConstraintAngle < 0.0f) ? -ConstraintAngle : ConstraintAngle;
 
 	RotationConstraints.UpdateDatas(ConstraintIndex, LocalAxis, LocalAngle, ConstraintRestitution, bCheckLimit);
-
-
 	RotationConstraints.Data.ConstraintLimits[ConstraintIndex] = JointSettings.AngularLimits[ConstraintIndex];
-
 	InitConstraintAxisAngularVelocities[ConstraintIndex] = FVec3::DotProduct(W(1) - W(0), LocalAxis);
 
 	InitRotationDatasMass(RotationConstraints, ConstraintIndex, Dt);
@@ -2153,9 +2114,6 @@ void FPBDJointCachedSolver::ApplyPositionProjection(
 
 	if (PositionConstraints.bUseSimd)
 	{
-		check(PositionConstraints.GetSoftLimit(0) == false);
-		check(PositionConstraints.GetSoftLimit(1) == false);
-		check(PositionConstraints.GetSoftLimit(2) == false);
 		ApplyPositionProjectionSimd(LinearProjection);
 	}
 	else
@@ -2571,24 +2529,16 @@ void FPBDJointCachedSolver::InitRotationConstraintDrive(
 
 void FPBDJointCachedSolver::InitRotationConstraintDriveSimd(
 	FVec3 ConstraintAxes[3],
-	const FReal Dt,
+	const FRealSingle Dtf,
 	const FVec3 DeltaAngles
 )
 {
 	FVec3 ConstraintHardIM;
-	FVec3 ConstraintSoftStiffness;
-	FVec3 ConstraintSoftDamping;
-	FVec3 ConstraintSoftIM;
-
 	RotationDrives.Simd.ConstraintCX = MakeVectorRegisterFloatFromDouble(MakeVectorRegister(DeltaAngles[0], DeltaAngles[1], DeltaAngles[2], 0.0f));
 
 	for (int32 ConstraintIndex = 0; ConstraintIndex < 3; ConstraintIndex++)
 	{
-		RotationDrives.SetValidDatas(ConstraintIndex, true);
-		RotationDrives.SetLimitsCheck(ConstraintIndex, true);
-
 		RotationDrives.ConstraintVX[ConstraintIndex] = 0.0;
-		RotationDrives.ConstraintRestitution[ConstraintIndex] = 0.0;
 
 		RotationDrives.Simd.ConstraintAxis[ConstraintIndex] = MakeVectorRegisterFloatFromDouble(MakeVectorRegister(ConstraintAxes[ConstraintIndex][0], ConstraintAxes[ConstraintIndex][1], ConstraintAxes[ConstraintIndex][2], 0.0f));
 
@@ -2616,21 +2566,18 @@ void FPBDJointCachedSolver::InitRotationConstraintDriveSimd(
 		RotationDrives.Simd.ConstraintDRAxis[ConstraintIndex][0] = IA0;
 		RotationDrives.Simd.ConstraintDRAxis[ConstraintIndex][1] = VectorNegate(IA1);
 
-
-		check(RotationDrives.GetSoftLimit(ConstraintIndex))
-		const FReal SpringMassScale = (RotationDrives.bAccelerationMode) ? (FReal)1 / (ConstraintHardIM[ConstraintIndex]) : (FReal)1;
-		ConstraintSoftStiffness[ConstraintIndex] = SpringMassScale * RotationDrives.SettingsSoftStiffness[ConstraintIndex] * Dt * Dt;
-		ConstraintSoftDamping[ConstraintIndex] = (bUsePositionBasedDrives) ? SpringMassScale * RotationDrives.SettingsSoftDamping[ConstraintIndex] * Dt : 0;
-		ConstraintSoftIM[ConstraintIndex] = (ConstraintSoftStiffness[ConstraintIndex] + ConstraintSoftDamping[ConstraintIndex]) * ConstraintHardIM[ConstraintIndex] + (FReal)1;
-
+		check(RotationDrives.GetSoftLimit(ConstraintIndex));
 	}
-
-	RotationDrives.Simd.ConstraintArms[0] = VectorZeroFloat();
-	RotationDrives.Simd.ConstraintArms[1] = VectorZeroFloat();
 	RotationDrives.Simd.ConstraintHardIM = MakeVectorRegisterFloatFromDouble(MakeVectorRegister(ConstraintHardIM[0], ConstraintHardIM[1], ConstraintHardIM[2], 0.0f));
-	RotationDrives.Simd.ConstraintSoftStiffness = MakeVectorRegisterFloatFromDouble(MakeVectorRegister(ConstraintSoftStiffness[0], ConstraintSoftStiffness[1], ConstraintSoftStiffness[2], 0.0f));
-	RotationDrives.Simd.ConstraintSoftDamping = MakeVectorRegisterFloatFromDouble(MakeVectorRegister(ConstraintSoftDamping[0], ConstraintSoftDamping[1], ConstraintSoftDamping[2], 0.0f));
-	RotationDrives.Simd.ConstraintSoftIM = MakeVectorRegisterFloatFromDouble(MakeVectorRegister(ConstraintSoftIM[0], ConstraintSoftIM[1], ConstraintSoftIM[2], 0.0f));
+	const VectorRegister4Float SoftDamping = MakeVectorRegisterFloatFromDouble(MakeVectorRegister(RotationDrives.SettingsSoftDamping[0], RotationDrives.SettingsSoftDamping[1], RotationDrives.SettingsSoftDamping[2], 0.0f));
+	const VectorRegister4Float SoftStiffness = MakeVectorRegisterFloatFromDouble(MakeVectorRegister(RotationDrives.SettingsSoftStiffness[0], RotationDrives.SettingsSoftStiffness[1], RotationDrives.SettingsSoftStiffness[2], 0.0f));
+	const VectorRegister4Float Dt = VectorLoadFloat1(&Dtf);
+
+	const VectorRegister4Float SpringMassScale = RotationDrives.bAccelerationMode ? VectorDivide(GlobalVectorConstants::FloatOne, RotationDrives.Simd.ConstraintHardIM) : GlobalVectorConstants::FloatOne;
+	RotationDrives.Simd.ConstraintSoftStiffness = VectorMultiply(VectorMultiply(SpringMassScale, SoftStiffness), VectorMultiply(Dt, Dt));
+	RotationDrives.Simd.ConstraintSoftDamping = bUsePositionBasedDrives ? VectorMultiply(SpringMassScale, VectorMultiply(SoftDamping, Dt)) : VectorZeroFloat();
+	RotationDrives.Simd.ConstraintSoftIM = VectorAdd(VectorMultiply(VectorAdd(RotationDrives.Simd.ConstraintSoftStiffness, RotationDrives.Simd.ConstraintSoftDamping), RotationDrives.Simd.ConstraintHardIM), GlobalVectorConstants::FloatOne);
+	
 }
 
 void FPBDJointCachedSolver::InitSwingTwistDrives(
@@ -2673,7 +2620,7 @@ void FPBDJointCachedSolver::InitSwingTwistDrives(
 	if (RotationDrives.bUseSimd)
 	{
 		FVec3 ConstraintAxes[3] = { ConnectorRs[1] * FJointConstants::TwistAxis(), ConnectorRs[1] * FJointConstants::Swing1Axis(), ConnectorRs[1] * FJointConstants::Swing2Axis() };
-		InitRotationConstraintDriveSimd(ConstraintAxes, Dt, FVec3(DTwistAngle, DSwing1Angle, DSwing2Angle));
+		InitRotationConstraintDriveSimd(ConstraintAxes, FRealSingle(Dt), FVec3(DTwistAngle, DSwing1Angle, DSwing2Angle));
 
 		RotationDrives.Simd.ConstraintHardStiffness = VectorZeroFloat();
 		RotationDrives.Simd.ConstraintLambda = VectorZeroFloat();
@@ -2739,7 +2686,7 @@ void FPBDJointCachedSolver::InitSLerpDrive(
 	{
 		RotationDrives.Simd.ConstraintLambda = VectorZeroFloat();
 		RotationDrives.Simd.ConstraintHardStiffness = VectorZeroFloat();
-		InitRotationConstraintDriveSimd(Axes, Dt, AxisAngles);
+		InitRotationConstraintDriveSimd(Axes, FRealSingle(Dt), AxisAngles);
 	}
 	else
 	{
diff --git a/Engine/Source/Runtime/Experimental/Chaos/Public/Chaos/Joint/PBDJointCachedSolverGaussSeidel.h b/Engine/Source/Runtime/Experimental/Chaos/Public/Chaos/Joint/PBDJointCachedSolverGaussSeidel.h
--- a/Engine/Source/Runtime/Experimental/Chaos/Public/Chaos/Joint/PBDJointCachedSolverGaussSeidel.h
+++ b/Engine/Source/Runtime/Experimental/Chaos/Public/Chaos/Joint/PBDJointCachedSolverGaussSeidel.h
@@ -674,7 +674,7 @@ struct FAxisConstraintDatas
 
 		void InitRotationConstraintsSimd(
 			const FPBDJointSettings& JointSettings,
-			const FReal Dt);
+			const FRealSingle Dt);
 
 		void CorrectAxisAngleConstraint(
 			const FPBDJointSettings& JointSettings,
@@ -822,7 +822,7 @@ struct FAxisConstraintDatas
 
 		void InitRotationConstraintDriveSimd(
 			FVec3 ConstraintAxes[3],
-			const FReal Dt,
+			const FRealSingle Dt,
 			const FVec3 DeltaAngles);
 		
 		void InitSwingTwistDrives(
