Chaos - fix contact friction for large world coordinates

#11756

diff --git a/Engine/Source/Runtime/Experimental/Chaos/Private/Chaos/Collision/PBDCollisionContainerSolver.cpp b/Engine/Source/Runtime/Experimental/Chaos/Private/Chaos/Collision/PBDCollisionContainerSolver.cpp
--- a/Engine/Source/Runtime/Experimental/Chaos/Private/Chaos/Collision/PBDCollisionContainerSolver.cpp
+++ b/Engine/Source/Runtime/Experimental/Chaos/Private/Chaos/Collision/PBDCollisionContainerSolver.cpp
@@ -129,9 +129,9 @@ namespace Chaos
 		FVec3f WorldFrictionDelta = FVec3f(0);
 		if (ManifoldPoint.Flags.bHasStaticFrictionAnchor)
 		{
-			const FVec3f FrictionDelta0 = ShapeWorldTransform0.TransformPositionNoScale(FVec3(ManifoldPoint.ShapeAnchorPoints[0]));
-			const FVec3f FrictionDelta1 = ShapeWorldTransform1.TransformPositionNoScale(FVec3(ManifoldPoint.ShapeAnchorPoints[1]));
-			WorldFrictionDelta = FrictionDelta0 - FrictionDelta1;
+			const FVec3 FrictionDelta0 = ShapeWorldTransform0.TransformPositionNoScale(FVec3(ManifoldPoint.ShapeAnchorPoints[0]));
+			const FVec3 FrictionDelta1 = ShapeWorldTransform1.TransformPositionNoScale(FVec3(ManifoldPoint.ShapeAnchorPoints[1]));
+			WorldFrictionDelta = FVec3f(FrictionDelta0 - FrictionDelta1);
 		}
 		else
 		{
