Chaos - fix memory leak in FScratchBuffer

diff --git a/Engine/Source/Runtime/Experimental/Chaos/Public/Chaos/Framework/ScratchBuffer.h b/Engine/Source/Runtime/Experimental/Chaos/Public/Chaos/Framework/ScratchBuffer.h
--- a/Engine/Source/Runtime/Experimental/Chaos/Public/Chaos/Framework/ScratchBuffer.h
+++ b/Engine/Source/Runtime/Experimental/Chaos/Public/Chaos/Framework/ScratchBuffer.h
@@ -28,6 +28,11 @@ namespace Chaos
 			{
 			}
 
+			~FScratchBuffer()
+			{
+				DestroyBuffer();
+			}
+
 			size_t BufferSize() const
 			{
 				// NOTE: this size does not include the sentinel
