Managed array collection : encapsulation of variable members of FValueType
- this is in preparation of copy on write attribute support in managed array collection

diff --git a/Engine/Source/Runtime/Experimental/Chaos/Private/GeometryCollection/GeometryCollectionUtility.cpp b/Engine/Source/Runtime/Experimental/Chaos/Private/GeometryCollection/GeometryCollectionUtility.cpp
--- a/Engine/Source/Runtime/Experimental/Chaos/Private/GeometryCollection/GeometryCollectionUtility.cpp
+++ b/Engine/Source/Runtime/Experimental/Chaos/Private/GeometryCollection/GeometryCollectionUtility.cpp
@@ -297,7 +297,7 @@ namespace GeometryCollection
 				if (NumVertices)
 				{
 					// transforms group
-					TManagedArray<FTransform>& Transform = InCollection->ModifyAttribute<FTransform>(FTransformCollection::TransformAttribute, FTransformCollection::TransformGroup);
+					TManagedArray<FTransform3f>& Transform = InCollection->ModifyAttribute<FTransform3f>(FTransformCollection::TransformAttribute, FTransformCollection::TransformGroup);
 					TManagedArray<int32>& TransformToGeometryIndex = InCollection->ModifyAttribute<int32>("TransformToGeometryIndex", FTransformCollection::TransformGroup);
 					// vertices group
 					TManagedArray<int32>& BoneMap = InCollection->ModifyAttribute<int32>("BoneMap", FGeometryCollection::VerticesGroup);
diff --git a/Engine/Source/Runtime/Experimental/Chaos/Private/GeometryCollection/ManagedArrayCollection.cpp b/Engine/Source/Runtime/Experimental/Chaos/Private/GeometryCollection/ManagedArrayCollection.cpp
--- a/Engine/Source/Runtime/Experimental/Chaos/Private/GeometryCollection/ManagedArrayCollection.cpp
+++ b/Engine/Source/Runtime/Experimental/Chaos/Private/GeometryCollection/ManagedArrayCollection.cpp
@@ -58,9 +58,9 @@ void FManagedArrayCollection::RemoveElements(const FName& Group, const TArray<in
 			//
 			// Reindex attributes dependent on the group being resized
 			//
-			if (Entry.Value.GroupIndexDependency == Group && Params.bReindexDependentAttibutes)
+			if (Entry.Value.GetGroupIndexDependency() == Group && Params.bReindexDependentAttibutes)
 			{
-				Entry.Value.Value->Reindex(Offsets, GroupSize - DelListNum, SortedDeletionList, DeletionSet);
+				Entry.Value.Modify().Reindex(Offsets, GroupSize - DelListNum, SortedDeletionList, DeletionSet);
 			}
 
 			//
@@ -68,7 +68,7 @@ void FManagedArrayCollection::RemoveElements(const FName& Group, const TArray<in
 			//
 			if (Entry.Key.Get<1>() == Group)
 			{
-				Entry.Value.Value->RemoveElements(SortedDeletionList);
+				Entry.Value.Modify().RemoveElements(SortedDeletionList);
 			}
 
 		}
@@ -107,12 +107,9 @@ bool FManagedArrayCollection::HasAttributes(const TArray<FManagedArrayCollection
 {
 	for (const FManagedType& ManagedType : Types)
 	{
-		FKeyType Key = FManagedArrayCollection::MakeMapKey(ManagedType.Name, ManagedType.Group);
-		if(Map.Contains(Key))
+		if(const FValueType* FoundValue = Map.Find({ ManagedType.Name, ManagedType.Group }))
 		{
-			const FValueType& FoundValue = Map[Key];
-
-			if (FoundValue.ArrayType != ManagedType.Type)
+			if (FoundValue->GetArrayType() != ManagedType.Type)
 			{
 				return false;
 			}
@@ -123,30 +120,27 @@ bool FManagedArrayCollection::HasAttributes(const TArray<FManagedArrayCollection
 
 FManagedArrayCollection::EArrayType FManagedArrayCollection::GetAttributeType(FName Name, FName Group) const
 {
-	const FKeyType Key = FManagedArrayCollection::MakeMapKey(Name, Group);
-	if (const FValueType* Attribute = Map.Find(Key))
+	if (const FValueType* Attribute = Map.Find({ Name, Group }))
 	{
-		return Attribute->ArrayType;
+		return Attribute->GetArrayType();
 	}
 	return EArrayType::FNoneType;
 }
 
 bool FManagedArrayCollection::IsAttributeDirty(FName Name, FName Group) const
 {
-	const FKeyType Key = FManagedArrayCollection::MakeMapKey(Name, Group);
-	if (const FValueType* Attribute = Map.Find(Key))
+	if (const FValueType* Attribute = Map.Find({ Name, Group }))
 	{
-		return (Attribute->Value)? Attribute->Value->IsDirty(): false;
+		return Attribute->IsDirty();
 	}
 	return false;
 }
 
 bool FManagedArrayCollection::IsAttributePersistent(FName Name, FName Group) const
 {
-	const FKeyType Key = FManagedArrayCollection::MakeMapKey(Name, Group);
-	if (const FValueType* Attribute = Map.Find(Key))
+	if (const FValueType* Attribute = Map.Find({ Name, Group }))
 	{
-		return Attribute->Saved;
+		return Attribute->IsPersistent();
 	}
 	return false;
 }
@@ -187,7 +181,7 @@ int32 FManagedArrayCollection::AddElements(int32 NumberElements, FName Group)
 	{
 		if (Entry.Key.Get<1>() == Group)
 		{
-			Entry.Value.Value->Resize(StartSize + NumberElements);
+			Entry.Value.Resize(StartSize + NumberElements);
 		}
 	}
 	GroupInfo[Group].Size += NumberElements;
@@ -237,7 +231,7 @@ void FManagedArrayCollection::Append(const FManagedArrayCollection& InCollection
 			const FValueType& DestValue = Map[Entry.Key];
 
 			// If we don't have a type match don't attempt the copy.
-			if (OriginalValue.ArrayType != DestValue.ArrayType)
+			if (OriginalValue.GetArrayType() != DestValue.GetArrayType())
 			{
 				bMatchingAttributes = false;
 				ensureMsgf(false, TEXT("Failed : Type error in FManagedArrayCollection::AppendCollection (%s:%s)"), 
@@ -264,35 +258,17 @@ void FManagedArrayCollection::Append(const FManagedArrayCollection& InCollection
 		// copy values
 		for (const TTuple<FKeyType, FValueType>& Entry : InCollection.Map)
 		{
-			FName AttributeName = Entry.Key.Get<0>();
-			FName GroupName = Entry.Key.Get<1>();
+			const FName AttributeName = Entry.Key.Get<0>();
+			const FName GroupName = Entry.Key.Get<1>();
 
 			if (HasAttribute(AttributeName, GroupName))
 			{
-				Map[Entry.Key].Value->CopyRange(*Entry.Value.Value, 0, Entry.Value.Value->Num());
+				Map[Entry.Key].CopyFrom(Entry.Value);
 			}
 			else
 			{
-				// Copied from FManagedArrayCollection::CopyAttribute, but this 
-				// does not manage destination size, as that is pre allocated above. 
-				FKeyType Key = FManagedArrayCollection::MakeMapKey(AttributeName, GroupName);
-
-				if (!HasAttribute(AttributeName, GroupName))
-				{
-					const FValueType& V = InCollection.Map[Key];
-					EArrayType Type = V.ArrayType;
-					FValueType Value(Type, *NewManagedTypedArray(Type));
-					Value.Value->Resize(NumElements(GroupName));
-					Value.GroupIndexDependency = V.GroupIndexDependency;
-					Value.Saved = V.Saved;
-					Value.bExternalValue = V.bExternalValue;
-					Map.Add(Key, MoveTemp(Value));
-				}
-
-				const FValueType& OriginalValue = InCollection.Map[Key];
-				const FValueType& DestValue = Map[Key];
-				check(OriginalValue.ArrayType == DestValue.ArrayType);
-				DestValue.Value->Init(*OriginalValue.Value);
+				FValueType NewAttribute(Entry.Value);
+				Map.Add(Entry.Key, MoveTemp(NewAttribute));
 			}
 		}
 	}
@@ -305,7 +281,9 @@ void FManagedArrayCollection::RemoveAttribute(FName Name, FName Group)
 	FValueType* FoundValue = Map.Find(Key);
 	if (FoundValue != nullptr)
 	{
-		FoundValue->Value->Empty();
+		// TODO : is that even just necessary only for external values ? 
+		// if so maybe we should have this logic inside the empty code of FValueType
+		FoundValue->Empty();
 	}
 	Map.Remove(Key);
 }
@@ -319,14 +297,11 @@ void FManagedArrayCollection::RemoveGroup(FName Group)
 		{
 			DelList.Add(Entry.Key.Get<0>());
 		}
-		if (Entry.Value.GroupIndexDependency == Group)
-		{
-			Entry.Value.GroupIndexDependency = "";
-		}
+		Entry.Value.RemoveGroupIndexDependency(Group);
 	}
 	for (const FName& AttrName : DelList)
 	{
-		Map.Remove(FManagedArrayCollection::MakeMapKey(AttrName, Group));
+		Map.Remove({ AttrName, Group });
 	}
 
 	GroupInfo.Remove(Group);
@@ -347,7 +322,7 @@ void FManagedArrayCollection::Resize(int32 Size, FName Group)
 	{
 		if (Entry.Key.Get<1>() == Group)
 		{
-			Entry.Value.Value->Resize(Size);
+			Entry.Value.Resize(Size);
 		}
 	}
 	GroupInfo[Group].Size = Size;
@@ -367,7 +342,7 @@ void FManagedArrayCollection::Reserve(int32 Size, FName Group)
 	{
 		if (Entry.Key.Get<1>() == Group)
 		{
-			Entry.Value.Value->Reserve(Size);
+			Entry.Value.Reserve(Size);
 		}
 	}
 }
@@ -380,7 +355,7 @@ void FManagedArrayCollection::EmptyGroup(FName Group)
 	{
 		if (Entry.Key.Get<1>() == Group)
 		{
-			Entry.Value.Value->Empty();
+			Entry.Value.Empty();
 		}
 	}
 
@@ -402,14 +377,14 @@ void FManagedArrayCollection::ReorderElements(FName Group, const TArray<int32>&
 	for (TTuple<FKeyType, FValueType>& Entry : Map)
 	{
 		// Reindex attributes dependent on the group being reordered
-		if (Entry.Value.GroupIndexDependency == Group)
+		if (Entry.Value.GetGroupIndexDependency() == Group)
 		{
-			Entry.Value.Value->ReindexFromLookup(InverseNewOrder);
+			Entry.Value.Modify().ReindexFromLookup(InverseNewOrder);
 		}
 
 		if (Entry.Key.Get<1>() == Group)
 		{
-			Entry.Value.Value->Reorder(NewOrder);
+			Entry.Value.Modify().Reorder(NewOrder);
 		}
 	}
 }
@@ -420,7 +395,7 @@ void FManagedArrayCollection::SetDependency(FName Name, FName Group, FName Depen
 	FKeyType Key = FManagedArrayCollection::MakeMapKey(Name, Group);
 	if (ensure(bAllowCircularDependency || !IsConnected(DependencyGroup, Group)))
 	{
-		Map[Key].GroupIndexDependency = DependencyGroup;
+		Map[Key].SetGroupIndexDependency(DependencyGroup);
 	}
 }
 
@@ -428,7 +403,7 @@ FName FManagedArrayCollection::GetDependency(FName Name, FName Group) const
 {
 	check(HasAttribute(Name, Group));
 	const FKeyType Key = FManagedArrayCollection::MakeMapKey(Name, Group);
-	return Map[Key].GroupIndexDependency;
+	return Map[Key].GetGroupIndexDependency();
 }
 
 void FManagedArrayCollection::RemoveDependencyFor(FName Group)
@@ -436,10 +411,7 @@ void FManagedArrayCollection::RemoveDependencyFor(FName Group)
 	ensure(HasGroup(Group));
 	for (TTuple<FKeyType, FValueType>& Entry : Map)
 	{
-		if (Entry.Value.GroupIndexDependency == Group)
-		{
-			Entry.Value.GroupIndexDependency = "";
-		}
+		Entry.Value.RemoveGroupIndexDependency(Group);
 	}
 }
 
@@ -449,7 +421,7 @@ void FManagedArrayCollection::SetDefaults(FName Group, uint32 StartSize, uint32
 	{
 		if (Entry.Key.Get<1>() == Group)
 		{
-			Entry.Value.Value->SetDefaults(StartSize, NumElements, Entry.Value.GroupIndexDependency != "");
+			Entry.Value.Modify().SetDefaults(StartSize, NumElements, !Entry.Value.GetGroupIndexDependency().IsNone());
 		}
 	}
 }
@@ -487,12 +459,12 @@ void FManagedArrayCollection::CopyMatchingAttributesFrom(const FManagedArrayColl
 		{
 			continue;
 		}
-		if (const FValueType* FromAttribute = FromCollection.Map.Find(MakeMapKey(AttributeName, GroupName)))
+		if (const FValueType* FromAttribute = FromCollection.Map.Find({ AttributeName, GroupName }))
 		{
 			FValueType& ToAttribute = Entry.Value;
-			if (ToAttribute.ArrayType == FromAttribute->ArrayType)
+			if (ToAttribute.GetArrayType() == FromAttribute->GetArrayType())
 			{
-				ToAttribute.Value->Init(*FromAttribute->Value);
+				ToAttribute.InitFrom(*FromAttribute);
 			}
 		}
 	}
@@ -526,7 +498,7 @@ void FManagedArrayCollection::CopyMatchingAttributesFrom(
 			const FValueType& DestValue = Map[Entry.Key];
 
 			// If we don't have a type match don't attempt the copy.
-			if(OriginalValue.ArrayType == DestValue.ArrayType)
+			if(OriginalValue.GetArrayType() == DestValue.GetArrayType())
 			{
 				CopyAttribute(InCollection, Entry.Key.Get<0>(), Entry.Key.Get<1>());
 			}
@@ -547,26 +519,12 @@ void FManagedArrayCollection::CopyAttribute(const FManagedArrayCollection& InCol
 
 	if (!HasAttribute(DestName, Group))
 	{
-		const FValueType& V = InCollection.Map[SrcKey];
-		EArrayType Type = V.ArrayType;
-		FValueType Value(Type, *NewManagedTypedArray(Type));
-		Value.Value->Resize(NumElements(Group));
-		Value.GroupIndexDependency = V.GroupIndexDependency;
-		Value.Saved = V.Saved;
-		Value.bExternalValue = V.bExternalValue;
-		Map.Add(DestKey, MoveTemp(Value));
+		FValueType NewAttribute(InCollection.Map[SrcKey]);
+		Map.Add(DestKey, MoveTemp(NewAttribute));
 	}
 
 	const FValueType& OriginalValue = InCollection.Map[SrcKey];
-	const FValueType& DestValue = Map[DestKey];
-	if (OriginalValue.ArrayType == DestValue.ArrayType)
-	{
-		DestValue.Value->Init(*OriginalValue.Value);
-	}
-	else
-	{
-		DestValue.Value->Convert(*OriginalValue.Value);
-	}
+	Map[DestKey].InitFrom(OriginalValue);
 }
 
 bool FManagedArrayCollection::IsConnected(FName StartingNode, FName TargetNode)
@@ -578,8 +536,8 @@ bool FManagedArrayCollection::IsConnected(FName StartingNode, FName TargetNode)
 		{
 			if (!DMap.Contains(Entry.Key.Get<1>()))
 				DMap.Add(Entry.Key.Get<1>(), TArray<FName>());
-			if (!Entry.Value.GroupIndexDependency.IsNone())
-				DMap[Entry.Key.Get<1>()].AddUnique(Entry.Value.GroupIndexDependency);
+			if (!Entry.Value.GetGroupIndexDependency().IsNone())
+				DMap[Entry.Key.Get<1>()].AddUnique(Entry.Value.GetGroupIndexDependency());
 		}
 
 		if (DMap.Contains(StartingNode))
@@ -625,7 +583,7 @@ FString FManagedArrayCollection::ToString() const
 			FKeyType Key = FManagedArrayCollection::MakeMapKey(AttributeName, GroupName);
 			const FValueType& Value = Map[Key];
 
-			const SIZE_T AttributeAllocatedSize = Value.Value? Value.Value->GetAllocatedSize() : 0;
+			const SIZE_T AttributeAllocatedSize = Value.Get().GetAllocatedSize();
 			const FString AttributeAllocatedSizeStr = FString::FormatAsNumber((int32)AttributeAllocatedSize);
 
 			const TArray<FStringFormatArg> AttributeInfos = { AttributeName.ToString(), AttributeAllocatedSizeStr };
@@ -640,10 +598,7 @@ SIZE_T FManagedArrayCollection::GetAllocatedSize() const
 	SIZE_T AllocatedSize = Map.GetAllocatedSize();
 	for (const TTuple<FKeyType, FValueType>& Entry : Map)
 	{
-		if (Entry.Value.Value)
-		{
-			AllocatedSize += Entry.Value.Value->GetAllocatedSize();
-		}
+		AllocatedSize += Entry.Value.Get().GetAllocatedSize();
 	}
 	return AllocatedSize;
 }
@@ -656,10 +611,7 @@ void FManagedArrayCollection::GetElementSizeInfoForGroups(TArray<TPair<FName, SI
 	for (const TPair<FKeyType, FValueType>& Attribute : Map)
 	{
 		SIZE_T& GroupSize = GroupToElementSizeMap.FindOrAdd(Attribute.Key.Get<1>());
-		if (FManagedArrayBase* Array = Attribute.Value.Value)
-		{
-			GroupSize += Array->GetTypeSize();
-		}
+		GroupSize += Attribute.Value.Get().GetTypeSize();
 	}
 
 	for (const TPair<FName, SIZE_T>& GroupAndElementSize : GroupToElementSizeMap)
@@ -717,14 +669,14 @@ void FManagedArrayCollection::Serialize(Chaos::FChaosArchive& Ar)
 		{
 			if (FValueType* Existing = Map.Find(Pair.Key))
 			{
-				if (Existing->ArrayType == Pair.Value.ArrayType)
+				if (Existing->GetArrayType() == Pair.Value.GetArrayType())
 				{
-					Existing->Value->ExchangeArrays(*Pair.Value.Value);	//if there is already an entry do an exchange. This way external arrays get correct serialization
+					Existing->Exchange(Pair.Value);	//if there is already an entry do an exchange. This way external arrays get correct serialization
 					//question: should we validate if group dependency has changed in some invalid way?
 				}
 				else
 				{
-					Existing->Value->Convert(*Pair.Value.Value);
+					Existing->Convert(Pair.Value);
 				}
 			}
 			else
@@ -738,7 +690,7 @@ void FManagedArrayCollection::Serialize(Chaos::FChaosArchive& Ar)
 		TArray<FKeyType> ToRemoveKeys;
 		for (const TTuple<FKeyType, FValueType>& Pair : Map)
 		{
-			if (!Pair.Value.bExternalValue && !TmpMap.Find(Pair.Key))
+			if (!Pair.Value.IsExternal() && !TmpMap.Find(Pair.Key))
 			{
 				ToRemoveKeys.Add(Pair.Key);
 			}
@@ -752,9 +704,9 @@ void FManagedArrayCollection::Serialize(Chaos::FChaosArchive& Ar)
 		for (TTuple<FKeyType, FValueType>& Pair : Map)
 		{
 			const int32 GroupSize = GroupInfo[Pair.Key.Get<1>()].Size;
-			if (GroupSize != Pair.Value.Value->Num())
+			if (GroupSize != Pair.Value.Get().Num())
 			{
-				Pair.Value.Value->Resize(GroupSize);
+				Pair.Value.Resize(GroupSize);
 			}
 		}
 
@@ -778,7 +730,7 @@ void FManagedArrayCollection::Serialize(Chaos::FChaosArchive& Ar)
 			TMap<FKeyType, FManagedArrayCollectionValueTypeWrapper> ToSaveMap;
 			for (TTuple<FKeyType, FValueType>& Pair : Map)
 			{
-				if (Pair.Value.Saved)
+				if (Pair.Value.IsPersistent())
 				{
 					ToSaveMap.Emplace(Pair.Key, FManagedArrayCollectionValueTypeWrapper(&Pair.Value));
 				}
@@ -800,39 +752,157 @@ FArchive& operator<<(FArchive& Ar, FManagedArrayCollection::FGroupInfo& GroupInf
 	return Ar;
 }
 
+/**
+*   operator<<(FValueType)
+*/
 FArchive& operator<<(FArchive& Ar, FManagedArrayCollection::FValueType& ValueIn)
 {
-	int Version = 4;	//todo: version per entry is really bloated
-	Ar << Version;
+	ValueIn.Serialize(Ar);
+	return Ar;
+}
 
-	int ArrayTypeAsInt = static_cast<int>(ValueIn.ArrayType);
+void FManagedArrayCollection::FValueType::Serialize(FArchive& Ar)
+{
+	int SerializationVersion = 4;	//todo: version per entry is really bloated
+	Ar << SerializationVersion;
+
+	int ArrayTypeAsInt = static_cast<int>(ArrayType);
 	Ar << ArrayTypeAsInt;
-	ValueIn.ArrayType = static_cast<FManagedArrayCollection::EArrayType>(ArrayTypeAsInt);
+	ArrayType = static_cast<FManagedArrayCollection::EArrayType>(ArrayTypeAsInt);
 
-	if (Version < 4)
+	if (SerializationVersion < 4)
 	{
 		int ArrayScopeAsInt;
 		Ar << ArrayScopeAsInt;	//assume all serialized old content was for rest collection
 	}
 
-	if (Version >= 2)
+	if (SerializationVersion >= 2)
 	{
-		Ar << ValueIn.GroupIndexDependency;
-		Ar << ValueIn.Saved;
+		Ar << GroupIndexDependency;
+		Ar << bPersistent;
 	}
 
-	if (ValueIn.Value == nullptr)
+	if (Value == nullptr)
 	{
-		ValueIn.Value = NewManagedTypedArray(ValueIn.ArrayType);
+		Value = NewManagedTypedArray(ArrayType);
 	}
 	
 	// Note: We switched to always saving the value here, and use the Saved flag
 	// to remove the property from the overall Map (see FManagedArrayCollection::Serialize above)
 	bool bNewSavedBehavior = Ar.CustomVer(FUE5MainStreamObjectVersion::GUID) >= FUE5MainStreamObjectVersion::ManagedArrayCollectionAlwaysSerializeValue;
-	if (bNewSavedBehavior || ValueIn.Saved)
+	if (bNewSavedBehavior || bPersistent)
 	{
-		ValueIn.Value->Serialize(static_cast<Chaos::FChaosArchive&>(Ar));
+		Value->Serialize(static_cast<Chaos::FChaosArchive&>(Ar));
 	}
+}
 
-	return Ar;
+FManagedArrayCollection::FValueType::FValueType()
+	: ArrayType(EArrayType::FNoneType)
+	, GroupIndexDependency(NAME_None)
+	, bPersistent(true)
+	, bExternalValue(false)
+	, Value(nullptr) 
+{};
+
+FManagedArrayCollection::FValueType::FValueType(const FValueType& Other)
+	: ArrayType(Other.ArrayType)
+	, GroupIndexDependency(Other.GroupIndexDependency)
+	, bPersistent(Other.bPersistent)
+	, bExternalValue(false)
+	, Value(nullptr)
+{
+	if (Other.Value)
+	{
+		this->Value = NewManagedTypedArray(this->ArrayType);
+		this->Value->Resize(Other.Value->Num());
+		this->Value->Init(*Other.Value);
+	}
+};
+
+FManagedArrayCollection::FValueType::FValueType(FValueType&& Other)
+	: ArrayType(Other.ArrayType)
+	, GroupIndexDependency(Other.GroupIndexDependency)
+	, bPersistent(Other.bPersistent)
+	, bExternalValue(Other.bExternalValue)
+	, Value(Other.Value)
+{
+	if (&Other != this)
+	{
+		Other.Value = nullptr;
+	}
 }
+
+FManagedArrayCollection::FValueType::~FValueType()
+{
+	if (Value && !bExternalValue)
+	{
+		delete Value;
+	}
+}
+
+FManagedArrayBase& FManagedArrayCollection::FValueType::Modify() const
+{
+	// todo : copy on write ?  ( use  NewManagedTypedArray )
+	Value->MarkDirty();
+	return *Value;
+}
+
+void FManagedArrayCollection::FValueType::Reserve(int32 ReservedSize)
+{
+	// todo : copy on write ? 
+	Value->Reserve(ReservedSize);
+}
+
+void FManagedArrayCollection::FValueType::Resize(int32 NewSize)
+{
+	// todo : copy on write ? 
+	Value->Resize(NewSize);
+}
+
+void FManagedArrayCollection::FValueType::InitFrom(const FManagedArrayCollection::FValueType& Other)
+{
+	// todo : copy on write ? 
+	if (ArrayType == Other.GetArrayType())
+	{
+		Value->Init(Other.Get());
+	}
+	else
+	{
+		Value->Convert(Other.Get());
+	}
+}
+
+void FManagedArrayCollection::FValueType::Exchange(FValueType& Other)
+{
+	// todo : copy on write ? 
+	check(ArrayType == Other.ArrayType);
+	Value->ExchangeArrays(*Other.Value);
+}
+
+void FManagedArrayCollection::FValueType::Convert(FValueType& Other)
+{
+	// todo : copy on write ? 
+	check(ArrayType != Other.ArrayType);
+	Value->Convert(*Other.Value);
+}
+
+void FManagedArrayCollection::FValueType::CopyFrom(const FValueType& Other)
+{
+	// todo : copy on write ? 
+	check(Other.Value->Num() <= Value->Num());
+	Value->CopyRange(*Other.Value, 0, Other.Value->Num());
+}
+
+void FManagedArrayCollection::FValueType::Empty()
+{
+	// todo : copy on write ? 
+	Value->Empty();
+}
+
+void FManagedArrayCollection::FValueType::RemoveGroupIndexDependency(FName Group)
+{
+	if (GroupIndexDependency == Group)
+	{
+		GroupIndexDependency = NAME_None;
+	}
+}
\ No newline at end of file
diff --git a/Engine/Source/Runtime/Experimental/Chaos/Private/GeometryCollectionProxyData.cpp b/Engine/Source/Runtime/Experimental/Chaos/Private/GeometryCollectionProxyData.cpp
--- a/Engine/Source/Runtime/Experimental/Chaos/Private/GeometryCollectionProxyData.cpp
+++ b/Engine/Source/Runtime/Experimental/Chaos/Private/GeometryCollectionProxyData.cpp
@@ -44,6 +44,7 @@ void FTransformDynamicCollection::Construct()
 
 	// Transform Group
 	AddExternalAttribute<bool>(FTransformCollection::ParentAttribute, FTransformCollection::TransformGroup, HasParent);
+	CopyAttribute(*RestCollectionShared, FTransformCollection::ParentAttribute, FTransformCollection::TransformGroup);
 }
 
 void FTransformDynamicCollection::InitializeTransforms()
diff --git a/Engine/Source/Runtime/Experimental/Chaos/Public/GeometryCollection/ManagedArrayCollection.h b/Engine/Source/Runtime/Experimental/Chaos/Public/GeometryCollection/ManagedArrayCollection.h
--- a/Engine/Source/Runtime/Experimental/Chaos/Public/GeometryCollection/ManagedArrayCollection.h
+++ b/Engine/Source/Runtime/Experimental/Chaos/Public/GeometryCollection/ManagedArrayCollection.h
@@ -75,7 +75,7 @@ public:
 	*/
 	struct FConstructionParameters {
 
-		FConstructionParameters(FName GroupIndexDependencyIn = ""
+		FConstructionParameters(FName GroupIndexDependencyIn = NAME_None
 			, bool SavedIn = true, bool bInAllowCircularDependency = false)
 			: GroupIndexDependency(GroupIndexDependencyIn)
 			, Saved(SavedIn)
@@ -223,19 +223,19 @@ public:
 			AddGroup(Group);
 		}
 
-		FValueType Value(ManagedArrayType<T>(), ValueIn);
-		Value.Value->Resize(NumElements(Group));
-		Value.Saved = Parameters.Saved;
+		const int32 InitialSize = NumElements(Group);
+		FName GroupIndexDependency;
 		if (ensure(Parameters.bAllowCircularDependency || !IsConnected(Parameters.GroupIndexDependency, Group)))
 		{
-			Value.GroupIndexDependency = Parameters.GroupIndexDependency;
+			GroupIndexDependency = Parameters.GroupIndexDependency;
 		}
 		else
 		{
-			Value.GroupIndexDependency = "";
+			GroupIndexDependency = NAME_None;
 		}
-		Value.bExternalValue = true;
-		Map.Add(FManagedArrayCollection::MakeMapKey(Name, Group), MoveTemp(Value));
+
+		FValueType NewAttribute = FValueType::MakeExternal<T>(&ValueIn, NumElements(Group), GroupIndexDependency, Parameters.Saved);
+		Map.Add({ Name, Group }, MoveTemp(NewAttribute));
 	}
 
 	/**
@@ -291,8 +291,7 @@ public:
 		const FKeyType Key = FManagedArrayCollection::MakeMapKey(Name, Group);
 		if (FValueType* FoundValue = Map.Find(Key))
 		{
-			checkSlow(Map[Key].ArrayType == ManagedArrayType<T>());
-			return static_cast<TManagedArray<T>*>(FoundValue->Value);
+			return FoundValue->ModifyTypedPtr<T>();
 		}
 		return nullptr;
 	};
@@ -300,10 +299,9 @@ public:
 	template<typename T>
 	const TManagedArray<T>* FindAttribute(FName Name, FName Group) const
 	{
-		const FKeyType Key = FManagedArrayCollection::MakeMapKey(Name, Group);
-		if (const FValueType* FoundValue = Map.Find(Key))
+		if (const FValueType* FoundValue = Map.Find({ Name, Group }))
 		{
-			return static_cast<const TManagedArray<T>*>(FoundValue->Value);
+			return FoundValue->GetTypedPtr<T>();
 		}
 		return nullptr;
 	};
@@ -317,12 +315,11 @@ public:
 	template<typename T>
 	TManagedArray<T>* FindAttributeTyped(FName Name, FName Group)
 	{
-		const FKeyType Key = FManagedArrayCollection::MakeMapKey(Name, Group);
-		if (FValueType* FoundValue = Map.Find(Key))
+		if (FValueType* FoundValue = Map.Find({ Name, Group }))
 		{
-			if(FoundValue->ArrayType == ManagedArrayType<T>())
+			if(FoundValue->IsSameType<T>())
 			{
-				return static_cast<TManagedArray<T>*>(FoundValue->Value);
+				return FoundValue->ModifyTypedPtr<T>();
 			}
 		}
 		return nullptr;
@@ -331,12 +328,11 @@ public:
 	template<typename T>
 	const TManagedArray<T>* FindAttributeTyped(FName Name, FName Group) const
 	{
-		const FKeyType Key = FManagedArrayCollection::MakeMapKey(Name, Group);
-		if (const FValueType* FoundValue = Map.Find(Key))
+		if (const FValueType* FoundValue = Map.Find({ Name, Group }))
 		{
-			if(FoundValue->ArrayType == ManagedArrayType<T>())
+			if(FoundValue->IsSameType<T>())
 			{
-				return static_cast<const TManagedArray<T>*>(FoundValue->Value);
+				return FoundValue->GetTypedPtr<T>();
 			}
 		}
 		return nullptr;
@@ -353,10 +349,7 @@ public:
 	TManagedArray<T>& ModifyAttribute(FName Name, FName Group)
 	{
 		check(HasAttribute(Name, Group))
-		const FKeyType Key = FManagedArrayCollection::MakeMapKey(Name, Group);
-		FManagedArrayBase* ManagedArray = Map[Key].Value;
-		ManagedArray->MarkDirty();
-		return *(static_cast<TManagedArray<T>*>(ManagedArray));
+		return Map[{Name, Group}].ModifyTyped<T>();
 	}
 
 	/**
@@ -369,13 +362,11 @@ public:
 	template<typename T>
 	TManagedArray<T>* ModifyAttributeTyped(FName Name, FName Group)
 	{
-		const FKeyType Key = FManagedArrayCollection::MakeMapKey(Name, Group);
-		if (FValueType* FoundValue = Map.Find(Key))
+		if (FValueType* FoundValue = Map.Find({ Name, Group }))
 		{
-			if (FoundValue->ArrayType == ManagedArrayType<T>())
+			if (FoundValue->IsSameType<T>())
 			{
-				FoundValue->Value->MarkDirty();
-				return static_cast<TManagedArray<T>*>(FoundValue->Value);
+				return FoundValue->ModifyTypedPtr<T>();
 			}
 		}
 		return nullptr;
@@ -387,23 +378,11 @@ public:
 	* @param Group - The group that manages the attribute
 	* @return ManagedArray<T> &
 	*/
-	
-	// template<typename T>
-	// UE_DEPRECATED(5.0, "non const GetAttribute() version is now deprecated, use ModifyAttribute instead")
-	// TManagedArray<T>& GetAttribute(FName Name, FName Group)
-	// {
-	// 	check(HasAttribute(Name, Group))
-	// 	const FKeyType Key = FManagedArrayCollection::MakeMapKey(Name, Group);
-	// 	return *(static_cast<TManagedArray<T>*>(Map[Key].Value));
-	// };
-
 	template<typename T>
 	const TManagedArray<T>& GetAttribute(FName Name, FName Group) const
 	{
-		check(HasAttribute(Name, Group))
-		const FKeyType Key = FManagedArrayCollection::MakeMapKey(Name, Group);
-		checkSlow(Map[Key].ArrayType == ManagedArrayType<T>());
-		return *(static_cast<TManagedArray<T>*>(Map[Key].Value));
+		check(HasAttribute(Name, Group));
+		return Map[{Name, Group}].GetTyped<T>();
 	};
 
 	/**
@@ -621,18 +600,22 @@ private:
 			AddGroup(Group);
 		}
 
-		FValueType Value(ManagedArrayType<T>(), *(new TManagedArray<T>()));
-		Value.Value->Resize(NumElements(Group));
-		Value.Saved = Parameters.Saved;
+		static const FName EmptyName = NAME_None;
+		
+		const int32 InitialSize = NumElements(Group);
+		FName GroupIndexDependency;
 		if (ensure(Parameters.bAllowCircularDependency || !IsConnected(Parameters.GroupIndexDependency, Group)))
 		{
-			Value.GroupIndexDependency = Parameters.GroupIndexDependency;
+			GroupIndexDependency = Parameters.GroupIndexDependency;
 		}
 		else
 		{
-			Value.GroupIndexDependency = "";
+			GroupIndexDependency = EmptyName;
 		}
-		Map.Add(FManagedArrayCollection::MakeMapKey(Name, Group), MoveTemp(Value));
+
+		TUniquePtr<TManagedArray<T>> ArrayPtr(new TManagedArray<T>());
+		FValueType NewAttribute = FValueType::MakeManaged(MoveTemp(ArrayPtr), InitialSize, GroupIndexDependency, Parameters.Saved);
+		Map.Add({ Name, Group }, MoveTemp(NewAttribute));
 	}
 
 	/****
@@ -662,65 +645,131 @@ private:
 	friend struct FManagedArrayCollectionValueTypeWrapper;
 	struct FValueType
 	{
+	private:
 		EArrayType ArrayType;
 		FName GroupIndexDependency;
-		bool Saved;
+		bool bPersistent;
 		bool bExternalValue;	//External arrays have external memory management.
 
 		FManagedArrayBase* Value;
 
-		FValueType()
-			: ArrayType(EArrayType::FNoneType)
-			, GroupIndexDependency("")
-			, Saved(true)
+	private:
+		/**
+		* Create a typed attribute
+		* the array is managed ( as opposed to being external ) and will be deleted by the destructor
+		*/
+		template<typename T>
+		FValueType(TUniquePtr<TManagedArray<T>>&& ArrayPtr, int32 InitialSize, FName InGroupIndexDependency, bool bInPersistent)
+			: ArrayType(ManagedArrayType<T>())
+			, GroupIndexDependency(InGroupIndexDependency)
+			, bPersistent(bInPersistent)
 			, bExternalValue(false)
-			, Value(nullptr) {};
-
-		FValueType(EArrayType ArrayTypeIn, FManagedArrayBase& In)
-			: ArrayType(ArrayTypeIn)
-			, GroupIndexDependency("")
-			, Saved(false)
-			, bExternalValue(false)
-			, Value(&In) {};
+			, Value(ArrayPtr.Release())
+		{
+			Value->Resize(InitialSize);
+		};
 
-		FValueType(const FValueType& Other)
-			: ArrayType(Other.ArrayType)
-			, GroupIndexDependency(Other.GroupIndexDependency)
-			, Saved(Other.Saved)
-			, bExternalValue(false)
-			, Value(nullptr)
+		/**
+		* Create a external typed attribute
+		* the array is external and it is the responsability of the original owner to destroy it
+		*/
+		template<typename T>
+		FValueType(TManagedArray<T>* ExternalArrayPtr, int32 InitialSize, FName InGroupIndexDependency, bool bInPersistent)
+			: ArrayType(ManagedArrayType<T>())
+			, GroupIndexDependency(InGroupIndexDependency)
+			, bPersistent(bInPersistent)
+			, bExternalValue(true)
+			, Value(ExternalArrayPtr)
 		{
-			if (Other.Value)
-			{
-				this->Value = NewManagedTypedArray(this->ArrayType);
-				this->Value->Resize(Other.Value->Num());
-				this->Value->Init(*Other.Value);
-			}
+			Value->Resize(InitialSize);
 		};
 
-		FValueType(FValueType&& Other)
-			: ArrayType(Other.ArrayType)
-			, GroupIndexDependency(Other.GroupIndexDependency)
-			, Saved(Other.Saved)
-			, bExternalValue(Other.bExternalValue)
-			, Value(Other.Value)
+	public:
+		/**
+		* Create a typed attribute
+		* the array is managed ( as opposed to being external ) and will be deleted by the destructor
+		*/
+		template <typename T>
+		static FValueType MakeManaged(TUniquePtr<TManagedArray<T>>&& ArrayPtr, int32 InitialSize, FName InGroupIndexDependency, bool bInPersistent)
 		{
-			if (&Other != this)
-			{
-				Other.Value = nullptr;
-			}
+			return FValueType(MoveTemp(ArrayPtr), InitialSize, InGroupIndexDependency, bInPersistent);
 		}
 
-		~FValueType()
+		template <typename T>
+		static FValueType MakeExternal(TManagedArray<T>* ExternalArrayPtr, int32 InitialSize, FName InGroupIndexDependency, bool bInPersistent)
 		{
-			if (Value && !bExternalValue)
-			{
-				delete Value;
-			}
+			return FValueType(ExternalArrayPtr, InitialSize, InGroupIndexDependency, bInPersistent);
 		}
 
+		CHAOS_API FValueType();
+		CHAOS_API FValueType(const FValueType& Other);
+		CHAOS_API FValueType(FValueType&& Other);
+
+		CHAOS_API ~FValueType();
+
+		EArrayType GetArrayType() const { return ArrayType; }
+		FName GetGroupIndexDependency() const { return GroupIndexDependency; };
+		bool IsPersistent() const { return bPersistent; }
+		bool IsExternal() const { return bExternalValue; }
+		bool IsDirty() const { return Value->IsDirty(); }
+
+		void SetGroupIndexDependency(FName NewGroupDependency) { GroupIndexDependency = NewGroupDependency; }
+
+		template <typename T>
+		bool IsSameType() const { return (ArrayType == ManagedArrayType<T>()); }
+		
+		template <typename T>
+		const TManagedArray<T>* GetTypedPtr() const
+		{ 
+			check(IsSameType<T>());
+			
+			const TManagedArray<T>* TypedManagedArray = static_cast<const TManagedArray<T>*>(Value);
+			check(TypedManagedArray != nullptr);
+			return TypedManagedArray;
+		}
+
+		template <typename T>
+		const TManagedArray<T>& GetTyped() const
+		{
+			return *GetTypedPtr<T>();
+		}
+
+		const FManagedArrayBase& Get() const { return *Value; }
+
+		template <typename T>
+		TManagedArray<T>* ModifyTypedPtr() const
+		{ 
+			check(IsSameType<T>());
+			// todo : copy on write ? 
+
+			TManagedArray<T>* TypedManagedArray = static_cast<TManagedArray<T>*>(Value);
+			check(TypedManagedArray != nullptr);
+
+			TypedManagedArray->MarkDirty();
+			return TypedManagedArray;
+		}
+
+		template <typename T>
+		TManagedArray<T>& ModifyTyped() const
+		{
+			return *ModifyTypedPtr<T>();
+		}
+
+		FManagedArrayBase& Modify() const;
+
+		void Reserve(int32 ReservedSize);
+		void Resize(int32 NewSize);
+		void InitFrom(const FValueType& Other);
+		void Exchange(FValueType& Other);
+		void Convert(FValueType& Other);
+		void CopyFrom(const FValueType& Other);
+		void Empty();
+		void RemoveGroupIndexDependency(FName Group);
+
 		FValueType& operator=(const FValueType& Other) = delete;
 		FValueType& operator=(FValueType&& Other) = delete;
+
+		void Serialize(FArchive& Ar);
 	};
 
 
diff --git a/Engine/Source/Runtime/Experimental/GeometryCollectionEngine/Private/GeometryCollection/GeometryCollectionComponent.cpp b/Engine/Source/Runtime/Experimental/GeometryCollectionEngine/Private/GeometryCollection/GeometryCollectionComponent.cpp
--- a/Engine/Source/Runtime/Experimental/GeometryCollectionEngine/Private/GeometryCollection/GeometryCollectionComponent.cpp
+++ b/Engine/Source/Runtime/Experimental/GeometryCollectionEngine/Private/GeometryCollection/GeometryCollectionComponent.cpp
@@ -556,17 +556,6 @@ private:
 // Define the methods
 COPY_ON_WRITE_ATTRIBUTES
 
-TManagedArray<int32>& UGeometryCollectionComponent::GetParentArrayCopyOnWrite()
-{
-	if (!IndirectParentArray)
-	{
-		DynamicCollection->AddAttribute<int32>(FTransformCollection::ParentAttribute, FTransformCollection::TransformGroup);
-		DynamicCollection->CopyAttribute(*RestCollection->GetGeometryCollection(), FTransformCollection::ParentAttribute, FTransformCollection::TransformGroup);
-		IndirectParentArray = &DynamicCollection->ModifyAttribute<int32>(FTransformCollection::ParentAttribute, FTransformCollection::TransformGroup);
-	}
-	return *IndirectParentArray;
-}
-
 int32 UGeometryCollectionComponent::GetParent(int32 Index) const
 {
 	if (DynamicCollection)
@@ -584,7 +573,6 @@ const TManagedArray<int32>& UGeometryCollectionComponent::GetParentArrayRest() c
 UGeometryCollectionComponent::UGeometryCollectionComponent(const FObjectInitializer& ObjectInitializer)
 	: Super(ObjectInitializer)
 	, ChaosSolverActor(nullptr)
-	, IndirectParentArray(nullptr)
 	, InitializationState(ESimulationInitializationState::Unintialized)
 	, ObjectType(EObjectStateTypeEnum::Chaos_Object_Dynamic)
 	, GravityGroupIndex(0)
@@ -3670,9 +3658,6 @@ void UGeometryCollectionComponent::ResetDynamicCollection()
 	{
 		DynamicCollection = MakeUnique<FGeometryDynamicCollection>(RestCollection->GetGeometryCollection());
 
-		IndirectParentArray = nullptr;
-		GetParentArrayCopyOnWrite();
-
 		if (bStoreVelocities || bNotifyTrailing)
 		{
 			DynamicCollection->AddVelocitiesAttributes();
@@ -7332,4 +7317,4 @@ void UGeometryCollectionComponent::RebaseDynamicCollectionTransformsOnNewWorldTr
 
 	PhysicsProxy->RebaseAllGameThreadCollectionTransformsOnNewWorldTransform_External();
 	ComponentSpaceTransforms.MarkDirty();
-}
\ No newline at end of file
+}
diff --git a/Engine/Source/Runtime/Experimental/GeometryCollectionEngine/Public/GeometryCollection/GeometryCollectionComponent.h b/Engine/Source/Runtime/Experimental/GeometryCollectionEngine/Public/GeometryCollection/GeometryCollectionComponent.h
--- a/Engine/Source/Runtime/Experimental/GeometryCollectionEngine/Public/GeometryCollection/GeometryCollectionComponent.h
+++ b/Engine/Source/Runtime/Experimental/GeometryCollectionEngine/Public/GeometryCollection/GeometryCollectionComponent.h
@@ -929,13 +929,8 @@ public:
 	// Declare all the methods
 	COPY_ON_WRITE_ATTRIBUTES
 
-	GEOMETRYCOLLECTIONENGINE_API TManagedArray<int32>& GetParentArrayCopyOnWrite();
 	GEOMETRYCOLLECTIONENGINE_API int32 GetParent(int32 Index) const;
 	GEOMETRYCOLLECTIONENGINE_API const TManagedArray<int32>& GetParentArrayRest() const;
-	private:
-		TManagedArray<int32>* IndirectParentArray;
-	public:
-
 
 	UPROPERTY(EditAnywhere, NoClear, BlueprintReadOnly, Category = "ChaosPhysics")
 	TObjectPtr<const UGeometryCollection> RestCollection;
