Chaos cloth: Make Buckling Ratio weighted

diff --git a/Engine/Plugins/ChaosClothAssetEditor/Source/ChaosClothAssetDataflowNodes/Private/ChaosClothAsset/SimulationBendingConfigNode.cpp b/Engine/Plugins/ChaosClothAssetEditor/Source/ChaosClothAssetDataflowNodes/Private/ChaosClothAsset/SimulationBendingConfigNode.cpp
--- a/Engine/Plugins/ChaosClothAssetEditor/Source/ChaosClothAssetDataflowNodes/Private/ChaosClothAsset/SimulationBendingConfigNode.cpp
+++ b/Engine/Plugins/ChaosClothAssetEditor/Source/ChaosClothAssetDataflowNodes/Private/ChaosClothAsset/SimulationBendingConfigNode.cpp
@@ -17,10 +17,12 @@ FChaosClothAssetSimulationBendingConfigNode::FChaosClothAssetSimulationBendingCo
 	RegisterInputConnection(&BendingStiffnessBias.WeightMap);
 	RegisterInputConnection(&BendingDamping.WeightMap);
 	RegisterInputConnection(&BendingAnisoDamping.WeightMap);
+	RegisterInputConnection(&AnisoBucklingRatio.WeightMap);
 	RegisterInputConnection(&BucklingStiffness.WeightMap);
 	RegisterInputConnection(&BucklingStiffnessWarp.WeightMap);
 	RegisterInputConnection(&BucklingStiffnessWeft.WeightMap);
 	RegisterInputConnection(&BucklingStiffnessBias.WeightMap);
+	RegisterInputConnection(&BucklingRatioWeighted.WeightMap);
 }
 
 void FChaosClothAssetSimulationBendingConfigNode::AddProperties(FPropertyHelper& PropertyHelper) const
@@ -57,7 +59,7 @@ void FChaosClothAssetSimulationBendingConfigNode::AddProperties(FPropertyHelper&
 					FName(TEXT("XPBDAnisoBucklingStiffnessWarp")),
 					FName(TEXT("BucklingStiffness"))});
 				
-				PropertyHelper.SetProperty(FName(TEXT("XPBDBucklingRatio")), BucklingRatio, {
+				PropertyHelper.SetPropertyWeighted(FName(TEXT("XPBDBucklingRatio")), BucklingRatioWeighted, {
 					FName(TEXT("XPBDAnisoBucklingRatio")),
 					FName(TEXT("BucklingRatios"))});
 			}
@@ -128,7 +130,11 @@ void FChaosClothAssetSimulationBendingConfigNode::AddProperties(FPropertyHelper&
 				FName(TEXT("XPBDBendingSpringDamping")),  
 				FName(TEXT("XPBDBendingElementDamping"))});
 
-			PropertyHelper.SetProperty(FName(TEXT("XPBDAnisoBucklingRatio")), BucklingRatio, {
+			PropertyHelper.SetFabricPropertyWeighted(FName(TEXT("XPBDAnisoBucklingRatio")), AnisoBucklingRatio, [](
+				const UE::Chaos::ClothAsset::FCollectionClothFabricFacade& FabricFacade)-> float
+			{
+				return FabricFacade.GetBucklingRatio();
+			}, {
 				FName(TEXT("BucklingRatio")),
 				FName(TEXT("XPBDBucklingRatio"))});
 
@@ -170,7 +176,7 @@ void FChaosClothAssetSimulationBendingConfigNode::AddProperties(FPropertyHelper&
 				FName(TEXT("XPBDBucklingStiffness")),
 				FName(TEXT("XPBDAnisoBucklingStiffnessWarp"))});
 			
-			PropertyHelper.SetProperty(FName(TEXT("BucklingRatios")), BucklingRatio, {
+			PropertyHelper.SetPropertyWeighted(FName(TEXT("BucklingRatio")), BucklingRatioWeighted, {
 				FName(TEXT("XPBDBucklingRatio")),
 				FName(TEXT("XPBDAnisoBucklingRatio"))});
 		}
@@ -190,3 +196,19 @@ void FChaosClothAssetSimulationBendingConfigNode::AddProperties(FPropertyHelper&
 		});
 	}
 }
+
+void FChaosClothAssetSimulationBendingConfigNode::Serialize(FArchive& Ar)
+{
+	Super::Serialize(Ar);
+	if (Ar.IsLoading())
+	{
+#if WITH_EDITORONLY_DATA
+		if (BucklingRatio_DEPRECATED != BucklingRatioDeprecatedDefault)
+		{
+			BucklingRatioWeighted.Low = BucklingRatioWeighted.High = BucklingRatio_DEPRECATED;
+			AnisoBucklingRatio.Low = AnisoBucklingRatio.High = BucklingRatio_DEPRECATED;
+			BucklingRatio_DEPRECATED = BucklingRatioDeprecatedDefault;
+		}
+#endif
+	}
+}
diff --git a/Engine/Plugins/ChaosClothAssetEditor/Source/ChaosClothAssetDataflowNodes/Private/ChaosClothAsset/SimulationCollisionConfigNode.cpp b/Engine/Plugins/ChaosClothAssetEditor/Source/ChaosClothAssetDataflowNodes/Private/ChaosClothAsset/SimulationCollisionConfigNode.cpp
--- a/Engine/Plugins/ChaosClothAssetEditor/Source/ChaosClothAssetDataflowNodes/Private/ChaosClothAsset/SimulationCollisionConfigNode.cpp
+++ b/Engine/Plugins/ChaosClothAssetEditor/Source/ChaosClothAssetDataflowNodes/Private/ChaosClothAsset/SimulationCollisionConfigNode.cpp
@@ -35,15 +35,15 @@ void FChaosClothAssetSimulationCollisionConfigNode::Serialize(FArchive& Ar)
 	if (Ar.IsLoading())
 	{
 #if WITH_EDITORONLY_DATA
-		if (FrictionCoefficient_DEPRECATED != UE::Chaos::ClothAsset::FDefaultFabric::Friction)
+		if (FrictionCoefficient_DEPRECATED != FrictionCoefficientDeprecatedDefault)
 		{
 			FrictionCoefficientImported.ImportedValue = FrictionCoefficient_DEPRECATED;
-			FrictionCoefficient_DEPRECATED = UE::Chaos::ClothAsset::FDefaultFabric::Friction;
+			FrictionCoefficient_DEPRECATED = FrictionCoefficientDeprecatedDefault;
 		}
-		if (CollisionThickness_DEPRECATED != UE::Chaos::ClothAsset::FDefaultFabric::CollisionThickness)
+		if (CollisionThickness_DEPRECATED != CollisionThicknessDeprecatedDefault)
 		{
 			CollisionThicknessImported.ImportedValue = CollisionThickness_DEPRECATED;
-			CollisionThickness_DEPRECATED = UE::Chaos::ClothAsset::FDefaultFabric::CollisionThickness;
+			CollisionThickness_DEPRECATED = CollisionThicknessDeprecatedDefault;
 		}
 #endif
 	}
diff --git a/Engine/Plugins/ChaosClothAssetEditor/Source/ChaosClothAssetDataflowNodes/Private/ChaosClothAsset/SimulationSelfCollisionConfigNode.cpp b/Engine/Plugins/ChaosClothAssetEditor/Source/ChaosClothAssetDataflowNodes/Private/ChaosClothAsset/SimulationSelfCollisionConfigNode.cpp
--- a/Engine/Plugins/ChaosClothAssetEditor/Source/ChaosClothAssetDataflowNodes/Private/ChaosClothAsset/SimulationSelfCollisionConfigNode.cpp
+++ b/Engine/Plugins/ChaosClothAssetEditor/Source/ChaosClothAssetDataflowNodes/Private/ChaosClothAsset/SimulationSelfCollisionConfigNode.cpp
@@ -63,10 +63,10 @@ void FChaosClothAssetSimulationSelfCollisionConfigNode::Serialize(FArchive& Ar)
 	if (Ar.IsLoading())
 	{
 #if WITH_EDITORONLY_DATA
-		if (SelfCollisionFriction_DEPRECATED != UE::Chaos::ClothAsset::FDefaultFabric::SelfFriction)
+		if (SelfCollisionFriction_DEPRECATED != SelfFrictionDeprecatedValue)
 		{
 			SelfCollisionFrictionImported.ImportedValue = SelfCollisionFriction_DEPRECATED;
-			SelfCollisionFriction_DEPRECATED = UE::Chaos::ClothAsset::FDefaultFabric::SelfFriction;
+			SelfCollisionFriction_DEPRECATED = SelfFrictionDeprecatedValue;
 		}
 		if (SelfCollisionKinematicColliderFriction_DEPRECATED != FrictionDeprecatedValue)
 		{
diff --git a/Engine/Plugins/ChaosClothAssetEditor/Source/ChaosClothAssetDataflowNodes/Public/ChaosClothAsset/SimulationBendingConfigNode.h b/Engine/Plugins/ChaosClothAssetEditor/Source/ChaosClothAssetDataflowNodes/Public/ChaosClothAsset/SimulationBendingConfigNode.h
--- a/Engine/Plugins/ChaosClothAssetEditor/Source/ChaosClothAssetDataflowNodes/Public/ChaosClothAsset/SimulationBendingConfigNode.h
+++ b/Engine/Plugins/ChaosClothAssetEditor/Source/ChaosClothAssetDataflowNodes/Public/ChaosClothAsset/SimulationBendingConfigNode.h
@@ -15,6 +15,8 @@ struct FChaosClothAssetSimulationBendingConfigNode : public FChaosClothAssetSimu
 	DATAFLOW_NODE_DEFINE_INTERNAL(FChaosClothAssetSimulationBendingConfigNode, "SimulationBendingConfig", "Cloth", "Cloth Simulation Bending Config")
 
 public:
+
+
 	/**
 	 * Method for calculating the rest angles of the constraints.
 	 */
@@ -52,7 +54,7 @@ public:
     EChaosClothAssetConstraintDistributionType DistributionType = EChaosClothAssetConstraintDistributionType::Isotropic;
 	
 	/** Constraint method type */
-	UPROPERTY(EditAnywhere, Category = "Bending Types", Meta = (EditCondition = "DistributionType == EChaosClothAssetConstraintDistributionType::Isotropic", EditConditionHides))
+	UPROPERTY(EditAnywhere, Category = "Bending Types", Meta = (EditCondition = "SolverType == EChaosClothAssetConstraintSolverType::PBD || DistributionType == EChaosClothAssetConstraintDistributionType::Isotropic", EditConditionHides))
 	EChaosClothAssetBendingConstraintType ConstraintType = EChaosClothAssetBendingConstraintType::HingeAngles;
 	
 	/**
@@ -90,6 +92,13 @@ public:
 	 */
 	UPROPERTY(EditAnywhere, Category = "Bending Properties", Meta = (UIMin = "0", UIMax = "10", ClampMin = "0", ClampMax = "1000"), Meta = (EditCondition = "SolverType == EChaosClothAssetConstraintSolverType::XPBD && DistributionType == EChaosClothAssetConstraintDistributionType::Anisotropic", EditConditionHides))
 	FChaosClothAssetWeightedValue BendingAnisoDamping = { true, 1.f, 1.f, TEXT("BendingAnisoDamping"), true };
+
+	/**
+	 * Once the element has bent such that it's folded more than this ratio from its rest angle ("buckled"), switch to using Buckling Stiffness instead of BendingElement Stiffness.
+	 * When Buckling Ratio = 0, the Buckling Stiffness will never be used. When BucklingRatio = 1, the Buckling Stiffness will be used as soon as it's bent past its rest configuration.
+	 */
+	UPROPERTY(EditAnywhere, Category = "Bending Properties", DisplayName = "Aniso Buckling Ratio", Meta = (UIMin = "0", UIMax = "1", ClampMin = "0", ClampMax = "1"), Meta = (EditCondition = "SolverType == EChaosClothAssetConstraintSolverType::XPBD && DistributionType == EChaosClothAssetConstraintDistributionType::Anisotropic", EditConditionHides))
+	FChaosClothAssetWeightedValue AnisoBucklingRatio= { true, 0.5f, 0.5f, TEXT("BucklingRatio"), true };
 	
 	/**
 	 * The stiffness after buckling in the warp (vertical) direction.
@@ -100,7 +109,7 @@ public:
 	 * are interpolated with the per particle weight to make the final value used for the simulation.
 	 * Otherwise all particles are considered to have a zero weight, and only the Low value is meaningful.
 	 */
-	UPROPERTY(EditAnywhere, Category = "Bending Properties", Meta = (UIMin = "0", UIMax = "10000", ClampMin = "0", ClampMax = "10000000", EditCondition = "BucklingRatio != 0 && DistributionType == EChaosClothAssetConstraintDistributionType::Anisotropic  && SolverType == EChaosClothAssetConstraintSolverType::XPBD", EditConditionHides))
+	UPROPERTY(EditAnywhere, Category = "Bending Properties", Meta = (UIMin = "0", UIMax = "10000", ClampMin = "0", ClampMax = "10000000", EditCondition = "DistributionType == EChaosClothAssetConstraintDistributionType::Anisotropic  && SolverType == EChaosClothAssetConstraintSolverType::XPBD", EditConditionHides))
 	FChaosClothAssetWeightedValue BucklingStiffnessWarp = { true, 50.f, 50.f, TEXT("BucklingStiffnessWarp"), true };
 
 	/**
@@ -112,7 +121,7 @@ public:
 	 * are interpolated with the per particle weight to make the final value used for the simulation.
 	 * Otherwise all particles are considered to have a zero weight, and only the Low value is meaningful.
 	 */
-	UPROPERTY(EditAnywhere, Category = "Bending Properties", Meta = (UIMin = "0", UIMax = "10000", ClampMin = "0", ClampMax = "10000000", EditCondition = "BucklingRatio != 0 && DistributionType == EChaosClothAssetConstraintDistributionType::Anisotropic && SolverType == EChaosClothAssetConstraintSolverType::XPBD", EditConditionHides))
+	UPROPERTY(EditAnywhere, Category = "Bending Properties", Meta = (UIMin = "0", UIMax = "10000", ClampMin = "0", ClampMax = "10000000", EditCondition = "DistributionType == EChaosClothAssetConstraintDistributionType::Anisotropic && SolverType == EChaosClothAssetConstraintSolverType::XPBD", EditConditionHides))
 	FChaosClothAssetWeightedValue BucklingStiffnessWeft = { true, 50.f, 50.f, TEXT("BucklingStiffnessWeft"),  true };
 
 	/**
@@ -124,7 +133,7 @@ public:
 	 * are interpolated with the per particle weight to make the final value used for the simulation.
 	 * Otherwise all particles are considered to have a zero weight, and only the Low value is meaningful.
 	 */
-	UPROPERTY(EditAnywhere, Category = "Bending Properties", Meta = (UIMin = "0", UIMax = "10000", ClampMin = "0", ClampMax = "10000000", EditCondition = "BucklingRatio != 0 && DistributionType == EChaosClothAssetConstraintDistributionType::Anisotropic && SolverType == EChaosClothAssetConstraintSolverType::XPBD", EditConditionHides))
+	UPROPERTY(EditAnywhere, Category = "Bending Properties", Meta = (UIMin = "0", UIMax = "10000", ClampMin = "0", ClampMax = "10000000", EditCondition = "DistributionType == EChaosClothAssetConstraintDistributionType::Anisotropic && SolverType == EChaosClothAssetConstraintSolverType::XPBD", EditConditionHides))
 	FChaosClothAssetWeightedValue BucklingStiffnessBias = { true, 50.f, 50.f, TEXT("BucklingStiffnessBias"), true };
 
 	/**
@@ -154,18 +163,25 @@ public:
 	 * are interpolated with the per particle weight to make the final value used for the simulation.
 	 * Otherwise all particles are considered to have a zero weight, and only the Low value is meaningful.
 	 */
-	UPROPERTY(EditAnywhere, Category = "Bending Properties", Meta = (UIMin = "0", UIMax = "1", ClampMin = "0", ClampMax = "1", EditCondition = "BucklingRatio != 0 && ((DistributionType == EChaosClothAssetConstraintDistributionType::Isotropic && SolverType == EChaosClothAssetConstraintSolverType::XPBD) || SolverType == EChaosClothAssetConstraintSolverType::PBD) && ConstraintType == EChaosClothAssetBendingConstraintType::HingeAngles", EditConditionHides))
+	UPROPERTY(EditAnywhere, Category = "Bending Properties", Meta = (UIMin = "0", UIMax = "1", ClampMin = "0", ClampMax = "1", EditCondition = "((DistributionType == EChaosClothAssetConstraintDistributionType::Isotropic && SolverType == EChaosClothAssetConstraintSolverType::XPBD) || SolverType == EChaosClothAssetConstraintSolverType::PBD) && ConstraintType == EChaosClothAssetBendingConstraintType::HingeAngles", EditConditionHides))
 	FChaosClothAssetWeightedValue BucklingStiffness = { true, 0.9f, 0.9f, TEXT("BucklingStiffness") };
 
 	/**
 	 * Once the element has bent such that it's folded more than this ratio from its rest angle ("buckled"), switch to using Buckling Stiffness instead of BendingElement Stiffness.
-	 * When Buckling Ratio = 0, the Buckling Stiffness will never be used. When BucklingRatio = 1, the Buckling Stiffness will be used as soon as its bent past its rest configuration.
+	 * When Buckling Ratio = 0, the Buckling Stiffness will never be used. When BucklingRatio = 1, the Buckling Stiffness will be used as soon as it's bent past its rest configuration.
 	 */
-	UPROPERTY(EditAnywhere, Category = "Bending Properties", Meta = (UIMin = "0", UIMax = "1", ClampMin = "0", ClampMax = "1"), Meta = (EditCondition = "ConstraintType == EChaosClothAssetBendingConstraintType::HingeAngles", EditConditionHides))
-	float BucklingRatio = 0.5f;
+	UPROPERTY(EditAnywhere, Category = "Bending Properties", DisplayName = "Buckling Ratio", Meta = (UIMin = "0", UIMax = "1", ClampMin = "0", ClampMax = "1"), Meta = (EditCondition = "((DistributionType == EChaosClothAssetConstraintDistributionType::Isotropic && SolverType == EChaosClothAssetConstraintSolverType::XPBD) || SolverType == EChaosClothAssetConstraintSolverType::PBD) && ConstraintType == EChaosClothAssetBendingConstraintType::HingeAngles", EditConditionHides))
+	FChaosClothAssetWeightedValue BucklingRatioWeighted = { true, 0.5f, 0.5f, TEXT("BucklingRatio") };
 
 	FChaosClothAssetSimulationBendingConfigNode(const Dataflow::FNodeParameters& InParam, FGuid InGuid = FGuid::NewGuid());
 
+	virtual void Serialize(FArchive& Ar) override;
 private:
+	// Deprecated properties
+#if WITH_EDITORONLY_DATA
+	static constexpr float BucklingRatioDeprecatedDefault = 0.5f;
+	UPROPERTY()
+	float BucklingRatio_DEPRECATED = BucklingRatioDeprecatedDefault;
+#endif
 	virtual void AddProperties(FPropertyHelper& PropertyHelper) const override;
 };
diff --git a/Engine/Plugins/ChaosClothAssetEditor/Source/ChaosClothAssetDataflowNodes/Public/ChaosClothAsset/SimulationCollisionConfigNode.h b/Engine/Plugins/ChaosClothAssetEditor/Source/ChaosClothAssetDataflowNodes/Public/ChaosClothAsset/SimulationCollisionConfigNode.h
--- a/Engine/Plugins/ChaosClothAssetEditor/Source/ChaosClothAssetDataflowNodes/Public/ChaosClothAsset/SimulationCollisionConfigNode.h
+++ b/Engine/Plugins/ChaosClothAssetEditor/Source/ChaosClothAssetDataflowNodes/Public/ChaosClothAsset/SimulationCollisionConfigNode.h
@@ -41,11 +41,12 @@ private:
 
 	// Deprecated properties
 #if WITH_EDITORONLY_DATA
-
+	static constexpr float CollisionThicknessDeprecatedDefault = 1.0f;
 	UPROPERTY()
-	float CollisionThickness_DEPRECATED = UE::Chaos::ClothAsset::FDefaultFabric::CollisionThickness;
+	float CollisionThickness_DEPRECATED = CollisionThicknessDeprecatedDefault;
 
+	static constexpr float FrictionCoefficientDeprecatedDefault = 0.8f;
 	UPROPERTY()
-	float FrictionCoefficient_DEPRECATED  = UE::Chaos::ClothAsset::FDefaultFabric::Friction;
+	float FrictionCoefficient_DEPRECATED  = FrictionCoefficientDeprecatedDefault;
 #endif
 };
diff --git a/Engine/Plugins/ChaosClothAssetEditor/Source/ChaosClothAssetDataflowNodes/Public/ChaosClothAsset/SimulationSelfCollisionConfigNode.h b/Engine/Plugins/ChaosClothAssetEditor/Source/ChaosClothAssetDataflowNodes/Public/ChaosClothAsset/SimulationSelfCollisionConfigNode.h
--- a/Engine/Plugins/ChaosClothAssetEditor/Source/ChaosClothAssetDataflowNodes/Public/ChaosClothAsset/SimulationSelfCollisionConfigNode.h
+++ b/Engine/Plugins/ChaosClothAssetEditor/Source/ChaosClothAssetDataflowNodes/Public/ChaosClothAsset/SimulationSelfCollisionConfigNode.h
@@ -100,8 +100,9 @@ private:
 	UPROPERTY()
 	float SelfCollisionKinematicColliderFriction_DEPRECATED = FrictionDeprecatedValue;
 
+	static constexpr float SelfFrictionDeprecatedValue = 0.0f;
 	UPROPERTY()
-	float SelfCollisionFriction_DEPRECATED = UE::Chaos::ClothAsset::FDefaultFabric::SelfFriction;
+	float SelfCollisionFriction_DEPRECATED = SelfFrictionDeprecatedValue;
 
 	static constexpr float SelfCollisionThicknessDeprecatedValue = -1.f;
 	UPROPERTY()
diff --git a/Engine/Source/Runtime/Experimental/Chaos/Private/Chaos/PBDBendingConstraints.cpp b/Engine/Source/Runtime/Experimental/Chaos/Private/Chaos/PBDBendingConstraints.cpp
--- a/Engine/Source/Runtime/Experimental/Chaos/Private/Chaos/PBDBendingConstraints.cpp
+++ b/Engine/Source/Runtime/Experimental/Chaos/Private/Chaos/PBDBendingConstraints.cpp
@@ -62,6 +62,7 @@ void FPBDBendingConstraints::InitColor(const SolverParticlesOrRange& InParticles
 		ConstraintSharedEdges = MoveTemp(ReorderedConstraintSharedEdges);
 		RestAngles = MoveTemp(ReorderedRestAngles);
 		Stiffness.ReorderIndices(OrigToReorderedIndices);
+		BucklingRatioWeighted.ReorderIndices(OrigToReorderedIndices);
 		BucklingStiffness.ReorderIndices(OrigToReorderedIndices);
 	}
 }
@@ -92,7 +93,21 @@ void FPBDBendingConstraints::SetProperties(
 	}
 	if (IsBucklingRatioMutable(PropertyCollection))
 	{
-		BucklingRatio = (FSolverReal)FMath::Clamp(GetBucklingRatio(PropertyCollection), 0.f, 1.);
+		const FSolverVec2 WeightedValue = FSolverVec2(GetWeightedFloatBucklingRatio(PropertyCollection)).ClampAxes(0.f, 1.f);
+		if (IsBucklingRatioStringDirty(PropertyCollection))
+		{
+			const FString& WeightMapName = GetBucklingRatioString(PropertyCollection);
+			BucklingRatioWeighted = FPBDWeightMap(
+				WeightedValue,
+				WeightMaps.FindRef(WeightMapName),
+				TConstArrayView<TVec2<int32>>(ConstraintSharedEdges),
+				ParticleOffset,
+				ParticleCount);
+		}
+		else
+		{
+			BucklingRatioWeighted.SetWeightedValue(WeightedValue);
+		}
 	}
 	if (IsBucklingStiffnessMutable(PropertyCollection))
 	{
diff --git a/Engine/Source/Runtime/Experimental/Chaos/Private/Chaos/PBDBendingConstraints.ispc b/Engine/Source/Runtime/Experimental/Chaos/Private/Chaos/PBDBendingConstraints.ispc
--- a/Engine/Source/Runtime/Experimental/Chaos/Private/Chaos/PBDBendingConstraints.ispc
+++ b/Engine/Source/Runtime/Experimental/Chaos/Private/Chaos/PBDBendingConstraints.ispc
@@ -95,6 +95,39 @@ export void InitBendingConstraintsIsBuckled(const uniform FVector3f X[],
 	}
 }
 
+export void InitBendingConstraintsIsBuckledWithMaps(const uniform FVector3f X[],
+									const uniform FIntVector4 Constraints[],
+									const uniform float RestAngles[],
+									uniform bool IsBuckled[],
+									const uniform uint8 BucklingRatioIndices[],
+									const uniform float BucklingRatioTable[],
+									const uniform int32 NumConstraints)
+{
+	foreach(i = 0 ... NumConstraints)
+	{
+		const varying FIntVector4 Constraint = VectorLoad(&Constraints[extract(i,0)]);
+
+		const varying int32 Index1 = Constraint.V[0];
+		const varying int32 Index2 = Constraint.V[1];
+		const varying int32 Index3 = Constraint.V[2];
+		const varying int32 Index4 = Constraint.V[3];
+		const varying int8 BucklingRatioIndex = BucklingRatioIndices[i];
+
+		const varying FVector3f X1 = VectorGather(&X[Index1]);
+		const varying FVector3f X2 = VectorGather(&X[Index2]);
+		const varying FVector3f X3 = VectorGather(&X[Index3]);
+		const varying FVector3f X4 = VectorGather(&X[Index4]);
+		
+		#pragma ignore warning(perf)
+		const varying float BucklingRatio = BucklingRatioTable[BucklingRatioIndex];
+		const varying float RestAngle = RestAngles[i];
+		const varying float BucklingAngle = BucklingRatio * (FLOAT_PI - abs(RestAngle));
+
+		const varying float Angle = CalcAngle(X1, X2, X3, X4);
+		IsBuckled[i] = FLOAT_PI - abs(Angle) < BucklingAngle;
+	}
+}
+
 export void ApplyBendingConstraints(uniform FVector4f PandInvM[],
 									const uniform FIntVector4 Constraints[],
 									const uniform float RestAngles[],
diff --git a/Engine/Source/Runtime/Experimental/Chaos/Private/Chaos/PBDBendingConstraintsBase.cpp b/Engine/Source/Runtime/Experimental/Chaos/Private/Chaos/PBDBendingConstraintsBase.cpp
--- a/Engine/Source/Runtime/Experimental/Chaos/Private/Chaos/PBDBendingConstraintsBase.cpp
+++ b/Engine/Source/Runtime/Experimental/Chaos/Private/Chaos/PBDBendingConstraintsBase.cpp
@@ -156,14 +156,29 @@ namespace Chaos::Softs
 #if INTEL_ISPC
 		if (bRealTypeCompatibleWithISPC && bChaos_Bending_ISPC_Enabled)
 		{
-			ispc::InitBendingConstraintsIsBuckled(
-				(const ispc::FVector3f*)InParticles.XArray().GetData(),
-				(const ispc::FIntVector4*)Constraints.GetData(),
-				RestAngles.GetData(),
-				IsBuckled.GetData(),
-				BucklingRatio,
-				Constraints.Num()
-			);
+			if (BucklingRatioWeighted.HasWeightMap())
+			{
+				ispc::InitBendingConstraintsIsBuckledWithMaps(
+					(const ispc::FVector3f*)InParticles.XArray().GetData(),
+					(const ispc::FIntVector4*)Constraints.GetData(),
+					RestAngles.GetData(),
+					IsBuckled.GetData(),
+					BucklingRatioWeighted.GetIndices().GetData(),
+					BucklingRatioWeighted.GetTable().GetData(),
+					Constraints.Num()
+				);
+			}
+			else
+			{
+				ispc::InitBendingConstraintsIsBuckled(
+					(const ispc::FVector3f*)InParticles.XArray().GetData(),
+					(const ispc::FIntVector4*)Constraints.GetData(),
+					RestAngles.GetData(),
+					IsBuckled.GetData(),
+					(FSolverReal)BucklingRatioWeighted,
+					Constraints.Num()
+				);
+			}
 		}
 		else
 #endif
@@ -180,7 +195,7 @@ namespace Chaos::Softs
 				const FSolverVec3& P3 = InParticles.X(Index3);
 				const FSolverVec3& P4 = InParticles.X(Index4);
 				const FSolverReal Angle = CalcAngle(P1, P2, P3, P4);
-				IsBuckled[ConstraintIndex] = AngleIsBuckled(Angle, RestAngles[ConstraintIndex]);
+				IsBuckled[ConstraintIndex] = AngleIsBuckled(Angle, ConstraintIndex);
 			}
 		}
 	}
diff --git a/Engine/Source/Runtime/Experimental/Chaos/Private/Chaos/XPBDAnisotropicBendingConstraints.cpp b/Engine/Source/Runtime/Experimental/Chaos/Private/Chaos/XPBDAnisotropicBendingConstraints.cpp
--- a/Engine/Source/Runtime/Experimental/Chaos/Private/Chaos/XPBDAnisotropicBendingConstraints.cpp
+++ b/Engine/Source/Runtime/Experimental/Chaos/Private/Chaos/XPBDAnisotropicBendingConstraints.cpp
@@ -38,10 +38,11 @@ FXPBDAnisotropicBendingConstraints::FXPBDAnisotropicBendingConstraints(const FSo
 		InParticles,
 		TriangleMesh.GetUniqueAdjacentElements(),
 		TConstArrayView<FRealSingle>(), // We don't use base stiffness weight maps
+		WeightMaps.FindRef(GetXPBDAnisoBucklingRatioString(PropertyCollection, XPBDAnisoBucklingRatioName.ToString())),
 		TConstArrayView<FRealSingle>(), // We don't use base stiffness weight maps
 		GetRestAngleMapFromCollection(WeightMaps, PropertyCollection),
 		FSolverVec2(GetWeightedFloatXPBDAnisoBendingStiffnessWarp(PropertyCollection, MaxStiffness)),
-		(FSolverReal)GetXPBDAnisoBucklingRatio(PropertyCollection, 0.f),
+		FSolverVec2(GetWeightedFloatXPBDAnisoBucklingRatio(PropertyCollection, 0.f)),
 		FSolverVec2(GetWeightedFloatXPBDAnisoBucklingStiffnessWarp(PropertyCollection, MaxStiffness)),
 		GetRestAngleValueFromCollection(PropertyCollection),
 		(ERestAngleConstructionType)GetXPBDAnisoRestAngleType(PropertyCollection, (int32)ERestAngleConstructionType::Use3DRestAngles),
@@ -119,10 +120,11 @@ FXPBDAnisotropicBendingConstraints::FXPBDAnisotropicBendingConstraints(const FSo
 		InParticleCount,
 		TriangleMesh.GetUniqueAdjacentElements(),
 		TConstArrayView<FRealSingle>(), // We don't use base stiffness weight maps
+		WeightMaps.FindRef(GetXPBDAnisoBucklingRatioString(PropertyCollection, XPBDAnisoBucklingRatioName.ToString())),
 		TConstArrayView<FRealSingle>(), // We don't use base stiffness weight maps
 		GetRestAngleMapFromCollection(WeightMaps, PropertyCollection),
 		FSolverVec2(GetWeightedFloatXPBDAnisoBendingStiffnessWarp(PropertyCollection, MaxStiffness)),
-		(FSolverReal)GetXPBDAnisoBucklingRatio(PropertyCollection, 0.f),
+		FSolverVec2(GetWeightedFloatXPBDAnisoBucklingRatio(PropertyCollection, 0.f)),
 		FSolverVec2(GetWeightedFloatXPBDAnisoBucklingStiffnessWarp(PropertyCollection, MaxStiffness)),
 		GetRestAngleValueFromCollection(PropertyCollection),
 		(ERestAngleConstructionType)GetXPBDAnisoRestAngleType(PropertyCollection, (int32)ERestAngleConstructionType::Use3DRestAngles),
@@ -196,28 +198,36 @@ FXPBDAnisotropicBendingConstraints::FXPBDAnisotropicBendingConstraints(const FSo
 	const TConstArrayView<FRealSingle>& StiffnessWarpMultipliers,
 	const TConstArrayView<FRealSingle>& StiffnessWeftMultipliers,
 	const TConstArrayView<FRealSingle>& StiffnessBiasMultipliers,
+	const TConstArrayView<FRealSingle>& BucklingRatioMultipliers,
 	const TConstArrayView<FRealSingle>& BucklingStiffnessWarpMultipliers,
 	const TConstArrayView<FRealSingle>& BucklingStiffnessWeftMultipliers,
 	const TConstArrayView<FRealSingle>& BucklingStiffnessBiasMultipliers,
 	const TConstArrayView<FRealSingle>& DampingMultipliers,
+	const TConstArrayView<FRealSingle>& RestAngleMap,
 	const FSolverVec2& InStiffnessWarp,
 	const FSolverVec2& InStiffnessWeft,
 	const FSolverVec2& InStiffnessBias,
-	const FSolverReal InBucklingRatio,
+	const FSolverVec2& InBucklingRatio,
 	const FSolverVec2& InBucklingStiffnessWarp,
 	const FSolverVec2& InBucklingStiffnessWeft,
 	const FSolverVec2& InBucklingStiffnessBias,
-	const FSolverVec2& InDampingRatio)
+	const FSolverVec2& InDampingRatio,
+	const FSolverVec2& RestAngleValue,
+	ERestAngleConstructionType RestAngleConstructionType)
 	: Base(
 		InParticles,
 		ParticleOffset,
 		ParticleCount,
 		TriangleMesh.GetUniqueAdjacentElements(),
 		TConstArrayView<FRealSingle>(), // We don't use base stiffness weight maps
+		BucklingRatioMultipliers,
 		TConstArrayView<FRealSingle>(), // We don't use base stiffness weight maps
+		RestAngleMap,
 		InStiffnessWarp,
 		InBucklingRatio,
 		InBucklingStiffnessWarp,
+		RestAngleValue,
+		RestAngleConstructionType,
 		true /*bTrimKinematicConstraints*/,
 		MaxStiffness)
 	, StiffnessWarp(
@@ -384,6 +394,7 @@ void FXPBDAnisotropicBendingConstraints::InitColor(const SolverParticlesOrRange&
 		StiffnessWarp.ReorderIndices(OrigToReorderedIndices);
 		StiffnessWeft.ReorderIndices(OrigToReorderedIndices);
 		StiffnessBias.ReorderIndices(OrigToReorderedIndices);
+		BucklingRatioWeighted.ReorderIndices(OrigToReorderedIndices);
 		BucklingStiffnessWarp.ReorderIndices(OrigToReorderedIndices);
 		BucklingStiffnessWeft.ReorderIndices(OrigToReorderedIndices);
 		BucklingStiffnessBias.ReorderIndices(OrigToReorderedIndices);
@@ -465,7 +476,21 @@ void FXPBDAnisotropicBendingConstraints::SetProperties(
 	}
 	if (IsXPBDAnisoBucklingRatioMutable(PropertyCollection))
 	{
-		BucklingRatio = FMath::Clamp(GetXPBDAnisoBucklingRatio(PropertyCollection), (FSolverReal)0., (FSolverReal)1.);
+		const FSolverVec2 WeightedValue = FSolverVec2(GetWeightedFloatXPBDAnisoBucklingRatio(PropertyCollection)).ClampAxes(0.f, 1.f);
+		if (IsXPBDAnisoBucklingRatioStringDirty(PropertyCollection))
+		{
+			const FString& WeightMapName = GetXPBDAnisoBucklingRatioString(PropertyCollection);
+			BucklingRatioWeighted = FPBDWeightMap(
+				WeightedValue,
+				WeightMaps.FindRef(WeightMapName),
+				TConstArrayView<TVec2<int32>>(ConstraintSharedEdges),
+				ParticleOffset,
+				ParticleCount);
+		}
+		else
+		{
+			BucklingRatioWeighted.SetWeightedValue(WeightedValue);
+		}
 	}
 	if (IsXPBDAnisoBucklingStiffnessWarpMutable(PropertyCollection))
 	{
@@ -626,21 +651,43 @@ void FXPBDAnisotropicBendingConstraints::Init(const SolverParticlesOrRange& Part
 	X4Array.SetNumUninitialized(Constraints.Num());
 	if (bRealTypeCompatibleWithISPC && bChaos_Bending_ISPC_Enabled && ConstraintsIndex1.Num() == Constraints.Num())
 	{
-		ispc::InitXPBDBendingConstraintsIsBuckled(
-			(const ispc::FVector3f*)Particles.XArray().GetData(),
-			ConstraintsIndex1.GetData(),
-			ConstraintsIndex2.GetData(),
-			ConstraintsIndex3.GetData(),
-			ConstraintsIndex4.GetData(),
-			RestAngles.GetData(),
-			IsBuckled.GetData(),
-			(ispc::FVector3f*)X1Array.GetData(),
-			(ispc::FVector3f*)X2Array.GetData(),
-			(ispc::FVector3f*)X3Array.GetData(),
-			(ispc::FVector3f*)X4Array.GetData(),
-			BucklingRatio,
-			Constraints.Num()
-		);
+		if (BucklingRatioWeighted.HasWeightMap())
+		{
+			ispc::InitXPBDBendingConstraintsIsBuckledWithMaps(
+				(const ispc::FVector3f*)Particles.XArray().GetData(),
+				ConstraintsIndex1.GetData(),
+				ConstraintsIndex2.GetData(),
+				ConstraintsIndex3.GetData(),
+				ConstraintsIndex4.GetData(),
+				RestAngles.GetData(),
+				IsBuckled.GetData(),
+				(ispc::FVector3f*)X1Array.GetData(),
+				(ispc::FVector3f*)X2Array.GetData(),
+				(ispc::FVector3f*)X3Array.GetData(),
+				(ispc::FVector3f*)X4Array.GetData(),
+				BucklingRatioWeighted.GetIndices().GetData(),
+				BucklingRatioWeighted.GetTable().GetData(),
+				Constraints.Num()
+			);
+		}
+		else
+		{
+			ispc::InitXPBDBendingConstraintsIsBuckled(
+				(const ispc::FVector3f*)Particles.XArray().GetData(),
+				ConstraintsIndex1.GetData(),
+				ConstraintsIndex2.GetData(),
+				ConstraintsIndex3.GetData(),
+				ConstraintsIndex4.GetData(),
+				RestAngles.GetData(),
+				IsBuckled.GetData(),
+				(ispc::FVector3f*)X1Array.GetData(),
+				(ispc::FVector3f*)X2Array.GetData(),
+				(ispc::FVector3f*)X3Array.GetData(),
+				(ispc::FVector3f*)X4Array.GetData(),
+				(FSolverReal)BucklingRatioWeighted,
+				Constraints.Num()
+			);
+		}
 	}
 	else
 #endif
diff --git a/Engine/Source/Runtime/Experimental/Chaos/Private/Chaos/XPBDBendingConstraints.cpp b/Engine/Source/Runtime/Experimental/Chaos/Private/Chaos/XPBDBendingConstraints.cpp
--- a/Engine/Source/Runtime/Experimental/Chaos/Private/Chaos/XPBDBendingConstraints.cpp
+++ b/Engine/Source/Runtime/Experimental/Chaos/Private/Chaos/XPBDBendingConstraints.cpp
@@ -71,6 +71,7 @@ void FXPBDBendingConstraints::InitColor(const SolverParticlesOrRange& InParticle
 		ConstraintSharedEdges = MoveTemp(ReorderedConstraintSharedEdges);
 		RestAngles = MoveTemp(ReorderedRestAngles);
 		XPBDStiffness.ReorderIndices(OrigToReorderedIndices);
+		BucklingRatioWeighted.ReorderIndices(OrigToReorderedIndices);
 		XPBDBucklingStiffness.ReorderIndices(OrigToReorderedIndices);
 
 #if INTEL_ISPC
@@ -107,21 +108,43 @@ void FXPBDBendingConstraints::Init(const SolverParticlesOrRange& InParticles)
 	X4Array.SetNumUninitialized(Constraints.Num());
 	if (bRealTypeCompatibleWithISPC && bChaos_Bending_ISPC_Enabled && ConstraintsIndex1.Num() == Constraints.Num())
 	{
-		ispc::InitXPBDBendingConstraintsIsBuckled(
-			(const ispc::FVector3f*)InParticles.XArray().GetData(),
-			ConstraintsIndex1.GetData(),
-			ConstraintsIndex2.GetData(),
-			ConstraintsIndex3.GetData(),
-			ConstraintsIndex4.GetData(),
-			RestAngles.GetData(),
-			IsBuckled.GetData(),
-			(ispc::FVector3f*)X1Array.GetData(),
-			(ispc::FVector3f*)X2Array.GetData(),
-			(ispc::FVector3f*)X3Array.GetData(),
-			(ispc::FVector3f*)X4Array.GetData(),
-			BucklingRatio,
-			Constraints.Num()
-		);
+		if (BucklingRatioWeighted.HasWeightMap())
+		{
+			ispc::InitXPBDBendingConstraintsIsBuckledWithMaps(
+				(const ispc::FVector3f*)InParticles.XArray().GetData(),
+				ConstraintsIndex1.GetData(),
+				ConstraintsIndex2.GetData(),
+				ConstraintsIndex3.GetData(),
+				ConstraintsIndex4.GetData(),
+				RestAngles.GetData(),
+				IsBuckled.GetData(),
+				(ispc::FVector3f*)X1Array.GetData(),
+				(ispc::FVector3f*)X2Array.GetData(),
+				(ispc::FVector3f*)X3Array.GetData(),
+				(ispc::FVector3f*)X4Array.GetData(),
+				BucklingRatioWeighted.GetIndices().GetData(),
+				BucklingRatioWeighted.GetTable().GetData(),
+				Constraints.Num()
+			);
+		}
+		else
+		{
+			ispc::InitXPBDBendingConstraintsIsBuckled(
+				(const ispc::FVector3f*)InParticles.XArray().GetData(),
+				ConstraintsIndex1.GetData(),
+				ConstraintsIndex2.GetData(),
+				ConstraintsIndex3.GetData(),
+				ConstraintsIndex4.GetData(),
+				RestAngles.GetData(),
+				IsBuckled.GetData(),
+				(ispc::FVector3f*)X1Array.GetData(),
+				(ispc::FVector3f*)X2Array.GetData(),
+				(ispc::FVector3f*)X3Array.GetData(),
+				(ispc::FVector3f*)X4Array.GetData(),
+				(FSolverReal)BucklingRatioWeighted,
+				Constraints.Num()
+			);
+		}
 	}
 	else
 #endif
@@ -156,7 +179,21 @@ void FXPBDBendingConstraints::SetProperties(
 	}
 	if (IsXPBDBucklingRatioMutable(PropertyCollection))
 	{
-		BucklingRatio = FMath::Clamp(GetXPBDBucklingRatio(PropertyCollection), (FSolverReal)0., (FSolverReal)1.);
+		const FSolverVec2 WeightedValue = FSolverVec2(GetWeightedFloatXPBDBucklingRatio(PropertyCollection)).ClampAxes(0.f, 1.f);
+		if (IsXPBDBucklingRatioStringDirty(PropertyCollection))
+		{
+			const FString& WeightMapName = GetXPBDBucklingRatioString(PropertyCollection);
+			BucklingRatioWeighted = FPBDWeightMap(
+				WeightedValue,
+				WeightMaps.FindRef(WeightMapName),
+				TConstArrayView<TVec2<int32>>(ConstraintSharedEdges),
+				ParticleOffset,
+				ParticleCount);
+		}
+		else
+		{
+			BucklingRatioWeighted.SetWeightedValue(WeightedValue);
+		}
 	}
 	if (IsXPBDBucklingStiffnessMutable(PropertyCollection))
 	{
diff --git a/Engine/Source/Runtime/Experimental/Chaos/Private/Chaos/XPBDBendingConstraints.ispc b/Engine/Source/Runtime/Experimental/Chaos/Private/Chaos/XPBDBendingConstraints.ispc
--- a/Engine/Source/Runtime/Experimental/Chaos/Private/Chaos/XPBDBendingConstraints.ispc
+++ b/Engine/Source/Runtime/Experimental/Chaos/Private/Chaos/XPBDBendingConstraints.ispc
@@ -62,6 +62,49 @@ export void InitXPBDBendingConstraintsIsBuckled(const uniform FVector3f X[],
 	}
 }
 
+export void InitXPBDBendingConstraintsIsBuckledWithMaps(const uniform FVector3f X[],
+							const uniform int32 ConstraintIndex1[],
+							const uniform int32 ConstraintIndex2[],
+							const uniform int32 ConstraintIndex3[],
+							const uniform int32 ConstraintIndex4[],
+							const uniform float RestAngles[],
+							uniform bool IsBuckled[],
+							uniform FVector3f X1Out[],
+							uniform FVector3f X2Out[],
+							uniform FVector3f X3Out[],
+							uniform FVector3f X4Out[],
+							const uniform uint8 BucklingRatioIndices[],
+							const uniform float BucklingRatioTable[],
+							const uniform int32 NumConstraints)
+{
+	foreach(i = 0 ... NumConstraints)
+	{
+		const varying int32 Index1 = ConstraintIndex1[i];
+		const varying int32 Index2 = ConstraintIndex2[i];
+		const varying int32 Index3 = ConstraintIndex3[i];
+		const varying int32 Index4 = ConstraintIndex4[i];
+		const varying int8 BucklingRatioIndex = BucklingRatioIndices[i];
+
+		const varying FVector3f X1 = VectorGather(&X[Index1]);
+		const varying FVector3f X2 = VectorGather(&X[Index2]);
+		const varying FVector3f X3 = VectorGather(&X[Index3]);
+		const varying FVector3f X4 = VectorGather(&X[Index4]);
+		
+		#pragma ignore warning(perf)
+		const varying float BucklingRatio = BucklingRatioTable[BucklingRatioIndex];
+		const varying float RestAngle = RestAngles[i];
+		const varying float BucklingAngle = BucklingRatio * (FLOAT_PI - abs(RestAngle));
+
+		const varying float Angle = CalcAngle(X1, X2, X3, X4);
+		IsBuckled[i] = FLOAT_PI - abs(Angle) < BucklingAngle;
+
+		VectorStore(&X1Out[extract(i,0)], X1);
+		VectorStore(&X2Out[extract(i,0)], X2);
+		VectorStore(&X3Out[extract(i,0)], X3);
+		VectorStore(&X4Out[extract(i,0)], X4);
+	}
+}
+
 export void ApplyXPBDBendingConstraints(uniform FVector4f PandInvM[],
 									const uniform int32 ConstraintIndex1[],
 									const uniform int32 ConstraintIndex2[],
diff --git a/Engine/Source/Runtime/Experimental/Chaos/Public/Chaos/PBDBendingConstraints.h b/Engine/Source/Runtime/Experimental/Chaos/Public/Chaos/PBDBendingConstraints.h
--- a/Engine/Source/Runtime/Experimental/Chaos/Public/Chaos/PBDBendingConstraints.h
+++ b/Engine/Source/Runtime/Experimental/Chaos/Public/Chaos/PBDBendingConstraints.h
@@ -26,10 +26,11 @@ public:
 			InParticles,
 			MoveTemp(InConstraints),
 			WeightMaps.FindRef(GetBendingElementStiffnessString(PropertyCollection, BendingElementStiffnessName.ToString())),
+			WeightMaps.FindRef(GetBucklingRatioString(PropertyCollection, BucklingRatioName.ToString())),
 			WeightMaps.FindRef(GetBucklingStiffnessString(PropertyCollection, BucklingStiffnessName.ToString())),
 			GetRestAngleMapFromCollection(WeightMaps, PropertyCollection),
 			FSolverVec2(GetWeightedFloatBendingElementStiffness(PropertyCollection, 1.f)),
-			(FSolverReal)GetBucklingRatio(PropertyCollection, 0.f),  // BucklingRatio is clamped in base class
+			FSolverVec2(GetWeightedFloatBucklingRatio(PropertyCollection, 0.f)),  // BucklingRatio is clamped in base class
 			FSolverVec2(GetWeightedFloatBucklingStiffness(PropertyCollection, 1.f)),
 			GetRestAngleValueFromCollection(PropertyCollection),
 			(ERestAngleConstructionType)GetRestAngleType(PropertyCollection, (int32)ERestAngleConstructionType::Use3DRestAngles),
@@ -57,10 +58,11 @@ public:
 			InParticleCount,
 			MoveTemp(InConstraints),
 			WeightMaps.FindRef(GetBendingElementStiffnessString(PropertyCollection, BendingElementStiffnessName.ToString())),
+			WeightMaps.FindRef(GetBucklingRatioString(PropertyCollection, BucklingRatioName.ToString())),
 			WeightMaps.FindRef(GetBucklingStiffnessString(PropertyCollection, BucklingStiffnessName.ToString())),
 			GetRestAngleMapFromCollection(WeightMaps, PropertyCollection),
 			FSolverVec2(GetWeightedFloatBendingElementStiffness(PropertyCollection, 1.f)),
-			(FSolverReal)GetBucklingRatio(PropertyCollection, 0.f),  // BucklingRatio is clamped in base class
+			FSolverVec2(GetWeightedFloatBucklingRatio(PropertyCollection, 0.f)),  // BucklingRatio is clamped in base class
 			FSolverVec2(GetWeightedFloatBucklingStiffness(PropertyCollection, 1.f)),
 			GetRestAngleValueFromCollection(PropertyCollection),
 			(ERestAngleConstructionType)GetRestAngleType(PropertyCollection, (int32)ERestAngleConstructionType::Use3DRestAngles),
@@ -90,10 +92,14 @@ public:
 			InParticleCount,
 			MoveTemp(InConstraints),
 			StiffnessMultipliers,
+			TConstArrayView<FRealSingle>(),
 			BucklingStiffnessMultipliers,
+			TConstArrayView<FRealSingle>(),
 			FSolverVec2(GetWeightedFloatBendingElementStiffness(PropertyCollection, 1.f)),
-			(FSolverReal)GetBucklingRatio(PropertyCollection, 0.f),  // BucklingRatio is clamped in base class
+			FSolverVec2(GetWeightedFloatBucklingRatio(PropertyCollection, 0.f)),  // BucklingRatio is clamped in base class
 			FSolverVec2(GetWeightedFloatBucklingStiffness(PropertyCollection, 1.f)),
+			FSolverVec2((FSolverReal)0.f),
+			ERestAngleConstructionType::Use3DRestAngles,
 			bTrimKinematicConstraints) 
 		, BendingElementStiffnessIndex(PropertyCollection)
 		, BucklingRatioIndex(PropertyCollection)
@@ -110,10 +116,14 @@ public:
 		int32 ParticleCount,
 		TArray<TVec4<int32>>&& InConstraints,
 		const TConstArrayView<FRealSingle>& StiffnessMultipliers,
+		const TConstArrayView<FRealSingle>& BucklingRatioMultipliers,
 		const TConstArrayView<FRealSingle>& BucklingStiffnessMultipliers,
+		const TConstArrayView<FRealSingle>& RestAngleMap,
 		const FSolverVec2& InStiffness,
-		const FSolverReal InBucklingRatio,
+		const FSolverVec2& InBucklingRatio,
 		const FSolverVec2& InBucklingStiffness,
+		const FSolverVec2& RestAngleValue,
+		ERestAngleConstructionType RestAngleConstructionType,
 		bool bTrimKinematicConstraints = false)
 		: Base(
 			InParticles,
@@ -121,10 +131,14 @@ public:
 			ParticleCount,
 			MoveTemp(InConstraints),
 			StiffnessMultipliers,
+			BucklingRatioMultipliers,
 			BucklingStiffnessMultipliers,
+			RestAngleMap,
 			InStiffness,
 			InBucklingRatio,
 			InBucklingStiffness,
+			RestAngleValue,
+			RestAngleConstructionType,
 			bTrimKinematicConstraints)
 		, BendingElementStiffnessIndex(ForceInit)
 		, BucklingRatioIndex(ForceInit)
@@ -136,6 +150,34 @@ public:
 		InitColor(InParticles);
 	}
 
+	UE_DEPRECATED(5.5, "Use constructor with BucklingRatioMultipliers")
+	FPBDBendingConstraints(const FSolverParticles& InParticles,
+		int32 ParticleOffset,
+		int32 ParticleCount,
+		TArray<TVec4<int32>>&& InConstraints,
+		const TConstArrayView<FRealSingle>& StiffnessMultipliers,
+		const TConstArrayView<FRealSingle>& BucklingStiffnessMultipliers,
+		const FSolverVec2& InStiffness,
+		const FSolverReal InBucklingRatio,
+		const FSolverVec2& InBucklingStiffness,
+		bool bTrimKinematicConstraints = false)
+		: FPBDBendingConstraints(
+			InParticles, 
+			ParticleOffset, 
+			ParticleCount, 
+			MoveTemp(InConstraints), 
+			StiffnessMultipliers, 
+			TConstArrayView<FRealSingle>(), 
+			BucklingStiffnessMultipliers,
+			TConstArrayView<FRealSingle>(),
+			InStiffness, 
+			FSolverVec2(InBucklingRatio), 
+			InBucklingStiffness,
+			FSolverVec2((FSolverReal)0.f),
+			ERestAngleConstructionType::Use3DRestAngles,
+			bTrimKinematicConstraints)
+	{}
+
 	virtual ~FPBDBendingConstraints() override {}
 
 	using Base::SetProperties;
@@ -199,7 +241,7 @@ private:
 	using Base::ParticleOffset;
 	using Base::ParticleCount;
 	using Base::Stiffness;
-	using Base::BucklingRatio;
+	using Base::BucklingRatioWeighted;
 	using Base::BucklingStiffness;
 
 	TArray<int32> ConstraintsPerColorStartIndex; // Constraints are ordered so each batch is contiguous. This is ColorNum + 1 length so it can be used as start and end.
diff --git a/Engine/Source/Runtime/Experimental/Chaos/Public/Chaos/PBDBendingConstraintsBase.h b/Engine/Source/Runtime/Experimental/Chaos/Public/Chaos/PBDBendingConstraintsBase.h
--- a/Engine/Source/Runtime/Experimental/Chaos/Public/Chaos/PBDBendingConstraintsBase.h
+++ b/Engine/Source/Runtime/Experimental/Chaos/Public/Chaos/PBDBendingConstraintsBase.h
@@ -24,10 +24,11 @@ public:
 	FPBDBendingConstraintsBase(const FSolverParticlesRange& InParticles,
 		TArray<TVec4<int32>>&& InConstraints,
 		const TConstArrayView<FRealSingle>& StiffnessMultipliers,
+		const TConstArrayView<FRealSingle>& BucklingRatioMultipliers,
 		const TConstArrayView<FRealSingle>& BucklingStiffnessMultipliers,
 		const TConstArrayView<FRealSingle>& RestAngleMap,
 		const FSolverVec2& InStiffness,
-		const FSolverReal InBucklingRatio,
+		const FSolverVec2& InBucklingRatio,
 		const FSolverVec2& InBucklingStiffness,
 		const FSolverVec2& RestAngleValue,
 		ERestAngleConstructionType RestAngleConstructionType,
@@ -46,7 +47,12 @@ public:
 			FPBDStiffness::DefaultTableSize,
 			FPBDStiffness::DefaultParameterFitBase,
 			MaxStiffness)
-		, BucklingRatio(FMath::Clamp(InBucklingRatio, (FSolverReal)0., (FSolverReal)1.))
+		, BucklingRatioWeighted(
+			InBucklingRatio.ClampAxes((FSolverReal)0., (FSolverReal)1.),
+			BucklingRatioMultipliers,
+			TConstArrayView<TVec2<int32>>(ConstraintSharedEdges),
+			ParticleOffset,
+			ParticleCount)
 		, BucklingStiffness(
 			InBucklingStiffness,
 			BucklingStiffnessMultipliers,
@@ -56,19 +62,52 @@ public:
 			FPBDStiffness::DefaultTableSize,
 			FPBDStiffness::DefaultParameterFitBase,
 			MaxStiffness)
+PRAGMA_DISABLE_DEPRECATION_WARNINGS
+		, BucklingRatio((FSolverReal)BucklingRatioWeighted)
+PRAGMA_ENABLE_DEPRECATION_WARNINGS
 	{
 		CalculateRestAngles(InParticles, ParticleOffset, ParticleCount, RestAngleMap, RestAngleValue, RestAngleConstructionType);
 	}
 
+	UE_DEPRECATED(5.5, "Use a constructor with BucklingRatioMultipliers.")
+	FPBDBendingConstraintsBase(const FSolverParticlesRange& InParticles,
+		TArray<TVec4<int32>>&& InConstraints,
+		const TConstArrayView<FRealSingle>& StiffnessMultipliers,
+		const TConstArrayView<FRealSingle>& BucklingStiffnessMultipliers,
+		const TConstArrayView<FRealSingle>& RestAngleMap,
+		const FSolverVec2& InStiffness,
+		const FSolverReal InBucklingRatio,
+		const FSolverVec2& InBucklingStiffness,
+		const FSolverVec2& RestAngleValue,
+		ERestAngleConstructionType RestAngleConstructionType,
+		bool bTrimKinematicConstraints = false,
+		FSolverReal MaxStiffness = FPBDStiffness::DefaultPBDMaxStiffness)
+		: FPBDBendingConstraintsBase(
+			InParticles,
+			MoveTemp(InConstraints),
+			StiffnessMultipliers,
+			TConstArrayView<float>(),
+			BucklingStiffnessMultipliers,
+			RestAngleMap,
+			InStiffness,
+			FSolverVec2(InBucklingRatio),
+			InBucklingStiffness,
+			RestAngleValue,
+			RestAngleConstructionType,
+			bTrimKinematicConstraints,
+			MaxStiffness)
+	{}
+
 	FPBDBendingConstraintsBase(const FSolverParticles& InParticles,
 		int32 InParticleOffset,
 		int32 InParticleCount,
 		TArray<TVec4<int32>>&& InConstraints,
 		const TConstArrayView<FRealSingle>& StiffnessMultipliers,
+		const TConstArrayView<FRealSingle>& BucklingRatioMultipliers,
 		const TConstArrayView<FRealSingle>& BucklingStiffnessMultipliers,
 		const TConstArrayView<FRealSingle>& RestAngleMap,
 		const FSolverVec2& InStiffness,
-		const FSolverReal InBucklingRatio,
+		const FSolverVec2& InBucklingRatio,
 		const FSolverVec2& InBucklingStiffness,
 		const FSolverVec2& RestAngleValue,
 		ERestAngleConstructionType RestAngleConstructionType,
@@ -87,7 +126,12 @@ public:
 			FPBDStiffness::DefaultTableSize,
 			FPBDStiffness::DefaultParameterFitBase,
 			MaxStiffness)
-		, BucklingRatio(FMath::Clamp(InBucklingRatio, (FSolverReal)0., (FSolverReal)1.))
+		, BucklingRatioWeighted(
+			InBucklingRatio.ClampAxes((FSolverReal)0., (FSolverReal)1.),
+			BucklingRatioMultipliers,
+			TConstArrayView<TVec2<int32>>(ConstraintSharedEdges),
+			ParticleOffset,
+			ParticleCount)
 		, BucklingStiffness(
 			InBucklingStiffness,
 			BucklingStiffnessMultipliers,
@@ -97,10 +141,48 @@ public:
 			FPBDStiffness::DefaultTableSize,
 			FPBDStiffness::DefaultParameterFitBase,
 			MaxStiffness)
+		PRAGMA_DISABLE_DEPRECATION_WARNINGS
+		, BucklingRatio((FSolverReal)BucklingRatioWeighted)
+		PRAGMA_ENABLE_DEPRECATION_WARNINGS
 	{
 		CalculateRestAngles(InParticles, ParticleOffset, ParticleCount, RestAngleMap, RestAngleValue, RestAngleConstructionType);
 	}
 
+	UE_DEPRECATED(5.5, "Use a constructor with BucklingRatioMultipliers.")
+	FPBDBendingConstraintsBase(const FSolverParticles& InParticles,
+		int32 InParticleOffset,
+		int32 InParticleCount,
+		TArray<TVec4<int32>>&& InConstraints,
+		const TConstArrayView<FRealSingle>& StiffnessMultipliers,
+		const TConstArrayView<FRealSingle>& BucklingStiffnessMultipliers,
+		const TConstArrayView<FRealSingle>& RestAngleMap,
+		const FSolverVec2& InStiffness,
+		const FSolverReal InBucklingRatio,
+		const FSolverVec2& InBucklingStiffness,
+		const FSolverVec2& RestAngleValue,
+		ERestAngleConstructionType RestAngleConstructionType,
+		bool bTrimKinematicConstraints = false,
+		FSolverReal MaxStiffness = FPBDStiffness::DefaultPBDMaxStiffness)
+		: FPBDBendingConstraintsBase(
+			InParticles,
+			InParticleOffset,
+			InParticleCount,
+			MoveTemp(InConstraints),
+			StiffnessMultipliers,
+			TConstArrayView<FRealSingle>(),
+			BucklingStiffnessMultipliers,
+			RestAngleMap,
+			InStiffness,
+			FSolverVec2(InBucklingRatio),
+			InBucklingStiffness,
+			RestAngleValue,
+			RestAngleConstructionType,
+			bTrimKinematicConstraints,
+			MaxStiffness)
+	{
+	}
+
+	UE_DEPRECATED(5.5, "Use a constructor with BucklingRatioMultipliers.")
 	FPBDBendingConstraintsBase(const FSolverParticles& InParticles,
 		int32 InParticleOffset,
 		int32 InParticleCount,
@@ -113,44 +195,57 @@ public:
 		bool bTrimKinematicConstraints = false,
 		FSolverReal MaxStiffness = FPBDStiffness::DefaultPBDMaxStiffness)
 		:FPBDBendingConstraintsBase(InParticles, InParticleOffset, InParticleCount, MoveTemp(InConstraints),
-			StiffnessMultipliers, BucklingStiffnessMultipliers, TConstArrayView<FRealSingle>(),
-			InStiffness, InBucklingRatio, InBucklingStiffness, FSolverVec2((FSolverReal)0.f),
+			StiffnessMultipliers, TConstArrayView<FRealSingle>(), BucklingStiffnessMultipliers, TConstArrayView<FRealSingle>(),
+			InStiffness, FSolverVec2(InBucklingRatio), InBucklingStiffness, FSolverVec2((FSolverReal)0.f),
 			ERestAngleConstructionType::Use3DRestAngles, bTrimKinematicConstraints, MaxStiffness)
 	{}
 
 	UE_DEPRECATED(5.2, "Use one of the other constructors instead.")
 	FPBDBendingConstraintsBase(const FSolverParticles& InParticles, TArray<TVec4<int32>>&& InConstraints, const FSolverReal InStiffness = (FSolverReal)1.)
-		: Constraints(MoveTemp(InConstraints))
-		, ConstraintSharedEdges(ExtractConstraintSharedEdges(Constraints))
-		, ParticleOffset(0)
-		, ParticleCount(InParticles.Size())
-		, Stiffness(FSolverVec2(InStiffness))
-		, BucklingRatio(0.f)
-		, BucklingStiffness(FSolverVec2(InStiffness))
+		: FPBDBendingConstraintsBase(
+			InParticles,
+			0, 
+			InParticles.Size(),
+			MoveTemp(InConstraints),
+			TConstArrayView<FRealSingle>(),
+			TConstArrayView<FRealSingle>(),
+			TConstArrayView<FRealSingle>(),
+			TConstArrayView<FRealSingle>(),
+			FSolverVec2(InStiffness),
+			FSolverVec2((FSolverReal)0.f),
+			FSolverVec2(InStiffness),
+			FSolverVec2((FSolverReal)0.f),
+			ERestAngleConstructionType::Use3DRestAngles)
 	{
-		CalculateRestAngles(InParticles, ParticleOffset, ParticleCount, TConstArrayView<FRealSingle>(), FSolverVec2(0.f), ERestAngleConstructionType::Use3DRestAngles);
 	}
 
 	virtual ~FPBDBendingConstraintsBase() {}
 
 	// Update stiffness values
-	void SetProperties(const FSolverVec2& InStiffness, const FSolverReal InBucklingRatio, const FSolverVec2& InBucklingStiffness)
+	void SetProperties(const FSolverVec2& InStiffness, const FSolverVec2& InBucklingRatio, const FSolverVec2& InBucklingStiffness)
 	{ 
 		Stiffness.SetWeightedValue(InStiffness);
-		BucklingRatio = FMath::Clamp(InBucklingRatio, (FSolverReal)0., (FSolverReal)1.);
+		BucklingRatioWeighted.SetWeightedValue(InBucklingRatio.ClampAxes((FSolverReal)0., (FSolverReal)1.));
 		BucklingStiffness.SetWeightedValue(InBucklingStiffness);
+		PRAGMA_DISABLE_DEPRECATION_WARNINGS
+		BucklingRatio = (FSolverReal)BucklingRatioWeighted;
+		PRAGMA_ENABLE_DEPRECATION_WARNINGS
+	}
+
+	UE_DEPRECATED(5.5, "Use version with FSolverVec2 InBucklingRatio.")
+	void SetProperties(const FSolverVec2& InStiffness, const FSolverReal InBucklingRatio, const FSolverVec2& InBucklingStiffness)
+	{
+		SetProperties(InStiffness, FSolverVec2(InBucklingRatio), InBucklingStiffness);
 	}
 
 	// Update stiffness table, as well as the simulation stiffness exponent
 	void ApplyProperties(const FSolverReal Dt, const int32 NumIterations)
 	{
 		Stiffness.ApplyPBDValues(Dt, NumIterations);
+		BucklingRatioWeighted.ApplyValues();
 		BucklingStiffness.ApplyPBDValues(Dt, NumIterations);
 	}
 
-	UE_DEPRECATED(5.1, "Use SetProperties instead.")
-	void SetStiffness(FSolverReal InStiffness) { SetProperties(FSolverVec2(InStiffness), 0.f, 1.f); }
-
 	template<typename SolverParticlesOrRange>
 	TStaticArray<FSolverVec3, 4> GetGradients(const SolverParticlesOrRange& InParticles, const int32 i) const
 	{
@@ -212,11 +307,19 @@ public:
 		return FMath::Atan2(SinPhi, CosPhi);
 	}
 
+	UE_DEPRECATED(5.5, "This version is deprecated and will produce the incorrect result if BucklingRatio is weighted. Use version that takes a ConstraintIndex instead.")
 	bool AngleIsBuckled(const FSolverReal Angle, const FSolverReal RestAngle) const
 	{
 		// Angle is 0 when completely flat. This is easier to think of in terms of Angle' = (PI - |Angle|), which is 0 when completely folded.
 		// Consider buckled when Angle' <= BucklingRatio * RestAngle', and use buckling stiffness instead of stiffness.
-		return UE_PI - FMath::Abs(Angle) < BucklingRatio * (UE_PI - FMath::Abs(RestAngle));
+		return UE_PI - FMath::Abs(Angle) < (FSolverReal)BucklingRatioWeighted * (UE_PI - FMath::Abs(RestAngle));
+	}
+
+	bool AngleIsBuckled(const FSolverReal Angle, const int32 ConstraintIndex) const
+	{
+		// Angle is 0 when completely flat. This is easier to think of in terms of Angle' = (PI - |Angle|), which is 0 when completely folded.
+		// Consider buckled when Angle' <= BucklingRatio * RestAngle', and use buckling stiffness instead of stiffness.
+		return UE_PI - FMath::Abs(Angle) < BucklingRatioWeighted.GetValue(ConstraintIndex) * (UE_PI - FMath::Abs(RestAngles[ConstraintIndex]));
 	}
 
 	template<typename SolverParticlesOrRange>
@@ -319,11 +422,14 @@ protected:
 	const int32 ParticleCount;
 
 	FPBDStiffness Stiffness;
-	FSolverReal BucklingRatio;
+	FPBDWeightMap BucklingRatioWeighted;
 	FPBDStiffness BucklingStiffness;
 
 	TArray<FSolverReal> RestAngles;
 	TArray<bool> IsBuckled;
+
+	UE_DEPRECATED(5.5, "Scalar Buckling Ratio has been deprecated. Use BucklingRatioWeighted instead.")
+	FSolverReal BucklingRatio;
 };
 
 }  // End namespace Chaos::Softs
diff --git a/Engine/Source/Runtime/Experimental/Chaos/Public/Chaos/PBDWeightMap.h b/Engine/Source/Runtime/Experimental/Chaos/Public/Chaos/PBDWeightMap.h
--- a/Engine/Source/Runtime/Experimental/Chaos/Public/Chaos/PBDWeightMap.h
+++ b/Engine/Source/Runtime/Experimental/Chaos/Public/Chaos/PBDWeightMap.h
@@ -93,6 +93,15 @@ public:
 	/** Return the exponential stiffness value when the weight map is not used. */
 	explicit operator FSolverReal() const { return GetLow(); }
 
+	FSolverReal GetValue(int32 Index) const
+	{
+		if (HasWeightMap())
+		{
+			return (*this)[Index];
+		}
+		return GetLow();
+	}
+
 	/** Return the particles/constraints indices to the stiffness table as a read only array. */
 	TConstArrayView<uint8> GetIndices() const { return TConstArrayView<uint8>(Indices); }
 
diff --git a/Engine/Source/Runtime/Experimental/Chaos/Public/Chaos/XPBDAnisotropicBendingConstraints.h b/Engine/Source/Runtime/Experimental/Chaos/Public/Chaos/XPBDAnisotropicBendingConstraints.h
--- a/Engine/Source/Runtime/Experimental/Chaos/Public/Chaos/XPBDAnisotropicBendingConstraints.h
+++ b/Engine/Source/Runtime/Experimental/Chaos/Public/Chaos/XPBDAnisotropicBendingConstraints.h
@@ -66,7 +66,33 @@ public:
 		: FXPBDAnisotropicBendingConstraints(InParticles, InParticleOffset, InParticleCount, TriangleMesh, FaceVertexPatternPositions, WeightMaps, PropertyCollection)
 	{}
 
-	CHAOS_API FXPBDAnisotropicBendingConstraints(const FSolverParticles& InParticles,
+	CHAOS_API FXPBDAnisotropicBendingConstraints(const FSolverParticles & InParticles,
+		int32 InParticleOffset,
+		int32 InParticleCount,
+		const FTriangleMesh& TriangleMesh,
+		const TArray<TVec3<FVec2f>>& FaceVertexPatternPositions,
+		const TConstArrayView<FRealSingle>& StiffnessWarpMultipliers,
+		const TConstArrayView<FRealSingle>& StiffnessWeftMultipliers,
+		const TConstArrayView<FRealSingle>& StiffnessBiasMultipliers,
+		const TConstArrayView<FRealSingle>& BucklingRatioMultipliers,
+		const TConstArrayView<FRealSingle>& BucklingStiffnessWarpMultipliers,
+		const TConstArrayView<FRealSingle>& BucklingStiffnessWeftMultipliers,
+		const TConstArrayView<FRealSingle>& BucklingStiffnessBiasMultipliers,
+		const TConstArrayView<FRealSingle>& DampingMultipliers,
+		const TConstArrayView<FRealSingle>& RestAngleMap,
+		const FSolverVec2& InStiffnessWarp,
+		const FSolverVec2& InStiffnessWeft,
+		const FSolverVec2& InStiffnessBias,
+		const FSolverVec2& InBucklingRatio,
+		const FSolverVec2& InBucklingStiffnessWarp,
+		const FSolverVec2& InBucklingStiffnessWeft,
+		const FSolverVec2& InBucklingStiffnessBias,
+		const FSolverVec2& InDampingRatio,
+		const FSolverVec2& RestAngleValue,
+		ERestAngleConstructionType RestAngleConstructionType);
+
+	UE_DEPRECATED(5.5, "Use constructor with BucklingRatioMultipliers")
+	FXPBDAnisotropicBendingConstraints(const FSolverParticles& InParticles,
 		int32 InParticleOffset,
 		int32 InParticleCount,
 		const FTriangleMesh& TriangleMesh,
@@ -85,7 +111,33 @@ public:
 		const FSolverVec2& InBucklingStiffnessWarp,
 		const FSolverVec2& InBucklingStiffnessWeft,
 		const FSolverVec2& InBucklingStiffnessBias,
-		const FSolverVec2& InDampingRatio);
+		const FSolverVec2& InDampingRatio)
+		:FXPBDAnisotropicBendingConstraints(
+			InParticles,
+			InParticleOffset,
+			InParticleCount,
+			TriangleMesh,
+			FaceVertexPatternPositions,
+			StiffnessWarpMultipliers,
+			StiffnessWeftMultipliers,
+			StiffnessBiasMultipliers,
+			TConstArrayView<FRealSingle>(),
+			BucklingStiffnessWarpMultipliers,
+			BucklingStiffnessWeftMultipliers,
+			BucklingStiffnessBiasMultipliers,
+			DampingMultipliers,
+			TConstArrayView<FRealSingle>(),
+			InStiffnessWarp,
+			InStiffnessWeft,
+			InStiffnessBias,
+			FSolverVec2(InBucklingRatio),
+			InBucklingStiffnessWarp,
+			InBucklingStiffnessWeft,
+			InBucklingStiffnessBias,
+			InDampingRatio,
+			FSolverVec2((FSolverReal)0.f),
+			ERestAngleConstructionType::Use3DRestAngles)
+	{}
 
 	UE_DEPRECATED(5.4, "XPBD Constraints must always trim kinematic constraints")
 	FXPBDAnisotropicBendingConstraints(const FSolverParticles& InParticles,
@@ -109,9 +161,31 @@ public:
 		const FSolverVec2& InBucklingStiffnessBias,
 		const FSolverVec2& InDampingRatio,
 		bool bTrimKinematicConstraints)
-		: FXPBDAnisotropicBendingConstraints(InParticles, InParticleOffset, InParticleCount, TriangleMesh, FaceVertexPatternPositions, StiffnessWarpMultipliers, StiffnessWeftMultipliers,
-			StiffnessBiasMultipliers, BucklingStiffnessWarpMultipliers, BucklingStiffnessWeftMultipliers, BucklingStiffnessBiasMultipliers, DampingMultipliers, InStiffnessWarp,
-			InStiffnessWeft, InStiffnessBias, InBucklingRatio, InBucklingStiffnessWarp, InBucklingStiffnessWeft, InBucklingStiffnessBias, InDampingRatio)
+		: FXPBDAnisotropicBendingConstraints(
+			InParticles,
+			InParticleOffset,
+			InParticleCount,
+			TriangleMesh,
+			FaceVertexPatternPositions,
+			StiffnessWarpMultipliers,
+			StiffnessWeftMultipliers,
+			StiffnessBiasMultipliers,
+			TConstArrayView<FRealSingle>(),
+			BucklingStiffnessWarpMultipliers,
+			BucklingStiffnessWeftMultipliers,
+			BucklingStiffnessBiasMultipliers,
+			DampingMultipliers,
+			TConstArrayView<FRealSingle>(),
+			InStiffnessWarp,
+			InStiffnessWeft,
+			InStiffnessBias,
+			FSolverVec2(InBucklingRatio),
+			InBucklingStiffnessWarp,
+			InBucklingStiffnessWeft,
+			InBucklingStiffnessBias,
+			InDampingRatio,
+			FSolverVec2((FSolverReal)0.f),
+			ERestAngleConstructionType::Use3DRestAngles)
 	{}
 
 	virtual ~FXPBDAnisotropicBendingConstraints() override {}
@@ -131,7 +205,7 @@ public:
 
 	void ApplyProperties(const FSolverReal /*Dt*/, const int32 /*NumIterations*/)
 	{
-		// Nothing to be done here for flat weight maps. Want to avoid base class from being called instead.
+		BucklingRatioWeighted.ApplyValues();
 	}
 
 	template<typename SolverParticlesOrRange>
diff --git a/Engine/Source/Runtime/Experimental/Chaos/Public/Chaos/XPBDBendingConstraints.h b/Engine/Source/Runtime/Experimental/Chaos/Public/Chaos/XPBDBendingConstraints.h
--- a/Engine/Source/Runtime/Experimental/Chaos/Public/Chaos/XPBDBendingConstraints.h
+++ b/Engine/Source/Runtime/Experimental/Chaos/Public/Chaos/XPBDBendingConstraints.h
@@ -39,10 +39,11 @@ public:
 			InParticles,
 			MoveTemp(InConstraints),
 			TConstArrayView<FRealSingle>(), // We don't use base stiffness weight maps
+			WeightMaps.FindRef(GetXPBDBucklingRatioString(PropertyCollection, XPBDBucklingRatioName.ToString())),
 			TConstArrayView<FRealSingle>(), // We don't use base stiffness weight maps
 			GetRestAngleMapFromCollection(WeightMaps, PropertyCollection),
 			FSolverVec2(GetWeightedFloatXPBDBendingElementStiffness(PropertyCollection, MaxStiffness)),
-			(FSolverReal)GetXPBDBucklingRatio(PropertyCollection, 0.f),
+			FSolverVec2(GetWeightedFloatXPBDBucklingRatio(PropertyCollection, 0.f)),
 			FSolverVec2(GetWeightedFloatXPBDBucklingStiffness(PropertyCollection, MaxStiffness)),
 			GetRestAngleValueFromCollection(PropertyCollection),
 			(ERestAngleConstructionType)GetXPBDRestAngleType(PropertyCollection, (int32)ERestAngleConstructionType::Use3DRestAngles),
@@ -100,10 +101,11 @@ public:
 			InParticleCount,
 			MoveTemp(InConstraints),
 			TConstArrayView<FRealSingle>(), // We don't use base stiffness weight maps
+			WeightMaps.FindRef(GetXPBDBucklingRatioString(PropertyCollection, XPBDBucklingRatioName.ToString())),
 			TConstArrayView<FRealSingle>(), // We don't use base stiffness weight maps
 			GetRestAngleMapFromCollection(WeightMaps, PropertyCollection),
 			FSolverVec2(GetWeightedFloatXPBDBendingElementStiffness(PropertyCollection, MaxStiffness)),
-			(FSolverReal)GetXPBDBucklingRatio(PropertyCollection, 0.f),
+			FSolverVec2(GetWeightedFloatXPBDBucklingRatio(PropertyCollection, 0.f)),
 			FSolverVec2(GetWeightedFloatXPBDBucklingStiffness(PropertyCollection, MaxStiffness)),
 			GetRestAngleValueFromCollection(PropertyCollection),
 			(ERestAngleConstructionType)GetXPBDRestAngleType(PropertyCollection, (int32)ERestAngleConstructionType::Use3DRestAngles),
@@ -167,10 +169,14 @@ public:
 			InParticleCount,
 			MoveTemp(InConstraints),
 			TConstArrayView<FRealSingle>(), // We don't use base stiffness weight maps
+			TConstArrayView<FRealSingle>(),
 			TConstArrayView<FRealSingle>(), // We don't use base stiffness weight maps
+			TConstArrayView<FRealSingle>(),
 			FSolverVec2(GetWeightedFloatXPBDBendingElementStiffness(PropertyCollection, MaxStiffness)),
-			(FSolverReal)GetXPBDBucklingRatio(PropertyCollection, 0.f),
+			FSolverVec2(GetWeightedFloatXPBDBucklingRatio(PropertyCollection, 0.f)),
 			FSolverVec2(GetWeightedFloatXPBDBucklingStiffness(PropertyCollection, MaxStiffness)),
+			FSolverVec2((FSolverReal)0.f),
+			ERestAngleConstructionType::Use3DRestAngles,
 			true /*bTrimKinematicConstraints*/,
 			MaxStiffness)
 		, XPBDStiffness(
@@ -204,66 +210,35 @@ public:
 		InitColor(InParticles);
 	}
 
-	UE_DEPRECATED(5.2, "Use one the other constructors instead.")
-	FXPBDBendingConstraints(const FSolverParticles& InParticles,
-		int32 InParticleOffset,
-		int32 InParticleCount,
-		TArray<TVec4<int32>>&& InConstraints,
-		const TConstArrayView<FRealSingle>& StiffnessMultipliers,
-		const TConstArrayView<FRealSingle>& BucklingStiffnessMultipliers,
-		const FSolverVec2& InStiffness,
-		const FSolverReal InBucklingRatio,
-		const FSolverVec2& InBucklingStiffness,
-		bool bTrimKinematicConstraints = false)
-		: Base(
-			InParticles,
-			InParticleOffset,
-			InParticleCount,
-			MoveTemp(InConstraints),
-			TConstArrayView<FRealSingle>(), // We don't use base stiffness weight maps
-			TConstArrayView<FRealSingle>(), // We don't use base stiffness weight maps
-			InStiffness,
-			InBucklingRatio,
-			InBucklingStiffness,
-			true /*bTrimKinematicConstraints*/,
-			MaxStiffness)
-		, XPBDStiffness(InStiffness.ClampAxes(0, MaxStiffness))
-		, XPBDBucklingStiffness(InBucklingStiffness.ClampAxes(0, MaxStiffness))
-		, DampingRatio(FSolverVec2::ZeroVector)
-		, XPBDBendingElementStiffnessIndex(ForceInit)
-		, XPBDBendingElementDampingIndex(ForceInit)
-		, XPBDBucklingRatioIndex(ForceInit)
-		, XPBDBucklingStiffnessIndex(ForceInit)
-		, XPBDFlatnessRatioIndex(ForceInit)
-		, XPBDRestAngleIndex(ForceInit)
-		, XPBDRestAngleTypeIndex(ForceInit)
-	{
-		Lambdas.Init((FSolverReal)0., Constraints.Num());
-		LambdasDamping.Init((FSolverReal)0., Constraints.Num());
-		InitColor(InParticles);
-	}
-
 	FXPBDBendingConstraints(const FSolverParticles& InParticles,
 		int32 ParticleOffset,
 		int32 ParticleCount,
 		TArray<TVec4<int32>>&& InConstraints,
 		const TConstArrayView<FRealSingle>& StiffnessMultipliers,
+		const TConstArrayView<FRealSingle>& BucklingRatioMultipliers,
 		const TConstArrayView<FRealSingle>& BucklingStiffnessMultipliers,
 		const TConstArrayView<FRealSingle>& InDampingMultipliers,
+		const TConstArrayView<FRealSingle>& RestAngleMap,
 		const FSolverVec2& InStiffness,
-		const FSolverReal InBucklingRatio,
+		const FSolverVec2& InBucklingRatio,
 		const FSolverVec2& InBucklingStiffness,
-		const FSolverVec2& InDampingRatio)
+		const FSolverVec2& InDampingRatio,
+		const FSolverVec2& RestAngleValue,
+		ERestAngleConstructionType RestAngleConstructionType)
 		: Base(
 			InParticles,
 			ParticleOffset,
 			ParticleCount,
 			MoveTemp(InConstraints),
 			TConstArrayView<FRealSingle>(), // We don't use base stiffness weight maps
+			BucklingRatioMultipliers,
 			TConstArrayView<FRealSingle>(), // We don't use base stiffness weight maps
+			RestAngleMap,
 			InStiffness,
 			InBucklingRatio,
 			InBucklingStiffness,
+			RestAngleValue,
+			RestAngleConstructionType,
 			true /*bTrimKinematicConstraints*/,
 			MaxStiffness)
 		, XPBDStiffness(
@@ -296,6 +271,37 @@ public:
 		LambdasDamping.Init((FSolverReal)0., Constraints.Num());
 		InitColor(InParticles);
 	}
+	
+
+	UE_DEPRECATED(5.5, "Use constructor with BucklingRatioMultipliers")
+	FXPBDBendingConstraints(const FSolverParticles& InParticles,
+		int32 ParticleOffset,
+		int32 ParticleCount,
+		TArray<TVec4<int32>>&& InConstraints,
+		const TConstArrayView<FRealSingle>& StiffnessMultipliers,
+		const TConstArrayView<FRealSingle>& BucklingStiffnessMultipliers,
+		const TConstArrayView<FRealSingle>& InDampingMultipliers,
+		const FSolverVec2& InStiffness,
+		const FSolverReal InBucklingRatio,
+		const FSolverVec2& InBucklingStiffness,
+		const FSolverVec2& InDampingRatio)
+		: FXPBDBendingConstraints(
+			InParticles, 
+			ParticleOffset, 
+			ParticleCount, 
+			MoveTemp(InConstraints), 
+			StiffnessMultipliers,
+			TConstArrayView<FRealSingle>(),
+			BucklingStiffnessMultipliers, 
+			InDampingMultipliers,
+			TConstArrayView<FRealSingle>(),
+			InStiffness, 
+			FSolverVec2(InBucklingRatio),
+			InBucklingStiffness, 
+			InDampingRatio,
+			FSolverVec2((FSolverReal)0.f),
+			ERestAngleConstructionType::Use3DRestAngles)
+	{}
 
 	UE_DEPRECATED(5.4, "XPBD Constraints must always trim kinematic constraints")
 	FXPBDBendingConstraints(const FSolverParticles& InParticles,
@@ -310,7 +316,22 @@ public:
 		const FSolverVec2& InBucklingStiffness,
 		const FSolverVec2& InDampingRatio,
 		bool bTrimKinematicConstraints)
-		: FXPBDBendingConstraints(InParticles, ParticleOffset, ParticleCount, MoveTemp(InConstraints), StiffnessMultipliers, BucklingStiffnessMultipliers, InDampingMultipliers, InStiffness, InBucklingRatio, InBucklingStiffness, InDampingRatio)
+		: FXPBDBendingConstraints(
+			InParticles, 
+			ParticleOffset, 
+			ParticleCount, 
+			MoveTemp(InConstraints), 
+			StiffnessMultipliers,
+			TConstArrayView<FRealSingle>(),
+			BucklingStiffnessMultipliers, 
+			InDampingMultipliers,
+			TConstArrayView<FRealSingle>(),
+			InStiffness, 
+			FSolverVec2(InBucklingRatio),
+			InBucklingStiffness, 
+			InDampingRatio,
+			FSolverVec2((FSolverReal)0.f),
+			ERestAngleConstructionType::Use3DRestAngles)
 	{}
 
 	virtual ~FXPBDBendingConstraints() override {}
@@ -328,17 +349,23 @@ public:
 		SetProperties(PropertyCollection, TMap<FString, TConstArrayView<FRealSingle>>());
 	}
 
-	void SetProperties(const FSolverVec2& InStiffness, const FSolverReal InBucklingRatio, const FSolverVec2& InBucklingStiffness, const FSolverVec2& InDampingRatio = FSolverVec2::ZeroVector)
+	void SetProperties(const FSolverVec2& InStiffness, const FSolverVec2& InBucklingRatio, const FSolverVec2& InBucklingStiffness, const FSolverVec2& InDampingRatio)
 	{
 		XPBDStiffness.SetWeightedValue(InStiffness.ClampAxes(0, MaxStiffness));
-		BucklingRatio = FMath::Clamp(InBucklingRatio, (FSolverReal)0., (FSolverReal)1.);
+		BucklingRatioWeighted.SetWeightedValue(InBucklingRatio.ClampAxes((FSolverReal)0., (FSolverReal)1.));
 		XPBDBucklingStiffness.SetWeightedValue(InBucklingStiffness.ClampAxes(0, MaxStiffness));
 		DampingRatio.SetWeightedValue(InDampingRatio.ClampAxes(MinDamping, MaxDamping));
 	}
 
+	UE_DEPRECATED(5.5, "Use SetProperties with Weighted Buckling Ratio")
+	void SetProperties(const FSolverVec2& InStiffness, const FSolverReal InBucklingRatio, const FSolverVec2& InBucklingStiffness, const FSolverVec2& InDampingRatio = FSolverVec2::ZeroVector)
+	{
+		SetProperties(InStiffness, FSolverVec2(InBucklingRatio), InBucklingStiffness, InDampingRatio);
+	}
+
 	void ApplyProperties(const FSolverReal /*Dt*/, const int32 /*NumIterations*/)
 	{
-		// Nothing to be done here for flat weight maps. Want to avoid base class from being called instead.
+		BucklingRatioWeighted.ApplyValues();
 	}
 
 	template<typename SolverParticlesOrRange>
