Chaos - spheres/capsule vs tri-mesh collision path now same as vs heightfield
- add option to sort mesh contacts by depth (cvar disabled)

diff --git a/Engine/Source/Runtime/Experimental/Chaos/Private/Chaos/Collision/MeshContactGenerator.cpp b/Engine/Source/Runtime/Experimental/Chaos/Private/Chaos/Collision/MeshContactGenerator.cpp
--- a/Engine/Source/Runtime/Experimental/Chaos/Private/Chaos/Collision/MeshContactGenerator.cpp
+++ b/Engine/Source/Runtime/Experimental/Chaos/Private/Chaos/Collision/MeshContactGenerator.cpp
@@ -30,7 +30,8 @@ namespace Chaos::Private
 		MaxContactsBufferSize = 1000;
 		bCullBackFaces = true;
 		bFixNormals = true;
-		bSortForSolver = false;
+		bSortByPhi = false;
+		bSortForSolverConvergence = false;
 		bUseTwoPassLoop = true;
 	}
 
@@ -166,9 +167,14 @@ namespace Chaos::Private
 		// Visited triangles are white, ignored triangles are gray
 		DebugDrawTriangles(ConvexTransform, FColor::White, FColor::Silver);
 
-		if (!!Settings.bSortForSolver)
+		if (!!Settings.bSortByPhi)
 		{
-			SortContactsForSolver();
+			SortContactByPhi();
+		}
+
+		if (!!Settings.bSortForSolverConvergence)
+		{
+			SortContactsForSolverConvergence();
 		}
 
 		FinalizeContacts(MeshToConvexTransform);
@@ -562,9 +568,18 @@ namespace Chaos::Private
 		return false;
 	}
 
+	void FMeshContactGenerator::SortContactByPhi()
+	{
+		Algo::Sort(Contacts,
+			[](const FContactPoint& L, const FContactPoint& R)
+			{
+				return L.Phi < R.Phi;
+			});
+	}
+
 	// Sort contacts on a shape pair in the order we like to solve them.
 	// NOTE: This relies on the enum order of EContactPointType.
-	void FMeshContactGenerator::SortContactsForSolver()
+	void FMeshContactGenerator::SortContactsForSolverConvergence()
 	{
 		// Sort TriangleContactPoints in solver preferred order, but ignore TriangleContactPointDatas
 		// NOTE: This should only be called at the end of the pruning proxess when we no longer care 
diff --git a/Engine/Source/Runtime/Experimental/Chaos/Private/Chaos/CollisionOneShotManifolds.cpp b/Engine/Source/Runtime/Experimental/Chaos/Private/Chaos/CollisionOneShotManifolds.cpp
--- a/Engine/Source/Runtime/Experimental/Chaos/Private/Chaos/CollisionOneShotManifolds.cpp
+++ b/Engine/Source/Runtime/Experimental/Chaos/Private/Chaos/CollisionOneShotManifolds.cpp
@@ -158,6 +158,10 @@ namespace Chaos
 	bool bChaos_Collision_UseConvexTriangleGJKSAT = true;
 	FAutoConsoleVariableRef CVarChaos_Collision_UseConvexTriangleSAT(TEXT("p.Chaos.Collision.UseConvexTriangleGJKSAT"), bChaos_Collision_UseConvexTriangleGJKSAT, TEXT(""));
 
+	// Whether to sort sphere/capsule vs mesh contacts by depth. Doing so leads to more consistent behaviour but is more expensive
+	bool bChaos_Collision_ConvexTriMeshSortByPhi = false;
+	FAutoConsoleVariableRef CVarChaos_Collision_ConvexTriMeshSortByPhi(TEXT("p.Chaos.Collision.ConvexTriMeshSortByPhi"), bChaos_Collision_ConvexTriMeshSortByPhi, TEXT(""));
+
 	namespace Collisions
 	{
 		// Forward delarations we need from CollisionRestitution.cpp
diff --git a/Engine/Source/Runtime/Experimental/Chaos/Private/Chaos/CollisionOneShotManifoldsMeshShapes.cpp b/Engine/Source/Runtime/Experimental/Chaos/Private/Chaos/CollisionOneShotManifoldsMeshShapes.cpp
--- a/Engine/Source/Runtime/Experimental/Chaos/Private/Chaos/CollisionOneShotManifoldsMeshShapes.cpp
+++ b/Engine/Source/Runtime/Experimental/Chaos/Private/Chaos/CollisionOneShotManifoldsMeshShapes.cpp
@@ -42,6 +42,8 @@ namespace Chaos
 	extern bool bChaos_Collision_ConvexTriMeshInsideCull;
 	extern bool bChaos_Collision_ConvexTriMeshBackFaceCull;
 
+	extern bool bChaos_Collision_ConvexTriMeshSortByPhi;
+
 	namespace CVars
 	{
 #if CHAOS_DEBUG_DRAW
@@ -759,24 +761,52 @@ namespace Chaos
 			check(Mesh != nullptr);
 
 			const FReal CullDistance = Constraint.GetCullDistance();
-			const FReal PhiTolerance = CalculateTriMeshPhiTolerance(CullDistance);
-			const FReal DistanceTolerance = Chaos_Collision_TriMeshDistanceTolerance;
-			FContactTriangleCollector MeshContacts(bChaos_Collision_OneSidedTriangleMesh, PhiTolerance, DistanceTolerance, QuadraticTransform);
 
-			if (const FImplicitSphere3* Sphere = Quadratic.template GetObject<FImplicitSphere3>())
-			{
-				ConstructConvexMeshOneShotManifold(*Sphere, QuadraticTransform, *Mesh, MeshTransform, MeshScale, CullDistance, MeshContacts);
-			}
-			else if (const FImplicitCapsule3* Capsule = Quadratic.template GetObject<FImplicitCapsule3>())
+			if (bChaos_Collision_EnableMeshManifoldOptimizedLoop)
 			{
-				ConstructConvexMeshOneShotManifold(*Capsule, QuadraticTransform, *Mesh, MeshTransform, MeshScale, CullDistance, MeshContacts);
+				// New version uses a two-pass loop over triangles to avoid visiting triangles whose vertices are all colliding as a result of checking adjacent triangles
+				Private::FMeshContactGeneratorSettings ContactGeneratorSettings;
+				ContactGeneratorSettings.FaceNormalDotThreshold = 0.9999;	// ~0.8deg Normals must be accurate or rolling will not work correctly
+				ContactGeneratorSettings.bUseTwoPassLoop = false;			// two-pass loop is not helpful for capsules and spheres
+				ContactGeneratorSettings.bSortByPhi = bChaos_Collision_ConvexTriMeshSortByPhi;
+				Private::FMeshContactGenerator ContactGenerator(ContactGeneratorSettings);
+
+				if (const FImplicitSphere3* Sphere = Quadratic.template GetObject<FImplicitSphere3>())
+				{
+					ConstructConvexMeshOneShotManifold2(*Sphere, QuadraticTransform, *Mesh, MeshTransform, MeshScale, CullDistance, ContactGenerator);
+				}
+				else if (const FImplicitCapsule3* Capsule = Quadratic.template GetObject<FImplicitCapsule3>())
+				{
+					ConstructConvexMeshOneShotManifold2(*Capsule, QuadraticTransform, *Mesh, MeshTransform, MeshScale, CullDistance, ContactGenerator);
+				}
+				else
+				{
+					check(false);
+				}
+
+				Constraint.SetOneShotManifoldContacts(ContactGenerator.GetContactPoints());
 			}
 			else
 			{
-				check(false);
-			}
+				const FReal PhiTolerance = CalculateTriMeshPhiTolerance(CullDistance);
+				const FReal DistanceTolerance = Chaos_Collision_TriMeshDistanceTolerance;
+				FContactTriangleCollector MeshContacts(bChaos_Collision_OneSidedTriangleMesh, PhiTolerance, DistanceTolerance, QuadraticTransform);
+
+				if (const FImplicitSphere3* Sphere = Quadratic.template GetObject<FImplicitSphere3>())
+				{
+					ConstructConvexMeshOneShotManifold(*Sphere, QuadraticTransform, *Mesh, MeshTransform, MeshScale, CullDistance, MeshContacts);
+				}
+				else if (const FImplicitCapsule3* Capsule = Quadratic.template GetObject<FImplicitCapsule3>())
+				{
+					ConstructConvexMeshOneShotManifold(*Capsule, QuadraticTransform, *Mesh, MeshTransform, MeshScale, CullDistance, MeshContacts);
+				}
+				else
+				{
+					check(false);
+				}
 
-			Constraint.SetOneShotManifoldContacts(MeshContacts.GetContactPoints());
+				Constraint.SetOneShotManifoldContacts(MeshContacts.GetContactPoints());
+			}
 		}
 
 		void ConstructQuadraticConvexHeightFieldOneShotManifold(const FImplicitObject& Quadratic, const FRigidTransform3& QuadraticTransform, const FHeightField& Mesh, const FRigidTransform3& MeshTransform, const FReal Dt, FPBDCollisionConstraint& Constraint)
@@ -788,8 +818,6 @@ namespace Chaos
 
 			const FVec3 MeshScale = FVec3(1);	// Scale is built into heightfield
 			const FReal CullDistance = Constraint.GetCullDistance();
-			const FReal PhiTolerance = CalculateTriMeshPhiTolerance(CullDistance);
-			const FReal DistanceTolerance = Chaos_Collision_TriMeshDistanceTolerance;
 
 			if (bChaos_Collision_EnableMeshManifoldOptimizedLoop)
 			{
@@ -797,6 +825,7 @@ namespace Chaos
 				Private::FMeshContactGeneratorSettings ContactGeneratorSettings;
 				ContactGeneratorSettings.FaceNormalDotThreshold = 0.9999;	// ~0.8deg Normals must be accurate or rolling will not work correctly
 				ContactGeneratorSettings.bUseTwoPassLoop = false;			// two-pass loop is not helpful for capsules and spheres
+				ContactGeneratorSettings.bSortByPhi = bChaos_Collision_ConvexTriMeshSortByPhi;
 				Private::FMeshContactGenerator ContactGenerator(ContactGeneratorSettings);
 
 				if (const FImplicitSphere3* Sphere = Quadratic.template GetObject<FImplicitSphere3>())
@@ -816,6 +845,8 @@ namespace Chaos
 			}
 			else
 			{
+				const FReal PhiTolerance = CalculateTriMeshPhiTolerance(CullDistance);
+				const FReal DistanceTolerance = Chaos_Collision_TriMeshDistanceTolerance;
 				FContactTriangleCollector MeshContacts(bChaos_Collision_OneSidedHeightField, PhiTolerance, DistanceTolerance, QuadraticTransform);
 
 				if (const FImplicitSphere3* Sphere = Quadratic.template GetObject<FImplicitSphere3>())
@@ -860,6 +891,7 @@ namespace Chaos
 			{
 				const FVec3 RelativeMovement = FVec3(Constraint.GetRelativeMovement());
 				Private::FMeshContactGeneratorSettings ContactGeneratorSettings;
+				ContactGeneratorSettings.bSortByPhi = bChaos_Collision_ConvexTriMeshSortByPhi;
 				Private::FMeshContactGenerator ContactGenerator(ContactGeneratorSettings);
 
 				if (const FImplicitBox3* RawBox = Convex.template GetObject<FImplicitBox3>())
@@ -933,6 +965,7 @@ namespace Chaos
 			{
 				// New version uses a two-pass loop over triangles to avoid visiting triangles whose vertices are all colliding as a result of checking adjacent triangles
 				Private::FMeshContactGeneratorSettings ContactGeneratorSettings;
+				ContactGeneratorSettings.bSortByPhi = bChaos_Collision_ConvexTriMeshSortByPhi;
 				Private::FMeshContactGenerator ContactGenerator(ContactGeneratorSettings);
 				const FVec3 RelativeMovement = FVec3(Constraint.GetRelativeMovement());
 
diff --git a/Engine/Source/Runtime/Experimental/Chaos/Public/Chaos/Collision/MeshContactGenerator.h b/Engine/Source/Runtime/Experimental/Chaos/Public/Chaos/Collision/MeshContactGenerator.h
--- a/Engine/Source/Runtime/Experimental/Chaos/Public/Chaos/Collision/MeshContactGenerator.h
+++ b/Engine/Source/Runtime/Experimental/Chaos/Public/Chaos/Collision/MeshContactGenerator.h
@@ -41,8 +41,11 @@ namespace Chaos::Private
 		// Whether to auto-correct normals
 		uint32 bFixNormals : 1;
 
-		// Whether to sort the contacts to improve solver convergence
-		uint32 bSortForSolver : 1;
+		// Whether to sort the contacts by depth
+		uint32 bSortByPhi : 1;
+
+		// Whether to sort the contacts to improve solver convergence (distance from the center of mass)
+		uint32 bSortForSolverConvergence : 1;
 
 		// Whether to use the optimized two-pass loop over triangles in GenerateMeshContacts which skips triangles
 		// that have contacts on all vertices in the second pass. This is only useful when this case occurs a lot
@@ -469,7 +472,8 @@ namespace Chaos::Private
 		void PruneAndCorrectContacts();
 		void FixContactNormal(const int32 ContactIndex);
 		void RemoveDisabledContacts();
-		void SortContactsForSolver();
+		void SortContactByPhi();
+		void SortContactsForSolverConvergence();
 		void FinalizeContacts(const FRigidTransform3& MeshToConvexTransform);
 
 		void DebugDrawContacts(const FRigidTransform3& ConvexTransform, const FColor& Color, const FReal LineScale);
