Fix for Geometry Component in-editor moving.
Refreshing custom renderer during OnUpdateTransform() happens before the fulll component hierachy has the transform applied. This means that the refreshed intance transform is based against the previous actor root transform and an offset can occur.
Move the refresh to the end of MoveComponentImpl() fixes it.
Submitting to unblock users, but will also investigate if using local transform space instead or world space is better.

diff --git a/Engine/Source/Runtime/Experimental/GeometryCollectionEngine/Private/GeometryCollection/GeometryCollectionComponent.cpp b/Engine/Source/Runtime/Experimental/GeometryCollectionEngine/Private/GeometryCollection/GeometryCollectionComponent.cpp
--- a/Engine/Source/Runtime/Experimental/GeometryCollectionEngine/Private/GeometryCollection/GeometryCollectionComponent.cpp
+++ b/Engine/Source/Runtime/Experimental/GeometryCollectionEngine/Private/GeometryCollection/GeometryCollectionComponent.cpp
@@ -3264,6 +3264,21 @@ FGeometryCollectionDynamicData* UGeometryCollectionComponent::InitDynamicData(bo
 	return DynamicData;
 }
 
+bool UGeometryCollectionComponent::MoveComponentImpl(const FVector& Delta, const FQuat& NewRotation, bool bSweep, FHitResult* Hit, EMoveComponentFlags MoveFlags, ETeleportType Teleport)
+{
+	const bool bResult = Super::MoveComponentImpl(Delta, NewRotation, bSweep, Hit, MoveFlags, Teleport);
+
+#if WITH_EDITOR
+	// Make sure that custom renderer is updated _after_ any move has been applied to the full component hierachy.
+	if (!GetWorld()->IsGameWorld())
+	{
+		RefreshCustomRenderer();
+	}
+#endif
+
+	return bResult;
+}
+
 void UGeometryCollectionComponent::OnUpdateTransform(EUpdateTransformFlags UpdateTransformFlags, ETeleportType Teleport)
 {
 	// Don't allow the primitive component to update physics as that does processing to set a transform for
@@ -3275,12 +3290,6 @@ void UGeometryCollectionComponent::OnUpdateTransform(EUpdateTransformFlags Updat
 	{
 		PhysicsProxy->SetWorldTransform_External(GetComponentTransform());
 	}
-
-	// Make sure that custom renderer is updated at in editor.
-	if (!GetWorld()->IsGameWorld())
-	{
-		RefreshCustomRenderer();
-	}
 }
 
 bool UGeometryCollectionComponent::HasAnySockets() const
@@ -3798,11 +3807,13 @@ void UGeometryCollectionComponent::OnCreatePhysicsState()
 			}
 		}
 
+#if WITH_EDITOR
 		// Make sure that custom renderer is updated at least once in editor.
 		if (!GetWorld()->IsGameWorld())
 		{
 			RefreshCustomRenderer();
 		}
+#endif
 	}
 }
 
diff --git a/Engine/Source/Runtime/Experimental/GeometryCollectionEngine/Public/GeometryCollection/GeometryCollectionComponent.h b/Engine/Source/Runtime/Experimental/GeometryCollectionEngine/Public/GeometryCollection/GeometryCollectionComponent.h
--- a/Engine/Source/Runtime/Experimental/GeometryCollectionEngine/Public/GeometryCollection/GeometryCollectionComponent.h
+++ b/Engine/Source/Runtime/Experimental/GeometryCollectionEngine/Public/GeometryCollection/GeometryCollectionComponent.h
@@ -1188,6 +1188,7 @@ public:
 	GEOMETRYCOLLECTIONENGINE_API virtual bool ShouldCreatePhysicsState() const override;
 	GEOMETRYCOLLECTIONENGINE_API virtual bool HasValidPhysicsState() const override;
 
+	GEOMETRYCOLLECTIONENGINE_API virtual bool MoveComponentImpl(const FVector& Delta, const FQuat& NewRotation, bool bSweep, FHitResult* Hit, EMoveComponentFlags MoveFlags, ETeleportType Teleport) override;
 	GEOMETRYCOLLECTIONENGINE_API virtual void OnUpdateTransform(EUpdateTransformFlags UpdateTransformFlags, ETeleportType Teleport = ETeleportType::None) override;
 
 	// Mirrored from the proxy on a sync
