Ruff draft for downloading I/O store chunks from CDN to local storage (IAD).

This change adds support for precaching/downloading on-demand container file(s)
to local storage. TOC files can now be mounted for:
  - installation/downloading
  - streaming (IAS)
  - or just mounted as a dependency for other TOC/container file(s)

diff --git a/Engine/Source/Runtime/Experimental/IoStore/OnDemand/Private/IoStoreOnDemand.cpp b/Engine/Source/Runtime/Experimental/IoStore/OnDemand/Private/IoStoreOnDemand.cpp
--- a/Engine/Source/Runtime/Experimental/IoStore/OnDemand/Private/IoStoreOnDemand.cpp
+++ b/Engine/Source/Runtime/Experimental/IoStore/OnDemand/Private/IoStoreOnDemand.cpp
@@ -2535,12 +2535,27 @@ void FIoStoreOnDemandModule::ReportAnalytics(TArray<FAnalyticsEventAttribute>& O
 
 void FIoStoreOnDemandModule::Mount(FOnDemandMountArgs&& Args, FOnDemandMountCompleted&& OnCompleted)
 {
-	OnCompleted(TIoStatusOr<FOnDemandMountResult>(FIoStatus(EIoErrorCode::InvalidCode, TEXT("Mount not implemented"))));
+	if (IoStore.IsValid() == false)
+	{
+		IoStore = MakeUnique<FOnDemandIoStore>();
+		if (FIoStatus Status = IoStore->Initialize(); !Status.IsOk())
+		{
+			UE_LOG(LogIas, Error, TEXT("Failed to initialize I/O store on-demand, reason '%s'"), *Status.ToString());
+			IoStore.Reset();
+			return OnCompleted(TIoStatusOr<FOnDemandMountResult>(Status));
+		}
+	}
+
+	IoStore->Mount(MoveTemp(Args), MoveTemp(OnCompleted));
 }
 
 FIoStatus FIoStoreOnDemandModule::Unmount(FStringView MountId)
 {
-	return FIoStatus(EIoErrorCode::InvalidCode, TEXT("Unmount not implemented"));
+	if (IoStore.IsValid())
+	{
+		return IoStore->Unmount(MountId);
+	}
+	return FIoStatus(EIoErrorCode::InvalidCode, TEXT("I/O store on-demand not initialized"));
 }
 
 void FIoStoreOnDemandModule::InitializeInternal()
@@ -2654,7 +2669,8 @@ void FIoStoreOnDemandModule::InitializeInternal()
 			MountArgs.Emplace(FOnDemandMountArgs
 			{
 				.MountId = EndpointConfig.TocFilePath,
-				.FilePath = EndpointConfig.TocFilePath
+				.FilePath = EndpointConfig.TocFilePath,
+				.Options = EOnDemandMountOptions::StreamOnDemand
 			});
 		}
 	}
@@ -2664,10 +2680,25 @@ void FIoStoreOnDemandModule::InitializeInternal()
 		MountArgs.Emplace(FOnDemandMountArgs
 		{
 			.MountId = TocUrl,
-			.Url = TocUrl
+			.Url = TocUrl,
+			.Options = EOnDemandMountOptions::StreamOnDemand
 		});
 	}
 
+#if !UE_BUILD_SHIPPING
+	if (FParse::Param(FCommandLine::Get(), TEXT("Iad")))
+	{
+		// Temporary switch for testing installation to local storage (IAD)
+		MountArgs.Emplace(FOnDemandMountArgs
+		{
+			.MountId = EndpointConfig.TocFilePath,
+			.Url = EndpointConfig.ServiceUrls[0] / EndpointConfig.TocPath,
+			.FilePath = EndpointConfig.TocFilePath,
+			.Options = EOnDemandMountOptions::Install
+		});
+	}
+#endif
+
 	if (MountArgs)
 	{
 		IoStore->Mount(
diff --git a/Engine/Source/Runtime/Experimental/IoStore/OnDemand/Private/OnDemandCacheIoDispatcherBackend.h b/Engine/Source/Runtime/Experimental/IoStore/OnDemand/Private/OnDemandCacheIoDispatcherBackend.h
deleted file mode 100644
diff --git a/Engine/Source/Runtime/Experimental/IoStore/OnDemand/Private/OnDemandHttpClient.cpp b/Engine/Source/Runtime/Experimental/IoStore/OnDemand/Private/OnDemandHttpClient.cpp
--- a/Engine/Source/Runtime/Experimental/IoStore/OnDemand/Private/OnDemandHttpClient.cpp
+++ b/Engine/Source/Runtime/Experimental/IoStore/OnDemand/Private/OnDemandHttpClient.cpp
@@ -263,4 +263,59 @@ const FString& FHttpClient::GetEndpointUrl(int32 Endpoint)
 	return Endpoint == INDEX_NONE ? None : Config.Endpoints[Endpoint];
 }
 
-} //namespace UE::IoStore
+TIoStatusOr<FIoBuffer> FHttpClient::Get(FAnsiStringView Url, uint32 RetryCount, EHttpRedirects Redirects)
+{
+	using namespace HTTP;
+
+	FEventLoop::FRequestParams Params = FEventLoop::FRequestParams
+	{
+		.bAutoRedirect = Redirects == EHttpRedirects::Follow
+	};
+
+	FEventLoop Loop;
+	FIoBuffer Body;
+	TStringBuilder<128> Reason;
+	uint32 StatusCode = 0;
+	const uint32 MaxAttempts = FMath::Min(RetryCount, 3u);
+
+	for (uint32 Attempt = 0; Attempt <= MaxAttempts; ++Attempt)
+	{
+		Loop.Send(Loop.Request("GET", Url, &Params), [&Body, &Reason, &StatusCode](const FTicketStatus& Status)
+		{
+			if (Status.GetId() == FTicketStatus::EId::Response)
+			{
+				Status.GetResponse().SetDestination(&Body);
+				StatusCode = Status.GetResponse().GetStatusCode();
+				return;
+			}
+
+			if (Status.GetId() == FTicketStatus::EId::Error)
+			{
+				Reason << FString(Status.GetErrorReason());
+			}
+		});
+
+		while (Loop.Tick(-1))
+			;
+
+		if (StatusCode > 199 && StatusCode < 300)
+		{
+			return Body;
+		}
+	}
+
+	if (Reason.Len() == 0)
+	{
+		Reason << TEXT("StatusCode: ") << StatusCode;
+	}
+
+	return FIoStatus(EIoErrorCode::ReadError, Reason.ToView());
+}
+
+TIoStatusOr<FIoBuffer> FHttpClient::Get(FStringView Url, uint32 RetryCount, EHttpRedirects Redirects)
+{
+	auto AnsiUrl = StringCast<ANSICHAR>(Url.GetData(), Url.Len());
+	return Get(AnsiUrl, RetryCount, Redirects);
+}
+
+} // namespace UE::IoStore
diff --git a/Engine/Source/Runtime/Experimental/IoStore/OnDemand/Private/OnDemandHttpClient.h b/Engine/Source/Runtime/Experimental/IoStore/OnDemand/Private/OnDemandHttpClient.h
--- a/Engine/Source/Runtime/Experimental/IoStore/OnDemand/Private/OnDemandHttpClient.h
+++ b/Engine/Source/Runtime/Experimental/IoStore/OnDemand/Private/OnDemandHttpClient.h
@@ -14,6 +14,14 @@
 namespace UE::IoStore
 {
 
+enum class EHttpRedirects
+{
+	/** Redirects will be rejected and handled as failed requests. */
+	Disabled,
+	/** Follow redirects automatically. */
+	Follow
+};
+
 struct FHttpClientConfig
 {
 	TArray<FString> Endpoints;
@@ -43,6 +51,9 @@ public:
 	void SetEndpoint(int32 Endpoint);
 	bool IsUsingPrimaryEndpoint() const { return CurrentEndpoint == Config.PrimaryEndpoint; }
 
+	static TIoStatusOr<FIoBuffer> Get(FAnsiStringView Url, uint32 RetryCount, EHttpRedirects Redirects);
+	static TIoStatusOr<FIoBuffer> Get(FStringView Url, uint32 RetryCount, EHttpRedirects Redirects);
+
 private:
 	struct FRequestParams
 	{
diff --git a/Engine/Source/Runtime/Experimental/IoStore/OnDemand/Private/OnDemandInstallCache.cpp b/Engine/Source/Runtime/Experimental/IoStore/OnDemand/Private/OnDemandInstallCache.cpp
new file mode 100644
--- /dev/null
+++ b/Engine/Source/Runtime/Experimental/IoStore/OnDemand/Private/OnDemandInstallCache.cpp
@@ -0,0 +1,337 @@
+// Copyright Epic Games, Inc. All Rights Reserved.
+
+#include "OnDemandInstallCache.h"
+#include "OnDemandHttpClient.h"
+#include "OnDemandIoStore.h"
+
+#include "Async/Mutex.h"
+#include "Async/UniqueLock.h"
+#include "HAL/FileManager.h"
+#include "HAL/PlatformFile.h"
+#include "HAL/PlatformFileManager.h"
+#include "IO/IoContainerHeader.h"
+#include "IO/IoChunkEncoding.h"
+#include "Misc/PathViews.h"
+#include "Serialization/MemoryReader.h"
+
+namespace UE::IoStore
+{
+
+///////////////////////////////////////////////////////////////////////////////
+class FOnDemandInstallCache final
+	: public IOnDemandInstallCache 
+{
+	using FSharedBackendContextRef	= TSharedRef<const FIoDispatcherBackendContext>;
+	using FSharedBackendContext		= TSharedPtr<const FIoDispatcherBackendContext>;
+
+	struct FChunkRequest
+	{
+		explicit FChunkRequest(
+			FIoRequestImpl* Request,
+			FOnDemandChunkInfo&& Info,
+			FIoOffsetAndLength Range,
+			uint64 RequestedRawSize)
+				: DispatcherRequest(Request)
+				, ChunkInfo(MoveTemp(Info))
+				, ChunkRange(Range)
+				, EncodedChunk(ChunkRange.GetLength())
+				, RawSize(RequestedRawSize)
+		{
+			check(DispatcherRequest != nullptr);
+			check(ChunkInfo.IsValid());
+			check(Request->NextRequest == nullptr);
+			check(Request->BackendData == nullptr);
+
+			DispatcherRequest->BackendData = this;
+		}
+
+		FIoRequestImpl*		DispatcherRequest;
+		FOnDemandChunkInfo	ChunkInfo;
+		FIoOffsetAndLength 	ChunkRange;
+		FIoBuffer			EncodedChunk;
+		uint64				RawSize;
+	};
+
+public:
+	FOnDemandInstallCache(const FOnDemandInstallCacheConfig& Config, FOnDemandIoStore& IoStore);
+	virtual ~FOnDemandInstallCache();
+
+	// IIoDispatcherBackend
+	virtual void				Initialize(FSharedBackendContextRef Context) override;
+	virtual void				Shutdown() override;
+	virtual void				ResolveIoRequests(FIoRequestList Requests, FIoRequestList& OutUnresolved) override;
+	virtual FIoRequestImpl*		GetCompletedIoRequests() override;
+	virtual void				CancelIoRequest(FIoRequestImpl* Request) override;
+	virtual void				UpdatePriorityForIoRequest(FIoRequestImpl* Request) override;
+	virtual bool				DoesChunkExist(const FIoChunkId& ChunkId) const override;
+	virtual TIoStatusOr<uint64> GetSizeForChunk(const FIoChunkId& ChunkId) const override;
+	virtual TIoStatusOr<FIoMappedRegion> OpenMapped(const FIoChunkId& ChunkId, const FIoReadOptions& Options) override;
+
+	// IOnDemandInstallCache
+	virtual FIoStatus Put(const FIoChunkId& ChunkId, FIoBuffer&& Chunk, const FIoHash& Hash) override;
+
+private:
+	bool						Resolve(FIoRequestImpl* Request);
+	void						CompleteRequest(TUniquePtr<FChunkRequest>&& ChunkRequest);
+	FIoStatus					WriteChunk(FIoBuffer Chunk, const FIoHash& ExpectedHash);
+	void						GetChunkFilename(const FIoHash& Hash, FStringBuilderBase& OutPath);
+
+	FOnDemandIoStore&		IoStore;
+	FSharedBackendContext	BackendContext;
+	FIoRequestList			CompletedRequests;
+	UE::FMutex				Mutex;
+	FString					CacheDirectory;
+};
+
+///////////////////////////////////////////////////////////////////////////////
+FOnDemandInstallCache::FOnDemandInstallCache(const FOnDemandInstallCacheConfig& Config, FOnDemandIoStore& InIoStore)
+	: IoStore(InIoStore)
+	, CacheDirectory(Config.RootDirectory)
+{
+}
+
+FOnDemandInstallCache::~FOnDemandInstallCache()
+{
+}
+
+void FOnDemandInstallCache::Initialize(FSharedBackendContextRef Context)
+{
+	BackendContext = Context;
+}
+
+void FOnDemandInstallCache::Shutdown()
+{
+}
+
+void FOnDemandInstallCache::ResolveIoRequests(FIoRequestList Requests, FIoRequestList& OutUnresolved)
+{
+	while (FIoRequestImpl* Request = Requests.PopHead())
+	{
+		if (Resolve(Request) == false)
+		{
+			OutUnresolved.AddTail(Request);
+		}
+	}
+}
+
+FIoRequestImpl* FOnDemandInstallCache::GetCompletedIoRequests()
+{
+	FIoRequestImpl* FirstCompleted = nullptr;
+	{
+		UE::TUniqueLock Lock(Mutex);
+		FirstCompleted = CompletedRequests.GetHead();
+		CompletedRequests = FIoRequestList();
+	}
+
+	return FirstCompleted;
+}
+
+void FOnDemandInstallCache::CancelIoRequest(FIoRequestImpl* Request)
+{
+}
+
+void FOnDemandInstallCache::UpdatePriorityForIoRequest(FIoRequestImpl* Request)
+{
+}
+
+bool FOnDemandInstallCache::DoesChunkExist(const FIoChunkId& ChunkId) const
+{
+	const TIoStatusOr<uint64> Status = GetSizeForChunk(ChunkId);
+	return Status.IsOk();
+}
+
+TIoStatusOr<uint64> FOnDemandInstallCache::GetSizeForChunk(const FIoChunkId& ChunkId) const
+{
+	if (FOnDemandChunkInfo ChunkInfo = IoStore.GetInstalledChunkInfo(ChunkId))
+	{
+		return ChunkInfo.RawSize();
+	}
+
+	return FIoStatus(EIoErrorCode::UnknownChunkID);
+}
+
+TIoStatusOr<FIoMappedRegion> FOnDemandInstallCache::OpenMapped(const FIoChunkId& ChunkId, const FIoReadOptions& Options)
+{
+	return FIoStatus(EIoErrorCode::FileOpenFailed);
+}
+
+bool FOnDemandInstallCache::Resolve(FIoRequestImpl* Request)
+{
+	FOnDemandChunkInfo ChunkInfo = IoStore.GetInstalledChunkInfo(Request->ChunkId);
+	if (ChunkInfo.IsValid() == false)
+	{
+		return false;
+	}
+
+	const uint64 RequestSize = FMath::Min<uint64>(
+		Request->Options.GetSize(),
+		ChunkInfo.RawSize() - Request->Options.GetOffset());
+
+	TIoStatusOr<FIoOffsetAndLength> ChunkRange = FIoChunkEncoding::GetChunkRange(
+		ChunkInfo.RawSize(),
+		ChunkInfo.BlockSize(),
+		ChunkInfo.Blocks(),
+		Request->Options.GetOffset(),
+		RequestSize);
+
+	if (ChunkRange.IsOk() == false)
+	{
+		UE_LOG(LogIoStoreOnDemand, Error, TEXT("Failed to get chunk range"));
+		return false;
+	}
+
+	TStringBuilder<256> Filename;
+	GetChunkFilename(ChunkInfo.Hash(), Filename);
+
+	IPlatformFile& Ipf = FPlatformFileManager::Get().GetPlatformFile();
+	TUniquePtr<IFileHandle> FileHandle(Ipf.OpenRead(Filename.ToString()));
+	if (FileHandle.IsValid() == false)
+	{
+		return false;
+	}
+
+	const FIoOffsetAndLength Range = ChunkRange.ConsumeValueOrDie();
+	if (!FileHandle->Seek(IntCastChecked<int64>(Range.GetOffset())))
+	{
+		UE_LOG(LogIoStoreOnDemand, Error, TEXT("Seek file failed"));
+		return false;
+	}
+
+	TUniquePtr<FChunkRequest> ChunkRequest = MakeUnique<FChunkRequest>(
+		Request,
+		MoveTemp(ChunkInfo),
+		Range,
+		RequestSize);
+
+	FIoBuffer& EncodedChunk = ChunkRequest->EncodedChunk;
+	if (!FileHandle->Read(reinterpret_cast<uint8*>(EncodedChunk.GetData()), IntCastChecked<int64>(EncodedChunk.GetSize())))
+	{
+		UE_LOG(LogIoStoreOnDemand, Error, TEXT("Failed to read '%s'"), Filename.ToString());
+		return false;
+	}
+
+	UE::Tasks::Launch(UE_SOURCE_LOCATION, [this, ChunkRequest = MoveTemp(ChunkRequest)]() mutable
+	{
+		CompleteRequest(MoveTemp(ChunkRequest));
+	});
+
+	return true;
+}
+
+FIoStatus FOnDemandInstallCache::Put(const FIoChunkId& ChunkId, FIoBuffer&& Chunk, const FIoHash& Hash)
+{
+	return WriteChunk(Chunk, Hash);
+}
+
+FIoStatus FOnDemandInstallCache::WriteChunk(FIoBuffer Chunk, const FIoHash& ExpectedHash)
+{
+	const FIoHash ChunkHash = FIoHash::HashBuffer(Chunk.GetView());
+
+	if (ChunkHash != ExpectedHash)
+	{
+		return FIoStatus(EIoErrorCode::ReadError, TEXTVIEW("Hash mismatch"));
+	}
+
+	TStringBuilder<256> Filename;
+	GetChunkFilename(ChunkHash, Filename);
+	const FString Directory = FString(FPathViews::GetPath(Filename));
+
+	const bool bTree = true;
+	if (IFileManager& Ifm = IFileManager::Get(); !Ifm.MakeDirectory(*Directory, bTree))
+	{
+		return FIoStatusBuilder(EIoErrorCode::WriteError)
+			<< TEXT("Failed to create directory '")
+			<< Directory
+			<< TEXT("'");
+	}
+
+	IPlatformFile& Ipf = FPlatformFileManager::Get().GetPlatformFile();
+	TUniquePtr<IFileHandle> FileHandle(Ipf.OpenWrite(Filename.ToString()));
+
+	if (FileHandle.IsValid() == false)
+	{
+		return EIoErrorCode::FileOpenFailed;
+	}
+
+	UE_LOG(LogIoStoreOnDemand, Log, TEXT("Writing '%s' %.2lf KiB"), Filename.ToString(), double(Chunk.GetSize()) / 1024.0);
+	if (!FileHandle->Write(reinterpret_cast<const uint8*>(Chunk.GetData()), Chunk.GetSize()))
+	{
+		return FIoStatusBuilder(EIoErrorCode::WriteError)
+			<< TEXT("Failed to write file '")
+			<< Filename.ToString()
+			<< TEXT("'");
+	}
+
+	return EIoErrorCode::Ok;
+}
+
+void FOnDemandInstallCache::GetChunkFilename(const FIoHash& Hash, FStringBuilderBase& OutPath)
+{
+	const FString HashString = LexToString(Hash);
+
+	FPathViews::Append(OutPath, CacheDirectory);
+	FPathViews::Append(OutPath, TEXTVIEW("chunks"));
+	FPathViews::Append(OutPath, HashString.Left(2));
+	FPathViews::Append(OutPath, HashString);
+	OutPath << TEXT(".iochunk");
+}
+
+void FOnDemandInstallCache::CompleteRequest(TUniquePtr<FChunkRequest>&& ChunkRequest)
+{
+	FIoRequestImpl* Request				= ChunkRequest->DispatcherRequest;
+	const FOnDemandChunkInfo& ChunkInfo = ChunkRequest->ChunkInfo;
+
+	FIoChunkDecodingParams Params;
+	Params.CompressionFormat	= ChunkInfo.CompressionFormat();
+	Params.EncryptionKey		= ChunkInfo.EncryptionKey();
+	Params.BlockSize			= ChunkInfo.BlockSize();
+	Params.TotalRawSize			= ChunkInfo.RawSize();
+	Params.RawOffset			= Request->Options.GetOffset();
+	Params.EncodedOffset		= ChunkRequest->ChunkRange.GetOffset();
+	Params.EncodedBlockSize		= ChunkInfo.Blocks();
+	Params.BlockHash			= ChunkInfo.BlockHashes();
+
+	Request->CreateBuffer(ChunkRequest->RawSize);
+
+	FMutableMemoryView RawChunk = Request->GetBuffer().GetMutableView();
+	FMemoryView EncodedChunk	= ChunkRequest->EncodedChunk.GetView();
+
+	if (FIoChunkEncoding::Decode(Params, EncodedChunk, RawChunk) == false)
+	{
+		UE_LOG(LogIoStoreOnDemand, Error, TEXT("Failed to read chunk, ChunkId='%s'"), *LexToString(Request->ChunkId));
+		Request->SetResult(FIoBuffer());
+		Request->SetFailed();
+	}
+
+	{
+		UE::TUniqueLock Lock(Mutex);
+		Request->BackendData = nullptr;
+		CompletedRequests.AddTail(Request);
+	}
+
+	BackendContext->WakeUpDispatcherThreadDelegate.Execute();
+}
+
+///////////////////////////////////////////////////////////////////////////////
+TSharedPtr<IOnDemandInstallCache> MakeOnDemandInstallCache(
+	FOnDemandIoStore& IoStore,
+	const FOnDemandInstallCacheConfig& Config)
+{
+	IFileManager& Ifm = IFileManager::Get();
+	if (Config.bDropCache)
+	{
+		UE_LOG(LogIoStoreOnDemand, Log, TEXT("Deleting install cache directory '%s'"), *Config.RootDirectory);
+		Ifm.DeleteDirectory(*Config.RootDirectory, false, true);
+	}
+
+	const bool bTree = true;
+	if (!Ifm.MakeDirectory(*Config.RootDirectory, bTree))
+	{
+		UE_LOG(LogIoStoreOnDemand, Error, TEXT("Failed to create directory '%s'"), *Config.RootDirectory);
+		return TSharedPtr<IOnDemandInstallCache>();
+	}
+
+	return MakeShareable<IOnDemandInstallCache>(new FOnDemandInstallCache(Config, IoStore));
+}
+
+} // namespace UE::IoStore
diff --git a/Engine/Source/Runtime/Experimental/IoStore/OnDemand/Private/OnDemandInstallCache.h b/Engine/Source/Runtime/Experimental/IoStore/OnDemand/Private/OnDemandInstallCache.h
new file mode 100644
--- /dev/null
+++ b/Engine/Source/Runtime/Experimental/IoStore/OnDemand/Private/OnDemandInstallCache.h
@@ -0,0 +1,35 @@
+// Copyright Epic Games, Inc. All Rights Reserved.
+
+#pragma once
+
+#include "IO/IoDispatcherBackend.h"
+#include "Templates/SharedPointer.h"
+#include "Containers/StringFwd.h"
+
+struct FIoContainerHeader; 
+
+namespace UE::IoStore
+{
+
+class FOnDemandIoStore;
+struct FOnDemandContainer;
+
+class IOnDemandInstallCache 
+	: public IIoDispatcherBackend
+{
+public:
+	virtual ~IOnDemandInstallCache() = default;
+	virtual FIoStatus Put(const FIoChunkId& ChunkId, FIoBuffer&& Chunk, const FIoHash& Hash) = 0;
+};
+
+struct FOnDemandInstallCacheConfig
+{
+	FString RootDirectory;
+	bool bDropCache = false;
+};
+
+TSharedPtr<IOnDemandInstallCache> MakeOnDemandInstallCache(
+	FOnDemandIoStore& IoStore,
+	const FOnDemandInstallCacheConfig& Config);
+
+} // namespace UE::IoStore
diff --git a/Engine/Source/Runtime/Experimental/IoStore/OnDemand/Private/OnDemandIoDispatcherBackend.cpp b/Engine/Source/Runtime/Experimental/IoStore/OnDemand/Private/OnDemandIoDispatcherBackend.cpp
--- a/Engine/Source/Runtime/Experimental/IoStore/OnDemand/Private/OnDemandIoDispatcherBackend.cpp
+++ b/Engine/Source/Runtime/Experimental/IoStore/OnDemand/Private/OnDemandIoDispatcherBackend.cpp
@@ -1258,7 +1258,7 @@ bool FOnDemandIoBackend::Resolve(FIoRequestImpl* Request)
 {
 	using namespace UE::Tasks;
 
-	FOnDemandChunkInfo ChunkInfo = IoStore.GetChunkInfo(Request->ChunkId);
+	FOnDemandChunkInfo ChunkInfo = IoStore.GetStreamingChunkInfo(Request->ChunkId);
 	if (!ChunkInfo.IsValid())
 	{
 		return false;
@@ -1387,7 +1387,7 @@ TIoStatusOr<uint64> FOnDemandIoBackend::GetSizeForChunk(const FIoChunkId& ChunkI
 {
 	OutAvailable = 0;
 
-	const FOnDemandChunkInfo ChunkInfo = IoStore.GetChunkInfo(ChunkId);
+	const FOnDemandChunkInfo ChunkInfo = IoStore.GetStreamingChunkInfo(ChunkId);
 	if (ChunkInfo.IsValid() == false)
 	{
 		return FIoStatus(EIoErrorCode::UnknownChunkID);
diff --git a/Engine/Source/Runtime/Experimental/IoStore/OnDemand/Private/OnDemandIoStore.cpp b/Engine/Source/Runtime/Experimental/IoStore/OnDemand/Private/OnDemandIoStore.cpp
--- a/Engine/Source/Runtime/Experimental/IoStore/OnDemand/Private/OnDemandIoStore.cpp
+++ b/Engine/Source/Runtime/Experimental/IoStore/OnDemand/Private/OnDemandIoStore.cpp
@@ -2,6 +2,8 @@
 
 #include "OnDemandIoStore.h"
 #include "OnDemandHttpClient.h"
+#include "OnDemandInstallCache.h"
+#include "OnDemandPackageStoreBackend.h"
 
 #include "Algo/RemoveIf.h"
 #include "Algo/Transform.h"
@@ -14,6 +16,9 @@
 #include "HAL/PlatformFileManager.h"
 #include "HAL/PlatformTime.h"
 #include "HAL/PlatformProcess.h"
+#include "IO/IoContainerHeader.h"
+#include "IO/PackageStore.h"
+#include "Misc/CommandLine.h"
 #include "Misc/CoreDelegatesInternal.h"
 #include "Misc/EncryptionKeyManager.h"
 #include "Misc/Guid.h"
@@ -21,6 +26,15 @@
 #include "Misc/PathViews.h"
 #include "Serialization/MemoryReader.h"
 
+///////////////////////////////////////////////////////////////////////////////
+bool GIoStoreOnDemandInstallCacheEnabled = true;
+static FAutoConsoleVariableRef CVar_IoStoreOnDemandInstallCacheEnabled (
+	TEXT("iostore.OnDemandInstallCacheEnabled"),
+	GIoStoreOnDemandInstallCacheEnabled,
+	TEXT("Whether the on-demand install cache is enabled."),
+	ECVF_ReadOnly
+);
+
 namespace UE::IoStore
 {
 
@@ -30,6 +44,11 @@ extern FString GIasOnDemandTocExt;
 namespace Private
 {
 
+static FString GetInstallCacheDirectory()
+{
+	return FPaths::ProjectPersistentDownloadDir() / TEXT("IoStore") / TEXT("InstallCache");
+}
+
 static void SplitHostUrl(const FStringView& Url, FStringView& OutHost, FStringView& OutRemainder)
 {
 	OutHost = OutRemainder = FStringView();
@@ -76,9 +95,54 @@ static TUniquePtr<FArchive> CreateReaderFromPlatformPackage(const FString& RelPa
 
 } // namespace UE::IoStore::Private
 
+///////////////////////////////////////////////////////////////////////////////
+void LexToString(EOnDemandContainerFlags Flags, FStringBuilderBase& Out)
+{
+	static const TCHAR* Names[]
+	{
+		TEXT("None"),
+		TEXT("PendingEncryptionKey"),
+		TEXT("Mounted"),
+		TEXT("Streaming"),
+		TEXT("Installed")
+	};
+
+	if (Flags == EOnDemandContainerFlags::None)
+	{
+		Out << TEXT("None");
+		return;
+	}
+
+	for (int32 Idx = 0, Count = int32(EOnDemandContainerFlags::Count); Idx < Count; ++Idx)
+	{
+		const EOnDemandContainerFlags FlagToTest = static_cast<EOnDemandContainerFlags>(1 << Idx);
+		if (EnumHasAnyFlags(Flags, FlagToTest))
+		{
+			if (Out.Len())
+			{
+				Out << TEXT("|");
+			}
+			Out << Names[Idx + 1];
+		}
+	}
+}
+
+FString LexToString(EOnDemandContainerFlags Flags)
+{
+	TStringBuilder<128> Sb;
+	LexToString(Flags, Sb);
+	return FString(Sb.ToString(), Sb.Len());
+}
+
 ///////////////////////////////////////////////////////////////////////////////
 const FOnDemandChunkEntry FOnDemandChunkEntry::Null = {};
 
+///////////////////////////////////////////////////////////////////////////////
+FString FOnDemandContainer::UniqueName() const
+{
+	return FString::Printf(TEXT("%s-%s"), *MountId, *Name);
+}
+
 ///////////////////////////////////////////////////////////////////////////////
 FOnDemandIoStore::FOnDemandIoStore()
 {
@@ -102,6 +166,39 @@ FOnDemandIoStore::~FOnDemandIoStore()
 
 FIoStatus FOnDemandIoStore::Initialize()
 {
+	bool bUseInstallCache = GIoStoreOnDemandInstallCacheEnabled;
+#if !UE_BUILD_SHIPPING
+	bUseInstallCache = FParse::Param(FCommandLine::Get(), TEXT("NoIAD")) == false;
+#endif
+	if (bUseInstallCache)
+	{
+		FOnDemandInstallCacheConfig CacheConfig;
+		CacheConfig.RootDirectory = Private::GetInstallCacheDirectory(); 
+#if !UE_BUILD_SHIPPING
+		CacheConfig.bDropCache = FParse::Param(FCommandLine::Get(), TEXT("Iad.DropCache"));
+#endif
+		InstallCache = MakeOnDemandInstallCache(*this, CacheConfig);
+		if (InstallCache.IsValid())
+		{
+			int32 BackendPriority = -5; // Lower than file (zero) but higher than streaming backend (-10)
+#if !UE_BUILD_SHIPPING
+			if (FParse::Param(FCommandLine::Get(), TEXT("Iad")))
+			{
+				// Bump the priority to be higher then the file system backend
+				BackendPriority = 5;
+			}
+#endif
+			FIoDispatcher::Get().Mount(InstallCache.ToSharedRef(), BackendPriority);
+			PackageStoreBackend = MakeOnDemandPackageStoreBackend();
+			FPackageStore::Get().Mount(PackageStoreBackend.ToSharedRef());
+		}
+		else
+		{
+			// Only warn until this is properly tested
+			UE_LOG(LogIoStoreOnDemand, Warning, TEXT("Failed to initialize install cache"));
+		}
+	}
+
 	OnMountPakHandle = FCoreInternalDelegates::GetOnPakMountOperation().AddLambda(
 		[this](EMountOperation Operation, const TCHAR* ContainerPath, int32 Order) -> void
 		{
@@ -146,10 +243,10 @@ FIoStatus FOnDemandIoStore::Initialize()
 	);
 
 	{
-		IPlatformFile& Ipf = IPlatformFile::GetPlatformPhysical();
 		FCurrentlyMountedPaksDelegate& Delegate = FCoreInternalDelegates::GetCurrentlyMountedPaksDelegate();
 		if (Delegate.IsBound())
 		{
+			IPlatformFile& Ipf = IPlatformFile::GetPlatformPhysical();
 			TArray<FMountedPakInfo> PakInfo = Delegate.Execute();
 			for (const FMountedPakInfo& Info : PakInfo)
 			{
@@ -172,10 +269,10 @@ FIoStatus FOnDemandIoStore::Initialize()
 					}));
 				}
 			}
-		}
 
-		// Process mount requests synchronously at startup
-		TickLoop();
+			// Process mount requests synchronously at startup
+			TickLoop();
+		}
 	}
 
 	return EIoErrorCode::Ok;
@@ -240,7 +337,17 @@ FIoStatus FOnDemandIoStore::Unmount(FStringView MountId)
 	return EIoErrorCode::Ok;
 }
 
-FOnDemandChunkInfo FOnDemandIoStore::GetChunkInfo(const FIoChunkId& ChunkId)
+FOnDemandChunkInfo FOnDemandIoStore::GetStreamingChunkInfo(const FIoChunkId& ChunkId)
+{
+	return GetChunkInfo(ChunkId, EOnDemandContainerFlags::Mounted | EOnDemandContainerFlags::Streaming);
+}
+
+FOnDemandChunkInfo FOnDemandIoStore::GetInstalledChunkInfo(const FIoChunkId& ChunkId)
+{
+	return GetChunkInfo(ChunkId, EOnDemandContainerFlags::Mounted | EOnDemandContainerFlags::Installed);
+}
+
+FOnDemandChunkInfo FOnDemandIoStore::GetChunkInfo(const FIoChunkId& ChunkId, EOnDemandContainerFlags ContainerFlags)
 {
 	TUniqueLock Lock(ContainerMutex);
 
@@ -323,13 +430,25 @@ void FOnDemandIoStore::Tick()
 			}
 		}
 
+		TStringBuilder<128> Sb;
 		TIoStatusOr<FOnDemandMountResult> MountStatus;
 		if (Status.IsOk())
 		{
 			UE::TUniqueLock Lock(ContainerMutex);
 			for (const FSharedOnDemandContainer& Container : Request->Containers)
 			{
+				if (EnumHasAnyFlags(Request->MountArgs.Options, EOnDemandMountOptions::StreamOnDemand))
+				{
+					EnumAddFlags(Container->Flags, EOnDemandContainerFlags::Streaming);
+				}
+
+				if (EnumHasAnyFlags(Container->Flags, EOnDemandContainerFlags::Mounted))
+				{
+					continue;
+				}
+
 				Containers.Add(Container);
+
 				if (Container->EncryptionKeyGuid.IsEmpty() == false)
 				{
 					FGuid KeyGuid;
@@ -343,8 +462,11 @@ void FOnDemandIoStore::Tick()
 					}
 				}
 
-				UE_LOG(LogIoStoreOnDemand, Log, TEXT("Mounting container '%s' (%d entries)"),
-					*Container->Name, Container->ChunkEntries.Num());
+				Sb.Reset();
+				LexToString(Container->Flags, Sb);
+				UE_LOG(LogIoStoreOnDemand, Log, TEXT("Mounting container '%s', Entries=%d, Flags='%s'"),
+					*Container->Name, Container->ChunkEntries.Num(), Sb.ToString());
+
 				EnumAddFlags(Container->Flags, EOnDemandContainerFlags::Mounted);
 			}
 
@@ -365,13 +487,13 @@ void FOnDemandIoStore::Tick()
 
 FIoStatus FOnDemandIoStore::ProcessMountRequest(FMountRequest& MountRequest)
 {
-	UE_LOG(LogIoStoreOnDemand, Log, TEXT("Processing mount request, MountId='%s'"), *MountRequest.MountArgs.MountId);
+	UE_LOG(LogIoStoreOnDemand, Verbose, TEXT("Processing mount request, MountId='%s'"), *MountRequest.MountArgs.MountId);
 
 	FOnDemandMountArgs& Args = MountRequest.MountArgs;
 	FStringView Host, TocRelUrl;
 	Private::SplitHostUrl(Args.Url, Host, TocRelUrl);
 	const FStringView TocPath = FPathViews::GetPath(TocRelUrl);
-	
+
 	if (Args.Toc.IsSet())
 	{
 		CreateContainersFromToc(Args.MountId, TocPath, Args.Toc.GetValue(), MountRequest.Containers);
@@ -444,42 +566,14 @@ FIoStatus FOnDemandIoStore::ProcessMountRequest(FMountRequest& MountRequest)
 			return FIoStatusBuilder(EIoErrorCode::InvalidParameter) << TEXT("No valid host URL");
 		}
 
-		FHttpClientConfig HttpCfg;
-		HttpCfg.Endpoints.Add(FString(Host));
-		HttpCfg.PrimaryEndpoint = 0; 
-		HttpCfg.MaxRetryCount = 2;
-		TUniquePtr<FHttpClient> Client = FHttpClient::Create(MoveTemp(HttpCfg));
-
-		TAnsiStringBuilder<256> Url;
-		if (!TocRelUrl.StartsWith(TCHAR('/')))
-		{
-			Url << "/";
-		}
-		Url << TocRelUrl;
-
-		FIoBuffer Body;
-		FIoStatus HttpStatus;
-		Client->Get(Url.ToView(), [&Body, &HttpStatus](TIoStatusOr<FIoBuffer> Response, uint64)
-			{
-				if (Response.IsOk())
-				{
-					Body = Response.ConsumeValueOrDie();
-				}
-				else
-				{
-					HttpStatus = Response.Status();
-				}
-			});
-
-		while (Client->Tick());
-
-		if (HttpStatus.IsOk() == false)
+		TIoStatusOr<FIoBuffer> Response = FHttpClient::Get(Args.Url, 2, EHttpRedirects::Follow);
+		if (Response.IsOk() == false)
 		{
-			return HttpStatus;
+			return FIoStatusBuilder(EIoErrorCode::InvalidCode) << TEXT("Failed to fetch TOC from URL");
 		}
 
 		FOnDemandToc Toc;
-		FMemoryReaderView Ar(Body.GetView());
+		FMemoryReaderView Ar(Response.ValueOrDie().GetView());
 		Ar << Toc;
 
 		if (Ar.IsError() || Ar.IsCriticalError())
@@ -489,9 +583,30 @@ FIoStatus FOnDemandIoStore::ProcessMountRequest(FMountRequest& MountRequest)
 
 		CreateContainersFromToc(Args.MountId, TocPath, Toc, MountRequest.Containers);
 	}
-	else
+
+	if (MountRequest.Containers.IsEmpty())
+	{
+		UE::TUniqueLock Lock(ContainerMutex);
+		for (FSharedOnDemandContainer& Container : Containers)
+		{
+			if (Container->MountId == Args.MountId)
+			{
+				MountRequest.Containers.Add(Container);
+			}
+		}
+	}
+
+	if (EnumHasAnyFlags(Args.Options, EOnDemandMountOptions::Install))
 	{
-		return FIoStatusBuilder(EIoErrorCode::InvalidParameter) << TEXT("Invalid mount arguments");
+		if (InstallCache.IsValid() == false || PackageStoreBackend.IsValid() == false)
+		{
+			return FIoStatusBuilder(EIoErrorCode::InvalidCode) << TEXT("Install cache not configured");
+		}
+
+		if (FIoStatus Status = InstallContainers(Args.Url, MountRequest.Containers); !Status.IsOk())
+		{
+			return Status;
+		}
 	}
 
 	return EIoErrorCode::Ok;
@@ -510,8 +625,8 @@ void FOnDemandIoStore::OnEncryptionKeyAdded(const FGuid& Id, const FAES::FAESKey
 
 			if (FEncryptionKeyManager::Get().TryGetKey(KeyGuid, Container->EncryptionKey))
 			{
-				UE_LOG(LogIoStoreOnDemand, Log, TEXT("Mounting container '%s' (%d entries)"),
-					*WriteToString<128>(Container->Name), Container->ChunkEntries.Num());
+				UE_LOG(LogIoStoreOnDemand, Log, TEXT("Mounting container '%s', Entries=%d, Flags='%s'"),
+					*Container->Name, Container->ChunkEntries.Num(), *LexToString(Container->Flags));
 
 				EnumRemoveFlags(Container->Flags, EOnDemandContainerFlags::PendingEncryptionKey);
 				EnumAddFlags(Container->Flags, EOnDemandContainerFlags::Mounted);
@@ -586,4 +701,144 @@ void FOnDemandIoStore::CreateContainersFromToc(
 	}
 }
 
+FIoStatus FOnDemandIoStore::InstallContainers(
+	const FString& Url,
+	const TConstArrayView<FSharedOnDemandContainer>& ContainersToInstall)
+{
+	FStringView Host, TocRelUrl;
+	Private::SplitHostUrl(Url, Host, TocRelUrl);
+	const FStringView TocPath = FPathViews::GetPath(TocRelUrl);
+	auto AnsiTocPath = StringCast<ANSICHAR>(TocPath.GetData(), TocPath.Len());
+
+	TAnsiStringBuilder<512> ChunkUrl;
+	for (const FSharedOnDemandContainer& Container : ContainersToInstall)
+	{
+		if (EnumHasAnyFlags(Container->Flags, EOnDemandContainerFlags::Installed))
+		{
+			UE_LOG(LogIoStoreOnDemand, Log, TEXT("Container '%s' is already installed, MountId='%s'"),
+				*Container->Name, *Container->MountId);
+			continue;
+		}
+
+		auto GetChunkUrl = [&AnsiTocPath, &Container](const FOnDemandChunkEntry& Entry, FAnsiStringBuilderBase& OutUrl) -> FAnsiStringBuilderBase&
+		{
+			OutUrl.Reset();
+			if (Container->ChunksDirectory.IsEmpty())
+			{
+				OutUrl << AnsiTocPath << "/" << "chunks";
+			}
+			else
+			{
+				OutUrl.Append(Container->ChunksDirectory);
+			}
+
+			const FString HashString = LexToString(Entry.Hash);
+			OutUrl << "/" << HashString.Left(2) << "/" << HashString << ".iochunk";
+
+			return OutUrl;
+		};
+
+		UE_LOG(LogIoStoreOnDemand, Log, TEXT("Installing container, Name='%s', MountId='%s'"),
+			*Container->Name, *Container->MountId);
+
+		FIoContainerHeader ContainerHeader;
+
+		const FIoChunkId HeaderChunkId = CreateContainerHeaderChunkId(Container->ContainerId);
+		if (const FOnDemandChunkEntry* Entry = Container->ChunkEntries.Find(HeaderChunkId))
+		{
+			TIoStatusOr<FIoBuffer> Response = FHttpClient::Get(GetChunkUrl(*Entry, ChunkUrl).ToView(), 2, EHttpRedirects::Follow);
+			if (Response.IsOk() == false)
+			{
+				return FIoStatusBuilder(EIoErrorCode::InvalidCode) << TEXT("Failed to fetch container header chunk from URL");
+			}
+
+			FMemoryReaderView Ar(Response.ValueOrDie().GetView());
+			if (Ar.IsError() || Ar.IsCriticalError())
+			{
+				return FIoStatusBuilder(EIoErrorCode::FileNotOpen) << TEXT("Failed to serialize container header");
+			}
+		}
+
+		if (ContainerHeader.PackageIds.Num())
+		{
+			TConstArrayView<FFilePackageStoreEntry> PackageStoreEntries(
+				reinterpret_cast<const FFilePackageStoreEntry*>(ContainerHeader.StoreEntries.GetData()),
+				ContainerHeader.PackageIds.Num());
+
+			//TODO: Find all referenced chunks
+		}
+
+		// Download all chunks
+		const int32 MaxConcurrentRequests = 32;
+		int32 ConcurrentRequests = 0;
+
+		FHttpClientConfig HttpConfig;
+		HttpConfig.MaxConnectionCount = 8;
+		HttpConfig.MaxRetryCount = 2;
+		HttpConfig.Endpoints.Add(Url);
+
+		TUniquePtr<FHttpClient> HttpClient = FHttpClient::Create(MoveTemp(HttpConfig));
+		if (HttpClient.IsValid() == false)
+		{
+			return FIoStatusBuilder(EIoErrorCode::InvalidCode) << TEXT("Failed to initialize HTTP client");
+		}
+
+		FIoStatus Status(EIoErrorCode::Ok);
+		for (const TPair<FIoChunkId, FOnDemandChunkEntry>& Kv : Container->ChunkEntries)
+		{
+			const FIoChunkId& ChunkId			= Kv.Key;
+			const FOnDemandChunkEntry& TocEntry	= Kv.Value;
+			
+			ConcurrentRequests++;
+			HttpClient->Get(
+				GetChunkUrl(TocEntry, ChunkUrl).ToView(),
+				[this, &Status, ChunkId, ExpectedHash = TocEntry.Hash, &ConcurrentRequests]
+				(TIoStatusOr<FIoBuffer> ChunkStatus, uint64 DurationMs)
+				{
+					ConcurrentRequests--;
+					if (ChunkStatus.IsOk())
+					{
+						if (FIoStatus PutStatus = InstallCache->Put(
+							ChunkId,
+							ChunkStatus.ConsumeValueOrDie(),
+							ExpectedHash); !PutStatus.IsOk())
+						{
+							Status = PutStatus;
+						}
+					}
+					else
+					{
+						Status = ChunkStatus.Status();
+					}
+				});
+
+			while (ConcurrentRequests >= MaxConcurrentRequests)
+			{
+				HttpClient->Tick();
+			}
+
+			if (Status.IsOk() == false)
+			{
+				return Status;
+			}
+		}
+
+		while (HttpClient->Tick());
+
+		if (ContainerHeader.PackageIds.IsEmpty() == false)
+		{
+			if (FIoStatus MountStatus = PackageStoreBackend->Mount(
+				Container->UniqueName(),
+				MoveTemp(ContainerHeader)); !MountStatus.IsOk())
+			{
+				return MountStatus;
+			}
+		}
+
+		EnumAddFlags(Container->Flags, EOnDemandContainerFlags::Installed);
+	}
+
+	return EIoErrorCode::Ok;
+}
+
 } // namespace UE::IoStore
diff --git a/Engine/Source/Runtime/Experimental/IoStore/OnDemand/Private/OnDemandIoStore.h b/Engine/Source/Runtime/Experimental/IoStore/OnDemand/Private/OnDemandIoStore.h
--- a/Engine/Source/Runtime/Experimental/IoStore/OnDemand/Private/OnDemandIoStore.h
+++ b/Engine/Source/Runtime/Experimental/IoStore/OnDemand/Private/OnDemandIoStore.h
@@ -17,15 +17,26 @@
 namespace UE::IoStore
 {
 
+class IOnDemandPackageStoreBackend;
+class IOnDemandInstallCache;
+using FSharedPackageStoreBackend	= TSharedPtr<IOnDemandPackageStoreBackend>;
+using FSharedInstallCache			= TSharedPtr<IOnDemandInstallCache>;
+
 ///////////////////////////////////////////////////////////////////////////////
 enum class EOnDemandContainerFlags : uint8
 {
 	None					= 0,
 	PendingEncryptionKey	= (1 << 0),
 	Mounted					= (1 << 1),
+	Streaming				= (1 << 2),
+	Installed				= (1 << 3),
+	Count
 };
 ENUM_CLASS_FLAGS(EOnDemandContainerFlags);
 
+void LexToString(EOnDemandContainerFlags Flags, FStringBuilderBase& Out);
+FString LexToString(EOnDemandContainerFlags Flags);
+
 ///////////////////////////////////////////////////////////////////////////////
 struct FOnDemandChunkEntry
 {
@@ -57,6 +68,8 @@ struct FOnDemandContainer
 	FIoContainerId			ContainerId;
 	uint32					BlockSize = 0;
 	EOnDemandContainerFlags Flags = EOnDemandContainerFlags ::None;
+
+	FString					UniqueName() const;
 };
 
 using FSharedOnDemandContainer = TSharedPtr<FOnDemandContainer, ESPMode::ThreadSafe>;
@@ -130,9 +143,11 @@ public:
 	FIoStatus				Initialize();
 	void					Mount(FOnDemandMountArgs&& Args, FOnDemandMountCompleted&& OnCompleted);
 	FIoStatus				Unmount(FStringView MountId);
-	FOnDemandChunkInfo		GetChunkInfo(const FIoChunkId& ChunkId);
+	FOnDemandChunkInfo		GetStreamingChunkInfo(const FIoChunkId& ChunkId);
+	FOnDemandChunkInfo		GetInstalledChunkInfo(const FIoChunkId& ChunkId);
 
 private:
+	FOnDemandChunkInfo		GetChunkInfo(const FIoChunkId& ChunkId, EOnDemandContainerFlags ContainerFlags);
 	void					TryEnterTickLoop();
 	void					TickLoop();
 	void					Tick();
@@ -144,7 +159,12 @@ private:
 								FStringView TocPath,
 								FOnDemandToc& Toc,
 								TArray<FSharedOnDemandContainer>& Out);
+	FIoStatus				InstallContainers(
+								const FString& Url,
+								const TConstArrayView<FSharedOnDemandContainer>& ContainersToInstall);
 
+	FSharedInstallCache					InstallCache;
+	FSharedPackageStoreBackend			PackageStoreBackend;
 	FDelegateHandle						OnMountPakHandle;
 	TArray<FSharedOnDemandContainer>	Containers;
 	UE::FMutex							ContainerMutex;
diff --git a/Engine/Source/Runtime/Experimental/IoStore/OnDemand/Public/IO/IoStoreOnDemand.h b/Engine/Source/Runtime/Experimental/IoStore/OnDemand/Public/IO/IoStoreOnDemand.h
--- a/Engine/Source/Runtime/Experimental/IoStore/OnDemand/Public/IO/IoStoreOnDemand.h
+++ b/Engine/Source/Runtime/Experimental/IoStore/OnDemand/Public/IO/IoStoreOnDemand.h
@@ -7,6 +7,7 @@
 #include "IO/IoContainerId.h"
 #include "IO/IoHash.h"
 #include "IO/IoStatus.h"
+#include "Misc/EnumClassFlags.h"
 #include "Misc/Guid.h"
 #include "Modules/ModuleInterface.h"
 #include "Templates/SharedPointer.h"
@@ -344,6 +345,18 @@ enum class EOnDemandInitResult
 
 #endif // UE_IAS_CUSTOM_INITIALIZATION
 
+/** Options for controlling the behavior of mount requests. */
+enum class EOnDemandMountOptions
+{
+	/** The TOC is loaded but not installed or available for streaming. */
+	None			= 0,
+	/** Make on-demand container(s) within a TOC available for streaming. */
+	StreamOnDemand	= 1 << 0,
+	/** Download and install on-demand contianer(s) to local storage. */
+	Install			= 1 << 1
+};
+ENUM_CLASS_FLAGS(EOnDemandMountOptions);
+
 struct FOnDemandMountArgs
 {
 	/** Mount an already serialized TOC. */
@@ -356,8 +369,8 @@ struct FOnDemandMountArgs
 	FString FilePath;
 	/** Directory path for any additonal content downloaded as part of the TOC. */
 	FString InstallDirectory;
-	/** Whether to download and install container chunks to local storage. */
-	bool bInstall = true;
+	/** Mount options. */
+	EOnDemandMountOptions Options;
 };
 
 struct FOnDemandMountResult
