Sends on keep alive peers are already negotiated.

#rnx

diff --git a/Engine/Source/Runtime/Experimental/IoStore/HttpClient/Private/Loop.inl b/Engine/Source/Runtime/Experimental/IoStore/HttpClient/Private/Loop.inl
--- a/Engine/Source/Runtime/Experimental/IoStore/HttpClient/Private/Loop.inl
+++ b/Engine/Source/Runtime/Experimental/IoStore/HttpClient/Private/Loop.inl
@@ -807,7 +807,11 @@ void FPeerGroup::TickSend(FTickState& State, FHost& Host)
 
 	if (!bWillBlock)
 	{
-		return Negotiate(State);
+		if (bNegotiating)
+		{
+			return Negotiate(State);
+		}
+		return SendInternal(State);
 	}
 
 	// Non-blocking connect
