Add an Iris stat for tracking the number of RepNotifies on the client.

diff --git a/Engine/Source/Runtime/Experimental/Iris/Core/Private/Iris/IrisConfig.cpp b/Engine/Source/Runtime/Experimental/Iris/Core/Private/Iris/IrisConfig.cpp
--- a/Engine/Source/Runtime/Experimental/Iris/Core/Private/Iris/IrisConfig.cpp
+++ b/Engine/Source/Runtime/Experimental/Iris/Core/Private/Iris/IrisConfig.cpp
@@ -39,4 +39,7 @@ EReplicationSystem GetUseIrisReplicationCmdlineValue()
 
 
 // Enable Iris category by default on servers
-CSV_DEFINE_CATEGORY(Iris, WITH_SERVER_CODE);
\ No newline at end of file
+CSV_DEFINE_CATEGORY(Iris, WITH_SERVER_CODE);
+
+// Enable IrisClient category on the client only
+CSV_DEFINE_CATEGORY(IrisClient, !WITH_SERVER_CODE);
\ No newline at end of file
diff --git a/Engine/Source/Runtime/Experimental/Iris/Core/Private/Iris/ReplicationState/PropertyReplicationState.cpp b/Engine/Source/Runtime/Experimental/Iris/Core/Private/Iris/ReplicationState/PropertyReplicationState.cpp
--- a/Engine/Source/Runtime/Experimental/Iris/Core/Private/Iris/ReplicationState/PropertyReplicationState.cpp
+++ b/Engine/Source/Runtime/Experimental/Iris/Core/Private/Iris/ReplicationState/PropertyReplicationState.cpp
@@ -15,9 +15,12 @@
 #include "UObject/UnrealType.h"
 #include "UObject/PropertyPortFlags.h"
 #include "Containers/StringFwd.h"
+#include "ProfilingDebugging/CsvProfiler.h"
 
 DEFINE_LOG_CATEGORY_STATIC(LogIrisRepNotify, Warning, All);
 
+CSV_DECLARE_CATEGORY_EXTERN(IrisClient);
+
 namespace UE::Net
 {
 
@@ -476,6 +479,8 @@ void FPropertyReplicationState::CallRepNotifies(void* RESTRICT DstData, const FC
 
 				if (bShouldCallRepNotify)
 				{
+					CSV_CUSTOM_STAT(IrisClient, RepNotifyCount, 1, ECsvCustomStatOp::Accumulate);
+
 					// We only want to call RepNotify once for c-arrays
 					LastPropertyWithRepNotify = Property;
 					Object->ProcessEvent(const_cast<UFunction*>(RepNotifyFunction), const_cast<uint8*>(PrevValuePtr));
diff --git a/Engine/Source/Runtime/Experimental/Iris/Core/Private/Iris/ReplicationSystem/ReplicationReader.cpp b/Engine/Source/Runtime/Experimental/Iris/Core/Private/Iris/ReplicationSystem/ReplicationReader.cpp
--- a/Engine/Source/Runtime/Experimental/Iris/Core/Private/Iris/ReplicationSystem/ReplicationReader.cpp
+++ b/Engine/Source/Runtime/Experimental/Iris/Core/Private/Iris/ReplicationSystem/ReplicationReader.cpp
@@ -37,7 +37,7 @@
 
 #define UE_LOG_REPLICATIONREADER_CONN_WARNING(Format, ...) UE_LOG(LogIris, Warning, TEXT("Conn: %u ") Format, Parameters.ConnectionId, ##__VA_ARGS__)
 
-CSV_DEFINE_CATEGORY(IrisClient, true);
+CSV_DECLARE_CATEGORY_EXTERN(IrisClient);
 
 namespace UE::Net::Private
 {
