Iris
*  Reduced log warning spam by only logging warnings once per RPC function
* Reduced RPC log warning when calling RPC on a netobject before BeginReplication was called on it.
* net.Iris.ThrottleRPCWarnings cvar can be set to false to log all occurences of the warning.

#rnx

diff --git a/Engine/Source/Runtime/Experimental/Iris/Core/Private/Iris/ReplicationSystem/NetBlob/NetBlobManager.cpp b/Engine/Source/Runtime/Experimental/Iris/Core/Private/Iris/ReplicationSystem/NetBlob/NetBlobManager.cpp
--- a/Engine/Source/Runtime/Experimental/Iris/Core/Private/Iris/ReplicationSystem/NetBlob/NetBlobManager.cpp
+++ b/Engine/Source/Runtime/Experimental/Iris/Core/Private/Iris/ReplicationSystem/NetBlob/NetBlobManager.cpp
@@ -18,6 +18,9 @@ namespace UE::Net::Private
 
 static TAutoConsoleVariable<int32> CVarEnableIrisRPCs(TEXT("net.Iris.EnableRPCs"), 1, TEXT( "If > 0 let Iris replicate and execute RPCs."));
 
+static bool bThrottleRPCWarnings = true;
+static FAutoConsoleVariableRef CVarThrottleRPCWarnings(TEXT("net.Iris.ThrottleRPCWarnings"), bThrottleRPCWarnings, TEXT("Only log send failure warnings once per RPC type."));
+
 FNetBlobManager::FNetBlobManager()
 {
 }
@@ -195,7 +198,14 @@ bool FNetBlobManager::GetRPCOwner(FRPCOwner& OutOwnerInfo, const UObject* RootOb
 
 		if (!bCanSendRpc)
 		{
-			UE_LOG(LogIris, Warning, TEXT("SendRPC %s for %s Failed. This rootobject is not yet replicated (RefHandle: %s Index: %u)."),
+			bool bLogRPCFailed = true;
+			if (UE::Net::Private::bThrottleRPCWarnings)
+			{
+				bool& bWasAlreadyLogged = RPCWarningThrottler.FindOrAdd(Function->GetFName(), false);
+				bLogRPCFailed = !bWasAlreadyLogged;
+				bWasAlreadyLogged = true;
+			}
+			UE_CLOG(bLogRPCFailed, LogIris, Warning, TEXT("SendRPC %s for %s Failed. This rootobject is not yet replicated (RefHandle: %s Index: %u)."),
 				ToCStr(Function->GetName()), *GetNameSafe(RootObject), ToCStr(OutOwnerInfo.CallerRef.GetRefHandle().ToString()), OutOwnerInfo.RootObjectIndex);
 		}
 	}
@@ -226,7 +236,15 @@ bool FNetBlobManager::GetRPCOwner(FRPCOwner& OutOwnerInfo, const UObject* RootOb
 		
 		if (!bCanSendRpc)
 		{
-			UE_LOG(LogIris, Warning, TEXT("SendRPC %s for %s::%s Failed. The root object (RefHandle: %s Index: %u) and subobject (RefHandle: %s Index: %u) is not yet replicated."),
+			bool bLogRPCFailed = true;
+			if (UE::Net::Private::bThrottleRPCWarnings)
+			{
+				bool& bWasAlreadyLogged = RPCWarningThrottler.FindOrAdd(Function->GetFName(), false);
+				bLogRPCFailed = !bWasAlreadyLogged;
+				bWasAlreadyLogged = true;
+			}
+
+			UE_CLOG(bLogRPCFailed, LogIris, Warning, TEXT("SendRPC %s for %s::%s Failed. The root object (RefHandle: %s Index: %u) and subobject (RefHandle: %s Index: %u) is not yet replicated."),
 				ToCStr(Function->GetName()), *GetNameSafe(RootObject), *GetNameSafe(SubObject), ToCStr(OutOwnerInfo.CallerRef.ToString()), OutOwnerInfo.RootObjectIndex, ToCStr(OutOwnerInfo.TargetRef.ToString()), OutOwnerInfo.SubObjectIndex);
 		}
 	}
diff --git a/Engine/Source/Runtime/Experimental/Iris/Core/Private/Iris/ReplicationSystem/NetBlob/NetBlobManager.h b/Engine/Source/Runtime/Experimental/Iris/Core/Private/Iris/ReplicationSystem/NetBlob/NetBlobManager.h
--- a/Engine/Source/Runtime/Experimental/Iris/Core/Private/Iris/ReplicationSystem/NetBlob/NetBlobManager.h
+++ b/Engine/Source/Runtime/Experimental/Iris/Core/Private/Iris/ReplicationSystem/NetBlob/NetBlobManager.h
@@ -172,6 +172,9 @@ private:
 	TStrongObjectPtr<UPartialNetObjectAttachmentHandler> PartialNetObjectAttachmentHandler;
 	TStrongObjectPtr<UNetObjectBlobHandler> NetObjectBlobHandler;
 
+	/** Track if a warning was already logged for a specific RPC. */
+	mutable TMap<FName, bool> RPCWarningThrottler;
+
 	UReplicationSystem* ReplicationSystem = nullptr;
 	FObjectReferenceCache* ObjectReferenceCache = nullptr;
 	FReplicationConnections* Connections = nullptr;
