Chaos - THashMappedArray API

diff --git a/Engine/Source/Runtime/Experimental/Chaos/Public/Chaos/Framework/HashMappedArray.h b/Engine/Source/Runtime/Experimental/Chaos/Public/Chaos/Framework/HashMappedArray.h
--- a/Engine/Source/Runtime/Experimental/Chaos/Public/Chaos/Framework/HashMappedArray.h
+++ b/Engine/Source/Runtime/Experimental/Chaos/Public/Chaos/Framework/HashMappedArray.h
@@ -70,7 +70,32 @@ namespace Chaos::Private
 			Elements.Reset(InReserveElements);
 		}
 
-		// Add an element with the specified ID to the map. 
+		// Try to add an element with the specified ID to the map. Does nothing if an Element with the same ID is already in the map.
+		// Returns true if the element was added, false if the ID already existed.
+		FORCEINLINE bool TryAdd(const FIDType ID, const FElementType& Element)
+		{
+			if (Find(ID) == nullptr)
+			{
+				Add(ID, Element);
+				return true;
+			}
+			return false;
+		}
+
+		// Try to add an element with the specified ID to the map. Does nothing if an Element with the same ID is already in the map.
+		// Returns true if the element was added, false if the ID already existed.
+		template <typename... ArgsType>
+		FORCEINLINE bool TryEmplace(const FIDType ID, ArgsType&&... Args)
+		{
+			if (Find(ID) == nullptr)
+			{
+				Emplace(ID, Forward<ArgsType>(Args)...);
+				return true;
+			}
+			return false;
+		}
+
+		// Add an element with the specified ID to the map. Asserts if an Element with the same ID is already in the map.
 		FORCEINLINE void Add(const FIDType ID, const FElementType& Element)
 		{
 			checkSlow(Find(ID) == nullptr);
