GeometryCollection : HasAttributes
- Fixed issue with returning true when one, or all, of the attributes are missing.

diff --git a/Engine/Source/Runtime/Experimental/Chaos/Private/GeometryCollection/ManagedArrayCollection.cpp b/Engine/Source/Runtime/Experimental/Chaos/Private/GeometryCollection/ManagedArrayCollection.cpp
--- a/Engine/Source/Runtime/Experimental/Chaos/Private/GeometryCollection/ManagedArrayCollection.cpp
+++ b/Engine/Source/Runtime/Experimental/Chaos/Private/GeometryCollection/ManagedArrayCollection.cpp
@@ -114,6 +114,10 @@ bool FManagedArrayCollection::HasAttributes(const TArray<FManagedArrayCollection
 				return false;
 			}
 		}
+		else
+		{
+			return false;
+		}
 	}
 	return true;
 }
