IASTool Chunk Plugin Command
- Add -IncludeSigPak to make inclusion of pak and sig files optional (default is now to not include them)
- Add -KeepContainerFiles to make deleting containers after chunking optional (default is now to delete containers)

diff --git a/Engine/Source/Runtime/Experimental/IoStore/OnDemand/Private/Tool/ChunkPlugin.cpp b/Engine/Source/Runtime/Experimental/IoStore/OnDemand/Private/Tool/ChunkPlugin.cpp
--- a/Engine/Source/Runtime/Experimental/IoStore/OnDemand/Private/Tool/ChunkPlugin.cpp
+++ b/Engine/Source/Runtime/Experimental/IoStore/OnDemand/Private/Tool/ChunkPlugin.cpp
@@ -71,7 +71,9 @@ static int32 ChunkPluginCommandEntry(const FContext& Context)
 	const FString InputFolder			= FString(Context.Get<FStringView>(TEXT("-InputFolder"), FString()));
 	const FString OutputFolder			= FString(Context.Get<FStringView>(TEXT("-OutputFolder"), FString()));
 	const FString IntermediateFolder	= FString(Context.Get<FStringView>(TEXT("-IntermediateFolder"), FString()));
-	FString ContainerFolder				= IntermediateFolder / TEXT("Cooked");
+	const bool bIncludeSigPak			= Context.Get<bool>(TEXT("-IncludeSigPak"), false);
+	const bool bDeleteContainerFiles	= !Context.Get<bool>(TEXT("-KeepContainerFiles"), false);
+	FString ContainerFolder				= InputFolder;
 	FString IoStoreOutputFolder			= OutputFolder / TEXT("iostore");
 	FString ChunksOutputFolder			= IoStoreOutputFolder / TEXT("chunks");
 
@@ -87,6 +89,8 @@ static int32 ChunkPluginCommandEntry(const FContext& Context)
 	UE_LOG(LogIoStore, Display, TEXT("\tInputFolder: %s"), *InputFolder);
 	UE_LOG(LogIoStore, Display, TEXT("\tOutputFolder: %s"), *OutputFolder);
 	UE_LOG(LogIoStore, Display, TEXT("\tIntermediateFolder: %s"), *IntermediateFolder);
+	UE_LOG(LogIoStore, Display, TEXT("\tIncludeSigPak: %s"), bIncludeSigPak ? TEXT("true") : TEXT("false"));
+	UE_LOG(LogIoStore, Display, TEXT("\tDeleteContainerFiles: %s"), bDeleteContainerFiles ? TEXT("true") : TEXT("false"));
 
 	IFileManager& FileMgr = IFileManager::Get();
 	if (FileMgr.MakeDirectory(*ChunksOutputFolder, true) == false)
@@ -110,6 +114,8 @@ static int32 ChunkPluginCommandEntry(const FContext& Context)
 	//OnDemandToc.Header.ChunksDirectory = TODO 
 	OnDemandToc.Containers.Reserve(ContainerFilenames.Num());
 
+	TArray<FString> FilesToDelete;
+
 	for (const FString& Filename : ContainerFilenames)
 	{
 		const FString FullPath = ContainerFolder / Filename;
@@ -211,9 +217,31 @@ static int32 ChunkPluginCommandEntry(const FContext& Context)
 			TocEntry.BlockOffset = BlockOffset;
 			TocEntry.BlockCount = BlockCount;
 		}
+
+		if (bDeleteContainerFiles)
+		{
+			FilesToDelete.Add(FullPath);
+			ContainerFileReader.GetContainerFilePaths(FilesToDelete);
+		}
+	}
+
+
+	IFileManager& FileMan = IFileManager::Get();
+	for (const FString& Path : FilesToDelete)
+	{
+		//UE_LOG(LogIas, Display, TEXT("Attempt Deleting '%s'"), *Path);
+		if (FileMan.FileExists(*Path))
+		{
+			UE_LOG(LogIas, Display, TEXT("Deleting '%s'"), *Path);
+			if (!FileMan.Delete(*Path, /*RequireExists*/true))
+			{
+				UE_LOG(LogIas, Error, TEXT("Failed to delete '%s'"), *Path);
+			}
+		}
 	}
 
 	// Write additional file(s)
+	if (bIncludeSigPak)
 	{
 		const FStringView AllowedExt[]
 		{
@@ -276,16 +304,17 @@ static int32 ChunkPluginCommandEntry(const FContext& Context)
 		if (TUniquePtr<FArchive> Ar(IFileManager::Get().CreateFileWriter(*TocPath)); Ar.IsValid())
 		{
 			*Ar << OnDemandToc;
+			Ar->Close();
 			if (Ar->IsError())
 			{
-				UE_LOG(LogIoStore, Error, TEXT("Failed to serialize TOC"));
+				UE_LOG(LogIoStore, Error, TEXT("Failed to serialize TOC '%s'"), *TocPath);
 				return -1;
 			}
 			else
 			{
 				const int64 TocSize = Ar->Tell();
 				UE_LOG(LogIoStore, Display, TEXT("Writing file '%s' (%.2lf KiB)"),
-					*Filename, double(TocSize) / 1024);
+					*TocPath, double(TocSize) / 1024);
 			}
 		}
 		else
@@ -311,6 +340,8 @@ static FCommand ChunkPluginCommand(
 		TArgument<FStringView>(TEXT("-OutputFolder"),		TEXT("Ouptut folder.")),
 		TArgument<FStringView>(TEXT("-IntermediateFolder"),	TEXT("Intermediate folder.")),
 		TArgument<FStringView>(TEXT("-ErrorOutput"),		TEXT("Error output.")),
+		TArgument<bool>(TEXT("-IncludeSigPak"),				TEXT("Include .sig and .pak file in the uondemandtoc")),
+		TArgument<bool>(TEXT("-KeepContainerFiles"),		TEXT("Should we keep the container files after processing them.")),
 	}
 );
 
