Chaos - fix CCD interactions with other CCD objects

If a fast CCD object hit a second sleeping or slow-moving CCD object then it would not update the collision constraint correctly after rewinding.

diff --git a/Engine/Source/Runtime/Experimental/Chaos/Private/Chaos/CCDUtilities.cpp b/Engine/Source/Runtime/Experimental/Chaos/Private/Chaos/CCDUtilities.cpp
--- a/Engine/Source/Runtime/Experimental/Chaos/Private/Chaos/CCDUtilities.cpp
+++ b/Engine/Source/Runtime/Experimental/Chaos/Private/Chaos/CCDUtilities.cpp
@@ -876,19 +876,13 @@ namespace Chaos
 		// logic executed in ApplySweptConstraints (i.e., contents of CCDConstraints)
 		// @todo(chaos): we could calculate the size of the Collisions array in Init
 		// @todo(chaos): could be parallelized
-		TArray<FPBDCollisionConstraint*> Collisions;
+		TSet<FPBDCollisionConstraint*> Collisions;
 		for (FCCDParticle& CCDParticle : CCDParticles)
 		{
 			CCDParticle.Particle->ParticleCollisions().VisitCollisions(
 				[&Collisions](FPBDCollisionConstraint& Collision)
 				{
-					// Avoid duplicates when both particles in the collision are CCD enabled by checking the particle ID
-					const FConstGenericParticleHandle P0 = FConstGenericParticleHandle(Collision.GetParticle0());
-					const FConstGenericParticleHandle P1 = FConstGenericParticleHandle(Collision.GetParticle1());
-					if (!P0->CCDEnabled() || !P1->CCDEnabled() || (P0->ParticleID() < P1->ParticleID()))
-					{
-						Collisions.Add(&Collision);
-					}
+					Collisions.Add(&Collision);
 					return ECollisionVisitorResult::Continue;
 				});
 		}
