Chaos - MACD auto falls back to non-MACD when not required
- MACD now builds the manifold based on closest contact normal, not the corrected normal

I had to disable the newer MACD mode which used mesh data to correct the contact normal before generating the manifold from it because it can lead to unexpected collisions on complicated meshes because we are no longer selected the shortest depentration vector. We also sometimes hit the malformed mesh ensure.

diff --git a/Engine/Source/Runtime/Experimental/Chaos/Private/Chaos/CollisionOneShotManifolds.cpp b/Engine/Source/Runtime/Experimental/Chaos/Private/Chaos/CollisionOneShotManifolds.cpp
--- a/Engine/Source/Runtime/Experimental/Chaos/Private/Chaos/CollisionOneShotManifolds.cpp
+++ b/Engine/Source/Runtime/Experimental/Chaos/Private/Chaos/CollisionOneShotManifolds.cpp
@@ -116,7 +116,8 @@ namespace Chaos
 	// The second method does a much better job at handling deep collisions where a single convex-triangle contact might generate a contact that pushes the convex out
 	// of the triangle edge, even though the triangle is part of a mesh and pushing out of the edge is not a valid option. However this option is also more likely to 
 	// generate ghost collisions with nearby faces when moving very fast.
-	bool bChaos_Collision_EnableMACDPreManifoldFix = true;
+	// @todo(chaos): the second method has problems with complicated meshes and need to be disabled.
+	bool bChaos_Collision_EnableMACDPreManifoldFix = false;
 	FAutoConsoleVariableRef CVarChaos_Collision_EnableMACDPreManifoldFix(TEXT("p.Chaos.Collision.EnableMACDPreManifoldFix"), bChaos_Collision_EnableMACDPreManifoldFix, TEXT(""));
 
 	// Whether to use the new index-less GJK. 
diff --git a/Engine/Source/Runtime/Experimental/Chaos/Private/Chaos/CollisionOneShotManifoldsMeshShapes.cpp b/Engine/Source/Runtime/Experimental/Chaos/Private/Chaos/CollisionOneShotManifoldsMeshShapes.cpp
--- a/Engine/Source/Runtime/Experimental/Chaos/Private/Chaos/CollisionOneShotManifoldsMeshShapes.cpp
+++ b/Engine/Source/Runtime/Experimental/Chaos/Private/Chaos/CollisionOneShotManifoldsMeshShapes.cpp
@@ -318,16 +318,29 @@ namespace Chaos
 		// Original MACD algorithm that uses mesh information to fix manifold point normals after the manifold is built.
 		// @todo(chaos): remove this when the new version is well tested.
 		template<typename ConvexType>
-		void GenerateConvexTriangleOneShotManifoldMACD_PostManifoldFix(const ConvexType& Convex, const FRigidTransform3& ConvexTransform, const FVec3& ConvexRelativeMovement, Private::FMeshContactGenerator& ContactGenerator, const int32 TriangleIndex, const FReal CullDistance, FContactPointManifold& OutContactPoints)
+		void GenerateConvexTriangleOneShotManifoldMACD_PostManifoldFix(const ConvexType& Convex, const FRigidTransform3& ConvexTransform, const FVec3& InConvexRelativeMovement, Private::FMeshContactGenerator& ContactGenerator, const int32 TriangleIndex, const FReal CullDistance, FContactPointManifold& OutContactPoints)
 		{
-			const FTriangle TriangleAtP = ContactGenerator.GetTriangle(TriangleIndex);
+			FTriangle Triangle = ContactGenerator.GetTriangle(TriangleIndex);
 			const FVec3 TriangleNormal = ContactGenerator.GetTriangleNormal(TriangleIndex);
+			const FReal ConvexTriangleDistanceAtP = FVec3::DotProduct(Convex.GetCenterOfMass() - Triangle.GetVertex(0), TriangleNormal);
+
+			// If we are outside the plane of the triangle at P, collide at P
+			bool bUseMACD = !InConvexRelativeMovement.IsZero() && (ConvexTriangleDistanceAtP < 0);
+			
+			// If desired, shift the triangle so it is relative to the convex when it was at X
+			FVec3 ConvexRelativeMovement = FVec3(0);
+			FReal ConvexTriangleDistance = ConvexTriangleDistanceAtP;
+			if (bUseMACD)
+			{
+				ConvexRelativeMovement = InConvexRelativeMovement;
+				Triangle[0] += ConvexRelativeMovement;
+				Triangle[1] += ConvexRelativeMovement;
+				Triangle[2] += ConvexRelativeMovement;
+				ConvexTriangleDistance -= FVec3::DotProduct(ConvexRelativeMovement, TriangleNormal);
+			}
 
-			FTriangle Triangle = FTriangle(TriangleAtP.GetVertex(0) + ConvexRelativeMovement, TriangleAtP.GetVertex(1) + ConvexRelativeMovement, TriangleAtP.GetVertex(2) + ConvexRelativeMovement);
-
-			// If we started inside the triangle we ignore this triangle
-			const FReal ConvexTriangleDistance = FVec3::DotProduct(Convex.GetCenterOfMass() - Triangle.GetVertex(0), TriangleNormal);
-			if (ConvexTriangleDistance < 0)
+			// If we were inside the triangle at X and P we ignore this triangle
+			if ((ConvexTriangleDistance < 0) && (ConvexTriangleDistanceAtP < 0))
 			{
 				return;
 			}
@@ -339,26 +352,39 @@ namespace Chaos
 				ClosestContact.Features[0].ObjectIndex = 0;
 				ClosestContact.Features[1].ObjectIndex = TriangleIndex;
 
+#if CHAOS_DEBUG_DRAW
+				if (CVars::ChaosSolverDebugDrawMeshContacts && FDebugDrawQueue::GetInstance().IsDebugDrawingEnabled())
+				{
+					const FVec3 P = ConvexTransform.TransformPositionNoScale(ClosestContact.ShapeContactPoints[1] - ConvexRelativeMovement);
+					const FVec3 N = ConvexTransform.TransformVectorNoScale(ClosestContact.ShapeContactNormal);
+					FDebugDrawQueue::GetInstance().DrawDebugLine(P, P + 10.0f * N, FColor::Black, false, CVars::ChaosSolverDebugDebugDrawSettings.DrawDuration, (uint8)CVars::ChaosSolverDebugDebugDrawSettings.DrawPriority, 1.25f * CVars::ChaosSolverDebugDebugDrawSettings.LineThickness);
+				}
+				if (CVars::ChaosSolverDebugDrawMeshContactDetails && FDebugDrawQueue::GetInstance().IsDebugDrawingEnabled())
+				{
+					const FVec3 XConvex = ConvexTransform.TransformPositionNoScale(-ConvexRelativeMovement);
+					const FColor Color = (bUseMACD) ? FColor::Red : FColor::Green;
+					FDebugDrawQueue::GetInstance().DrawDebugPoint(XConvex, Color, false, CVars::ChaosSolverDebugDebugDrawSettings.DrawDuration, (uint8)CVars::ChaosSolverDebugDebugDrawSettings.DrawPriority, 20.0f * CVars::ChaosSolverDebugDebugDrawSettings.LineThickness);
+				}
+#endif
+
 				// Back face culling
 				const FReal TriangleDotNormal = FVec3::DotProduct(TriangleNormal, ClosestContact.ShapeContactNormal);
 				if (TriangleDotNormal < 0)
 				{
-#if CHAOS_DEBUG_DRAW
-					if (CVars::ChaosSolverDebugDrawMeshContacts && FDebugDrawQueue::GetInstance().IsDebugDrawingEnabled())
-					{
-						const FVec3 P = ConvexTransform.TransformPositionNoScale(ClosestContact.ShapeContactPoints[1] - ConvexRelativeMovement);
-						const FVec3 N = ConvexTransform.TransformVectorNoScale(ClosestContact.ShapeContactNormal);
-						FDebugDrawQueue::GetInstance().DrawDebugLine(P, P + 10.0f * N, FColor::Black, false, CVars::ChaosSolverDebugDebugDrawSettings.DrawDuration, (uint8)CVars::ChaosSolverDebugDebugDrawSettings.DrawPriority, 1.25f * CVars::ChaosSolverDebugDebugDrawSettings.LineThickness);
-					}
-#endif
 					return;
 				}
 
+				// Calculate cull distance that takes movement direction and distance into account
+				FReal NetCullDistance = CullDistance;
+				if (bUseMACD)
+				{
+					const FReal ConvexMotionAlongNormal = FVec3::DotProduct(ConvexRelativeMovement, ClosestContact.ShapeContactNormal);
+					const FReal CullDistancePadding = FMath::Max(0, -ConvexMotionAlongNormal);
+					NetCullDistance += CullDistancePadding;
+				}
+
 				// Generate a manifold based on the closest features
 				// NOTE: normal points from triangle to convex
-				const FReal ConvexMotionAlongNormal = FVec3::DotProduct(ConvexRelativeMovement, ClosestContact.ShapeContactNormal);
-				const FReal CullDistancePadding = FMath::Max(0, -ConvexMotionAlongNormal);
-				const FReal NetCullDistance = CullDistance + CullDistancePadding;
 				Private::ConvexTriangleManifoldFromContact(Convex, Triangle, TriangleNormal, ClosestContact, NetCullDistance, OutContactPoints);
 
 				if (OutContactPoints.Num() > 0)
@@ -389,8 +415,9 @@ namespace Chaos
 					}
 #endif
 
-					// Adjust the feature if it is invalid. E.g., if we collide with a triangle edge and the normal is outside the range allowed by
-					// the triangles sharing the edge we will project the normal into the valid range
+					// Adjust the closest feature if it is invalid. E.g., if we collide with a triangle edge and the normal is outside the range allowed by
+					// the triangles sharing the edge we will project the normal into the valid range. If the feature gets chnaged, we will alter all
+					// of the manifold points to use the new normal.
 					const bool bFeatureChanged = ContactGenerator.FixFeature(TriangleIndex, ClosestContact.Features[1].FeatureType, ClosestContact.Features[1].PlaneFeatureIndex, ClosestContact.ShapeContactNormal);
 					if (bFeatureChanged)
 					{
@@ -404,20 +431,23 @@ namespace Chaos
 
 							// Remap the triangle contact onto the new plane, keeping the contact point on the convex shape where it is.
 							// NOTE: This means that the triangle contact point may be outside the triangle, but for contact separation
-							// we really only care about the contact plane. Howeber, the contact positions are used for static friction, 
-							// which assumes the contacts have zero separation perpendicular to the normal on the frame they are generated.
+							// we really only care about the contact plane. This is required for static friction, which assumes the contacts 
+							// have zero tangential separation on the frame they are generated.
 							ContactPoint.ShapeContactPoints[1] = ContactPoint.ShapeContactPoints[0] - (ContactPoint.Phi * ContactPoint.ShapeContactNormal);
 						}
 					}
 
-					// Correct the contact points based on convex movement
-					for (int32 ContactIndex = 0; ContactIndex < OutContactPoints.Num(); ++ContactIndex)
+					// Correct the contact points if we ran collision detection at X rather than P
+					if (bUseMACD)
 					{
-						FContactPoint& ContactPoint = OutContactPoints[ContactIndex];
+						for (int32 ContactIndex = 0; ContactIndex < OutContactPoints.Num(); ++ContactIndex)
+						{
+							FContactPoint& ContactPoint = OutContactPoints[ContactIndex];
 
-						const FReal ShiftDotNormal = FVec3::DotProduct(ConvexRelativeMovement, ContactPoint.ShapeContactNormal);
-						ContactPoint.ShapeContactPoints[1] += -ShiftDotNormal * ContactPoint.ShapeContactNormal;
-						ContactPoint.Phi += ShiftDotNormal;
+							const FReal ShiftDotNormal = FVec3::DotProduct(ConvexRelativeMovement, ContactPoint.ShapeContactNormal);
+							ContactPoint.ShapeContactPoints[1] += -ShiftDotNormal * ContactPoint.ShapeContactNormal;
+							ContactPoint.Phi += ShiftDotNormal;
+						}
 					}
 				}
 			}
@@ -666,25 +696,10 @@ namespace Chaos
 				const auto& GenerateConvexTriangleContacts =
 					[&Convex, &ConvexTransform, &ConvexRelativeMovement, CullDistance, &Manifold](Private::FMeshContactGenerator& ContactGenerator, const int32 TriangleIndex)
 				{
-					// If we are outside the plane of the triangle, just use non-MACD collision
-					const FTriangle& Triangle = ContactGenerator.GetTriangle(TriangleIndex);
-					bool bUseNonMACDPath = false;
-					if (bChaos_Collision_EnableMACDFallback)
-					{
-						const FVec3 TriangleNormal = ContactGenerator.GetTriangleNormal(TriangleIndex);
-						const FReal Distance = FVec3::DotProduct(Convex.GetCenterOfMass() - Triangle.GetVertex(0), TriangleNormal);
-						bUseNonMACDPath = Distance > 0;
-					}
-
 					Manifold.Reset();
-					if (bUseNonMACDPath)
-					{
-						GenerateConvexTriangleOneShotManifold(Convex, Triangle, CullDistance, Manifold);
-					}
-					else
-					{
-						GenerateConvexTriangleOneShotManifoldMACD(Convex, ConvexTransform, ConvexRelativeMovement, ContactGenerator, TriangleIndex, CullDistance, Manifold);
-					}
+
+					GenerateConvexTriangleOneShotManifoldMACD(Convex, ConvexTransform, ConvexRelativeMovement, ContactGenerator, TriangleIndex, CullDistance, Manifold);
+
 					ContactGenerator.AddTriangleContacts(TriangleIndex, MakeArrayView(Manifold));
 				};
 
