[NetPhysics] Don't automatically initialize RewindData from NetworkPhysicsComponent unless it's using resimulation.
- Also adjustments to make NetworkPhysicsComponent run properly without RewindData being initalized.

diff --git a/Engine/Source/Runtime/Engine/Private/PhysicsEngine/NetworkPhysicsComponent.cpp b/Engine/Source/Runtime/Engine/Private/PhysicsEngine/NetworkPhysicsComponent.cpp
--- a/Engine/Source/Runtime/Engine/Private/PhysicsEngine/NetworkPhysicsComponent.cpp
+++ b/Engine/Source/Runtime/Engine/Private/PhysicsEngine/NetworkPhysicsComponent.cpp
@@ -1139,6 +1139,13 @@ int32 UNetworkPhysicsComponent::SetupRewindData()
 {
 	int32 NumFrames = UPhysicsSettings::Get()->GetPhysicsHistoryCount();
 
+	AActor* Owner = GetOwner();
+	// Don't let this actor initialize RewindData if the actor is not using resimulation
+	if (Owner && Owner->GetPhysicsReplicationMode() != EPhysicsReplicationMode::Resimulation)
+	{
+		return NumFrames;
+	}
+
 	if (FPhysScene* PhysScene = GetWorld()->GetPhysicsScene())
 	{
 		if (Chaos::FPhysicsSolver* Solver = PhysScene->GetSolver())
diff --git a/Engine/Source/Runtime/Experimental/Chaos/Private/Chaos/Framework/PhysicsSolverBase.cpp b/Engine/Source/Runtime/Experimental/Chaos/Private/Chaos/Framework/PhysicsSolverBase.cpp
--- a/Engine/Source/Runtime/Experimental/Chaos/Private/Chaos/Framework/PhysicsSolverBase.cpp
+++ b/Engine/Source/Runtime/Experimental/Chaos/Private/Chaos/Framework/PhysicsSolverBase.cpp
@@ -496,7 +496,7 @@ namespace Chaos
 
 		while(FPushPhysicsData* PushData = MarshallingManager.StepInternalTime_External())
 		{
-			if(ShouldApplyRewindCallbacks() && !bIsShuttingDown)
+			if(MRewindCallback && !bIsShuttingDown)
 			{
 				MRewindCallback->ProcessInputs_External(PushData->InternalStep, PushData->SimCallbackInputs);
 			}
diff --git a/Engine/Source/Runtime/Experimental/Chaos/Private/PBDRigidsSolver.cpp b/Engine/Source/Runtime/Experimental/Chaos/Private/PBDRigidsSolver.cpp
--- a/Engine/Source/Runtime/Experimental/Chaos/Private/PBDRigidsSolver.cpp
+++ b/Engine/Source/Runtime/Experimental/Chaos/Private/PBDRigidsSolver.cpp
@@ -1459,7 +1459,7 @@ namespace Chaos
 
 		GetEvolution()->GetBroadPhase().GetIgnoreCollisionManager().PushProducerStorageData_External(MarshallingManager.GetExternalTimestamp_External());
 
-		if (ShouldApplyRewindCallbacks() && !IsShuttingDown())
+		if (MRewindCallback && !IsShuttingDown())
 		{
 			MRewindCallback->InjectInputs_External(MarshallingManager.GetInternalStep_External(), NumSteps);
 		}
@@ -1861,16 +1861,16 @@ namespace Chaos
 			SimCallbackObject->PostInitialize_Internal();
 		}
 
-		if (MRewindCallback && MRewindData && !IsShuttingDown())
+		if (MRewindCallback && !IsShuttingDown())
 		{
-			MRewindCallback->ProcessInputs_Internal(MRewindData->CurrentFrame(), PushData.SimCallbackInputs);
+			MRewindCallback->ProcessInputs_Internal(GetCurrentFrame(), PushData.SimCallbackInputs);
 		}
 	}
 
 	void FPBDRigidsSolver::ConditionalApplyRewind_Internal()
 	{
 		// Note: checking MRewindData->IsResim() can lead to recursion into this function on the last resim frame since the call to AdvanceSolver is what advances RewindData's internal frame
-		if(!IsShuttingDown() && ShouldApplyRewindCallbacks() && MRewindData && !GetEvolution()->IsResimming())
+		if(!IsShuttingDown() && ShouldApplyRewindCallbacks() && !GetEvolution()->IsResimming())
 		{
 			const int32 LastStep = MRewindData->CurrentFrame() - 1;
 			const int32 ResimStep = MRewindCallback->TriggerRewindIfNeeded_Internal(LastStep);
