Fix crash when moving a GC component while in a blueprint editor

diff --git a/Engine/Source/Runtime/Experimental/GeometryCollectionEngine/Private/GeometryCollection/GeometryCollectionComponent.cpp b/Engine/Source/Runtime/Experimental/GeometryCollectionEngine/Private/GeometryCollection/GeometryCollectionComponent.cpp
--- a/Engine/Source/Runtime/Experimental/GeometryCollectionEngine/Private/GeometryCollection/GeometryCollectionComponent.cpp
+++ b/Engine/Source/Runtime/Experimental/GeometryCollectionEngine/Private/GeometryCollection/GeometryCollectionComponent.cpp
@@ -3270,9 +3270,12 @@ bool UGeometryCollectionComponent::MoveComponentImpl(const FVector& Delta, const
 
 #if WITH_EDITOR
 	// Make sure that custom renderer is updated _after_ any move has been applied to the full component hierachy.
-	if (!GetWorld()->IsGameWorld())
+	if (UWorld* World = GetWorld())
 	{
-		RefreshCustomRenderer();
+		if (!World->IsGameWorld())
+		{
+			RefreshCustomRenderer();
+		}
 	}
 #endif
 
