Fix vectorized solver

diff --git a/Engine/Source/Runtime/Experimental/Chaos/Private/Chaos/Joint/PBDJointCachedSolverGaussSeidel.cpp b/Engine/Source/Runtime/Experimental/Chaos/Private/Chaos/Joint/PBDJointCachedSolverGaussSeidel.cpp
--- a/Engine/Source/Runtime/Experimental/Chaos/Private/Chaos/Joint/PBDJointCachedSolverGaussSeidel.cpp
+++ b/Engine/Source/Runtime/Experimental/Chaos/Private/Chaos/Joint/PBDJointCachedSolverGaussSeidel.cpp
@@ -1037,13 +1037,13 @@ void FPBDJointCachedSolver::ApplyVelocityConstraintSimd()
 		}
 		if (Body(1).IsDynamic())
 		{
-			const FRealSingle Inv1f = -FRealSingle(InvM(1));
-			const VectorRegister4Float InvM1 = VectorLoadFloat1(&Inv1f);
+			const FRealSingle OppInv1f = -FRealSingle(InvM(1));
+			const VectorRegister4Float OppInvM1 = VectorLoadFloat1(&OppInv1f);
 			VectorRegister4Float DV1 = VectorZeroFloat();
 			VectorRegister4Float DW1 = VectorZeroFloat();
 			for (int32 CIndex = 0; CIndex < 3; CIndex++)
 			{
-				DV1 = VectorMultiplyAdd(InvM1, VectorMultiply(DeltaLambdas[CIndex], PositionConstraints.Simd.ConstraintAxis[CIndex]), DV1);
+				DV1 = VectorMultiplyAdd(OppInvM1, VectorMultiply(DeltaLambdas[CIndex], PositionConstraints.Simd.ConstraintAxis[CIndex]), DV1);
 				DW1 = VectorMultiplyAdd(PositionConstraints.Simd.ConstraintDRAxis[CIndex][1], DeltaLambdas[CIndex], DW1);
 			}
 			FVec3f DV1f;
@@ -1295,11 +1295,7 @@ void FPBDJointCachedSolver::InitRotationConstraintsSimd(
 	RotationConstraints.Simd.ConstraintSoftDamping = bUsePositionBasedDrives ? VectorMultiply(SpringMassScale, VectorMultiply(SoftDamping, Dt)) : VectorZeroFloat();
 	RotationConstraints.Simd.ConstraintSoftIM = VectorAdd(VectorMultiply(VectorAdd(RotationConstraints.Simd.ConstraintSoftStiffness, RotationConstraints.Simd.ConstraintSoftDamping), RotationConstraints.Simd.ConstraintHardIM), GlobalVectorConstants::FloatOne);
 
-
-
 	RotationConstraints.Simd.ConstraintCX = MakeVectorRegisterFloat(LocalAngles[0], LocalAngles[1], LocalAngles[2], 0.0f);
-	RotationConstraints.Simd.ConstraintHardIM = MakeVectorRegisterFloat(ConstraintHardIM[0], ConstraintHardIM[1], ConstraintHardIM[2], 0.0f);
-
 	RotationConstraints.Simd.ConstraintLimits = MakeVectorRegisterFloatFromDouble(MakeVectorRegister(JointSettings.AngularLimits[0], JointSettings.AngularLimits[1], JointSettings.AngularLimits[2], UE_BIG_NUMBER));
 }
 
@@ -1808,7 +1804,7 @@ void FPBDJointCachedSolver::ApplyAngularVelocityConstraintSimd()
 			}
 			FVec3f DW1f;
 			VectorStoreFloat3(DW1, &DW1f[0]);
-			Body(0).ApplyAngularVelocityDelta(FVec3(DW1f));
+			Body(1).ApplyAngularVelocityDelta(FVec3(DW1f));
 		}
 	}
 }
