Fix issue with blueprint thumbnails not rendering geometry collections with custom renderer.
For thumbnail rendering OnCreatePhysicsState() won't be called and so we need to rely on component registration to call RefreshCustomRenderer().

diff --git a/Engine/Source/Runtime/Experimental/GeometryCollectionEngine/Private/GeometryCollection/GeometryCollectionComponent.cpp b/Engine/Source/Runtime/Experimental/GeometryCollectionEngine/Private/GeometryCollection/GeometryCollectionComponent.cpp
--- a/Engine/Source/Runtime/Experimental/GeometryCollectionEngine/Private/GeometryCollection/GeometryCollectionComponent.cpp
+++ b/Engine/Source/Runtime/Experimental/GeometryCollectionEngine/Private/GeometryCollection/GeometryCollectionComponent.cpp
@@ -3566,6 +3566,14 @@ void UGeometryCollectionComponent::OnRegister()
 			CustomRenderer = NewObject<UObject>(this, Type, NAME_None, RF_Transient | RF_DuplicateTransient);
 		}
 		RegisterCustomRenderer();
+
+#if WITH_EDITOR
+		// Make sure that custom renderer is updated at least once in editor.
+		if (GetWorld() && !GetWorld()->IsGameWorld())
+		{
+			RefreshCustomRenderer();
+		}
+#endif
 	}
 	else
 	{
@@ -3809,14 +3817,6 @@ void UGeometryCollectionComponent::OnCreatePhysicsState()
 				OnComponentPhysicsStateChanged.Broadcast(this, EComponentPhysicsStateChange::Created);
 			}
 		}
-
-#if WITH_EDITOR
-		// Make sure that custom renderer is updated at least once in editor.
-		if (!GetWorld()->IsGameWorld())
-		{
-			RefreshCustomRenderer();
-		}
-#endif
 	}
 }
 
