Dataflow:
- Added mechanism to version Dataflow nodes (it is backward compatible with existing nodes and graphs)
- Added capability to deprecate Dataflow nodes
- Added capability to specify Experimental Dataflow nodes
- Node version,  deprecated and experimental info displayed in node tooltip
- Experimental and deprecated nodes have different visual in the graph

diff --git a/Engine/Plugins/ChaosClothAssetEditor/Source/ChaosClothAssetDataflowNodes/Public/ChaosClothAsset/AttributeNode.h b/Engine/Plugins/ChaosClothAssetEditor/Source/ChaosClothAssetDataflowNodes/Public/ChaosClothAsset/AttributeNode.h
--- a/Engine/Plugins/ChaosClothAssetEditor/Source/ChaosClothAssetDataflowNodes/Public/ChaosClothAsset/AttributeNode.h
+++ b/Engine/Plugins/ChaosClothAssetEditor/Source/ChaosClothAssetDataflowNodes/Public/ChaosClothAsset/AttributeNode.h
@@ -29,7 +29,7 @@ enum class EChaosClothAssetNodeAttributeType : uint8
 };
 
 /** Create a new attribute for the specified group. */
-USTRUCT(Meta = (DataflowCloth))
+USTRUCT(Meta = (DataflowCloth, Experimental))
 struct FChaosClothAssetAttributeNode : public FDataflowNode
 {
 	GENERATED_USTRUCT_BODY()
@@ -74,7 +74,6 @@ private:
 	virtual void Evaluate(Dataflow::FContext& Context, const FDataflowOutput* Out) const override;
 	virtual void OnSelected(Dataflow::FContext& Context) override;
 	virtual void OnDeselected() override;
-	virtual bool IsExperimental() override { return true; }
 	//~ End FDataflowNode interface
 
 	TArray<FName> CachedCollectionGroupNames;
diff --git a/Engine/Plugins/ChaosClothAssetEditor/Source/ChaosClothAssetDataflowNodes/Public/ChaosClothAsset/DatasmithImportNode.h b/Engine/Plugins/ChaosClothAssetEditor/Source/ChaosClothAssetDataflowNodes/Public/ChaosClothAsset/DatasmithImportNode.h
--- a/Engine/Plugins/ChaosClothAssetEditor/Source/ChaosClothAssetDataflowNodes/Public/ChaosClothAsset/DatasmithImportNode.h
+++ b/Engine/Plugins/ChaosClothAssetEditor/Source/ChaosClothAssetDataflowNodes/Public/ChaosClothAsset/DatasmithImportNode.h
@@ -39,7 +39,7 @@ public:
 };
 
 /** Import a file from a third party garment construction package compatible with the Datasmith scene format. */
-USTRUCT(meta = (DataflowCloth))
+USTRUCT(meta = (DataflowCloth, Deprecated = "5.4"))
 struct FChaosClothAssetDatasmithImportNode : public FDataflowNode
 {
 	GENERATED_USTRUCT_BODY()
@@ -64,7 +64,6 @@ private:
 	//~ Begin FDataflowNode interface
 	virtual void Evaluate(Dataflow::FContext& Context, const FDataflowOutput* Out) const override;
 	virtual void Serialize(FArchive& Archive) override;
-	virtual bool IsDeprecated() override { return true; }
 	//~ End FDataflowNode interface
 
 	bool EvaluateImpl(Dataflow::FContext& Context, FManagedArrayCollection& OutCollection) const;
diff --git a/Engine/Plugins/ChaosClothAssetEditor/Source/ChaosClothAssetDataflowNodes/Public/ChaosClothAsset/SimulationBendingOverrideConfigNode.h b/Engine/Plugins/ChaosClothAssetEditor/Source/ChaosClothAssetDataflowNodes/Public/ChaosClothAsset/SimulationBendingOverrideConfigNode.h
--- a/Engine/Plugins/ChaosClothAssetEditor/Source/ChaosClothAssetDataflowNodes/Public/ChaosClothAsset/SimulationBendingOverrideConfigNode.h
+++ b/Engine/Plugins/ChaosClothAssetEditor/Source/ChaosClothAssetDataflowNodes/Public/ChaosClothAsset/SimulationBendingOverrideConfigNode.h
@@ -8,7 +8,7 @@
 #include "SimulationBendingOverrideConfigNode.generated.h"
 
 /** Bending constraint property override configuration node. */
-USTRUCT(Meta = (DataflowCloth))
+USTRUCT(Meta = (DataflowCloth, Experimental))
 struct FChaosClothAssetSimulationBendingOverrideConfigNode : public FChaosClothAssetSimulationBaseConfigNode
 {
 	GENERATED_USTRUCT_BODY()
@@ -120,5 +120,4 @@ public:
 
 private:
 	virtual void AddProperties(FPropertyHelper& PropertyHelper) const override;
-	virtual bool IsExperimental() override { return true; }
 };
diff --git a/Engine/Plugins/ChaosClothAssetEditor/Source/ChaosClothAssetDataflowNodes/Public/ChaosClothAsset/SimulationMultiResConfigNode.h b/Engine/Plugins/ChaosClothAssetEditor/Source/ChaosClothAssetDataflowNodes/Public/ChaosClothAsset/SimulationMultiResConfigNode.h
--- a/Engine/Plugins/ChaosClothAssetEditor/Source/ChaosClothAssetDataflowNodes/Public/ChaosClothAsset/SimulationMultiResConfigNode.h
+++ b/Engine/Plugins/ChaosClothAssetEditor/Source/ChaosClothAssetDataflowNodes/Public/ChaosClothAsset/SimulationMultiResConfigNode.h
@@ -6,7 +6,7 @@
 #include "SimulationMultiResConfigNode.generated.h"
 
 /** Experimental Solver multires configuration node. */
-USTRUCT(Meta = (DataflowCloth))
+USTRUCT(Meta = (DataflowCloth, Experimental))
 struct FChaosClothAssetSimulationMultiResConfigNode : public FChaosClothAssetSimulationBaseConfigNode
 {
 	GENERATED_USTRUCT_BODY()
@@ -42,5 +42,4 @@ public:
 
 private:
 	virtual void AddProperties(FPropertyHelper& PropertyHelper) const override;
-	virtual bool IsExperimental() override { return true; }
 };
\ No newline at end of file
diff --git a/Engine/Plugins/ChaosClothAssetEditor/Source/ChaosClothAssetDataflowNodes/Public/ChaosClothAsset/SimulationPBDAreaSpringConfigNode.h b/Engine/Plugins/ChaosClothAssetEditor/Source/ChaosClothAssetDataflowNodes/Public/ChaosClothAsset/SimulationPBDAreaSpringConfigNode.h
--- a/Engine/Plugins/ChaosClothAssetEditor/Source/ChaosClothAssetDataflowNodes/Public/ChaosClothAsset/SimulationPBDAreaSpringConfigNode.h
+++ b/Engine/Plugins/ChaosClothAssetEditor/Source/ChaosClothAssetDataflowNodes/Public/ChaosClothAsset/SimulationPBDAreaSpringConfigNode.h
@@ -7,7 +7,7 @@
 #include "SimulationPBDAreaSpringConfigNode.generated.h"
 
 /** Area spring constraint property configuration node. */
-USTRUCT(Meta = (DataflowCloth))
+USTRUCT(Meta = (DataflowCloth, Deprecated = "5.4"))
 struct FChaosClothAssetSimulationPBDAreaSpringConfigNode : public FChaosClothAssetSimulationBaseConfigNode
 {
 	GENERATED_USTRUCT_BODY()
@@ -28,5 +28,4 @@ public:
 
 private:
 	virtual void AddProperties(FPropertyHelper& PropertyHelper) const override;
-	virtual bool IsDeprecated() override { return true; }
 };
diff --git a/Engine/Plugins/ChaosClothAssetEditor/Source/ChaosClothAssetDataflowNodes/Public/ChaosClothAsset/SimulationPBDBendingElementConfigNode.h b/Engine/Plugins/ChaosClothAssetEditor/Source/ChaosClothAssetDataflowNodes/Public/ChaosClothAsset/SimulationPBDBendingElementConfigNode.h
--- a/Engine/Plugins/ChaosClothAssetEditor/Source/ChaosClothAssetDataflowNodes/Public/ChaosClothAsset/SimulationPBDBendingElementConfigNode.h
+++ b/Engine/Plugins/ChaosClothAssetEditor/Source/ChaosClothAssetDataflowNodes/Public/ChaosClothAsset/SimulationPBDBendingElementConfigNode.h
@@ -8,7 +8,7 @@
 #include "SimulationPBDBendingElementConfigNode.generated.h"
 
 /** Bending element constraint property configuration node. */
-USTRUCT(Meta = (DataflowCloth))
+USTRUCT(Meta = (DataflowCloth, Deprecated = "5.4"))
 struct FChaosClothAssetSimulationPBDBendingElementConfigNode : public FChaosClothAssetSimulationBaseConfigNode
 {
 	GENERATED_USTRUCT_BODY()
@@ -75,5 +75,4 @@ public:
 
 private:
 	virtual void AddProperties(FPropertyHelper& PropertyHelper) const override;
-	virtual bool IsDeprecated() override { return true; }
 };
diff --git a/Engine/Plugins/ChaosClothAssetEditor/Source/ChaosClothAssetDataflowNodes/Public/ChaosClothAsset/SimulationPBDBendingSpringConfigNode.h b/Engine/Plugins/ChaosClothAssetEditor/Source/ChaosClothAssetDataflowNodes/Public/ChaosClothAsset/SimulationPBDBendingSpringConfigNode.h
--- a/Engine/Plugins/ChaosClothAssetEditor/Source/ChaosClothAssetDataflowNodes/Public/ChaosClothAsset/SimulationPBDBendingSpringConfigNode.h
+++ b/Engine/Plugins/ChaosClothAssetEditor/Source/ChaosClothAssetDataflowNodes/Public/ChaosClothAsset/SimulationPBDBendingSpringConfigNode.h
@@ -7,7 +7,7 @@
 #include "SimulationPBDBendingSpringConfigNode.generated.h"
 
 /** Bending spring constraint property configuration node. */
-USTRUCT(Meta = (DataflowCloth))
+USTRUCT(Meta = (DataflowCloth, Deprecated = "5.4"))
 struct FChaosClothAssetSimulationPBDBendingSpringConfigNode : public FChaosClothAssetSimulationBaseConfigNode
 {
 	GENERATED_USTRUCT_BODY()
@@ -28,5 +28,4 @@ public:
 
 private:
 	virtual void AddProperties(FPropertyHelper& PropertyHelper) const override;
-	virtual bool IsDeprecated() override { return true; }
 };
diff --git a/Engine/Plugins/ChaosClothAssetEditor/Source/ChaosClothAssetDataflowNodes/Public/ChaosClothAsset/SimulationPBDEdgeSpringConfigNode.h b/Engine/Plugins/ChaosClothAssetEditor/Source/ChaosClothAssetDataflowNodes/Public/ChaosClothAsset/SimulationPBDEdgeSpringConfigNode.h
--- a/Engine/Plugins/ChaosClothAssetEditor/Source/ChaosClothAssetDataflowNodes/Public/ChaosClothAsset/SimulationPBDEdgeSpringConfigNode.h
+++ b/Engine/Plugins/ChaosClothAssetEditor/Source/ChaosClothAssetDataflowNodes/Public/ChaosClothAsset/SimulationPBDEdgeSpringConfigNode.h
@@ -7,7 +7,7 @@
 #include "SimulationPBDEdgeSpringConfigNode.generated.h"
 
 /** Edge spring constraint property configuration node. */
-USTRUCT(Meta = (DataflowCloth))
+USTRUCT(Meta = (DataflowCloth, Deprecated = "5.4"))
 struct FChaosClothAssetSimulationPBDEdgeSpringConfigNode : public FChaosClothAssetSimulationBaseConfigNode
 {
 	GENERATED_USTRUCT_BODY()
@@ -28,5 +28,4 @@ public:
 
 private:
 	virtual void AddProperties(FPropertyHelper& PropertyHelper) const override;
-	virtual bool IsDeprecated() override { return true; }
 };
diff --git a/Engine/Plugins/ChaosClothAssetEditor/Source/ChaosClothAssetDataflowNodes/Public/ChaosClothAsset/SimulationStretchOverrideConfigNode.h b/Engine/Plugins/ChaosClothAssetEditor/Source/ChaosClothAssetDataflowNodes/Public/ChaosClothAsset/SimulationStretchOverrideConfigNode.h
--- a/Engine/Plugins/ChaosClothAssetEditor/Source/ChaosClothAssetDataflowNodes/Public/ChaosClothAsset/SimulationStretchOverrideConfigNode.h
+++ b/Engine/Plugins/ChaosClothAssetEditor/Source/ChaosClothAssetDataflowNodes/Public/ChaosClothAsset/SimulationStretchOverrideConfigNode.h
@@ -8,7 +8,7 @@
 #include "SimulationStretchOverrideConfigNode.generated.h"
 
 /** Stretching constraint property override configuration node. */
-USTRUCT(Meta = (DataflowCloth))
+USTRUCT(Meta = (DataflowCloth, Experimental))
 struct FChaosClothAssetSimulationStretchOverrideConfigNode : public FChaosClothAssetSimulationBaseConfigNode
 {
 	GENERATED_USTRUCT_BODY()
@@ -89,5 +89,4 @@ public:
 
 private:
 	virtual void AddProperties(FPropertyHelper& PropertyHelper) const override;
-	virtual bool IsExperimental() override { return true; }
 };
diff --git a/Engine/Plugins/ChaosClothAssetEditor/Source/ChaosClothAssetDataflowNodes/Public/ChaosClothAsset/SimulationXPBDAnisoBendingConfigNode.h b/Engine/Plugins/ChaosClothAssetEditor/Source/ChaosClothAssetDataflowNodes/Public/ChaosClothAsset/SimulationXPBDAnisoBendingConfigNode.h
--- a/Engine/Plugins/ChaosClothAssetEditor/Source/ChaosClothAssetDataflowNodes/Public/ChaosClothAsset/SimulationXPBDAnisoBendingConfigNode.h
+++ b/Engine/Plugins/ChaosClothAssetEditor/Source/ChaosClothAssetDataflowNodes/Public/ChaosClothAsset/SimulationXPBDAnisoBendingConfigNode.h
@@ -8,7 +8,7 @@
 #include "SimulationXPBDAnisoBendingConfigNode.generated.h"
 
 /** XPBD anisotropic bending constraint property configuration node. */
-USTRUCT(Meta = (DataflowCloth))
+USTRUCT(Meta = (DataflowCloth, Deprecated = "5.4"))
 struct FChaosClothAssetSimulationXPBDAnisoBendingConfigNode : public FChaosClothAssetSimulationBaseConfigNode
 {
 	GENERATED_USTRUCT_BODY()
@@ -126,5 +126,4 @@ public:
 
 private:
 	virtual void AddProperties(FPropertyHelper& PropertyHelper) const override;
-	virtual bool IsDeprecated() override { return true; }
 };
diff --git a/Engine/Plugins/ChaosClothAssetEditor/Source/ChaosClothAssetDataflowNodes/Public/ChaosClothAsset/SimulationXPBDAnisoSpringConfigNode.h b/Engine/Plugins/ChaosClothAssetEditor/Source/ChaosClothAssetDataflowNodes/Public/ChaosClothAsset/SimulationXPBDAnisoSpringConfigNode.h
--- a/Engine/Plugins/ChaosClothAssetEditor/Source/ChaosClothAssetDataflowNodes/Public/ChaosClothAsset/SimulationXPBDAnisoSpringConfigNode.h
+++ b/Engine/Plugins/ChaosClothAssetEditor/Source/ChaosClothAssetDataflowNodes/Public/ChaosClothAsset/SimulationXPBDAnisoSpringConfigNode.h
@@ -7,7 +7,7 @@
 #include "SimulationXPBDAnisoSpringConfigNode.generated.h"
 
 /** XPBD anisotropic spring constraint property configuration node. */
-USTRUCT(Meta = (DataflowCloth))
+USTRUCT(Meta = (DataflowCloth, Deprecated = "5.4"))
 struct FChaosClothAssetSimulationXPBDAnisoSpringConfigNode : public FChaosClothAssetSimulationBaseConfigNode
 {
 	GENERATED_USTRUCT_BODY()
@@ -76,5 +76,4 @@ public:
 
 private:
 	virtual void AddProperties(FPropertyHelper& PropertyHelper) const override;
-	virtual bool IsDeprecated() override { return true; }
 };
diff --git a/Engine/Plugins/ChaosClothAssetEditor/Source/ChaosClothAssetDataflowNodes/Public/ChaosClothAsset/SimulationXPBDAnisoStretchConfigNode.h b/Engine/Plugins/ChaosClothAssetEditor/Source/ChaosClothAssetDataflowNodes/Public/ChaosClothAsset/SimulationXPBDAnisoStretchConfigNode.h
--- a/Engine/Plugins/ChaosClothAssetEditor/Source/ChaosClothAssetDataflowNodes/Public/ChaosClothAsset/SimulationXPBDAnisoStretchConfigNode.h
+++ b/Engine/Plugins/ChaosClothAssetEditor/Source/ChaosClothAssetDataflowNodes/Public/ChaosClothAsset/SimulationXPBDAnisoStretchConfigNode.h
@@ -7,7 +7,7 @@
 #include "SimulationXPBDAnisoStretchConfigNode.generated.h"
 
 /** XPBD anisotropic stretch constraint property configuration node. */
-USTRUCT(Meta = (DataflowCloth))
+USTRUCT(Meta = (DataflowCloth, Deprecated = "5.4"))
 struct FChaosClothAssetSimulationXPBDAnisoStretchConfigNode : public FChaosClothAssetSimulationBaseConfigNode
 {
 	GENERATED_USTRUCT_BODY()
@@ -76,5 +76,4 @@ public:
 
 private:
 	virtual void AddProperties(FPropertyHelper& PropertyHelper) const override;
-	virtual bool IsDeprecated() override { return true; }
 };
diff --git a/Engine/Plugins/ChaosClothAssetEditor/Source/ChaosClothAssetDataflowNodes/Public/ChaosClothAsset/SimulationXPBDAreaSpringConfigNode.h b/Engine/Plugins/ChaosClothAssetEditor/Source/ChaosClothAssetDataflowNodes/Public/ChaosClothAsset/SimulationXPBDAreaSpringConfigNode.h
--- a/Engine/Plugins/ChaosClothAssetEditor/Source/ChaosClothAssetDataflowNodes/Public/ChaosClothAsset/SimulationXPBDAreaSpringConfigNode.h
+++ b/Engine/Plugins/ChaosClothAssetEditor/Source/ChaosClothAssetDataflowNodes/Public/ChaosClothAsset/SimulationXPBDAreaSpringConfigNode.h
@@ -7,7 +7,7 @@
 #include "SimulationXPBDAreaSpringConfigNode.generated.h"
 
 /** XPBD area spring constraint property configuration node. */
-USTRUCT(Meta = (DataflowCloth))
+USTRUCT(Meta = (DataflowCloth, Deprecated = "5.4"))
 struct FChaosClothAssetSimulationXPBDAreaSpringConfigNode : public FChaosClothAssetSimulationBaseConfigNode
 {
 	GENERATED_USTRUCT_BODY()
@@ -28,5 +28,4 @@ public:
 
 private:
 	virtual void AddProperties(FPropertyHelper& PropertyHelper) const override;
-	virtual bool IsDeprecated() override { return true; }
 };
diff --git a/Engine/Plugins/ChaosClothAssetEditor/Source/ChaosClothAssetDataflowNodes/Public/ChaosClothAsset/SimulationXPBDBendingElementConfigNode.h b/Engine/Plugins/ChaosClothAssetEditor/Source/ChaosClothAssetDataflowNodes/Public/ChaosClothAsset/SimulationXPBDBendingElementConfigNode.h
--- a/Engine/Plugins/ChaosClothAssetEditor/Source/ChaosClothAssetDataflowNodes/Public/ChaosClothAsset/SimulationXPBDBendingElementConfigNode.h
+++ b/Engine/Plugins/ChaosClothAssetEditor/Source/ChaosClothAssetDataflowNodes/Public/ChaosClothAsset/SimulationXPBDBendingElementConfigNode.h
@@ -8,7 +8,7 @@
 #include "SimulationXPBDBendingElementConfigNode.generated.h"
 
 /** XPBD bending element constraint property configuration node. */
-USTRUCT(Meta = (DataflowCloth))
+USTRUCT(Meta = (DataflowCloth, Deprecated = "5.4"))
 struct FChaosClothAssetSimulationXPBDBendingElementConfigNode : public FChaosClothAssetSimulationBaseConfigNode
 {
 	GENERATED_USTRUCT_BODY()
@@ -84,5 +84,4 @@ public:
 
 private:
 	virtual void AddProperties(FPropertyHelper& PropertyHelper) const override;
-	virtual bool IsDeprecated() override { return true; }
 };
diff --git a/Engine/Plugins/ChaosClothAssetEditor/Source/ChaosClothAssetDataflowNodes/Public/ChaosClothAsset/SimulationXPBDBendingSpringConfigNode.h b/Engine/Plugins/ChaosClothAssetEditor/Source/ChaosClothAssetDataflowNodes/Public/ChaosClothAsset/SimulationXPBDBendingSpringConfigNode.h
--- a/Engine/Plugins/ChaosClothAssetEditor/Source/ChaosClothAssetDataflowNodes/Public/ChaosClothAsset/SimulationXPBDBendingSpringConfigNode.h
+++ b/Engine/Plugins/ChaosClothAssetEditor/Source/ChaosClothAssetDataflowNodes/Public/ChaosClothAsset/SimulationXPBDBendingSpringConfigNode.h
@@ -7,7 +7,7 @@
 #include "SimulationXPBDBendingSpringConfigNode.generated.h"
 
 /** XPBD bending spring constraint property configuration node. */
-USTRUCT(Meta = (DataflowCloth))
+USTRUCT(Meta = (DataflowCloth, Deprecated = "5.4"))
 struct FChaosClothAssetSimulationXPBDBendingSpringConfigNode : public FChaosClothAssetSimulationBaseConfigNode
 {
 	GENERATED_USTRUCT_BODY()
@@ -36,5 +36,4 @@ public:
 
 private:
 	virtual void AddProperties(FPropertyHelper& PropertyHelper) const override;
-	virtual bool IsDeprecated() override { return true; }
 };
diff --git a/Engine/Plugins/ChaosClothAssetEditor/Source/ChaosClothAssetDataflowNodes/Public/ChaosClothAsset/SimulationXPBDEdgeSpringConfigNode.h b/Engine/Plugins/ChaosClothAssetEditor/Source/ChaosClothAssetDataflowNodes/Public/ChaosClothAsset/SimulationXPBDEdgeSpringConfigNode.h
--- a/Engine/Plugins/ChaosClothAssetEditor/Source/ChaosClothAssetDataflowNodes/Public/ChaosClothAsset/SimulationXPBDEdgeSpringConfigNode.h
+++ b/Engine/Plugins/ChaosClothAssetEditor/Source/ChaosClothAssetDataflowNodes/Public/ChaosClothAsset/SimulationXPBDEdgeSpringConfigNode.h
@@ -7,7 +7,7 @@
 #include "SimulationXPBDEdgeSpringConfigNode.generated.h"
 
 /** XPBD edge spring constraint property configuration node. */
-USTRUCT(Meta = (DataflowCloth))
+USTRUCT(Meta = (DataflowCloth, Deprecated = "5.2"))
 struct FChaosClothAssetSimulationXPBDEdgeSpringConfigNode : public FChaosClothAssetSimulationBaseConfigNode
 {
 	GENERATED_USTRUCT_BODY()
@@ -36,5 +36,4 @@ public:
 
 private:
 	virtual void AddProperties(FPropertyHelper& PropertyHelper) const override;
-	virtual bool IsDeprecated() override { return true; }
 };
diff --git a/Engine/Plugins/Experimental/Dataflow/Source/DataflowEditor/Private/Dataflow/DataflowSNode.cpp b/Engine/Plugins/Experimental/Dataflow/Source/DataflowEditor/Private/Dataflow/DataflowSNode.cpp
--- a/Engine/Plugins/Experimental/Dataflow/Source/DataflowEditor/Private/Dataflow/DataflowSNode.cpp
+++ b/Engine/Plugins/Experimental/Dataflow/Source/DataflowEditor/Private/Dataflow/DataflowSNode.cpp
@@ -131,19 +131,18 @@ TArray<FOverlayWidgetInfo> SDataflowEdNode::GetOverlayWidgets(bool bSelected, co
 	return Widgets;
 }
 
-
 void SDataflowEdNode::UpdateErrorInfo()
 {
 	if (DataflowGraphNode)
 	{
 		if (const TSharedPtr<FDataflowNode> DataflowNode = DataflowGraphNode->GetDataflowNode())
 		{
-			if (DataflowNode->IsExperimental())
+			if (Dataflow::FNodeFactory::IsNodeExperimental(DataflowNode->GetType()))
 			{
 				ErrorMsg = FString(TEXT("Experimental"));
 				ErrorColor = FAppStyle::GetColor("ErrorReporting.WarningBackgroundColor");
 			} 
-			else if (DataflowNode->IsDeprecated())
+			if (Dataflow::FNodeFactory::IsNodeDeprecated(DataflowNode->GetType()))
 			{
 				ErrorMsg = FString(TEXT("Deprecated"));
 				ErrorColor = FAppStyle::GetColor("ErrorReporting.WarningBackgroundColor");
@@ -153,7 +152,6 @@ void SDataflowEdNode::UpdateErrorInfo()
 	
 }
 
-
 FReply SDataflowEdNode::OnMouseButtonDoubleClick(const FGeometry& InMyGeometry, const FPointerEvent& InMouseEvent)
 {
 	if (GraphNode)
@@ -187,11 +185,10 @@ void SDataflowEdNode::AddReferencedObjects(FReferenceCollector& Collector)
 	}
 }
 
-
 //
 // Add a menu option to create a graph node.
 //
-TSharedPtr<FAssetSchemaAction_Dataflow_CreateNode_DataflowEdNode> FAssetSchemaAction_Dataflow_CreateNode_DataflowEdNode::CreateAction(UEdGraph* ParentGraph, const FName & InNodeTypeName)
+TSharedPtr<FAssetSchemaAction_Dataflow_CreateNode_DataflowEdNode> FAssetSchemaAction_Dataflow_CreateNode_DataflowEdNode::CreateAction(UEdGraph* ParentGraph, const FName & InNodeTypeName, const FName& InOverrideNodeName)
 {
 	if (Dataflow::FNodeFactory* Factory = Dataflow::FNodeFactory::GetInstance())
 	{
@@ -199,7 +196,11 @@ TSharedPtr<FAssetSchemaAction_Dataflow_CreateNode_DataflowEdNode> FAssetSchemaAc
 		if (Param.IsValid())
 		{
 			const FText ToolTip = FText::FromString(Param.ToolTip.IsEmpty() ? FString("Add a Dataflow node.") : Param.ToolTip);
-			const FText NodeName = FText::FromString(Param.DisplayName.ToString());
+			FText NodeName = FText::FromString(Param.DisplayName.ToString());
+			if (!InOverrideNodeName.IsNone())
+			{
+				NodeName = FText::FromName(InOverrideNodeName);
+			}
 			const FText Category = FText::FromString(Param.Category.ToString().IsEmpty() ? FString("Dataflow") : Param.Category.ToString());
 			const FText Tags = FText::FromString(Param.Tags);
 			TSharedPtr<FAssetSchemaAction_Dataflow_CreateNode_DataflowEdNode> NewNodeAction(
diff --git a/Engine/Plugins/Experimental/Dataflow/Source/DataflowEditor/Private/Dataflow/DataflowSchema.cpp b/Engine/Plugins/Experimental/Dataflow/Source/DataflowEditor/Private/Dataflow/DataflowSchema.cpp
--- a/Engine/Plugins/Experimental/Dataflow/Source/DataflowEditor/Private/Dataflow/DataflowSchema.cpp
+++ b/Engine/Plugins/Experimental/Dataflow/Source/DataflowEditor/Private/Dataflow/DataflowSchema.cpp
@@ -95,7 +95,7 @@ void UDataflowSchema::GetGraphContextActions(FGraphContextMenuBuilder& ContextMe
 			if (FDataflowEditorCommands::Get().CreateNodesMap.Contains(NodeParameters.TypeName))
 			{
 				if (TSharedPtr<FAssetSchemaAction_Dataflow_CreateNode_DataflowEdNode> Action =
-					FAssetSchemaAction_Dataflow_CreateNode_DataflowEdNode::CreateAction(ContextMenuBuilder.OwnerOfTemporaries, NodeParameters.TypeName))
+					FAssetSchemaAction_Dataflow_CreateNode_DataflowEdNode::CreateAction(ContextMenuBuilder.OwnerOfTemporaries, NodeParameters.TypeName, NodeParameters.DisplayName))
 				{
 					ContextMenuBuilder.AddAction(Action);
 				}
diff --git a/Engine/Plugins/Experimental/Dataflow/Source/DataflowEditor/Public/Dataflow/DataflowSNode.h b/Engine/Plugins/Experimental/Dataflow/Source/DataflowEditor/Public/Dataflow/DataflowSNode.h
--- a/Engine/Plugins/Experimental/Dataflow/Source/DataflowEditor/Public/Dataflow/DataflowSNode.h
+++ b/Engine/Plugins/Experimental/Dataflow/Source/DataflowEditor/Public/Dataflow/DataflowSNode.h
@@ -71,7 +71,7 @@ public:
 	FAssetSchemaAction_Dataflow_CreateNode_DataflowEdNode(const FName& InType, FText InNodeCategory, FText InMenuDesc, FText InToolTip, FText InKeywords)
 		: FEdGraphSchemaAction(InNodeCategory, InMenuDesc, InToolTip, 0, InKeywords), NodeTypeName(InType) {}
 
-	static TSharedPtr<FAssetSchemaAction_Dataflow_CreateNode_DataflowEdNode> CreateAction(UEdGraph* ParentGraph, const FName & NodeTypeName);
+	static TSharedPtr<FAssetSchemaAction_Dataflow_CreateNode_DataflowEdNode> CreateAction(UEdGraph* ParentGraph, const FName & NodeTypeName, const FName& InOverrideNodeName = NAME_None);
 
 	virtual UEdGraphNode* PerformAction(class UEdGraph* ParentGraph, UEdGraphPin* FromPin, const FVector2D Location, bool bSelectNewNode = true) override;
 	//virtual void AddReferencedObjects(FReferenceCollector& Collector) override;
diff --git a/Engine/Plugins/Experimental/GeometryCollectionPlugin/Source/GeometryCollectionNodes/Public/Dataflow/GeometryCollectionNodes.h b/Engine/Plugins/Experimental/GeometryCollectionPlugin/Source/GeometryCollectionNodes/Public/Dataflow/GeometryCollectionNodes.h
--- a/Engine/Plugins/Experimental/GeometryCollectionPlugin/Source/GeometryCollectionNodes/Public/Dataflow/GeometryCollectionNodes.h
+++ b/Engine/Plugins/Experimental/GeometryCollectionPlugin/Source/GeometryCollectionNodes/Public/Dataflow/GeometryCollectionNodes.h
@@ -1704,6 +1704,18 @@ public:
 
 };
 
+// Example to deprecate a Dataflow node
+ //USTRUCT(meta = (Deprecated = "5.1"))
+ //struct FLogStringDataflowNode : public FDataflowNode
+
+ // Example to version up a Dataflow node, the type name needs to be versioned up
+ // Important: don't change the display name!
+  //USTRUCT()
+  //struct FLogStringDataflowNode_v3 : public FDataflowNode
+
+// Example of experimental Dataflow node
+ //USTRUCT(meta = (Experimental))
+ //struct FLogStringDataflowNode_v2 : public FDataflowNode
 
 namespace Dataflow
 {
diff --git a/Engine/Source/Runtime/Experimental/Dataflow/Core/Private/Dataflow/DataflowNode.cpp b/Engine/Source/Runtime/Experimental/Dataflow/Core/Private/Dataflow/DataflowNode.cpp
--- a/Engine/Source/Runtime/Experimental/Dataflow/Core/Private/Dataflow/DataflowNode.cpp
+++ b/Engine/Source/Runtime/Experimental/Dataflow/Core/Private/Dataflow/DataflowNode.cpp
@@ -8,6 +8,7 @@
 #include "Serialization/ObjectWriter.h"
 #include "Serialization/ObjectReader.h"
 #include "Templates/TypeHash.h"
+#include "Dataflow/DataflowNodeFactory.h"
 
 #include UE_INLINE_GENERATED_CPP_BY_NAME(DataflowNode)
 
@@ -676,91 +677,9 @@ bool FDataflowNode::ValidateConnections()
 
 FString FDataflowNode::GetToolTip()
 {
-	if (const TUniquePtr<FStructOnScope> ScriptOnStruct = TUniquePtr<FStructOnScope>(NewStructOnScope()))
-	{
-#if WITH_EDITORONLY_DATA
-		if (const UStruct* Struct = ScriptOnStruct->GetStruct())
-		{
-			FString OutStr, InputsStr, OutputsStr;
-
-			FText StructText = Struct->GetToolTipText();
-
-			OutStr.Appendf(TEXT("%s\n\n%s\n"), *GetDisplayName().ToString(), *StructText.ToString());
-
-			for (FPropertyValueIterator PropertyIt(FProperty::StaticClass(), Struct, ScriptOnStruct->GetStructMemory()); PropertyIt; ++PropertyIt)
-			{
-				const FProperty* const Property = PropertyIt.Key();
-				check(Property);
-
-				if (Property->HasMetaData(TEXT("Tooltip")))
-				{
-					FString ToolTipStr = Property->GetToolTipText(true).ToString();
-					if (ToolTipStr.Len() > 0)
-					{
-						TArray<FString> OutArr;
-						ToolTipStr.ParseIntoArray(OutArr, TEXT(":\r\n"));
-						
-						if (OutArr.Num() == 0)
-						{
-							break;
-						}
-
-						TArray<const FProperty*> PropertyChain;
-						PropertyIt.GetPropertyChain(PropertyChain);
+	Dataflow::FFactoryParameters FactoryParameters = ::Dataflow::FNodeFactory::GetInstance()->GetParameters(GetType());
 
-						const FName PropName(GetPropertyFullName(PropertyChain));
-
-						const FString& MainTooltipText = (OutArr.Num() > 1) ? OutArr[1] : OutArr[0];
-
-						if (Property->HasMetaData(FDataflowNode::DataflowInput) &&
-							Property->HasMetaData(FDataflowNode::DataflowOutput) &&
-							Property->HasMetaData(FDataflowNode::DataflowPassthrough))
-						{
-							if (Property->HasMetaData(FDataflowNode::DataflowIntrinsic))
-							{
-								InputsStr.Appendf(TEXT("    %s [Intrinsic] - %s\n"), *PropName.ToString(), *MainTooltipText);
-							}
-							else
-							{
-								InputsStr.Appendf(TEXT("    %s - %s\n"), *PropName.ToString(), *MainTooltipText);
-							}
-
-							OutputsStr.Appendf(TEXT("    %s [Passthrough] - %s\n"), *PropName.ToString(), *MainTooltipText);
-						}					
-						else if (Property->HasMetaData(FDataflowNode::DataflowInput))
-						{
-							if (Property->HasMetaData(FDataflowNode::DataflowIntrinsic))
-							{
-								InputsStr.Appendf(TEXT("    %s [Intrinsic] - %s\n"), *PropName.ToString(), *MainTooltipText);
-							}
-							else
-							{
-								InputsStr.Appendf(TEXT("    %s - %s\n"), *PropName.ToString(), *MainTooltipText);
-							}
-						}
-						else if (Property->HasMetaData(FDataflowNode::DataflowOutput))
-						{
-							OutputsStr.Appendf(TEXT("    %s - %s\n"), *PropName.ToString(), *MainTooltipText);
-						}
-					}
-				}
-			}
-
-			if (InputsStr.Len() > 0)
-			{
-				OutStr.Appendf(TEXT("\n Input(s) :\n % s"), *InputsStr);
-			}
-			if (OutputsStr.Len() > 0)
-			{
-				OutStr.Appendf(TEXT("\n Output(s):\n%s"), *OutputsStr);
-			}
-
-			return OutStr;
-		}
-#endif
-	}
-
-	return "";
+	return FactoryParameters.ToolTip;
 }
 
 FText FDataflowNode::GetPinDisplayName(const FName& PropertyFullName)
diff --git a/Engine/Source/Runtime/Experimental/Dataflow/Core/Private/Dataflow/DataflowNodeFactory.cpp b/Engine/Source/Runtime/Experimental/Dataflow/Core/Private/Dataflow/DataflowNodeFactory.cpp
--- a/Engine/Source/Runtime/Experimental/Dataflow/Core/Private/Dataflow/DataflowNodeFactory.cpp
+++ b/Engine/Source/Runtime/Experimental/Dataflow/Core/Private/Dataflow/DataflowNodeFactory.cpp
@@ -42,41 +42,337 @@ namespace Dataflow
 	void FNodeFactory::RegisterNode(const FFactoryParameters& Parameters, FNewNodeFunction NewFunction)
 	{
 		bool bRegisterNode = true;
-		if (ClassMap.Contains(Parameters.TypeName) || DisplayMap.Contains(Parameters.DisplayName))
+
+		// To specify a new version of a node the node TypeName have to be versioned up in the definition in the header file, for example: FLogStringDataflowNode_v2
+		// or leave it without a version number and it will be modified automatically to <...>_V1
+		// The DisplayName doesn't have to have a version, just for example: LogString. But it has to be unique for DisplayMap and therefore it automatically 
+		// gets modified to NewDisplayName = DisplayName + _v<VERSION>
+		// 
+		// Update DisplayName with version from TypeName
+		FString NewDisplayNameString = Parameters.DisplayName.ToString() + "_" + Parameters.NodeVersion.ToString();
+
+		FFactoryParameters NewParameters(Parameters.TypeName, FName(*NewDisplayNameString), Parameters.Category, Parameters.Tags, Parameters.ToolTip, Parameters.bIsDeprecated, Parameters.bIsExperimental, Parameters.NodeVersion);
+
+		if (ClassMap.Contains(NewParameters.TypeName) || DisplayMap.Contains(NewParameters.DisplayName))
 		{
-			if (ParametersMap[Parameters.TypeName].DisplayName.IsEqual(Parameters.DisplayName))
+			if (ParametersMap[NewParameters.TypeName].DisplayName.IsEqual(NewParameters.DisplayName))
 			{
 				UE_LOG(LogDataflowFactory, Warning,
 					TEXT("Warning : Dataflow node registration mismatch with type(%s).The \
 						nodes have inconsistent display names(%s) vs(%s).There are two nodes \
-						with the same type being registered."), *Parameters.TypeName.ToString(),
-					*ParametersMap[Parameters.TypeName].DisplayName.ToString(),
-					*Parameters.DisplayName.ToString(), *Parameters.TypeName.ToString());
+						with the same type being registered."), *NewParameters.TypeName.ToString(),
+					*ParametersMap[NewParameters.TypeName].DisplayName.ToString(),
+					*NewParameters.DisplayName.ToString(), *NewParameters.TypeName.ToString());
 			}
-			if (ParametersMap[Parameters.TypeName].Category.IsEqual(Parameters.Category))
+			if (ParametersMap[NewParameters.TypeName].Category.IsEqual(NewParameters.Category))
 			{
 				UE_LOG(LogDataflowFactory, Warning,
 					TEXT("Warning : Dataflow node registration mismatch with type (%s). The nodes \
 						have inconsistent categories names (%s) vs (%s). There are two different nodes \
-						with the same type being registered. "), *Parameters.TypeName.ToString(),
-					*ParametersMap[Parameters.TypeName].DisplayName.ToString(),
-					*Parameters.DisplayName.ToString(), *Parameters.TypeName.ToString());
+						with the same type being registered. "), *NewParameters.TypeName.ToString(),
+					*ParametersMap[NewParameters.TypeName].DisplayName.ToString(),
+					*NewParameters.DisplayName.ToString(), *NewParameters.TypeName.ToString());
 			}
-			if (!ClassMap.Contains(Parameters.TypeName))
+			if (!ClassMap.Contains(NewParameters.TypeName))
 			{
 				UE_LOG(LogDataflowFactory, Warning,
 					TEXT("Warning: Attempted to register node type(%s) with display name (%s) \
 						that conflicts with an existing nodes display name (%s)."),
-					*Parameters.TypeName.ToString(), *Parameters.DisplayName.ToString(),
-					*ParametersMap[Parameters.TypeName].DisplayName.ToString());
+					*NewParameters.TypeName.ToString(), *NewParameters.DisplayName.ToString(),
+					*ParametersMap[NewParameters.TypeName].DisplayName.ToString());
 			}
 		}
 		else
 		{
-			ClassMap.Add(Parameters.TypeName, NewFunction);
-			ParametersMap.Add(Parameters.TypeName, Parameters);
-			DisplayMap.Add(Parameters.DisplayName, Parameters.TypeName);
+			ClassMap.Add(NewParameters.TypeName, NewFunction);
+			ParametersMap.Add(NewParameters.TypeName, NewParameters);
+			DisplayMap.Add(NewParameters.DisplayName, NewParameters.TypeName);
+
+			const FName TypeNameNoVersion = GetTypeNameNoVersion(NewParameters.TypeName);
+			VersionMap.FindOrAdd(TypeNameNoVersion).AddUnique(NewParameters.TypeName);
+		}
+	}
+
+	FName FNodeFactory::GetVersionFromTypeName(const FName& TypeName)
+	{
+		const FString String = TypeName.ToString();
+		const int32 Index = String.Find("_v", ESearchCase::IgnoreCase, ESearchDir::FromEnd);
+		if (Index == INDEX_NONE)
+		{
+			return FName(TEXT("v1"));
+		}
+		else
+		{
+			FString VersionString = String.Mid(Index + 1);
+			return FName(*VersionString);
+		}
+	}
+
+	int32 FNodeFactory::GetNumVersionFromVersion(const FName& Version)
+	{
+		const FString String = Version.ToString();
+		const FString NumVersionString = String.Mid(1);
+		return FCString::Atoi(*NumVersionString);
+	}
+
+	bool FNodeFactory::IsNodeDeprecated(const FName NodeType)
+	{
+		// Display node deprecated if InNode is deprecated
+		if (Dataflow::FNodeFactory* Factory = GetInstance())
+		{
+			const Dataflow::FFactoryParameters& Param = Factory->GetParameters(NodeType);
+			if (Param.IsValid() && Param.IsDeprecated())
+			{
+				return true;
+			}
+		}
+
+		return false;
+	}
+
+	bool FNodeFactory::IsNodeExperimental(const FName NodeType)
+	{
+		// Display node experimental if InNode is experimental
+		if (Dataflow::FNodeFactory* Factory = GetInstance())
+		{
+			const Dataflow::FFactoryParameters& Param = Factory->GetParameters(NodeType);
+			if (Param.IsValid() && Param.IsExperimental())
+			{
+				return true;
+			}
+		}
+
+		return false;
+	}
+
+	bool FNodeFactory::IsNodeDeprecated(const UStruct* Struct)
+	{
+#if WITH_EDITOR
+		return Struct->HasMetaData("Deprecated");
+#else
+		return false;
+#endif // WITH_EDITOR
+	}
+
+	bool FNodeFactory::IsNodeExperimental(const UStruct* Struct)
+	{
+#if WITH_EDITOR
+		return Struct->HasMetaData("Experimental");
+#else
+		return false;
+#endif // WITH_EDITOR
+	}
+
+	FName FNodeFactory::TypeNameFromDisplayName(const FName& DisplayName)
+	{
+		if (DisplayMap.Contains(DisplayName))
+		{
+			return DisplayMap[DisplayName];
 		}
+		return "";
 	}
+
+	const FFactoryParameters& FNodeFactory::GetParameters(FName InTypeName) const
+	{
+		static FFactoryParameters EmptyFactoryParameters;
+
+		if (ParametersMap.Contains(InTypeName))
+		{
+			return ParametersMap[InTypeName];
+		}
+
+		return EmptyFactoryParameters;
+	}
+
+	FName FNodeFactory::GetTypeNameNoVersion(const FName& TypeName)
+	{
+		FName TypeNameNoVersion = TypeName;
+
+		FString String = TypeName.ToString();
+		int32 Index = String.Find("_v", ESearchCase::IgnoreCase, ESearchDir::FromEnd);
+		if (Index != INDEX_NONE)
+		{
+			FString NoVersionString = String.Mid(0, Index);
+			TypeNameNoVersion = FName(*NoVersionString);
+
+			return TypeNameNoVersion;
+		}
+
+		return TypeNameNoVersion;
+	}
+
+	FName FNodeFactory::GetDisplayNameNoVersion(const FName& DisplayName)
+	{
+		FName DisplayNameNoVersion = DisplayName;
+
+		FString String = DisplayName.ToString();
+		int32 Index = String.Find("_v", ESearchCase::IgnoreCase, ESearchDir::FromEnd);
+		if (Index != INDEX_NONE)
+		{
+			FString NoVersionString = String.Mid(0, Index);
+			DisplayNameNoVersion = FName(*NoVersionString);
+
+			return DisplayNameNoVersion;
+		}
+
+		return DisplayNameNoVersion;
+	}
+
+	FString FNodeFactory::GetToolTipFromStruct(UScriptStruct* InStruct, const FName& InTypeName, const FName& InDisplayName)
+	{
+		FString OutStr, InputsStr, OutputsStr;
+#if WITH_EDITOR
+		FName NodeVersion = GetVersionFromTypeName(InTypeName);
+		bool bIsDeprecated = IsNodeDeprecated(InStruct);
+		bool bIsExperimental = IsNodeExperimental(InStruct);
+
+		FName NewDisplayName = GetDisplayNameNoVersion(InDisplayName);
+
+		OutStr.Appendf(TEXT("%s (%s)\n"), *NewDisplayName.ToString(), *NodeVersion.ToString());
+		if (bIsDeprecated)
+		{
+			OutStr.Appendf(TEXT("Deprecated\n"));
+		}
+		if (bIsExperimental)
+		{
+			OutStr.Appendf(TEXT("Experimental\n"));
+		}
+
+		FText StructText = InStruct->GetToolTipText();
+
+		OutStr.Appendf(TEXT("\n%s\n"), *StructText.ToString());
+
+		// Iterate over the properties
+		for (const FField* ChildProperty = InStruct->ChildProperties; ChildProperty; ChildProperty = ChildProperty->Next)
+		{
+			FName PropertyName = ChildProperty->GetFName();
+
+			if (ChildProperty->HasMetaData(TEXT("Tooltip")))
+			{
+				FString ToolTipStr = ChildProperty->GetToolTipText(true).ToString();
+
+				if (ToolTipStr.Len() > 0)
+				{
+					TArray<FString> OutArr;
+					ToolTipStr.ParseIntoArray(OutArr, TEXT(":\r\n"));
+
+					if (OutArr.Num() == 0)
+					{
+						break;
+					}
+
+					const FString& MainTooltipText = (OutArr.Num() > 1) ? OutArr[1] : OutArr[0];
+
+					if (ChildProperty->HasMetaData(FDataflowNode::DataflowInput) &&
+						ChildProperty->HasMetaData(FDataflowNode::DataflowOutput) &&
+						ChildProperty->HasMetaData(FDataflowNode::DataflowPassthrough))
+					{
+						if (ChildProperty->HasMetaData(FDataflowNode::DataflowIntrinsic))
+						{
+							InputsStr.Appendf(TEXT("    %s [Intrinsic] - %s\n"), *PropertyName.ToString(), *MainTooltipText);
+						}
+						else
+						{
+							InputsStr.Appendf(TEXT("    %s - %s\n"), *PropertyName.ToString(), *MainTooltipText);
+						}
+
+						OutputsStr.Appendf(TEXT("    %s [Passthrough] - %s\n"), *PropertyName.ToString(), *MainTooltipText);
+					}
+					else if (ChildProperty->HasMetaData(FDataflowNode::DataflowInput))
+					{
+						if (ChildProperty->HasMetaData(FDataflowNode::DataflowIntrinsic))
+						{
+							InputsStr.Appendf(TEXT("    %s [Intrinsic] - %s\n"), *PropertyName.ToString(), *MainTooltipText);
+						}
+						else
+						{
+							InputsStr.Appendf(TEXT("    %s - %s\n"), *PropertyName.ToString(), *MainTooltipText);
+						}
+					}
+					else if (ChildProperty->HasMetaData(FDataflowNode::DataflowOutput))
+					{
+						OutputsStr.Appendf(TEXT("    %s - %s\n"), *PropertyName.ToString(), *MainTooltipText);
+					}
+				}
+			}
+		}
+
+		if (InputsStr.Len() > 0)
+		{
+			OutStr.Appendf(TEXT("\n Input(s) :\n % s"), *InputsStr);
+		}
+
+		if (OutputsStr.Len() > 0)
+		{
+			OutStr.Appendf(TEXT("\n Output(s):\n%s"), *OutputsStr);
+		}
+#endif // WITH_EDITOR
+		return OutStr;
+	}
+
+	//
+	// Building Context menu
+	//
+	// If a node is deprecated -> omit from this list
+	// If a node has single version -> don't display version
+	//
+	TArray<FFactoryParameters> FNodeFactory::RegisteredParameters() const
+	{
+		TArray<FFactoryParameters> RetVal;
+
+		for (auto Elem : VersionMap)
+		{
+			// Check how many version(s) the node has
+			TArray<FFactoryParameters> ParametersArray;
+			for (const FName& VersionedTypeName : Elem.Value)
+			{
+				const FFactoryParameters& FactoryParameters = GetParameters(VersionedTypeName);
+				if (!FactoryParameters.IsDeprecated())
+				{
+					ParametersArray.Add(FactoryParameters);
+				}
+			}
+
+			if (ParametersArray.Num() == 1)
+			{
+				// 
+				// There is only one version of the node
+				// Do not display version in DisplayName
+				//
+				ParametersArray[0].DisplayName = GetDisplayNameNoVersion(ParametersArray[0].DisplayName);
+				
+				if (ParametersArray[0].bIsExperimental)
+				{
+					FString NewDisplayName = ParametersArray[0].DisplayName.ToString() + " (Experimental)";
+					ParametersArray[0].DisplayName = FName(*NewDisplayName);
+				}
+
+				RetVal.Add(ParametersArray[0]);
+			}
+			else if (ParametersArray.Num() > 1)
+			{
+				// 
+				// There are multiple versions of the node
+				// DisplayName: DisplayName (v2)
+				//
+				for (FFactoryParameters& FactoryParameters : ParametersArray)
+				{
+					FString NewDisplayName = GetDisplayNameNoVersion(FactoryParameters.DisplayName).ToString() + " (v" + FString::FromInt(GetNumVersionFromVersion(FactoryParameters.NodeVersion)) + ")";
+					FactoryParameters.DisplayName = FName(*NewDisplayName);
+
+					if (FactoryParameters.bIsExperimental)
+					{
+						NewDisplayName = FactoryParameters.DisplayName.ToString() + " (Experimental)";
+						FactoryParameters.DisplayName = FName(*NewDisplayName);
+					}
+
+					RetVal.Add(FactoryParameters);
+				}
+			}
+		}
+
+		return RetVal;
+	}
+
 }
 
diff --git a/Engine/Source/Runtime/Experimental/Dataflow/Core/Public/Dataflow/DataflowNode.h b/Engine/Source/Runtime/Experimental/Dataflow/Core/Public/Dataflow/DataflowNode.h
--- a/Engine/Source/Runtime/Experimental/Dataflow/Core/Public/Dataflow/DataflowNode.h
+++ b/Engine/Source/Runtime/Experimental/Dataflow/Core/Public/Dataflow/DataflowNode.h
@@ -85,7 +85,9 @@ struct FDataflowNode
 	UE_DEPRECATED(5.4, "FDataflowNode::CopyNodeProperties is deprecated.")
 	DATAFLOWCORE_API void CopyNodeProperties(const TSharedPtr<FDataflowNode> CopyFromDataflowNode);
 
+	UE_DEPRECATED(5.4, "FDataflowNode::IsDeprecated is deprecated.")
 	virtual bool IsDeprecated() { return false; }
+	UE_DEPRECATED(5.4, "FDataflowNode::IsExperimental is deprecated.")
 	virtual bool IsExperimental() { return false; }
 
 	//
@@ -343,13 +345,7 @@ namespace Dataflow
 	//
 
 #define DATAFLOW_NODE_REGISTER_CREATION_FACTORY(A)									\
-	::Dataflow::FNodeFactory::GetInstance()->RegisterNode(							\
-		{A::StaticType(),A::StaticDisplay(),A::StaticCategory(),					\
-			A::StaticTags(),A::StaticToolTip()},									\
-		[](const ::Dataflow::FNewNodeParameters& InParam){							\
-				TUniquePtr<A> Val = MakeUnique<A>(::Dataflow::FNodeParameters{		\
-					InParam.Name, InParam.OwningObject}, InParam.Guid);				\
-				Val->ValidateProperties();Val->ValidateConnections();return Val;});
+	::Dataflow::FNodeFactory::RegisterNodeFromType<A>();
 
 #define DATAFLOW_NODE_RENDER_TYPE(A, B)												\
 	virtual TArray<::Dataflow::FRenderingParameter> GetRenderParametersImpl() const {		\
diff --git a/Engine/Source/Runtime/Experimental/Dataflow/Core/Public/Dataflow/DataflowNodeFactory.h b/Engine/Source/Runtime/Experimental/Dataflow/Core/Public/Dataflow/DataflowNodeFactory.h
--- a/Engine/Source/Runtime/Experimental/Dataflow/Core/Public/Dataflow/DataflowNodeFactory.h
+++ b/Engine/Source/Runtime/Experimental/Dataflow/Core/Public/Dataflow/DataflowNodeFactory.h
@@ -20,18 +20,33 @@ namespace Dataflow
 
 	struct FFactoryParameters {
 		FFactoryParameters() {}
-		FFactoryParameters(FName InTypeName, FName InDisplayName, FName InCategory, FString InTags, FString InToolTip)
-		: TypeName(InTypeName), DisplayName(InDisplayName), Category(InCategory), Tags(InTags), ToolTip(InToolTip) {}
+		FFactoryParameters(FName InTypeName, FName InDisplayName, FName InCategory, FString InTags, FString InToolTip, bool bInIsDeprecated, bool bInIsExperimental, FName InNodeVersion)
+		: TypeName(InTypeName), DisplayName(InDisplayName), Category(InCategory), Tags(InTags), ToolTip(InToolTip), bIsDeprecated(bInIsDeprecated), bIsExperimental(bInIsExperimental), NodeVersion(InNodeVersion) {}
 
 		FName TypeName = FName("");
 		FName DisplayName = FName("");
 		FName Category = FName("");
 		FString Tags = FString("");
 		FString ToolTip = FString("");
+		bool bIsDeprecated = false;
+		bool bIsExperimental = false;
+		FName NodeVersion = FName(TEXT("v1"));
 
 		bool IsValid() const {
 			return !TypeName.ToString().IsEmpty() && !DisplayName.ToString().IsEmpty();
 		}
+
+		bool IsDeprecated() const {
+			return bIsDeprecated;
+		}
+
+		bool IsExperimental() const {
+			return bIsExperimental;
+		}
+
+		FName GetVersion() const {
+			return NodeVersion;
+		}
 	};
 
 	//
@@ -42,9 +57,10 @@ namespace Dataflow
 		typedef TFunction<TUniquePtr<FDataflowNode> (const FNewNodeParameters&)> FNewNodeFunction;
 
 		// All Maps indexed by TypeName
-		TMap<FName, FNewNodeFunction > ClassMap;		// [TypeName] -> NewNodeFunction
-		TMap<FName, FFactoryParameters > ParametersMap;	// [TypeName] -> Parameters
-		TMap<FName, FName > DisplayMap;					// [DisplayName] -> TypeName
+		TMap<FName, FNewNodeFunction > ClassMap;			// [TypeName] -> NewNodeFunction
+		TMap<FName, FFactoryParameters > ParametersMap;		// [TypeName] -> Parameters
+		TMap<FName, FName > DisplayMap;						// [DisplayName] -> TypeName
+		TMap<FName, TArray<FName>> VersionMap;				// [TypeNameNoVersion] -> Array of TypeName(versions)
 
 		DATAFLOWCORE_API static FNodeFactory* Instance;
 		FNodeFactory() {};
@@ -62,42 +78,68 @@ namespace Dataflow
 			return Instance;
 		}
 
-		DATAFLOWCORE_API void RegisterNode(const FFactoryParameters& Parameters, FNewNodeFunction NewFunction);
-
-		FName TypeNameFromDisplayName(const FName& DisplayName)
+		template<typename T>
+		static void RegisterNodeFromType()
 		{
-			if (DisplayMap.Contains(DisplayName))
+			FFactoryParameters FactoryParameters =
 			{
-				return DisplayMap[DisplayName];
-			}
-			return "";
-		}
-
-		FFactoryParameters GetParameters(FName InTypeName) 
-		{
-			if (ParametersMap.Contains(InTypeName))
+				T::StaticType(),
+				T::StaticDisplay(),
+				T::StaticCategory(),
+				T::StaticTags(),
+				GetToolTipFromStruct(T::StaticStruct(), T::StaticType(), T::StaticDisplay()),
+				IsNodeDeprecated(T::StaticStruct()),
+				IsNodeExperimental(T::StaticStruct()),
+				GetVersionFromTypeName(T::StaticType())
+			};
+
+			auto CreationFunction = [](const ::Dataflow::FNewNodeParameters& InParam)
 			{
-				return ParametersMap[InTypeName];
-			}
-			return FFactoryParameters();
+				const FNodeParameters NodeParameters = { InParam.Name, InParam.OwningObject };
+				TUniquePtr<T> Val = MakeUnique<T>(NodeParameters, InParam.Guid);
+				Val->ValidateProperties();
+				Val->ValidateConnections();
+				return Val;
+			};
+
+			GetInstance()->RegisterNode(FactoryParameters, CreationFunction);
 		}
 
+		DATAFLOWCORE_API const FFactoryParameters& GetParameters(FName InTypeName) const;
+
 		DATAFLOWCORE_API TSharedPtr<FDataflowNode> NewNodeFromRegisteredType(FGraph& Graph, const FNewNodeParameters& Param);
 
+		DATAFLOWCORE_API TArray<FFactoryParameters> RegisteredParameters() const;
+
+		DATAFLOWCORE_API static bool IsNodeDeprecated(const FName NodeType);
+
+		DATAFLOWCORE_API static bool IsNodeExperimental(const FName NodeType);
+
+	private:
+		DATAFLOWCORE_API void RegisterNode(const FFactoryParameters& Parameters, FNewNodeFunction NewFunction);
+
+		DATAFLOWCORE_API static FName GetVersionFromTypeName(const FName& TypeName);
+
+		DATAFLOWCORE_API static int32 GetNumVersionFromVersion(const FName& Version);
+
+		DATAFLOWCORE_API static bool IsNodeDeprecated(const UStruct* Struct);
+
+		DATAFLOWCORE_API static bool IsNodeExperimental(const UStruct* Struct);
+
+		DATAFLOWCORE_API static FName GetTypeNameNoVersion(const FName& TypeName);
+
+		DATAFLOWCORE_API static FName GetDisplayNameNoVersion(const FName& DisplayName);
+
+		DATAFLOWCORE_API static FString GetToolTipFromStruct(UScriptStruct* InStruct, const FName& InTypeName, const FName& InDisplayName);
+
+		DATAFLOWCORE_API FName TypeNameFromDisplayName(const FName& DisplayName);
+
 		template<class T> TSharedPtr<T> NewNode(FGraph& Graph, const FNewNodeParameters& Param)
 		{
 			return Graph.AddNode(new T(Param.Name, Param.Guid));
 		}
 
-		TArray<FFactoryParameters> RegisteredParameters() const
-		{
-			TArray<FFactoryParameters> RetVal;
-			for (auto Elem : ParametersMap) RetVal.Add(Elem.Value);
-			return RetVal;
-		}
-	private:
 		DATAFLOWCORE_API void RegisterDefaultNodes();
-
 	};
 }
 
