Iris
 - Implemented support for capturing and caching CreationInfo which allows us to capture final snapshot of replicated objects data including information we need to send creation info to clients. This allows flushing of replicated state and rpcs even if the owning actor is destroyed.
 - Iris now supports flushing state for Networked objects created and destroyed during the same frame.
 - Autoflush behavior can be disabled by cvar net.Iris.EnableFlushReliableRPCOnDestroy 0

diff --git a/Engine/Plugins/Runtime/ReplicationSystemTestPlugin/Source/Private/Tests/ReplicationSystem/ReplicatedTestObject.cpp b/Engine/Plugins/Runtime/ReplicationSystemTestPlugin/Source/Private/Tests/ReplicationSystem/ReplicatedTestObject.cpp
--- a/Engine/Plugins/Runtime/ReplicationSystemTestPlugin/Source/Private/Tests/ReplicationSystem/ReplicatedTestObject.cpp
+++ b/Engine/Plugins/Runtime/ReplicationSystemTestPlugin/Source/Private/Tests/ReplicationSystem/ReplicatedTestObject.cpp
@@ -102,52 +102,82 @@ UE::Net::FNetRefHandle UReplicatedTestObjectBridge::BeginReplication(FNetRefHand
 
 bool UReplicatedTestObjectBridge::WriteCreationHeader(UE::Net::FNetSerializationContext& Context, FNetRefHandle Handle)
 {
-	UE::Net::FNetBitStreamWriter& Writer = *Context.GetBitStreamWriter();
-
-	uint16 NumComponentsToSpawn = 0U;
-	uint16 NumIrisComponentsToSpawn = 0U;
-	uint16 NumDynamicComponentsToSpawn = 0U;
-	uint16 NumConnectionFilteredComponentsToSpawn = 0U;
-	uint16 NumObjectReferenceComponentsToSpawn = 0U;
-	bool bForceFailToInstantiateOnRemote = false;
+	TUniquePtr<FCreationHeader> Header(GetCreationHeader(Handle));
+	return WriteCreationHeader(Context, Header.Get());
+}
 
+TUniquePtr<UObjectReplicationBridge::FCreationHeader> UReplicatedTestObjectBridge::GetCreationHeader(FNetRefHandle Handle)
+{
 	const UObject* Object = GetReplicatedObject(Handle);
+
+	if (!ensure(Object))
+	{
+		return nullptr;
+	}
+
+	TUniquePtr<FReplicationTestObjectCreationHeader> Header(new FReplicationTestObjectCreationHeader);
+
+	Header->NumComponentsToSpawn = 0U;
+	Header->NumIrisComponentsToSpawn = 0U;
+	Header->NumDynamicComponentsToSpawn = 0U;
+	Header->NumConnectionFilteredComponentsToSpawn = 0U;
+	Header->NumObjectReferenceComponentsToSpawn = 0U;
+	Header->bForceFailCreateRemoteInstance = false;
+
 	if (const UReplicatedTestObject* ReplicatedTestObject = Cast<UReplicatedTestObject>(Object))
 	{
-		bForceFailToInstantiateOnRemote = ReplicatedTestObject->bForceFailToInstantiateOnRemote;
+		Header->bForceFailCreateRemoteInstance = ReplicatedTestObject->bForceFailToInstantiateOnRemote;
 	}
 	if (const UTestReplicatedIrisObject* TestReplicatedIrisObject = Cast<UTestReplicatedIrisObject>(Object))
 	{
-		NumComponentsToSpawn = IntCastChecked<uint16>(TestReplicatedIrisObject->Components.Num());
-		NumIrisComponentsToSpawn = IntCastChecked<uint16>(TestReplicatedIrisObject->IrisComponents.Num());
-		NumDynamicComponentsToSpawn = IntCastChecked<uint16>(TestReplicatedIrisObject->DynamicStateComponents.Num());
-		NumConnectionFilteredComponentsToSpawn = IntCastChecked<uint16>(TestReplicatedIrisObject->ConnectionFilteredComponents.Num());
-		NumObjectReferenceComponentsToSpawn = IntCastChecked<uint16>(TestReplicatedIrisObject->ObjectReferenceComponents.Num());
+		Header->NumComponentsToSpawn = IntCastChecked<uint16>(TestReplicatedIrisObject->Components.Num());
+		Header->NumIrisComponentsToSpawn = IntCastChecked<uint16>(TestReplicatedIrisObject->IrisComponents.Num());
+		Header->NumDynamicComponentsToSpawn = IntCastChecked<uint16>(TestReplicatedIrisObject->DynamicStateComponents.Num());
+		Header->NumConnectionFilteredComponentsToSpawn = IntCastChecked<uint16>(TestReplicatedIrisObject->ConnectionFilteredComponents.Num());
+		Header->NumObjectReferenceComponentsToSpawn = IntCastChecked<uint16>(TestReplicatedIrisObject->ObjectReferenceComponents.Num());
 	}
 
 	UObject* Archetype = Object->GetArchetype();
-	check(Archetype);
+	if (!Archetype)
+	{
+		check(Archetype);
+		return nullptr;
+	}
+
+	Header->ArchetypeName = Archetype->GetPathName();
+
+	return Header;
+}
 
-	WriteString(&Writer, Archetype->GetPathName());
-	Writer.WriteBits(NumComponentsToSpawn, 16);
-	Writer.WriteBits(NumIrisComponentsToSpawn, 16);
-	Writer.WriteBits(NumDynamicComponentsToSpawn, 16);
-	Writer.WriteBits(NumConnectionFilteredComponentsToSpawn, 16);
-	Writer.WriteBits(NumObjectReferenceComponentsToSpawn, 16);
-	Writer.WriteBool(bForceFailToInstantiateOnRemote);
+bool UReplicatedTestObjectBridge::WriteCreationHeader(UE::Net::FNetSerializationContext& Context, const UObjectReplicationBridge::FCreationHeader* InHeader)
+{
+	if (!InHeader)
+	{
+		return false;
+	}
+
+	const FReplicationTestObjectCreationHeader* Header = static_cast<const FReplicationTestObjectCreationHeader*>(InHeader);
+
+	UE::Net::FNetBitStreamWriter& Writer = *Context.GetBitStreamWriter();
+
+	WriteString(&Writer, Header->ArchetypeName);
+	Writer.WriteBits(Header->NumComponentsToSpawn, 16);
+	Writer.WriteBits(Header->NumIrisComponentsToSpawn, 16);
+	Writer.WriteBits(Header->NumDynamicComponentsToSpawn, 16);
+	Writer.WriteBits(Header->NumConnectionFilteredComponentsToSpawn, 16);
+	Writer.WriteBits(Header->NumObjectReferenceComponentsToSpawn, 16);
+	Writer.WriteBool(Header->bForceFailCreateRemoteInstance);
 
 	return !Writer.IsOverflown();
 }
 
-UObjectReplicationBridge::FCreationHeader* UReplicatedTestObjectBridge::ReadCreationHeader(UE::Net::FNetSerializationContext& Context)
+TUniquePtr<UObjectReplicationBridge::FCreationHeader> UReplicatedTestObjectBridge::ReadCreationHeader(UE::Net::FNetSerializationContext& Context)
 {
 	UE::Net::FNetBitStreamReader& Reader = *Context.GetBitStreamReader();
 	TUniquePtr<FReplicationTestObjectCreationHeader> Header(new FReplicationTestObjectCreationHeader);
 
 	ReadString(&Reader, Header->ArchetypeName);
 
-	//FNetObjectReference ArcheTypeRef;
-
 	Header->NumComponentsToSpawn = Reader.ReadBits(16);
 	Header->NumIrisComponentsToSpawn = Reader.ReadBits(16);
 	Header->NumDynamicComponentsToSpawn = Reader.ReadBits(16);
@@ -160,7 +190,7 @@ UObjectReplicationBridge::FCreationHeader* UReplicatedTestObjectBridge::ReadCrea
 		return nullptr;
 	}
 
-	return Header.Release();
+	return Header;
 }
 
 FObjectReplicationBridgeInstantiateResult UReplicatedTestObjectBridge::BeginInstantiateFromRemote(FNetRefHandle RootObjectOfSubObject, const UE::Net::FNetObjectResolveContext& ResolveContext, const FCreationHeader* InHeader)
diff --git a/Engine/Plugins/Runtime/ReplicationSystemTestPlugin/Source/Private/Tests/ReplicationSystem/ReplicatedTestObject.h b/Engine/Plugins/Runtime/ReplicationSystemTestPlugin/Source/Private/Tests/ReplicationSystem/ReplicatedTestObject.h
--- a/Engine/Plugins/Runtime/ReplicationSystemTestPlugin/Source/Private/Tests/ReplicationSystem/ReplicatedTestObject.h
+++ b/Engine/Plugins/Runtime/ReplicationSystemTestPlugin/Source/Private/Tests/ReplicationSystem/ReplicatedTestObject.h
@@ -547,8 +547,10 @@ protected:
 		bool bForceFailCreateRemoteInstance;
 	};
 
+	virtual TUniquePtr<FCreationHeader> GetCreationHeader(FNetRefHandle Handle) override;
 	virtual bool WriteCreationHeader(UE::Net::FNetSerializationContext& Context, FNetRefHandle Handle) override;
-	virtual FCreationHeader* ReadCreationHeader(UE::Net::FNetSerializationContext& Context) override;
+	virtual bool WriteCreationHeader(UE::Net::FNetSerializationContext& Context, const FCreationHeader* Header) override;
+	virtual TUniquePtr<FCreationHeader> ReadCreationHeader(UE::Net::FNetSerializationContext& Context) override;
 
 	virtual FObjectReplicationBridgeInstantiateResult BeginInstantiateFromRemote(FNetRefHandle RootObjectOfSubObject, const UE::Net::FNetObjectResolveContext& ResolveContext, const FCreationHeader* InHeader) override;
 	virtual void EndInstantiateFromRemote(FNetRefHandle Handle) override;
diff --git a/Engine/Plugins/Runtime/ReplicationSystemTestPlugin/Source/Private/Tests/ReplicationSystem/TestFlushBeforeDestroy.cpp b/Engine/Plugins/Runtime/ReplicationSystemTestPlugin/Source/Private/Tests/ReplicationSystem/TestFlushBeforeDestroy.cpp
--- a/Engine/Plugins/Runtime/ReplicationSystemTestPlugin/Source/Private/Tests/ReplicationSystem/TestFlushBeforeDestroy.cpp
+++ b/Engine/Plugins/Runtime/ReplicationSystemTestPlugin/Source/Private/Tests/ReplicationSystem/TestFlushBeforeDestroy.cpp
@@ -52,12 +52,199 @@ UE_NET_TEST_FIXTURE(FTestFlushBeforeDestroyFixture, TestReliableAttachmentFlushe
 	Server->PostSendUpdate();
 
 	// Verify that object is destroyed
+	UE_NET_ASSERT_EQ(Client->GetReplicationBridge()->GetReplicatedObject(ObjectHandle), nullptr);
+}
+
+UE_NET_TEST_FIXTURE(FTestFlushBeforeDestroyFixture, TestObjectCreatedAndDestroyedSameFrameReplicatesIfFlushed)
+{
+	FReplicationSystemTestClient* Client = CreateClient();
+	RegisterNetBlobHandlers(Client);
+
+	// Create and start to replicate object
+	UReplicatedTestObject* ServerObject = Server->CreateObject(0, 0);
+	FNetRefHandle ObjectHandle = ServerObject->NetRefHandle;
+
+	// Destroy object indicating that it should be flushed that is that the final state should be replicated to all clients with the object in scope, this invalidates the creationinfo which has to be cached in order for this to work.
+	Server->DestroyObject(ServerObject, EEndReplicationFlags::Destroy | EEndReplicationFlags::Flush);
+
+	// Send update, it should send the data.
+	Server->UpdateAndSend({Client});
+
+	// Verify that object is created
+	UE_NET_ASSERT_NE(Client->GetReplicationBridge()->GetReplicatedObject(ObjectHandle), nullptr);
+
+	// Deliver a packet, make sure that object is destroyed on the client.
+	Server->UpdateAndSend({Client});
+
+	// Verify that object is destroyed
+	UE_NET_ASSERT_EQ(Client->GetReplicationBridge()->GetReplicatedObject(ObjectHandle), nullptr);
+}
+
+UE_NET_TEST_FIXTURE(FTestFlushBeforeDestroyFixture, TestObjectAndSubObjectCreatedAndDestroyedSameFrameReplicatesIfFlushed)
+{
+	FReplicationSystemTestClient* Client = CreateClient();
+	RegisterNetBlobHandlers(Client);
+
+	// Create and start to replicate object
+	UReplicatedTestObject* ServerObject = Server->CreateObject(0, 0);
+	FNetRefHandle ObjectHandle = ServerObject->NetRefHandle;
+	UReplicatedTestObject* ServerSubObject = Server->CreateSubObject(ObjectHandle, 0, 0);
+	FNetRefHandle SubObjectHandle = ServerSubObject->NetRefHandle;
+
+	// Destroy object indicating that it should be flushed that is that the final state should be replicated to all clients with the object in scope, this invalidates the creationinfo which has to be cached in order for this to work.
+	Server->DestroyObject(ServerObject, EEndReplicationFlags::Destroy | EEndReplicationFlags::Flush);
+
+	// Send update, it should send the data.
+	Server->UpdateAndSend({Client});
+
+	// Verify that objects are created
+	UE_NET_ASSERT_TRUE(Client->GetReplicationBridge()->GetReplicatedObject(ObjectHandle) != nullptr);
+	UE_NET_ASSERT_TRUE(Client->GetReplicationBridge()->GetReplicatedObject(SubObjectHandle) != nullptr);
+
+	// Deliver a packet, make sure that object is destroyed on the client.
+	Server->UpdateAndSend({Client});
+
+	// Verify that objects are destroyed
 	UE_NET_ASSERT_TRUE(Client->GetReplicationBridge()->GetReplicatedObject(ObjectHandle) == nullptr);
+	UE_NET_ASSERT_TRUE(Client->GetReplicationBridge()->GetReplicatedObject(SubObjectHandle) == nullptr);
 }
 
-UE_NET_TEST_FIXTURE(FTestFlushBeforeDestroyFixture, TestReliableAttachmentFlushedWithDataInflightBeforeDestroy)
+UE_NET_TEST_FIXTURE(FTestFlushBeforeDestroyFixture, TestSubObjectCreatedAndDestroyedSameFrameReplicatesIfFlushed)
 {
+	FReplicationSystemTestClient* Client = CreateClient();
+	RegisterNetBlobHandlers(Client);
+
+	// Create and start to replicate object
+	UReplicatedTestObject* ServerObject = Server->CreateObject(0, 0);
+	FNetRefHandle ObjectHandle = ServerObject->NetRefHandle;
+	UReplicatedTestObject* ServerSubObject = Server->CreateSubObject(ObjectHandle, 0, 0);
+	FNetRefHandle SubObjectHandle = ServerSubObject->NetRefHandle;
 
+	// Destroy SubObject indicating that it should be flushed that is that the final state should be replicated to all clients with the object in scope, this invalidates the creationinfo which has to be cached in order for this to work.
+	Server->DestroyObject(ServerSubObject, EEndReplicationFlags::Destroy | EEndReplicationFlags::Flush);
+
+	// Send update, it should send the data.
+	Server->UpdateAndSend({Client});
+
+	// Verify that objects are created
+	UE_NET_ASSERT_TRUE(Client->GetReplicationBridge()->GetReplicatedObject(ObjectHandle) != nullptr);
+	UE_NET_ASSERT_TRUE(Client->GetReplicationBridge()->GetReplicatedObject(SubObjectHandle) != nullptr);
+
+	// Deliver a packet, make sure that object is destroyed on the client.
+	Server->UpdateAndSend({Client});
+
+	// Verify that objects are destroyed
+	UE_NET_ASSERT_TRUE(Client->GetReplicationBridge()->GetReplicatedObject(ObjectHandle) != nullptr);
+	UE_NET_ASSERT_TRUE(Client->GetReplicationBridge()->GetReplicatedObject(SubObjectHandle) == nullptr);
+}
+
+UE_NET_TEST_FIXTURE(FTestFlushBeforeDestroyFixture, TestReliableAttachmentFlushedBeforeDestroyIfObjectCreatedAndDestroyedSameFrame)
+{
+	FReplicationSystemTestClient* Client = CreateClient();
+	RegisterNetBlobHandlers(Client);
+
+	// Create and start to replicate object
+	UReplicatedTestObject* ServerObject = Server->CreateObject(0, 0);
+	FNetRefHandle ObjectHandle = ServerObject->NetRefHandle;
+
+	// Create attachment
+	{
+		constexpr uint32 PayloadBitCount = 24;
+		const TRefCountPtr<FNetObjectAttachment>& Attachment = MockNetObjectAttachmentHandler->CreateReliableNetObjectAttachment(PayloadBitCount);
+		FNetObjectReference AttachmentTarget = FObjectReferenceCache::MakeNetObjectReference(ServerObject->NetRefHandle);
+		Server->GetReplicationSystem()->QueueNetObjectAttachment(Client->ConnectionIdOnServer, AttachmentTarget, Attachment);
+	}
+
+	// Destroy object, it should be implicitly flushed due to pending attachment.
+	Server->DestroyObject(ServerObject, EEndReplicationFlags::Destroy);
+
+	// Send update, it should send the data.
+	Server->UpdateAndSend({Client});
+
+	// Verify that the attachment has been received
+	UE_NET_ASSERT_EQ(ClientMockNetObjectAttachmentHandler->GetFunctionCallCounts().OnNetBlobReceived, 1U);
+
+	// Deliver a packet, make sure that object is destroyed on the client.
+	Server->UpdateAndSend({Client});
+
+	// Verify that object is destroyed
+	UE_NET_ASSERT_EQ(Client->GetReplicationBridge()->GetReplicatedObject(ObjectHandle), nullptr);
+}
+
+UE_NET_TEST_FIXTURE(FTestFlushBeforeDestroyFixture, TestReliableAttachmentForSubObjectFlushedBeforeDestroyIfObjectCreatedAndDestroyedSameFrame)
+{
+	FReplicationSystemTestClient* Client = CreateClient();
+	RegisterNetBlobHandlers(Client);
+
+	// Create and start to replicate object with subobject
+	UReplicatedTestObject* ServerObject = Server->CreateObject(0, 0);
+	FNetRefHandle ObjectHandle = ServerObject->NetRefHandle;
+	UReplicatedTestObject* ServerSubObject = Server->CreateSubObject(ObjectHandle, 0, 0);
+	FNetRefHandle SubObjectHandle = ServerSubObject->NetRefHandle;
+
+	// Create attachment
+	{
+		constexpr uint32 PayloadBitCount = 24;
+		const TRefCountPtr<FNetObjectAttachment>& Attachment = MockNetObjectAttachmentHandler->CreateReliableNetObjectAttachment(PayloadBitCount);
+		FNetObjectReference AttachmentTarget = FObjectReferenceCache::MakeNetObjectReference(SubObjectHandle);
+		Server->GetReplicationSystem()->QueueNetObjectAttachment(Client->ConnectionIdOnServer, AttachmentTarget, Attachment);
+	}
+
+	// Destroy object, it should be implicitly flushed due to pending attachment.
+	Server->DestroyObject(ServerObject, EEndReplicationFlags::Destroy);
+
+	// Send update, it should send the data.
+	Server->UpdateAndSend({Client});
+
+	// Verify that the attachment has been received
+	UE_NET_ASSERT_EQ(ClientMockNetObjectAttachmentHandler->GetFunctionCallCounts().OnNetBlobReceived, 1U);
+
+	// Deliver a packet, make sure that objects is destroyed on the client.
+	Server->UpdateAndSend({Client});
+
+	// Verify that objects are destroyed
+	UE_NET_ASSERT_EQ(Client->GetReplicationBridge()->GetReplicatedObject(ObjectHandle), nullptr);
+	UE_NET_ASSERT_EQ(Client->GetReplicationBridge()->GetReplicatedObject(SubObjectHandle), nullptr);
+}
+
+UE_NET_TEST_FIXTURE(FTestFlushBeforeDestroyFixture, TestReliableAttachmentForSubObjecFlushedBeforeDestroyIfSubObjectCreatedAndDestroyedSameFrame)
+{
+	FReplicationSystemTestClient* Client = CreateClient();
+	RegisterNetBlobHandlers(Client);
+
+	// Create and start to replicate object with subobject
+	UReplicatedTestObject* ServerObject = Server->CreateObject(0, 0);
+	FNetRefHandle ObjectHandle = ServerObject->NetRefHandle;
+	UReplicatedTestObject* ServerSubObject = Server->CreateSubObject(ObjectHandle, 0, 0);
+	FNetRefHandle SubObjectHandle = ServerSubObject->NetRefHandle;
+
+	// Create attachment
+	{
+		constexpr uint32 PayloadBitCount = 24;
+		const TRefCountPtr<FNetObjectAttachment>& Attachment = MockNetObjectAttachmentHandler->CreateReliableNetObjectAttachment(PayloadBitCount);
+		FNetObjectReference AttachmentTarget = FObjectReferenceCache::MakeNetObjectReference(SubObjectHandle);
+		Server->GetReplicationSystem()->QueueNetObjectAttachment(Client->ConnectionIdOnServer, AttachmentTarget, Attachment);
+	}
+
+	// Destroy subobject, it should be implicitly flushed due to pending attachment.
+	Server->DestroyObject(ServerSubObject, EEndReplicationFlags::Destroy);
+
+	// Send update, it should send the data.
+	Server->UpdateAndSend({Client});
+
+	// Verify that the attachment has been received
+	UE_NET_ASSERT_EQ(ClientMockNetObjectAttachmentHandler->GetFunctionCallCounts().OnNetBlobReceived, 1U);
+
+	// Deliver a packet, make sure that object is destroyed on the client.
+	Server->UpdateAndSend({Client});
+
+	// Verify that objects is destroyed
+	UE_NET_ASSERT_NE(Client->GetReplicationBridge()->GetReplicatedObject(ObjectHandle), nullptr);
+	UE_NET_ASSERT_EQ(Client->GetReplicationBridge()->GetReplicatedObject(SubObjectHandle), nullptr);
+}
+
+UE_NET_TEST_FIXTURE(FTestFlushBeforeDestroyFixture, TestReliableAttachmentFlushedWithDataInflightBeforeDestroy)
+{
 	FReplicationSystemTestClient* Client = CreateClient();
 	RegisterNetBlobHandlers(Client);
 
@@ -103,7 +290,7 @@ UE_NET_TEST_FIXTURE(FTestFlushBeforeDestroyFixture, TestReliableAttachmentFlushe
 	Server->PostSendUpdate();
 
 	// Verify that object is destroyed
-	UE_NET_ASSERT_TRUE(Client->GetReplicationBridge()->GetReplicatedObject(ObjectHandle) == nullptr);
+	UE_NET_ASSERT_EQ(Client->GetReplicationBridge()->GetReplicatedObject(ObjectHandle), nullptr);
 }
 
 // This test exercises what was a bad case where we was posting RPC:s to a not yet confirmed objects which also was marked for destroy
@@ -161,7 +348,7 @@ UE_NET_TEST_FIXTURE(FTestFlushBeforeDestroyFixture, TestReliableAttachmentFlushe
 	UE_NET_ASSERT_EQ(ClientMockNetObjectAttachmentHandler->GetFunctionCallCounts().OnNetBlobReceived, 0U);
 
 	// Verify that object does not exist
-	UE_NET_ASSERT_TRUE(Client->GetReplicationBridge()->GetReplicatedObject(ObjectHandle) == nullptr);
+	UE_NET_ASSERT_EQ(Client->GetReplicationBridge()->GetReplicatedObject(ObjectHandle), nullptr);
 }
 
 UE_NET_TEST_FIXTURE(FTestFlushBeforeDestroyFixture, TestReliableAttachmentFlushedWithPendingCreationInflightBeforeDestroy)
@@ -222,7 +409,7 @@ UE_NET_TEST_FIXTURE(FTestFlushBeforeDestroyFixture, TestReliableAttachmentFlushe
 	Server->PostSendUpdate();
 
 	// Verify that object does not exist
-	UE_NET_ASSERT_TRUE(Client->GetReplicationBridge()->GetReplicatedObject(ObjectHandle) == nullptr);
+	UE_NET_ASSERT_EQ(Client->GetReplicationBridge()->GetReplicatedObject(ObjectHandle), nullptr);
 }
 
 UE_NET_TEST_FIXTURE(FTestFlushBeforeDestroyFixture, TestReliableAttachmentSubObjectFlushedBeforeDestroy)
@@ -266,7 +453,7 @@ UE_NET_TEST_FIXTURE(FTestFlushBeforeDestroyFixture, TestReliableAttachmentSubObj
 	Server->PostSendUpdate();
 
 	// Verify that object is destroyed
-	UE_NET_ASSERT_TRUE(Client->GetReplicationBridge()->GetReplicatedObject(SubObjectHandle) == nullptr);
+	UE_NET_ASSERT_EQ(Client->GetReplicationBridge()->GetReplicatedObject(SubObjectHandle), nullptr);
 }
 
 UE_NET_TEST_FIXTURE(FTestFlushBeforeDestroyFixture, TestReliableAttachmentSubObjectFlushedBeforeDestroyIfOwnerIsDestroyed)
@@ -310,8 +497,8 @@ UE_NET_TEST_FIXTURE(FTestFlushBeforeDestroyFixture, TestReliableAttachmentSubObj
 	Server->PostSendUpdate();
 
 	// Verify that both object and subobject are destroyed
-	UE_NET_ASSERT_TRUE(Client->GetReplicationBridge()->GetReplicatedObject(SubObjectHandle) == nullptr);
-	UE_NET_ASSERT_TRUE(Client->GetReplicationBridge()->GetReplicatedObject(ObjectHandle) == nullptr);
+	UE_NET_ASSERT_EQ(Client->GetReplicationBridge()->GetReplicatedObject(SubObjectHandle), nullptr);
+	UE_NET_ASSERT_EQ(Client->GetReplicationBridge()->GetReplicatedObject(ObjectHandle), nullptr);
 }
 
 UE_NET_TEST_FIXTURE(FTestFlushBeforeDestroyFixture, TestStateFlushedBeforeDestroy)
@@ -354,7 +541,7 @@ UE_NET_TEST_FIXTURE(FTestFlushBeforeDestroyFixture, TestStateFlushedBeforeDestro
 	Server->PostSendUpdate();
 
 	// Verify that object is destroyed
-	UE_NET_ASSERT_TRUE(Client->GetReplicationBridge()->GetReplicatedObject(ObjectHandle) == nullptr);
+	UE_NET_ASSERT_EQ(Client->GetReplicationBridge()->GetReplicatedObject(ObjectHandle), nullptr);
 }
 
 UE_NET_TEST_FIXTURE(FTestFlushBeforeDestroyFixture, TestStateInFlightFlushedBeforeDestroy)
@@ -408,7 +595,7 @@ UE_NET_TEST_FIXTURE(FTestFlushBeforeDestroyFixture, TestStateInFlightFlushedBefo
 	Server->PostSendUpdate();
 
 	// Verify that object is destroyed
-	UE_NET_ASSERT_TRUE(Client->GetReplicationBridge()->GetReplicatedObject(ObjectHandle) == nullptr);
+	UE_NET_ASSERT_EQ(Client->GetReplicationBridge()->GetReplicatedObject(ObjectHandle), nullptr);
 }
 
 UE_NET_TEST_FIXTURE(FTestFlushBeforeDestroyFixture, TestSubObjectStateFlushedBeforeOwnerDestroy)
diff --git a/Engine/Plugins/Runtime/ReplicationSystemTestPlugin/Source/Private/Tests/ReplicationSystem/TestReplicationSystem.cpp b/Engine/Plugins/Runtime/ReplicationSystemTestPlugin/Source/Private/Tests/ReplicationSystem/TestReplicationSystem.cpp
--- a/Engine/Plugins/Runtime/ReplicationSystemTestPlugin/Source/Private/Tests/ReplicationSystem/TestReplicationSystem.cpp
+++ b/Engine/Plugins/Runtime/ReplicationSystemTestPlugin/Source/Private/Tests/ReplicationSystem/TestReplicationSystem.cpp
@@ -973,8 +973,7 @@ UE_NET_TEST_FIXTURE(FReplicationSystemServerClientTestFixture, TestTearOffResend
 	UE_NET_ASSERT_EQ(Cast<UTestReplicatedIrisObject>(Client->GetReplicationBridge()->GetReplicatedObject(ServerObject->NetRefHandle)), nullptr);
 }
 
-// Test TearOff for new object and resend (should not work or is this what we want?)
-#if 0 // Until we either keep object around but out of scope, or cache creation info + deps
+// Test TearOff for new object and resend, this requires creation info to be cached.
 UE_NET_TEST_FIXTURE(FReplicationSystemServerClientTestFixture, TestTearOffImmediateOnNewlyCreatedObjectResend)
 {
 	UReplicationSystem* ReplicationSystem = Server->ReplicationSystem;
@@ -995,14 +994,10 @@ UE_NET_TEST_FIXTURE(FReplicationSystemServerClientTestFixture, TestTearOffImmedi
 	Server->ReplicationBridge->EndReplication(ServerObject, EEndReplicationFlags::TearOff);
 
 	// Send and deliver packet
-	Server->PreSendUpdate();
-	Server->SendAndDeliverTo(Client, false);
-	Server->PostSendUpdate();
+	Server->UpdateAndSend({Client}, false);
 
 	// Send and deliver packet
-	Server->PreSendUpdate();
-	Server->SendAndDeliverTo(Client, true);
-	Server->PostSendUpdate();
+	Server->UpdateAndSend({Client});
 
 	// Client should have created a object
 	UE_NET_ASSERT_EQ(NumObjectsCreatedOnClientBeforeReplication + 1, Client->CreatedObjects.Num());
@@ -1015,8 +1010,55 @@ UE_NET_TEST_FIXTURE(FReplicationSystemServerClientTestFixture, TestTearOffImmedi
 	// Verify that we replicated the expected state
 	UE_NET_ASSERT_EQ(ServerObject->IntA, ClientObjectThatWillBeTornOff->IntA);
 }
-#endif
 
+// Test TearOff for new subobject and resend, this requires creation info to be cached.
+UE_NET_TEST_FIXTURE(FReplicationSystemServerClientTestFixture, TestTearOffImmediateOnNewlyCreatedSubObjectResend)
+{
+	UReplicationSystem* ReplicationSystem = Server->ReplicationSystem;
+
+	// Add a client
+	FReplicationSystemTestClient* Client = CreateClient();
+
+	// We should not have any created objects
+	const int32 NumObjectsCreatedOnClientBeforeReplication = Client->CreatedObjects.Num();
+
+	// Spawn object on server
+	UTestReplicatedIrisObject* ServerObject = Server->CreateObject(0,0);
+
+	// Set state
+	ServerObject->IntA = 1;
+
+	// Send and deliver packet
+	Server->UpdateAndSend({Client});
+
+	// Spawn second object on server as a subobject
+	UTestReplicatedIrisObject* ServerSubObject = Server->CreateSubObject(ServerObject->NetRefHandle, 0, 0);
+
+	// Set state
+	ServerSubObject->IntA = 1;
+
+	// TearOff the subobject before first replication
+	Server->ReplicationBridge->EndReplication(ServerSubObject, EEndReplicationFlags::TearOff);
+
+	// Send and drop
+	Server->UpdateAndSend({Client}, false);
+
+	// Send and deliver packet
+	Server->UpdateAndSend({Client});
+
+	// Client should have created a object + subobject
+	UE_NET_ASSERT_EQ(NumObjectsCreatedOnClientBeforeReplication + 2, Client->CreatedObjects.Num());
+
+	// But as we have torn off the subobject it should no longer be a replicated object
+	UE_NET_ASSERT_TRUE(Cast<UTestReplicatedIrisObject>(Client->GetReplicationBridge()->GetReplicatedObject(ServerObject->NetRefHandle)) != nullptr);
+	UE_NET_ASSERT_TRUE(Cast<UTestReplicatedIrisObject>(Client->GetReplicationBridge()->GetReplicatedObject(ServerSubObject->NetRefHandle)) == nullptr);
+
+	// We should be able to get the object from the created objects array to validate the state
+	UTestReplicatedIrisObject* ClientSubObjectThatWillBeTornOff = Cast<UTestReplicatedIrisObject>(Client->CreatedObjects[NumObjectsCreatedOnClientBeforeReplication + 1].Get());
+
+	// Verify that we replicated the expected state
+	UE_NET_ASSERT_EQ(ServerSubObject->IntA, ClientSubObjectThatWillBeTornOff->IntA);
+}
 
 UE_NET_TEST_FIXTURE(FReplicationSystemServerClientTestFixture, TestDefferedTearOffOnNewlyCreatedObjectResend)
 {
diff --git a/Engine/Source/Runtime/Engine/Private/Net/Iris/ReplicationSystem/ActorReplicationBridge.cpp b/Engine/Source/Runtime/Engine/Private/Net/Iris/ReplicationSystem/ActorReplicationBridge.cpp
--- a/Engine/Source/Runtime/Engine/Private/Net/Iris/ReplicationSystem/ActorReplicationBridge.cpp
+++ b/Engine/Source/Runtime/Engine/Private/Net/Iris/ReplicationSystem/ActorReplicationBridge.cpp
@@ -519,7 +519,80 @@ bool UActorReplicationBridge::WriteCreationHeader(UE::Net::FNetSerializationCont
 	return false;
 }
 
-UObjectReplicationBridge::FCreationHeader* UActorReplicationBridge::ReadCreationHeader(UE::Net::FNetSerializationContext& Context)
+TUniquePtr<UObjectReplicationBridge::FCreationHeader> UActorReplicationBridge::GetCreationHeader(FNetRefHandle Handle)
+{
+	using namespace UE::Net;
+	using namespace UE::Net::Private;
+
+	if (!ensure(IsReplicatedHandle(Handle)))
+	{
+		return nullptr;
+	}
+
+	const UObject* Object = GetReplicatedObject(Handle);
+	if (const AActor* Actor = Cast<AActor>(Object))
+	{
+		// Get Header
+		TUniquePtr<FActorCreationHeader> Header(new FActorCreationHeader);
+		GetActorCreationHeader(Actor, *(Header.Get()));
+	
+		return Header;
+	}
+	else if (Object)
+	{
+		const UObject* RootObject = GetReplicatedObject(GetRootObjectOfSubObject(Handle));
+
+		// Get Header
+		TUniquePtr<FSubObjectCreationHeader> Header(new FSubObjectCreationHeader);
+		GetSubObjectCreationHeader(Object, RootObject, *(Header.Get()));
+
+		return Header;
+	}
+
+	ensureMsgf(false, TEXT("UActorReplicationBridge::GetCreationHeader Failed to get creationHeader for NetRefHandle (Id=%u) %s"), Handle.GetId(), ToCStr(GetReplicationSystem()->GetDebugName(Handle)));
+
+	return nullptr;
+}
+
+bool UActorReplicationBridge::WriteCreationHeader(UE::Net::FNetSerializationContext& Context, const FCreationHeader* InHeader)
+{
+	using namespace UE::Net;
+	using namespace UE::Net::Private;
+
+	FNetBitStreamWriter* Writer = Context.GetBitStreamWriter();
+
+	if (!InHeader)
+	{
+		return false;
+	}
+
+	const FActorReplicationBridgeCreationHeader* BridgeHeader = static_cast<const FActorReplicationBridgeCreationHeader*>(InHeader);
+
+	if (BridgeHeader->bIsActor)
+	{
+		const FActorCreationHeader* Header = static_cast<const FActorCreationHeader*>(BridgeHeader);
+
+		// Serialize the data
+		// Indicate that this is an actor
+		Writer->WriteBool(true);
+		WriteActorCreationHeader(Context, *Header, SpawnInfoFlags);
+
+		return !Writer->IsOverflown();
+	}
+	else
+	{
+		const FSubObjectCreationHeader* Header = static_cast<const FSubObjectCreationHeader*>(BridgeHeader);
+
+		// Serialize the data
+		// Indicate that this is a SubObject
+		Writer->WriteBool(false);
+		WriteSubObjectCreationHeader(Context, *Header);
+
+		return !Writer->IsOverflown();
+	}
+}
+
+TUniquePtr<UObjectReplicationBridge::FCreationHeader> UActorReplicationBridge::ReadCreationHeader(UE::Net::FNetSerializationContext& Context)
 {
 	using namespace UE::Net;
 	using namespace UE::Net::Private;
@@ -533,7 +606,7 @@ UObjectReplicationBridge::FCreationHeader* UActorReplicationBridge::ReadCreation
 
 		if (!Reader->IsOverflown())
 		{
-			return Header.Release();
+			return Header;
 		}
 	}
 	else
@@ -543,7 +616,7 @@ UObjectReplicationBridge::FCreationHeader* UActorReplicationBridge::ReadCreation
 
 		if (!Reader->IsOverflown())
 		{
-			return Header.Release();
+			return Header;
 		}
 	}
 
diff --git a/Engine/Source/Runtime/Engine/Public/Net/Iris/ReplicationSystem/ActorReplicationBridge.h b/Engine/Source/Runtime/Engine/Public/Net/Iris/ReplicationSystem/ActorReplicationBridge.h
--- a/Engine/Source/Runtime/Engine/Public/Net/Iris/ReplicationSystem/ActorReplicationBridge.h
+++ b/Engine/Source/Runtime/Engine/Public/Net/Iris/ReplicationSystem/ActorReplicationBridge.h
@@ -97,7 +97,9 @@ protected:
 	virtual void Initialize(UReplicationSystem* ReplicationSystem) override;
 	virtual void Deinitialize() override;
 	virtual bool WriteCreationHeader(UE::Net::FNetSerializationContext& Context, FNetRefHandle Handle) override;
-	virtual FCreationHeader* ReadCreationHeader(UE::Net::FNetSerializationContext& Context) override;
+	virtual TUniquePtr<FCreationHeader> GetCreationHeader(FNetRefHandle Handle) override;
+	virtual bool WriteCreationHeader(UE::Net::FNetSerializationContext& Context, const FCreationHeader* Header) override;
+	virtual TUniquePtr<FCreationHeader> ReadCreationHeader(UE::Net::FNetSerializationContext& Context) override;
 	virtual FObjectReplicationBridgeInstantiateResult BeginInstantiateFromRemote(FNetRefHandle RootObjectOfSubObject, const UE::Net::FNetObjectResolveContext& ResolveContext, const FCreationHeader* InHeader) override;
 	virtual bool OnInstantiatedFromRemote(UObject* Instance, const FCreationHeader* InHeader, uint32 ConnectionId) const override;
 	virtual void OnSubObjectCreatedFromReplication(FNetRefHandle SubObjectHandle) override;
diff --git a/Engine/Source/Runtime/Experimental/Iris/Core/Private/Iris/ReplicationSystem/NetBlob/NetBlobManager.cpp b/Engine/Source/Runtime/Experimental/Iris/Core/Private/Iris/ReplicationSystem/NetBlob/NetBlobManager.cpp
--- a/Engine/Source/Runtime/Experimental/Iris/Core/Private/Iris/ReplicationSystem/NetBlob/NetBlobManager.cpp
+++ b/Engine/Source/Runtime/Experimental/Iris/Core/Private/Iris/ReplicationSystem/NetBlob/NetBlobManager.cpp
@@ -560,6 +560,17 @@ void FNetBlobManager::FNetObjectAttachmentSendQueue::ResetProcessQueue()
 	ProcessContext.Reset();
 }
 
+bool FNetBlobManager::HasUnprocessedReliableAttachments(FInternalNetRefIndex InternalIndex) const
+{
+	return AttachmentSendQueue.HasUnprocessedReliableAttachments(InternalIndex);
+}
+
+bool FNetBlobManager::FNetObjectAttachmentSendQueue::HasUnprocessedReliableAttachments(FInternalNetRefIndex InternalIndex)  const
+{
+	// For the moment we only need to check the AttachmentQueue as reliable attachments are not schedules as immediate.
+	return AttachmentQueue.ContainsByPredicate([&InternalIndex](const FNetObjectAttachmentQueueEntry& Entry) { return (Entry.OwnerIndex == InternalIndex || Entry.SubObjectIndex == InternalIndex) && Entry.Attachment->IsReliable();} );
+}
+
 void FNetBlobManager::FNetObjectAttachmentSendQueue::ProcessQueue(EProcessMode ProcessMode)
 {
 	if (!ProcessContext.IsValid())
diff --git a/Engine/Source/Runtime/Experimental/Iris/Core/Private/Iris/ReplicationSystem/NetBlob/NetBlobManager.h b/Engine/Source/Runtime/Experimental/Iris/Core/Private/Iris/ReplicationSystem/NetBlob/NetBlobManager.h
--- a/Engine/Source/Runtime/Experimental/Iris/Core/Private/Iris/ReplicationSystem/NetBlob/NetBlobManager.h
+++ b/Engine/Source/Runtime/Experimental/Iris/Core/Private/Iris/ReplicationSystem/NetBlob/NetBlobManager.h
@@ -49,6 +49,8 @@ public:
 	bool SendRPC(const UObject* Object, const UObject* SubObject, const UFunction* Function, const void* Parameters, ENetObjectAttachmentSendPolicyFlags SendFlags = ENetObjectAttachmentSendPolicyFlags::None);
 	bool SendRPC(uint32 ConnectionId, const UObject* Object, const UObject* SubObject, const UFunction* Function, const void* Parameters, ENetObjectAttachmentSendPolicyFlags SendFlags = ENetObjectAttachmentSendPolicyFlags::None);
 
+	bool HasUnprocessedReliableAttachments(FInternalNetRefIndex InternalIndex)  const;
+
 	enum class EProcessMode 
 	{
 		ProcessObjectsGoingOutOfScope,
@@ -122,6 +124,8 @@ private:
 		void ProcessQueue(EProcessMode ProcessMode);
 		void ResetProcessQueue();
 		void PrepareAndProcessOOBAttachmentQueue(FReplicationConnections* InConnections, const FNetRefHandleManager* InNetRefHandleManager, FNetBitArray& OutConnetionsPendingImmediateSend);
+
+		bool HasUnprocessedReliableAttachments(FInternalNetRefIndex InternalIndex)  const;
 	
 	private:
 		struct FNetObjectAttachmentQueueEntry
diff --git a/Engine/Source/Runtime/Experimental/Iris/Core/Private/Iris/ReplicationSystem/NetRefHandleManager.h b/Engine/Source/Runtime/Experimental/Iris/Core/Private/Iris/ReplicationSystem/NetRefHandleManager.h
--- a/Engine/Source/Runtime/Experimental/Iris/Core/Private/Iris/ReplicationSystem/NetRefHandleManager.h
+++ b/Engine/Source/Runtime/Experimental/Iris/Core/Private/Iris/ReplicationSystem/NetRefHandleManager.h
@@ -104,6 +104,7 @@ public:
 				uint32 bNeedsFullCopyAndQuantize : 1U;
 				uint32 bWantsFullPoll : 1U;
 				uint32 bPendingEndReplication : 1U;
+				uint32 bHasCachedCreationInfo : 1U;
 			};
 		};
 	
diff --git a/Engine/Source/Runtime/Experimental/Iris/Core/Private/Iris/ReplicationSystem/ObjectReplicationBridge.cpp b/Engine/Source/Runtime/Experimental/Iris/Core/Private/Iris/ReplicationSystem/ObjectReplicationBridge.cpp
--- a/Engine/Source/Runtime/Experimental/Iris/Core/Private/Iris/ReplicationSystem/ObjectReplicationBridge.cpp
+++ b/Engine/Source/Runtime/Experimental/Iris/Core/Private/Iris/ReplicationSystem/ObjectReplicationBridge.cpp
@@ -596,8 +596,28 @@ bool UObjectReplicationBridge::WriteNetRefHandleCreationInfo(FReplicationBridgeS
 	const UE::Net::FReplicationProtocol* Protocol = GetReplicationSystem()->GetReplicationProtocol(Handle);
 	WriteUint64(Context.SerializationContext.GetBitStreamWriter(), Protocol->ProtocolIdentifier);
 
-	// Write Type header
-	return WriteCreationHeader(Context.SerializationContext, Handle);
+	// Write Type header, if there is a cached one, use it!
+	if (TUniquePtr<const FCreationHeader>* CachedHeader = CachedCreationHeaders.Find(Handle))
+	{
+		return WriteCreationHeader(Context.SerializationContext, (*CachedHeader).Get());
+	}
+	else
+	{
+		return WriteCreationHeader(Context.SerializationContext, Handle);
+	}
+}
+
+bool UObjectReplicationBridge::CacheNetRefHandleCreationInfo(FNetRefHandle Handle)
+{
+	TUniquePtr<const FCreationHeader> Header(GetCreationHeader(Handle));
+
+	if (Header)
+	{
+		CachedCreationHeaders.Add(Handle, MoveTemp(Header));
+		return true;
+	}
+
+	return false;
 }
 
 void UObjectReplicationBridge::EndReplication(UObject* Instance, EEndReplicationFlags EndReplicationFlags, FEndReplicationParameters* Parameters)
@@ -638,11 +658,15 @@ void UObjectReplicationBridge::DetachInstanceFromRemote(FNetRefHandle Handle, ER
 	}
 
 	// $IRIS TODO: Cleanup any pending creation data if we have not yet instantiated the instance.
+
+	Super::DetachInstanceFromRemote(Handle, DestroyReason, DestroyFlags);
 }
 
 void UObjectReplicationBridge::DetachInstance(FNetRefHandle RefHandle)
 {
 	UnregisterInstance(RefHandle);
+	CachedCreationHeaders.Remove(RefHandle);
+	Super::DetachInstance(RefHandle);
 }
 
 void UObjectReplicationBridge::UnregisterInstance(FNetRefHandle RefHandle)
diff --git a/Engine/Source/Runtime/Experimental/Iris/Core/Private/Iris/ReplicationSystem/ReplicationBridge.cpp b/Engine/Source/Runtime/Experimental/Iris/Core/Private/Iris/ReplicationSystem/ReplicationBridge.cpp
--- a/Engine/Source/Runtime/Experimental/Iris/Core/Private/Iris/ReplicationSystem/ReplicationBridge.cpp
+++ b/Engine/Source/Runtime/Experimental/Iris/Core/Private/Iris/ReplicationSystem/ReplicationBridge.cpp
@@ -2,7 +2,9 @@
 
 #include "Iris/ReplicationSystem/ReplicationBridge.h"
 #include "Containers/ArrayView.h"
+#include "HAL/IConsoleManager.h"
 #include "Iris/IrisConfigInternal.h"
+#include "Iris/Core/IrisCsv.h"
 #include "Iris/Core/IrisLog.h"
 #include "Iris/Core/IrisProfiler.h"
 #include "Net/Core/NetBitArray.h"
@@ -30,6 +32,14 @@ DEFINE_LOG_CATEGORY(LogIrisBridge)
 
 #define UE_LOG_REPLICATIONBRIDGE(Category, Format, ...)  UE_LOG(LogIrisBridge, Category, TEXT("ReplicationBridge(%u)::") Format, ReplicationSystem->GetId(), ##__VA_ARGS__)
 
+
+static bool bEnableFlushReliableRPCOnDestroy = true;
+static FAutoConsoleVariableRef CVarEnableFlushReliableRPCOnDestroy(
+	TEXT("net.Iris.EnableFlushReliableRPCOnDestroy"),
+	bEnableFlushReliableRPCOnDestroy,
+	TEXT("When true EEndReplicationFlags::Flush flag will be appended in EndReplication if we have pending unprocessed attachments/RPC:s when destroying a replicated object.")
+);
+
 /**
  * ReplicationBridge Implementation
  */
@@ -47,6 +57,11 @@ bool UReplicationBridge::WriteNetRefHandleCreationInfo(FReplicationBridgeSeriali
 	return true;
 }
 
+bool UReplicationBridge::CacheNetRefHandleCreationInfo(FNetRefHandle Handle)
+{
+	return false;
+}
+
 FReplicationBridgeCreateNetRefHandleResult UReplicationBridge::CreateNetRefHandleFromRemote(FNetRefHandle RootObjectOfSubObject, FNetRefHandle WantedNetHandle, FReplicationBridgeSerializationContext& Context)
 {
 	return FReplicationBridgeCreateNetRefHandleResult();
@@ -203,6 +218,11 @@ void UReplicationBridge::CallPruneStaleObjects()
 	PruneStaleObjects();
 }
 
+bool UReplicationBridge::CallCacheNetRefHandleCreationInfo(FNetRefHandle Handle)
+{
+	return CacheNetRefHandleCreationInfo(Handle);
+}
+
 bool UReplicationBridge::CallWriteNetRefHandleCreationInfo(FReplicationBridgeSerializationContext& Context, FNetRefHandle Handle)
 {
 	using namespace UE::Net;
@@ -466,7 +486,7 @@ void UReplicationBridge::DestroyLocalNetHandle(FNetRefHandle Handle, EEndReplica
 	InternalDetachInstanceFromNetRefHandle(Handle);
 
 	// Allow derived bridges to cleanup any instance info they have stored
-	DetachInstance(Handle);
+	CallDetachInstance(Handle);
 
 	// If the object is in any groups we need to remove it to make sure that we update filtering
 	GetReplicationSystem()->RemoveFromAllGroups(Handle);
@@ -533,17 +553,21 @@ void UReplicationBridge::EndReplication(FNetRefHandle Handle, EEndReplicationFla
 	const FInternalNetRefIndex InternalReplicationIndex = NetRefHandleManager->GetInternalIndex(Handle);
 	if (NetRefHandleManager->IsLocal(InternalReplicationIndex))
 	{
+		FNetRefHandleManager::FReplicatedObjectData& ObjectData = NetRefHandleManager->GetReplicatedObjectDataNoCheck(InternalReplicationIndex);
 		if (EnumHasAnyFlags(EndReplicationFlags, EEndReplicationFlags::TearOff))
 		{
-			// Add handle to list of objects pending tear off and indicate that it should be destroyed during next update
+			// Add handle to list of objects pending EndReplication indicate that it should be destroyed during next update
 			// We need to do this to cover the case where the torn off object not yet has been added to the scope.
-			constexpr bool bIsImmediate = true;
-			TearOff(Handle, EndReplicationFlags, bIsImmediate);
+			AddPendingEndReplication(Handle, EndReplicationFlags);
 
 			// We do however copy the final state data and mark object to stop propagating state changes
 			// Note: Current implementation keeps the mapping between instance and handle alive so that we can handle the case where we both start replication 
 			// and tear the object off during the same frame
-			InternalTearOff(Handle);
+			InternalTearOff(Handle, ECacheCreationInfo::Yes);
+
+			// Detach instance as we must assume that we should not access the object after this call.
+			ObjectData.bPendingEndReplication = 1U;
+			InternalDetachInstanceFromNetRefHandle(Handle);
 		}
 		else 
 		{
@@ -555,14 +579,31 @@ void UReplicationBridge::EndReplication(FNetRefHandle Handle, EEndReplicationFla
 				InternalAddDestructionInfo(Handle, *Parameters);
 			}
 
-			// If the flush flag is set, we have to explicitly copy the final state and propagate pending changes as we do not expect the source object to be kept alive after the call to EndReplication
-			if (EnumHasAnyFlags(EndReplicationFlags, EEndReplicationFlags::Flush))
+			// New objects, destroyed during the same frame with posted attachments(RPC):s needs to request a flush to ensure that they get a scope update
+			const UE::Net::Private::FNetBlobManager& NetBlobManager = GetReplicationSystem()->GetReplicationSystemInternal()->GetNetBlobManager();
+			const bool bAllowAutoFlushOfUnProcessedReliableRPCs = bEnableFlushReliableRPCOnDestroy && ObjectData.bNeedsFullCopyAndQuantize;
+			if (bAllowAutoFlushOfUnProcessedReliableRPCs && NetBlobManager.HasUnprocessedReliableAttachments(InternalReplicationIndex))
 			{
-				InternalFlushStateData(Handle);
+				EnumAddFlags(EndReplicationFlags, EEndReplicationFlags::Flush);
 			}
 
-			NetRefHandleManager->GetReplicatedObjectDataNoCheck(InternalReplicationIndex).bPendingEndReplication = 1U;
-			DestroyLocalNetHandle(Handle, EndReplicationFlags);	
+			if (EnumHasAnyFlags(EndReplicationFlags, EEndReplicationFlags::Flush))
+			{
+				// Defer destroy until after scope update to allow create/destroy on same frame
+				AddPendingEndReplication(Handle, EndReplicationFlags);
+			
+				// Capture final state
+				InternalFlushStateData(Handle, ECacheCreationInfo::Yes);
+
+				// Detach instance as we must assume that we should not access the object after this call.
+				ObjectData.bPendingEndReplication = 1U;
+				InternalDetachInstanceFromNetRefHandle(Handle);
+			}
+			else
+			{
+				ObjectData.bPendingEndReplication = 1U;
+				DestroyLocalNetHandle(Handle, EndReplicationFlags);	
+			}	
 		}
 	}
 	else
@@ -636,24 +677,27 @@ void UReplicationBridge::RemoveDestructionInfosForGroup(UE::Net::FNetObjectGroup
 void UReplicationBridge::TearOffHandlesPendingTearOff()
 {
 	// Initiate tear off
-	for (FTearOffInfo Info : MakeArrayView(HandlesPendingTearOff))
+	for (FPendingEndReplicationInfo Info : MakeArrayView(HandlesPendingEndReplication))
 	{
-		InternalTearOff(Info.Handle);
+		if (EnumHasAnyFlags(Info.DestroyFlags, EEndReplicationFlags::TearOff))
+		{
+			InternalTearOff(Info.Handle, ECacheCreationInfo::Yes);
+		}
 	}
 }
 
-void UReplicationBridge::UpdateHandlesPendingTearOff()
+void UReplicationBridge::UpdateHandlesPendingEndReplication()
 {
 	using namespace UE::Net;
 	using namespace UE::Net::Private;
 
-	TArray<FTearOffInfo, TInlineAllocator<32>> ObjectsStillPendingTearOff;
-	for (FTearOffInfo Info : MakeArrayView(HandlesPendingTearOff))
+	TArray<FPendingEndReplicationInfo, TInlineAllocator<32>> ObjectsStillPendingEndReplication;
+	for (FPendingEndReplicationInfo Info : MakeArrayView(HandlesPendingEndReplication))
 	{
 		if (FInternalNetRefIndex ObjectInternalIndex = NetRefHandleManager->GetInternalIndex(Info.Handle))
 		{
-			// Immediate tear-off or object that no longer are referenced by any connections are destroyed
-			if (NetRefHandleManager->GetNetObjectRefCount(ObjectInternalIndex) == 0U || Info.bIsImmediate)
+			// Immediate destroy or objects that are no longer are referenced by any connections are destroyed
+			if (NetRefHandleManager->GetNetObjectRefCount(ObjectInternalIndex) == 0U || Info.Immediate == EPendingEndReplicationImmediate::Yes)
 			{
 				FNetRefHandleManager::FReplicatedObjectData& ObjectData = NetRefHandleManager->GetReplicatedObjectDataNoCheck(ObjectInternalIndex);
 				ObjectData.bPendingEndReplication = 1U;
@@ -661,10 +705,10 @@ void UReplicationBridge::UpdateHandlesPendingTearOff()
 			}
 			else
 			{
-				// If the object is still in scope remove it from scope as objects being torn-off should not be added to new connections
+				// If the object is still in scope remove it from scope as objects pending EndReplication should not be added to new connections after the first update.
 				if (NetRefHandleManager->IsScopableIndex(ObjectInternalIndex))
 				{
-					// Mark object and subobjects as no longer scopable, and that we should not propagate changed states
+					// Mark object and subobjects as no longer scopeable, and that we should not propagate changed states
 					NetRefHandleManager->RemoveFromScope(ObjectInternalIndex);
 					for (FInternalNetRefIndex SubObjectIndex : NetRefHandleManager->GetSubObjects(ObjectInternalIndex))
 					{
@@ -672,26 +716,30 @@ void UReplicationBridge::UpdateHandlesPendingTearOff()
 					}
 				}
 			
-				// Keep object in the pending tear-off list until the object is no longer referenced by any ReplicationWriter
-				constexpr bool bIsImmediate = false;
-				ObjectsStillPendingTearOff.Add(FTearOffInfo(Info.Handle, Info.DestroyFlags, bIsImmediate));
+				// Keep object in the pending EndReplication-list until the object is no longer referenced by any ReplicationWriter
+				ObjectsStillPendingEndReplication.Add(FPendingEndReplicationInfo(Info.Handle, Info.DestroyFlags, EPendingEndReplicationImmediate::No));
 			}
 		}
 	}
 
-	HandlesPendingTearOff.Reset();
-	HandlesPendingTearOff.Insert(ObjectsStillPendingTearOff.GetData(), ObjectsStillPendingTearOff.Num(), 0);
+	HandlesPendingEndReplication.Reset();
+	HandlesPendingEndReplication.Insert(ObjectsStillPendingEndReplication.GetData(), ObjectsStillPendingEndReplication.Num(), 0);
+
+	CSV_CUSTOM_STAT(Iris, NumHandlesPendingEndRepliction, (float)HandlesPendingEndReplication.Num(), ECsvCustomStatOp::Set);
 }
 
-void UReplicationBridge::TearOff(FNetRefHandle Handle, EEndReplicationFlags DestroyFlags, bool bIsImmediate)
+void UReplicationBridge::AddPendingEndReplication(FNetRefHandle Handle, EEndReplicationFlags DestroyFlags, EPendingEndReplicationImmediate Immediate)
 {
-	if (!HandlesPendingTearOff.FindByPredicate([&](const FTearOffInfo& Info){ return Info.Handle == Handle; }))
+	if (ensure(EnumHasAnyFlags(DestroyFlags, EEndReplicationFlags::Flush | EEndReplicationFlags::TearOff)))
 	{
-		HandlesPendingTearOff.Add(FTearOffInfo(Handle, DestroyFlags, bIsImmediate));
-	}	
+		if (!HandlesPendingEndReplication.FindByPredicate([&](const FPendingEndReplicationInfo& Info){ return Info.Handle == Handle; }))
+		{
+			HandlesPendingEndReplication.Add(FPendingEndReplicationInfo(Handle, DestroyFlags, Immediate));
+		}
+	}
 }
 
-void UReplicationBridge::InternalFlushStateData(UE::Net::FNetSerializationContext& SerializationContext, UE::Net::Private::FChangeMaskCache& ChangeMaskCache, UE::Net::FNetBitStreamWriter& ChangeMaskWriter, uint32 InternalObjectIndex)
+void UReplicationBridge::InternalFlushStateData(UE::Net::FNetSerializationContext& SerializationContext, UE::Net::Private::FChangeMaskCache& ChangeMaskCache, UE::Net::FNetBitStreamWriter& ChangeMaskWriter, uint32 InternalObjectIndex, ECacheCreationInfo CacheCreationInfo)
 {
 	using namespace UE::Net;
 	using namespace UE::Net::Private;
@@ -709,25 +757,34 @@ void UReplicationBridge::InternalFlushStateData(UE::Net::FNetSerializationContex
 		return;
 	}
 
-	for (FInternalNetRefIndex SubObjectInternalIndex : NetRefHandleManager->GetChildSubObjects(InternalObjectIndex))
+	if (ObjectData.InstanceProtocol)
 	{
-		InternalFlushStateData(SerializationContext, ChangeMaskCache, ChangeMaskWriter, SubObjectInternalIndex);
-	}
+		if (EnumHasAnyFlags(ObjectData.InstanceProtocol->InstanceTraits, EReplicationInstanceProtocolTraits::NeedsPoll | EReplicationInstanceProtocolTraits::NeedsPreSendUpdate))
+		{
+			CallPreSendUpdateSingleHandle(ObjectData.RefHandle);
+		}
 
-	if (EnumHasAnyFlags(ObjectData.InstanceProtocol->InstanceTraits, EReplicationInstanceProtocolTraits::NeedsPoll | EReplicationInstanceProtocolTraits::NeedsPreSendUpdate))
-	{
-		CallPreSendUpdateSingleHandle(ObjectData.RefHandle);
-	} 
+		// Cache creation info
+		if (CacheCreationInfo == ECacheCreationInfo::Yes)
+		{
+			ObjectData.bHasCachedCreationInfo =  CallCacheNetRefHandleCreationInfo(ObjectData.RefHandle) ? 1U : 0U;
+		}
 
-	FReplicationInstanceOperationsInternal::QuantizeObjectStateData(ChangeMaskWriter, ChangeMaskCache, *NetRefHandleManager, SerializationContext, InternalObjectIndex);
+		FReplicationInstanceOperationsInternal::QuantizeObjectStateData(ChangeMaskWriter, ChangeMaskCache, *NetRefHandleManager, SerializationContext, InternalObjectIndex);
 
-	// Clear the quantize flag since it was done directly here.
-	NetRefHandleManager->GetDirtyObjectsToQuantize().ClearBit(InternalObjectIndex);
+		// Clear the quantize flag since it was done directly here.
+		NetRefHandleManager->GetDirtyObjectsToQuantize().ClearBit(InternalObjectIndex);
+	}
+
+	for (FInternalNetRefIndex SubObjectInternalIndex : NetRefHandleManager->GetChildSubObjects(InternalObjectIndex))
+	{
+		InternalFlushStateData(SerializationContext, ChangeMaskCache, ChangeMaskWriter, SubObjectInternalIndex, CacheCreationInfo);
+	}
 
 	// $IRIS TODO:  Should we also clear the DirtyTracker flags for this flushed object ?
 }
 
-void UReplicationBridge::InternalFlushStateData(FNetRefHandle Handle)
+void UReplicationBridge::InternalFlushStateData(FNetRefHandle Handle, ECacheCreationInfo CacheCreationInfo)
 {
 	using namespace UE::Net;
 	using namespace UE::Net::Private;
@@ -751,7 +808,7 @@ void UReplicationBridge::InternalFlushStateData(FNetRefHandle Handle)
 	SerializationContext.SetInternalContext(&InternalContext);
 	SerializationContext.SetNetStatsContext(ReplicationSystem->GetReplicationSystemInternal()->GetNetTypeStats().GetNetStatsContext());
 
-	InternalFlushStateData(SerializationContext, ChangeMaskCache, ChangeMaskWriter, InternalObjectIndex);
+	InternalFlushStateData(SerializationContext, ChangeMaskCache, ChangeMaskWriter, InternalObjectIndex, CacheCreationInfo);
 
 	// Iterate over connections and propagate dirty changemasks to all connections already scoping this object
 	if (ChangeMaskCache.Indices.Num() > 0)
@@ -770,7 +827,7 @@ void UReplicationBridge::InternalFlushStateData(FNetRefHandle Handle)
 	}
 }
 
-void UReplicationBridge::InternalTearOff(FNetRefHandle Handle)
+void UReplicationBridge::InternalTearOff(FNetRefHandle Handle, ECacheCreationInfo CacheCreationInfo)
 {
 	using namespace UE::Net;
 	using namespace UE::Net::Private;
@@ -789,12 +846,6 @@ void UReplicationBridge::InternalTearOff(FNetRefHandle Handle)
 	// Copy state data and tear off now
 	if (!ObjectData.bTearOff)
 	{
-		// Mark subobjects for TearOff as well.
-		for (FInternalNetRefIndex SubObjectInternalIndex : NetRefHandleManager->GetChildSubObjects(InternalObjectIndex))
-		{
-			InternalTearOff(NetRefHandleManager->GetNetRefHandleFromInternalIndex(SubObjectInternalIndex));
-		}	
-
 		UE_LOG_REPLICATIONBRIDGE(Verbose, TEXT("TearOff %s"), *Handle.ToString());
 
 		// Force copy of final state data as we will detach the object after scope update
@@ -807,14 +858,26 @@ void UReplicationBridge::InternalTearOff(FNetRefHandle Handle)
 		SerializationContext.SetInternalContext(&InternalContext);
 		SerializationContext.SetNetStatsContext(ReplicationSystem->GetReplicationSystemInternal()->GetNetTypeStats().GetNetStatsContext());
 
-		if (ObjectData.InstanceProtocol && EnumHasAnyFlags(ObjectData.InstanceProtocol->InstanceTraits, EReplicationInstanceProtocolTraits::NeedsPoll | EReplicationInstanceProtocolTraits::NeedsPreSendUpdate))
+		if (ObjectData.InstanceProtocol)
 		{
-			CallPreSendUpdateSingleHandle(Handle);
-		} 
+			if (EnumHasAnyFlags(ObjectData.InstanceProtocol->InstanceTraits, EReplicationInstanceProtocolTraits::NeedsPoll | EReplicationInstanceProtocolTraits::NeedsPreSendUpdate))
+			{
+				CallPreSendUpdateSingleHandle(Handle);
+			} 
+
+			// Cache creation info
+			if (CacheCreationInfo == ECacheCreationInfo::Yes)
+			{
+				ObjectData.bHasCachedCreationInfo =  CallCacheNetRefHandleCreationInfo(Handle) ? 1U : 0U;
+			}
+		}
 
 		if (ObjectData.InstanceProtocol && ObjectData.Protocol->InternalTotalSize > 0U)
 		{
 			FReplicationInstanceOperationsInternal::QuantizeObjectStateData(ChangeMaskWriter, ChangeMaskCache, *NetRefHandleManager, SerializationContext, InternalObjectIndex);
+
+			// Clear the quantize flag since it was done directly here.
+			NetRefHandleManager->GetDirtyObjectsToQuantize().ClearBit(InternalObjectIndex);
 		}
 		else
 		{
@@ -841,6 +904,12 @@ void UReplicationBridge::InternalTearOff(FNetRefHandle Handle)
 		const FNetBitArray& ValidConnections = Connections.GetValidConnections();
 		ValidConnections.ForAllSetBits(UpdateDirtyChangeMasks);		
 
+		// TearOff subobjects as well.
+		for (FInternalNetRefIndex SubObjectInternalIndex : NetRefHandleManager->GetChildSubObjects(InternalObjectIndex))
+		{
+			InternalTearOff(NetRefHandleManager->GetNetRefHandleFromInternalIndex(SubObjectInternalIndex), CacheCreationInfo);
+		}	
+
 		// Mark object as being torn-off and that we should no longer propagate state changes
 		ObjectData.bTearOff = 1U;
 		ObjectData.bShouldPropagateChangedStates = 0U;
diff --git a/Engine/Source/Runtime/Experimental/Iris/Core/Private/Iris/ReplicationSystem/ReplicationSystem.cpp b/Engine/Source/Runtime/Experimental/Iris/Core/Private/Iris/ReplicationSystem/ReplicationSystem.cpp
--- a/Engine/Source/Runtime/Experimental/Iris/Core/Private/Iris/ReplicationSystem/ReplicationSystem.cpp
+++ b/Engine/Source/Runtime/Experimental/Iris/Core/Private/Iris/ReplicationSystem/ReplicationSystem.cpp
@@ -292,8 +292,8 @@ public:
 		// Store the scope list for the next SendUpdate.
 		ReplicationSystemInternal.GetNetRefHandleManager().OnPostSendUpdate();
 
-		// Update handles pending tear-off
-		ReplicationSystemInternal.GetReplicationBridge()->UpdateHandlesPendingTearOff();
+		// Update handles pending tear-off/end-replication
+		ReplicationSystemInternal.GetReplicationBridge()->UpdateHandlesPendingEndReplication();
 
 		// Reset baseline invalidation
 		ReplicationSystemInternal.GetDeltaCompressionBaselineInvalidationTracker().PostSendUpdate();
@@ -1572,8 +1572,8 @@ void UReplicationSystem::SetIsNetTemporary(FNetRefHandle Handle)
 
 void UReplicationSystem::TearOffNextUpdate(FNetRefHandle Handle)
 {
-	constexpr EEndReplicationFlags DestroyFlags = EEndReplicationFlags::DestroyNetHandle | EEndReplicationFlags::ClearNetPushId;
-	Impl->ReplicationSystemInternal.GetReplicationBridge()->TearOff(Handle, DestroyFlags, false);
+	constexpr EEndReplicationFlags DestroyFlags = EEndReplicationFlags::TearOff | EEndReplicationFlags::ClearNetPushId;
+	Impl->ReplicationSystemInternal.GetReplicationBridge()->AddPendingEndReplication(Handle, DestroyFlags);
 }
 
 void UReplicationSystem::ForceNetUpdate(FNetRefHandle Handle)
diff --git a/Engine/Source/Runtime/Experimental/Iris/Core/Private/Iris/ReplicationSystem/ReplicationWriter.cpp b/Engine/Source/Runtime/Experimental/Iris/Core/Private/Iris/ReplicationSystem/ReplicationWriter.cpp
--- a/Engine/Source/Runtime/Experimental/Iris/Core/Private/Iris/ReplicationSystem/ReplicationWriter.cpp
+++ b/Engine/Source/Runtime/Experimental/Iris/Core/Private/Iris/ReplicationSystem/ReplicationWriter.cpp
@@ -2115,7 +2115,7 @@ FReplicationWriter::EWriteObjectStatus FReplicationWriter::WriteObjectAndSubObje
 					UE_NET_TRACE_SCOPE(CreationInfo, Writer, Context.GetTraceCollector(), ENetTraceVerbosity::Trace);
 
 					// Warn if we cannot replicate this object
-					if (!ObjectData.InstanceProtocol)
+					if (!ObjectData.InstanceProtocol && !(ObjectData.bHasCachedCreationInfo == 1U))
 					{
 						UE_LOG_REPLICATIONWRITER_WARNING(TEXT("Failed to replicate ( InternalIndex: %u ) %s, ProtocolName: %s, Currently we do not support creating a remote instance when the instance has been detached."), InternalIndex, *NetRefHandle.ToString(), ToCStr(ObjectData.Protocol->DebugName));
 						return EWriteObjectStatus::NoInstanceProtocol;
diff --git a/Engine/Source/Runtime/Experimental/Iris/Core/Public/Iris/ReplicationSystem/ObjectReplicationBridge.h b/Engine/Source/Runtime/Experimental/Iris/Core/Public/Iris/ReplicationSystem/ObjectReplicationBridge.h
--- a/Engine/Source/Runtime/Experimental/Iris/Core/Public/Iris/ReplicationSystem/ObjectReplicationBridge.h
+++ b/Engine/Source/Runtime/Experimental/Iris/Core/Public/Iris/ReplicationSystem/ObjectReplicationBridge.h
@@ -210,6 +210,7 @@ protected:
 	IRISCORE_API virtual void UpdateInstancesWorldLocation() override;
 	IRISCORE_API virtual void PruneStaleObjects() override;	
 	IRISCORE_API virtual bool WriteNetRefHandleCreationInfo(FReplicationBridgeSerializationContext& Context, FNetRefHandle Handle) override;
+	IRISCORE_API virtual bool CacheNetRefHandleCreationInfo(FNetRefHandle Handle) override;
 	IRISCORE_API virtual FReplicationBridgeCreateNetRefHandleResult CreateNetRefHandleFromRemote(FNetRefHandle RootObjectNetHandle, FNetRefHandle WantedNetHandle, FReplicationBridgeSerializationContext& Context) override;
 	IRISCORE_API virtual void SubObjectCreatedFromReplication(FNetRefHandle SubObjectHandle) override;
 	IRISCORE_API virtual void PostApplyInitialState(FNetRefHandle Handle) override;
@@ -226,11 +227,17 @@ protected:
 	* but only ever is replicated when we first instantiate an object of the specific type 
 	*/
 
+	/** Allocate and fill in creation header. */
+	virtual TUniquePtr<FCreationHeader> GetCreationHeader(FNetRefHandle Handle) { return nullptr; };
+
+	/** Write creation header. */
+	virtual bool WriteCreationHeader(UE::Net::FNetSerializationContext& Context, const FCreationHeader* CreationHeader) { return false; };
+
 	/** Write data required to instantiate the Handle. */
 	virtual bool WriteCreationHeader(UE::Net::FNetSerializationContext& Context, FNetRefHandle Handle) { return false; };
 
 	/** Create Header and read data required to instantiate the Handle. */
-	virtual FCreationHeader* ReadCreationHeader(UE::Net::FNetSerializationContext& Context) { return nullptr; };
+	virtual TUniquePtr<FCreationHeader> ReadCreationHeader(UE::Net::FNetSerializationContext& Context) { return nullptr; };
 
 	/** Called when we instantiate/find object instance requested by remote. */
 	virtual FObjectReplicationBridgeInstantiateResult BeginInstantiateFromRemote(FNetRefHandle RootObjectOfSubObject, const UE::Net::FNetObjectResolveContext& ResolveContext, const FCreationHeader* Header) { return FObjectReplicationBridgeInstantiateResult(); };
@@ -448,6 +455,9 @@ private:
 	// Array of dormant objects that has requested a flush
 	TArray<FNetRefHandle> DormantHandlesPendingFlush;
 
+	// When we flush objects, we might need to defer sending creation info.
+	TMap<FNetRefHandle, TUniquePtr<const FCreationHeader>> CachedCreationHeaders;
+
 	// Objects which has object references and could be affected by garbage collection.
 	UE::Net::FNetBitArray ObjectsWithObjectReferences;
 	// Objects that needs to update their object references due to garbage collection.
diff --git a/Engine/Source/Runtime/Experimental/Iris/Core/Public/Iris/ReplicationSystem/ReplicationBridge.h b/Engine/Source/Runtime/Experimental/Iris/Core/Public/Iris/ReplicationSystem/ReplicationBridge.h
--- a/Engine/Source/Runtime/Experimental/Iris/Core/Public/Iris/ReplicationSystem/ReplicationBridge.h
+++ b/Engine/Source/Runtime/Experimental/Iris/Core/Public/Iris/ReplicationSystem/ReplicationBridge.h
@@ -200,6 +200,13 @@ protected:
 	 */
 	IRISCORE_API virtual bool WriteNetRefHandleCreationInfo(FReplicationBridgeSerializationContext& Context, FNetRefHandle Handle);
 
+	/**
+	 * Cache info required to allow deferred writing of NetRefHandleCreationInfo
+	 * @param Handle The handle of the object to store creation data for.
+	 * return whether cached data is stored or not.
+	*/
+	IRISCORE_API virtual bool CacheNetRefHandleCreationInfo(FNetRefHandle Handle);
+
 	/** Read data required to instantiate NetObject from bitstream. */
 	IRISCORE_API virtual FReplicationBridgeCreateNetRefHandleResult CreateNetRefHandleFromRemote(FNetRefHandle RootObjectOfSubObject, FNetRefHandle WantedNetHandle, FReplicationBridgeSerializationContext& Context);
 
@@ -283,14 +290,21 @@ private:
 	void DetachSubObjectInstancesFromRemote(FNetRefHandle Handle, EReplicationBridgeDestroyInstanceReason DestroyReason, EReplicationBridgeDestroyInstanceFlags DestroyFlags);
 	void DestroyNetObjectFromRemote(FNetRefHandle Handle, EReplicationBridgeDestroyInstanceReason DestroyReason, EReplicationBridgeDestroyInstanceFlags DestroyFlags);
 
-	// Adds the Handle to the list of handles pending tear-off, if bIsImmediate is true the object will be destroyed after the next update, otherwise
-	// it will be kept around until EndReplication is called.
-	void TearOff(FNetRefHandle Handle, EEndReplicationFlags DestroyFlags, bool bIsImmediate);
+	enum class EPendingEndReplicationImmediate : uint8
+	{
+		Yes,
+		No,
+	};
+
+	// Adds the Handle to the list of handles pending deferred EndReplication, if bIsImmediate is true the object will be destroyed after the next update, otherwise
+	// it will be kept around until the handle is no longer ref-counted by any connection. It will however be removed from the set of scopeable objects after the first update so new connections will not add it to their scope.
+	void AddPendingEndReplication(FNetRefHandle Handle, EEndReplicationFlags DestroyFlags, EPendingEndReplicationImmediate Immediate = EPendingEndReplicationImmediate::No);
 
 	FReplicationBridgeCreateNetRefHandleResult CallCreateNetRefHandleFromRemote(FNetRefHandle RootObjectOfSubObject, FNetRefHandle WantedNetHandle, FReplicationBridgeSerializationContext& Context);
 	void CallPreSendUpdate(float DeltaSeconds);	
 	void CallPreSendUpdateSingleHandle(FNetRefHandle Handle);
 	void CallUpdateInstancesWorldLocation();
+	bool CallCacheNetRefHandleCreationInfo(FNetRefHandle Handle);
 	bool CallWriteNetRefHandleCreationInfo(FReplicationBridgeSerializationContext& Context, FNetRefHandle Handle);
 	bool CallWriteNetRefHandleDestructionInfo(FReplicationBridgeSerializationContext& Context, FNetRefHandle Handle);
 	void CallSubObjectCreatedFromReplication(FNetRefHandle SubObjectHandle);
@@ -301,14 +315,19 @@ private:
 	void CallDetachInstanceFromRemote(FNetRefHandle Handle, EReplicationBridgeDestroyInstanceReason DestroyReason, EReplicationBridgeDestroyInstanceFlags DestroyFlags);
 
 private:
+	// Specifies if we should capture and cache CreationHeader or not for flush and tearoff.
+	enum class ECacheCreationInfo : uint8
+	{
+		No,
+		Yes
+	};
 
-	void InternalFlushStateData(UE::Net::FNetSerializationContext& SerializationContext, UE::Net::Private::FChangeMaskCache& ChangeMaskCache, UE::Net::FNetBitStreamWriter& ChangeMaskWriter, uint32 InternalObjectIndex);
-
+	void InternalFlushStateData(UE::Net::FNetSerializationContext& SerializationContext, UE::Net::Private::FChangeMaskCache& ChangeMaskCache, UE::Net::FNetBitStreamWriter& ChangeMaskWriter, uint32 InternalObjectIndex, ECacheCreationInfo CacheCreationInfo);
 	// Internal method to copy state data for Handle
-	void InternalFlushStateData(FNetRefHandle Handle);
+	void InternalFlushStateData(FNetRefHandle Handle, ECacheCreationInfo CacheCreationInfo);
 
 	// Internal method to copy state data for Handle and any SubObjects and mark them as being torn-off
-	void InternalTearOff(FNetRefHandle OwnerHandle);
+	void InternalTearOff(FNetRefHandle OwnerHandle, ECacheCreationInfo CacheCreationInfo);
 
 	// Destroy all SubObjects owned by provided handle
 	void InternalDestroySubObjects(FNetRefHandle OwnerHandle, EEndReplicationFlags Flags);
@@ -330,8 +349,8 @@ private:
 	// Tear-off all handles in the PendingTearOff list that has not yet been torn-off
 	void TearOffHandlesPendingTearOff();
 
-	// Update all the handles pending tear-off
-	void UpdateHandlesPendingTearOff();
+	// Update all the handles pending EndReplication
+	void UpdateHandlesPendingEndReplication();
 
 	void SetNetPushIdOnFragments(const TArrayView<const UE::Net::FReplicationFragment*const>& Fragments, const UE::Net::Private::FNetPushObjectHandle& PushHandle);
 	void ClearNetPushIdOnFragments(const TArrayView<const UE::Net::FReplicationFragment*const>& Fragments);
@@ -371,15 +390,15 @@ private:
 	// We use this to be able ask remote to destroy static objects
 	TMap<FNetRefHandle, FDestructionInfo> StaticObjectsPendingDestroy;
 
-	struct FTearOffInfo
+	struct FPendingEndReplicationInfo
 	{
-		FTearOffInfo(FNetRefHandle InHandle, EEndReplicationFlags InDestroyFlags, bool bInIsImmediate) : Handle(InHandle), DestroyFlags(InDestroyFlags), bIsImmediate(bInIsImmediate) {}
+		FPendingEndReplicationInfo(FNetRefHandle InHandle, EEndReplicationFlags InDestroyFlags, EPendingEndReplicationImmediate InImmediate) : Handle(InHandle), DestroyFlags(InDestroyFlags), Immediate(InImmediate) {}
 
 		FNetRefHandle Handle;
 		EEndReplicationFlags DestroyFlags;
-		bool bIsImmediate;
+		EPendingEndReplicationImmediate Immediate;
 	};
-	TArray<FTearOffInfo> HandlesPendingTearOff;
+	TArray<FPendingEndReplicationInfo> HandlesPendingEndReplication;
 };
 
 inline FReplicationBridgeSerializationContext::FReplicationBridgeSerializationContext(UE::Net::FNetSerializationContext& InSerialiazationContext, uint32 InConnectionId, bool bInIsDestructionInfo)
diff --git a/Engine/Source/Runtime/Experimental/Iris/Core/Public/Iris/ReplicationSystem/ReplicationSystem.h b/Engine/Source/Runtime/Experimental/Iris/Core/Public/Iris/ReplicationSystem/ReplicationSystem.h
--- a/Engine/Source/Runtime/Experimental/Iris/Core/Public/Iris/ReplicationSystem/ReplicationSystem.h
+++ b/Engine/Source/Runtime/Experimental/Iris/Core/Public/Iris/ReplicationSystem/ReplicationSystem.h
@@ -453,7 +453,10 @@ public:
 	/** Set status of GroupFilter for specific connection. */
 	IRISCORE_API void SetGroupFilterStatus(FNetObjectGroupHandle GroupHandle, uint32 ConnectionId, UE::Net::ENetFilterStatus ReplicationStatus);
 
-	/** Set status of GroupFilter for connection marked in the Connections BitArray. */
+	/**
+	 * Set status of GroupFilter for connection marked in the Connections BitArray to the passed ReplicationStatus,
+	 * Connections not marked in the BitArray will be set to the opposite status.
+	*/
 	IRISCORE_API void SetGroupFilterStatus(FNetObjectGroupHandle GroupHandle, const UE::Net::FNetBitArray& Connections, UE::Net::ENetFilterStatus ReplicationStatus);
 
 	/** Set status of GroupFilter for all connections. */
