Correctly filter on-demand I/O store containers

diff --git a/Engine/Source/Runtime/Experimental/IoStore/OnDemand/Private/OnDemandIoStore.cpp b/Engine/Source/Runtime/Experimental/IoStore/OnDemand/Private/OnDemandIoStore.cpp
--- a/Engine/Source/Runtime/Experimental/IoStore/OnDemand/Private/OnDemandIoStore.cpp
+++ b/Engine/Source/Runtime/Experimental/IoStore/OnDemand/Private/OnDemandIoStore.cpp
@@ -353,7 +353,7 @@ FOnDemandChunkInfo FOnDemandIoStore::GetChunkInfo(const FIoChunkId& ChunkId, EOn
 
 	for (FSharedOnDemandContainer& Container : Containers)
 	{
-		if (EnumHasAnyFlags(Container->Flags, EOnDemandContainerFlags::Mounted))
+		if (EnumHasAllFlags(Container->Flags, ContainerFlags))
 		{
 			if (const FOnDemandChunkEntry* Entry = Container->ChunkEntries.Find(ChunkId))
 			{
