Compile fix.

#rnx

diff --git a/Engine/Source/Runtime/Experimental/IoStore/HttpClient/Private/Module.cpp b/Engine/Source/Runtime/Experimental/IoStore/HttpClient/Private/Module.cpp
--- a/Engine/Source/Runtime/Experimental/IoStore/HttpClient/Private/Module.cpp
+++ b/Engine/Source/Runtime/Experimental/IoStore/HttpClient/Private/Module.cpp
@@ -2,4 +2,6 @@
 
 #include "Modules/Boilerplate/ModuleBoilerplate.h"
 
+#if !IS_MONOLITHIC
 REPLACEMENT_OPERATOR_NEW_AND_DELETE
+#endif
