Fix for an server crash when calculating the Iris grid cells for an out of bounds position.

diff --git a/Engine/Source/Runtime/Experimental/Iris/Core/Private/Iris/ReplicationSystem/Filtering/NetObjectGridFilter.cpp b/Engine/Source/Runtime/Experimental/Iris/Core/Private/Iris/ReplicationSystem/Filtering/NetObjectGridFilter.cpp
--- a/Engine/Source/Runtime/Experimental/Iris/Core/Private/Iris/ReplicationSystem/Filtering/NetObjectGridFilter.cpp
+++ b/Engine/Source/Runtime/Experimental/Iris/Core/Private/Iris/ReplicationSystem/Filtering/NetObjectGridFilter.cpp
@@ -482,9 +482,6 @@ void UNetObjectGridFilter::CalculateCellBox(const UNetObjectGridFilter::FPerObje
 	FVector MinPosition = Position - CullDistance;
 	FVector MaxPosition = Position + CullDistance;
 
-	MinPosition = MinPosition.ComponentMax(Config->MinPos);
-	MaxPosition = MaxPosition.ComponentMin(Config->MaxPos);
-
 	const int64 MinX = FPlatformMath::FloorToInt(MinPosition.X/Config->CellSizeX);
 	const int64 MinY = FPlatformMath::FloorToInt(MinPosition.Y/Config->CellSizeY);
 	const int64 MaxX = FPlatformMath::FloorToInt(MaxPosition.X/Config->CellSizeX);
diff --git a/Engine/Source/Runtime/Experimental/Iris/Core/Private/Iris/ReplicationSystem/ObjectReplicationBridge.cpp b/Engine/Source/Runtime/Experimental/Iris/Core/Private/Iris/ReplicationSystem/ObjectReplicationBridge.cpp
--- a/Engine/Source/Runtime/Experimental/Iris/Core/Private/Iris/ReplicationSystem/ObjectReplicationBridge.cpp
+++ b/Engine/Source/Runtime/Experimental/Iris/Core/Private/Iris/ReplicationSystem/ObjectReplicationBridge.cpp
@@ -1206,6 +1206,11 @@ void UObjectReplicationBridge::UpdateInstancesWorldLocation()
 			FWorldLocations::FObjectInfo CachedObjectInfo;
 			
 			GetInstanceWorldObjectInfoFunction(ObjectData.RefHandle, ReplicatedInstances[InternalObjectIndex], CachedObjectInfo.WorldLocation, CachedObjectInfo.CullDistance);
+
+			// Check the position lies within world boundaries.
+			ensureMsgf(CachedObjectInfo.WorldLocation.X >= WorldLocations.GetWorldMinPos().X && CachedObjectInfo.WorldLocation.X <= WorldLocations.GetWorldMaxPos().X && 
+					   CachedObjectInfo.WorldLocation.Y >= WorldLocations.GetWorldMinPos().Y && CachedObjectInfo.WorldLocation.Y <= WorldLocations.GetWorldMaxPos().Y, TEXT("Object %s with position %s lies outside configured world boundary."), ToCStr(NetRefHandleManager->PrintObjectFromIndex((InternalObjectIndex))), *CachedObjectInfo.WorldLocation.ToString());
+			
 			WorldLocations.SetObjectInfo(InternalObjectIndex, CachedObjectInfo);
 		}
 	};
diff --git a/Engine/Source/Runtime/Experimental/Iris/Core/Private/Iris/ReplicationSystem/ReplicationBridge.cpp b/Engine/Source/Runtime/Experimental/Iris/Core/Private/Iris/ReplicationSystem/ReplicationBridge.cpp
--- a/Engine/Source/Runtime/Experimental/Iris/Core/Private/Iris/ReplicationSystem/ReplicationBridge.cpp
+++ b/Engine/Source/Runtime/Experimental/Iris/Core/Private/Iris/ReplicationSystem/ReplicationBridge.cpp
@@ -948,6 +948,11 @@ UE::Net::FNetRefHandle UReplicationBridge::InternalAddDestructionInfo(FNetRefHan
 	{
 		// Use WorldLocations to feed the location of the destruction info so that it can be prioritized properly.
 		FWorldLocations& WorldLocations = GetReplicationSystem()->GetReplicationSystemInternal()->GetWorldLocations();
+
+		// Check the position lies within world boundaries.
+		ensureMsgf(Parameters.Location.X >= WorldLocations.GetWorldMinPos().X && Parameters.Location.X <= WorldLocations.GetWorldMaxPos().X && 
+				   Parameters.Location.Y >= WorldLocations.GetWorldMinPos().Y && Parameters.Location.Y <= WorldLocations.GetWorldMaxPos().Y, TEXT("Object %s with position %s lies outside configured world boundary."), ToCStr(NetRefHandleManager->PrintObjectFromIndex((InternalReplicationIndex))), *Parameters.Location.ToString());
+
 		WorldLocations.InitObjectInfoCache(InternalReplicationIndex);
 		WorldLocations.UpdateWorldLocation(InternalReplicationIndex, Parameters.Location);
 
diff --git a/Engine/Source/Runtime/Experimental/Iris/Core/Private/Iris/ReplicationSystem/WorldLocations.cpp b/Engine/Source/Runtime/Experimental/Iris/Core/Private/Iris/ReplicationSystem/WorldLocations.cpp
--- a/Engine/Source/Runtime/Experimental/Iris/Core/Private/Iris/ReplicationSystem/WorldLocations.cpp
+++ b/Engine/Source/Runtime/Experimental/Iris/Core/Private/Iris/ReplicationSystem/WorldLocations.cpp
@@ -3,6 +3,8 @@
 #include "Iris/ReplicationSystem/WorldLocations.h"
 #include "Iris/Core/IrisMemoryTracker.h"
 
+#include UE_INLINE_GENERATED_CPP_BY_NAME(WorldLocations)
+
 namespace UE::Net
 {
 
@@ -11,6 +13,9 @@ void FWorldLocations::Init(const FWorldLocationsInitParams& InitParams)
 	ValidInfoIndexes.Init(InitParams.MaxObjectCount);
 	ObjectsWithDirtyInfo.Init(InitParams.MaxObjectCount);
 	ObjectsRequiringFrequentWorldLocationUpdate.Init(InitParams.MaxObjectCount);
+
+	MinWorldPos = GetDefault<UWorldLocationsConfig>()->MinPos;
+	MaxWorldPos = GetDefault<UWorldLocationsConfig>()->MaxPos;
 }
 
 void FWorldLocations::InitObjectInfoCache(uint32 ObjectIndex)
@@ -46,16 +51,18 @@ void FWorldLocations::SetObjectInfo(uint32 ObjectIndex, const FWorldLocations::F
 	FObjectInfo& TargetObjectInfo = StoredObjectInfo[ObjectIndex];
 	const bool bHasInfoChanged = ObjectsWithDirtyInfo.GetBit(ObjectIndex) || TargetObjectInfo.WorldLocation != ObjectInfo.WorldLocation || TargetObjectInfo.CullDistance != ObjectInfo.CullDistance;
 	TargetObjectInfo = ObjectInfo;
+	TargetObjectInfo.WorldLocation = ClampPositionToBoundary(ObjectInfo.WorldLocation);
 
 	ObjectsWithDirtyInfo.SetBitValue(ObjectIndex, bHasInfoChanged);
 }
 
 void FWorldLocations::UpdateWorldLocation(uint32 ObjectIndex, const FVector& WorldLocation)
 {
+	const FVector InBoundsWorldLocation = ClampPositionToBoundary(WorldLocation);
 	checkSlow(ValidInfoIndexes.GetBit(ObjectIndex));
 	FVector& TargetWorldLocation = StoredObjectInfo[ObjectIndex].WorldLocation;
-	TargetWorldLocation = WorldLocation;
-	const bool bHasInfoChanged = ObjectsWithDirtyInfo.GetBit(ObjectIndex) || TargetWorldLocation != WorldLocation;
+	TargetWorldLocation = InBoundsWorldLocation;
+	const bool bHasInfoChanged = ObjectsWithDirtyInfo.GetBit(ObjectIndex) || TargetWorldLocation != InBoundsWorldLocation;
 	
 	ObjectsWithDirtyInfo.SetBitValue(ObjectIndex, bHasInfoChanged);
 }
@@ -65,4 +72,9 @@ void FWorldLocations::ResetObjectsWithDirtyInfo()
 	ObjectsWithDirtyInfo.Reset();
 }
 
+FVector FWorldLocations::ClampPositionToBoundary(const FVector& Position)
+{
+	return Position.BoundToBox(GetWorldMinPos(), GetWorldMaxPos());
+}
+
 }
diff --git a/Engine/Source/Runtime/Experimental/Iris/Core/Public/Iris/ReplicationSystem/Filtering/NetObjectGridFilter.h b/Engine/Source/Runtime/Experimental/Iris/Core/Public/Iris/ReplicationSystem/Filtering/NetObjectGridFilter.h
--- a/Engine/Source/Runtime/Experimental/Iris/Core/Public/Iris/ReplicationSystem/Filtering/NetObjectGridFilter.h
+++ b/Engine/Source/Runtime/Experimental/Iris/Core/Public/Iris/ReplicationSystem/Filtering/NetObjectGridFilter.h
@@ -75,14 +75,6 @@ public:
 	UPROPERTY(Config)
 	float DefaultCullDistance = 15000.0f;
 
-	/** Coordinates will be clamped to MinPos and MaxPos. */
-	UPROPERTY(Config)
-	FVector MinPos = {-0.5f*2097152.0f, -0.5f*2097152.0f, -0.5f*2097152.0f};
-
-	/** Coordinates will be clamped to MinPos and MaxPos. */
-	UPROPERTY(Config)
-	FVector MaxPos = {+0.5f*2097152.0f, +0.5f*2097152.0f, +0.5f*2097152.0f};
-
 	/** 
 	 * If true: use the exact distance between an object and the viewer to determine if the object is relevant or should be culled out.
 	 * When false: consider all objects within a grid cell to be relevant when a viewer is located within the cell. This can extend the relevant distance of objects beyond their cull distance. 
diff --git a/Engine/Source/Runtime/Experimental/Iris/Core/Public/Iris/ReplicationSystem/WorldLocations.h b/Engine/Source/Runtime/Experimental/Iris/Core/Public/Iris/ReplicationSystem/WorldLocations.h
--- a/Engine/Source/Runtime/Experimental/Iris/Core/Public/Iris/ReplicationSystem/WorldLocations.h
+++ b/Engine/Source/Runtime/Experimental/Iris/Core/Public/Iris/ReplicationSystem/WorldLocations.h
@@ -6,6 +6,25 @@
 #include "Containers/ChunkedArray.h"
 #include "Net/Core/NetBitArray.h"
 #include "Math/Vector.h"
+#include "WorldLocations.generated.h"
+
+/**
+* Common settings used to configure how the GridFilter behaves
+*/
+UCLASS(Config=Engine)
+class UWorldLocationsConfig : public UObject
+{
+	GENERATED_BODY()
+
+public:
+	/** All world positions will be clamped to MinPos and MaxPos. */
+	UPROPERTY(Config)
+	FVector MinPos = { -0.5f * 2097152.0f, -0.5f * 2097152.0f, -0.5f * 2097152.0f };
+
+	/** All world positions will be clamped to MinPos and MaxPos. */
+	UPROPERTY(Config)
+	FVector MaxPos = { +0.5f * 2097152.0f, +0.5f * 2097152.0f, +0.5f * 2097152.0f };
+};
 
 namespace UE::Net
 {
@@ -67,6 +86,13 @@ public:
 	/** Returns the list of objects that registered world location information */
 	const FNetBitArrayView GetObjectsWithWorldInfo() const { return MakeNetBitArrayView(ValidInfoIndexes); }
 
+	/** Return the world boundaries (min and max position). */
+	const FVector& GetWorldMinPos() const { return MinWorldPos; };
+	const FVector& GetWorldMaxPos() const { return MaxWorldPos; };
+	
+	/** Return a position clamped to the configured world boundary. */
+	FVector ClampPositionToBoundary(const FVector& Position);
+
 private:
 	enum : uint32
 	{
@@ -81,6 +107,10 @@ private:
 	FNetBitArray ObjectsRequiringFrequentWorldLocationUpdate;
 
 	TChunkedArray<FObjectInfo, BytesPerLocationChunk> StoredObjectInfo;
+
+	/** World boundaries (min and max position). */
+	FVector MinWorldPos;
+	FVector MaxWorldPos;
 };
 
 inline bool FWorldLocations::HasInfoForObject(uint32 ObjectIndex) const
