Add AutoRTFM Support for FChaosRefCountedObject

diff --git a/Engine/Source/Runtime/Experimental/Chaos/Public/Chaos/ImplicitObject.h b/Engine/Source/Runtime/Experimental/Chaos/Public/Chaos/ImplicitObject.h
--- a/Engine/Source/Runtime/Experimental/Chaos/Public/Chaos/ImplicitObject.h
+++ b/Engine/Source/Runtime/Experimental/Chaos/Public/Chaos/ImplicitObject.h
@@ -9,6 +9,7 @@
 #include "Chaos/ImplicitObjectType.h"
 #include "Chaos/AABB.h"
 #include "Templates/RefCounting.h"
+#include "AutoRTFM/AutoRTFM.h"
 
 #ifndef TRACK_CHAOS_GEOMETRY
 #define TRACK_CHAOS_GEOMETRY !(UE_BUILD_SHIPPING || UE_BUILD_TEST)
@@ -105,6 +106,8 @@ struct TImplicitTypeInfo
 #define DISALLOW_FIMPLICIT_OBJECT_TAIL_PADDING INTEL_ISPC
 
 // Chaos ref counted object
+//  * @note AutoRTFM means that the return value of AddRef/Release is nonsense (as the ref-count doesn't change until the
+//  *       transaction is committed), but this is fine for use with TRefCountPtr (as it doesn't use those return values).
 class FChaosRefCountedObject
 {
 public:
@@ -114,23 +117,40 @@ public:
 	FChaosRefCountedObject& operator=(const FChaosRefCountedObject& Rhs) = delete;
 	uint32 AddRef() const
 	{
-		return uint32(NumRefs.Increment());
+		UE_AUTORTFM_ONCOMMIT(
+		{
+			NumRefs.Increment();
+		});
+
+		// Note: TRefCountPtr doesn't use the return value
+		return 0;
 	}
 	uint32 Release() const
 	{
-		uint32 Refs = uint32(NumRefs.Decrement());
-		if (Refs == 0)
+		UE_AUTORTFM_ONCOMMIT(
 		{
-			if(bTransientFlag)
-			{ 
-				delete this;
+			uint32 Refs = uint32(NumRefs.Decrement());
+			if (Refs == 0)
+			{
+				if(bTransientFlag)
+				{
+					delete this;
+				}
 			}
-		}
-		return Refs;
+		});
+
+		// Note: TRefCountPtr doesn't use the return value
+		return 0;
 	}
 	uint32 GetRefCount() const
 	{
-		return uint32(NumRefs.GetValue());
+		uint32 Ret = 0;
+		UE_AUTORTFM_OPEN(
+		{
+			Ret = uint32(NumRefs.GetValue());
+		});
+
+		return Ret;
 	}
 
 	void MakePersistent() const
