Geometry collection : when loading geometry collection with older version of the implicits / sharedimplicits attributes, make sure to account for the fact that the collection may contains multiple of those

diff --git a/Engine/Source/Runtime/Experimental/GeometryCollectionEngine/Private/GeometryCollection/GeometryCollectionObject.cpp b/Engine/Source/Runtime/Experimental/GeometryCollectionEngine/Private/GeometryCollection/GeometryCollectionObject.cpp
--- a/Engine/Source/Runtime/Experimental/GeometryCollectionEngine/Private/GeometryCollection/GeometryCollectionObject.cpp
+++ b/Engine/Source/Runtime/Experimental/GeometryCollectionEngine/Private/GeometryCollection/GeometryCollectionObject.cpp
@@ -1032,8 +1032,13 @@ void UGeometryCollection::Serialize(FArchive& Ar)
 			TArray<Chaos::FImplicitObjectPtr> ImplicitObjects;
 			ImplicitObjects.SetNum(NumElems);
 			
-			if( TManagedArray<TUniquePtr<Chaos::FImplicitObject>>* OldAttrA = ArchiveGeometryCollection->FindAttributeTyped<TUniquePtr<Chaos::FImplicitObject>>(
-				FGeometryDynamicCollection::ImplicitsAttribute, FTransformCollection::TransformGroup))
+			const TManagedArray<TUniquePtr<Chaos::FImplicitObject>>* OldAttrA = ArchiveGeometryCollection->FindAttributeTyped<TUniquePtr<Chaos::FImplicitObject>>(FGeometryDynamicCollection::ImplicitsAttribute, FTransformCollection::TransformGroup);
+			const TManagedArray<TSharedPtr<Chaos::FImplicitObject, ESPMode::ThreadSafe>>* OldAttrB = ArchiveGeometryCollection->FindAttributeTyped<TSharedPtr<Chaos::FImplicitObject, ESPMode::ThreadSafe>>(FGeometryDynamicCollection::SharedImplicitsAttribute, FTransformCollection::TransformGroup);
+			const TManagedArray<TSharedPtr<Chaos::FImplicitObject, ESPMode::ThreadSafe>>* OldAttrC = ArchiveGeometryCollection->FindAttributeTyped<TSharedPtr<Chaos::FImplicitObject, ESPMode::ThreadSafe>>(FGeometryDynamicCollection::ImplicitsAttribute, FTransformCollection::TransformGroup);
+
+			// Some geometry collection can still store several of those arrays
+			// We need to make sure to remove all of them and keep the last good one 
+			if (OldAttrA)
 			{
 				for (int32 Index = 0; Index < NumElems; ++Index)
 				{
@@ -1044,8 +1049,7 @@ void UGeometryCollection::Serialize(FArchive& Ar)
 				}
 				ArchiveGeometryCollection->RemoveAttribute(FGeometryDynamicCollection::ImplicitsAttribute, FTransformCollection::TransformGroup);
 			}
-			else if(TManagedArray<TSharedPtr<Chaos::FImplicitObject, ESPMode::ThreadSafe>>* OldAttrB = ArchiveGeometryCollection->FindAttributeTyped<TSharedPtr<Chaos::FImplicitObject, ESPMode::ThreadSafe>>(
-				FGeometryDynamicCollection::SharedImplicitsAttribute, FTransformCollection::TransformGroup))
+			if (OldAttrB)
 			{
 				for (int32 Index = 0; Index < NumElems; ++Index)
 				{
@@ -1056,8 +1060,7 @@ void UGeometryCollection::Serialize(FArchive& Ar)
 				}
 				ArchiveGeometryCollection->RemoveAttribute(FGeometryDynamicCollection::SharedImplicitsAttribute, FTransformCollection::TransformGroup);
 			}
-			else if(TManagedArray<TSharedPtr<Chaos::FImplicitObject, ESPMode::ThreadSafe>>* OldAttrC = ArchiveGeometryCollection->FindAttributeTyped<TSharedPtr<Chaos::FImplicitObject, ESPMode::ThreadSafe>>(
-				FGeometryDynamicCollection::ImplicitsAttribute, FTransformCollection::TransformGroup))
+			if (OldAttrC)
 			{
 				for (int32 Index = 0; Index < NumElems; ++Index)
 				{
@@ -2020,4 +2023,4 @@ bool UGeometryCollection::CanEditChange(const FProperty* InProperty) const
 
 	return true;
 }
-#endif
\ No newline at end of file
+#endif
