Copy On Write implementation for managed array collections

diff --git a/Engine/Source/Runtime/Experimental/Chaos/Private/GeometryCollection/ManagedArrayCollection.cpp b/Engine/Source/Runtime/Experimental/Chaos/Private/GeometryCollection/ManagedArrayCollection.cpp
--- a/Engine/Source/Runtime/Experimental/Chaos/Private/GeometryCollection/ManagedArrayCollection.cpp
+++ b/Engine/Source/Runtime/Experimental/Chaos/Private/GeometryCollection/ManagedArrayCollection.cpp
@@ -53,7 +53,7 @@ void FManagedArrayCollection::RemoveElements(const FName& Group, const TArray<in
 		GeometryCollectionAlgo::BuildIncrementMask(SortedDeletionList, GroupSize, Offsets);
 
 		TSet<int32> DeletionSet(SortedDeletionList);
-		for (const TTuple<FKeyType, FValueType>& Entry : Map)
+		for (TTuple<FKeyType, FValueType>& Entry : Map)
 		{
 			//
 			// Reindex attributes dependent on the group being resized
@@ -786,9 +786,15 @@ void FManagedArrayCollection::FValueType::Serialize(FArchive& Ar)
 		Ar << bPersistent;
 	}
 
-	if (Value == nullptr)
+	if (ManagedArray == nullptr)
+	{
+		ensure(Ar.IsLoading());
+		ManagedArray = NewManagedTypedArray(ArrayType);
+		SharedManagedArray = TSharedPtr<FManagedArrayBase, ESPMode::NotThreadSafe>(ManagedArray);
+	}
+	else
 	{
-		Value = NewManagedTypedArray(ArrayType);
+		ensure(Ar.IsSaving());
 	}
 	
 	// Note: We switched to always saving the value here, and use the Saved flag
@@ -796,7 +802,7 @@ void FManagedArrayCollection::FValueType::Serialize(FArchive& Ar)
 	bool bNewSavedBehavior = Ar.CustomVer(FUE5MainStreamObjectVersion::GUID) >= FUE5MainStreamObjectVersion::ManagedArrayCollectionAlwaysSerializeValue;
 	if (bNewSavedBehavior || bPersistent)
 	{
-		Value->Serialize(static_cast<Chaos::FChaosArchive&>(Ar));
+		ManagedArray->Serialize(static_cast<Chaos::FChaosArchive&>(Ar));
 	}
 }
 
@@ -805,7 +811,8 @@ FManagedArrayCollection::FValueType::FValueType()
 	, GroupIndexDependency(NAME_None)
 	, bPersistent(true)
 	, bExternalValue(false)
-	, Value(nullptr) 
+	, SharedManagedArray(nullptr)
+	, ManagedArray(nullptr)
 {};
 
 FManagedArrayCollection::FValueType::FValueType(const FValueType& Other)
@@ -813,13 +820,24 @@ FManagedArrayCollection::FValueType::FValueType(const FValueType& Other)
 	, GroupIndexDependency(Other.GroupIndexDependency)
 	, bPersistent(Other.bPersistent)
 	, bExternalValue(false)
-	, Value(nullptr)
+	, SharedManagedArray(nullptr)
+	, ManagedArray(nullptr)
 {
-	if (Other.Value)
+	if (Other.ManagedArray)
 	{
-		this->Value = NewManagedTypedArray(this->ArrayType);
-		this->Value->Resize(Other.Value->Num());
-		this->Value->Init(*Other.Value);
+		if (Other.bExternalValue)
+		{
+			ManagedArray = NewManagedTypedArray(ArrayType);
+			ManagedArray->Resize(Other.ManagedArray->Num());
+			ManagedArray->Init(*Other.ManagedArray);
+		}
+		else
+		{
+			// we only copy the shared pointer as we don't want to pay for the full copy
+			// copy on write will make this unique if necessary
+			SharedManagedArray = Other.SharedManagedArray;
+			ManagedArray = SharedManagedArray.Get();
+		}
 	}
 };
 
@@ -828,79 +846,107 @@ FManagedArrayCollection::FValueType::FValueType(FValueType&& Other)
 	, GroupIndexDependency(Other.GroupIndexDependency)
 	, bPersistent(Other.bPersistent)
 	, bExternalValue(Other.bExternalValue)
-	, Value(Other.Value)
+	, SharedManagedArray(nullptr)
+	, ManagedArray(nullptr)
 {
 	if (&Other != this)
 	{
-		Other.Value = nullptr;
+		SharedManagedArray = MoveTemp(Other.SharedManagedArray);
+		ManagedArray = Other.ManagedArray;
+		Other.ManagedArray = nullptr;
+		Other.bExternalValue = false;
 	}
 }
 
 FManagedArrayCollection::FValueType::~FValueType()
 {
-	if (Value && !bExternalValue)
+	if (bExternalValue)
 	{
-		delete Value;
+		check(!SharedManagedArray.IsValid());
 	}
 }
 
-FManagedArrayBase& FManagedArrayCollection::FValueType::Modify() const
+void FManagedArrayCollection::FValueType::MakeUniqueForWrite()
 {
-	// todo : copy on write ?  ( use  NewManagedTypedArray )
-	Value->MarkDirty();
-	return *Value;
+	// IsUnique is performant in that case because we are using NonThreadSafe SharedPtr
+	// this is a requirement for CopyOnWrite as we forbid cross-thread CopyOnWrite
+	if (SharedManagedArray && !SharedManagedArray.IsUnique())
+	{
+		check(!IsExternal());
+
+		ManagedArray = NewManagedTypedArray(ArrayType);
+		ManagedArray->Resize(SharedManagedArray->Num());
+		ManagedArray->Init(*SharedManagedArray);
+
+		SharedManagedArray = TSharedPtr<FManagedArrayBase, ESPMode::NotThreadSafe>(ManagedArray);
+	}
+	check(ManagedArray);
+}
+
+FManagedArrayBase& FManagedArrayCollection::FValueType::Modify()
+{
+	MakeUniqueForWrite();
+
+	ManagedArray->MarkDirty();
+	return *ManagedArray;
 }
 
 void FManagedArrayCollection::FValueType::Reserve(int32 ReservedSize)
 {
-	// todo : copy on write ? 
-	Value->Reserve(ReservedSize);
+	if (ReservedSize > ManagedArray->Max())
+	{
+		MakeUniqueForWrite();
+		ManagedArray->Reserve(ReservedSize);
+	}
 }
 
 void FManagedArrayCollection::FValueType::Resize(int32 NewSize)
 {
-	// todo : copy on write ? 
-	Value->Resize(NewSize);
+	if (NewSize > ManagedArray->Num())
+	{
+		MakeUniqueForWrite();
+		ManagedArray->Resize(NewSize);
+	}
 }
 
 void FManagedArrayCollection::FValueType::InitFrom(const FManagedArrayCollection::FValueType& Other)
 {
-	// todo : copy on write ? 
+	MakeUniqueForWrite();
 	if (ArrayType == Other.GetArrayType())
 	{
-		Value->Init(Other.Get());
+		ManagedArray->Init(Other.Get());
 	}
 	else
 	{
-		Value->Convert(Other.Get());
+		ManagedArray->Convert(Other.Get());
 	}
 }
 
 void FManagedArrayCollection::FValueType::Exchange(FValueType& Other)
 {
-	// todo : copy on write ? 
+	MakeUniqueForWrite();
 	check(ArrayType == Other.ArrayType);
-	Value->ExchangeArrays(*Other.Value);
+	ManagedArray->ExchangeArrays(*Other.ManagedArray);
 }
 
 void FManagedArrayCollection::FValueType::Convert(FValueType& Other)
 {
-	// todo : copy on write ? 
+	MakeUniqueForWrite();
 	check(ArrayType != Other.ArrayType);
-	Value->Convert(*Other.Value);
+	ManagedArray->Convert(*Other.ManagedArray);
 }
 
 void FManagedArrayCollection::FValueType::CopyFrom(const FValueType& Other)
 {
-	// todo : copy on write ? 
-	check(Other.Value->Num() <= Value->Num());
-	Value->CopyRange(*Other.Value, 0, Other.Value->Num());
+	MakeUniqueForWrite();
+	check(Other.ManagedArray->Num() <= ManagedArray->Num());
+	ManagedArray->CopyRange(*Other.ManagedArray, 0, Other.ManagedArray->Num());
 }
 
 void FManagedArrayCollection::FValueType::Empty()
 {
-	// todo : copy on write ? 
-	Value->Empty();
+	MakeUniqueForWrite();
+	ManagedArray->Empty();
 }
 
 void FManagedArrayCollection::FValueType::RemoveGroupIndexDependency(FName Group)
diff --git a/Engine/Source/Runtime/Experimental/Chaos/Public/GeometryCollection/ManagedArrayCollection.h b/Engine/Source/Runtime/Experimental/Chaos/Public/GeometryCollection/ManagedArrayCollection.h
--- a/Engine/Source/Runtime/Experimental/Chaos/Public/GeometryCollection/ManagedArrayCollection.h
+++ b/Engine/Source/Runtime/Experimental/Chaos/Public/GeometryCollection/ManagedArrayCollection.h
@@ -651,7 +651,10 @@ private:
 		bool bPersistent;
 		bool bExternalValue;	//External arrays have external memory management.
 
-		FManagedArrayBase* Value;
+		// TODO : save the thread on which the object has been created to call out when unsafe access is done ? 
+
+		TSharedPtr<FManagedArrayBase, ESPMode::NotThreadSafe> SharedManagedArray; // external arrays have this set this to null
+		FManagedArrayBase* ManagedArray = nullptr;;
 
 	private:
 		/**
@@ -664,9 +667,11 @@ private:
 			, GroupIndexDependency(InGroupIndexDependency)
 			, bPersistent(bInPersistent)
 			, bExternalValue(false)
-			, Value(ArrayPtr.Release())
+			, SharedManagedArray(ArrayPtr.Release())
+			, ManagedArray(SharedManagedArray.Get())
 		{
-			Value->Resize(InitialSize);
+			MakeUniqueForWrite();
+			ManagedArray->Resize(InitialSize);
 		};
 
 		/**
@@ -679,11 +684,14 @@ private:
 			, GroupIndexDependency(InGroupIndexDependency)
 			, bPersistent(bInPersistent)
 			, bExternalValue(true)
-			, Value(ExternalArrayPtr)
+			, SharedManagedArray(nullptr)
+			, ManagedArray(ExternalArrayPtr)
 		{
-			Value->Resize(InitialSize);
+			ManagedArray->Resize(InitialSize);
 		};
 
+		CHAOS_API void MakeUniqueForWrite();
+
 	public:
 		/**
 		* Create a typed attribute
@@ -711,7 +719,7 @@ private:
 		FName GetGroupIndexDependency() const { return GroupIndexDependency; };
 		bool IsPersistent() const { return bPersistent; }
 		bool IsExternal() const { return bExternalValue; }
-		bool IsDirty() const { return Value->IsDirty(); }
+		bool IsDirty() const { return ManagedArray->IsDirty(); }
 
 		void SetGroupIndexDependency(FName NewGroupDependency) { GroupIndexDependency = NewGroupDependency; }
 
@@ -723,7 +731,7 @@ private:
 		{ 
 			check(IsSameType<T>());
 			
-			const TManagedArray<T>* TypedManagedArray = static_cast<const TManagedArray<T>*>(Value);
+			const TManagedArray<T>* TypedManagedArray = static_cast<const TManagedArray<T>*>(ManagedArray);
 			check(TypedManagedArray != nullptr);
 			return TypedManagedArray;
 		}
@@ -734,15 +742,15 @@ private:
 			return *GetTypedPtr<T>();
 		}
 
-		const FManagedArrayBase& Get() const { return *Value; }
+		const FManagedArrayBase& Get() const { return *ManagedArray; }
 
 		template <typename T>
-		TManagedArray<T>* ModifyTypedPtr() const
+		TManagedArray<T>* ModifyTypedPtr()
 		{ 
 			check(IsSameType<T>());
-			// todo : copy on write ? 
+			MakeUniqueForWrite();
 
-			TManagedArray<T>* TypedManagedArray = static_cast<TManagedArray<T>*>(Value);
+			TManagedArray<T>* TypedManagedArray = static_cast<TManagedArray<T>*>(ManagedArray);
 			check(TypedManagedArray != nullptr);
 
 			TypedManagedArray->MarkDirty();
@@ -750,12 +758,12 @@ private:
 		}
 
 		template <typename T>
-		TManagedArray<T>& ModifyTyped() const
+		TManagedArray<T>& ModifyTyped()
 		{
 			return *ModifyTypedPtr<T>();
 		}
 
-		FManagedArrayBase& Modify() const;
+		FManagedArrayBase& Modify();
 
 		void Reserve(int32 ReservedSize);
 		void Resize(int32 NewSize);
