managed array collection : Fix crash when loading a dataflow asset

diff --git a/Engine/Source/Runtime/Experimental/Chaos/Private/GeometryCollection/ManagedArrayCollection.cpp b/Engine/Source/Runtime/Experimental/Chaos/Private/GeometryCollection/ManagedArrayCollection.cpp
--- a/Engine/Source/Runtime/Experimental/Chaos/Private/GeometryCollection/ManagedArrayCollection.cpp
+++ b/Engine/Source/Runtime/Experimental/Chaos/Private/GeometryCollection/ManagedArrayCollection.cpp
@@ -830,6 +830,9 @@ FManagedArrayCollection::FValueType::FValueType(const FValueType& Other)
 			ManagedArray = NewManagedTypedArray(ArrayType);
 			ManagedArray->Resize(Other.ManagedArray->Num());
 			ManagedArray->Init(*Other.ManagedArray);
+
+			SharedManagedArray = TSharedPtr<FManagedArrayBase, ESPMode::NotThreadSafe>(ManagedArray);
+			bExternalValue = false;
 		}
 		else
 		{
@@ -856,6 +859,7 @@ FManagedArrayCollection::FValueType::FValueType(FValueType&& Other)
 		Other.ManagedArray = nullptr;
 		Other.bExternalValue = false;
 	}
+	check(ManagedArray != nullptr);
 }
 
 FManagedArrayCollection::FValueType::~FValueType()
