remove unecessary ensure when serializing a manged array collection ( side effect of my recent changes )

diff --git a/Engine/Source/Runtime/Experimental/Chaos/Private/GeometryCollection/ManagedArrayCollection.cpp b/Engine/Source/Runtime/Experimental/Chaos/Private/GeometryCollection/ManagedArrayCollection.cpp
--- a/Engine/Source/Runtime/Experimental/Chaos/Private/GeometryCollection/ManagedArrayCollection.cpp
+++ b/Engine/Source/Runtime/Experimental/Chaos/Private/GeometryCollection/ManagedArrayCollection.cpp
@@ -792,10 +792,6 @@ void FManagedArrayCollection::FValueType::Serialize(FArchive& Ar)
 		ManagedArray = NewManagedTypedArray(ArrayType);
 		SharedManagedArray = TSharedPtr<FManagedArrayBase, ESPMode::NotThreadSafe>(ManagedArray);
 	}
-	else
-	{
-		ensure(Ar.IsSaving());
-	}
 	
 	// Note: We switched to always saving the value here, and use the Saved flag
 	// to remove the property from the overall Map (see FManagedArrayCollection::Serialize above)
