ChaosFlesh: Add safe guard to the code for blow up simulation and add pervertex stiffness control code

diff --git a/Engine/Source/Runtime/Experimental/Chaos/Private/Chaos/Deformable/ChaosDeformableSolver.cpp b/Engine/Source/Runtime/Experimental/Chaos/Private/Chaos/Deformable/ChaosDeformableSolver.cpp
--- a/Engine/Source/Runtime/Experimental/Chaos/Private/Chaos/Deformable/ChaosDeformableSolver.cpp
+++ b/Engine/Source/Runtime/Experimental/Chaos/Private/Chaos/Deformable/ChaosDeformableSolver.cpp
@@ -98,6 +98,9 @@ namespace Chaos::Softs
 	int32 GSParallelMax = 100;
 	FAutoConsoleVariableRef CVarDeformableGSParrelMax(TEXT("p.Chaos.Deformable.GSParallelMax"), GSParallelMax, TEXT("Minimal number of particles to process in parallel for Gauss Seidel constraints. [def: 100]"));
 
+	float MaxDxRatio = .01f;
+	FAutoConsoleVariableRef CVarDeformableGSMaxDxRatio(TEXT("p.Chaos.Deformable.GSMaxDxRatio"), MaxDxRatio, TEXT("Max size for dx in each iteration for Gauss Seidel constraints. [def: .01]"));
+
 	FDeformableSolver::FDeformableSolver(FDeformableSolverProperties InProp)
 		: CurrentInputPackage(TUniquePtr < FDeformablePackage >(nullptr))
 		, PreviousInputPackage(TUniquePtr < FDeformablePackage >(nullptr))
@@ -161,6 +164,8 @@ namespace Chaos::Softs
 			AllTetEMeshArray.Reset(new TArray<FSolverReal>());
 			AllTetNuMeshArray.Reset(new TArray<FSolverReal>());
 			AllTetAlphaJArray.Reset(new TArray<FSolverReal>());
+			AllCorotatedCodEMeshArray.Reset(new TArray<FSolverReal>());
+			AllSkinEMeshArray.Reset(new TArray<FSolverReal>());
 			GSWeakConstraints.Reset(new FGaussSeidelWeakConstraints<FSolverReal, FSolverParticles>({}, {}, {}, {}, {}, GDeformableXPBDWeakConstraintParams));
 			MuscleActivationConstraints.Reset(new FMuscleActivationConstraints<FSolverReal, FSolverParticles>());
 		}
@@ -762,6 +767,8 @@ namespace Chaos::Softs
 
 		auto ChaosTet = [](FIntVector4 V, int32 dp) { return Chaos::TVec4<int32>(dp + V.X, dp + V.Y, dp + V.Z, dp + V.W); };
 
+		TArray<FSolverReal> StiffnessWithMultiplier;
+
 		const TManagedArray<FIntVector4>& Tetrahedron = Rest.GetAttribute<FIntVector4>("Tetrahedron", "Tetrahedral");
 		if (uint32 NumElements = Tetrahedron.Num())
 		{
@@ -789,7 +796,7 @@ namespace Chaos::Softs
 			if (Rest.HasAttributes({ FManagedArrayCollection::TManagedType<FSolverReal>("Stiffness", FGeometryCollection::VerticesGroup) }))
 			{
 				uint32 NumParticles = Rest.NumElements(FGeometryCollection::VerticesGroup);
-				TArray<FSolverReal> StiffnessWithMultiplier;
+
 				StiffnessWithMultiplier.Init(0.f, NumParticles);
 				FSolverReal StiffnessMultiplier = 1.f;
 				FSolverReal IncompressibilityMultiplier = 1.f;
@@ -1001,6 +1008,22 @@ namespace Chaos::Softs
 			}
 		
 		}
+
+		AllCorotatedCodEMeshArray->Init(Property.EMesh, AllUnconstrainedSurfaceElementsCorotatedCod->Num());
+		AllSkinEMeshArray->Init(Property.EMesh, AllUnconstrainedSurfaceElementsSkin->Num());
+		if (StiffnessWithMultiplier.Num() > 0)
+		{
+			for (int32 i = 0; i < AllUnconstrainedSurfaceElementsCorotatedCod->Num(); i++)
+			{
+				(*AllCorotatedCodEMeshArray)[i] = (StiffnessWithMultiplier[(*AllUnconstrainedSurfaceElementsCorotatedCod)[i][0]] + StiffnessWithMultiplier[(*AllUnconstrainedSurfaceElementsCorotatedCod)[i][1]]
+								+ StiffnessWithMultiplier[(*AllUnconstrainedSurfaceElementsCorotatedCod)[i][2]]) / 3.f;
+			}
+			for (int32 i = 0; i < AllSkinEMeshArray->Num(); i++)
+			{
+				(*AllSkinEMeshArray)[i] = (StiffnessWithMultiplier[(*AllUnconstrainedSurfaceElementsSkin)[i][0]] + StiffnessWithMultiplier[(*AllUnconstrainedSurfaceElementsSkin)[i][1]]
+								+ StiffnessWithMultiplier[(*AllUnconstrainedSurfaceElementsSkin)[i][2]]) / 3.f;
+			}
+		}
 	}
 
 
@@ -1393,12 +1416,12 @@ namespace Chaos::Softs
 	{
 		PERF_SCOPE(STAT_ChaosDeformableSolver_InitializeGaussSeidelConstraintVariables);
 
-		GSMainConstraint.Reset(new Chaos::Softs::FGaussSeidelMainConstraint<FSolverReal, FSolverParticles>(Evolution->Particles(), Property.bDoQuasistatics, Property.bUseSOR, Property.OmegaSOR, GSParallelMax));
+		GSMainConstraint.Reset(new Chaos::Softs::FGaussSeidelMainConstraint<FSolverReal, FSolverParticles>(Evolution->Particles(), Property.bDoQuasistatics, Property.bUseSOR, Property.OmegaSOR, GSParallelMax, MaxDxRatio));
 
 		if (AllUnconstrainedSurfaceElementsCorotatedCod->Num() > 0) 
 		{
 			GSCorotatedCodConstraints.Reset(new Chaos::Softs::FGaussSeidelCorotatedCodimensionalConstraints<FSolverReal, FSolverParticles>(
-				Evolution->Particles(), *AllUnconstrainedSurfaceElementsCorotatedCod, false, Property.EMesh));
+				Evolution->Particles(), *AllUnconstrainedSurfaceElementsCorotatedCod, *AllCorotatedCodEMeshArray));
 			TArray<TArray<int32>> IncidentElements, IncidentElementsLocal;
 			GSMainConstraint->AddStaticConstraints(GSCorotatedCodConstraints->GetConstraintsArray(), IncidentElements, IncidentElementsLocal);
 
@@ -1413,7 +1436,7 @@ namespace Chaos::Softs
 		if (AllUnconstrainedSurfaceElementsSkin->Num() > 0)
 		{
 			GSLinearCodConstraints.Reset(new Chaos::Softs::FGaussSeidelLinearCodimensionalConstraints<FSolverReal, FSolverParticles>(
-				Evolution->Particles(), *AllUnconstrainedSurfaceElementsSkin, false, Property.EMesh));
+				Evolution->Particles(), *AllUnconstrainedSurfaceElementsSkin, *AllSkinEMeshArray));
 			TArray<TArray<int32>> IncidentElements, IncidentElementsLocal;
 			GSMainConstraint->AddStaticConstraints(GSLinearCodConstraints->GetConstraintsArray(), IncidentElements, IncidentElementsLocal);
 
diff --git a/Engine/Source/Runtime/Experimental/Chaos/Public/Chaos/Deformable/ChaosDeformableSolver.h b/Engine/Source/Runtime/Experimental/Chaos/Public/Chaos/Deformable/ChaosDeformableSolver.h
--- a/Engine/Source/Runtime/Experimental/Chaos/Public/Chaos/Deformable/ChaosDeformableSolver.h
+++ b/Engine/Source/Runtime/Experimental/Chaos/Public/Chaos/Deformable/ChaosDeformableSolver.h
@@ -211,6 +211,8 @@ namespace Chaos::Softs
 		TUniquePtr <TArray<FSolverReal>> AllSecondWeights;
 		TUniquePtr <TArray<TVec3<int32>>> AllUnconstrainedSurfaceElementsCorotatedCod;  //correspond to the triangle mesh elements that are simulated using corotated cod
 		TUniquePtr <TArray<TVec3<int32>>> AllUnconstrainedSurfaceElementsSkin;          //correspond to the triangle mesh elements that are simulated using linear cod constraints
+		TUniquePtr <TArray<FSolverReal>> AllCorotatedCodEMeshArray;
+		TUniquePtr <TArray<FSolverReal>> AllSkinEMeshArray;
 		TArray<int32> ParticleComponentIndex;
 		TMap<int32, TSet<int32>> ParticleTriangleExclusionMap;
 		//Muscle Activation Variables
diff --git a/Engine/Source/Runtime/Experimental/Chaos/Public/Chaos/Deformable/GaussSeidelCorotatedCodimensionalConstraints.h b/Engine/Source/Runtime/Experimental/Chaos/Public/Chaos/Deformable/GaussSeidelCorotatedCodimensionalConstraints.h
--- a/Engine/Source/Runtime/Experimental/Chaos/Public/Chaos/Deformable/GaussSeidelCorotatedCodimensionalConstraints.h
+++ b/Engine/Source/Runtime/Experimental/Chaos/Public/Chaos/Deformable/GaussSeidelCorotatedCodimensionalConstraints.h
@@ -17,7 +17,6 @@ namespace Chaos::Softs
 	{
 
 	public:
-		//this one only accepts tetmesh input and mesh
 		FGaussSeidelCorotatedCodimensionalConstraints(
 			const ParticleType& InParticles,
 			const TArray<TVector<int32, 3>>& InMesh,
@@ -37,6 +36,28 @@ namespace Chaos::Softs
 			InitializeCodimensionData(InParticles);
 		}
 
+		FGaussSeidelCorotatedCodimensionalConstraints(
+			const ParticleType& InParticles,
+			const TArray<TVector<int32, 3>>& InMesh,
+			const TArray<T>& EMeshArray,
+			const T& NuMesh = (T).3
+		)
+			: MeshConstraints(InMesh)
+		{
+			ensureMsgf(EMeshArray.Num() == InMesh.Num(), TEXT("Input Young Modulus Array Size is wrong"));
+			Measure.Init((T)0., MeshConstraints.Num());
+
+			InitializeCodimensionData(InParticles);
+			LambdaElementArray.Init((T)0., MeshConstraints.Num());
+			MuElementArray.Init((T)0., MeshConstraints.Num());
+			
+			for (int32 e = 0; e < InMesh.Num(); e++)
+			{
+				LambdaElementArray[e] = EMeshArray[e] * NuMesh / (((T)1. + NuMesh) * ((T)1. - (T)2. * NuMesh));
+				MuElementArray[e] = EMeshArray[e] / ((T)2. * ((T)1. + NuMesh));
+			}
+		}
+
 		virtual ~FGaussSeidelCorotatedCodimensionalConstraints() {}
 
 		PMatrix<T, 3, 3> DsInit(const int32 E, const ParticleType& InParticles)  
diff --git a/Engine/Source/Runtime/Experimental/Chaos/Public/Chaos/Deformable/GaussSeidelLinearCodimensionalConstraints.h b/Engine/Source/Runtime/Experimental/Chaos/Public/Chaos/Deformable/GaussSeidelLinearCodimensionalConstraints.h
--- a/Engine/Source/Runtime/Experimental/Chaos/Public/Chaos/Deformable/GaussSeidelLinearCodimensionalConstraints.h
+++ b/Engine/Source/Runtime/Experimental/Chaos/Public/Chaos/Deformable/GaussSeidelLinearCodimensionalConstraints.h
@@ -17,15 +17,13 @@ namespace Chaos::Softs
 	{
 
 	public:
-		//this one only accepts tetmesh input and mesh
 		FGaussSeidelLinearCodimensionalConstraints(
 			const ParticleType& InParticles,
 			const TArray<TVector<int32, 3>>& InMesh,
-			const bool bRecordMetricIn = true,
 			const T& EMesh = (T)10.0,
 			const T& NuMesh = (T).3
 		)
-			: bRecordMetric(bRecordMetricIn), MeshConstraints(InMesh)
+			: MeshConstraints(InMesh)
 		{
 			Measure.Init((T)0., MeshConstraints.Num());
 			Lambda = EMesh * NuMesh / (((T)1. + NuMesh) * ((T)1. - (T)2. * NuMesh));
@@ -37,6 +35,27 @@ namespace Chaos::Softs
 			InitializeCodimensionData(InParticles);
 		}
 
+		FGaussSeidelLinearCodimensionalConstraints(
+			const ParticleType& InParticles,
+			const TArray<TVector<int32, 3>>& InMesh,
+			const TArray<T>& EMeshArray,
+			const T& NuMesh = (T).3
+		)
+			: MeshConstraints(InMesh)
+		{
+			Measure.Init((T)0., MeshConstraints.Num());
+
+			InitializeCodimensionData(InParticles);
+			LambdaElementArray.Init((T)0., MeshConstraints.Num());
+			MuElementArray.Init((T)0., MeshConstraints.Num());
+			
+			for (int32 e = 0; e < InMesh.Num(); e++)
+			{
+				LambdaElementArray[e] = EMeshArray[e] * NuMesh / (((T)1. + NuMesh) * ((T)1. - (T)2. * NuMesh));
+				MuElementArray[e] = EMeshArray[e] / ((T)2. * ((T)1. + NuMesh));
+			}
+		}
+
 		virtual ~FGaussSeidelLinearCodimensionalConstraints() {}
 
 		PMatrix<T, 3, 3> DsInit(const int32 E, const ParticleType& InParticles)  
@@ -230,7 +249,6 @@ namespace Chaos::Softs
 		TArray<T> MuElementArray;
 		TArray<T> LambdaElementArray;
 		TArray<T> AlphaJArray;
-		bool bRecordMetric;
 
 		TArray<TVector<int32, 3>> MeshConstraints;
 		mutable TArray<T> Measure;
diff --git a/Engine/Source/Runtime/Experimental/Chaos/Public/Chaos/Deformable/GaussSeidelMainConstraint.h b/Engine/Source/Runtime/Experimental/Chaos/Public/Chaos/Deformable/GaussSeidelMainConstraint.h
--- a/Engine/Source/Runtime/Experimental/Chaos/Public/Chaos/Deformable/GaussSeidelMainConstraint.h
+++ b/Engine/Source/Runtime/Experimental/Chaos/Public/Chaos/Deformable/GaussSeidelMainConstraint.h
@@ -18,6 +18,8 @@ DECLARE_CYCLE_STAT(TEXT("Chaos.Deformable.GSMainConstraint.Acceleration"), STAT_
 DECLARE_CYCLE_STAT(TEXT("Chaos.Deformable.GSMainConstraint.Init"), STAT_ChaosGSMainConstraint_Init, STATGROUP_Chaos);
 DECLARE_CYCLE_STAT(TEXT("Chaos.Deformable.GSMainConstraint.InitDynamicColor"), STAT_ChaosGSMainConstraint_InitDynamicColor, STATGROUP_Chaos);
 
+DEFINE_LOG_CATEGORY_STATIC(LogDeformableGaussSeidelMainConstraint, Log, All);
+
 namespace Chaos::Softs
 {
 	template <typename T, typename ParticleType>
@@ -31,6 +33,7 @@ namespace Chaos::Softs
 			const bool bDoSORIn = true,
 			const T InOmegaSOR = (T)1.6,
 			const int32 ParallelMaxIn = 1000, 
+			const T MaxDxRatioIn = T(1),
 			const FDeformableXPBDCorotatedParams& InParams = FDeformableXPBDCorotatedParams())
 			: bDoQuasistatics(bDoQuasistaticsIn)
 			, bDoAcceleration(bDoSORIn)
@@ -41,6 +44,23 @@ namespace Chaos::Softs
 			Resize((int32)InParticles.Size());
 
 			InitializeLambdas();
+			
+			TVec3<T> MaxCoord((T)100.), MinCoord((T)-100.);
+			for (int32 i = 0; i < (int32)InParticles.Size(); i++)
+			{
+				for (int32 j = 0; j < 3; j++) 
+				{
+					if (InParticles.X(i)[j] < MinCoord[j])
+					{
+						MinCoord[j] = InParticles.X(i)[j];
+					}
+					if (InParticles.X(i)[j] > MaxCoord[j])
+					{
+						MaxCoord[j] = InParticles.X(i)[j];
+					}
+				}
+			}
+			MaxDxSize = (MaxCoord - MinCoord).Size() * MaxDxRatioIn;
 		}
 
 		virtual ~FGaussSeidelMainConstraint() {}
@@ -253,7 +273,14 @@ namespace Chaos::Softs
 				HessianInv *= T(1) / HessianDet;
 				Chaos::TVector<T, 3> Dx = HessianInv.GetTransposed() * (-ParticleResidual);
 
-				Particles.P(p) += Dx;
+				if (Dx.Size() < MaxDxSize)
+				{
+					Particles.P(p) += Dx;
+				}
+				else
+				{
+					UE_LOG(LogDeformableGaussSeidelMainConstraint, Warning, TEXT("Following Particle is skipped because of too large dx size: %d"), p);
+				}
 			}
 		}
 		
@@ -340,6 +367,8 @@ namespace Chaos::Softs
 		int32 NumTotalParticles;
 		TArray<TVec3<T>> ReorderedPs;
 
+		T MaxDxSize = T(0);
+
 		public:
 		bool DebugResidual = false;
 		bool IsFirstFrame = true;
