Trim leading/trailing slashes from chunks directory

diff --git a/Engine/Source/Runtime/Experimental/IoStore/OnDemand/Private/OnDemandIoStore.cpp b/Engine/Source/Runtime/Experimental/IoStore/OnDemand/Private/OnDemandIoStore.cpp
--- a/Engine/Source/Runtime/Experimental/IoStore/OnDemand/Private/OnDemandIoStore.cpp
+++ b/Engine/Source/Runtime/Experimental/IoStore/OnDemand/Private/OnDemandIoStore.cpp
@@ -660,11 +660,11 @@ void FOnDemandIoStore::CreateContainersFromToc(
 		ChunksDirectory = Sb;
 		if (ChunksDirectory.StartsWith('/'))
 		{
-			TocPath.RemovePrefix(1);
+			ChunksDirectory.RemovePrefix(1);
 		}
 		if (ChunksDirectory.EndsWith('/'))
 		{
-			TocPath.RemoveSuffix(1);
+			ChunksDirectory.RemoveSuffix(1);
 		}
 	}
 
