Use absolute URL when retrieving the header chunk

diff --git a/Engine/Source/Runtime/Experimental/IoStore/OnDemand/Private/OnDemandIoStore.cpp b/Engine/Source/Runtime/Experimental/IoStore/OnDemand/Private/OnDemandIoStore.cpp
--- a/Engine/Source/Runtime/Experimental/IoStore/OnDemand/Private/OnDemandIoStore.cpp
+++ b/Engine/Source/Runtime/Experimental/IoStore/OnDemand/Private/OnDemandIoStore.cpp
@@ -720,16 +720,24 @@ FIoStatus FOnDemandIoStore::InstallContainers(
 			continue;
 		}
 
-		auto GetChunkUrl = [&AnsiTocPath, &Container](const FOnDemandChunkEntry& Entry, FAnsiStringBuilderBase& OutUrl) -> FAnsiStringBuilderBase&
+		auto GetChunkUrl = [&AnsiTocPath, &Container](
+			const FStringView& Host,
+			const FOnDemandChunkEntry& Entry,
+			FAnsiStringBuilderBase& OutUrl) -> FAnsiStringBuilderBase&
 		{
 			OutUrl.Reset();
+			if (Host.IsEmpty() == false)
+			{
+				OutUrl << Host;
+			}
+
 			if (Container->ChunksDirectory.IsEmpty())
 			{
 				OutUrl << AnsiTocPath << "/" << "chunks";
 			}
 			else
 			{
-				OutUrl.Append(Container->ChunksDirectory);
+				OutUrl << "/" << Container->ChunksDirectory;
 			}
 
 			const FString HashString = LexToString(Entry.Hash);
@@ -746,7 +754,7 @@ FIoStatus FOnDemandIoStore::InstallContainers(
 		const FIoChunkId HeaderChunkId = CreateContainerHeaderChunkId(Container->ContainerId);
 		if (const FOnDemandChunkEntry* Entry = Container->ChunkEntries.Find(HeaderChunkId))
 		{
-			TIoStatusOr<FIoBuffer> Response = FHttpClient::Get(GetChunkUrl(*Entry, ChunkUrl).ToView(), 2, EHttpRedirects::Follow);
+			TIoStatusOr<FIoBuffer> Response = FHttpClient::Get(GetChunkUrl(Host, *Entry, ChunkUrl).ToView(), 2, EHttpRedirects::Follow);
 			if (Response.IsOk() == false)
 			{
 				return FIoStatusBuilder(EIoErrorCode::InvalidCode) << TEXT("Failed to fetch container header chunk from URL");
@@ -791,7 +799,7 @@ FIoStatus FOnDemandIoStore::InstallContainers(
 			
 			ConcurrentRequests++;
 			HttpClient->Get(
-				GetChunkUrl(TocEntry, ChunkUrl).ToView(),
+				GetChunkUrl(FStringView(), TocEntry, ChunkUrl).ToView(),
 				[this, &Status, ChunkId, ExpectedHash = TocEntry.Hash, &ConcurrentRequests]
 				(TIoStatusOr<FIoBuffer> ChunkStatus, uint64 DurationMs)
 				{
