Disable on-demand TOC file validation until the file sentinal is always included.

diff --git a/Engine/Source/Runtime/Experimental/IoStore/OnDemand/Private/OnDemandIoStore.cpp b/Engine/Source/Runtime/Experimental/IoStore/OnDemand/Private/OnDemandIoStore.cpp
--- a/Engine/Source/Runtime/Experimental/IoStore/OnDemand/Private/OnDemandIoStore.cpp
+++ b/Engine/Source/Runtime/Experimental/IoStore/OnDemand/Private/OnDemandIoStore.cpp
@@ -520,9 +520,9 @@ FIoStatus FOnDemandIoStore::ProcessMountRequest(FMountRequest& MountRequest)
 				<< TEXT("'");
 		}
 
-		// TODO: Add sentinel to all serialization paths for TOC files
-		const bool bIncludeDot = true;	
-		const bool bValidate = FPathViews::GetExtension(Args.FilePath, bIncludeDot) == GIasOnDemandTocExt;
+		// TODO: Enable validation when the sentinal is included in all serialization paths 
+		//const bool bIncludeDot = true;	
+		const bool bValidate = false; //FPathViews::GetExtension(Args.FilePath, bIncludeDot) == GIasOnDemandTocExt;
 		if (bValidate)
 		{
 			const int64 SentinelPos = Ar->TotalSize() - FOnDemandTocSentinel::SentinelSize;
