Fix GJK instability on XSX and PS5

diff --git a/Engine/Source/Runtime/Experimental/Chaos/Public/Chaos/SimplexVectorized.h b/Engine/Source/Runtime/Experimental/Chaos/Public/Chaos/SimplexVectorized.h
--- a/Engine/Source/Runtime/Experimental/Chaos/Public/Chaos/SimplexVectorized.h
+++ b/Engine/Source/Runtime/Experimental/Chaos/Public/Chaos/SimplexVectorized.h
@@ -75,7 +75,7 @@ namespace Chaos
 
 
 		// Handle degenerate triangle
-		const VectorRegister4Float TriNormal = VectorCross(AB, AC);
+		const VectorRegister4Float TriNormal = Private::VectorCrossNoFMA(AB, AC);
 		const VectorRegister4Float TriNormal2 = VectorDot3(TriNormal, TriNormal);
 		const VectorRegister4Float MinFloat = MakeVectorRegisterFloatConstant(std::numeric_limits<float>::min(), std::numeric_limits<float>::min(), std::numeric_limits<float>::min(), std::numeric_limits<float>::min());
 		const VectorRegister4Float AMin = VectorMultiply(A, MinFloat);
@@ -308,10 +308,10 @@ namespace Chaos
 		*/
 
 		VectorRegister4Float Cofactors[4];
-		Cofactors[0] = VectorNegate(VectorDot3(X1, VectorCross(X2, X3)));
-		Cofactors[1] = VectorDot3(X0, VectorCross(X2, X3));
-		Cofactors[2] = VectorNegate(VectorDot3(X0, VectorCross(X1, X3)));
-		Cofactors[3] = VectorDot3(X0, VectorCross(X1, X2));
+		Cofactors[0] = VectorNegate(VectorDot3(X1, Private::VectorCrossNoFMA(X2, X3)));
+		Cofactors[1] = VectorDot3(X0, Private::VectorCrossNoFMA(X2, X3));
+		Cofactors[2] = VectorNegate(VectorDot3(X0, Private::VectorCrossNoFMA(X1, X3)));
+		Cofactors[3] = VectorDot3(X0, Private::VectorCrossNoFMA(X1, X2));
 		VectorRegister4Float DetM = VectorAdd(VectorAdd(Cofactors[0], Cofactors[1]), VectorAdd(Cofactors[2], Cofactors[3]));
 
 		bool bSignMatch[4];
diff --git a/Engine/Source/Runtime/Experimental/ChaosCore/Public/Chaos/VectorUtility.h b/Engine/Source/Runtime/Experimental/ChaosCore/Public/Chaos/VectorUtility.h
--- a/Engine/Source/Runtime/Experimental/ChaosCore/Public/Chaos/VectorUtility.h
+++ b/Engine/Source/Runtime/Experimental/ChaosCore/Public/Chaos/VectorUtility.h
@@ -129,6 +129,34 @@ namespace Chaos::Private
 		const VectorRegister4Float R2 = MakeVectorRegisterFloatFromDouble(MakeVectorRegister(M.M[2][0], M.M[2][1], M.M[2][2], 0.0));
 		return VectorMultiplyAdd(R0, VecX, VectorMultiplyAdd(R1, VecY, VectorMultiply(R2, VecZ)));
 	}
+
+	/**
+	 * Calculates the cross product of two vectors (XYZ components). W of the input should be 0, and will remain 0.
+	 * This function is not using FMA for stability reason, rounding with FMA could cause numerical instability. 
+	 * 
+	 * @param Vec1	1st vector
+	 * @param Vec2	2nd vector
+	 * @return		cross(Vec1.xyz, Vec2.xyz). W of the input should be 0, and will remain 0.
+	 */
+	FORCEINLINE VectorRegister4Float VectorCrossNoFMA(const VectorRegister4Float& Vec1, const VectorRegister4Float& Vec2)
+	{
+#if PLATFORM_ENABLE_VECTORINTRINSICS_NEON
+		return VectorCross(Vec1, Vec2);
+#elif PLATFORM_ENABLE_VECTORINTRINSICS
+		// YZX
+		VectorRegister4Float A = VectorSwizzle(Vec2, 1, 2, 0, 3);
+		VectorRegister4Float B = VectorSwizzle(Vec1, 1, 2, 0, 3);
+		// XY, YZ, ZX
+		A = VectorMultiply(A, Vec1);
+		// XY-YX, YZ-ZY, ZX-XZ
+		A = VectorSubtract(A, VectorMultiply(B, Vec2));
+		// YZ-ZY, ZX-XZ, XY-YX
+		return VectorSwizzle(A, 1, 2, 0, 3);
+#else
+		return VectorCross(Vec1, Vec2);
+#endif
+	}
+
 }
 
 /**
