Iris
* net.Iris.PrintNetCullDistances now takes into account cull distance overrides
* Added verbose logs when setting and clearing cull distance overrides

#rnx

diff --git a/Engine/Source/Runtime/Experimental/Iris/Core/Private/Iris/ReplicationSystem/NetCullDistanceOverrides.cpp b/Engine/Source/Runtime/Experimental/Iris/Core/Private/Iris/ReplicationSystem/NetCullDistanceOverrides.cpp
--- a/Engine/Source/Runtime/Experimental/Iris/Core/Private/Iris/ReplicationSystem/NetCullDistanceOverrides.cpp
+++ b/Engine/Source/Runtime/Experimental/Iris/Core/Private/Iris/ReplicationSystem/NetCullDistanceOverrides.cpp
@@ -11,9 +11,11 @@ void FNetCullDistanceOverrides::Init(const FNetCullDistanceOverridesInitParams&
 	ValidCullDistanceSqr.Init(InitParams.MaxObjectCount);
 }
 
-void FNetCullDistanceOverrides::ClearCullDistanceSqr(uint32 ObjectIndex)
+bool FNetCullDistanceOverrides::ClearCullDistanceSqr(uint32 ObjectIndex)
 {
+	const bool bWasBitSet = ValidCullDistanceSqr.IsBitSet(ObjectIndex);
 	ValidCullDistanceSqr.ClearBit(ObjectIndex);
+	return bWasBitSet;
 }
 
 void FNetCullDistanceOverrides::SetCullDistanceSqr(uint32 ObjectIndex, float CullDistSqr)
diff --git a/Engine/Source/Runtime/Experimental/Iris/Core/Private/Iris/ReplicationSystem/ObjectReplicationBridgeDebugging.cpp b/Engine/Source/Runtime/Experimental/Iris/Core/Private/Iris/ReplicationSystem/ObjectReplicationBridgeDebugging.cpp
--- a/Engine/Source/Runtime/Experimental/Iris/Core/Private/Iris/ReplicationSystem/ObjectReplicationBridgeDebugging.cpp
+++ b/Engine/Source/Runtime/Experimental/Iris/Core/Private/Iris/ReplicationSystem/ObjectReplicationBridgeDebugging.cpp
@@ -7,6 +7,8 @@
 #include "Iris/Core/IrisLog.h"
 #include "Iris/Core/IrisDebugging.h"
 
+#include "Iris/ReplicationSystem/NetCullDistanceOverrides.h"
+
 #include "Iris/ReplicationSystem/ReplicationSystem.h"
 #include "Iris/ReplicationSystem/ReplicationSystemTypes.h"
 #include "Iris/ReplicationSystem/ReplicationSystemInternal.h"
@@ -593,17 +595,34 @@ void UObjectReplicationBridge::PrintNetCullDistances(const TArray<FString>& Args
 
 	FReplicationSystemInternal* ReplicationSystemInternal = GetReplicationSystem()->GetReplicationSystemInternal();
 	const FWorldLocations& WorldLocations = ReplicationSystemInternal->GetWorldLocations();
+	const FNetCullDistanceOverrides& CullDistanceOverrides = ReplicationSystemInternal->GetNetCullDistanceOverrides();
 
 	struct FCullDistanceInfo
 	{
 		UClass* Class = nullptr;
 		float CDOCullDistance = 0.0f;
 
-		uint32 NumTotal = 0; // Total replicated rootobjects of this class
-		uint32 NumCDOCullDistance = 0; // Total replicated rootobjects
+		// Total replicated root objects of this class
+		uint32 NumTotal = 0;
+
+		// Track unique culldistance values for replicated root objects
+		TMap<float /*CullDistance*/, uint32 /*ActorCount with culldistance value*/> UniqueCullDistances; 
+
+		const float FindMostUsedCullDistance() const
+		{
+			float MostUsedCullDistance = 0.0;
+			uint32 MostUsedCount = 0;
+			for (auto It : UniqueCullDistances)
+			{
+				if (It.Value >= MostUsedCount)
+				{
+					MostUsedCount = It.Value;
+					MostUsedCullDistance = FMath::Max(It.Key, MostUsedCullDistance);
+				}
+			}
 
-		// Track culldistance values for actors that are different from the CDO
-		TMap<float /*CullDistance*/, uint32 /*ActorCount with culldistance value*/> DivergentCullDistances; 
+			return MostUsedCullDistance;
+		}
 	};
 
 	TMap<UClass*, FCullDistanceInfo> ClassCullDistanceMap;
@@ -624,47 +643,47 @@ void UObjectReplicationBridge::PrintNetCullDistances(const TArray<FString>& Args
 
 	RootObjects.ForAllSetBits([&](uint32 RootObjectIndex)
 	{
-		const float CurrentCullDistance = WorldLocations.GetCullDistance(RootObjectIndex);
-
 		if( UObject* RepObj = NetRefHandleManager->GetReplicatedObjectInstance(RootObjectIndex) )
 		{
 			UClass* RepObjClass = RepObj->GetClass();
-			UObject* RepClassCDO = RepObjClass->GetDefaultObject();
-
-			float CDOCullDistance = 0.0;
-
-			// Try to find the object's actual culldistance and that of the CDO.
-			if (GetInstanceWorldObjectInfoFunction)
+			// Find this object's current net cull distance
+			float RootObjectCullDistance = WorldLocations.GetCullDistance(RootObjectIndex);
+			if (CullDistanceOverrides.HasCullDistanceOverride(RootObjectIndex))
 			{
-				FVector Loc;
-				GetInstanceWorldObjectInfoFunction(NetRefHandleManager->GetNetRefHandleFromInternalIndex(RootObjectIndex), RepClassCDO, Loc, CDOCullDistance);
+				RootObjectCullDistance = FMath::Sqrt(CullDistanceOverrides.GetCullDistanceSqr(RootObjectIndex));
 			}
 
 			FCullDistanceInfo& Info = ClassCullDistanceMap.FindOrAdd(RepObjClass);
-
 			if (Info.Class == nullptr)
 			{
+				UObject* RepClassCDO = RepObjClass->GetDefaultObject();
+
+				// Find the CullDistance of the CDO.
+				float CDOCullDistance = 0.0;
+				if (GetInstanceWorldObjectInfoFunction)
+				{
+					FVector Loc;
+					GetInstanceWorldObjectInfoFunction(NetRefHandleManager->GetNetRefHandleFromInternalIndex(RootObjectIndex), RepClassCDO, Loc, CDOCullDistance);
+				}
+
 				Info.Class = RepObjClass;
 				Info.CDOCullDistance = CDOCullDistance;
 			}
 
 			Info.NumTotal++;
-			
-			// Check if Obj has diverged from the CDO
-			if(CurrentCullDistance != CDOCullDistance)
-			{
-				uint32& NumDivergent = Info.DivergentCullDistances.FindOrAdd(CurrentCullDistance, 0);
-				++NumDivergent;
-			}
-			else
-			{
-				Info.NumCDOCullDistance++;
-			}
+
+			uint32& NumUsingCullDistance = Info.UniqueCullDistances.FindOrAdd(RootObjectCullDistance, 0);
+			++NumUsingCullDistance;
 		}
 	});
 
 	// Sort from highest to lowest
-	ClassCullDistanceMap.ValueSort([](const FCullDistanceInfo& lhs, const FCullDistanceInfo& rhs) { return lhs.CDOCullDistance >= rhs.CDOCullDistance; });
+	ClassCullDistanceMap.ValueSort([](const FCullDistanceInfo& lhs, const FCullDistanceInfo& rhs) 
+	{ 
+		const float LHSSortingCullDistance = lhs.FindMostUsedCullDistance();
+		const float RHSSortingCullDistance = rhs.FindMostUsedCullDistance();
+		return LHSSortingCullDistance >= RHSSortingCullDistance;
+	});
 
 	UE_LOG(LogIrisBridge, Display, TEXT("################ Start Printing NetCullDistance Values ################"));
 	UE_LOG(LogIrisBridge, Display, TEXT(""));
@@ -674,14 +693,12 @@ void UObjectReplicationBridge::PrintNetCullDistances(const TArray<FString>& Args
 		FCullDistanceInfo& Info = ClassIt.Value();
 		UClass* Class = Info.Class;
 
-		UE_LOG(LogIrisBridge, Display, TEXT("NetCullDistance: %f | Class: %s | ReplicatedCount: %u | Using CDO CullDistance: %u (%.2f%%)"),
-			Info.CDOCullDistance, *Info.Class->GetName(), Info.NumTotal, Info.NumCDOCullDistance, ((float)Info.NumCDOCullDistance/(float)Info.NumTotal)*100.f);
-
-		Info.DivergentCullDistances.KeySort([](const float& lhs, const float& rhs){ return lhs >= rhs; });
+		UE_LOG(LogIrisBridge, Display, TEXT("MostCommon NetCullDistance: %f | Class: %s | Instances: %u"), Info.FindMostUsedCullDistance(), *Info.Class->GetName(), Info.NumTotal);
 
-		for (auto DivergentIt = Info.DivergentCullDistances.CreateConstIterator(); DivergentIt; ++DivergentIt)
+		Info.UniqueCullDistances.KeySort([](const float& lhs, const float& rhs){ return lhs >= rhs; });
+		for (auto DivergentIt = Info.UniqueCullDistances.CreateConstIterator(); DivergentIt; ++DivergentIt)
 		{
-			UE_LOG(LogIrisBridge, Display, TEXT("\tNetCullDistance: %f | UseCount: %d (%.2f%%)"), DivergentIt.Key(), DivergentIt.Value(), ((float)DivergentIt.Value()/(float)Info.NumTotal)*100.f);
+			UE_LOG(LogIrisBridge, Display, TEXT("\tNetCullDistance: %f | UseCount: %d/%d (%.2f%%)"), DivergentIt.Key(), DivergentIt.Value(), Info.NumTotal,((float)DivergentIt.Value()/(float)Info.NumTotal)*100.f);
 		}
 	}
 	
diff --git a/Engine/Source/Runtime/Experimental/Iris/Core/Private/Iris/ReplicationSystem/ReplicationSystem.cpp b/Engine/Source/Runtime/Experimental/Iris/Core/Private/Iris/ReplicationSystem/ReplicationSystem.cpp
--- a/Engine/Source/Runtime/Experimental/Iris/Core/Private/Iris/ReplicationSystem/ReplicationSystem.cpp
+++ b/Engine/Source/Runtime/Experimental/Iris/Core/Private/Iris/ReplicationSystem/ReplicationSystem.cpp
@@ -1626,6 +1626,8 @@ void UReplicationSystem::SetCullDistanceSqrOverride(FNetRefHandle Handle, float
 		return;
 	}
 
+	UE_LOG(LogIris, Verbose, TEXT("UReplicationSystem::SetCullDistanceSqrOverride: %s is now overridden to %f"), *Impl->ReplicationSystemInternal.GetNetRefHandleManager().PrintObjectFromNetRefHandle(Handle), FMath::Sqrt(DistSqr));
+	
 	return Impl->ReplicationSystemInternal.GetNetCullDistanceOverrides().SetCullDistanceSqr(ObjectInternalIndex, DistSqr);
 }
 
@@ -1637,7 +1639,8 @@ void UReplicationSystem::ClearCullDistanceSqrOverride(FNetRefHandle Handle)
 		return;
 	}
 
-	Impl->ReplicationSystemInternal.GetNetCullDistanceOverrides().ClearCullDistanceSqr(ObjectInternalIndex);
+	const bool bWasCullDistanceOverriden = Impl->ReplicationSystemInternal.GetNetCullDistanceOverrides().ClearCullDistanceSqr(ObjectInternalIndex);
+	UE_CLOG(bWasCullDistanceOverriden, LogIris, Verbose, TEXT("UReplicationSystem::ClearCullDistanceSqrOverride: %s is no longer overridden."), *Impl->ReplicationSystemInternal.GetNetRefHandleManager().PrintObjectFromNetRefHandle(Handle));
 }
 
 float UReplicationSystem::GetCullDistanceSqrOverride(FNetRefHandle Handle, float DefaultValue) const
diff --git a/Engine/Source/Runtime/Experimental/Iris/Core/Public/Iris/ReplicationSystem/NetCullDistanceOverrides.h b/Engine/Source/Runtime/Experimental/Iris/Core/Public/Iris/ReplicationSystem/NetCullDistanceOverrides.h
--- a/Engine/Source/Runtime/Experimental/Iris/Core/Public/Iris/ReplicationSystem/NetCullDistanceOverrides.h
+++ b/Engine/Source/Runtime/Experimental/Iris/Core/Public/Iris/ReplicationSystem/NetCullDistanceOverrides.h
@@ -30,7 +30,7 @@ public:
 
 public:
 	/** Remove cull distance override for object. */
-	void ClearCullDistanceSqr(uint32 ObjectIndex);
+	bool ClearCullDistanceSqr(uint32 ObjectIndex);
 
 	/** Set cull distance override for object. */
 	void SetCullDistanceSqr(uint32 ObjectIndex, float CullDistSqr);
