Chaos - new debug draw backend to handle debugdraw from async phsyics

Work in progress. Debug draw from the physics thread now works correctlyl Still need to add debug draw context to ParallelFor and Tasks. Also RBANs need their own timeline and begin/end frame sequence. If the debug draw context is not set up on a thread it falls back to the original global singleton mechanism.

diff --git a/Engine/Source/Runtime/Experimental/Chaos/Private/Chaos/DebugDrawQueue.cpp b/Engine/Source/Runtime/Experimental/Chaos/Private/Chaos/DebugDrawQueue.cpp
--- a/Engine/Source/Runtime/Experimental/Chaos/Private/Chaos/DebugDrawQueue.cpp
+++ b/Engine/Source/Runtime/Experimental/Chaos/Private/Chaos/DebugDrawQueue.cpp
@@ -7,6 +7,9 @@
 
 namespace Chaos
 {
+	bool bChaosDebugDraw_UseLegacyQueue = false;
+	FAutoConsoleVariableRef CVarChaos_DebugDraw_UseLegacyQueue(TEXT("p.Chaos.DebugDraw.UseLegacyQueue"), bChaosDebugDraw_UseLegacyQueue, TEXT(""));
+
 	void FDebugDrawQueue::SetConsumerActive(void* Consumer, bool bConsumerActive)
 	{
 		FScopeLock Lock(&ConsumersCS);
diff --git a/Engine/Source/Runtime/Experimental/Chaos/Private/Chaos/PBDRigidsEvolutionGBF.cpp b/Engine/Source/Runtime/Experimental/Chaos/Private/Chaos/PBDRigidsEvolutionGBF.cpp
--- a/Engine/Source/Runtime/Experimental/Chaos/Private/Chaos/PBDRigidsEvolutionGBF.cpp
+++ b/Engine/Source/Runtime/Experimental/Chaos/Private/Chaos/PBDRigidsEvolutionGBF.cpp
@@ -27,6 +27,10 @@
 
 #include "ChaosVisualDebugger/ChaosVisualDebuggerTrace.h"
 
+#include "ChaosDebugDraw/ChaosDDContext.h"
+#include "ChaosDebugDraw/ChaosDDScene.h"
+#include "ChaosDebugDraw/ChaosDDTimeline.h"
+
 //UE_DISABLE_OPTIMIZATION
 
 namespace Chaos
@@ -358,12 +362,39 @@ void FPBDRigidsEvolutionGBF::Advance(const FReal Dt,const FReal MaxStepDt,const
 
 void FPBDRigidsEvolutionGBF::AdvanceOneTimeStep(const FReal Dt,const FSubStepInfo& SubStepInfo)
 {	
+#if CHAOS_DEBUG_DRAW
+	// @todo(chaos): need the time when we want to show historical debug draw
+	ChaosDD::Private::FChaosDDContext::BeginFrame(CDDTickTimeline, 0.0f, Dt);
+#endif
+
 	PrepareTick();
 
 	AdvanceOneTimeStepImpl(Dt, SubStepInfo);
 
 	UnprepareTick();
+
+#if CHAOS_DEBUG_DRAW
+	ChaosDD::Private::FChaosDDContext::EndFrame();
+#endif
+}
+
+#if CHAOS_DEBUG_DRAW
+void FPBDRigidsEvolutionGBF::SetDebugDrawScene(const ChaosDD::Private::FChaosDDScenePtr& InCDDScene)
+{
+	if (CDDScene.IsValid() && CDDTickTimeline.IsValid())
+	{
+		CDDScene->ReleaseTimeline(CDDTickTimeline);
+	}
+
+	CDDScene = InCDDScene;
+	CDDTickTimeline.Reset();
+
+	if (CDDScene.IsValid())
+	{
+		CDDTickTimeline = CDDScene->CreateTimeline(FString::Format(TEXT("{0} {1}"), { CDDScene->GetName(), "Tick" }));
+	}
 }
+#endif
 
 void FPBDRigidsEvolutionGBF::ReloadParticlesCache()
 {
diff --git a/Engine/Source/Runtime/Experimental/Chaos/Private/ChaosDebugDraw/ChaosDDContext.cpp b/Engine/Source/Runtime/Experimental/Chaos/Private/ChaosDebugDraw/ChaosDDContext.cpp
new file mode 100644
--- /dev/null
+++ b/Engine/Source/Runtime/Experimental/Chaos/Private/ChaosDebugDraw/ChaosDDContext.cpp
@@ -0,0 +1,63 @@
+// Copyright Epic Games, Inc. All Rights Reserved.
+#include "ChaosDebugDraw/ChaosDDContext.h"
+#include "ChaosDebugDraw/ChaosDDLog.h"
+#include "ChaosDebugDraw/ChaosDDFrame.h"
+#include "ChaosDebugDraw/ChaosDDTimeline.h"
+
+#include "ChaosLog.h"
+
+DEFINE_LOG_CATEGORY(LogChaosDD);
+
+#if CHAOS_DEBUG_DRAW
+
+namespace ChaosDD::Private
+{
+	FChaosDDContext::FChaosDDContext()
+	{
+		TimelineStack.Push(FChaosDDTimelinePtr());
+	}
+
+	void FChaosDDContext::BeginFrame(const FChaosDDTimelinePtr& InTimeline, double Time, double Dt)
+	{
+		FChaosDDContext& Context = Get();
+
+		if (InTimeline.IsValid())
+		{
+			Context.TimelineStack.Push(InTimeline);
+
+			Context.GetTimeline()->BeginFrame(Time, Dt);
+		}
+
+		Context.UpdateCachedFramePtr();
+
+	}
+
+	void FChaosDDContext::EndFrame()
+	{
+		FChaosDDContext& Context = Get();
+
+		if (ensure(Context.GetTimeline().IsValid()))
+		{
+			Context.GetTimeline()->EndFrame();
+
+			Context.TimelineStack.Pop();
+		}
+
+		Context.UpdateCachedFramePtr();
+	}
+
+	void FChaosDDContext::UpdateCachedFramePtr()
+	{
+		if (GetTimeline().IsValid())
+		{
+			Frame = GetTimeline()->GetActiveFrame();
+		}
+		else
+		{
+			Frame.Reset();
+		}
+	}
+
+}
+
+#endif
\ No newline at end of file
diff --git a/Engine/Source/Runtime/Experimental/Chaos/Private/ChaosDebugDraw/ChaosDDScene.cpp b/Engine/Source/Runtime/Experimental/Chaos/Private/ChaosDebugDraw/ChaosDDScene.cpp
new file mode 100644
--- /dev/null
+++ b/Engine/Source/Runtime/Experimental/Chaos/Private/ChaosDebugDraw/ChaosDDScene.cpp
@@ -0,0 +1,100 @@
+// Copyright Epic Games, Inc. All Rights Reserved.
+#include "ChaosDebugDraw/ChaosDDScene.h"
+#include "ChaosDebugDraw/ChaosDDFrame.h"
+#include "ChaosDebugDraw/ChaosDDLog.h"
+#include "ChaosDebugDraw/ChaosDDRenderer.h"
+#include "ChaosDebugDraw/ChaosDDTimeline.h"
+
+#if CHAOS_DEBUG_DRAW
+
+namespace ChaosDD::Private
+{
+	FChaosDDScene::FChaosDDScene(const FString& InName)
+		: Name(InName)
+		, DrawRegion(FVector::Zero(), 0.0)
+		, CommandBudget(20000)
+	{
+	}
+
+	FChaosDDScene::~FChaosDDScene()
+	{
+	}
+
+	void FChaosDDScene::SetDrawRegion(const FSphere3d& InDrawRegion)
+	{
+		FScopeLock Lock(&TimelinesCS);
+
+		DrawRegion = InDrawRegion;
+
+		for (const FChaosDDTimelinePtr& Timeline : Timelines)
+		{
+			Timeline->SetDrawRegion(InDrawRegion);
+		}
+	}
+
+	void FChaosDDScene::SetCommandBudget(int32 InCommandBudget)
+	{
+		FScopeLock Lock(&TimelinesCS);
+
+		CommandBudget = InCommandBudget;
+
+		for (const FChaosDDTimelinePtr& Timeline : Timelines)
+		{
+			Timeline->SetCommandBudget(InCommandBudget);
+		}
+	}
+
+	FChaosDDTimelinePtr FChaosDDScene::CreateTimeline(const FString& InName)
+	{
+		FScopeLock Lock(&TimelinesCS);
+
+		FChaosDDTimelinePtr Timeline = MakeShared<FChaosDDTimeline>(InName, CommandBudget);
+
+		Timelines.Add(Timeline);
+
+		return Timeline;
+	}
+
+	void FChaosDDScene::ReleaseTimeline(const FChaosDDTimelinePtr& InTimeline)
+	{
+		FScopeLock Lock(&TimelinesCS);
+
+		for (int32 TimelineIndex = 0; TimelineIndex < Timelines.Num(); ++TimelineIndex)
+		{
+			if (Timelines[TimelineIndex].Get() == InTimeline.Get())
+			{
+				Timelines.RemoveAt(TimelineIndex, 1, EAllowShrinking::No);
+			}
+		}
+	}
+
+	void FChaosDDScene::RenderLatestFrames(IChaosDDRenderer& Renderer)
+	{
+		TArray<FChaosDDFramePtr> Frames = GetFrames();
+		for (const FChaosDDFramePtr& Frame : Frames)
+		{
+			UE_LOG(LogChaosDD, VeryVerbose, TEXT("Render %s %d %d Commands"), *Frame->GetTimeline()->GetName(), Frame->GetFrameIndex(), Frame->GetNumCommands());
+
+			// Render the legacy commands
+			Frame->VisitCommands(
+				[&Renderer](const Chaos::FLatentDrawCommand& Command)
+				{
+					Renderer.DrawCommand(Command);
+				});
+		}
+	}
+
+	TArray<FChaosDDFramePtr> FChaosDDScene::GetFrames()
+	{
+		FScopeLock Lock(&TimelinesCS);
+
+		TArray<FChaosDDFramePtr> Frames;
+		for (const FChaosDDTimelinePtr& Timeline : Timelines)
+		{
+			Timeline->GetFrames(Frames);
+		}
+		return Frames;
+	}
+}
+
+#endif
diff --git a/Engine/Source/Runtime/Experimental/Chaos/Private/PBDRigidsSolver.cpp b/Engine/Source/Runtime/Experimental/Chaos/Private/PBDRigidsSolver.cpp
--- a/Engine/Source/Runtime/Experimental/Chaos/Private/PBDRigidsSolver.cpp
+++ b/Engine/Source/Runtime/Experimental/Chaos/Private/PBDRigidsSolver.cpp
@@ -29,6 +29,10 @@
 #include "Chaos/PhysicsSolverBaseImpl.h"
 #include "Chaos/ConvexOptimizer.h"
 
+#include "ChaosDebugDraw/ChaosDDContext.h"
+#include "ChaosDebugDraw/ChaosDDScene.h"
+#include "ChaosDebugDraw/ChaosDDTimeline.h"
+
 #include "ProfilingDebugging/CountersTrace.h"
 #include "ProfilingDebugging/CsvProfiler.h"
 #include "ChaosLog.h"
@@ -1297,6 +1301,10 @@ namespace Chaos
 	{
 		const FReal StartSimTime = GetSolverTime();
 
+#if CHAOS_DEBUG_DRAW
+		ChaosDD::Private::FChaosDDContext::BeginFrame(CDDFrameTimeline, StartSimTime, GetLastDt());
+#endif
+
 #if !(UE_BUILD_SHIPPING || UE_BUILD_TEST)
 		if(IsNetworkPhysicsPredictionEnabled() && CanDebugNetworkPhysicsPrediction())
 		{
@@ -1320,8 +1328,32 @@ namespace Chaos
 			//final step so we can destroy proxies
 			DestroyPendingProxies_Internal();
 		}
+
+#if CHAOS_DEBUG_DRAW
+		ChaosDD::Private::FChaosDDContext::EndFrame();
+#endif
 	}
 
+#if CHAOS_DEBUG_DRAW
+	void FPBDRigidsSolver::SetDebugDrawScene(const ChaosDD::Private::FChaosDDScenePtr& InCDDScene)
+	{
+		if (CDDScene.IsValid() && CDDFrameTimeline.IsValid())
+		{
+			CDDScene->ReleaseTimeline(CDDFrameTimeline);
+		}
+
+		CDDScene = InCDDScene;
+		CDDFrameTimeline.Reset();
+
+		if (CDDScene.IsValid())
+		{
+			CDDFrameTimeline = CDDScene->CreateTimeline(FString::Format(TEXT("{0} {1}"), { CDDScene->GetName(), "Frame" }));
+		}
+
+		GetEvolution()->SetDebugDrawScene(InCDDScene);
+	}
+#endif
+
 	void FPBDRigidsSolver::SetExternalTimestampConsumed_Internal(const int32 Timestamp)
 	{
 		MEvolution->LatestExternalTimestampConsumed_Internal = Timestamp;
diff --git a/Engine/Source/Runtime/Experimental/Chaos/Public/Chaos/DebugDrawCommand.h b/Engine/Source/Runtime/Experimental/Chaos/Public/Chaos/DebugDrawCommand.h
new file mode 100644
--- /dev/null
+++ b/Engine/Source/Runtime/Experimental/Chaos/Public/Chaos/DebugDrawCommand.h
@@ -0,0 +1,184 @@
+// Copyright Epic Games, Inc. All Rights Reserved.
+#pragma once
+
+#include "Chaos/ChaosDebugDrawDeclares.h"
+#include "Containers/List.h"
+#include "Math/Vector.h"
+#include "Math/Color.h"
+#include "Math/Quat.h"
+
+class AActor;
+
+#if CHAOS_DEBUG_DRAW
+
+namespace Chaos
+{
+	struct FLatentDrawCommand
+	{
+		FVector LineStart;
+		FVector LineEnd;
+		FColor Color;
+		int32 Segments;
+		bool bPersistentLines;
+		float ArrowSize;
+		float LifeTime;
+		uint8 DepthPriority;
+		float Thickness;
+		FReal Radius;
+		FReal HalfHeight;
+		FVector Center;
+		FVector Extent;
+		FQuat Rotation;
+		FVector TextLocation;
+		FString Text;
+		class AActor* TestBaseActor;
+		bool bDrawShadow;
+		float FontScale;
+		float Duration;
+		FMatrix TransformMatrix;
+		bool bDrawAxis;
+		FVector YAxis;
+		FVector ZAxis;
+
+		enum class EDrawType
+		{
+			Point,
+			Line,
+			DirectionalArrow,
+			Sphere,
+			Box,
+			String,
+			Circle,
+			Capsule,
+		} Type;
+
+		FLatentDrawCommand()
+			: TestBaseActor(nullptr)
+		{
+		}
+
+		static FLatentDrawCommand DrawPoint(const FVector& Position, const FColor& Color, bool bPersistentLines, float LifeTime, uint8 DepthPriority, float Thickness)
+		{
+			FLatentDrawCommand Command;
+			Command.LineStart = Position;
+			Command.Color = Color;
+			Command.bPersistentLines = bPersistentLines;
+			Command.LifeTime = LifeTime;
+			Command.DepthPriority = DepthPriority;
+			Command.Thickness = Thickness;
+			Command.Type = EDrawType::Point;
+			return Command;
+		}
+
+
+		static FLatentDrawCommand DrawLine(const FVector& LineStart, const FVector& LineEnd, const FColor& Color, bool bPersistentLines, float LifeTime, uint8 DepthPriority, float Thickness)
+		{
+			FLatentDrawCommand Command;
+			Command.LineStart = LineStart;
+			Command.LineEnd = LineEnd;
+			Command.Color = Color;
+			Command.bPersistentLines = bPersistentLines;
+			Command.LifeTime = LifeTime;
+			Command.DepthPriority = DepthPriority;
+			Command.Thickness = Thickness;
+			Command.Type = EDrawType::Line;
+			return Command;
+		}
+
+		static FLatentDrawCommand DrawDirectionalArrow(const FVector& LineStart, FVector const& LineEnd, float ArrowSize, const FColor& Color, bool bPersistentLines, float LifeTime, uint8 DepthPriority, float Thickness)
+		{
+			FLatentDrawCommand Command;
+			Command.LineStart = LineStart;
+			Command.LineEnd = LineEnd;
+			Command.ArrowSize = ArrowSize;
+			Command.Color = Color;
+			Command.bPersistentLines = bPersistentLines;
+			Command.LifeTime = LifeTime;
+			Command.DepthPriority = DepthPriority;
+			Command.Thickness = Thickness;
+			Command.Type = EDrawType::DirectionalArrow;
+			return Command;
+		}
+
+		static FLatentDrawCommand DrawDebugSphere(const FVector& Center, FVector::FReal Radius, int32 Segments, const FColor& Color, bool bPersistentLines, float LifeTime, uint8 DepthPriority, float Thickness)
+		{
+			FLatentDrawCommand Command;
+			Command.LineStart = Center;
+			Command.Radius = Radius;
+			Command.Color = Color;
+			Command.Segments = Segments;
+			Command.bPersistentLines = bPersistentLines;
+			Command.LifeTime = LifeTime;
+			Command.DepthPriority = DepthPriority;
+			Command.Thickness = Thickness;
+			Command.Type = EDrawType::Sphere;
+			return Command;
+		}
+
+		static FLatentDrawCommand DrawDebugBox(const FVector& Center, const FVector& Extent, const FQuat& Rotation, const FColor& Color, bool bPersistentLines, float LifeTime, uint8 DepthPriority, float Thickness)
+		{
+			FLatentDrawCommand Command;
+			Command.Center = Center;
+			Command.Extent = Extent;
+			Command.Rotation = Rotation;
+			Command.Color = Color;
+			Command.bPersistentLines = bPersistentLines;
+			Command.LifeTime = LifeTime;
+			Command.DepthPriority = DepthPriority;
+			Command.Thickness = Thickness;
+			Command.Type = EDrawType::Box;
+			return Command;
+		}
+
+		static FLatentDrawCommand DrawDebugString(const FVector& TextLocation, const FString& Text, class AActor* TestBaseActor, const FColor& Color, float Duration, bool bDrawShadow, float FontScale)
+		{
+			FLatentDrawCommand Command;
+			Command.TextLocation = TextLocation;
+			Command.Text = Text;
+			Command.TestBaseActor = TestBaseActor;
+			Command.Color = Color;
+			Command.Duration = Duration;
+			Command.LifeTime = Duration;
+			Command.bDrawShadow = bDrawShadow;
+			Command.FontScale = FontScale;
+			Command.Type = EDrawType::String;
+			return Command;
+		}
+
+		static FLatentDrawCommand DrawDebugCircle(const FVector& Center, FReal Radius, int32 Segments, const FColor& Color, bool bPersistentLines, float LifeTime, uint8 DepthPriority, float Thickness, const FVector& YAxis, const FVector& ZAxis, bool bDrawAxis)
+		{
+			FLatentDrawCommand Command;
+			Command.Center = Center;
+			Command.Radius = Radius;
+			Command.Segments = Segments;
+			Command.Color = Color;
+			Command.bPersistentLines = bPersistentLines;
+			Command.LifeTime = LifeTime;
+			Command.DepthPriority = DepthPriority;
+			Command.Thickness = Thickness;
+			Command.YAxis = YAxis;
+			Command.ZAxis = ZAxis;
+			Command.bDrawAxis = bDrawAxis;
+			Command.Type = EDrawType::Circle;
+			return Command;
+		}
+
+		static FLatentDrawCommand DrawDebugCapsule(const FVector& Center, FReal HalfHeight, FReal Radius, const FQuat& Rotation, const FColor& Color, bool bPersistentLines, float LifeTime, uint8 DepthPriority, float Thickness)
+		{
+			FLatentDrawCommand Command;
+			Command.Center = Center;
+			Command.HalfHeight = HalfHeight;
+			Command.Radius = Radius;
+			Command.Rotation = Rotation;
+			Command.Color = Color;
+			Command.bPersistentLines = bPersistentLines;
+			Command.LifeTime = LifeTime;
+			Command.DepthPriority = DepthPriority;
+			Command.Thickness = Thickness;
+			Command.Type = EDrawType::Capsule;
+			return Command;
+		}
+	};
+}
+
+#endif
\ No newline at end of file
diff --git a/Engine/Source/Runtime/Experimental/Chaos/Public/Chaos/DebugDrawQueue.h b/Engine/Source/Runtime/Experimental/Chaos/Public/Chaos/DebugDrawQueue.h
--- a/Engine/Source/Runtime/Experimental/Chaos/Public/Chaos/DebugDrawQueue.h
+++ b/Engine/Source/Runtime/Experimental/Chaos/Public/Chaos/DebugDrawQueue.h
@@ -3,6 +3,7 @@
 
 #include "Chaos/ChaosDebugDrawDeclares.h"
 #include "Chaos/AABB.h"
+#include "Chaos/DebugDrawCommand.h"
 #include "Containers/List.h"
 #include "Math/Vector.h"
 #include "Math/Color.h"
@@ -11,11 +12,17 @@
 #include "HAL/ThreadSafeBool.h"
 #include "HAL/IConsoleManager.h"
 
+#include "ChaosDebugDraw/ChaosDDContext.h"
+#include "ChaosDebugDraw/ChaosDDFrame.h"
+
 class AActor;
 
+
 #if CHAOS_DEBUG_DRAW
 namespace Chaos
 {
+extern CHAOS_API bool bChaosDebugDraw_UseLegacyQueue;
+
 /** Thread-safe single-linked list (lock-free). (Taken from Light Mass, should probably just move into core) */
 template<typename ElementType>
 class TListThreadSafe
@@ -82,174 +89,6 @@ private:
 	TList<ElementType>* FirstElement;
 };
 
-struct FLatentDrawCommand
-{
-	FVector LineStart;
-	FVector LineEnd;
-	FColor Color;
-	int32 Segments;
-	bool bPersistentLines;
-	float ArrowSize;
-	float LifeTime;
-	uint8 DepthPriority;
-	float Thickness;
-	FReal Radius;
-	FReal HalfHeight;
-	FVector Center;
-	FVector Extent;
-	FQuat Rotation;
-	FVector TextLocation;
-	FString Text;
-	class AActor* TestBaseActor;
-	bool bDrawShadow;
-	float FontScale;
-	float Duration;
-	FMatrix TransformMatrix;
-	bool bDrawAxis;
-	FVector YAxis;
-	FVector ZAxis;
-
-	enum class EDrawType
-	{
-		Point,
-		Line,
-		DirectionalArrow,
-		Sphere,
-		Box,
-		String,
-		Circle,
-		Capsule,
-	} Type;
-
-	FLatentDrawCommand()
-		: TestBaseActor(nullptr)
-	{
-	}
-
-	static FLatentDrawCommand DrawPoint(const FVector& Position, const FColor& Color, bool bPersistentLines, float LifeTime, uint8 DepthPriority, float Thickness)
-	{
-		FLatentDrawCommand Command;
-		Command.LineStart = Position;
-		Command.Color = Color;
-		Command.bPersistentLines = bPersistentLines;
-		Command.LifeTime = LifeTime;
-		Command.DepthPriority = DepthPriority;
-		Command.Thickness = Thickness;
-		Command.Type = EDrawType::Point;
-		return Command;
-	}
-
-
-	static FLatentDrawCommand DrawLine(const FVector& LineStart, const FVector& LineEnd, const FColor& Color, bool bPersistentLines, float LifeTime, uint8 DepthPriority, float Thickness)
-	{
-		FLatentDrawCommand Command;
-		Command.LineStart = LineStart;
-		Command.LineEnd = LineEnd;
-		Command.Color = Color;
-		Command.bPersistentLines = bPersistentLines;
-		Command.LifeTime = LifeTime;
-		Command.DepthPriority = DepthPriority;
-		Command.Thickness = Thickness;
-		Command.Type = EDrawType::Line;
-		return Command;
-	}
-
-	static FLatentDrawCommand DrawDirectionalArrow(const FVector& LineStart, FVector const& LineEnd, float ArrowSize, const FColor& Color, bool bPersistentLines, float LifeTime, uint8 DepthPriority, float Thickness)
-	{
-		FLatentDrawCommand Command;
-		Command.LineStart = LineStart;
-		Command.LineEnd = LineEnd;
-		Command.ArrowSize = ArrowSize;
-		Command.Color = Color;
-		Command.bPersistentLines = bPersistentLines;
-		Command.LifeTime = LifeTime;
-		Command.DepthPriority = DepthPriority;
-		Command.Thickness = Thickness;
-		Command.Type = EDrawType::DirectionalArrow;
-		return Command;
-	}
-
-	static FLatentDrawCommand DrawDebugSphere(const FVector& Center, FVector::FReal Radius, int32 Segments, const FColor& Color, bool bPersistentLines, float LifeTime, uint8 DepthPriority, float Thickness)
-	{
-		FLatentDrawCommand Command;
-		Command.LineStart = Center;
-		Command.Radius = Radius;
-		Command.Color = Color;
-		Command.Segments = Segments;
-		Command.bPersistentLines = bPersistentLines;
-		Command.LifeTime = LifeTime;
-		Command.DepthPriority = DepthPriority;
-		Command.Thickness = Thickness;
-		Command.Type = EDrawType::Sphere;
-		return Command;
-	}
-
-	static FLatentDrawCommand DrawDebugBox(const FVector& Center, const FVector& Extent, const FQuat& Rotation, const FColor& Color, bool bPersistentLines, float LifeTime, uint8 DepthPriority, float Thickness)
-	{
-		FLatentDrawCommand Command;
-		Command.Center = Center;
-		Command.Extent = Extent;
-		Command.Rotation = Rotation;
-		Command.Color = Color;
-		Command.bPersistentLines = bPersistentLines;
-		Command.LifeTime = LifeTime;
-		Command.DepthPriority = DepthPriority;
-		Command.Thickness = Thickness;
-		Command.Type = EDrawType::Box;
-		return Command;
-	}
-
-	static FLatentDrawCommand DrawDebugString(const FVector& TextLocation, const FString& Text, class AActor* TestBaseActor, const FColor& Color, float Duration, bool bDrawShadow, float FontScale)
-	{
-		FLatentDrawCommand Command;
-		Command.TextLocation = TextLocation;
-		Command.Text = Text;
-		Command.TestBaseActor = TestBaseActor;
-		Command.Color = Color;
-		Command.Duration = Duration;
-		Command.LifeTime = Duration;
-		Command.bDrawShadow = bDrawShadow;
-		Command.FontScale = FontScale;
-		Command.Type = EDrawType::String;
-		return Command;
-	}
-
-	static FLatentDrawCommand DrawDebugCircle(const FVector& Center, FReal Radius, int32 Segments, const FColor& Color, bool bPersistentLines, float LifeTime, uint8 DepthPriority, float Thickness, const FVector& YAxis, const FVector& ZAxis, bool bDrawAxis)
-	{
-		FLatentDrawCommand Command;
-		Command.Center = Center;
-		Command.Radius = Radius;
-		Command.Segments = Segments;
-		Command.Color = Color;
-		Command.bPersistentLines = bPersistentLines;
-		Command.LifeTime = LifeTime;
-		Command.DepthPriority = DepthPriority;
-		Command.Thickness = Thickness;
-		Command.YAxis = YAxis;
-		Command.ZAxis = ZAxis;
-		Command.bDrawAxis = bDrawAxis;
-		Command.Type = EDrawType::Circle;
-		return Command;
-	}
-
-	static FLatentDrawCommand DrawDebugCapsule(const FVector& Center, FReal HalfHeight, FReal Radius, const FQuat& Rotation, const FColor& Color, bool bPersistentLines, float LifeTime, uint8 DepthPriority, float Thickness)
-	{
-		FLatentDrawCommand Command;
-		Command.Center = Center;
-		Command.HalfHeight = HalfHeight;
-		Command.Radius = Radius;
-		Command.Rotation = Rotation;
-		Command.Color = Color;
-		Command.bPersistentLines = bPersistentLines;
-		Command.LifeTime = LifeTime;
-		Command.DepthPriority = DepthPriority;
-		Command.Thickness = Thickness;
-		Command.Type = EDrawType::Capsule;
-		return Command;
-	}
-
-};
-
 /** A thread safe way to generate latent debug drawing. (This is picked up later by the geometry collection component which is a total hack for now, but needed to get into an engine world ) */
 class FDebugDrawQueue
 {
@@ -278,7 +117,7 @@ public:
 			FScopeLock Lock(&CommandQueueCS);
 			if (AcceptCommand(1, Position))
 			{
-				CommandQueue.Emplace(FLatentDrawCommand::DrawPoint(Position, Color, bPersistentLines, LifeTime, DepthPriority, Thickness));
+				AddCommand(FLatentDrawCommand::DrawPoint(Position, Color, bPersistentLines, LifeTime, DepthPriority, Thickness));
 			}
 		}
 	}
@@ -290,7 +129,7 @@ public:
 			FScopeLock Lock(&CommandQueueCS);
 			if (AcceptCommand(1, LineStart))
 			{
-				CommandQueue.Emplace(FLatentDrawCommand::DrawLine(LineStart, LineEnd, Color, bPersistentLines, LifeTime, DepthPriority, Thickness));
+				AddCommand(FLatentDrawCommand::DrawLine(LineStart, LineEnd, Color, bPersistentLines, LifeTime, DepthPriority, Thickness));
 			}
 		}
 	}
@@ -302,7 +141,7 @@ public:
 			FScopeLock Lock(&CommandQueueCS);
 			if (AcceptCommand(3, LineStart))
 			{
-				CommandQueue.Emplace(FLatentDrawCommand::DrawDirectionalArrow(LineStart, LineEnd, ArrowSize, Color, bPersistentLines, LifeTime, DepthPriority, Thickness));
+				AddCommand(FLatentDrawCommand::DrawDirectionalArrow(LineStart, LineEnd, ArrowSize, Color, bPersistentLines, LifeTime, DepthPriority, Thickness));
 			}
 		}
 	}
@@ -320,9 +159,9 @@ public:
 
 			if (AcceptCommand(3, Position))
 			{
-				CommandQueue.Emplace(FLatentDrawCommand::DrawLine(Position, Position + X * Scale, FColor::Red, bPersistentLines, LifeTime, DepthPriority, Thickness));
-				CommandQueue.Emplace(FLatentDrawCommand::DrawLine(Position, Position + Y * Scale, FColor::Green, bPersistentLines, LifeTime, DepthPriority, Thickness));
-				CommandQueue.Emplace(FLatentDrawCommand::DrawLine(Position, Position + Z * Scale, FColor::Blue, bPersistentLines, LifeTime, DepthPriority, Thickness));
+				AddCommand(FLatentDrawCommand::DrawLine(Position, Position + X * Scale, FColor::Red, bPersistentLines, LifeTime, DepthPriority, Thickness));
+				AddCommand(FLatentDrawCommand::DrawLine(Position, Position + Y * Scale, FColor::Green, bPersistentLines, LifeTime, DepthPriority, Thickness));
+				AddCommand(FLatentDrawCommand::DrawLine(Position, Position + Z * Scale, FColor::Blue, bPersistentLines, LifeTime, DepthPriority, Thickness));
 			}
 		}
 	}
@@ -336,7 +175,7 @@ public:
 			const int Cost = Segments * Segments;
 			if (AcceptCommand(Cost, Center))
 			{
-				CommandQueue.Emplace(FLatentDrawCommand::DrawDebugSphere(Center, Radius, Segments, Color, bPersistentLines, LifeTime, DepthPriority, Thickness));
+				AddCommand(FLatentDrawCommand::DrawDebugSphere(Center, Radius, Segments, Color, bPersistentLines, LifeTime, DepthPriority, Thickness));
 			}
 		}
 	}
@@ -348,7 +187,7 @@ public:
 			FScopeLock Lock(&CommandQueueCS);
 			if (AcceptCommand(12, Center))
 			{
-				CommandQueue.Emplace(FLatentDrawCommand::DrawDebugBox(Center, Extent, Rotation, Color, bPersistentLines, LifeTime, DepthPriority, Thickness));
+				AddCommand(FLatentDrawCommand::DrawDebugBox(Center, Extent, Rotation, Color, bPersistentLines, LifeTime, DepthPriority, Thickness));
 			}
 		}
 	}
@@ -361,7 +200,7 @@ public:
 			int Cost = Text.Len();
 			if (AcceptCommand(Cost, TextLocation))
 			{
-				CommandQueue.Emplace(FLatentDrawCommand::DrawDebugString(TextLocation, Text, TestBaseActor, Color, Duration, bDrawShadow, FontScale));
+				AddCommand(FLatentDrawCommand::DrawDebugString(TextLocation, Text, TestBaseActor, Color, Duration, bDrawShadow, FontScale));
 			}
 		}
 	}
@@ -373,7 +212,7 @@ public:
 			FScopeLock Lock(&CommandQueueCS);
 			if (AcceptCommand(Segments, Center))
 			{
-				CommandQueue.Emplace(FLatentDrawCommand::DrawDebugCircle(Center, Radius, Segments, Color, bPersistentLines, LifeTime, DepthPriority, Thickness, YAxis, ZAxis, bDrawAxis));
+				AddCommand(FLatentDrawCommand::DrawDebugCircle(Center, Radius, Segments, Color, bPersistentLines, LifeTime, DepthPriority, Thickness, YAxis, ZAxis, bDrawAxis));
 			}
 		}
 	}
@@ -385,7 +224,7 @@ public:
 			FScopeLock Lock(&CommandQueueCS);
 			if (AcceptCommand(16, Center))
 			{
-				CommandQueue.Emplace(FLatentDrawCommand::DrawDebugCapsule(Center, HalfHeight, Radius, Rotation, Color, bPersistentLines, LifeTime, DepthPriority, Thickness));
+				AddCommand(FLatentDrawCommand::DrawDebugCapsule(Center, HalfHeight, Radius, Rotation, Color, bPersistentLines, LifeTime, DepthPriority, Thickness));
 			}
 		}
 	}
@@ -456,6 +295,19 @@ private:
 
 	bool AcceptCommand(int Cost, const FVec3& Position)
 	{
+		if (!bChaosDebugDraw_UseLegacyQueue)
+		{
+			const ChaosDD::Private::FChaosDDFramePtr& DDFrame = ChaosDD::Private::FChaosDDContext::Get().GetFrame();
+			if (DDFrame.IsValid())
+			{
+				if (DDFrame->IsInDrawRegion(Position))
+				{
+					return DDFrame->AddToCost(Cost);
+				}
+				return false;
+			}
+		}
+
 		if (IsInRegionOfInterest(Position))
 		{
 			RequestedCommandCost += Cost;
@@ -464,6 +316,24 @@ private:
 		return false;
 	}
 
+	void AddCommand(const FLatentDrawCommand& Command)
+	{
+		// Try to add to the new debug draw system which queues commands per World and ticking thread
+		if (!bChaosDebugDraw_UseLegacyQueue)
+		{
+			const ChaosDD::Private::FChaosDDFramePtr& DDFrame = ChaosDD::Private::FChaosDDContext::Get().GetFrame();
+			if (DDFrame.IsValid())
+			{
+				DDFrame->EnqueueCommand(Command);
+				return;
+			}
+		}
+
+		// This thread has not been set up for the new debug draw system so default back to the old
+		// way that has issues with timing in async mode.
+		CommandQueue.Add(Command);
+	}
+
 	bool IsInBudget() const
 	{
 		return (MaxCommandCost <= 0) || (RequestedCommandCost <= MaxCommandCost);
diff --git a/Engine/Source/Runtime/Experimental/Chaos/Public/Chaos/Framework/PhysicsSolverBase.h b/Engine/Source/Runtime/Experimental/Chaos/Public/Chaos/Framework/PhysicsSolverBase.h
--- a/Engine/Source/Runtime/Experimental/Chaos/Public/Chaos/Framework/PhysicsSolverBase.h
+++ b/Engine/Source/Runtime/Experimental/Chaos/Public/Chaos/Framework/PhysicsSolverBase.h
@@ -18,6 +18,8 @@
 #include "ChaosVisualDebugger/ChaosVDContextProvider.h"
 #endif
 
+#include "ChaosDebugDraw/ChaosDDTypes.h"
+
 class FChaosSolversModule;
 class FPhysicsReplicationAsync;
 
@@ -895,6 +897,11 @@ namespace Chaos
 			return CVDContextData;
 		};
 #endif
+
+#if CHAOS_DEBUG_DRAW
+	public:
+		virtual void SetDebugDrawScene(const ChaosDD::Private::FChaosDDScenePtr& InCDDScene) = 0;
+#endif
 	};
 }
 
diff --git a/Engine/Source/Runtime/Experimental/Chaos/Public/Chaos/PBDRigidsEvolutionGBF.h b/Engine/Source/Runtime/Experimental/Chaos/Public/Chaos/PBDRigidsEvolutionGBF.h
--- a/Engine/Source/Runtime/Experimental/Chaos/Public/Chaos/PBDRigidsEvolutionGBF.h
+++ b/Engine/Source/Runtime/Experimental/Chaos/Public/Chaos/PBDRigidsEvolutionGBF.h
@@ -317,6 +317,15 @@ namespace Chaos
 #if CHAOS_EVOLUTION_COLLISION_TESTMODE
 		void TestModeResetCollisions();
 #endif
+
+#if CHAOS_DEBUG_DRAW
+	public:
+		void SetDebugDrawScene(const ChaosDD::Private::FChaosDDScenePtr& InCDDScene);
+
+	private:
+		ChaosDD::Private::FChaosDDScenePtr CDDScene;
+		ChaosDD::Private::FChaosDDTimelinePtr CDDTickTimeline;
+#endif
 	};
 
 }
diff --git a/Engine/Source/Runtime/Experimental/Chaos/Public/ChaosDebugDraw/ChaosDDContext.h b/Engine/Source/Runtime/Experimental/Chaos/Public/ChaosDebugDraw/ChaosDDContext.h
new file mode 100644
--- /dev/null
+++ b/Engine/Source/Runtime/Experimental/Chaos/Public/ChaosDebugDraw/ChaosDDContext.h
@@ -0,0 +1,48 @@
+// Copyright Epic Games, Inc. All Rights Reserved.
+#pragma once
+
+#include "ChaosDebugDraw/ChaosDDTypes.h"
+#include "HAL/ThreadSingleton.h"
+
+#if CHAOS_DEBUG_DRAW
+
+namespace ChaosDD::Private
+{
+	//
+	// A thread-local debug draw context used to access the queue to draw to
+	// for any thread on which debug draw has been set up.
+	//
+	class CHAOS_API FChaosDDContext : public TThreadSingleton<FChaosDDContext>
+	{
+	public:
+		FChaosDDContext();
+
+		// The frame we should be drawing into on this thread
+		const FChaosDDFramePtr& GetFrame() const
+		{ 
+			return Frame;
+		}
+
+		// The timeline that holds the frames from this thread
+		const FChaosDDTimelinePtr& GetTimeline() const
+		{
+			// NOTE: TimelineStack has null as its first element so that we can always return a shared pointer by ref
+			return TimelineStack.Top();
+		}
+
+		// Allocate a new frame for the timeline and set up the context
+		static void BeginFrame(const FChaosDDTimelinePtr& InTimeline, double Time, double Dt);
+
+		// Signal that the frame is complete and can be rendered
+		static void EndFrame();
+
+	private:
+		void UpdateCachedFramePtr();
+
+		// We need a stack of timelines to handle nested BeginFrame calls. E.g., for sub-steps in the solver.
+		TArray<FChaosDDTimelinePtr> TimelineStack;
+		FChaosDDFramePtr Frame;
+	};
+}
+
+#endif
\ No newline at end of file
diff --git a/Engine/Source/Runtime/Experimental/Chaos/Public/ChaosDebugDraw/ChaosDDFrame.h b/Engine/Source/Runtime/Experimental/Chaos/Public/ChaosDebugDraw/ChaosDDFrame.h
new file mode 100644
--- /dev/null
+++ b/Engine/Source/Runtime/Experimental/Chaos/Public/ChaosDebugDraw/ChaosDDFrame.h
@@ -0,0 +1,141 @@
+// Copyright Epic Games, Inc. All Rights Reserved.
+#pragma once
+
+#include "ChaosDebugDraw/ChaosDDTypes.h"
+#include "Chaos/DebugDrawCommand.h"
+#include "Containers/Map.h"
+#include "HAL/PlatformTLS.h"
+#include "HAL/CriticalSection.h"
+#include "Math/Box.h"
+#include "Math/Sphere.h"
+
+
+#if CHAOS_DEBUG_DRAW
+
+namespace ChaosDD::Private
+{
+	//
+	// A single frame of debug draw data
+	//
+	class CHAOS_API FChaosDDFrame : public TSharedFromThis<FChaosDDFrame>
+	{
+	public:
+		FChaosDDFrame(const FChaosDDTimelineWeakPtr& InTimeline, int64 InFrameIndex, double InTime, double InDt, const FSphere3d& InDrawRegion, int32 InCommandBudget, int32 InCommandQueueLength)
+			: Timeline(InTimeline)
+			, FrameIndex(InFrameIndex)
+			, Time(InTime)
+			, Dt(InDt)
+			, DrawRegion(InDrawRegion)
+			, CommandBudget(InCommandBudget)
+			, CommandCost(0)
+		{
+			if (InCommandQueueLength > 0)
+			{
+				Commands.Reserve(InCommandQueueLength);
+			}
+		}
+
+		~FChaosDDFrame()
+		{
+		}
+
+		FChaosDDTimelinePtr GetTimeline() const
+		{
+			return Timeline.Pin();
+		}
+
+		int64 GetFrameIndex() const
+		{
+			return FrameIndex;
+		}
+
+		double GetTime() const
+		{
+			return Time;
+		}
+
+		double GetDt() const
+		{
+			return Dt;
+		}
+
+		bool IsInDrawRegion(const FVector& InPos)
+		{
+			if (DrawRegion.W > 0.0)
+			{
+				return DrawRegion.IsInside(InPos);
+			}
+			return true;
+		}
+
+		bool IsInDrawRegion(const FSphere3d& InSphere)
+		{
+			if (DrawRegion.W > 0.0)
+			{
+				return DrawRegion.Intersects(InSphere);
+			}
+			return true;
+		}
+
+		bool IsInDrawRegion(const FBox3d& InBox)
+		{
+			if (DrawRegion.W > 0.0)
+			{
+				const FBox3d DrawRegionBox = FBox3d(DrawRegion.Center - FVector(DrawRegion.W), DrawRegion.Center + FVector(DrawRegion.W));
+				return DrawRegionBox.Intersect(InBox);
+			}
+			return true;
+		}
+
+		bool AddToCost(int32 InCost)
+		{
+			FScopeLock Lock(&CommandsCS);
+
+			CommandCost += InCost;
+
+			// A budget of zero means infinite
+			return ((CommandBudget == 0) || (CommandCost <= CommandBudget));
+		}
+
+		void EnqueueCommand(const Chaos::FLatentDrawCommand& InCommand)
+		{
+			FScopeLock Lock(&CommandsCS);
+
+			Commands.Add(InCommand);
+		}
+
+		int32 GetNumCommands() const
+		{
+			FScopeLock Lock(&CommandsCS);
+
+			return Commands.Num();
+		}
+
+		template<typename VisitorType>
+		void VisitCommands(const VisitorType& Visitor)
+		{
+			FScopeLock Lock(&CommandsCS);
+
+			for (const Chaos::FLatentDrawCommand& Command : Commands)
+			{
+				Visitor(Command);
+			}
+		}
+
+	private:
+		FChaosDDTimelineWeakPtr Timeline;
+		int64 FrameIndex;
+		double Time;
+		double Dt;
+		const FSphere3d& DrawRegion;
+		int32 CommandBudget;
+		int32 CommandCost;
+
+		// Legacy debug draw commands in a lock-protected queue
+		// @todo(chaos): we could move these to the per-thread buffer
+		TArray<Chaos::FLatentDrawCommand> Commands;
+		mutable FCriticalSection CommandsCS;
+	};
+}
+
+#endif
diff --git a/Engine/Source/Runtime/Experimental/Chaos/Public/ChaosDebugDraw/ChaosDDLog.h b/Engine/Source/Runtime/Experimental/Chaos/Public/ChaosDebugDraw/ChaosDDLog.h
new file mode 100644
--- /dev/null
+++ b/Engine/Source/Runtime/Experimental/Chaos/Public/ChaosDebugDraw/ChaosDDLog.h
@@ -0,0 +1,6 @@
+// Copyright Epic Games, Inc. All Rights Reserved.
+#pragma once
+
+#include "CoreMinimal.h"
+
+CHAOS_API DECLARE_LOG_CATEGORY_EXTERN(LogChaosDD, Log, All);
diff --git a/Engine/Source/Runtime/Experimental/Chaos/Public/ChaosDebugDraw/ChaosDDRenderer.h b/Engine/Source/Runtime/Experimental/Chaos/Public/ChaosDebugDraw/ChaosDDRenderer.h
new file mode 100644
--- /dev/null
+++ b/Engine/Source/Runtime/Experimental/Chaos/Public/ChaosDebugDraw/ChaosDDRenderer.h
@@ -0,0 +1,29 @@
+// Copyright Epic Games, Inc. All Rights Reserved.
+#pragma once
+
+#include "Chaos/Core.h"
+#include "ChaosDebugDraw/ChaosDDTypes.h"
+
+#if CHAOS_DEBUG_DRAW
+
+namespace Chaos
+{
+	struct FLatentDrawCommand;
+}
+
+namespace ChaosDD::Private
+{
+	//
+	// Primitive rendering API for use by debug draw objects
+	//
+	class CHAOS_API IChaosDDRenderer
+	{
+	public:
+		IChaosDDRenderer() {}
+		virtual ~IChaosDDRenderer() {}
+
+		virtual void DrawCommand(const Chaos::FLatentDrawCommand& Command) const = 0;
+	};
+}
+
+#endif
diff --git a/Engine/Source/Runtime/Experimental/Chaos/Public/ChaosDebugDraw/ChaosDDScene.h b/Engine/Source/Runtime/Experimental/Chaos/Public/ChaosDebugDraw/ChaosDDScene.h
new file mode 100644
--- /dev/null
+++ b/Engine/Source/Runtime/Experimental/Chaos/Public/ChaosDebugDraw/ChaosDDScene.h
@@ -0,0 +1,57 @@
+// Copyright Epic Games, Inc. All Rights Reserved.
+#pragma once
+
+#include "ChaosDebugDraw/ChaosDDTypes.h"
+#include "Containers/Array.h"
+#include "HAL/CriticalSection.h"
+#include "Math/Sphere.h"
+
+#if CHAOS_DEBUG_DRAW
+
+namespace ChaosDD::Private
+{
+	class IChaosDDRenderer;
+	class FChaosDDTimeline;
+
+	//
+	// Debug draw system for a world. In PIE there will be one of these for the server and each client.
+	//
+	// @todo(chaos): enable retention of debug draw frames and debug draw from a specific time
+	class CHAOS_API FChaosDDScene
+	{
+	public:
+		FChaosDDScene(const FString& InName);
+		~FChaosDDScene();
+
+		const FString& GetName() const
+		{
+			return Name;
+		}
+
+		// Specify the region of in which we wish to enable debug draw. A radius of zero means everywhere.
+		void SetDrawRegion(const FSphere3d& InDrawRegion);
+
+		// Set the line budget for debug draw
+		void SetCommandBudget(int32 InCommandBudget);
+
+		// Create a new timeline. E.g., PT, GT, RBAN
+		FChaosDDTimelinePtr CreateTimeline(const FString& Name);
+		
+		// Release a timeline
+		void ReleaseTimeline(const FChaosDDTimelinePtr& Scene);
+
+		// Render the most recent frame from each timeline
+		void RenderLatestFrames(IChaosDDRenderer& Renderer);
+
+	private:
+		TArray<FChaosDDFramePtr> GetFrames();
+
+		FString Name;
+		FCriticalSection TimelinesCS;
+		TArray<FChaosDDTimelinePtr> Timelines;
+		FSphere3d DrawRegion;
+		int32 CommandBudget;
+	};
+}
+
+#endif
\ No newline at end of file
diff --git a/Engine/Source/Runtime/Experimental/Chaos/Public/ChaosDebugDraw/ChaosDDTimeline.h b/Engine/Source/Runtime/Experimental/Chaos/Public/ChaosDebugDraw/ChaosDDTimeline.h
new file mode 100644
--- /dev/null
+++ b/Engine/Source/Runtime/Experimental/Chaos/Public/ChaosDebugDraw/ChaosDDTimeline.h
@@ -0,0 +1,116 @@
+// Copyright Epic Games, Inc. All Rights Reserved.
+#pragma once
+
+#include "ChaosDebugDraw/ChaosDDTypes.h"
+#include "ChaosDebugDraw/ChaosDDFrame.h"
+#include "ChaosDebugDraw/ChaosDDLog.h"
+#include "Containers/Map.h"
+#include "HAL/CriticalSection.h"
+#include "Math/Sphere.h"
+
+#if CHAOS_DEBUG_DRAW
+
+namespace ChaosDD::Private
+{
+	//
+	// Holds a sequence of debug draw frames. Every system that runs its own loop 
+	// will need a timeline. E.g., Physics Thread, Game Thread, RBANs
+	// 
+	// @todo(chaos): Add per-thread buffers to avoid locks
+	//
+	class CHAOS_API FChaosDDTimeline : public TSharedFromThis<FChaosDDTimeline>
+	{
+	public:
+		FChaosDDTimeline(const FString& InName, int32 InCommandBudget)
+			: Name(InName)
+			, MaxFrames(1)
+			, NextFrameIndex(0)
+			, DrawRegion(FSphere3d(FVector::Zero(), 0.0))
+			, CommandBudget(InCommandBudget)
+			, CommandQueueLength(0)
+		{
+		}
+
+		const FString& GetName() const
+		{
+			return Name;
+		}
+
+		void SetDrawRegion(const FSphere3d& InDrawRegion)
+		{
+			DrawRegion = InDrawRegion;
+		}
+
+		void SetCommandBudget(int32 InCommandBudget)
+		{
+			CommandBudget = InCommandBudget;
+		}
+
+		void BeginFrame(double InTime, double InDt)
+		{
+			FScopeLock Lock(&FramesCS);
+
+			// Duplicate BeginFrame?
+			check(!ActiveFrame.IsValid());
+
+			ActiveFrame = MakeShared<FChaosDDFrame>(this->AsShared(), NextFrameIndex++, InTime, InDt, DrawRegion, CommandBudget, CommandQueueLength);
+
+			UE_LOG(LogChaosDD, VeryVerbose, TEXT("BeginFrame %s %lld"), *GetName(), ActiveFrame->GetFrameIndex());
+		}
+
+		void EndFrame()
+		{
+			FScopeLock Lock(&FramesCS);
+
+			// Missing BeginFrame?
+			check(ActiveFrame.IsValid());
+
+			UE_LOG(LogChaosDD, VeryVerbose, TEXT("EndFrame %s %lld (%d Commands)"), *GetName(), ActiveFrame->GetFrameIndex(), ActiveFrame->GetNumCommands());
+
+			// Remeber the queue size to prevent array growth every frame
+			CommandQueueLength = FMath::RoundUpToPowerOfTwo(FMath::Max(CommandQueueLength, ActiveFrame->GetNumCommands()));
+
+			Frames.Add(ActiveFrame);
+			ActiveFrame.Reset();
+
+			PruneFrames();
+		}
+
+		const FChaosDDFramePtr& GetActiveFrame() const
+		{
+			return ActiveFrame;
+		}
+
+		void GetFrames(TArray<FChaosDDFramePtr>& InOutFrames)
+		{
+			FScopeLock Lock(&FramesCS);
+
+			for (const FChaosDDFramePtr& Frame : Frames)
+			{
+				InOutFrames.Add(Frame);
+			}
+		}
+
+	private:
+		void PruneFrames()
+		{
+			const int32 NumToRemove = Frames.Num() - MaxFrames;
+			if (NumToRemove > 0)
+			{
+				Frames.RemoveAt(0, NumToRemove, EAllowShrinking::No);
+			}
+		}
+
+		FString Name;
+		TArray<FChaosDDFramePtr> Frames;
+		FChaosDDFramePtr ActiveFrame;
+		int32 MaxFrames;
+		int64 NextFrameIndex;
+		FSphere3d DrawRegion;
+		int32 CommandBudget;
+		int32 CommandQueueLength;
+		FCriticalSection FramesCS;
+	};
+}
+
+#endif
diff --git a/Engine/Source/Runtime/Experimental/Chaos/Public/ChaosDebugDraw/ChaosDDTypes.h b/Engine/Source/Runtime/Experimental/Chaos/Public/ChaosDebugDraw/ChaosDDTypes.h
new file mode 100644
--- /dev/null
+++ b/Engine/Source/Runtime/Experimental/Chaos/Public/ChaosDebugDraw/ChaosDDTypes.h
@@ -0,0 +1,26 @@
+// Copyright Epic Games, Inc. All Rights Reserved.
+#pragma once
+
+#include "Chaos/Declares.h"
+#include "Templates/SharedPointer.h"
+#include "Templates/UniquePtr.h"
+
+#if CHAOS_DEBUG_DRAW
+
+// @todo(chaos): Move to ChaosDD when API is decided
+namespace ChaosDD::Private
+{
+	class FChaosDDContext;
+	class FChaosDDFrame;
+	class FChaosDDScene;
+	class FChaosDDTimeline;
+
+	class IChaosDDRenderer;
+
+	using FChaosDDFramePtr = TSharedPtr<FChaosDDFrame, ESPMode::ThreadSafe>;
+	using FChaosDDScenePtr = TSharedPtr<FChaosDDScene, ESPMode::ThreadSafe>;
+	using FChaosDDTimelinePtr = TSharedPtr<FChaosDDTimeline, ESPMode::ThreadSafe>;
+	using FChaosDDTimelineWeakPtr = TWeakPtr<FChaosDDTimeline, ESPMode::ThreadSafe>;
+}
+
+#endif
\ No newline at end of file
diff --git a/Engine/Source/Runtime/Experimental/Chaos/Public/PBDRigidsSolver.h b/Engine/Source/Runtime/Experimental/Chaos/Public/PBDRigidsSolver.h
--- a/Engine/Source/Runtime/Experimental/Chaos/Public/PBDRigidsSolver.h
+++ b/Engine/Source/Runtime/Experimental/Chaos/Public/PBDRigidsSolver.h
@@ -431,6 +431,17 @@ namespace Chaos
 
 		/** Sets if we are resimming or not */
 		void SetIsResimming(bool bIsResimming);
+
+#if CHAOS_DEBUG_DRAW
+	public:
+		CHAOS_API virtual void SetDebugDrawScene(const ChaosDD::Private::FChaosDDScenePtr& InCDDScene) override final;
+
+	private:
+		// NOTE: A physics frame may consist of multiple sub-steps (ticks) so this debug draw timeline only
+		// captures debug draw that happens outside of the tick(s). See FPBDRigidsEvolutionGBF::SetDebugDrawScene
+		ChaosDD::Private::FChaosDDScenePtr CDDScene;
+		ChaosDD::Private::FChaosDDTimelinePtr CDDFrameTimeline;
+#endif
 	};
 
 	template<>
diff --git a/Engine/Source/Runtime/Experimental/ChaosSolverEngine/Private/Chaos/ChaosDebugDrawComponent.cpp b/Engine/Source/Runtime/Experimental/ChaosSolverEngine/Private/Chaos/ChaosDebugDrawComponent.cpp
--- a/Engine/Source/Runtime/Experimental/ChaosSolverEngine/Private/Chaos/ChaosDebugDrawComponent.cpp
+++ b/Engine/Source/Runtime/Experimental/ChaosSolverEngine/Private/Chaos/ChaosDebugDrawComponent.cpp
@@ -2,6 +2,8 @@
 
 #include "Chaos/ChaosDebugDrawComponent.h"
 #include "Chaos/DebugDrawQueue.h"
+#include "ChaosDebugDraw/ChaosDDRenderer.h"
+#include "ChaosDebugDraw/ChaosDDScene.h"
 #include "ChaosLog.h"
 #include "Debug/DebugDrawService.h"
 #include "DrawDebugHelpers.h"
@@ -14,6 +16,7 @@
 #include "Engine/World.h"
 #include "GameFramework/Pawn.h"
 #include "GameFramework/PlayerController.h"
+#include "Physics/Experimental/PhysScene_Chaos.h"
 #include "VisualLogger/VisualLogger.h"
 
 #include UE_INLINE_GENERATED_CPP_BY_NAME(ChaosDebugDrawComponent)
@@ -55,7 +58,7 @@ int bChaosDebugDraw_DrawMode = 0;
 FAutoConsoleVariableRef CVarArrowSize(TEXT("p.Chaos.DebugDraw.Mode"), bChaosDebugDraw_DrawMode, TEXT("Where to send debug draw commands. 0 = UE Debug Draw; 1 = VisLog; 2 = Both"));
 
 
-float CommandLifeTime(const Chaos::FLatentDrawCommand& Command, const bool bIsPaused)
+float CommandLifeTime(const float LifeTime)
 {
 	// The linebatch time handling is a bit awkward and we need to translate
 	// Linebatcher Lifetime < 0: eternal (regardless of bPersistent flag)
@@ -64,14 +67,99 @@ float CommandLifeTime(const Chaos::FLatentDrawCommand& Command, const bool bIsPa
 	// Whereas in our Command, 
 	// Command lifetime <= 0: 1 frame
 	// Command lifetime > 0: specified duration
-	if (Command.LifeTime <= 0)
+	if (LifeTime <= 0)
 	{
 		// One frame - must be non-zero but also less than the next frame time which we don't know
 		// NOTE: this only works because UChaosDebugDrawComponent ticks after the line batcher
 		return UE_SMALL_NUMBER;
 	}
 	
-	return Command.LifeTime;
+	return LifeTime;
+}
+
+void DebugDrawChaosCommand(const UWorld* World, const Chaos::FLatentDrawCommand& Command)
+{
+	using namespace Chaos;
+
+	const uint8 DepthPriority = ChaosDebugDraw_SeeThrough ? Command.DepthPriority : 0;
+	switch (Command.Type)
+	{
+	case FLatentDrawCommand::EDrawType::Point:
+		DrawDebugPoint(World, Command.LineStart, Command.Thickness, Command.Color, Command.bPersistentLines, CommandLifeTime(Command.LifeTime), DepthPriority);
+		break;
+	case FLatentDrawCommand::EDrawType::Line:
+		DrawDebugLine(World, Command.LineStart, Command.LineEnd, Command.Color, Command.bPersistentLines, CommandLifeTime(Command.LifeTime), DepthPriority, Command.Thickness);
+		break;
+	case FLatentDrawCommand::EDrawType::DirectionalArrow:
+		DrawDebugDirectionalArrow(World, Command.LineStart, Command.LineEnd, Command.ArrowSize, Command.Color, Command.bPersistentLines, CommandLifeTime(Command.LifeTime), DepthPriority, Command.Thickness);
+		break;
+	case FLatentDrawCommand::EDrawType::Sphere:
+		DrawDebugSphere(World, Command.LineStart, Command.Radius, Command.Segments, Command.Color, Command.bPersistentLines, CommandLifeTime(Command.LifeTime), DepthPriority, Command.Thickness);
+		break;
+	case FLatentDrawCommand::EDrawType::Box:
+		DrawDebugBox(World, Command.Center, Command.Extent, Command.Rotation, Command.Color, Command.bPersistentLines, CommandLifeTime(Command.LifeTime), DepthPriority, Command.Thickness);
+		break;
+	case FLatentDrawCommand::EDrawType::String:
+		DrawDebugString(World, Command.TextLocation, Command.Text, Command.TestBaseActor, Command.Color, CommandLifeTime(Command.LifeTime), Command.bDrawShadow, Command.FontScale);
+		break;
+	case FLatentDrawCommand::EDrawType::Circle:
+	{
+		FMatrix M = FRotationMatrix::MakeFromYZ(Command.YAxis, Command.ZAxis);
+		M.SetOrigin(Command.Center);
+		DrawDebugCircle(World, M, Command.Radius, Command.Segments, Command.Color, Command.bPersistentLines, CommandLifeTime(Command.LifeTime), DepthPriority, Command.Thickness, Command.bDrawAxis);
+		break;
+	}
+	case FLatentDrawCommand::EDrawType::Capsule:
+		DrawDebugCapsule(World, Command.Center, Command.HalfHeight, Command.Radius, Command.Rotation, Command.Color, Command.bPersistentLines, CommandLifeTime(Command.LifeTime), DepthPriority, Command.Thickness);
+	default:
+		break;
+	}
+}
+
+void VisLogChaosCommand(const AActor* Actor, const Chaos::FLatentDrawCommand& Command)
+{
+	using namespace Chaos;
+
+	switch (Command.Type)
+	{
+	case FLatentDrawCommand::EDrawType::Point:
+		UE_VLOG_SEGMENT_THICK(Actor, LogChaos, Log, Command.LineStart, Command.LineStart, Command.Color, Command.Thickness, TEXT_EMPTY);
+		break;
+	case FLatentDrawCommand::EDrawType::Line:
+		UE_VLOG_SEGMENT(Actor, LogChaos, Log, Command.LineStart, Command.LineEnd, Command.Color, TEXT_EMPTY);
+		break;
+	case FLatentDrawCommand::EDrawType::DirectionalArrow:
+		UE_VLOG_SEGMENT(Actor, LogChaos, Log, Command.LineStart, Command.LineEnd, Command.Color, TEXT_EMPTY);
+		break;
+	case FLatentDrawCommand::EDrawType::Sphere:
+	{
+		// VLOG Capsule uses the bottom end as the origin (though the variable is named Center)
+		FVector Base = Command.LineStart - Command.Radius * FVector::UpVector;
+		UE_VLOG_CAPSULE(Actor, LogChaos, Log, Base, Command.Radius + KINDA_SMALL_NUMBER, Command.Radius, FQuat::Identity, Command.Color, TEXT_EMPTY);
+		break;
+	}
+	case FLatentDrawCommand::EDrawType::Box:
+		UE_VLOG_OBOX(Actor, LogChaos, Log, FBox(-Command.Extent, Command.Extent), FQuatRotationTranslationMatrix::Make(Command.Rotation, Command.Center), Command.Color, TEXT_EMPTY);
+		break;
+	case FLatentDrawCommand::EDrawType::String:
+		UE_VLOG(Command.TestBaseActor, LogChaos, Log, TEXT("%s"), *Command.Text);
+		break;
+	case FLatentDrawCommand::EDrawType::Circle:
+	{
+		const FMatrix M = FRotationMatrix::MakeFromYZ(Command.YAxis, Command.ZAxis);
+		UE_VLOG_CIRCLE(Actor, LogChaos, Log, Command.Center, M.GetUnitAxis(EAxis::X), Command.Radius, Command.Color, TEXT_EMPTY);
+		break;
+	}
+	case FLatentDrawCommand::EDrawType::Capsule:
+	{
+		// VLOG Capsule uses the bottom end as the origin (though the variable is named Center)
+		FVector Base = Command.Center - Command.HalfHeight * (Command.Rotation * FVector::UpVector);
+		UE_VLOG_CAPSULE(Actor, LogChaos, Log, Base, Command.HalfHeight, Command.Radius, Command.Rotation, Command.Color, TEXT_EMPTY);
+		break;
+	}
+	default:
+		break;
+	}
 }
 
 void DebugDrawChaos(const AActor* DebugDrawActor, const TArray<Chaos::FLatentDrawCommand>& DrawCommands, const bool bIsPaused)
@@ -110,39 +198,7 @@ void DebugDrawChaos(const AActor* DebugDrawActor, const TArray<Chaos::FLatentDra
 	{
 		for (const FLatentDrawCommand& Command : DrawCommands)
 		{
-			const uint8 DepthPriority = ChaosDebugDraw_SeeThrough ? Command.DepthPriority: 0;
-			switch (Command.Type)
-			{
-			case FLatentDrawCommand::EDrawType::Point:
-				DrawDebugPoint(World, Command.LineStart, Command.Thickness, Command.Color, Command.bPersistentLines, CommandLifeTime(Command, bIsPaused), DepthPriority);
-				break;
-			case FLatentDrawCommand::EDrawType::Line:
-				DrawDebugLine(World, Command.LineStart, Command.LineEnd, Command.Color, Command.bPersistentLines, CommandLifeTime(Command, bIsPaused), DepthPriority, Command.Thickness);
-				break;
-			case FLatentDrawCommand::EDrawType::DirectionalArrow:
-				DrawDebugDirectionalArrow(World, Command.LineStart, Command.LineEnd, Command.ArrowSize, Command.Color, Command.bPersistentLines, CommandLifeTime(Command, bIsPaused), DepthPriority, Command.Thickness);
-				break;
-			case FLatentDrawCommand::EDrawType::Sphere:
-				DrawDebugSphere(World, Command.LineStart, Command.Radius, Command.Segments, Command.Color, Command.bPersistentLines, CommandLifeTime(Command, bIsPaused), DepthPriority, Command.Thickness);
-				break;
-			case FLatentDrawCommand::EDrawType::Box:
-				DrawDebugBox(World, Command.Center, Command.Extent, Command.Rotation, Command.Color, Command.bPersistentLines, CommandLifeTime(Command, bIsPaused), DepthPriority, Command.Thickness);
-				break;
-			case FLatentDrawCommand::EDrawType::String:
-				DrawDebugString(World, Command.TextLocation, Command.Text, Command.TestBaseActor, Command.Color, CommandLifeTime(Command, bIsPaused), Command.bDrawShadow, Command.FontScale);
-				break;
-			case FLatentDrawCommand::EDrawType::Circle:
-			{
-				FMatrix M = FRotationMatrix::MakeFromYZ(Command.YAxis, Command.ZAxis);
-				M.SetOrigin(Command.Center);
-				DrawDebugCircle(World, M, Command.Radius, Command.Segments, Command.Color, Command.bPersistentLines, CommandLifeTime(Command, bIsPaused), DepthPriority, Command.Thickness, Command.bDrawAxis);
-				break;
-			}
-			case FLatentDrawCommand::EDrawType::Capsule:
-				DrawDebugCapsule(World, Command.Center, Command.HalfHeight, Command.Radius, Command.Rotation, Command.Color, Command.bPersistentLines, CommandLifeTime(Command, bIsPaused), DepthPriority, Command.Thickness);
-			default:
-				break;
-			}
+			DebugDrawChaosCommand(World, Command);
 		}
 	}
 
@@ -153,52 +209,49 @@ void DebugDrawChaos(const AActor* DebugDrawActor, const TArray<Chaos::FLatentDra
 		for (const FLatentDrawCommand& Command : DrawCommands)
 		{
 			const AActor* Actor = (Command.TestBaseActor) ? Command.TestBaseActor : DebugDrawActor;
-
-			switch (Command.Type)
-			{
-			case FLatentDrawCommand::EDrawType::Point:
-				UE_VLOG_SEGMENT_THICK(Actor, LogChaos, Log, Command.LineStart, Command.LineStart, Command.Color, Command.Thickness, TEXT_EMPTY);
-				break;
-			case FLatentDrawCommand::EDrawType::Line:
-				UE_VLOG_SEGMENT(Actor, LogChaos, Log, Command.LineStart, Command.LineEnd, Command.Color, TEXT_EMPTY);
-				break;
-			case FLatentDrawCommand::EDrawType::DirectionalArrow:
-				UE_VLOG_SEGMENT(Actor, LogChaos, Log, Command.LineStart, Command.LineEnd, Command.Color, TEXT_EMPTY);
-				break;
-			case FLatentDrawCommand::EDrawType::Sphere:
-			{
-				// VLOG Capsule uses the bottom end as the origin (though the variable is named Center)
-				FVector Base = Command.LineStart - Command.Radius * FVector::UpVector;
-				UE_VLOG_CAPSULE(Actor, LogChaos, Log, Base, Command.Radius + KINDA_SMALL_NUMBER, Command.Radius, FQuat::Identity, Command.Color, TEXT_EMPTY);
-				break;
-			}
-			case FLatentDrawCommand::EDrawType::Box:
-				UE_VLOG_OBOX(Actor, LogChaos, Log, FBox(-Command.Extent, Command.Extent), FQuatRotationTranslationMatrix::Make(Command.Rotation, Command.Center), Command.Color, TEXT_EMPTY);
-				break;
-			case FLatentDrawCommand::EDrawType::String:
-				UE_VLOG(Command.TestBaseActor, LogChaos, Log, TEXT("%s"), *Command.Text);
-				break;
-			case FLatentDrawCommand::EDrawType::Circle:
-			{
-				const FMatrix M = FRotationMatrix::MakeFromYZ(Command.YAxis, Command.ZAxis);
-				UE_VLOG_CIRCLE(Actor, LogChaos, Log, Command.Center, M.GetUnitAxis(EAxis::X), Command.Radius, Command.Color, TEXT_EMPTY);
-				break;
-			}
-			case FLatentDrawCommand::EDrawType::Capsule:
-			{
-				// VLOG Capsule uses the bottom end as the origin (though the variable is named Center)
-				FVector Base = Command.Center - Command.HalfHeight * (Command.Rotation * FVector::UpVector);
-				UE_VLOG_CAPSULE(Actor, LogChaos, Log, Base, Command.HalfHeight, Command.Radius, Command.Rotation, Command.Color, TEXT_EMPTY);
-				break;
-			}
-			default:
-				break;
-			}
+			VisLogChaosCommand(Actor, Command);
 		}
 	}
 }
 #endif
 
+#if CHAOS_DEBUG_DRAW
+
+class FChaosDDRenderer : public ChaosDD::Private::IChaosDDRenderer
+{
+public:
+	FChaosDDRenderer(UWorld* InWorld, AActor* InDebugDrawActor)
+		: World(InWorld)
+		, DebugDrawActor(InDebugDrawActor)
+	{
+		DepthPriority = 10;
+		bIsPaused = false;
+	}
+
+	virtual void DrawCommand(const Chaos::FLatentDrawCommand& Command) const override final
+	{
+		const bool bDrawUe = bChaosDebugDraw_DrawMode != 1;
+		if (bDrawUe)
+		{
+			DebugDrawChaosCommand(World, Command);
+		}
+
+		const bool bDrawVisLog = bChaosDebugDraw_DrawMode != 0;
+		if (bDrawVisLog)
+		{
+			const AActor* Actor = (Command.TestBaseActor) ? Command.TestBaseActor : DebugDrawActor;
+			VisLogChaosCommand(Actor, Command);
+		}
+	}
+
+private:
+	UWorld* World;
+	AActor* DebugDrawActor;
+	int8 DepthPriority;
+	bool bIsPaused;
+};
+
+#endif
 
 
 UChaosDebugDrawComponent::UChaosDebugDrawComponent()
@@ -289,8 +342,13 @@ void UChaosDebugDrawComponent::TickComponent(float DeltaTime, enum ELevelTick Ti
 	UWorld* World = GetWorld();
 	if (World != nullptr)
 	{
+		FVector RegionOfInterestOrigin = FVec3(0);
+		float RegionOfInterestRadius = ChaosDebugDraw_Radius;
+
 		if (World->ViewLocationsRenderedLastFrame.Num() > 0)
 		{
+			RegionOfInterestOrigin = World->ViewLocationsRenderedLastFrame[0];
+
 			if (bChaosDebugDraw_SingleActor)
 			{
 				if (const APlayerController* Controller = GEngine->GetFirstLocalPlayerController(World))
@@ -311,22 +369,20 @@ void UChaosDebugDrawComponent::TickComponent(float DeltaTime, enum ELevelTick Ti
 					{
 						FVector Origin, BoxExtent;
 						HitResult.GetActor()->GetActorBounds(true, Origin, BoxExtent);
-						const float Radius = BoxExtent.Size();
-						if (Radius <= ChaosDebugDraw_SingleActorMaxRadius)
-						{
-							FDebugDrawQueue::GetInstance().SetRegionOfInterest(Origin, Radius);
-						}
+						const float Radius = FMath::Min(BoxExtent.Size(), ChaosDebugDraw_SingleActorMaxRadius);
+						RegionOfInterestOrigin = Origin;
+						RegionOfInterestRadius = Radius;
 					}
 				}
 			}
-			else
-			{
-				FDebugDrawQueue::GetInstance().SetRegionOfInterest(World->ViewLocationsRenderedLastFrame[0], ChaosDebugDraw_Radius);
-			}
+
 		}
 
+		FDebugDrawQueue::GetInstance().SetRegionOfInterest(RegionOfInterestOrigin, RegionOfInterestRadius);
 		FDebugDrawQueue::GetInstance().SetMaxCost(ChaosDebugDraw_MaxElements);
 
+		// Draw any commands from threads that did not have the debug draw context initialized
+		// @todo(chaos): we should make sure all threads initialize the context...
 		const bool bIsPaused = World->IsPaused();
 		if (!bIsPaused)
 		{
@@ -334,6 +390,19 @@ void UChaosDebugDrawComponent::TickComponent(float DeltaTime, enum ELevelTick Ti
 
 			DebugDrawChaos(GetOwner(), DrawCommands, World->IsPaused());
 		}
+
+		const ChaosDD::Private::FChaosDDScenePtr& CDDScene = World->GetPhysicsScene()->GetDebugDrawScene();
+		if (CDDScene.IsValid())
+		{
+			CDDScene->SetDrawRegion(FSphere3d(RegionOfInterestOrigin, RegionOfInterestRadius));
+			CDDScene->SetCommandBudget(ChaosDebugDraw_MaxElements);
+
+			if (!bIsPaused)
+			{
+				FChaosDDRenderer CDDRenderer = FChaosDDRenderer(World, GetOwner());
+				CDDScene->RenderLatestFrames(CDDRenderer);
+			}
+		}
 	}
 #endif
 }
diff --git a/Engine/Source/Runtime/PhysicsCore/Private/ChaosScene.cpp b/Engine/Source/Runtime/PhysicsCore/Private/ChaosScene.cpp
--- a/Engine/Source/Runtime/PhysicsCore/Private/ChaosScene.cpp
+++ b/Engine/Source/Runtime/PhysicsCore/Private/ChaosScene.cpp
@@ -38,6 +38,9 @@
 
 #include "ChaosVisualDebugger/ChaosVisualDebuggerTrace.h"
 
+#include "ChaosDebugDraw/ChaosDDScene.h"
+#include "ChaosDebugDraw/ChaosDDTimeline.h"
+
 DECLARE_CYCLE_STAT(TEXT("Update Kinematics On Deferred SkelMeshes"),STAT_UpdateKinematicsOnDeferredSkelMeshesChaos,STATGROUP_Physics);
 CSV_DEFINE_CATEGORY(ChaosPhysics,true);
 CSV_DEFINE_CATEGORY(AABBTreeExpensiveStats, false);
@@ -79,6 +82,11 @@ FChaosScene::FChaosScene(
 	SceneSolver->GetChaosVDContextData().Type = static_cast<int32>(EChaosVDContextType::Solver);
 #endif
 
+#if CHAOS_DEBUG_DRAW
+	CDDScene = MakeShared<ChaosDD::Private::FChaosDDScene>(DebugName.ToString());
+	SceneSolver->SetDebugDrawScene(CDDScene);
+#endif
+
 	SceneSolver->PhysSceneHack = this;
 	SimCallback = SceneSolver->CreateAndRegisterSimCallbackObject_External<FChaosSceneSimCallback>();
 
diff --git a/Engine/Source/Runtime/PhysicsCore/Public/Chaos/ChaosScene.h b/Engine/Source/Runtime/PhysicsCore/Public/Chaos/ChaosScene.h
--- a/Engine/Source/Runtime/PhysicsCore/Public/Chaos/ChaosScene.h
+++ b/Engine/Source/Runtime/PhysicsCore/Public/Chaos/ChaosScene.h
@@ -20,6 +20,8 @@
 
 #include "ChaosVisualDebugger/ChaosVDContextProvider.h"
 
+#include "ChaosDebugDraw/ChaosDDTypes.h"
+
 #ifndef CHAOS_DEBUG_NAME
 #define CHAOS_DEBUG_NAME 0
 #endif
@@ -203,4 +205,15 @@ public:
 		return CVDContextData;
 	};
 #endif
+
+#if CHAOS_DEBUG_DRAW
+public:
+	const ChaosDD::Private::FChaosDDScenePtr& GetDebugDrawScene() const
+	{ 
+		return CDDScene;
+	}
+
+private:
+	ChaosDD::Private::FChaosDDScenePtr CDDScene;
+#endif
 };
