Chaos - fix new debug draw backend ensure in programs that do not set up the DDContext at all

diff --git a/Engine/Source/Runtime/Experimental/Chaos/Private/ChaosDebugDraw/ChaosDDContext.cpp b/Engine/Source/Runtime/Experimental/Chaos/Private/ChaosDebugDraw/ChaosDDContext.cpp
--- a/Engine/Source/Runtime/Experimental/Chaos/Private/ChaosDebugDraw/ChaosDDContext.cpp
+++ b/Engine/Source/Runtime/Experimental/Chaos/Private/ChaosDebugDraw/ChaosDDContext.cpp
@@ -36,7 +36,7 @@ namespace ChaosDD::Private
 	{
 		FChaosDDContext& Context = Get();
 
-		if (ensure(Context.GetTimeline().IsValid()))
+		if (Context.GetTimeline().IsValid())
 		{
 			Context.GetTimeline()->EndFrame();
 
