Chaos Modular Vehicles:
Create a suspension base interface for the suspension module so that it can be extended by users.
(based on Bill.Henderson's shelf 33170285)

diff --git a/Engine/Plugins/Experimental/ChaosModularVehicle/Source/ChaosModularVehicleEngine/Private/ChaosModularVehicle/ModularVehicleBuilder.cpp b/Engine/Plugins/Experimental/ChaosModularVehicle/Source/ChaosModularVehicleEngine/Private/ChaosModularVehicle/ModularVehicleBuilder.cpp
--- a/Engine/Plugins/Experimental/ChaosModularVehicle/Source/ChaosModularVehicleEngine/Private/ChaosModularVehicle/ModularVehicleBuilder.cpp
+++ b/Engine/Plugins/Experimental/ChaosModularVehicle/Source/ChaosModularVehicleEngine/Private/ChaosModularVehicle/ModularVehicleBuilder.cpp
@@ -51,7 +51,7 @@ void FModularVehicleBuilder::FixupTreeLinks(TUniquePtr<Chaos::FSimModuleTree>& S
 
 				if (ParentModule && ParentModule->GetSimType() == eSimType::Wheel)
 				{
-					FSuspensionSimModule* Suspension = static_cast<FSuspensionSimModule*>(Module);
+					FSuspensionBaseInterface* Suspension = static_cast<FSuspensionBaseInterface*>(Module);
 					FWheelBaseInterface* Wheel = static_cast<FWheelBaseInterface*>(ParentModule);
 					Wheel->SetSuspensionSimTreeIndex(Suspension->GetTreeIndex());
 					Suspension->SetWheelSimTreeIndex(Wheel->GetTreeIndex());
@@ -68,7 +68,7 @@ void FModularVehicleBuilder::FixupTreeLinks(TUniquePtr<Chaos::FSimModuleTree>& S
 
 					if (ChildModule && ChildModule->GetSimType() == eSimType::Wheel)
 					{
-						FSuspensionSimModule* Suspension = static_cast<FSuspensionSimModule*>(Module);
+						FSuspensionBaseInterface* Suspension = static_cast<FSuspensionBaseInterface*>(Module);
 						FWheelBaseInterface* Wheel = static_cast<FWheelBaseInterface*>(ChildModule);
 						Wheel->SetSuspensionSimTreeIndex(Suspension->GetTreeIndex());
 						Suspension->SetWheelSimTreeIndex(Wheel->GetTreeIndex());
diff --git a/Engine/Plugins/Experimental/ChaosModularVehicle/Source/ChaosModularVehicleEngine/Private/ChaosModularVehicle/ModularVehicleSimulationCU.cpp b/Engine/Plugins/Experimental/ChaosModularVehicle/Source/ChaosModularVehicleEngine/Private/ChaosModularVehicle/ModularVehicleSimulationCU.cpp
--- a/Engine/Plugins/Experimental/ChaosModularVehicle/Source/ChaosModularVehicleEngine/Private/ChaosModularVehicle/ModularVehicleSimulationCU.cpp
+++ b/Engine/Plugins/Experimental/ChaosModularVehicle/Source/ChaosModularVehicleEngine/Private/ChaosModularVehicle/ModularVehicleSimulationCU.cpp
@@ -194,7 +194,7 @@ void FModularVehicleSimulationCU::PerformAdditionalSimWork(UWorld* InWorld, cons
 							if (Node.SimModule->IsClustered() && Node.SimModule->IsBehaviourType(Chaos::eSimModuleTypeFlags::Raycast))
 							{
 								Chaos::FSpringTrace OutTrace;
-								Chaos::FSuspensionSimModule* Suspension = static_cast<Chaos::FSuspensionSimModule*>(Node.SimModule);
+								Chaos::FSuspensionBaseInterface* Suspension = static_cast<Chaos::FSuspensionBaseInterface*>(Node.SimModule);
 
 								// would be cleaner an faster to just store radius in suspension also
 								float WheelRadius = 0;
@@ -243,7 +243,7 @@ void FModularVehicleSimulationCU::PerformAdditionalSimWork(UWorld* InWorld, cons
 									}
 								}
 
-								float Offset = Suspension->Setup().MaxLength;
+								float Offset = Suspension->GetMaxSpringLength();
 								if (HitResult.bBlockingHit && GModularVehicleDebugParams.SuspensionRaycastsEnabled)
 								{
 									Offset = HitResult.Distance - WheelRadius;
@@ -345,17 +345,19 @@ void FModularVehicleSimulationCU::PerformAdditionalSimWork(UWorld* InWorld, cons
 								FVector Up = ClusterWorldTM.GetUnitAxis(EAxis::Z);
 
 								FVector HitPoint;
-
+								float HitDistance = 0.f;
 								if (InputData.PhysicsInputs.TraceType == ETraceType::Spherecast)
 								{
 									HitPoint = HitResult.Location;
+									HitDistance = HitResult.Distance;
 								}
 								else
 								{
 									HitPoint = HitResult.ImpactPoint + Up * WheelRadius;
+									HitDistance = HitResult.Distance - WheelRadius;
 								}
 
-								Suspension->SetTargetPoint(HitPoint, HitResult.ImpactNormal, HitResult.bBlockingHit);
+								Suspension->SetTargetPoint(HitPoint, HitResult.ImpactNormal, HitDistance, HitResult.bBlockingHit);
 							}
 
 						}
diff --git a/Engine/Source/Runtime/Experimental/ChaosVehicles/ChaosVehiclesCore/Private/SimModule/SuspensionModule.cpp b/Engine/Source/Runtime/Experimental/ChaosVehicles/ChaosVehiclesCore/Private/SimModule/SuspensionModule.cpp
--- a/Engine/Source/Runtime/Experimental/ChaosVehicles/ChaosVehiclesCore/Private/SimModule/SuspensionModule.cpp
+++ b/Engine/Source/Runtime/Experimental/ChaosVehicles/ChaosVehiclesCore/Private/SimModule/SuspensionModule.cpp
@@ -20,12 +20,8 @@ namespace Chaos
 		, SpringDisplacement(0.f)
 		, LastDisplacement(0.f)
 		, SpringSpeed(0.f)
-		, WheelSimTreeIndex(INVALID_IDX)
 		, Constraint(nullptr)
 		, ConstraintIndex(INVALID_IDX)
-		, TargetPos(FVector::ZeroVector)
-		, ImpactNormal(FVector::ZeroVector)
-		, WheelInContact(false)
 	{
 		AccessSetup().MaxLength = FMath::Abs(Settings.MaxRaise + Settings.MaxDrop);
 
diff --git a/Engine/Source/Runtime/Experimental/ChaosVehicles/ChaosVehiclesCore/Public/SimModule/SuspensionBaseInterface.h b/Engine/Source/Runtime/Experimental/ChaosVehicles/ChaosVehiclesCore/Public/SimModule/SuspensionBaseInterface.h
new file mode 100644
--- /dev/null
+++ b/Engine/Source/Runtime/Experimental/ChaosVehicles/ChaosVehiclesCore/Public/SimModule/SuspensionBaseInterface.h
@@ -0,0 +1,74 @@
+// Copyright Epic Games, Inc. All Rights Reserved.
+
+#pragma once
+
+#include "SimModule/SimulationModuleBase.h"
+
+namespace Chaos
+{
+
+/** Suspension world ray/shape trace start and end positions */
+struct CHAOSVEHICLESCORE_API FSpringTrace
+{
+	FVector Start;
+	FVector End;
+
+	FVector TraceDir()
+	{
+		FVector Dir(End - Start);
+		return Dir.FVector::GetSafeNormal();
+	}
+
+	float Length()
+	{
+		FVector Dir(End - Start);
+		return Dir.Size();
+	}
+};
+
+
+class CHAOSVEHICLESCORE_API FSuspensionBaseInterface : public ISimulationModuleBase
+{
+public:
+	FSuspensionBaseInterface()
+		: WheelSimTreeIndex(INVALID_IDX)
+		, TargetPos(FVector::ZeroVector)
+		, ImpactNormal(FVector::ZeroVector)
+		, HitDistance(0.f)
+		, WheelInContact(false)
+	{
+	}
+
+	virtual ~FSuspensionBaseInterface() {}
+
+	virtual bool IsBehaviourType(eSimModuleTypeFlags InType) const override { return (InType & Raycast); }
+
+	virtual eSimType GetSimType() const override { return eSimType::Suspension; }
+
+	virtual float GetMaxSpringLength() const = 0;
+	virtual float GetSpringLength() const = 0;
+	virtual void SetSpringLength(float InLength, float WheelRadius) = 0;
+	virtual void GetWorldRaycastLocation(const FTransform& BodyTransform, float WheelRadius, FSpringTrace& OutTrace) = 0;
+
+	void SetTargetPoint(const FVector& InTargetPoint, const FVector& InImpactNormal, float InHitDistance, bool InWheelInContact)
+	{
+		TargetPos = InTargetPoint;
+		ImpactNormal = InImpactNormal;
+		HitDistance = InHitDistance;
+		WheelInContact = InWheelInContact;
+	}
+	void SetWheelSimTreeIndex(int WheelTreeIndexIn) { WheelSimTreeIndex = WheelTreeIndexIn; }
+	int GetWheelSimTreeIndex() const { return WheelSimTreeIndex; }
+
+protected:
+	int WheelSimTreeIndex;
+
+	FVector TargetPos;
+	FVector ImpactNormal;
+	float HitDistance;
+	bool WheelInContact;
+
+};
+
+
+} // namespace Chaos
diff --git a/Engine/Source/Runtime/Experimental/ChaosVehicles/ChaosVehiclesCore/Public/SimModule/SuspensionModule.h b/Engine/Source/Runtime/Experimental/ChaosVehicles/ChaosVehiclesCore/Public/SimModule/SuspensionModule.h
--- a/Engine/Source/Runtime/Experimental/ChaosVehicles/ChaosVehiclesCore/Public/SimModule/SuspensionModule.h
+++ b/Engine/Source/Runtime/Experimental/ChaosVehicles/ChaosVehiclesCore/Public/SimModule/SuspensionModule.h
@@ -2,7 +2,7 @@
 
 #pragma once
 
-#include "SimModule/SimulationModuleBase.h"
+#include "SimModule/SuspensionBaseInterface.h"
 
 namespace Chaos
 {
@@ -92,26 +92,7 @@ namespace Chaos
 		//	float Swaybar;				// Anti-roll bar
 	};
 
-	/** Suspension world ray/shape trace start and end positions */
-	struct CHAOSVEHICLESCORE_API FSpringTrace
-	{
-		FVector Start;
-		FVector End;
-
-		FVector TraceDir()
-		{
-			FVector Dir(End - Start);
-			return Dir.FVector::GetSafeNormal();
-		}
-
-		float Length()
-		{
-			FVector Dir(End - Start);
-			return Dir.Size();
-		}
-	};
-
-	class CHAOSVEHICLESCORE_API FSuspensionSimModule : public ISimulationModuleBase, public TSimModuleSettings<FSuspensionSettings>
+	class CHAOSVEHICLESCORE_API FSuspensionSimModule : public FSuspensionBaseInterface, public TSimModuleSettings<FSuspensionSettings>
 	{
 		friend FSuspensionSimModuleDatas;
 		friend FSuspensionOutputData;
@@ -120,7 +101,6 @@ namespace Chaos
 
 		FSuspensionSimModule(const FSuspensionSettings& Settings);
 
-		virtual bool IsBehaviourType(eSimModuleTypeFlags InType) const override { return (InType & Raycast); }
 		virtual TSharedPtr<FModuleNetData> GenerateNetData(int SimArrayIndex) const
 		{
 			return MakeShared<FSuspensionSimModuleDatas>(
@@ -136,13 +116,12 @@ namespace Chaos
 			return FSuspensionOutputData::MakeNew();
 		}
 
-		virtual eSimType GetSimType() const { return eSimType::Suspension; }
-
 		virtual const FString GetDebugName() const { return TEXT("Suspension"); }
 
-		float GetSpringLength() const;
-		void SetSpringLength(float InLength, float WheelRadius);
-		void GetWorldRaycastLocation(const FTransform& BodyTransform, float WheelRadius, FSpringTrace& OutTrace);
+		virtual float GetMaxSpringLength() const override { return Setup().MaxLength; }
+		virtual float GetSpringLength() const override;
+		virtual  void SetSpringLength(float InLength, float WheelRadius) override;
+		virtual void GetWorldRaycastLocation(const FTransform& BodyTransform, float WheelRadius, FSpringTrace& OutTrace) override;
 
 		virtual void Simulate(float DeltaTime, const FAllInputs& Inputs, FSimModuleTree& VehicleModuleSystem) override;
 
@@ -150,32 +129,19 @@ namespace Chaos
 
 		const FVector& GetRestLocation() const { return Setup().RestOffset; }
 
-		void SetWheelSimTreeIndex(int WheelTreeIndexIn) { WheelSimTreeIndex = WheelTreeIndexIn; }
-		int GetWheelSimTreeIndex() const { return WheelSimTreeIndex; } 
-
 		void UpdateConstraint();
 
 		void SetSuspensionConstraint(FSuspensionConstraint* InConstraint);
 		void SetConstraintIndex(int32 InConstraintIndex) { ConstraintIndex = InConstraintIndex; }
 		int32 GetConstraintIndex() const { return ConstraintIndex; }
-		void SetTargetPoint(const FVector& InTargetPoint, const FVector& InImpactNormal, bool InWheelInContact)
-		{
-			TargetPos = InTargetPoint;
-			ImpactNormal = InImpactNormal;
-			WheelInContact = InWheelInContact;
-		}
 	private:
 
 		float SpringDisplacement;
 		float LastDisplacement;
 		float SpringSpeed;
-		int WheelSimTreeIndex;
 
 		FSuspensionConstraint* Constraint;
 		int32 ConstraintIndex;
-		FVector TargetPos;
-		FVector ImpactNormal;
-		bool WheelInContact;
 	};
 
 
