Fix potential crash with half connected dataflow geometry collection terminal nodes

diff --git a/Engine/Source/Runtime/Experimental/Chaos/Private/GeometryCollection/Facades/CollectionInstancedMeshFacade.cpp b/Engine/Source/Runtime/Experimental/Chaos/Private/GeometryCollection/Facades/CollectionInstancedMeshFacade.cpp
--- a/Engine/Source/Runtime/Experimental/Chaos/Private/GeometryCollection/Facades/CollectionInstancedMeshFacade.cpp
+++ b/Engine/Source/Runtime/Experimental/Chaos/Private/GeometryCollection/Facades/CollectionInstancedMeshFacade.cpp
@@ -32,7 +32,7 @@ namespace GeometryCollection::Facades
 
 	int32 FCollectionInstancedMeshFacade::GetNumIndices() const
 	{
-		return InstancedMeshIndexAttribute.Get().Num();
+		return IsValid() ? InstancedMeshIndexAttribute.Get().Num(): 0;
 	}
 
 	int32 FCollectionInstancedMeshFacade::GetIndex(int32 TransformIndex) const
