GeometryCollection : TransformFacade updates.
-Move algos for Parent and Unparent into Facade.

diff --git a/Engine/Source/Runtime/Experimental/Chaos/Private/GeometryCollection/Facades/CollectionTransformFacade.cpp b/Engine/Source/Runtime/Experimental/Chaos/Private/GeometryCollection/Facades/CollectionTransformFacade.cpp
--- a/Engine/Source/Runtime/Experimental/Chaos/Private/GeometryCollection/Facades/CollectionTransformFacade.cpp
+++ b/Engine/Source/Runtime/Experimental/Chaos/Private/GeometryCollection/Facades/CollectionTransformFacade.cpp
@@ -8,14 +8,18 @@
 namespace GeometryCollection::Facades
 {
 	FCollectionTransformFacade::FCollectionTransformFacade(FManagedArrayCollection& InCollection)
-		: ParentAttribute(InCollection, FTransformCollection::ParentAttribute, FTransformCollection::TransformGroup)
+		: ConstCollection(InCollection)
+		, Collection(&InCollection)
+		, ParentAttribute(InCollection, FTransformCollection::ParentAttribute, FTransformCollection::TransformGroup)
 		, ChildrenAttribute(InCollection, FTransformCollection::ChildrenAttribute, FTransformCollection::TransformGroup)
 		, TransformAttribute(InCollection, FTransformCollection::TransformAttribute, FTransformCollection::TransformGroup)
 		, BoneNameAttribute(InCollection, "BoneName", FTransformCollection::TransformGroup)
 	{}
 
 	FCollectionTransformFacade::FCollectionTransformFacade(const FManagedArrayCollection& InCollection)
-		: ParentAttribute(InCollection, FTransformCollection::ParentAttribute, FTransformCollection::TransformGroup)
+		: ConstCollection(InCollection)
+		, Collection(nullptr)
+		, ParentAttribute(InCollection, FTransformCollection::ParentAttribute, FTransformCollection::TransformGroup)
 		, ChildrenAttribute(InCollection, FTransformCollection::ChildrenAttribute, FTransformCollection::TransformGroup)
 		, TransformAttribute(InCollection, FTransformCollection::TransformAttribute, FTransformCollection::TransformGroup)
 		, BoneNameAttribute(InCollection, "BoneName", FTransformCollection::TransformGroup)
@@ -241,6 +245,144 @@ namespace GeometryCollection::Facades
 		}
 	}
 
+	bool FCollectionTransformFacade::HasCycle(const TManagedArray<int32>& Parents, int32 Node)
+	{
+		const int32 NumParents = Parents.Num();
+		int32 WalkNode = Node;
+		for (int32 Iters = 0; WalkNode != INDEX_NONE && Iters < NumParents; ++Iters)
+		{
+			WalkNode = Parents[WalkNode];
+		}
+		return WalkNode != INDEX_NONE;
+	}
+
+	bool FCollectionTransformFacade::HasCycle(const TManagedArray<int32>& Parents, const TArray<int32>& SelectedBones)
+	{
+		for (int32 Bone : SelectedBones)
+		{
+			if (HasCycle(Parents, Bone))
+			{
+				return true;
+			}
+		}
+		return false;
+	}
+
+	void FCollectionTransformFacade::ParentTransform(const int32 TransformIndex, const int32 ChildIndex)
+	{
+		check(Collection);
 
+		TArray<int32> SelectedBones;
+		SelectedBones.Add(ChildIndex);
+		ParentTransforms(TransformIndex, SelectedBones);
+	}
+
+	void FCollectionTransformFacade::ParentTransforms(const int32 TransformIndex, const TArray<int32>& SelectedBones)
+	{
+		check(Collection);
+		TManagedArray<FTransform3f>& Transform = TransformAttribute.Modify();
+		TManagedArray<int32>& Parents = ParentAttribute.Modify();
+		TManagedArray<TSet<int32>>& Children = ChildrenAttribute.Modify();
+		
+		if (ensure(-1 <= TransformIndex && TransformIndex < Transform.Num()))
+		{
+			// pre calculate global positions
+			TArray<FTransform3f> GlobalTransform;
+			GeometryCollectionAlgo::GlobalMatrices(Transform, Parents, GlobalTransform);
+
+			// append children 
+			for (int32 Index = 0; Index < SelectedBones.Num(); Index++)
+			{
+				int32 BoneIndex = SelectedBones[Index];
+				if (ensure(0 <= BoneIndex && BoneIndex < Parents.Num()))
+				{
+					// remove entry in previous parent
+					int32 ParentIndex = Parents[BoneIndex];
+					if (ParentIndex != INDEX_NONE)
+					{
+						if (ensure(0 <= ParentIndex && ParentIndex < Parents.Num()))
+						{
+							Children[ParentIndex].Remove(BoneIndex);
+						}
+					}
+
+					// set new parent
+					Parents[BoneIndex] = TransformIndex;
+				}
+			}
+
+			FTransform3f ParentInverse = FTransform3f::Identity;
+			if (TransformIndex != INDEX_NONE)
+			{
+				Children[TransformIndex].Append(SelectedBones);
+				ParentInverse = GlobalTransform[TransformIndex].Inverse();
+			}
+
+			// move the children to the local space of the transform. 
+			for (int32 Index = 0; Index < SelectedBones.Num(); Index++)
+			{
+				int32 BoneIndex = SelectedBones[Index];
+				Transform[BoneIndex] = GlobalTransform[BoneIndex] * ParentInverse;
+			}
+
+		}
+
+		// error check for circular dependencies
+		ensure(!HasCycle(Parents, TransformIndex));
+		ensure(!HasCycle(Parents, SelectedBones));
+	}
+
+	void FCollectionTransformFacade::UnparentTransform(const int32 ChildIndex)
+	{
+		check(Collection);
+
+		if (IsValid())
+		{
+			int32 NumTransforms = Collection->NumElements(FTransformCollection::TransformGroup);
+
+			if (0 < ChildIndex && ChildIndex < NumTransforms)
+			{
+				TManagedArray<int32>& Parent = ParentAttribute.Modify();
+				TManagedArray<TSet<int32>>& Children = ChildrenAttribute.Modify();
+
+				int32 ParentIndex = Parent[ChildIndex];
+				if (0 <= ParentIndex && ParentIndex < NumTransforms)
+				{
+					Children[ParentIndex].Remove(ChildIndex);
+					Parent[ChildIndex] = INDEX_NONE;
+				}
+			}
+		}
+	}
+
+
+	CHAOS_API void FCollectionTransformFacade::EnforceSingleRoot(FString RootName)
+	{
+		TArray<int32> Roots = GetRootIndices();
+		if (Roots.Num() > 1)
+		{
+			int32 Idx = TransformAttribute.AddElements(1);
+
+			if (HasBoneNameAttribute())
+			{
+				const TManagedArray<FString>& BoneName = BoneNameAttribute.Get();
+
+				auto MakeUnique = [&BoneName](FString& NewName)
+				{
+					int32 CurrentIndex = 1;
+					FString TestName = NewName;
+					while (BoneName.Contains(TestName))
+					{
+						TestName = FString::Printf(TEXT("%s%d"), *NewName, CurrentIndex++);
+					}
+					return TestName;
+				};
+
+				BoneNameAttribute.Modify()[Idx]=MakeUnique(RootName);
+			}
+
+			ParentTransforms(Idx, Roots);
+		}
+	}
 
 }
\ No newline at end of file
diff --git a/Engine/Source/Runtime/Experimental/Chaos/Private/GeometryCollection/GeometryCollectionAlgo.cpp b/Engine/Source/Runtime/Experimental/Chaos/Private/GeometryCollection/GeometryCollectionAlgo.cpp
--- a/Engine/Source/Runtime/Experimental/Chaos/Private/GeometryCollection/GeometryCollectionAlgo.cpp
+++ b/Engine/Source/Runtime/Experimental/Chaos/Private/GeometryCollection/GeometryCollectionAlgo.cpp
@@ -7,6 +7,7 @@
 #include "GeometryCollection/GeometryCollection.h"
 #include "GeometryCollection/RecordedTransformTrack.h"
 #include "GeometryCollection/Facades/CollectionHierarchyFacade.h"
+#include "GeometryCollection/Facades/CollectionTransformFacade.h"
 #include "GeometryCollectionProxyData.h"
 #include "Async/ParallelFor.h"
 
@@ -177,116 +178,39 @@ namespace GeometryCollectionAlgo
 		return false;
 	}
 
-	bool HasCycle(TManagedArray<int32>& Parents, int32 Node)
+	bool HasCycle(const TManagedArray<int32>& Parents, int32 Node)
 	{
-		const int32 NumParents = Parents.Num();
-		int32 WalkNode = Node;
-		for (int32 Iters = 0; WalkNode != FGeometryCollection::Invalid && Iters < NumParents; ++Iters)
-		{
-			WalkNode = Parents[WalkNode];
-		}
-		return WalkNode != FGeometryCollection::Invalid;
+		return GeometryCollection::Facades::FCollectionTransformFacade::HasCycle(Parents, Node);
 	}
-	bool HasCycle(TManagedArray<int32>& Parents, const TArray<int32>& SelectedBones)
+	bool HasCycle(const TManagedArray<int32>& Parents, const TArray<int32>& SelectedBones)
 	{
-		for (int32 Bone : SelectedBones)
-		{
-			if (HasCycle(Parents, Bone))
-			{
-				return true;
-			}
-		}
-		return false;
+		return GeometryCollection::Facades::FCollectionTransformFacade::HasCycle(Parents, SelectedBones);
 	}
 
-	void ParentTransform(FTransformCollection* GeometryCollection, const int32 TransformIndex, const int32 ChildIndex)
+	void ParentTransform(FManagedArrayCollection* ManagedArrayCollection, const int32 TransformIndex, const int32 ChildIndex)
 	{
-		TArray<int32> SelectedBones;
-		SelectedBones.Add(ChildIndex);
-		ParentTransforms(GeometryCollection, TransformIndex, SelectedBones);
+		if (ManagedArrayCollection)
+		{
+			GeometryCollection::Facades::FCollectionTransformFacade Facade(*ManagedArrayCollection);
+			Facade.ParentTransform(TransformIndex, ChildIndex);
+		}
 	}
 
-	void ParentTransforms(FTransformCollection* GeometryCollection, const int32 TransformIndex,
-		const TArray<int32>& SelectedBones)
+	void ParentTransforms(FManagedArrayCollection* ManagedArrayCollection, const int32 TransformIndex, const TArray<int32>& SelectedBones)
 	{
-		check(GeometryCollection != nullptr);
-
-		TManagedArray<FTransform3f>& Transform = GeometryCollection->Transform;
-		TManagedArray<int32>& Parents = GeometryCollection->Parent;
-		TManagedArray<TSet<int32>>& Children = GeometryCollection->Children;
-
-		if (ensure(-1 <= TransformIndex && TransformIndex < GeometryCollection->NumElements(FGeometryCollection::TransformGroup)))
+		if (ManagedArrayCollection)
 		{
-			// pre calculate global positions
-			TArray<FTransform3f> GlobalTransform;
-			GeometryCollectionAlgo::GlobalMatrices(GeometryCollection->Transform, Parents, GlobalTransform);
-
-			// append children 
-			for (int32 Index = 0; Index < SelectedBones.Num(); Index++)
-			{
-				int32 BoneIndex = SelectedBones[Index];
-				if (ensure(0 <= BoneIndex && BoneIndex < Parents.Num()))
-				{
-					// remove entry in previous parent
-					int32 ParentIndex = Parents[BoneIndex];
-					if (ParentIndex != FGeometryCollection::Invalid)
-					{
-						if (ensure(0 <= ParentIndex && ParentIndex < Parents.Num()))
-						{
-							Children[ParentIndex].Remove(BoneIndex);
-						}
-					}
-
-					// set new parent
-					Parents[BoneIndex] = TransformIndex;
-				}
-			}
-
-			FTransform3f ParentInverse = FTransform3f::Identity;
-			if (TransformIndex != FGeometryCollection::Invalid)
-			{
-				Children[TransformIndex].Append(SelectedBones);
-				ParentInverse = GlobalTransform[TransformIndex].Inverse();
-			}
-
-			// move the children to the local space of the transform. 
-			for (int32 Index = 0; Index < SelectedBones.Num(); Index++)
-			{
-				int32 BoneIndex = SelectedBones[Index];
-				Transform[BoneIndex] = GlobalTransform[BoneIndex] * ParentInverse;
-			}
-
+			GeometryCollection::Facades::FCollectionTransformFacade Facade(*ManagedArrayCollection);
+			Facade.ParentTransforms(TransformIndex, SelectedBones);
 		}
-
-		// error check for circular dependencies
-		ensure(!HasCycle(Parents, TransformIndex));
-		ensure(!HasCycle(Parents, SelectedBones));
 	}
 
-	void UnparentTransform(FManagedArrayCollection* Collection, const int32 ChildIndex)
+	void UnparentTransform(FManagedArrayCollection* ManagedArrayCollection, const int32 ChildIndex)
 	{
-		if (Collection)
+		if (ManagedArrayCollection)
 		{
-			if (Collection->HasAttribute(FTransformCollection::ChildrenAttribute, FTransformCollection::TransformGroup))
-			{
-				if (Collection->HasAttribute(FTransformCollection::ParentAttribute, FTransformCollection::TransformGroup))
-				{
-					int32 NumTransforms = Collection->NumElements(FTransformCollection::TransformGroup);
-
-					if (0 < ChildIndex && ChildIndex < NumTransforms)
-					{
-						TManagedArray<int32>& Parent = Collection->ModifyAttribute<int32>(FTransformCollection::ParentAttribute, FTransformCollection::TransformGroup);
-						TManagedArray< TArray<int32> >& Children = Collection->ModifyAttribute< TArray<int32> >(FTransformCollection::ChildrenAttribute, FTransformCollection::TransformGroup);
-
-						int32 ParentIndex = Parent[ChildIndex];
-						if (0 <= ParentIndex && ParentIndex < NumTransforms)
-						{
-							Children[ParentIndex].Remove(ChildIndex);
-							Parent[ChildIndex] = INDEX_NONE;
-						}
-					}
-				}
-			}
+			GeometryCollection::Facades::FCollectionTransformFacade Facade(*ManagedArrayCollection);
+			Facade.UnparentTransform(ChildIndex);
 		}
 	}
 
diff --git a/Engine/Source/Runtime/Experimental/Chaos/Public/GeometryCollection/Facades/CollectionTransformFacade.h b/Engine/Source/Runtime/Experimental/Chaos/Public/GeometryCollection/Facades/CollectionTransformFacade.h
--- a/Engine/Source/Runtime/Experimental/Chaos/Public/GeometryCollection/Facades/CollectionTransformFacade.h
+++ b/Engine/Source/Runtime/Experimental/Chaos/Public/GeometryCollection/Facades/CollectionTransformFacade.h
@@ -105,7 +105,28 @@ namespace GeometryCollection::Facades
 		/** Sets the selected bone's transform to identity */
 		CHAOS_API void SetBoneTransformToIdentity(int32 BoneIdx);
 
+		/** Does the transform heirarchy have a cycle*/
+		static bool CHAOS_API HasCycle(const TManagedArray<int32>& Parents, int32 Node);
+
+		/** Does the transform heirarchy have a cycle*/
+		static bool CHAOS_API HasCycle(const TManagedArray<int32>& Parents, const TArray<int32>& SelectedBones);
+
+		/** Parent a single transform */
+		void CHAOS_API ParentTransform(const int32 TransformIndex, const int32 ChildIndex);
+
+		/**  Parent the list of transforms to the selected index. */
+		void CHAOS_API ParentTransforms(const int32 TransformIndex, const TArray<int32>& SelectedBones);
+
+		/**  Unparent the child index from its parent */
+		void CHAOS_API UnparentTransform(const int32 ChildIndex);
+
+		/** Adds a Identity transform and nests all roots under the new transform */
+		CHAOS_API void EnforceSingleRoot(FString RootName);
+
 	private:
+		const FManagedArrayCollection& ConstCollection;
+		FManagedArrayCollection* Collection = nullptr;
+
 		TManagedArrayAccessor<int32>		ParentAttribute;
 		TManagedArrayAccessor<TSet<int32>>	ChildrenAttribute;
 		TManagedArrayAccessor<FTransform3f>	TransformAttribute;
diff --git a/Engine/Source/Runtime/Experimental/Chaos/Public/GeometryCollection/GeometryCollectionAlgo.h b/Engine/Source/Runtime/Experimental/Chaos/Public/GeometryCollection/GeometryCollectionAlgo.h
--- a/Engine/Source/Runtime/Experimental/Chaos/Public/GeometryCollection/GeometryCollectionAlgo.h
+++ b/Engine/Source/Runtime/Experimental/Chaos/Public/GeometryCollection/GeometryCollectionAlgo.h
@@ -91,28 +91,28 @@ namespace GeometryCollectionAlgo
 	*/
 	bool 
 	CHAOS_API
-	HasCycle(TManagedArray<int32>& Parents, int32 Node);
+	HasCycle(const TManagedArray<int32>& Parents, int32 Node);
 
 	/*
 	*
 	*/
 	bool 
 	CHAOS_API
-	HasCycle(TManagedArray<int32>& Parents, const TArray<int32>& SelectedBones);
+	HasCycle(const TManagedArray<int32>& Parents, const TArray<int32>& SelectedBones);
 
 	/*
 	* Parent a single transform
 	*/
 	void
 	CHAOS_API 
-	ParentTransform(FTransformCollection* GeometryCollection, const int32 TransformIndex, const int32 ChildIndex);
+	ParentTransform(FManagedArrayCollection* ManagedArrayCollection, const int32 TransformIndex, const int32 ChildIndex);
 		
 	/*
 	*  Parent the list of transforms to the selected index. 
 	*/
 	void 
 	CHAOS_API 
-	ParentTransforms(FTransformCollection* GeometryCollection, const int32 TransformIndex, const TArray<int32>& SelectedBones);
+	ParentTransforms(FManagedArrayCollection* ManagedArrayCollection, const int32 TransformIndex, const TArray<int32>& SelectedBones);
 
 
 	/*
@@ -120,7 +120,7 @@ namespace GeometryCollectionAlgo
 	*/
 	void
 	CHAOS_API
-	UnparentTransform(FManagedArrayCollection* Collection, const int32 ChildIndex);
+	UnparentTransform(FManagedArrayCollection* ManagedArrayCollection, const int32 ChildIndex);
 
 	/*
 	*  Find the average position of the transforms.
