GeometryCollection : SkeletalMeshToCollectionConversion
- Updated split mesh bone name to start from zero.

diff --git a/Engine/Plugins/Experimental/ChaosFlesh/Source/ChaosFleshNodes/Private/Dataflow/ChaosFleshCreateTetrahedronNode.cpp b/Engine/Plugins/Experimental/ChaosFlesh/Source/ChaosFleshNodes/Private/Dataflow/ChaosFleshCreateTetrahedronNode.cpp
--- a/Engine/Plugins/Experimental/ChaosFlesh/Source/ChaosFleshNodes/Private/Dataflow/ChaosFleshCreateTetrahedronNode.cpp
+++ b/Engine/Plugins/Experimental/ChaosFlesh/Source/ChaosFleshNodes/Private/Dataflow/ChaosFleshCreateTetrahedronNode.cpp
@@ -22,6 +22,7 @@
 #include "Generate/IsosurfaceStuffing.h"
 #include "GeometryCollection/ManagedArrayCollection.h"
 #include "GeometryCollection/Facades/CollectionTetrahedralMetricsFacade.h"
+#include "GeometryCollection/Facades/CollectionTransformFacade.h"
 #include "GeometryCollection/GeometryCollectionAlgo.h"
 #include "MeshDescription.h"
 #include "MeshDescriptionToDynamicMesh.h"
@@ -181,6 +182,8 @@ void FCreateTetrahedronDataflowNode::Evaluate(Dataflow::FContext& Context, const
 				}
 			};
 			AppendProcessedGeometry();
+
+			GeometryCollection::Facades::FCollectionTransformFacade(*InCollectionVal.Get()).EnforceSingleRoot("root");
 		}
 		SetValue<const DataType&>(Context, *InCollectionVal, &Collection);
 	}
diff --git a/Engine/Plugins/Experimental/GeometryCollectionPlugin/Source/GeometryCollectionNodes/Private/Dataflow/GeometryCollectionTransferVertexScalarAttributeNode.cpp b/Engine/Plugins/Experimental/GeometryCollectionPlugin/Source/GeometryCollectionNodes/Private/Dataflow/GeometryCollectionTransferVertexScalarAttributeNode.cpp
--- a/Engine/Plugins/Experimental/GeometryCollectionPlugin/Source/GeometryCollectionNodes/Private/Dataflow/GeometryCollectionTransferVertexScalarAttributeNode.cpp
+++ b/Engine/Plugins/Experimental/GeometryCollectionPlugin/Source/GeometryCollectionNodes/Private/Dataflow/GeometryCollectionTransferVertexScalarAttributeNode.cpp
@@ -288,7 +288,13 @@ void FGeometryCollectionTransferVertexScalarAttributeNode::NearestVertexTransfer
 float FGeometryCollectionTransferVertexScalarAttributeNode::MaxEdgeLength(TArray<FVector>& Vert, const TManagedArray<FIntVector3>& Tri, int VertexOFfset, int TriStart, int TriCount)
 {
 	auto TriInRange = [](const FIntVector3& T, int Max) {
-		for (int k = 0; k < 3; k++) if (T[0] < 0 || Max <= T[0]) return false;
+		for (int k = 0; k < 3; k++)
+		{
+			if (ensure(T[k] < 0 || Max <= T[k]))
+			{
+				return false;
+			}
+		}
 		return true;
 	};
 
diff --git a/Engine/Source/Runtime/Experimental/GeometryCollectionEngine/Private/GeometryCollection/GeometryCollectionEngineConversion.cpp b/Engine/Source/Runtime/Experimental/GeometryCollectionEngine/Private/GeometryCollection/GeometryCollectionEngineConversion.cpp
--- a/Engine/Source/Runtime/Experimental/GeometryCollectionEngine/Private/GeometryCollection/GeometryCollectionEngineConversion.cpp
+++ b/Engine/Source/Runtime/Experimental/GeometryCollectionEngine/Private/GeometryCollection/GeometryCollectionEngineConversion.cpp
@@ -1265,25 +1265,40 @@ bool FGeometryCollectionEngineConversion::AppendSkeletalMesh(const USkeletalMesh
 	// 
 	// Transform Attributes 
 	// 
+	auto MakeUnique = [&BoneName](FString& NewName, int32& CurrentIndex)
+	{
+		FString TestName = FString::Printf(TEXT("%s%d"), *NewName, CurrentIndex);
+		while (BoneName.Contains(TestName))
+		{
+			TestName = FString::Printf(TEXT("%s%d"), *NewName, CurrentIndex);
+			CurrentIndex++;
+		}
+		return TestName;
+	};
 
 	// add transforms for the separated geometry components. 
+	int32 SplitMeshIndex = 1;
 	TArray<int32> ComponentToTransformGroupIndex;
 	ComponentToTransformGroupIndex.Init(INDEX_NONE, ComponentsSourceIndices.Num());
 	int32 ComponentTransformBaseIndex = GeometryCollection.AddElements(ComponentsSourceIndices.Num(), FTransformCollection::TransformGroup);
 	for (int ComponentIndex = 0; ComponentIndex < ComponentsSourceIndices.Num(); ComponentIndex++)
 	{
+		FString BaseName("SplitMesh");
+
 		int32 ComponentTransformIndex = ComponentTransformBaseIndex + ComponentIndex;
 		Parent[ComponentTransformIndex] = RootIndex;
 		if (RootIndex != INDEX_NONE)
 		{
 			Children[RootIndex].Add(ComponentTransformIndex);
+			BaseName = FString::Printf(TEXT("%s_SplitMesh"), *BoneName[RootIndex]);
 		}
 
-		BoneName[ComponentTransformIndex] = FString::Printf(TEXT("%s_Mesh%d"), *BoneName[RootIndex], ComponentTransformIndex);
+		BoneName[ComponentTransformIndex] = MakeUnique(BaseName, SplitMeshIndex);
 		LocalSpaceTransform[ComponentTransformIndex] = FTransform3f::Identity;
 		SimulationType[ComponentTransformIndex] = FGeometryCollection::ESimulationTypes::FST_None;
 		BoneColor[ComponentTransformIndex] = FLinearColor::MakeRandomColor();
 		ComponentToTransformGroupIndex[ComponentIndex] = ComponentTransformIndex;
+		SplitMeshIndex++;
 	}
 
 	TArray<FTransform> ComponentTransform;
@@ -1728,6 +1743,7 @@ void FGeometryCollectionEngineConversion::ConvertCollectionToSkeleton(const FMan
 			}
 		}
 
+		Transforms.EnforceSingleRoot("root");
 		UE_LOG(UGeometryCollectionConversionLogging, Log, TEXT("FGeometryCollectionEngineConversion::ConvertCollectionToSkeleton(NumTransforms:%d)"), OutSkeleton->GetReferenceSkeleton().GetRawBoneNum());
 	}
 }
