Prevent on-demand containers to be mounted twice

diff --git a/Engine/Source/Runtime/Experimental/IoStore/OnDemand/Private/OnDemandIoStore.cpp b/Engine/Source/Runtime/Experimental/IoStore/OnDemand/Private/OnDemandIoStore.cpp
--- a/Engine/Source/Runtime/Experimental/IoStore/OnDemand/Private/OnDemandIoStore.cpp
+++ b/Engine/Source/Runtime/Experimental/IoStore/OnDemand/Private/OnDemandIoStore.cpp
@@ -5,6 +5,7 @@
 #include "OnDemandInstallCache.h"
 #include "OnDemandPackageStoreBackend.h"
 
+#include "Algo/Find.h"
 #include "Algo/RemoveIf.h"
 #include "Algo/Transform.h"
 #include "Async/ManualResetEvent.h"
@@ -584,6 +585,31 @@ FIoStatus FOnDemandIoStore::ProcessMountRequest(FMountRequest& MountRequest)
 		CreateContainersFromToc(Args.MountId, TocPath, Toc, MountRequest.Containers);
 	}
 
+	// Remove already mounted containers
+	// TODO: Treat this as an error?
+	{
+		MountRequest.Containers.SetNum(
+			Algo::RemoveIf(
+				MountRequest.Containers,
+				[MountedContainers = GetMountedContainers()](const FSharedOnDemandContainer& Container)
+				{
+					const FSharedOnDemandContainer* Existing =
+						Algo::FindBy(
+							MountedContainers,
+							Container->UniqueName(),
+							[](const FSharedOnDemandContainer& C) { return C->UniqueName(); });
+
+					if (Existing != nullptr)
+					{
+						UE_LOG(LogIoStoreOnDemand, Warning, TEXT("Container already mounted, ContainerName='%s', MountId='%s'"),
+							*WriteToString<128>(Container->Name), *WriteToString<128>(Container->MountId));
+						return true;
+					}
+					return false;
+				}));
+	}
+
+	// Find containers matching the mount ID if this is a request to install/download already mounted containers 
 	if (MountRequest.Containers.IsEmpty())
 	{
 		UE::TUniqueLock Lock(ContainerMutex);
@@ -849,4 +875,10 @@ FIoStatus FOnDemandIoStore::InstallContainers(
 	return EIoErrorCode::Ok;
 }
 
+TArray<FSharedOnDemandContainer> FOnDemandIoStore::GetMountedContainers()
+{
+	UE::TUniqueLock Lock(ContainerMutex);
+	return Containers;
+}
+
 } // namespace UE::IoStore
diff --git a/Engine/Source/Runtime/Experimental/IoStore/OnDemand/Private/OnDemandIoStore.h b/Engine/Source/Runtime/Experimental/IoStore/OnDemand/Private/OnDemandIoStore.h
--- a/Engine/Source/Runtime/Experimental/IoStore/OnDemand/Private/OnDemandIoStore.h
+++ b/Engine/Source/Runtime/Experimental/IoStore/OnDemand/Private/OnDemandIoStore.h
@@ -163,6 +163,8 @@ private:
 								const FString& Url,
 								const TConstArrayView<FSharedOnDemandContainer>& ContainersToInstall);
 
+	TArray<FSharedOnDemandContainer> GetMountedContainers();
+
 	FSharedInstallCache					InstallCache;
 	FSharedPackageStoreBackend			PackageStoreBackend;
 	FDelegateHandle						OnMountPakHandle;
