FOnDemandIoStore::InstallContainers - Fix container headers not being serialized

diff --git a/Engine/Source/Runtime/Experimental/IoStore/OnDemand/Private/OnDemandIoStore.cpp b/Engine/Source/Runtime/Experimental/IoStore/OnDemand/Private/OnDemandIoStore.cpp
--- a/Engine/Source/Runtime/Experimental/IoStore/OnDemand/Private/OnDemandIoStore.cpp
+++ b/Engine/Source/Runtime/Experimental/IoStore/OnDemand/Private/OnDemandIoStore.cpp
@@ -791,6 +791,8 @@ FIoStatus FOnDemandIoStore::InstallContainers(
 			{
 				return FIoStatusBuilder(EIoErrorCode::FileNotOpen) << TEXT("Failed to serialize container header");
 			}
+
+			Ar << ContainerHeader;
 		}
 
 		if (ContainerHeader.PackageIds.Num())
