FOnDemandIoStore
- Checking for ContainerHeader serialization errors should happen after serialization

diff --git a/Engine/Source/Runtime/Experimental/IoStore/OnDemand/Private/OnDemandIoStore.cpp b/Engine/Source/Runtime/Experimental/IoStore/OnDemand/Private/OnDemandIoStore.cpp
--- a/Engine/Source/Runtime/Experimental/IoStore/OnDemand/Private/OnDemandIoStore.cpp
+++ b/Engine/Source/Runtime/Experimental/IoStore/OnDemand/Private/OnDemandIoStore.cpp
@@ -787,12 +787,12 @@ FIoStatus FOnDemandIoStore::InstallContainers(
 			}
 
 			FMemoryReaderView Ar(Response.ValueOrDie().GetView());
+			Ar << ContainerHeader;
+			Ar.Close();
 			if (Ar.IsError() || Ar.IsCriticalError())
 			{
 				return FIoStatusBuilder(EIoErrorCode::FileNotOpen) << TEXT("Failed to serialize container header");
 			}
-
-			Ar << ContainerHeader;
 		}
 
 		if (ContainerHeader.PackageIds.Num())
