Chaos - fix potential null pointer dereferences in ChaosDebugDrawComponent

diff --git a/Engine/Source/Runtime/Experimental/ChaosSolverEngine/Private/Chaos/ChaosDebugDrawComponent.cpp b/Engine/Source/Runtime/Experimental/ChaosSolverEngine/Private/Chaos/ChaosDebugDrawComponent.cpp
--- a/Engine/Source/Runtime/Experimental/ChaosSolverEngine/Private/Chaos/ChaosDebugDrawComponent.cpp
+++ b/Engine/Source/Runtime/Experimental/ChaosSolverEngine/Private/Chaos/ChaosDebugDrawComponent.cpp
@@ -391,16 +391,19 @@ void UChaosDebugDrawComponent::TickComponent(float DeltaTime, enum ELevelTick Ti
 			DebugDrawChaos(GetOwner(), DrawCommands, World->IsPaused());
 		}
 
-		const ChaosDD::Private::FChaosDDScenePtr& CDDScene = World->GetPhysicsScene()->GetDebugDrawScene();
-		if (CDDScene.IsValid())
+		if (World->GetPhysicsScene() != nullptr)
 		{
-			CDDScene->SetDrawRegion(FSphere3d(RegionOfInterestOrigin, RegionOfInterestRadius));
-			CDDScene->SetCommandBudget(ChaosDebugDraw_MaxElements);
-
-			if (!bIsPaused)
+			const ChaosDD::Private::FChaosDDScenePtr& CDDScene = World->GetPhysicsScene()->GetDebugDrawScene();
+			if (CDDScene.IsValid())
 			{
-				FChaosDDRenderer CDDRenderer = FChaosDDRenderer(World, GetOwner());
-				CDDScene->RenderLatestFrames(CDDRenderer);
+				CDDScene->SetDrawRegion(FSphere3d(RegionOfInterestOrigin, RegionOfInterestRadius));
+				CDDScene->SetCommandBudget(ChaosDebugDraw_MaxElements);
+
+				if (!bIsPaused)
+				{
+					FChaosDDRenderer CDDRenderer = FChaosDDRenderer(World, GetOwner());
+					CDDScene->RenderLatestFrames(CDDRenderer);
+				}
 			}
 		}
 	}
