Constraints: fixed crash while using invalid constraints
    - speculative fix as I couldn't reproduce

#rnx

diff --git a/Engine/Plugins/Animation/ControlRig/Source/ControlRigEditor/Private/Editor/Constraints/SConstraintsWidget.cpp b/Engine/Plugins/Animation/ControlRig/Source/ControlRigEditor/Private/Editor/Constraints/SConstraintsWidget.cpp
--- a/Engine/Plugins/Animation/ControlRig/Source/ControlRigEditor/Private/Editor/Constraints/SConstraintsWidget.cpp
+++ b/Engine/Plugins/Animation/ControlRig/Source/ControlRigEditor/Private/Editor/Constraints/SConstraintsWidget.cpp
@@ -975,6 +975,8 @@ int32 FBaseConstraintListWidget::RefreshConstraintList()
 		return 0;
 	}
 
+	static constexpr bool bSorted = true;
+	
 	const TArray<AActor*> Selection = GetCurrentSelection();
 
 	const bool bIsConstraintsActor = Selection.Num() == 1 && Selection[0]->IsA<AConstraintsActor>();
@@ -985,14 +987,7 @@ int32 FBaseConstraintListWidget::RefreshConstraintList()
 		if (bIsConstraintsActor)
 		{
 			const FConstraintsManagerController& Controller = FConstraintsManagerController::Get(World);
-			static constexpr bool bSorted = true;
-			TArray< TObjectPtr<UTickableConstraint> > StaticConstraints;
-
-			StaticConstraints = Controller.GetStaticConstraints(bSorted);
-			for (TObjectPtr<UTickableConstraint>& Constraint : StaticConstraints)
-			{
-				Constraints.Add(Constraint);
-			}
+			Constraints = Controller.GetStaticConstraints(bSorted);
 		}
 		else
 		{
@@ -1000,6 +995,7 @@ int32 FBaseConstraintListWidget::RefreshConstraintList()
 			{
 				FTransformConstraintUtils::GetParentConstraints(World, Actor, Constraints);
 			}
+			
 			//remove if not active...
 			for (int32 Index = Constraints.Num() - 1; Index >= 0; --Index)
 			{
@@ -1013,7 +1009,6 @@ int32 FBaseConstraintListWidget::RefreshConstraintList()
 	else
 	{
 		const FConstraintsManagerController& Controller = FConstraintsManagerController::Get(World);
-		static constexpr bool bSorted = true;
 		Constraints = Controller.GetAllConstraints(bSorted);
 		if (ShowConstraints == EShowConstraints::ShowLevelSequence)
 		{
diff --git a/Engine/Source/Runtime/Experimental/Animation/Constraints/Private/ConstraintsManager.cpp b/Engine/Source/Runtime/Experimental/Animation/Constraints/Private/ConstraintsManager.cpp
--- a/Engine/Source/Runtime/Experimental/Animation/Constraints/Private/ConstraintsManager.cpp
+++ b/Engine/Source/Runtime/Experimental/Animation/Constraints/Private/ConstraintsManager.cpp
@@ -815,77 +815,30 @@ static void SortConstraints(UWorld* InWorld, TArray< TWeakObjectPtr<UTickableCon
 	InOutSortedConstraints = SortedConstraints;
 }
 
-static void SortConstraints(UWorld* InWorld, TArray< TObjectPtr<UTickableConstraint> >& InOutSortedConstraints)
-{
-	using ConstraintPtr = TObjectPtr<UTickableConstraint>;
-
-	// store original indices
-	TArray<int32> ConstraintIndices;
-	ConstraintIndices.Reserve(InOutSortedConstraints.Num());
-	for (int32 Index = 0; Index < InOutSortedConstraints.Num(); ++Index)
-	{
-		ConstraintIndices[Index] = Index;
-	}
-	
-	// LHS ticks before RHS = LHS is a prerex of RHS 
-	auto TicksBefore = [InWorld, &InOutSortedConstraints](const int32 LHSIndex, const int32 RHSIndex)
-	{
-		const ConstraintPtr& LHS = InOutSortedConstraints[LHSIndex];
-		const ConstraintPtr& RHS = InOutSortedConstraints[RHSIndex];
-		
-		const TArray<FTickPrerequisite>& RHSPrerex = RHS->GetTickFunction(InWorld).GetPrerequisites();
-		const FConstraintTickFunction& LHSTickFunction = LHS->GetTickFunction(InWorld);
-		const bool bIsLHSAPrerexOfRHS = RHSPrerex.ContainsByPredicate([&LHSTickFunction](const FTickPrerequisite& Prerex)
-		{
-			return Prerex.PrerequisiteTickFunction == &LHSTickFunction;
-		});
-
-		// if not a prerequisite then compare constraints indices
-		return bIsLHSAPrerexOfRHS ? bIsLHSAPrerexOfRHS : LHSIndex < RHSIndex;
-	};
-
-	// sort
-	Algo::Sort(ConstraintIndices, [TicksBefore](const int32 LHSIndex, const int32 RHSIndex)
-	{
-		return TicksBefore(LHSIndex, RHSIndex);
-	});
-
-	// re-order using sorted indices
-	TArray<TObjectPtr<UTickableConstraint> > SortedConstraints;
-	SortedConstraints.Reserve(InOutSortedConstraints.Num());
-	for (int32 Index = 0; Index < ConstraintIndices.Num(); ++Index)
-	{
-		SortedConstraints[Index] = InOutSortedConstraints[ConstraintIndices[Index]];
-	}
-	InOutSortedConstraints = SortedConstraints;
-}
-
-TArray< TObjectPtr<UTickableConstraint> > FConstraintsManagerController::GetStaticConstraints(const bool bSorted) const
+TArray< TWeakObjectPtr<UTickableConstraint> > FConstraintsManagerController::GetStaticConstraints(const bool bSorted) const
 {
 	UConstraintsManager* Manager = FindManager();
 	if (!Manager)
 	{
-		static const TArray< TObjectPtr<UTickableConstraint> > Empty;
+		static const TArray< TWeakObjectPtr<UTickableConstraint> > Empty;
 		return Empty;
 	}
-	TArray<TObjectPtr<UTickableConstraint>> Constraints = Manager->Constraints;
-
-	// Remove stale constraints. Stale constraints may be caused to due unexpected unloading
-	// of constributing objects, such as a level sequence.
+	
+	// Remove stale constraints. Stale constraints may be caused by unexpected unloading of contributing objects, such as a level sequence.
 	Manager->Constraints.RemoveAll([](const TObjectPtr<UTickableConstraint>& ExistingConstraint) -> bool
 	{
-		return !ExistingConstraint;
+		return !IsValid(ExistingConstraint);
 	});
 
+	TArray< TWeakObjectPtr<UTickableConstraint> > Constraints(Manager->Constraints);
 	if (!bSorted)
 	{
 		return Constraints;
 	}
 
-	TArray< TObjectPtr<UTickableConstraint> > SortedConstraints(Constraints);
-	SortConstraints(World,SortedConstraints);
+	SortConstraints(World, Constraints);
 
-	return SortedConstraints;
+	return Constraints;
 }
 
 TArray< TWeakObjectPtr<UTickableConstraint> > FConstraintsManagerController::GetAllConstraints(const bool bSorted) const
@@ -896,22 +849,23 @@ TArray< TWeakObjectPtr<UTickableConstraint> > FConstraintsManagerController::Get
 		static const TArray< TWeakObjectPtr<UTickableConstraint> > Empty;
 		return Empty;
 	}
+
 	TArray<TWeakObjectPtr<UTickableConstraint>> Constraints = Subsystem->GetConstraints(World);
+
+	// Remove stale constraints (note that GCd constraints should already have been caught by UConstraintSubsystem)
+	Constraints.RemoveAll([](const TWeakObjectPtr<UTickableConstraint>& ExistingConstraint) -> bool
+	{
+		return ExistingConstraint.IsStale();
+	});
 	
 	if (!bSorted)
 	{
 		return Constraints;
 	}
 
-	TArray< TWeakObjectPtr<UTickableConstraint> > SortedConstraints(Constraints);
-	// Remove stale constraints
-	Constraints.RemoveAll([](const TWeakObjectPtr<UTickableConstraint>& ExistingConstraint) -> bool
-	{
-		return ExistingConstraint.IsStale();
-	});
-	SortConstraints(World,SortedConstraints);
+	SortConstraints(World, Constraints);
 	
-	return SortedConstraints;
+	return Constraints;
 }
 
 void FConstraintsManagerController::EvaluateAllConstraints() const
diff --git a/Engine/Source/Runtime/Experimental/Animation/Constraints/Public/ConstraintsManager.h b/Engine/Source/Runtime/Experimental/Animation/Constraints/Public/ConstraintsManager.h
--- a/Engine/Source/Runtime/Experimental/Animation/Constraints/Public/ConstraintsManager.h
+++ b/Engine/Source/Runtime/Experimental/Animation/Constraints/Public/ConstraintsManager.h
@@ -277,7 +277,7 @@ public:
 	CONSTRAINTS_API TArray< TWeakObjectPtr<UTickableConstraint> > GetAllConstraints(const bool bSorted = false) const;
 
 	/** Returns the static/non-animated constraints*/
-	CONSTRAINTS_API TArray< TObjectPtr<UTickableConstraint> >GetStaticConstraints(const bool bSorted = false) const;
+	CONSTRAINTS_API TArray< TWeakObjectPtr<UTickableConstraint> > GetStaticConstraints(const bool bSorted = false) const;
 
 	/** Returns a filtered constraints array checking if the predicate for each element is true. */
 	template <typename Predicate>
