FOnDemandMountArgs
- remove unused InstallDirectory member

diff --git a/Engine/Source/Runtime/Experimental/IoStore/OnDemand/Public/IO/IoStoreOnDemand.h b/Engine/Source/Runtime/Experimental/IoStore/OnDemand/Public/IO/IoStoreOnDemand.h
--- a/Engine/Source/Runtime/Experimental/IoStore/OnDemand/Public/IO/IoStoreOnDemand.h
+++ b/Engine/Source/Runtime/Experimental/IoStore/OnDemand/Public/IO/IoStoreOnDemand.h
@@ -367,8 +367,6 @@ struct FOnDemandMountArgs
 	FString Url;
 	/** Serialize the TOC from the specified file path. */
 	FString FilePath;
-	/** Directory path for any additonal content downloaded as part of the TOC. */
-	FString InstallDirectory;
 	/** Mount options. */
 	EOnDemandMountOptions Options;
 };
