FEventLoop::Send - move sink to activity instead of copying.

diff --git a/Engine/Source/Runtime/Experimental/IoStore/HttpClient/Private/Loop.inl b/Engine/Source/Runtime/Experimental/IoStore/HttpClient/Private/Loop.inl
--- a/Engine/Source/Runtime/Experimental/IoStore/HttpClient/Private/Loop.inl
+++ b/Engine/Source/Runtime/Experimental/IoStore/HttpClient/Private/Loop.inl
@@ -1457,7 +1457,7 @@ FTicket FEventLoop::Send(FRequest&& Request, FTicketSink Sink, UPTRINT SinkParam
 	FActivity* Activity = nullptr;
 	Swap(Activity, Request.Ptr);
 	Activity->SinkParam = SinkParam;
-	Activity->Sink = Sink;
+	Activity->Sink = MoveTemp(Sink);
 
 	// Intercept sink calls to catch 30x status codes and follow them
 	if (Activity->bFollow30x)
