Dataflow : Properly  extract full templated types ( for arrays and other containers )

diff --git a/Engine/Source/Runtime/Core/Public/UObject/FortniteMainBranchObjectVersion.h b/Engine/Source/Runtime/Core/Public/UObject/FortniteMainBranchObjectVersion.h
--- a/Engine/Source/Runtime/Core/Public/UObject/FortniteMainBranchObjectVersion.h
+++ b/Engine/Source/Runtime/Core/Public/UObject/FortniteMainBranchObjectVersion.h
@@ -459,6 +459,9 @@ struct FFortniteMainBranchObjectVersion
 		// Target layers are now defined in the Landscape actor and not continuously synced from the assigned material.
 		LandscapeTargetLayersInLandscapeActor,
 
+		// Fix to get full name of templated type ( Tarray > TArray<Float> for example )
+		DataflowTemplatedTypeFix,
+
 		// -----<new versions can be added above this line>-------------------------------------------------
 		VersionPlusOne,
 		LatestVersion = VersionPlusOne - 1
diff --git a/Engine/Source/Runtime/Experimental/Dataflow/Core/Private/Dataflow/DataflowConnection.cpp b/Engine/Source/Runtime/Experimental/Dataflow/Core/Private/Dataflow/DataflowConnection.cpp
--- a/Engine/Source/Runtime/Experimental/Dataflow/Core/Private/Dataflow/DataflowConnection.cpp
+++ b/Engine/Source/Runtime/Experimental/Dataflow/Core/Private/Dataflow/DataflowConnection.cpp
@@ -57,11 +57,26 @@ void FDataflowConnection::SetAsAnyType(bool bAnyType, const FName& ConcreteType)
 	}
 }
 
-void FDataflowConnection::SetConcreteType(const FName& InType)
+void FDataflowConnection::SetConcreteType(FName InType)
 {
 	// Can only change from AnyType to a concrete type
-	if (ensure(IsAnyType() && IsAnyType(Type) && !IsAnyType(InType)))
+	if (ensure(IsAnyType() && !IsAnyType(InType)))
 	{
 		Type = InType;
 	}
+}
+
+void FDataflowConnection::ForceSimpleType(FName InType)
+{
+	check(Type.ToString().StartsWith(InType.ToString()));
+	Type = InType;
+}
+
+void FDataflowConnection::FixAndPropagateType()
+{
+	FString ExtendedType;
+	const FString CPPType = Property->GetCPPType(&ExtendedType);
+	FName FixedType(CPPType + ExtendedType);
+
+	FixAndPropagateType(FixedType);
 }
\ No newline at end of file
diff --git a/Engine/Source/Runtime/Experimental/Dataflow/Core/Private/Dataflow/DataflowGraph.cpp b/Engine/Source/Runtime/Experimental/Dataflow/Core/Private/Dataflow/DataflowGraph.cpp
--- a/Engine/Source/Runtime/Experimental/Dataflow/Core/Private/Dataflow/DataflowGraph.cpp
+++ b/Engine/Source/Runtime/Experimental/Dataflow/Core/Private/Dataflow/DataflowGraph.cpp
@@ -199,12 +199,34 @@ namespace Dataflow
 
 	void FGraph::SerializeForLoading(FArchive& Ar, FGraph* InGraph, UObject* OwningObject)
 	{
+		const bool bDataflowSeparateInputOutputSerialization = (Ar.CustomVer(FUE5MainStreamObjectVersion::GUID) >= FUE5MainStreamObjectVersion::DataflowSeparateInputOutputSerialization);
+		const bool bDataflowAnyTypeSupport = (Ar.CustomVer(FFortniteMainBranchObjectVersion::GUID) >= FFortniteMainBranchObjectVersion::DataflowAnyTypeSupport);
+		const bool bDataflowTemplateTypeFix = (Ar.CustomVer(FFortniteMainBranchObjectVersion::GUID) >= FFortniteMainBranchObjectVersion::DataflowTemplatedTypeFix);
+
 		FGuid ArGuid;
 		FName ArType, ArName;
 		int32 ArNum = 0;
 
 		TMap<FGuid, TSharedPtr<FDataflowNode> > NodeGuidMap;
 		TMap<FGuid, FDataflowConnection* > ConnectionGuidMap;
+		TArray<FDataflowConnection*> ConnectionsToFix;
+
+		// returns true if the connection is to be fixed
+		auto AddTemplateTypedConnectionToBeFixed = [&ConnectionsToFix, bDataflowTemplateTypeFix](FDataflowConnection* Connection, FName SerializedType) -> bool
+			{
+				if (Connection && !bDataflowTemplateTypeFix)
+				{
+					const bool bSametype = (Connection->GetType() == SerializedType);
+					const bool bIsOldTemplatedType = !bSametype && Connection->GetType().ToString().StartsWith(SerializedType.ToString());
+					if (bIsOldTemplatedType)
+					{
+						Connection->ForceSimpleType(SerializedType);
+						ConnectionsToFix.Add(Connection);
+						return true;
+					}
+				}
+				return false;
+			};
 
 		Ar << ArNum;
 		for (int32 Ndx = ArNum; Ndx > 0; Ndx--)
@@ -218,9 +240,6 @@ namespace Dataflow
 				ensure(!NodeGuidMap.Contains(ArGuid));
 				NodeGuidMap.Add(ArGuid, Node);
 
-				const bool bDataflowSeparateInputOutputSerialization = (Ar.CustomVer(FUE5MainStreamObjectVersion::GUID) >= FUE5MainStreamObjectVersion::DataflowSeparateInputOutputSerialization);
-				const bool bDataflowAnyTypeSupport = (Ar.CustomVer(FFortniteMainBranchObjectVersion::GUID) >= FFortniteMainBranchObjectVersion::DataflowAnyTypeSupport);
-
 				if (!bDataflowSeparateInputOutputSerialization)
 				{
 
@@ -263,7 +282,10 @@ namespace Dataflow
 							}
 							if (Connection)
 							{
-								check(Connection->GetType() == ArType);
+								if (!AddTemplateTypedConnectionToBeFixed(Connection, ArType))
+								{
+									check(Connection->GetType() == ArType);
+								}
 								Connection->SetGuid(ArGuid);
 								ensure(!ConnectionGuidMap.Contains(ArGuid));
 								ConnectionGuidMap.Add(ArGuid, Connection);
@@ -299,7 +321,10 @@ namespace Dataflow
 								{
 									Output->SetAsAnyType(bIsAnyType, ArType);
 								}
-								check(Output->GetType() == ArType);
+								if (!AddTemplateTypedConnectionToBeFixed(Output, ArType))
+								{
+									check(Output->GetType() == ArType || bIsAnyType);
+								}
 								Output->SetGuid(ArGuid);
 								ensure(!ConnectionGuidMap.Contains(ArGuid));
 								ConnectionGuidMap.Add(ArGuid, Output);
@@ -333,7 +358,10 @@ namespace Dataflow
 								{
 									Input->SetAsAnyType(bIsAnyType, ArType);
 								}
-								check(Input->GetType() == ArType || bIsAnyType);
+								if (!AddTemplateTypedConnectionToBeFixed(Input, ArType))
+								{
+									check(Input->GetType() == ArType || bIsAnyType);
+								}
 								Input->SetGuid(ArGuid);
 								ensure(!ConnectionGuidMap.Contains(ArGuid));
 								ConnectionGuidMap.Add(ArGuid, Input);
@@ -367,13 +395,21 @@ namespace Dataflow
 			{
 				if (ConnectionGuidMap.Contains(Con.Input) && ConnectionGuidMap.Contains(Con.Output))
 				{
-					if (ConnectionGuidMap[Con.Input]->GetType() == ConnectionGuidMap[Con.Output]->GetType())
+					FDataflowOutput* Output = static_cast<FDataflowOutput*>(ConnectionGuidMap[Con.Output]);
+					FDataflowInput* Input = static_cast<FDataflowInput*>(ConnectionGuidMap[Con.Input]);
+					if (Input->GetType() == Output->GetType())
 					{
-						InGraph->Connect(static_cast<FDataflowOutput*>(ConnectionGuidMap[Con.Output]), static_cast<FDataflowInput*>(ConnectionGuidMap[Con.Input]));
+						InGraph->Connect(Output, Input);
 					}
 				}
 			}
 		}
+
+		// fix templated types if any : see bDataflowTemplateTypeFix
+		for (FDataflowConnection* ConnectionToFix : ConnectionsToFix)
+		{
+			ConnectionToFix->FixAndPropagateType();
+		}
 	}
 }
 
diff --git a/Engine/Source/Runtime/Experimental/Dataflow/Core/Private/Dataflow/DataflowInputOutput.cpp b/Engine/Source/Runtime/Experimental/Dataflow/Core/Private/Dataflow/DataflowInputOutput.cpp
--- a/Engine/Source/Runtime/Experimental/Dataflow/Core/Private/Dataflow/DataflowInputOutput.cpp
+++ b/Engine/Source/Runtime/Experimental/Dataflow/Core/Private/Dataflow/DataflowInputOutput.cpp
@@ -78,6 +78,24 @@ void FDataflowInput::PullValue(Dataflow::FContext& Context) const
 	}
 }
 
+void FDataflowInput::FixAndPropagateType(FName InType)
+{
+	check(InType.ToString().StartsWith(Type.ToString()));
+	Type = InType;
+
+	if (FDataflowOutput* Output = GetConnection())
+	{
+		if (Output && Output->GetType() != InType)
+		{
+			if (Output->IsAnyType())
+			{
+				Output->GetOwningNode()->PropagateTypeToAllAnyTypeInputsAndOutputs(InType);
+			}
+			Output->FixAndPropagateType(InType);
+		}
+	}
+}
+
 //
 //
 //  Output
@@ -210,3 +228,21 @@ void FDataflowOutput::ForwardInput(const void* InputReference, Dataflow::FContex
 		}
 	}
 }
+
+void FDataflowOutput::FixAndPropagateType(FName InType)
+{
+	check(InType.ToString().StartsWith(Type.ToString()));
+	Type = InType;
+
+	for (FDataflowInput* Input: Connections)
+	{
+		if (Input && Input->GetType() != InType)
+		{
+			if (Input->IsAnyType())
+			{
+				Input->GetOwningNode()->PropagateTypeToAllAnyTypeInputsAndOutputs(InType);
+			}
+			Input->FixAndPropagateType(InType);
+		}
+	}
+}
\ No newline at end of file
diff --git a/Engine/Source/Runtime/Experimental/Dataflow/Core/Private/Dataflow/DataflowNode.cpp b/Engine/Source/Runtime/Experimental/Dataflow/Core/Private/Dataflow/DataflowNode.cpp
--- a/Engine/Source/Runtime/Experimental/Dataflow/Core/Private/Dataflow/DataflowNode.cpp
+++ b/Engine/Source/Runtime/Experimental/Dataflow/Core/Private/Dataflow/DataflowNode.cpp
@@ -315,6 +315,26 @@ void FDataflowNode::UnregisterPinConnection(const Dataflow::FPin& Pin)
 	}
 }
 
+void FDataflowNode::PropagateTypeToAllAnyTypeInputsAndOutputs(FName ConcreteType)
+{
+	for (TPair<uint32, FDataflowOutput*> Elem : Outputs)
+	{
+		FDataflowOutput* Output = Elem.Value;
+		if (Output->IsAnyType())
+		{
+			Output->SetConcreteType(ConcreteType);
+		}
+	}
+	for (TPair<uint32, FDataflowInput*> Elem : Inputs)
+	{
+		FDataflowInput* Input = Elem.Value;
+		if (Input->IsAnyType())
+		{
+			Input->SetConcreteType(ConcreteType);
+		}
+	}
+}
+
 void FDataflowNode::Invalidate(const Dataflow::FTimestamp& InModifiedTimestamp)
 {
 	if (LastModifiedTimestamp < InModifiedTimestamp)
@@ -440,12 +460,13 @@ void FDataflowNode::RegisterInputConnection(const void* InProperty, const FName&
 		if (const UStruct* Struct = ScriptOnStruct->GetStruct())
 		{
 			TArray<const FProperty*> PropertyChain;
-			const FProperty* const Property =
-				FindProperty(Struct, InProperty, PropertyName, &PropertyChain);
+			const FProperty* const Property = FindProperty(Struct, InProperty, PropertyName, &PropertyChain);
 			if (ensure(Property && PropertyChain.Num()))
 			{
 				const FName PropName(GetPropertyFullName(PropertyChain));
-				const FName PropType(Property->GetCPPType());
+				FString ExtendedType;
+				const FString CPPType = Property->GetCPPType(&ExtendedType);
+				const FName PropType(CPPType + ExtendedType);
 				AddInput(new FDataflowInput({ PropType, PropName, this, Property }));
 			}
 		}
@@ -490,7 +511,9 @@ void FDataflowNode::RegisterOutputConnection(const void* InProperty, const void*
 			if (ensure(Property && PropertyChain.Num()))
 			{
 				const FName PropName(GetPropertyFullName(PropertyChain));
-				const FName PropType(Property->GetCPPType());
+				FString ExtendedType;
+				const FString CPPType = Property->GetCPPType(&ExtendedType);
+				const FName PropType(CPPType + ExtendedType);
 				OutputConnection = new FDataflowOutput({ PropType, PropName, this, Property });
 
 				TArray<const FProperty*> PassthroughPropertyChain;
diff --git a/Engine/Source/Runtime/Experimental/Dataflow/Core/Private/Dataflow/DataflowNodeParameters.cpp b/Engine/Source/Runtime/Experimental/Dataflow/Core/Private/Dataflow/DataflowNodeParameters.cpp
--- a/Engine/Source/Runtime/Experimental/Dataflow/Core/Private/Dataflow/DataflowNodeParameters.cpp
+++ b/Engine/Source/Runtime/Experimental/Dataflow/Core/Private/Dataflow/DataflowNodeParameters.cpp
@@ -117,7 +117,9 @@ namespace Dataflow
 				if (Elem.Value && Elem.Value->Property && Elem.Value->Type == FContextCacheElementBase::EType::CacheElementTyped)
 				{
 					FProperty* Property = (FProperty*)Elem.Value->Property;
-					FName TypeName(Elem.Value->Property->GetCPPType());
+					FString ExtendedType;
+					const FString CPPType = Property->GetCPPType(&ExtendedType);
+					FName TypeName(CPPType + ExtendedType);
 					FGuid NodeGuid = Elem.Value->NodeGuid;
 					uint32 NodeHash = Elem.Value->NodeHash;
 
diff --git a/Engine/Source/Runtime/Experimental/Dataflow/Core/Public/Dataflow/DataflowConnection.h b/Engine/Source/Runtime/Experimental/Dataflow/Core/Public/Dataflow/DataflowConnection.h
--- a/Engine/Source/Runtime/Experimental/Dataflow/Core/Public/Dataflow/DataflowConnection.h
+++ b/Engine/Source/Runtime/Experimental/Dataflow/Core/Public/Dataflow/DataflowConnection.h
@@ -65,6 +65,11 @@ protected:
 	/** this should only be used for serialization */
 	DATAFLOWCORE_API void SetAsAnyType(bool bAnyType, const FName& ConcreteType);
 
+	/** this should only be used for serialization - for support of old simple TArray types*/
+	DATAFLOWCORE_API void ForceSimpleType(FName InType);
+	DATAFLOWCORE_API void FixAndPropagateType();
+	DATAFLOWCORE_API virtual void FixAndPropagateType(FName InType) { ensure(false); }
+
 public:
 	FDataflowConnection() {};
 	DATAFLOWCORE_API FDataflowConnection(Dataflow::FPin::EDirection Direction, FName InType, FName InName, FDataflowNode* OwningNode = nullptr, const FProperty* InProperty = nullptr, FGuid InGuid = FGuid::NewGuid());
@@ -95,7 +100,7 @@ public:
 	DATAFLOWCORE_API bool IsAnyType() const { return bIsAnyType; }
 	static bool IsAnyType(const FName& InType);
 
-	DATAFLOWCORE_API void SetConcreteType(const FName& InType);
+	DATAFLOWCORE_API void SetConcreteType(FName InType);
 
 	template<class T>
 	bool IsA(const T* InVar) const
diff --git a/Engine/Source/Runtime/Experimental/Dataflow/Core/Public/Dataflow/DataflowInputOutput.h b/Engine/Source/Runtime/Experimental/Dataflow/Core/Public/Dataflow/DataflowInputOutput.h
--- a/Engine/Source/Runtime/Experimental/Dataflow/Core/Public/Dataflow/DataflowInputOutput.h
+++ b/Engine/Source/Runtime/Experimental/Dataflow/Core/Public/Dataflow/DataflowInputOutput.h
@@ -42,6 +42,11 @@ struct FDataflowInput : public FDataflowConnection
 	friend struct FDataflowConnection;
 
 	FDataflowOutput* Connection;
+
+protected:
+	friend struct FDataflowOutput;
+	virtual void FixAndPropagateType(FName InType) override;
+
 public:
 	FDataflowInput(const Dataflow::FInputParameters& Param = {}, FGuid InGuid = FGuid::NewGuid());
 
@@ -103,6 +108,10 @@ struct FDataflowOutput : public FDataflowConnection
 
 	uint32 PassthroughOffset = INDEX_NONE;
 
+protected:
+	friend struct FDataflowInput;
+	virtual void FixAndPropagateType(FName InType) override;
+
 public:
 	static DATAFLOWCORE_API FDataflowOutput NoOpOutput;
 	
diff --git a/Engine/Source/Runtime/Experimental/Dataflow/Core/Public/Dataflow/DataflowNode.h b/Engine/Source/Runtime/Experimental/Dataflow/Core/Public/Dataflow/DataflowNode.h
--- a/Engine/Source/Runtime/Experimental/Dataflow/Core/Public/Dataflow/DataflowNode.h
+++ b/Engine/Source/Runtime/Experimental/Dataflow/Core/Public/Dataflow/DataflowNode.h
@@ -173,6 +173,8 @@ struct FDataflowNode
 	/** Unregister the connection if one exists matching this pin, then invalidate the graph. */
 	DATAFLOWCORE_API void UnregisterPinConnection(const Dataflow::FPin& Pin);
 
+	DATAFLOWCORE_API void PropagateTypeToAllAnyTypeInputsAndOutputs(FName ConcreteType);
+
 	//
 	// Evaluation
 	//
diff --git a/Engine/Source/Runtime/Experimental/Dataflow/Engine/Private/Dataflow/DataflowEdNode.cpp b/Engine/Source/Runtime/Experimental/Dataflow/Engine/Private/Dataflow/DataflowEdNode.cpp
--- a/Engine/Source/Runtime/Experimental/Dataflow/Engine/Private/Dataflow/DataflowEdNode.cpp
+++ b/Engine/Source/Runtime/Experimental/Dataflow/Engine/Private/Dataflow/DataflowEdNode.cpp
@@ -114,17 +114,37 @@ void UDataflowEdNode::UpdatePinsFromDataflowNode()
 					{
 						if (Pin->Direction == EEdGraphPinDirection::EGPD_Input)
 						{
-							if (!DataflowNode->FindInput(Pin->GetFName()))
+							const FDataflowInput* DataflowInput = DataflowNode->FindInput(Pin->GetFName());
+							if (!DataflowInput)
 							{
 								PinsToRemove.Add(Pin);
 							}
+							else if (DataflowInput->GetType() != Pin->PinType.PinCategory)
+							{
+								Pin->PinType = FEdGraphPinType();
+								Pin->PinType.bIsReference = false;
+								Pin->PinType.bIsConst = false;
+								Pin->PinType.PinCategory = DataflowInput->GetType();
+								Pin->PinType.PinSubCategory = NAME_None;
+								Pin->PinType.PinSubCategoryObject = nullptr;
+							}
 						}
 						else if (Pin->Direction == EEdGraphPinDirection::EGPD_Output)
 						{
-							if (!DataflowNode->FindOutput(Pin->GetFName()))
+							const FDataflowOutput* DataflowOutput = DataflowNode->FindOutput(Pin->GetFName());
+							if (!DataflowOutput)
 							{
 								PinsToRemove.Add(Pin);
 							}
+							else if (DataflowOutput->GetType() != Pin->PinType.PinCategory)
+							{
+								Pin->PinType = FEdGraphPinType();
+								Pin->PinType.bIsReference = false;
+								Pin->PinType.bIsConst = false;
+								Pin->PinType.PinCategory = DataflowOutput->GetType();
+								Pin->PinType.PinSubCategory = NAME_None;
+								Pin->PinType.PinSubCategoryObject = nullptr;
+							}
 						}
 					}
 				}
