Iris
* NetObjectFilter now owns the FilteredObjects list
* NetObjectFilter's now need to set the NeedsUpdate trait if they want to receive UpdateObjects callback telling them of dirty objects
* NetObjectFilter::UpdateObjects is no longer pure virtual
* NetObjectFilter's OnInit params is now const
* Added some safety tests by using ArrayView's instead of raw buffers

diff --git a/Engine/Plugins/Runtime/ReplicationSystemTestPlugin/Source/Private/Tests/ReplicationSystem/Filtering/MockNetObjectFilter.cpp b/Engine/Plugins/Runtime/ReplicationSystemTestPlugin/Source/Private/Tests/ReplicationSystem/Filtering/MockNetObjectFilter.cpp
--- a/Engine/Plugins/Runtime/ReplicationSystemTestPlugin/Source/Private/Tests/ReplicationSystem/Filtering/MockNetObjectFilter.cpp
+++ b/Engine/Plugins/Runtime/ReplicationSystemTestPlugin/Source/Private/Tests/ReplicationSystem/Filtering/MockNetObjectFilter.cpp
@@ -13,8 +13,10 @@ UMockNetObjectFilter::UMockNetObjectFilter()
 {
 }
 
-void UMockNetObjectFilter::OnInit(FNetObjectFilterInitParams& Params)
+void UMockNetObjectFilter::OnInit(const FNetObjectFilterInitParams& Params)
 {
+	AddFilterTraits(ENetFilterTraits::NeedsUpdate);
+
 	++CallStatus.CallCounts.Init;
 
 	CallStatus.SuccessfulCallCounts.Init += Cast<UMockNetObjectFilterConfig>(Params.Config) != nullptr;
@@ -116,7 +118,7 @@ void UMockNetObjectFilter::Filter(FNetObjectFilteringParams& Params)
 	++CallStatus.CallCounts.Filter;
 
 	bool bIsProperCall = true;
-	UE::Net::FNetBitArrayView::ForAllExclusiveBits(Params.FilteredObjects, MakeNetBitArrayView(AddedObjectIndices), [](...) {}, [&bIsProperCall](...) {bIsProperCall = false; });
+	UE::Net::FNetBitArray::ForAllExclusiveBits(FilteredObjects, AddedObjectIndices, [](...) {}, [&bIsProperCall](...) {bIsProperCall = false; });
 	CallStatus.SuccessfulCallCounts.Filter += bIsProperCall;
 
 	if (CallSetup.Filter.bFilterOutByDefault)
@@ -141,7 +143,7 @@ void UMockNetObjectFilter::PostFilter(FNetObjectPostFilteringParams& Params)
 // UMockNetObjectFilterWithCondition
 //**************************************************************************************************
 
-void UMockNetObjectFilterWithCondition::OnInit(FNetObjectFilterInitParams& Params)
+void UMockNetObjectFilterWithCondition::OnInit(const FNetObjectFilterInitParams& Params)
 {
 	Super::OnInit(Params);
 
@@ -165,7 +167,7 @@ void UMockNetObjectFilterWithCondition::Filter(FNetObjectFilteringParams& Params
 
 	bool bIsProperCall = true;
 
-	Params.FilteredObjects.ForAllSetBits([&](uint32 ObjectIndex)
+	FilteredObjects.ForAllSetBits([&](uint32 ObjectIndex)
 	{
 		UTestFilteringObject* FilterObject = CastChecked<UTestFilteringObject>(ReplicationSystem->GetReplicationSystemInternal()->GetNetRefHandleManager().GetReplicatedObjectInstance(ObjectIndex));
 		
diff --git a/Engine/Plugins/Runtime/ReplicationSystemTestPlugin/Source/Private/Tests/ReplicationSystem/Filtering/MockNetObjectFilter.h b/Engine/Plugins/Runtime/ReplicationSystemTestPlugin/Source/Private/Tests/ReplicationSystem/Filtering/MockNetObjectFilter.h
--- a/Engine/Plugins/Runtime/ReplicationSystemTestPlugin/Source/Private/Tests/ReplicationSystem/Filtering/MockNetObjectFilter.h
+++ b/Engine/Plugins/Runtime/ReplicationSystemTestPlugin/Source/Private/Tests/ReplicationSystem/Filtering/MockNetObjectFilter.h
@@ -66,7 +66,7 @@ public:
 	inline void ResetFunctionCallStatus() { CallStatus = FFunctionCallStatus({}); }
 
 protected:
-	virtual void OnInit(FNetObjectFilterInitParams&) override;
+	virtual void OnInit(const FNetObjectFilterInitParams&) override;
 	virtual void AddConnection(uint32 ConnectionId) override;
 	virtual void RemoveConnection(uint32 ConnectionId) override;
 	virtual bool AddObject(uint32 ObjectIndex, FNetObjectFilterAddObjectParams&) override;
@@ -98,7 +98,7 @@ class UMockNetObjectFilterWithCondition : public UMockNetObjectFilter
 
 protected:
 
-	virtual void OnInit(FNetObjectFilterInitParams&) override;
+	virtual void OnInit(const FNetObjectFilterInitParams&) override;
 	virtual bool AddObject(uint32 ObjectIndex, FNetObjectFilterAddObjectParams&) override;
 	virtual void Filter(FNetObjectFilteringParams&) override;
 
diff --git a/Engine/Source/Runtime/Experimental/Iris/Core/Private/Iris/ReplicationSystem/Filtering/FilterOutNetObjectFilter.cpp b/Engine/Source/Runtime/Experimental/Iris/Core/Private/Iris/ReplicationSystem/Filtering/FilterOutNetObjectFilter.cpp
--- a/Engine/Source/Runtime/Experimental/Iris/Core/Private/Iris/ReplicationSystem/Filtering/FilterOutNetObjectFilter.cpp
+++ b/Engine/Source/Runtime/Experimental/Iris/Core/Private/Iris/ReplicationSystem/Filtering/FilterOutNetObjectFilter.cpp
@@ -2,7 +2,7 @@
 
 #include "Iris/ReplicationSystem/Filtering/FilterOutNetObjectFilter.h"
 
-void UFilterOutNetObjectFilter::OnInit(FNetObjectFilterInitParams& Params)
+void UFilterOutNetObjectFilter::OnInit(const FNetObjectFilterInitParams& Params)
 {
 }
 
@@ -15,10 +15,6 @@ void UFilterOutNetObjectFilter::RemoveObject(uint32 ObjectIndex, const FNetObjec
 {
 }
 
-void UFilterOutNetObjectFilter::UpdateObjects(FNetObjectFilterUpdateParams& Params)
-{
-}
-
 void UFilterOutNetObjectFilter::Filter(FNetObjectFilteringParams& Params)
 {
 	// Filter out everything
diff --git a/Engine/Source/Runtime/Experimental/Iris/Core/Private/Iris/ReplicationSystem/Filtering/NetObjectFilter.cpp b/Engine/Source/Runtime/Experimental/Iris/Core/Private/Iris/ReplicationSystem/Filtering/NetObjectFilter.cpp
--- a/Engine/Source/Runtime/Experimental/Iris/Core/Private/Iris/ReplicationSystem/Filtering/NetObjectFilter.cpp
+++ b/Engine/Source/Runtime/Experimental/Iris/Core/Private/Iris/ReplicationSystem/Filtering/NetObjectFilter.cpp
@@ -3,28 +3,17 @@
 #include "Iris/ReplicationSystem/Filtering/NetObjectFilter.h"
 #include "Iris/ReplicationSystem/Filtering/ReplicationFiltering.h"
 
-FNetObjectFilteringParams::FNetObjectFilteringParams(const UE::Net::FNetBitArrayView InFilteredObjects)
-: FilteredObjects(InFilteredObjects)
-, FilteringInfos(nullptr)
-, ConnectionId(0)
-{
-}
-
-FNetObjectPreFilteringParams::FNetObjectPreFilteringParams(const UE::Net::FNetBitArrayView InFilteredObjects)
-: FilteredObjects(InFilteredObjects)
-{
-}
-
 UNetObjectFilter::UNetObjectFilter()
 {
 }
 
-void UNetObjectFilter::Init(FNetObjectFilterInitParams& Params)
+void UNetObjectFilter::Init(const FNetObjectFilterInitParams& Params)
 {
+	FilteredObjects.Init(Params.MaxObjectCount);
+
 	{
 		UE::Net::Private::FNetObjectFilteringInfoAccessor FilteringInfoAccessor;
-		TArrayView<FNetObjectFilteringInfo> NetObjectFilteringInfos = FilteringInfoAccessor.GetNetObjectFilteringInfos(Params.ReplicationSystem);
-		FilterInfo = FFilterInfo(Params.FilteredObjects, NetObjectFilteringInfos);
+		FilteringInfos = FilteringInfoAccessor.GetNetObjectFilteringInfos(Params.ReplicationSystem);
 	}
 
 	OnInit(Params);
@@ -38,37 +27,26 @@ void UNetObjectFilter::RemoveConnection(uint32 ConnectionId)
 {
 }
 
-void UNetObjectFilter::PreFilter(FNetObjectPreFilteringParams&)
-{
-}
-
-void UNetObjectFilter::Filter(FNetObjectFilteringParams&)
+void UNetObjectFilter::UpdateObjects(FNetObjectFilterUpdateParams&)
 {
 }
 
-void UNetObjectFilter::PostFilter(FNetObjectPostFilteringParams&)
+void UNetObjectFilter::PreFilter(FNetObjectPreFilteringParams&)
 {
 }
 
-// FFilterInfo implementation
-UNetObjectFilter::FFilterInfo::FFilterInfo(const UE::Net::FNetBitArrayView InFilteredObjects, const TArrayView<FNetObjectFilteringInfo> InFilteringInfos)
-: FilteredObjects(InFilteredObjects)
-, FilteringInfos(InFilteringInfos)
+void UNetObjectFilter::Filter(FNetObjectFilteringParams&)
 {
 }
 
-UNetObjectFilter::FFilterInfo& UNetObjectFilter::FFilterInfo::operator=(const UNetObjectFilter::FFilterInfo& Other)
+void UNetObjectFilter::PostFilter(FNetObjectPostFilteringParams&)
 {
-	// Re-initialize this instance.
-	this->~FFilterInfo();
-	new (this) FFilterInfo(Other.FilteredObjects, Other.FilteringInfos);
-	return *this;
 }
 
-FNetObjectFilteringInfo* UNetObjectFilter::FFilterInfo::GetFilteringInfo(uint32 ObjectIndex)
+FNetObjectFilteringInfo* UNetObjectFilter::GetFilteringInfo(uint32 ObjectIndex)
 {
-	// Only allow retriving infos for objects handled by this instance.
-	if (!IsAddedToFilter(ObjectIndex))
+	// Only allow retreiving infos for objects handled by this instance.
+	if (!IsObjectFiltered(ObjectIndex))
 	{
 		return nullptr;
 	}
diff --git a/Engine/Source/Runtime/Experimental/Iris/Core/Private/Iris/ReplicationSystem/Filtering/NetObjectGridFilter.cpp b/Engine/Source/Runtime/Experimental/Iris/Core/Private/Iris/ReplicationSystem/Filtering/NetObjectGridFilter.cpp
--- a/Engine/Source/Runtime/Experimental/Iris/Core/Private/Iris/ReplicationSystem/Filtering/NetObjectGridFilter.cpp
+++ b/Engine/Source/Runtime/Experimental/Iris/Core/Private/Iris/ReplicationSystem/Filtering/NetObjectGridFilter.cpp
@@ -12,7 +12,7 @@
 #include "Iris/Core/IrisLog.h"
 #include "Iris/Core/IrisProfiler.h"
 
-void UNetObjectGridFilter::OnInit(FNetObjectFilterInitParams& Params)
+void UNetObjectGridFilter::OnInit(const FNetObjectFilterInitParams& Params)
 {
 	AddFilterTraits(ENetFilterTraits::Spatial);
 
@@ -555,17 +555,13 @@ FString UNetObjectGridFilter::PrintDebugInfoForObject(const FDebugInfoParams& Pa
 // UNetObjectGridWorldLocFilter
 //*************************************************************************************************
 
-void UNetObjectGridWorldLocFilter::OnInit(FNetObjectFilterInitParams& Params)
+void UNetObjectGridWorldLocFilter::OnInit(const FNetObjectFilterInitParams& Params)
 {
 	Super::OnInit(Params);
 
 	WorldLocations = &Params.ReplicationSystem->GetWorldLocations();
 }
 
-void UNetObjectGridWorldLocFilter::UpdateObjects(FNetObjectFilterUpdateParams&)
-{
-}
-
 void UNetObjectGridWorldLocFilter::PreFilter(FNetObjectPreFilteringParams& Params)
 {
 	Super::PreFilter(Params);
@@ -579,7 +575,7 @@ void UNetObjectGridWorldLocFilter::PreFilter(FNetObjectPreFilteringParams& Param
 
 	// Update cell info for all objects that have moved.
 	UE::Net::FNetBitArrayView ObjectsWithDirtyWorldLocations = WorldLocations->GetObjectsWithDirtyInfo();
-	UE::Net::FNetBitArrayView::ForAllSetBits(Params.FilteredObjects, ObjectsWithDirtyWorldLocations, UE::Net::FNetBitArrayBase::AndOp, UpdateCells);
+	UE::Net::FNetBitArrayView::ForAllSetBits(MakeNetBitArrayView(FilteredObjects), ObjectsWithDirtyWorldLocations, UE::Net::FNetBitArrayBase::AndOp, UpdateCells);
 }
 
 void UNetObjectGridWorldLocFilter::UpdateObjectInfo(UNetObjectGridFilter::FPerObjectInfo& PerObjectInfo, const UNetObjectGridFilter::FObjectLocationInfo& ObjectLocationInfo, const UE::Net::FReplicationInstanceProtocol* InstanceProtocol)
diff --git a/Engine/Source/Runtime/Experimental/Iris/Core/Private/Iris/ReplicationSystem/Filtering/NopNetObjectFilter.cpp b/Engine/Source/Runtime/Experimental/Iris/Core/Private/Iris/ReplicationSystem/Filtering/NopNetObjectFilter.cpp
--- a/Engine/Source/Runtime/Experimental/Iris/Core/Private/Iris/ReplicationSystem/Filtering/NopNetObjectFilter.cpp
+++ b/Engine/Source/Runtime/Experimental/Iris/Core/Private/Iris/ReplicationSystem/Filtering/NopNetObjectFilter.cpp
@@ -2,7 +2,7 @@
 
 #include "Iris/ReplicationSystem/Filtering/NopNetObjectFilter.h"
 
-void UNopNetObjectFilter::OnInit(FNetObjectFilterInitParams& Params)
+void UNopNetObjectFilter::OnInit(const FNetObjectFilterInitParams& Params)
 {
 }
 
@@ -15,10 +15,6 @@ void UNopNetObjectFilter::RemoveObject(uint32 ObjectIndex, const FNetObjectFilte
 {
 }
 
-void UNopNetObjectFilter::UpdateObjects(FNetObjectFilterUpdateParams& Params)
-{
-}
-
 void UNopNetObjectFilter::Filter(FNetObjectFilteringParams& Params)
 {
 	// Allow all objects to be replicated, just as if no dynamic filtering was applied.
diff --git a/Engine/Source/Runtime/Experimental/Iris/Core/Private/Iris/ReplicationSystem/Filtering/ReplicationFiltering.cpp b/Engine/Source/Runtime/Experimental/Iris/Core/Private/Iris/ReplicationSystem/Filtering/ReplicationFiltering.cpp
--- a/Engine/Source/Runtime/Experimental/Iris/Core/Private/Iris/ReplicationSystem/Filtering/ReplicationFiltering.cpp
+++ b/Engine/Source/Runtime/Experimental/Iris/Core/Private/Iris/ReplicationSystem/Filtering/ReplicationFiltering.cpp
@@ -121,7 +121,7 @@ public:
 		}
 	}
 
-	void PrepareBatch(const uint32* ObjectIndices, uint32 ObjectCount, const uint8* FilterIndices)
+	void PrepareBatch(const uint32* ObjectIndices, uint32 ObjectCount, const TArray<uint8>& FilterIndices)
 	{
 		ResetBatch();
 
@@ -171,6 +171,7 @@ FReplicationFiltering::FReplicationFiltering()
 , bHasDynamicFilters(0)
 , bHasDirtyExclusionFilterGroup(0)
 , bHasDirtyInclusionFilterGroup(0)
+, bHasDynamicFiltersWithUpdateTrait(0)
 {
 	StaticChecks();
 }
@@ -422,7 +423,7 @@ bool FReplicationFiltering::SetFilter(FInternalNetRefIndex ObjectIndex, FNetObje
 		if (FilterInfo.Filter->AddObject(ObjIndex, AddParams))
 		{
 			++FilterInfo.ObjectCount;
-			FilterInfo.FilteredObjects.SetBit(ObjIndex);
+			FilterInfo.Filter->GetFilteredObjects().SetBit(ObjIndex);
 			this->ObjectIndexToDynamicFilterIndex[ObjIndex] = static_cast<uint8>(FilterIndex);
 			this->DynamicFilterEnabledObjects.SetBit(ObjIndex);
 			return true;
@@ -1224,7 +1225,7 @@ void FReplicationFiltering::PreUpdateDynamicFiltering()
 				continue;
 			}
 
-			FNetObjectPreFilteringParams PreFilteringParams(MakeNetBitArrayView(Info.FilteredObjects));
+			FNetObjectPreFilteringParams PreFilteringParams;
 			PreFilteringParams.FilteringInfos = MakeArrayView(NetObjectFilteringInfos);
 			PreFilteringParams.ValidConnections = MakeNetBitArrayView(ValidConnections);
 			Info.Filter->PreFilter(PreFilteringParams);
@@ -1285,16 +1286,16 @@ void FReplicationFiltering::UpdateDynamicFiltering()
 					continue;
 				}
 
-				FNetObjectFilteringParams FilteringParams(MakeNetBitArrayView(Info.FilteredObjects));
+				FNetObjectFilteringParams FilteringParams;
 				FilteringParams.OutAllowedObjects = AllowedObjects;
-				FilteringParams.FilteringInfos = NetObjectFilteringInfos.GetData();
+				FilteringParams.FilteringInfos = MakeArrayView(NetObjectFilteringInfos);
 				FilteringParams.ConnectionId = ConnId;
 				FilteringParams.View = Connections->GetReplicationView(ConnId);
 
 				// Execute the filter here
 				Info.Filter->Filter(FilteringParams);
 
-				const uint32* FilteredObjectsData = Info.FilteredObjects.GetData();
+				const uint32* FilteredObjectsData = Info.Filter->GetFilteredObjects().GetData();
 				for (SIZE_T WordIt = 0, WordEndIt = WordCountForObjectBitArrays; WordIt != WordEndIt; ++WordIt)
 				{
 					const uint32 FilteredObjects = FilteredObjectsData[WordIt];
@@ -1439,6 +1440,12 @@ void FReplicationFiltering::PostUpdateDynamicFiltering()
 
 void FReplicationFiltering::NotifyFiltersOfDirtyObjects()
 {
+	if (!bHasDynamicFiltersWithUpdateTrait)
+	{
+		// No filters use UpdateObjects so we can skip early.
+		return;
+	}
+
 	IRIS_PROFILER_SCOPE(FReplicationFiltering_UpdateFilterWithDirtyObjects);
 
 	FDirtyObjectsAccessor DirtyObjectsAccessor(ReplicationSystem->GetReplicationSystemInternal()->GetDirtyNetObjectTracker());
@@ -1464,12 +1471,11 @@ void FReplicationFiltering::NotifyFiltersOfDirtyObjects()
 
 void FReplicationFiltering::BatchNotifyFiltersOfDirtyObjects(FUpdateDirtyObjectsBatchHelper& BatchHelper, const uint32* DirtyObjectIndices, uint32 ObjectCount)
 {
-	BatchHelper.PrepareBatch(DirtyObjectIndices, ObjectCount, ObjectIndexToDynamicFilterIndex.GetData());
+	BatchHelper.PrepareBatch(DirtyObjectIndices, ObjectCount, ObjectIndexToDynamicFilterIndex);
 
 	FNetObjectFilterUpdateParams UpdateParameters;
-	UpdateParameters.FilteringInfos = NetObjectFilteringInfos.GetData();
+	UpdateParameters.FilteringInfos = MakeArrayView(NetObjectFilteringInfos);
 
-	// $IRIS TODO: We should probably have a trait asking if the Filter needs to receive UpdateObjects.
 	for (const FUpdateDirtyObjectsBatchHelper::FPerFilterInfo& PerFilterInfo : BatchHelper.PerFilterInfos)
 	{
 		if (PerFilterInfo.ObjectCount == 0)
@@ -1479,7 +1485,6 @@ void FReplicationFiltering::BatchNotifyFiltersOfDirtyObjects(FUpdateDirtyObjects
 
 		UpdateParameters.ObjectIndices = PerFilterInfo.ObjectIndices;
 		UpdateParameters.ObjectCount = PerFilterInfo.ObjectCount;
-		// $IRIS TODO: Add a trait asking if the Filter needs to receive UpdateObjects. Most do nothing in UpdateObjects()
 		const int32 FilterIndex = static_cast<int32>(&PerFilterInfo - BatchHelper.PerFilterInfos.GetData());
 		UNetObjectFilter* Filter = DynamicFilterInfos[FilterIndex].Filter.Get();
 		Filter->UpdateObjects(UpdateParameters);
@@ -2393,9 +2398,8 @@ void FReplicationFiltering::InitFilters()
 		check(Info.Filter.IsValid());
 		Info.Name = FilterDefinition.FilterName;
 		Info.ObjectCount = 0;
-		Info.FilteredObjects.Init(MaxObjectCount);
 
-		FNetObjectFilterInitParams InitParams(MakeNetBitArrayView(Info.FilteredObjects));
+		FNetObjectFilterInitParams InitParams;
 		InitParams.ReplicationSystem = ReplicationSystem;
 		InitParams.Config = (NetObjectFilterConfigClass ? NewObject<UNetObjectFilterConfig>((UObject*)GetTransientPackage(), NetObjectFilterConfigClass) : nullptr);
 		InitParams.MaxObjectCount = MaxObjectCount;
@@ -2404,6 +2408,7 @@ void FReplicationFiltering::InitFilters()
 		Info.Filter->Init(InitParams);
 
 		bHasDynamicFilters = true;
+		bHasDynamicFiltersWithUpdateTrait = bHasDynamicFiltersWithUpdateTrait || Info.Filter->HasFilterTrait(ENetFilterTraits::NeedsUpdate);
 	}
 }
 
@@ -2416,7 +2421,7 @@ void FReplicationFiltering::RemoveFromDynamicFilter(uint32 ObjectIndex, uint32 F
 	FNetObjectFilteringInfo& NetObjectFilteringInfo = NetObjectFilteringInfos[ObjectIndex];
 	FFilterInfo& FilterInfo = DynamicFilterInfos[FilterIndex];
 	--FilterInfo.ObjectCount;
-	FilterInfo.FilteredObjects.ClearBit(ObjectIndex);
+	FilterInfo.Filter->GetFilteredObjects().ClearBit(ObjectIndex);
 	FilterInfo.Filter->RemoveObject(ObjectIndex, NetObjectFilteringInfo);
 
 	DynamicFilterEnabledObjects.ClearBit(ObjectIndex);
@@ -2477,7 +2482,7 @@ FString FReplicationFiltering::PrintFilterObjectInfo(FInternalNetRefIndex Object
 
 	const FFilterInfo& FilterInfo = DynamicFilterInfos[DynamicFilterIndex];
 
-	if (!FilterInfo.FilteredObjects.IsBitSet(ObjectIndex))
+	if (!FilterInfo.Filter->GetFilteredObjects().IsBitSet(ObjectIndex))
 	{
 		ensureMsgf(false, TEXT("Problem with Filter configs for %s.  DynamicIndex %u Filter %s but not in FilteredObjects list"), *NetRefHandleManager->PrintObjectFromIndex(ObjectIndex), DynamicFilterIndex, *FilterInfo.Name.ToString());
 		return TEXT("[WrongSettings]");
@@ -2485,7 +2490,7 @@ FString FReplicationFiltering::PrintFilterObjectInfo(FInternalNetRefIndex Object
 
 	UNetObjectFilter::FDebugInfoParams DebugParams;
 	DebugParams.FilterName = FilterInfo.Name;
-	DebugParams.FilteringInfos = NetObjectFilteringInfos.GetData();
+	DebugParams.FilteringInfos = MakeArrayView(NetObjectFilteringInfos);
 	DebugParams.ConnectionId = ConnectionId;
 	DebugParams.View = Connections->GetReplicationView(ConnectionId);
 
@@ -2498,7 +2503,7 @@ void FReplicationFiltering::BuildObjectsInFilterList(FNetBitArrayView OutObjects
 	{
 		if (FilterInfo.Name == FilterName)
 		{
-			OutObjectsInFilter.Copy(FilterInfo.FilteredObjects);
+			OutObjectsInFilter.Copy(FilterInfo.Filter->GetFilteredObjects());
 			return;
 		}
 	}
diff --git a/Engine/Source/Runtime/Experimental/Iris/Core/Private/Iris/ReplicationSystem/Filtering/ReplicationFiltering.h b/Engine/Source/Runtime/Experimental/Iris/Core/Private/Iris/ReplicationSystem/Filtering/ReplicationFiltering.h
--- a/Engine/Source/Runtime/Experimental/Iris/Core/Private/Iris/ReplicationSystem/Filtering/ReplicationFiltering.h
+++ b/Engine/Source/Runtime/Experimental/Iris/Core/Private/Iris/ReplicationSystem/Filtering/ReplicationFiltering.h
@@ -183,8 +183,6 @@ private:
 		TStrongObjectPtr<UNetObjectFilter> Filter;
 		FName Name;
 		uint32 ObjectCount = 0;
-		// Objects with this filter set.
-		FNetBitArray FilteredObjects;
 	};
 
 private:
@@ -335,6 +333,8 @@ private:
 	uint32 bHasDynamicFilters : 1;
 	uint32 bHasDirtyExclusionFilterGroup : 1;
 	uint32 bHasDirtyInclusionFilterGroup : 1;
+	// Is true if any initialized DynamicFilter has the NeedsUpdate trait
+	uint32 bHasDynamicFiltersWithUpdateTrait : 1;
 };
 
 inline bool FReplicationFiltering::HasDynamicFilters() const
diff --git a/Engine/Source/Runtime/Experimental/Iris/Core/Private/Iris/ReplicationSystem/Prioritization/NetObjectConnectionFilter.cpp b/Engine/Source/Runtime/Experimental/Iris/Core/Private/Iris/ReplicationSystem/Prioritization/NetObjectConnectionFilter.cpp
--- a/Engine/Source/Runtime/Experimental/Iris/Core/Private/Iris/ReplicationSystem/Prioritization/NetObjectConnectionFilter.cpp
+++ b/Engine/Source/Runtime/Experimental/Iris/Core/Private/Iris/ReplicationSystem/Prioritization/NetObjectConnectionFilter.cpp
@@ -27,7 +27,7 @@ void UNetObjectConnectionFilter::SetReplicateToConnection(UE::Net::FNetRefHandle
 	}
 }
 
-void UNetObjectConnectionFilter::OnInit(FNetObjectFilterInitParams& Params)
+void UNetObjectConnectionFilter::OnInit(const FNetObjectFilterInitParams& Params)
 {
 	NetRefHandleManager = &Params.ReplicationSystem->GetReplicationSystemInternal()->GetNetRefHandleManager();
 
@@ -83,11 +83,6 @@ void UNetObjectConnectionFilter::RemoveObject(uint32 ObjectIndex, const FNetObje
 	LocalToNetRefIndex[LocalIndex] = UE::Net::Private::FNetRefHandleManager::InvalidInternalIndex;
 }
 
-void UNetObjectConnectionFilter::UpdateObjects(FNetObjectFilterUpdateParams&)
-{
-	// Intentionally left empty.
-}
-
 void UNetObjectConnectionFilter::PreFilter(FNetObjectPreFilteringParams& Params)
 {
 	if (!bObjectRemoved)
diff --git a/Engine/Source/Runtime/Experimental/Iris/Core/Public/Iris/ReplicationSystem/Filtering/FilterOutNetObjectFilter.h b/Engine/Source/Runtime/Experimental/Iris/Core/Public/Iris/ReplicationSystem/Filtering/FilterOutNetObjectFilter.h
--- a/Engine/Source/Runtime/Experimental/Iris/Core/Public/Iris/ReplicationSystem/Filtering/FilterOutNetObjectFilter.h
+++ b/Engine/Source/Runtime/Experimental/Iris/Core/Public/Iris/ReplicationSystem/Filtering/FilterOutNetObjectFilter.h
@@ -18,9 +18,8 @@ class UFilterOutNetObjectFilter final : public UNetObjectFilter
 
 protected:
 	// UNetObjectFilter interface
-	IRISCORE_API virtual void OnInit(FNetObjectFilterInitParams&) override;
+	IRISCORE_API virtual void OnInit(const FNetObjectFilterInitParams&) override;
 	IRISCORE_API virtual bool AddObject(uint32 ObjectIndex, FNetObjectFilterAddObjectParams&) override;
 	IRISCORE_API virtual void RemoveObject(uint32 ObjectIndex, const FNetObjectFilteringInfo&) override;
-	IRISCORE_API virtual void UpdateObjects(FNetObjectFilterUpdateParams&) override;
 	IRISCORE_API virtual void Filter(FNetObjectFilteringParams&) override;
 };
diff --git a/Engine/Source/Runtime/Experimental/Iris/Core/Public/Iris/ReplicationSystem/Filtering/NetObjectConnectionFilter.h b/Engine/Source/Runtime/Experimental/Iris/Core/Public/Iris/ReplicationSystem/Filtering/NetObjectConnectionFilter.h
--- a/Engine/Source/Runtime/Experimental/Iris/Core/Public/Iris/ReplicationSystem/Filtering/NetObjectConnectionFilter.h
+++ b/Engine/Source/Runtime/Experimental/Iris/Core/Public/Iris/ReplicationSystem/Filtering/NetObjectConnectionFilter.h
@@ -50,12 +50,11 @@ protected:
 	};
 
 	// UNetObjectFilter interface
-	IRISCORE_API virtual void OnInit(FNetObjectFilterInitParams&) override;
+	IRISCORE_API virtual void OnInit(const FNetObjectFilterInitParams&) override;
 	IRISCORE_API virtual void AddConnection(uint32 ConnectionId) override;
 	IRISCORE_API virtual void RemoveConnection(uint32 ConnectionId) override;
 	IRISCORE_API virtual bool AddObject(uint32 ObjectIndex, FNetObjectFilterAddObjectParams&) override;
 	IRISCORE_API virtual void RemoveObject(uint32 ObjectIndex, const FNetObjectFilteringInfo&) override;
-	IRISCORE_API virtual void UpdateObjects(FNetObjectFilterUpdateParams&) override;
 	IRISCORE_API virtual void PreFilter(FNetObjectPreFilteringParams&) override;
 	IRISCORE_API virtual void Filter(FNetObjectFilteringParams&) override;
 
diff --git a/Engine/Source/Runtime/Experimental/Iris/Core/Public/Iris/ReplicationSystem/Filtering/NetObjectFilter.h b/Engine/Source/Runtime/Experimental/Iris/Core/Public/Iris/ReplicationSystem/Filtering/NetObjectFilter.h
--- a/Engine/Source/Runtime/Experimental/Iris/Core/Public/Iris/ReplicationSystem/Filtering/NetObjectFilter.h
+++ b/Engine/Source/Runtime/Experimental/Iris/Core/Public/Iris/ReplicationSystem/Filtering/NetObjectFilter.h
@@ -41,11 +41,6 @@ enum class ENetFilterStatus : uint32
  */
 struct FNetObjectFilteringParams
 {
-	FNetObjectFilteringParams(const UE::Net::FNetBitArrayView InFilteredObjects);
-
-	/** The indices of the objects that have this filter set. The indices of set bits correspond to the object indices. */
-	const UE::Net::FNetBitArrayView FilteredObjects;
-
 	/**
 	 * The contents of OutAllowedObjects is undefined when passed to Filter(). The filter is responsible
 	 * for setting and clearing bits for objects that have this filter set, which is provided in the
@@ -55,13 +50,13 @@ struct FNetObjectFilteringParams
 	UE::Net::FNetBitArrayView OutAllowedObjects;
 
 	/** FilteringInfos for all objects. Index using the set bit indices in FilteredObjects. */
-	const FNetObjectFilteringInfo* FilteringInfos = nullptr;
+	TArrayView<const FNetObjectFilteringInfo> FilteringInfos;
 
 	/** State buffers for all objects. Index using the set bit indices in FilteredObjects. */
 	const UE::Net::TNetChunkedArray<uint8*>* StateBuffers = nullptr;
 
 	/** ID of the connection that the filtering applies to. */
-	uint32 ConnectionId;
+	uint32 ConnectionId = 0;
 
 	/** The view associated with the connection and its sub-connections that objects are filtered for. */
 	UE::Net::FReplicationView View;
@@ -72,14 +67,9 @@ struct FNetObjectFilteringParams
  */
 struct FNetObjectPreFilteringParams
 {
-	FNetObjectPreFilteringParams(const UE::Net::FNetBitArrayView InFilteredObjects);
-
 	// The IDs of all valid connections.
 	UE::Net::FNetBitArrayView ValidConnections;
 
-	/** The indices of the objects that have this filter set. The indices of set bits correspond to the object indices. */
-	const UE::Net::FNetBitArrayView FilteredObjects;
-
 	/** FilteringInfos for all objects. Index using the set bit indices in FilteredObjects. */
 	TArrayView<const FNetObjectFilteringInfo> FilteringInfos;
 };
@@ -102,8 +92,11 @@ struct alignas(8) FNetObjectFilteringInfo
 
 enum class ENetFilterTraits : uint8
 {
-	None = 0,
-	Spatial = 1,
+	None = 0x00,
+	/** Set this trait for NetFilters that filter according to the WorldLocation of it's objects. */
+	Spatial = 0x01,
+	/** Set this trait so that UpdateObjects will be called on your NetFilter. Default is to not call the virtual */
+	NeedsUpdate = 0x02,
 };
 ENUM_CLASS_FLAGS(ENetFilterTraits);
 
@@ -123,20 +116,13 @@ public:
 struct FNetObjectFilterInitParams
 {
 public:
-	FNetObjectFilterInitParams(const UE::Net::FNetBitArrayView InFilteredObjects)
-	: FilteredObjects(InFilteredObjects)
-	{
-	}
-
-	TObjectPtr<UReplicationSystem> ReplicationSystem;
+	TObjectPtr<UReplicationSystem> ReplicationSystem = nullptr;
 	/** Optional config as set in the FNetObjectFilterDefinition. */
 	UNetObjectFilterConfig* Config = nullptr;
 	/** The maximum number of objects in the system. */
 	uint32 MaxObjectCount = 0;
 	/** The maximum number of connections in the system. */
 	uint32 MaxConnectionCount = 0;
-	/** The objects that are handled by the filter. */
-	const UE::Net::FNetBitArrayView FilteredObjects;
 };
 
 struct FNetObjectFilterAddObjectParams
@@ -170,7 +156,7 @@ struct FNetObjectFilterUpdateParams
 	uint32 ObjectCount = 0;
 
 	/** Infos for all objects. Index using ObjectIndices[0..ObjectCount-1]. */
-	FNetObjectFilteringInfo* FilteringInfos = nullptr;
+	TArrayView<FNetObjectFilteringInfo> FilteringInfos;
 };
 
 UCLASS(Abstract, MinimalAPI)
@@ -179,8 +165,8 @@ class UNetObjectFilter : public UObject
 	GENERATED_BODY()
 
 public:
-	IRISCORE_API void Init(FNetObjectFilterInitParams& Params);
-
+	IRISCORE_API void Init(const FNetObjectFilterInitParams& Params);
+	
 	/** A new connection has been added. An opportunity for the filter to allocate per connection info. */
 	IRISCORE_API virtual void AddConnection(uint32 ConnectionId);
 
@@ -193,8 +179,8 @@ public:
 	/** An object no longer wants to use this filter. */
 	IRISCORE_API virtual void RemoveObject(uint32 ObjectIndex, const FNetObjectFilteringInfo&) PURE_VIRTUAL(RemoveObject,)
 
-	/** A set of objects using this filter has been updated. An opportunity for the filter to update cached data. */
-	IRISCORE_API virtual void UpdateObjects(FNetObjectFilterUpdateParams&) PURE_VIRTUAL(UpdateObjects,)
+	/** A set of objects using this filter are dirty since the last update. An opportunity for the filter to update cached data. */
+	IRISCORE_API virtual void UpdateObjects(FNetObjectFilterUpdateParams&);
 
 	/**
 	 * If there are any connections being replicated and there's a chance Filter() will be called then PreFilter()
@@ -210,14 +196,21 @@ public:
 	 */
 	IRISCORE_API virtual void PostFilter(FNetObjectPostFilteringParams&);
 
-	/** Returns the filter's traits. */
+	/** Returns all the filter's traits. */
 	ENetFilterTraits GetFilterTraits() const { return FilterTraits; }
 
+	/** Tells if the filter was assigned a specific trait */
+	bool HasFilterTrait(ENetFilterTraits FilterTrait) const { return EnumHasAnyFlags(FilterTraits, FilterTrait); }
+
+	/** The list of objects that are filtered by this filter */
+	UE::Net::FNetBitArrayView GetFilteredObjects() { return MakeNetBitArrayView(FilteredObjects); }
+
+
 	struct FDebugInfoParams
 	{
 		FName FilterName;
 
-		const FNetObjectFilteringInfo* FilteringInfos = nullptr;
+		TArrayView<const FNetObjectFilteringInfo> FilteringInfos;
 
 		/** ID of the connection that the filtering applies to. */
 		uint32 ConnectionId = 0;
@@ -225,60 +218,38 @@ public:
 		/** The view associated with the connection and its sub-connections that objects are filtered for. */
 		UE::Net::FReplicationView View;
 	};
-
 	IRISCORE_API virtual FString PrintDebugInfoForObject(const FDebugInfoParams& Params, uint32 ObjectIndex) const { return Params.FilterName.ToString(); };
 
 protected:
 	IRISCORE_API UNetObjectFilter();
 
 	/** Called right after constructor for enabled filters. Must be overriden. */
-	IRISCORE_API virtual void OnInit(FNetObjectFilterInitParams&) PURE_VIRTUAL(OnInit, );
-
-	/* Returns true if the object is added to this filter, false otherwise. */
-	bool IsAddedToFilter(uint32 ObjectIndex) const;
+	IRISCORE_API virtual void OnInit(const FNetObjectFilterInitParams&) PURE_VIRTUAL(OnInit, );
 
 	/* Returns the filtering info for this object if it's handled by this filter, nullptr otherwise. */
-	FNetObjectFilteringInfo* GetFilteringInfo(uint32 ObjectIndex);
+	IRISCORE_API FNetObjectFilteringInfo* GetFilteringInfo(uint32 ObjectIndex);
+
+	/* Returns true if the object is assigned to be filtered by this filter.*/
+	inline bool IsObjectFiltered(uint32 ObjectIndex) const;
 
 	/** Adds traits. */
-	void AddFilterTraits(ENetFilterTraits Traits);
+	inline void AddFilterTraits(ENetFilterTraits Traits);
 
 	/** Sets the traits specified by TraitsMask to Traits. */
-	void SetFilterTraits(ENetFilterTraits Traits, ENetFilterTraits TraitsMask);
+	inline void SetFilterTraits(ENetFilterTraits Traits, ENetFilterTraits TraitsMask);
 
-private:
-	class FFilterInfo
-	{
-	public:
-		FFilterInfo() = default;
-		FFilterInfo(const UE::Net::FNetBitArrayView FilteredObjects, const TArrayView<FNetObjectFilteringInfo> FileringInfos);
-
-		FFilterInfo& operator=(const FFilterInfo&);
-
-		bool IsAddedToFilter(uint32 ObjectIndex) const;
-
-		FNetObjectFilteringInfo* GetFilteringInfo(uint32 ObjectIndex);
+protected:
 
-	private:
-		const UE::Net::FNetBitArrayView FilteredObjects;
-		TArrayView<FNetObjectFilteringInfo> FilteringInfos;
-	};
+	/** The indices of the objects that have this filter set. The indices of set bits correspond to the object indices. */
+	UE::Net::FNetBitArray FilteredObjects;
 
+private:
 	ENetFilterTraits FilterTraits = ENetFilterTraits::None;
-	FFilterInfo FilterInfo;
+	
+	TArrayView<FNetObjectFilteringInfo> FilteringInfos; 
 };
 
-inline bool UNetObjectFilter::IsAddedToFilter(uint32 ObjectIndex) const
-{
-	return FilterInfo.IsAddedToFilter(ObjectIndex);
-}
-
-inline FNetObjectFilteringInfo* UNetObjectFilter::GetFilteringInfo(uint32 ObjectIndex)
-{
-	return FilterInfo.GetFilteringInfo(ObjectIndex);
-}
-
-inline bool UNetObjectFilter::FFilterInfo::IsAddedToFilter(uint32 ObjectIndex) const
+inline bool UNetObjectFilter::IsObjectFiltered(uint32 ObjectIndex) const
 {
 	return ObjectIndex < FilteredObjects.GetNumBits() && FilteredObjects.IsBitSet(ObjectIndex);
 }
diff --git a/Engine/Source/Runtime/Experimental/Iris/Core/Public/Iris/ReplicationSystem/Filtering/NetObjectGridFilter.h b/Engine/Source/Runtime/Experimental/Iris/Core/Public/Iris/ReplicationSystem/Filtering/NetObjectGridFilter.h
--- a/Engine/Source/Runtime/Experimental/Iris/Core/Public/Iris/ReplicationSystem/Filtering/NetObjectGridFilter.h
+++ b/Engine/Source/Runtime/Experimental/Iris/Core/Public/Iris/ReplicationSystem/Filtering/NetObjectGridFilter.h
@@ -94,7 +94,7 @@ class UNetObjectGridFilter : public UNetObjectFilter
 
 protected:
 	// UNetObjectFilter interface
-	IRISCORE_API virtual void OnInit(FNetObjectFilterInitParams&) override;
+	IRISCORE_API virtual void OnInit(const FNetObjectFilterInitParams&) override;
 	IRISCORE_API virtual void AddConnection(uint32 ConnectionId) override;
 	IRISCORE_API virtual void RemoveConnection(uint32 ConnectionId) override;
 	IRISCORE_API virtual bool AddObject(uint32 ObjectIndex, FNetObjectFilterAddObjectParams&) override;
@@ -307,8 +307,7 @@ class UNetObjectGridWorldLocFilter : public UNetObjectGridFilter
 
 protected:
 
-	virtual void OnInit(FNetObjectFilterInitParams&) override;
-	virtual void UpdateObjects(FNetObjectFilterUpdateParams&) override;
+	virtual void OnInit(const FNetObjectFilterInitParams&) override;
 	virtual void PreFilter(FNetObjectPreFilteringParams&) override;
 	virtual void UpdateObjectInfo(FPerObjectInfo& PerObjectInfo, const UNetObjectGridFilter::FObjectLocationInfo& ObjectLocationInfo, const UE::Net::FReplicationInstanceProtocol* InstanceProtocol) override;
 	virtual bool BuildObjectInfo(uint32 ObjectIndex, FNetObjectFilterAddObjectParams& Params) override;
diff --git a/Engine/Source/Runtime/Experimental/Iris/Core/Public/Iris/ReplicationSystem/Filtering/NopNetObjectFilter.h b/Engine/Source/Runtime/Experimental/Iris/Core/Public/Iris/ReplicationSystem/Filtering/NopNetObjectFilter.h
--- a/Engine/Source/Runtime/Experimental/Iris/Core/Public/Iris/ReplicationSystem/Filtering/NopNetObjectFilter.h
+++ b/Engine/Source/Runtime/Experimental/Iris/Core/Public/Iris/ReplicationSystem/Filtering/NopNetObjectFilter.h
@@ -18,9 +18,8 @@ class UNopNetObjectFilter final : public UNetObjectFilter
 
 protected:
 	// UNetObjectFilter interface
-	IRISCORE_API virtual void OnInit(FNetObjectFilterInitParams&) override;
+	IRISCORE_API virtual void OnInit(const FNetObjectFilterInitParams&) override;
 	IRISCORE_API virtual bool AddObject(uint32 ObjectIndex, FNetObjectFilterAddObjectParams&) override;
 	IRISCORE_API virtual void RemoveObject(uint32 ObjectIndex, const FNetObjectFilteringInfo&) override;
-	IRISCORE_API virtual void UpdateObjects(FNetObjectFilterUpdateParams&) override;
 	IRISCORE_API virtual void Filter(FNetObjectFilteringParams&) override;
 };
