Iris
* DirtyNetObjectTrackerContainer now using proper NetBitArray instead of raw array
* Added LogIrisDirtyTracker and some verbose logs (non-shipping) to MarkDirty and ForceNetUpdate functions.

diff --git a/Engine/Source/Runtime/Experimental/Iris/Core/Private/Iris/ReplicationSystem/DirtyNetObjectTracker.cpp b/Engine/Source/Runtime/Experimental/Iris/Core/Private/Iris/ReplicationSystem/DirtyNetObjectTracker.cpp
--- a/Engine/Source/Runtime/Experimental/Iris/Core/Private/Iris/ReplicationSystem/DirtyNetObjectTracker.cpp
+++ b/Engine/Source/Runtime/Experimental/Iris/Core/Private/Iris/ReplicationSystem/DirtyNetObjectTracker.cpp
@@ -10,19 +10,14 @@
 #include "Traits/IntType.h"
 #include <atomic>
 
-#if (UE_BUILD_SHIPPING || UE_BUILD_TEST)
-#	define UE_NET_ENABLE_DIRTYOBJECTTRACKER_LOG 0
+// Don't compile verbose logs in Shipping builds
+#if (UE_BUILD_SHIPPING)
+#	define UE_NET_DIRTYOBJECTTRACKER_LOG_COMPIL_VERBOSITY Log
 #else
-#	define UE_NET_ENABLE_DIRTYOBJECTTRACKER_LOG 0
+#	define UE_NET_DIRTYOBJECTTRACKER_LOG_COMPIL_VERBOSITY All
 #endif 
 
-#if UE_NET_ENABLE_DIRTYOBJECTTRACKER_LOG
-#	define UE_LOG_DIRTYOBJECTTRACKER(Format, ...)  UE_LOG(LogIris, Log, Format, ##__VA_ARGS__)
-#else
-#	define UE_LOG_DIRTYOBJECTTRACKER(...)
-#endif
-
-#define UE_LOG_DIRTYOBJECTTRACKER_WARNING(Format, ...)  UE_LOG(LogIris, Warning, Format, ##__VA_ARGS__)
+DEFINE_LOG_CATEGORY_STATIC(LogIrisDirtyTracker, Log, UE_NET_DIRTYOBJECTTRACKER_LOG_COMPIL_VERBOSITY)
 
 namespace UE::Net::Private
 {
@@ -39,41 +34,28 @@ FDirtyNetObjectTracker::~FDirtyNetObjectTracker()
 
 void FDirtyNetObjectTracker::Init(const FDirtyNetObjectTrackerInitParams& Params)
 {
-	check(Params.NetObjectIndexRangeEnd >= Params.NetObjectIndexRangeStart);
 	check(Params.NetRefHandleManager != nullptr);
-	check(DirtyNetObjectContainer == nullptr);
 
 	NetRefHandleManager = Params.NetRefHandleManager;
 	ReplicationSystemId = Params.ReplicationSystemId;
-	NetObjectIdRangeStart = Params.NetObjectIndexRangeStart;
-	NetObjectIdRangeEnd = Params.NetObjectIndexRangeEnd;
-	/* 
-	 * For now we support all IDs up to RangeEnd. This could be expensive if we partition things more in some way or other.
-	 * In the latter case we would have to add functionality to FNetBitArrayView to handle an offset or add a "FNetSparseBitArray".
-	 */
+	
 	NetObjectIdCount = Params.MaxObjectCount;
 
 	GlobalDirtyTrackerPollHandle = FGlobalDirtyNetObjectTracker::CreatePoller();
 
-	DirtyNetObjectWordCount = (NetObjectIdCount + StorageTypeBitCount - 1)/StorageTypeBitCount;
-	DirtyNetObjectContainer = new StorageType[DirtyNetObjectWordCount];
-	FMemory::Memzero(DirtyNetObjectContainer, DirtyNetObjectWordCount * sizeof(StorageType));
-
 	AccumulatedDirtyNetObjects.Init(NetObjectIdCount);
 	ForceNetUpdateObjects.Init(NetObjectIdCount);
+	DirtyNetObjects.Init(NetObjectIdCount);
 
 	AllowExternalAccess();
 
-	UE_LOG_DIRTYOBJECTTRACKER(TEXT("FDirtyNetObjectTracker::Init %u Id, Start:%u, End: %u"), ReplicationSystemId, NetObjectIdRangeStart, NetObjectIdRangeEnd);
+	UE_LOG(LogIrisDirtyTracker, Log, TEXT("FDirtyNetObjectTracker::Init[%u]: CurrentMaxSize: %u"), ReplicationSystemId, NetObjectIdCount);
 }
 
 void FDirtyNetObjectTracker::Deinit()
 {
 	GlobalDirtyTrackerPollHandle.Destroy();
 	bShouldResetPolledGlobalDirtyTracker = false;
-
-	delete[] DirtyNetObjectContainer;
-	DirtyNetObjectContainer = nullptr;
 }
 
 void FDirtyNetObjectTracker::GrabAndApplyGlobalDirtyObjectList()
@@ -85,9 +67,7 @@ void FDirtyNetObjectTracker::GrabAndApplyGlobalDirtyObjectList()
 			const FInternalNetRefIndex NetObjectIndex = NetRefHandleManager->GetInternalIndexFromNetHandle(NetHandle);
 			if (NetObjectIndex != FNetRefHandleManager::InvalidInternalIndex)
 			{
-				const uint32 BitOffset = NetObjectIndex;
-				const StorageType BitMask = StorageType(1) << (BitOffset & (StorageTypeBitCount - 1));
-				DirtyNetObjectContainer[BitOffset / StorageTypeBitCount] |= BitMask;
+				DirtyNetObjects.SetBit(NetObjectIndex);
 			}
 		}
 	}
@@ -112,18 +92,24 @@ void FDirtyNetObjectTracker::UpdateDirtyNetObjects()
 
 	//$IRIS TODO:  We could look if any objects where actually in the global list and skip the array iteration if not needed.
 
-	const uint32* GlobalScopeListData = NetRefHandleManager->GetCurrentFrameScopableInternalIndices().GetData();
+	const FNetBitArrayView GlobalScopeList = NetRefHandleManager->GetCurrentFrameScopableInternalIndices();
+	const uint32* GlobalScopeListData = GlobalScopeList.GetData();
+	
 	uint32* AccumulatedDirtyNetObjectsData = AccumulatedDirtyNetObjects.GetData();
+	uint32* DirtyNetObjectsData = DirtyNetObjects.GetData();
 
 	const uint32 NumWords = AccumulatedDirtyNetObjects.GetNumWords();
-	for (uint32 WordIndex = 0; WordIndex != NumWords; ++WordIndex)
+	check(NumWords == DirtyNetObjects.GetNumWords());
+	check(NumWords == GlobalScopeList.GetNumWords());
+
+	for (uint32 WordIndex = 0; WordIndex < NumWords; ++WordIndex)
 	{
 		// Due to objects having been marked as dirty and later removed we must make sure that all dirty objects are still in scope.
-		uint32 DirtyObjectWord = DirtyNetObjectContainer[WordIndex] & GlobalScopeListData[WordIndex];
-		DirtyNetObjectContainer[WordIndex] = DirtyObjectWord;
+		uint32 DirtyObjectWord = DirtyNetObjectsData[WordIndex] & GlobalScopeListData[WordIndex];
+		DirtyNetObjectsData[WordIndex] = DirtyObjectWord;
 
 		// Add the latest dirty objects to the accumulated list and remove no-longer scoped objects that have never been copied.
-		AccumulatedDirtyNetObjectsData[WordIndex] = (AccumulatedDirtyNetObjectsData[WordIndex] | DirtyNetObjectContainer[WordIndex]) & GlobalScopeListData[WordIndex];
+		AccumulatedDirtyNetObjectsData[WordIndex] = (AccumulatedDirtyNetObjectsData[WordIndex] | DirtyNetObjectsData[WordIndex]) & GlobalScopeListData[WordIndex];
 	}
 
 	AllowExternalAccess();
@@ -144,10 +130,7 @@ void FDirtyNetObjectTracker::UpdateAndLockDirtyNetObjects()
 void FDirtyNetObjectTracker::UpdateAccumulatedDirtyList()
 {
 	IRIS_PROFILER_SCOPE(FDirtyNetObjectTracker_UpdateDirtyNetObjects)
-
-	FNetBitArrayView DirtyObjectsThisFrame(DirtyNetObjectContainer, NetObjectIdCount);
-
-	MakeNetBitArrayView(AccumulatedDirtyNetObjects).Combine(DirtyObjectsThisFrame, FNetBitArrayView::OrOp);
+	AccumulatedDirtyNetObjects.Combine(DirtyNetObjects, FNetBitArray::OrOp);
 }
 
 void FDirtyNetObjectTracker::MarkNetObjectDirty(FInternalNetRefIndex NetObjectIndex)
@@ -156,16 +139,21 @@ void FDirtyNetObjectTracker::MarkNetObjectDirty(FInternalNetRefIndex NetObjectIn
 	checkf(bIsExternalAccessAllowed, TEXT("Cannot mark objects dirty while the bitarray is locked for modifications."));
 #endif
 
-	if ((NetObjectIndex >= NetObjectIdRangeStart) & (NetObjectIndex <= NetObjectIdRangeEnd))
+	if (NetObjectIndex >= NetObjectIdCount || NetObjectIndex == FNetRefHandleManager::InvalidInternalIndex)
 	{
-		const uint32 BitOffset = NetObjectIndex;
-		const StorageType BitMask = StorageType(1) << (BitOffset & (StorageTypeBitCount - 1));
+		UE_LOG(LogIrisDirtyTracker, Warning, TEXT("FDirtyNetObjectTracker::MarkNetObjectDirty received invalid NetObjectIndex: %u | Max: %u"), NetObjectIndex, NetObjectIdCount);
+		return;
+	}
 
-		// ideally we'd have c++20 std::atomic_ref for this
-		FPlatformAtomics::InterlockedOr((TSignedIntType<sizeof(StorageType)>::Type*)(&DirtyNetObjectContainer[BitOffset/StorageTypeBitCount]), BitMask);
+	const uint32 BitOffset = NetObjectIndex;
+	const StorageType BitMask = StorageType(1) << (BitOffset & (StorageTypeBitCount - 1));
 
-		UE_LOG_DIRTYOBJECTTRACKER(TEXT("FDirtyNetObjectTracker::MarkNetObjectDirty %u ( InternalIndex: %u )"), ReplicationSystemId, NetObjectIndex);
-	}
+	uint32* DirtyNetObjectsData = DirtyNetObjects.GetData();
+
+	// ideally we'd have c++20 std::atomic_ref for this
+	FPlatformAtomics::InterlockedOr((TSignedIntType<sizeof(StorageType)>::Type*)(&DirtyNetObjectsData[BitOffset/StorageTypeBitCount]), BitMask);
+
+	UE_LOG(LogIrisDirtyTracker, Verbose, TEXT("FDirtyNetObjectTracker::MarkNetObjectDirty[%u]: %s"), ReplicationSystemId, *NetRefHandleManager->PrintObjectFromIndex(NetObjectIndex));
 }
 
 void FDirtyNetObjectTracker::ForceNetUpdate(FInternalNetRefIndex NetObjectIndex)
@@ -175,7 +163,7 @@ void FDirtyNetObjectTracker::ForceNetUpdate(FInternalNetRefIndex NetObjectIndex)
 	// Flag the object dirty so we update his filters too
 	MarkNetObjectDirty(NetObjectIndex);
 
-	UE_LOG_DIRTYOBJECTTRACKER(TEXT("FDirtyNetObjectTracker::ForceNetUpdateObjects %u ( InternalIndex: %u )"), ReplicationSystemId, NetObjectIndex);
+	UE_LOG(LogIrisDirtyTracker, Verbose, TEXT("FDirtyNetObjectTracker::ForceNetUpdateObjects[%u]: %s"), ReplicationSystemId, *NetRefHandleManager->PrintObjectFromIndex(NetObjectIndex));
 }
 
 void FDirtyNetObjectTracker::LockExternalAccess()
@@ -197,7 +185,7 @@ FNetBitArrayView FDirtyNetObjectTracker::GetDirtyNetObjectsThisFrame()
 #if UE_NET_THREAD_SAFETY_CHECK
 	checkf(!bIsExternalAccessAllowed, TEXT("Cannot access the DirtyNetObjects bitarray unless its locked for multithread access."));
 #endif
-	return FNetBitArrayView(DirtyNetObjectContainer, NetObjectIdCount);
+	return MakeNetBitArrayView(DirtyNetObjects);
 }
 
 void FDirtyNetObjectTracker::ReconcilePolledList(const FNetBitArrayView& ObjectsPolled)
@@ -217,7 +205,7 @@ void FDirtyNetObjectTracker::ReconcilePolledList(const FNetBitArrayView& Objects
 	MakeNetBitArrayView(AccumulatedDirtyNetObjects).Combine(ObjectsPolled, FNetBitArrayView::AndNotOp);
 
 	// Clear the current frame dirty objects
-	FMemory::Memzero(DirtyNetObjectContainer, DirtyNetObjectWordCount*sizeof(StorageType));
+	DirtyNetObjects.Reset();
 
 	AllowExternalAccess();
 
diff --git a/Engine/Source/Runtime/Experimental/Iris/Core/Private/Iris/ReplicationSystem/DirtyNetObjectTracker.h b/Engine/Source/Runtime/Experimental/Iris/Core/Private/Iris/ReplicationSystem/DirtyNetObjectTracker.h
--- a/Engine/Source/Runtime/Experimental/Iris/Core/Private/Iris/ReplicationSystem/DirtyNetObjectTracker.h
+++ b/Engine/Source/Runtime/Experimental/Iris/Core/Private/Iris/ReplicationSystem/DirtyNetObjectTracker.h
@@ -26,8 +26,6 @@ struct FDirtyNetObjectTrackerInitParams
 	const FNetRefHandleManager* NetRefHandleManager = nullptr;
 	uint32 ReplicationSystemId = 0;
 	uint32 MaxObjectCount = 0;
-	uint32 NetObjectIndexRangeStart = 0;
-	uint32 NetObjectIndexRangeEnd = 0;
 };
 
 class FDirtyNetObjectTracker
@@ -92,7 +90,7 @@ private:
 	FNetBitArray ForceNetUpdateObjects;
 
 	// List of objects set to be dirty this frame. Is always reset at the end of the net tick flush
-	StorageType* DirtyNetObjectContainer = nullptr;
+	FNetBitArray DirtyNetObjects;
 
 	const FNetRefHandleManager* NetRefHandleManager = nullptr;
 	
@@ -100,9 +98,6 @@ private:
 
 	uint32 ReplicationSystemId;
 
-	uint32 DirtyNetObjectWordCount = 0;
-	uint32 NetObjectIdRangeStart = 0;
-	uint32 NetObjectIdRangeEnd = 0;
 	uint32 NetObjectIdCount = 0;
 	
 	bool bShouldResetPolledGlobalDirtyTracker = false;
diff --git a/Engine/Source/Runtime/Experimental/Iris/Core/Private/Iris/ReplicationSystem/ReplicationSystem.cpp b/Engine/Source/Runtime/Experimental/Iris/Core/Private/Iris/ReplicationSystem/ReplicationSystem.cpp
--- a/Engine/Source/Runtime/Experimental/Iris/Core/Private/Iris/ReplicationSystem/ReplicationSystem.cpp
+++ b/Engine/Source/Runtime/Experimental/Iris/Core/Private/Iris/ReplicationSystem/ReplicationSystem.cpp
@@ -106,8 +106,6 @@ public:
 			DirtyNetObjectTrackerInitParams.NetRefHandleManager = &NetRefHandleManager;
 			DirtyNetObjectTrackerInitParams.ReplicationSystemId = ReplicationSystemId;
 			DirtyNetObjectTrackerInitParams.MaxObjectCount = MaxObjectCount;
-			DirtyNetObjectTrackerInitParams.NetObjectIndexRangeStart = 1;
-			DirtyNetObjectTrackerInitParams.NetObjectIndexRangeEnd = MaxObjectCount - 1U;
 
 			ReplicationSystemInternal.InitDirtyNetObjectTracker(DirtyNetObjectTrackerInitParams);
 		}
