[Backout] - CL33441738 due to build health error.
#rnx
Original CL Desc
-----------------------------------------------------------------
Dataflow - Add branch node
- fix Forward input when trying to forward to an input that is not also the output
- move the reroute node and branch node to a core nodes file

diff --git a/Engine/Plugins/Experimental/Dataflow/Source/DataflowEditor/Private/Dataflow/DataflowSchema.cpp b/Engine/Plugins/Experimental/Dataflow/Source/DataflowEditor/Private/Dataflow/DataflowSchema.cpp
--- a/Engine/Plugins/Experimental/Dataflow/Source/DataflowEditor/Private/Dataflow/DataflowSchema.cpp
+++ b/Engine/Plugins/Experimental/Dataflow/Source/DataflowEditor/Private/Dataflow/DataflowSchema.cpp
@@ -2,7 +2,6 @@
 
 #include "Dataflow/DataflowSchema.h"
 
-#include "Dataflow/DataflowCoreNodes.h"
 #include "Dataflow/DataflowEdNode.h"
 #include "Dataflow/DataflowEditorCommands.h"
 #include "Dataflow/DataflowEdNode.h"
diff --git a/Engine/Source/Runtime/Experimental/Dataflow/Core/Private/Dataflow/DataflowCoreNodes.cpp b/Engine/Source/Runtime/Experimental/Dataflow/Core/Private/Dataflow/DataflowCoreNodes.cpp
deleted file mode 100644
diff --git a/Engine/Source/Runtime/Experimental/Dataflow/Core/Private/Dataflow/DataflowNode.cpp b/Engine/Source/Runtime/Experimental/Dataflow/Core/Private/Dataflow/DataflowNode.cpp
--- a/Engine/Source/Runtime/Experimental/Dataflow/Core/Private/Dataflow/DataflowNode.cpp
+++ b/Engine/Source/Runtime/Experimental/Dataflow/Core/Private/Dataflow/DataflowNode.cpp
@@ -801,7 +801,7 @@ void FDataflowNode::ForwardInput(Dataflow::FContext& Context, const void* InputR
 {
 	if (const FDataflowOutput* Output = FindOutput(Reference))
 	{
-		if (const FDataflowInput* Input = FindInput(InputReference))
+		if (const FDataflowInput* Input = FindInput(Reference))
 		{
 			// we need to pull the value first so the upstream of the graph evaluate 
 			Input->PullValue(Context);
diff --git a/Engine/Source/Runtime/Experimental/Dataflow/Core/Private/Dataflow/DataflowNodeFactory.cpp b/Engine/Source/Runtime/Experimental/Dataflow/Core/Private/Dataflow/DataflowNodeFactory.cpp
--- a/Engine/Source/Runtime/Experimental/Dataflow/Core/Private/Dataflow/DataflowNodeFactory.cpp
+++ b/Engine/Source/Runtime/Experimental/Dataflow/Core/Private/Dataflow/DataflowNodeFactory.cpp
@@ -3,7 +3,6 @@
 #include "Dataflow/DataflowNodeFactory.h"
 
 #include "Dataflow/DataflowNode.h"
-#include "Dataflow/DataflowCoreNodes.h"
 #include "Misc/MessageDialog.h"
 #include "Logging/LogMacros.h"
 
@@ -17,7 +16,6 @@ namespace Dataflow
 	{
 		// by default always register a reroute node
 		DATAFLOW_NODE_REGISTER_CREATION_FACTORY(FDataflowReRouteNode);
-		DATAFLOW_NODE_REGISTER_CREATION_FACTORY(FDataflowBranchNode);
 	}
 
 
diff --git a/Engine/Source/Runtime/Experimental/Dataflow/Core/Public/Dataflow/DataflowCoreNodes.h b/Engine/Source/Runtime/Experimental/Dataflow/Core/Public/Dataflow/DataflowCoreNodes.h
deleted file mode 100644
diff --git a/Engine/Source/Runtime/Experimental/Dataflow/Core/Public/Dataflow/DataflowNode.h b/Engine/Source/Runtime/Experimental/Dataflow/Core/Public/Dataflow/DataflowNode.h
--- a/Engine/Source/Runtime/Experimental/Dataflow/Core/Public/Dataflow/DataflowNode.h
+++ b/Engine/Source/Runtime/Experimental/Dataflow/Core/Public/Dataflow/DataflowNode.h
@@ -396,5 +396,29 @@ struct FDataflowAnyType
 	DATAFLOWCORE_API static const FName TypeName;
 };
 
+USTRUCT()
+struct FDataflowReRouteNode : public FDataflowNode
+{
+	GENERATED_USTRUCT_BODY()
+
+	DATAFLOW_NODE_DEFINE_INTERNAL(FDataflowReRouteNode, "ReRouteNode", "Core", "")
+
+public:
+	FDataflowReRouteNode(const Dataflow::FNodeParameters& Param, FGuid InGuid = FGuid::NewGuid())
+		: Super(Param, InGuid)
+	{
+		RegisterInputConnection(&Value);
+		RegisterOutputConnection(&Value, &Value);
+	}
+
+	virtual void Evaluate(Dataflow::FContext& Context, const FDataflowOutput* Out) const
+	{
+		ForwardInput(Context, &Value, &Value);
+	}
+
+public:
+	UPROPERTY(meta = (DataflowInput, DataflowOutput, DataflowPassthrough = "Value", DisplayName = "Value"))
+	FDataflowAnyType Value;
+};
 
 
diff --git a/Engine/Source/Runtime/Experimental/Dataflow/Engine/Private/Dataflow/DataflowEdNode.cpp b/Engine/Source/Runtime/Experimental/Dataflow/Engine/Private/Dataflow/DataflowEdNode.cpp
--- a/Engine/Source/Runtime/Experimental/Dataflow/Engine/Private/Dataflow/DataflowEdNode.cpp
+++ b/Engine/Source/Runtime/Experimental/Dataflow/Engine/Private/Dataflow/DataflowEdNode.cpp
@@ -4,7 +4,6 @@
 
 #include "Dataflow/DataflowNode.h"
 #include "Dataflow/DataflowCore.h"
-#include "Dataflow/DataflowCoreNodes.h"
 #include "Dataflow/DataflowObject.h"
 #include "Dataflow/DataflowRenderingFactory.h"
 #include "GeometryCollection/Facades/CollectionRenderingFacade.h"
