Constraints: cleanup useless/invalid constraints tick prerequisites after GC

#rnx

diff --git a/Engine/Source/Runtime/Experimental/Animation/Constraints/Private/ConstraintSubsystem.cpp b/Engine/Source/Runtime/Experimental/Animation/Constraints/Private/ConstraintSubsystem.cpp
--- a/Engine/Source/Runtime/Experimental/Animation/Constraints/Private/ConstraintSubsystem.cpp
+++ b/Engine/Source/Runtime/Experimental/Animation/Constraints/Private/ConstraintSubsystem.cpp
@@ -278,10 +278,50 @@ void UConstraintSubsystem::CleanupInvalidConstraints() const
 	{
 		for (FConstraintsInWorld& WorldConstraints: System->ConstraintsInWorld)
 		{
-			WorldConstraints.Constraints.RemoveAll( [](const TWeakObjectPtr<UTickableConstraint>& InConstraint)
+			UWorld* World = WorldConstraints.World.Get();
+			
+			TSet<FTickFunction*> ConstraintsTickFunctions;
+			ConstraintsTickFunctions.Reserve(WorldConstraints.Constraints.Num());
+			
+			// remove stale constraints and store valid constraints tick functions
+			WorldConstraints.Constraints.RemoveAll([World, &ConstraintsTickFunctions](const TWeakObjectPtr<UTickableConstraint>& InConstraint)
 			{
-				return !InConstraint.IsValid() || InConstraint.IsStale();
+				const bool bRemove = !InConstraint.IsValid() || InConstraint.IsStale();
+				if (!bRemove && World)
+				{
+					// store tick functions
+					ConstraintsTickFunctions.Add(&InConstraint->GetTickFunction(World));
+				}
+				return bRemove;
 			});
+
+			if (World)
+			{
+				static constexpr bool bEvenIfPendingKill = true;
+				
+				// cleanup useless tick prerequisites
+				for (const TWeakObjectPtr<UTickableConstraint>& Constraint : WorldConstraints.Constraints)
+				{
+					TArray<FTickPrerequisite>& Prerequisites = Constraint->GetTickFunction(World).GetPrerequisites();
+					for (int32 PrereqIndex = 0; PrereqIndex < Prerequisites.Num(); PrereqIndex++)
+					{
+						FTickPrerequisite& Prerequisite = Prerequisites[PrereqIndex];
+						
+						UObject* PrereqObject = Prerequisite.PrerequisiteObject.Get(bEvenIfPendingKill);
+						if (!PrereqObject)
+						{
+							// remove prerequisite coming from stale object (cf. FTickFunction::QueueTickFunction)
+							Prerequisites.RemoveAtSwap(PrereqIndex--);
+						}
+						else if (PrereqObject == System && !ConstraintsTickFunctions.Contains(Prerequisite.PrerequisiteTickFunction))
+						{
+							// remove prerequisite coming from GCd constraint (cf. UConstraintSubsystem::SetConstraintDependencies) 
+							Prerequisites.RemoveAtSwap(PrereqIndex--);
+						}
+					}
+				}
+			}
+			
 			WorldConstraints.InvalidateGraph();
 		}
 		bNeedsCleanup = false;
