Add dummy utoc files after IAD plugin chunking to make pak mounting happy.

diff --git a/Engine/Source/Runtime/Core/Private/IO/IoStore.cpp b/Engine/Source/Runtime/Core/Private/IO/IoStore.cpp
--- a/Engine/Source/Runtime/Core/Private/IO/IoStore.cpp
+++ b/Engine/Source/Runtime/Core/Private/IO/IoStore.cpp
@@ -1166,7 +1166,7 @@ FIoStatus FIoStoreTocResource::Read(const TCHAR* TocFilePath, EIoStoreTocReadOpt
 
 	TUniquePtr<uint8[]> TocBuffer = MakeUnique<uint8[]>(TocSize);
 
-	if (!TocFileHandle->Read(TocBuffer.Get(), TocSize))
+	if (TocSize != 0 && !TocFileHandle->Read(TocBuffer.Get(), TocSize))
 	{
 		return FIoStatusBuilder(EIoErrorCode::CorruptToc) << TEXT("Failed to read IoStore TOC file '") << TocFilePath << TEXT("'");
 	}
diff --git a/Engine/Source/Runtime/Experimental/IoStore/OnDemand/Private/Tool/ChunkPlugin.cpp b/Engine/Source/Runtime/Experimental/IoStore/OnDemand/Private/Tool/ChunkPlugin.cpp
--- a/Engine/Source/Runtime/Experimental/IoStore/OnDemand/Private/Tool/ChunkPlugin.cpp
+++ b/Engine/Source/Runtime/Experimental/IoStore/OnDemand/Private/Tool/ChunkPlugin.cpp
@@ -100,7 +100,7 @@ static int32 ChunkPluginCommandEntry(const FContext& Context)
 	}
 
 	TMap<FGuid, FAES::FAESKey> EncryptionKeys;
-	if (IFileManager::Get().DirectoryExists(*ContainerFolder) == false)
+	if (FileMgr.DirectoryExists(*ContainerFolder) == false)
 	{
 		UE_LOG(LogIoStore, Error, TEXT("Directory '%s' does not exist"), *ContainerFolder);
 		return -1;
@@ -225,7 +225,6 @@ static int32 ChunkPluginCommandEntry(const FContext& Context)
 		}
 	}
 
-
 	IFileManager& FileMan = IFileManager::Get();
 	for (const FString& Path : FilesToDelete)
 	{
@@ -324,6 +323,51 @@ static int32 ChunkPluginCommandEntry(const FContext& Context)
 		}
 	}
 
+	// Write dummy containers if necessary
+	if (bDeleteContainerFiles)
+	{
+		TSet<FString> DummyContainerPaths;
+		FileMgr.IterateDirectory(*ContainerFolder, [&DummyContainerPaths](const TCHAR* FilenameOrDirectory, bool bIsDirectory)
+		{
+			if (!bIsDirectory)
+			{
+				FStringView MaybePak(FilenameOrDirectory);
+				if (FPathViews::GetExtension(MaybePak) == TEXTVIEW("pak"))
+				{
+					DummyContainerPaths.Emplace(FPathViews::GetBaseFilenameWithPath(MaybePak));
+				}
+			}
+			return true;
+		});
+
+		for (const FString& Filename : ContainerFilenames)
+		{
+			const FString FullPath = ContainerFolder / Filename;
+			const FStringView ContainerPath = FPathViews::GetBaseFilenameWithPath(FullPath);
+			if (!DummyContainerPaths.ContainsByHash(GetTypeHash(ContainerPath), ContainerPath))
+			{
+				continue;
+			}
+
+			FNameBuilder NameBuilder;
+			NameBuilder << BuildVersion << FPathViews::GetBaseFilename(Filename) << TEXTVIEW("dummy");
+
+			FIoContainerSettings ContainerSettings;
+			ContainerSettings.ContainerId = FIoContainerId::FromName(FName(NameBuilder));
+
+			FIoStoreTocResource Toc;
+
+			TIoStatusOr<uint64> Status = FIoStoreTocResource::Write(*FullPath, Toc, 0, 0, ContainerSettings);
+			if (!Status.IsOk())
+			{
+				UE_LOG(LogIoStore, Error, TEXT("Failed to write dummy container '%s' (%s)"), *FullPath, *Status.Status().ToString());
+				return -1;
+			}
+
+			UE_LOG(LogIoStore, Display, TEXT("Wrote dummy file '%s' (%.2lf KiB)"), *FullPath, double(Status.ValueOrDie()) / 1024);
+		}
+	}
+
 	return 0;
 }
 
diff --git a/Engine/Source/Runtime/Experimental/IoStore/OnDemand/Public/IO/IoStoreOnDemand.h b/Engine/Source/Runtime/Experimental/IoStore/OnDemand/Public/IO/IoStoreOnDemand.h
--- a/Engine/Source/Runtime/Experimental/IoStore/OnDemand/Public/IO/IoStoreOnDemand.h
+++ b/Engine/Source/Runtime/Experimental/IoStore/OnDemand/Public/IO/IoStoreOnDemand.h
@@ -24,7 +24,6 @@ class FArchive;
 class FCbFieldView;
 class FCbWriter;
 struct FKeyChain;
-class IIoStoreWriter;
 struct FAnalyticsEventAttribute;
 struct FIoContainerSettings;
 struct FIoStoreWriterSettings;
