Make the geometry collection neighbor/contact selection work on clusters + add supporting methods to enumerate neighbors w/ support for non-leaf transforms
+ make corresponding contact select dataflow nodes use the updated neighbor selection logic

diff --git a/Engine/Plugins/Experimental/GeometryCollectionPlugin/Source/GeometryCollectionNodes/Private/Dataflow/GeometryCollectionSelectionNodes.cpp b/Engine/Plugins/Experimental/GeometryCollectionPlugin/Source/GeometryCollectionNodes/Private/Dataflow/GeometryCollectionSelectionNodes.cpp
--- a/Engine/Plugins/Experimental/GeometryCollectionPlugin/Source/GeometryCollectionNodes/Private/Dataflow/GeometryCollectionSelectionNodes.cpp
+++ b/Engine/Plugins/Experimental/GeometryCollectionPlugin/Source/GeometryCollectionNodes/Private/Dataflow/GeometryCollectionSelectionNodes.cpp
@@ -599,7 +599,7 @@ void FCollectionTransformSelectionContactDataflowNode::Evaluate(Dataflow::FConte
 		GeometryCollection::Facades::FCollectionTransformSelectionFacade TransformSelectionFacade(InCollection);
 		TArray<int32> SelectionArr = InTransformSelection.AsArray();
 
-		TransformSelectionFacade.SelectLevel(SelectionArr);
+		TransformSelectionFacade.SelectContact(SelectionArr, bAllowContactInParentLevels);
 		InTransformSelection.SetFromArray(SelectionArr);
 
 		SetValue(Context, MoveTemp(InTransformSelection), &TransformSelection);
diff --git a/Engine/Plugins/Experimental/GeometryCollectionPlugin/Source/GeometryCollectionNodes/Public/Dataflow/GeometryCollectionSelectionNodes.h b/Engine/Plugins/Experimental/GeometryCollectionPlugin/Source/GeometryCollectionNodes/Public/Dataflow/GeometryCollectionSelectionNodes.h
--- a/Engine/Plugins/Experimental/GeometryCollectionPlugin/Source/GeometryCollectionNodes/Public/Dataflow/GeometryCollectionSelectionNodes.h
+++ b/Engine/Plugins/Experimental/GeometryCollectionPlugin/Source/GeometryCollectionNodes/Public/Dataflow/GeometryCollectionSelectionNodes.h
@@ -602,6 +602,10 @@ public:
 	UPROPERTY(meta = (DataflowInput, DataflowOutput, DataflowPassthrough = "Collection", DataflowIntrinsic))
 	FManagedArrayCollection Collection;
 
+	/** Whether to allow contact with bones that are in a parent level */
+	UPROPERTY(EditAnywhere, Category = Options)
+	bool bAllowContactInParentLevels = true;
+
 	FCollectionTransformSelectionContactDataflowNode(const Dataflow::FNodeParameters& InParam, FGuid InGuid = FGuid::NewGuid())
 		: FDataflowNode(InParam, InGuid)
 	{
diff --git a/Engine/Source/Runtime/Experimental/Chaos/Private/GeometryCollection/Facades/CollectionHierarchyFacade.cpp b/Engine/Source/Runtime/Experimental/Chaos/Private/GeometryCollection/Facades/CollectionHierarchyFacade.cpp
--- a/Engine/Source/Runtime/Experimental/Chaos/Private/GeometryCollection/Facades/CollectionHierarchyFacade.cpp
+++ b/Engine/Source/Runtime/Experimental/Chaos/Private/GeometryCollection/Facades/CollectionHierarchyFacade.cpp
@@ -79,6 +79,23 @@ namespace Chaos::Facades
 		return ParentAttribute[TransformIndex];
 	}
 
+	void FCollectionHierarchyFacade::EnumerateChildTransforms(int32 TransformIndex, TFunctionRef<bool(int32)> ShouldProcess, TFunctionRef<void(int32)> ProcessTransformFunc) const
+	{
+		checkSlow(ChildrenAttribute.IsValid() && ChildrenAttribute.IsValidIndex(TransformIndex));
+		
+		if (ShouldProcess(TransformIndex))
+		{
+			ProcessTransformFunc(TransformIndex);
+		}
+		else
+		{
+			for (int32 ChildTransformIndex : ChildrenAttribute[TransformIndex])
+			{
+				EnumerateChildTransforms(ChildTransformIndex, ShouldProcess, ProcessTransformFunc);
+			}
+		}
+	}
+
 	int32 FCollectionHierarchyFacade::GetInitialLevel(int32 TransformIndex) const
 	{
 		if (HasLevelAttribute())
diff --git a/Engine/Source/Runtime/Experimental/Chaos/Private/GeometryCollection/Facades/CollectionTransformSelectionFacade.cpp b/Engine/Source/Runtime/Experimental/Chaos/Private/GeometryCollection/Facades/CollectionTransformSelectionFacade.cpp
--- a/Engine/Source/Runtime/Experimental/Chaos/Private/GeometryCollection/Facades/CollectionTransformSelectionFacade.cpp
+++ b/Engine/Source/Runtime/Experimental/Chaos/Private/GeometryCollection/Facades/CollectionTransformSelectionFacade.cpp
@@ -464,34 +464,29 @@ namespace GeometryCollection::Facades
 		return OutSelection;
 	}
 
-	void FCollectionTransformSelectionFacade::SelectContact(TArray<int32>& InOutSelection) const
+	void FCollectionTransformSelectionFacade::SelectContact(TArray<int32>& InOutSelection, bool bIncludeNeighborsInParentLevels) const
 	{
 		TSet<int32> NewSelection;
 
 		if (CanSelectContact())
 		{
-			// TODO: This needs to be looked at
 			if (TUniquePtr<FGeometryCollection> TempGeomCollection = TUniquePtr<FGeometryCollection>(ConstCollection.NewCopy<FGeometryCollection>()))
 			{
 				FGeometryCollectionProximityUtility ProximityUtility(TempGeomCollection.Get());
 				ProximityUtility.RequireProximity();
 
-				const TManagedArray<int32>& TransformIndex = TransformIndexAttribute.Get();
-				const TManagedArray<int32>& TransformToGeometryIndex = TransformToGeometryIndexAttribute.Get();
-				const TManagedArray<TSet<int32>>& Proximity = TempGeomCollection->GetAttribute<TSet<int32>>("Proximity", FGeometryCollection::TransformGroup);
+				const TManagedArray<int32>& TransformIndex = TempGeomCollection->TransformIndex;
+				const TManagedArray<int32>& TransformToGeometryIndex = TempGeomCollection->TransformToGeometryIndex;
+				const TManagedArray<TSet<int32>>& Proximity = TempGeomCollection->GetAttribute<TSet<int32>>("Proximity", FGeometryCollection::GeometryGroup);
+				Chaos::Facades::FCollectionHierarchyFacade HierarchyFacade(*TempGeomCollection);
 
 				for (int32 Bone : InOutSelection)
 				{
 					NewSelection.Add(Bone);
-					int32 GeometryIdx = TransformToGeometryIndex[Bone];
-					if (GeometryIdx != INDEX_NONE)
-					{
-						const TSet<int32>& Neighbors = Proximity[GeometryIdx];
-						for (int32 NeighborGeometryIndex : Neighbors)
+					ProximityUtility.EnumerateNeighbors(HierarchyFacade, Bone, [&NewSelection](int32 NeighborTransformIdx)
 						{
-							NewSelection.Add(TransformIndex[NeighborGeometryIndex]);
-						}
-					}
+							NewSelection.Add(NeighborTransformIdx);
+						}, bIncludeNeighborsInParentLevels, false /*bFilterDuplicates, not needed since we add to a set*/);
 				}
 			}
 		}
diff --git a/Engine/Source/Runtime/Experimental/Chaos/Private/GeometryCollection/GeometryCollectionProximityUtility.cpp b/Engine/Source/Runtime/Experimental/Chaos/Private/GeometryCollection/GeometryCollectionProximityUtility.cpp
--- a/Engine/Source/Runtime/Experimental/Chaos/Private/GeometryCollection/GeometryCollectionProximityUtility.cpp
+++ b/Engine/Source/Runtime/Experimental/Chaos/Private/GeometryCollection/GeometryCollectionProximityUtility.cpp
@@ -529,6 +529,65 @@ void FGeometryCollectionProximityUtility::ClearConnectionGraph()
 	ConnectionsFacade.ClearAttributes();
 }
 
+void FGeometryCollectionProximityUtility::EnumerateNeighbors(const Chaos::Facades::FCollectionHierarchyFacade& Hierarchy, 
+	int32 TransformIdx, TFunctionRef<void(int32)> NeighborFunc, bool bIncludeNeighborsInParentLevels, bool bFilterDuplicates)
+{
+	if (!ensureMsgf(Collection->HasAttribute("Proximity", FGeometryCollection::GeometryGroup),
+		TEXT("Attempted to enumerate geometry collection neighbors, but proximity graph was not available")))
+	{
+		return;
+	}
+	if (!Hierarchy.IsValid() || !Hierarchy.HasLevelAttribute())
+	{
+		ensureMsgf(false, TEXT("Must have a valid hierarchy with a level attribute to enumerate geometry collection neighbors"));
+	}
+
+	TSet<int32> AlreadyProcessed;
+
+	int32 Level = Hierarchy.GetInitialLevel(TransformIdx);
+
+	const TManagedArray<TSet<int32>>& Proximity = Collection->GetAttribute<TSet<int32>>("Proximity", FGeometryCollection::GeometryGroup);
+	const TManagedArray<int32>& GeoToTransformIndex = Collection->GetAttribute<int32>("TransformIndex", FGeometryCollection::GeometryGroup);
+	const TManagedArray<int32>& TransformToGeometryIndex = Collection->GetAttribute<int32>("TransformToGeometryIndex", FGeometryCollection::TransformGroup);
+	const TManagedArray<int32>& SimulationType = Collection->GetAttribute<int32>("SimulationType", FGeometryCollection::TransformGroup);
+	Hierarchy.EnumerateChildTransforms(TransformIdx, 
+		[&SimulationType](int32 TransformIndex)
+		{
+			return SimulationType[TransformIndex] == (int32)FGeometryCollection::ESimulationTypes::FST_Rigid;
+		}, 
+		[&Hierarchy, TransformIdx, &NeighborFunc, bIncludeNeighborsInParentLevels, bFilterDuplicates, &AlreadyProcessed, Level, &Proximity, &GeoToTransformIndex, &TransformToGeometryIndex, &SimulationType]
+		(int32 ChildTransformIndex)
+		{
+			int32 GeoIdx = TransformToGeometryIndex[ChildTransformIndex];
+			for (int32 ProxGeoIndex : Proximity[GeoIdx])
+			{
+				int32 ProxTransformIndex = GeoToTransformIndex[ProxGeoIndex];
+				int32 ProxLevel = Hierarchy.GetInitialLevel(ProxTransformIndex);
+				if (bIncludeNeighborsInParentLevels && ProxLevel < Level)
+				{
+					NeighborFunc(ProxTransformIndex);
+				}
+				while (ProxLevel > Level && ProxTransformIndex != INDEX_NONE)
+				{
+					ProxTransformIndex = Hierarchy.GetParent(ProxTransformIndex);
+					ProxLevel--;
+				}
+				if (ProxTransformIndex != TransformIdx && ProxTransformIndex != INDEX_NONE)
+				{
+					if (bFilterDuplicates)
+					{
+						if (AlreadyProcessed.Contains(ProxTransformIndex))
+						{
+							continue;
+						}
+						AlreadyProcessed.Add(ProxTransformIndex);
+					}
+					NeighborFunc(ProxTransformIndex);
+				}
+			}
+		});
+}
+
 void FGeometryCollectionProximityUtility::CopyProximityToConnectionGraph(const TArray<FGeometryContactEdge>* ContactEdges)
 {
 	if (!Collection->HasAttribute("Proximity", FGeometryCollection::GeometryGroup))
diff --git a/Engine/Source/Runtime/Experimental/Chaos/Public/GeometryCollection/Facades/CollectionHierarchyFacade.h b/Engine/Source/Runtime/Experimental/Chaos/Public/GeometryCollection/Facades/CollectionHierarchyFacade.h
--- a/Engine/Source/Runtime/Experimental/Chaos/Public/GeometryCollection/Facades/CollectionHierarchyFacade.h
+++ b/Engine/Source/Runtime/Experimental/Chaos/Public/GeometryCollection/Facades/CollectionHierarchyFacade.h
@@ -53,6 +53,11 @@ namespace Chaos::Facades
 		/** Get parent of a specific transform index */
 		CHAOS_API int32 GetParent(int32 TransformIndex) const;
 
+		/**
+		 * Enumerate transforms at or under TransformIndex where ShouldProcess returns true, stopping recursion and calling ProcessTransformFunc on such transforms.
+		 */
+		CHAOS_API void EnumerateChildTransforms(int32 TransformIndex, TFunctionRef<bool(int32)> ShouldProcess, TFunctionRef<void(int32)> ProcessTransformFunc) const;
+
 		/** 
 		* Get initial level of a specific transform index 
 		* If the attribute is missing return INDEX_NONE
diff --git a/Engine/Source/Runtime/Experimental/Chaos/Public/GeometryCollection/Facades/CollectionTransformSelectionFacade.h b/Engine/Source/Runtime/Experimental/Chaos/Public/GeometryCollection/Facades/CollectionTransformSelectionFacade.h
--- a/Engine/Source/Runtime/Experimental/Chaos/Public/GeometryCollection/Facades/CollectionTransformSelectionFacade.h
+++ b/Engine/Source/Runtime/Experimental/Chaos/Public/GeometryCollection/Facades/CollectionTransformSelectionFacade.h
@@ -101,8 +101,8 @@ namespace GeometryCollection::Facades
 		CHAOS_API TArray<int32> SelectCluster() const;
 
 		/**  */
-		bool CanSelectContact() const { return TransformIndexAttribute.IsValid() && TransformToGeometryIndexAttribute.IsValid(); }
-		CHAOS_API void SelectContact(TArray<int32>& InOutSelection) const;
+		bool CanSelectContact() const { return ParentAttribute.IsValid() && LevelAttribute.IsValid() && SimulationTypeAttribute.IsValid(); }
+		CHAOS_API void SelectContact(TArray<int32>& InOutSelection, bool bIncludeNeighborsInParentLevels = true) const;
 
 		/**  */
 		bool CanSelectParent() const { return ParentAttribute.IsValid(); }
diff --git a/Engine/Source/Runtime/Experimental/Chaos/Public/GeometryCollection/GeometryCollectionProximityUtility.h b/Engine/Source/Runtime/Experimental/Chaos/Public/GeometryCollection/GeometryCollectionProximityUtility.h
--- a/Engine/Source/Runtime/Experimental/Chaos/Public/GeometryCollection/GeometryCollectionProximityUtility.h
+++ b/Engine/Source/Runtime/Experimental/Chaos/Public/GeometryCollection/GeometryCollectionProximityUtility.h
@@ -13,6 +13,11 @@ namespace UE::GeometryCollectionConvexUtility
 	struct FConvexHulls;
 }
 
+namespace Chaos::Facades
+{
+	class FCollectionHierarchyFacade;
+}
+
 UENUM()
 enum class EProximityMethod : int32
 {
@@ -56,6 +61,8 @@ public:
 
 	CHAOS_API void InvalidateProximity();
 
+	CHAOS_API void EnumerateNeighbors(const Chaos::Facades::FCollectionHierarchyFacade& Hierarchy, int32 BoneIdx, TFunctionRef<void(int32)> NeighborFunc, bool bIncludeNeighborsInParentLevels = true, bool bFilterDuplicates = true);
+
 	// Stores stats about the contact between two geometries
 	struct FGeometryContactEdge
 	{
diff --git a/Engine/Source/Runtime/Experimental/GeometryCollectionEngine/Private/GeometryCollection/GeometryCollectionComponent.cpp b/Engine/Source/Runtime/Experimental/GeometryCollectionEngine/Private/GeometryCollection/GeometryCollectionComponent.cpp
--- a/Engine/Source/Runtime/Experimental/GeometryCollectionEngine/Private/GeometryCollection/GeometryCollectionComponent.cpp
+++ b/Engine/Source/Runtime/Experimental/GeometryCollectionEngine/Private/GeometryCollection/GeometryCollectionComponent.cpp
@@ -5468,6 +5468,7 @@ void FScopedColorEdit::SelectBones(GeometryCollection::ESelectionMode SelectionM
 			const TManagedArray<int32>& TransformIndex = GeometryCollectionPtr->TransformIndex;
 			const TManagedArray<int32>& TransformToGeometryIndex = GeometryCollectionPtr->TransformToGeometryIndex;
 			const TManagedArray<TSet<int32>>& Proximity = GeometryCollectionPtr->GetAttribute<TSet<int32>>("Proximity", FGeometryCollection::GeometryGroup);
+			Chaos::Facades::FCollectionHierarchyFacade HierarchyFacade(*GeometryCollectionPtr);
 
 			const TArray<int32> SelectedBones = GetSelectedBones();
 
@@ -5475,15 +5476,10 @@ void FScopedColorEdit::SelectBones(GeometryCollection::ESelectionMode SelectionM
 			for (int32 Bone : SelectedBones)
 			{
 				NewSelection.Add(Bone);
-				int32 GeometryIdx = TransformToGeometryIndex[Bone];
-				if (GeometryIdx != INDEX_NONE)
+				ProximityUtility.EnumerateNeighbors(HierarchyFacade, Bone, [&NewSelection](int32 NeighborTransformIdx)
 				{
-					const TSet<int32>& Neighbors = Proximity[GeometryIdx];
-					for (int32 NeighborGeometryIndex : Neighbors)
-					{
-						NewSelection.Add(TransformIndex[NeighborGeometryIndex]);
-					}
-				}
+					NewSelection.Add(NeighborTransformIdx);
+				}, true /*allow neighbors in parent levels*/, false /*bFilterDuplicates, not needed since we add to a set*/);
 			}
 
 			ResetBoneSelection();
