Iris
* FNetRefHandleManager now takes FInitParams in it's ctor.
* Tranformed more FNetBitArray& into FNetBitArrayView

#rnx

diff --git a/Engine/Source/Runtime/Experimental/Iris/Core/Private/Iris/ReplicationSystem/NetRefHandleManager.cpp b/Engine/Source/Runtime/Experimental/Iris/Core/Private/Iris/ReplicationSystem/NetRefHandleManager.cpp
--- a/Engine/Source/Runtime/Experimental/Iris/Core/Private/Iris/ReplicationSystem/NetRefHandleManager.cpp
+++ b/Engine/Source/Runtime/Experimental/Iris/Core/Private/Iris/ReplicationSystem/NetRefHandleManager.cpp
@@ -15,33 +15,26 @@
 #include "UObject/UObjectGlobals.h"
 #include "UObject/CoreNetTypes.h"
 #include "ProfilingDebugging/CsvProfiler.h"
+#include "Stats/StatsMisc.h"
 
 CSV_DEFINE_CATEGORY(IrisCommon, true);
 
 namespace UE::Net::Private
 {
 
-FNetRefHandleManager::FNetRefHandleManager(FReplicationProtocolManager& InReplicationProtocolManager, uint32 InReplicationSystemId, uint32 InMaxActiveObjectCount, uint32 InPreAllocatedObjectCount)
-: ActiveObjectCount(0)
-, MaxActiveObjectCount(InMaxActiveObjectCount)
-, PreAllocatedObjectCount(InPreAllocatedObjectCount > 0 ? InPreAllocatedObjectCount : 1)	// PreAllocatedObjectCount must be a minimum of 1 to account for InvalidInternalIndex.
-, ReplicationSystemId(InReplicationSystemId)
-, GlobalScopableInternalIndices(MaxActiveObjectCount)
-, ScopeFrameData(MaxActiveObjectCount)
-, RelevantObjectsInternalIndices(MaxActiveObjectCount)
-, PolledObjectsInternalIndices(MaxActiveObjectCount)
-, DirtyObjectsToQuantize(MaxActiveObjectCount)
-, AssignedInternalIndices(MaxActiveObjectCount)
-, SubObjectInternalIndices(MaxActiveObjectCount)
-, DependentObjectInternalIndices(MaxActiveObjectCount)
-, ObjectsWithDependentObjectsInternalIndices(MaxActiveObjectCount)
-, DestroyedStartupObjectInternalIndices(MaxActiveObjectCount)
-, WantToBeDormantInternalIndices(MaxActiveObjectCount)
-, ObjectsWithPreUpdate(MaxActiveObjectCount)
-, NextStaticHandleIndex(1) // Index 0 is always reserved, for both static and dynamic handles
-, NextDynamicHandleIndex(1)
-, ReplicationProtocolManager(InReplicationProtocolManager)
+FNetRefHandleManager::FNetRefHandleManager(FReplicationProtocolManager& InReplicationProtocolManager)
+	: ReplicationProtocolManager(InReplicationProtocolManager)
 {
+}
+
+void FNetRefHandleManager::Init(const FInitParams& InitParams)
+{
+	MaxActiveObjectCount = InitParams.MaxActiveObjectCount;
+	ReplicationSystemId = InitParams.ReplicationSystemId;
+
+	// PreAllocatedObjectCount must be a minimum of 1 to account for InvalidInternalIndex.
+	PreAllocatedObjectCount = InitParams.PreAllocatedObjectCount > 0 ? InitParams.PreAllocatedObjectCount : 1;
+
 	// Ensure that the pre-allocated object count is not greater than the maximum. If it is, just set it to the maximum.
 	if (PreAllocatedObjectCount > MaxActiveObjectCount)
 	{
@@ -52,20 +45,49 @@ FNetRefHandleManager::FNetRefHandleManager(FReplicationProtocolManager& InReplic
 	// Calculate the largest internal index and must be a minimum of 0 to support InvalidInternalIndex.
 	LargestInternalIndex = PreAllocatedObjectCount > 0 ? PreAllocatedObjectCount - 1 : 0;
 
-	UE_LOG(LogIris, Log, TEXT("NetRefHandleManager: Configured with MaxActiveObjectCount=%d, PreAllocatedObjectCount=%d and LargestInternalIndex=%d."), MaxActiveObjectCount, PreAllocatedObjectCount, LargestInternalIndex);
+	UE_LOG(LogIris, Log, TEXT("NetRefHandleManager: Configured with MaxActiveObjectCount=%d, PreAllocatedObjectCount=%d, LargestInternalIndex=%d"),
+		MaxActiveObjectCount, PreAllocatedObjectCount, LargestInternalIndex);
 
 	static_assert(InvalidInternalIndex == 0, "FNetRefHandleManager::InvalidInternalIndex has an unexpected value");
-	// Mark the invalid index as used
-	AssignedInternalIndices.SetBit(0);
 
+	// Initialize TNetChunkedArrays with PreAllocatedObjectCount
 	ReplicatedObjectData = TNetChunkedArray<FReplicatedObjectData>(PreAllocatedObjectCount, EInitMemory::Constructor);
 	ReplicatedObjectRefCount = TNetChunkedArray<uint16>(PreAllocatedObjectCount, EInitMemory::Zero);
 	ReplicatedObjectStateBuffers = TNetChunkedArray<uint8*>(PreAllocatedObjectCount, EInitMemory::Zero);
-
 	ReplicatedInstances = TNetChunkedArray<TObjectPtr<UObject>>(PreAllocatedObjectCount, EInitMemory::Zero);
 
 	// For convenience we initialize ReplicatedObjectData for InvalidInternalIndex so that GetReplicatedObjectDataNoCheck returns something useful.
 	ReplicatedObjectData[InvalidInternalIndex] = FReplicatedObjectData();
+
+	// Init all NetBitArrays here
+	{
+		InitNetBitArray(&ScopeFrameData.CurrentFrameScopableInternalIndices);
+		InitNetBitArray(&ScopeFrameData.PrevFrameScopableInternalIndices);
+		InitNetBitArray(&GlobalScopableInternalIndices);
+		InitNetBitArray(&RelevantObjectsInternalIndices);
+		InitNetBitArray(&PolledObjectsInternalIndices);
+		InitNetBitArray(&DirtyObjectsToQuantize);
+		InitNetBitArray(&AssignedInternalIndices);
+		InitNetBitArray(&SubObjectInternalIndices);
+		InitNetBitArray(&DependentObjectInternalIndices);
+		InitNetBitArray(&ObjectsWithDependentObjectsInternalIndices);
+		InitNetBitArray(&DestroyedStartupObjectInternalIndices);
+		InitNetBitArray(&WantToBeDormantInternalIndices);
+		InitNetBitArray(&ObjectsWithPreUpdate);
+	}
+
+	// Mark the invalid index as used
+	AssignedInternalIndices.SetBit(0);
+}
+
+void FNetRefHandleManager::Deinit()
+{
+
+}
+
+void FNetRefHandleManager::InitNetBitArray(FNetBitArray* NetBitArray)
+{
+	NetBitArray->Init(MaxActiveObjectCount);
 }
 
 void FNetRefHandleManager::GrowBuffersToLargestIndex(uint32 InternalIndex)
@@ -75,8 +97,8 @@ void FNetRefHandleManager::GrowBuffersToLargestIndex(uint32 InternalIndex)
 		QUICK_SCOPE_CYCLE_COUNTER(STAT_FNetRefHandleManager_GrowBuffersToLargestIndex);
 		CSV_CUSTOM_STAT(IrisCommon, DynamicBufferGrowCount, 1, ECsvCustomStatOp::Accumulate);
 
-		// This call will add the neccessary number of elements and chunks to ReplicatedObjectRefCount
-		// to accomidate InternalIndex. Once this is done, we determine how many more elements could be
+		// This call will add the necessary number of elements and chunks to ReplicatedObjectRefCount
+		// to accommodate InternalIndex. Once this is done, we determine how many more elements could be
 		// added to the array without adding a new chunk and ensure all of the other buffers have this
 		// many elements. 
 		// 
@@ -134,7 +156,7 @@ FInternalNetRefIndex FNetRefHandleManager::InternalCreateNetObject(const FNetRef
 	const uint32 InternalIndex = GetNextFreeInternalIndex();
 	if (InternalIndex != InvalidInternalIndex)
 	{
-		UE_LOG(LogIris, Verbose, TEXT("FNetRefHandleManager::InternalCreateNetObject ( InternalIndex: %u ) %s"), InternalIndex, *NetRefHandle.ToString());
+		UE_LOG(LogIris, Verbose, TEXT("FNetRefHandleManager::InternalCreateNetObject: (InternalIndex: %u) (%s)"), InternalIndex, *NetRefHandle.ToString());
 
 		// Track the largest internal index and grow internal buffers if neccessary.
 		GrowBuffersToLargestIndex(InternalIndex);
@@ -323,7 +345,7 @@ void FNetRefHandleManager::InternalDestroyNetObject(FInternalNetRefIndex Interna
 {
 	FReplicatedObjectData& Data = ReplicatedObjectData[InternalIndex];
 
-	UE_LOG(LogIris, Verbose, TEXT("FNetRefHandleManager::InternalDestroyNetObject ( InternalIndex: %u ) %s"), InternalIndex, *Data.RefHandle.ToString());
+	UE_LOG(LogIris, Verbose, TEXT("FNetRefHandleManager::InternalDestroyNetObject: (InternalIndex: %u) (%s)"), InternalIndex, *Data.RefHandle.ToString());
 
 	uint8* StateBuffer = ReplicatedObjectStateBuffers[InternalIndex];
 	// Free any allocated resources
diff --git a/Engine/Source/Runtime/Experimental/Iris/Core/Private/Iris/ReplicationSystem/NetRefHandleManager.h b/Engine/Source/Runtime/Experimental/Iris/Core/Private/Iris/ReplicationSystem/NetRefHandleManager.h
--- a/Engine/Source/Runtime/Experimental/Iris/Core/Private/Iris/ReplicationSystem/NetRefHandleManager.h
+++ b/Engine/Source/Runtime/Experimental/Iris/Core/Private/Iris/ReplicationSystem/NetRefHandleManager.h
@@ -8,6 +8,7 @@
 #include "Containers/ArrayView.h"
 #include "Net/Core/NetBitArray.h"
 #include "Net/Core/NetHandle/NetHandle.h"
+#include "Iris/IrisConstants.h"
 #include "Iris/Core/NetChunkedArray.h"
 #include "Iris/ReplicationState/ReplicationStateDescriptor.h"
 #include "Iris/ReplicationSystem/NetRefHandle.h"
@@ -122,7 +123,18 @@ public:
 	typedef TMap<FNetHandle, FInternalNetRefIndex> FNetHandleMap;
 
 public:
-	FNetRefHandleManager(FReplicationProtocolManager& InReplicationProtocolManager, uint32 InReplicationSystemId, uint32 InMaxActiveObjectCount, uint32 InPreAllocatedObjectCount);
+
+	struct FInitParams
+	{
+		uint32 ReplicationSystemId = 0;
+		uint32 MaxActiveObjectCount = 0;
+		uint32 PreAllocatedObjectCount = 0;
+	};
+
+	FNetRefHandleManager(FReplicationProtocolManager& InReplicationProtocolManager);
+
+	void Init(const FInitParams& InitParams);
+	void Deinit();
 
 	/** Callback triggered at the beginning of PreSendUpdate. Used to sync current frame data. */
 	void OnPreSendUpdate();
@@ -217,7 +229,7 @@ public:
 	FNetBitArrayView GetDirtyObjectsToQuantize() const { return MakeNetBitArrayView(DirtyObjectsToQuantize); }
 
 	// Get bitarray for all internal indices that currently are assigned
-	const FNetBitArray& GetAssignedInternalIndices() const { return AssignedInternalIndices; }
+	const FNetBitArrayView GetAssignedInternalIndices() const { return MakeNetBitArrayView(AssignedInternalIndices); }
 
 	// SubObjects
 	const FNetBitArray& GetSubObjectInternalIndices() const { return SubObjectInternalIndices; }
@@ -233,8 +245,8 @@ public:
 	bool SetSubObjectNetCondition(FInternalNetRefIndex SubObjectInternalIndex, FLifeTimeConditionStorage SubObjectCondition);
 
 	// DependentObjects
-	const FNetBitArray& GetDependentObjectInternalIndices() const { return DependentObjectInternalIndices; }
-	const FNetBitArray& GetObjectsWithDependentObjectsInternalIndices() const { return ObjectsWithDependentObjectsInternalIndices; }
+	const FNetBitArrayView GetDependentObjectInternalIndices() const { return MakeNetBitArrayView(DependentObjectInternalIndices); }
+	const FNetBitArrayView GetObjectsWithDependentObjectsInternalIndices() const { return MakeNetBitArrayView(ObjectsWithDependentObjectsInternalIndices); }
 	bool AddDependentObject(FNetRefHandle ParentHandle, FNetRefHandle DependentHandle, EDependentObjectSchedulingHint SchedulingHint, EAddDependentObjectFlags Flags = EAddDependentObjectFlags::WarnIfAlreadyDependentObject);
 	void RemoveDependentObject(FNetRefHandle ParentHandle, FNetRefHandle DependentHandle);
 
@@ -283,11 +295,11 @@ public:
 	bool GetIsDestroyedStartupObject(FInternalNetRefIndex InternalIndex) const { return DestroyedStartupObjectInternalIndices.GetBit(InternalIndex); }
 	inline uint32 GetOriginalDestroyedStartupObjectIndex(FInternalNetRefIndex InternalIndex) const;
 
-	const FNetBitArray& GetDestroyedStartupObjectInternalIndices() const { return DestroyedStartupObjectInternalIndices; }
+	const FNetBitArrayView GetDestroyedStartupObjectInternalIndices() const { return MakeNetBitArrayView(DestroyedStartupObjectInternalIndices); }
 
 	/** List of replicated objects that want to be dormant */
-	const FNetBitArray& GetWantToBeDormantInternalIndices() const { return WantToBeDormantInternalIndices; }
-	FNetBitArray& GetWantToBeDormantInternalIndices() { return WantToBeDormantInternalIndices; }
+	const FNetBitArrayView GetWantToBeDormantInternalIndices() const { return MakeNetBitArrayView(WantToBeDormantInternalIndices); }
+	FNetBitArrayView GetWantToBeDormantInternalIndices() { return MakeNetBitArrayView(WantToBeDormantInternalIndices); }
 
 	/** Return a string to identify the object linked to an index in logs */
 	FString PrintObjectFromIndex(FInternalNetRefIndex ObjectIndex) const;
@@ -322,6 +334,9 @@ public:
 	};
 
 private:
+
+	void InitNetBitArray(FNetBitArray* NetBitArray);
+
 	FInternalNetRefIndex InternalCreateNetObject(const FNetRefHandle NetRefHandle, const FNetHandle GlobalHandle, const FReplicationProtocol* ReplicationProtocol);
 	void InternalDestroyNetObject(FInternalNetRefIndex InternalIndex);
 
@@ -343,33 +358,28 @@ private:
 	void GrowBuffersToLargestIndex(uint32 InternalIndex);
 
 	// The current replicated object count
-	uint32 ActiveObjectCount;
+	uint32 ActiveObjectCount = 0;
 
 	// Max allowed replicated object count
-	uint32 MaxActiveObjectCount;
+	uint32 MaxActiveObjectCount = 0;
 
 	// The number of pre-allocated objects used by internal buffers.
-	uint32 PreAllocatedObjectCount;
+	uint32 PreAllocatedObjectCount = 0;
 
 	// The largest internal index value that has been used.
-	uint32 LargestInternalIndex;
+	uint32 LargestInternalIndex = 0;
 
 	// A delegate that is triggered when the largest encountered internal index increases.
 	FOnLargestIndexIncrease OnLargestIndexIncreaseDelegate;
 
-	uint32 ReplicationSystemId;
+	uint32 ReplicationSystemId = UE::Net::InvalidReplicationSystemId;
 
 	FRefHandleMap RefHandleToInternalIndex;
 	FNetHandleMap NetHandleToInternalIndex;
 
 	struct FScopeFrameData
 	{
-		FScopeFrameData(uint32 InMaxActiveObjectCount)
-			: bIsValid(false)
-			, CurrentFrameScopableInternalIndices(InMaxActiveObjectCount)
-			, PrevFrameScopableInternalIndices(InMaxActiveObjectCount)
-			
-		{ }
+		FScopeFrameData() : bIsValid(false) { }
 
 		// Controls if the frame data can be read or not
 		uint32 bIsValid : 1;
@@ -441,9 +451,10 @@ private:
 	// Raw pointers to all bound instances
 	TNetChunkedArray<TObjectPtr<UObject>> ReplicatedInstances;
 
-	// Assign handles
-	uint64 NextStaticHandleIndex;
-	uint64 NextDynamicHandleIndex;
+	// Index ready to be assigned to the next replicated object
+	// Index 0 is always reserved, for both static and dynamic handles
+	uint64 NextStaticHandleIndex = 1;
+	uint64 NextDynamicHandleIndex = 1;
 
 	FNetDependencyData SubObjects;
 
diff --git a/Engine/Source/Runtime/Experimental/Iris/Core/Private/Iris/ReplicationSystem/ObjectReplicationBridge.cpp b/Engine/Source/Runtime/Experimental/Iris/Core/Private/Iris/ReplicationSystem/ObjectReplicationBridge.cpp
--- a/Engine/Source/Runtime/Experimental/Iris/Core/Private/Iris/ReplicationSystem/ObjectReplicationBridge.cpp
+++ b/Engine/Source/Runtime/Experimental/Iris/Core/Private/Iris/ReplicationSystem/ObjectReplicationBridge.cpp
@@ -899,7 +899,7 @@ void UObjectReplicationBridge::PruneStaleObjects()
 	};
 
 	// Iterate over assigned indices and detect if any of the replicated instances has been garbagecollected (excluding DestroyedStartupObjectInternalIndices) as they never have an instance
-	FNetBitArray::ForAllSetBits(LocalNetRefHandleManager.GetAssignedInternalIndices(), LocalNetRefHandleManager.GetDestroyedStartupObjectInternalIndices(), FNetBitArray::AndNotOp, DetectStaleObjectsFunc);
+	FNetBitArrayView::ForAllSetBits(LocalNetRefHandleManager.GetAssignedInternalIndices(), LocalNetRefHandleManager.GetDestroyedStartupObjectInternalIndices(), FNetBitArrayView::AndNotOp, DetectStaleObjectsFunc);
 
 	// EndReplication/detach stale instances
 	for (FNetRefHandle Handle : MakeArrayView(StaleObjects.GetData(), StaleObjects.Num()))
@@ -951,7 +951,7 @@ void UObjectReplicationBridge::BuildPollList(UE::Net::FNetBitArrayView ObjectsCo
 
 	const FNetRefHandleManager& LocalNetRefHandleManager = ReplicationSystemInternal->GetNetRefHandleManager();
 	const FNetBitArrayView RelevantObjects = LocalNetRefHandleManager.GetRelevantObjectsInternalIndices();
-	const FNetBitArrayView WantToBeDormantObjects = MakeNetBitArrayView(LocalNetRefHandleManager.GetWantToBeDormantInternalIndices());
+	const FNetBitArrayView WantToBeDormantObjects = LocalNetRefHandleManager.GetWantToBeDormantInternalIndices();
 
 	if (bUseFrequencyBasedPolling)
 	{
@@ -1063,7 +1063,7 @@ void UObjectReplicationBridge::BuildPollList(UE::Net::FNetBitArrayView ObjectsCo
 
 			FNetBitArray TempObjectsConsideredForPolling;
 			TempObjectsConsideredForPolling.InitAndCopy(ObjectsConsideredForPolling);
-			FNetBitArray::ForAllSetBits(TempObjectsConsideredForPolling, NetRefHandleManager->GetObjectsWithDependentObjectsInternalIndices(), FNetBitArray::AndOp,
+			FNetBitArrayView::ForAllSetBits(MakeNetBitArrayView(TempObjectsConsideredForPolling), NetRefHandleManager->GetObjectsWithDependentObjectsInternalIndices(), FNetBitArray::AndOp,
 				[&LocalNetRefHandleManager, &ObjectsConsideredForPolling](FInternalNetRefIndex ObjectIndex) 
 				{
 					LocalNetRefHandleManager.ForAllDependentObjectsRecursive(ObjectIndex, [&ObjectsConsideredForPolling, ObjectIndex](FInternalNetRefIndex DependentObjectIndex)
@@ -1254,7 +1254,7 @@ void UObjectReplicationBridge::SetObjectWantsToBeDormant(FNetRefHandle Handle, b
 
 	if (const FInternalNetRefIndex InternalObjectIndex = LocalNetRefHandleManager.GetInternalIndex(Handle))
 	{
-		FNetBitArray& WantToBeDormantObjects = LocalNetRefHandleManager.GetWantToBeDormantInternalIndices();
+		FNetBitArrayView WantToBeDormantObjects = LocalNetRefHandleManager.GetWantToBeDormantInternalIndices();
 
 		// Update pending dormancy status
 		WantToBeDormantObjects.SetBitValue(InternalObjectIndex, bWantsToBeDormant);
diff --git a/Engine/Source/Runtime/Experimental/Iris/Core/Private/Iris/ReplicationSystem/ReplicationSystem.cpp b/Engine/Source/Runtime/Experimental/Iris/Core/Private/Iris/ReplicationSystem/ReplicationSystem.cpp
--- a/Engine/Source/Runtime/Experimental/Iris/Core/Private/Iris/ReplicationSystem/ReplicationSystem.cpp
+++ b/Engine/Source/Runtime/Experimental/Iris/Core/Private/Iris/ReplicationSystem/ReplicationSystem.cpp
@@ -95,6 +95,17 @@ public:
 
 		FNetRefHandleManager& NetRefHandleManager = ReplicationSystemInternal.GetNetRefHandleManager();
 
+		{
+			FNetRefHandleManager::FInitParams HandleInitParams
+			{
+				.ReplicationSystemId = ReplicationSystemId,
+				.MaxActiveObjectCount = Params.MaxReplicatedObjectCount,
+				.PreAllocatedObjectCount = Params.PreAllocatedReplicatedObjectCount,
+			};
+
+			NetRefHandleManager.Init(HandleInitParams);
+		}
+
 		const uint32 MaxObjectCount =  NetRefHandleManager.GetMaxActiveObjectCount();
 
 		// DirtyNetObjectTracking is only needed when object replication is allowed
@@ -261,6 +272,8 @@ public:
 			ObjectReferencePackageMap->MarkAsGarbage();
 			ReplicationSystemInternal.SetIrisObjectReferencePackageMap(static_cast<UIrisObjectReferencePackageMap*>(nullptr));
 		}
+
+		ReplicationSystemInternal.GetNetRefHandleManager().Deinit();
 	}
 
 	void StartPreSendUpdate()
@@ -558,8 +571,8 @@ public:
 			ReplicationConditionals.RemoveConnection(ConnectionId);
 		}
 
-		FReplicationConnections& Connections = ReplicationSystemInternal.GetConnections();
 		{
+			FReplicationConnections& Connections = ReplicationSystemInternal.GetConnections();
 			Connections.RemoveConnection(ConnectionId);
 		}
 	}
diff --git a/Engine/Source/Runtime/Experimental/Iris/Core/Private/Iris/ReplicationSystem/ReplicationSystemInternal.h b/Engine/Source/Runtime/Experimental/Iris/Core/Private/Iris/ReplicationSystem/ReplicationSystemInternal.h
--- a/Engine/Source/Runtime/Experimental/Iris/Core/Private/Iris/ReplicationSystem/ReplicationSystemInternal.h
+++ b/Engine/Source/Runtime/Experimental/Iris/Core/Private/Iris/ReplicationSystem/ReplicationSystemInternal.h
@@ -43,7 +43,7 @@ class FReplicationSystemInternal
 {
 public:
 	explicit FReplicationSystemInternal(const FReplicationSystemInternalInitParams& Params)
-	: NetRefHandleManager(ReplicationProtocolManager, Params.ReplicationSystemId, Params.MaxReplicatedObjectCount, Params.PreAllocatedReplicatedObjectCount)
+	: NetRefHandleManager(ReplicationProtocolManager)
 	, InternalInitParams(Params)
 	, DirtyNetObjectTracker()
 	, ReplicationBridge(nullptr)
diff --git a/Engine/Source/Runtime/Experimental/Iris/Core/Private/Iris/ReplicationSystem/ReplicationTypes.h b/Engine/Source/Runtime/Experimental/Iris/Core/Private/Iris/ReplicationSystem/ReplicationTypes.h
--- a/Engine/Source/Runtime/Experimental/Iris/Core/Private/Iris/ReplicationSystem/ReplicationTypes.h
+++ b/Engine/Source/Runtime/Experimental/Iris/Core/Private/Iris/ReplicationSystem/ReplicationTypes.h
@@ -10,11 +10,11 @@ namespace UE::Net::Private
 
 struct FReplicationParameters
 {
-	uint32 MaxActiveReplicatedObjectCount;	
-	uint32 PreAllocatedReplicatedObjectCount;
-	uint32 MaxReplicatedWriterObjectCount;
-	uint32 PacketSendWindowSize;
-	uint32 ConnectionId;
+	uint32 MaxActiveReplicatedObjectCount = 0;
+	uint32 PreAllocatedReplicatedObjectCount = 0;
+	uint32 MaxReplicatedWriterObjectCount = 0;
+	uint32 PacketSendWindowSize = 0;
+	uint32 ConnectionId = 0;
 	UReplicationSystem* ReplicationSystem = nullptr;
 	bool bAllowSendingAttachmentsToObjectsNotInScope = false;
 	bool bAllowReceivingAttachmentsFromRemoteObjectsNotInScope = false;
diff --git a/Engine/Source/Runtime/Experimental/Iris/Core/Public/Iris/ReplicationSystem/ReplicationSystem.h b/Engine/Source/Runtime/Experimental/Iris/Core/Public/Iris/ReplicationSystem/ReplicationSystem.h
--- a/Engine/Source/Runtime/Experimental/Iris/Core/Public/Iris/ReplicationSystem/ReplicationSystem.h
+++ b/Engine/Source/Runtime/Experimental/Iris/Core/Public/Iris/ReplicationSystem/ReplicationSystem.h
@@ -61,15 +61,40 @@ public:
 
 	struct FReplicationSystemParams
 	{
-		//$IRIS TODO: These need documentation
+		/** The replication bridge that allows communication between the replication system and the game engine  */
 		UReplicationBridge* ReplicationBridge = nullptr;
+
+		/** The maximum amount of netobjects that can be registered to the replication system */
 		uint32 MaxReplicatedObjectCount = 65535U;
+
+		/**
+		 * The amount of netobjects to preallocate internal memory buffers for. Preallocated memory provides faster cache-friendly operations but
+		 * has the downside of using much more memory than might actually be needed. Reduce this value if you are operating on a memory constrained platform.
+		 */
 		uint32 PreAllocatedReplicatedObjectCount = 65535U;
+
+		/**
+		 * The maximum amount of netobjects that can replicate properties to remote connection. Can be much lower on clients where very few
+		 * netobjects have authority and support property replication
+		 */
 		uint32 MaxReplicatedWriterObjectCount = 65535U;
+
+		/** The maximum amount of netobjects that can be added to the delta compression manager */
 		uint32 MaxDeltaCompressedObjectCount = 2048U;
+
+		/** The maximum amount of filter groups that can be created. @see UReplicationSystem::CreateGroup */
 		uint32 MaxNetObjectGroupCount = 2048U;
+
+		/** Is this replication system owned by a server or a client. */
 		bool bIsServer = false;
+
+		/**
+		 * When true enable netobject subsystems like: property replication, filtering, prioritization, deltacompression, dirtytracking, etc.
+		 * Generally false on clients or on lightweight RPC-only systems.
+		 */
 		bool bAllowObjectReplication = false;
+
+		/** Delegate that receives every RPC executed locally. */
 		UE::Net::FForwardNetRPCCallDelegate ForwardNetRPCCallDelegate;
 	};
 
