Dataflow : Add support for icons
- Icons are set in the node metadata
- set the branch core node icon

diff --git a/Engine/Source/Runtime/Experimental/Dataflow/Core/Public/Dataflow/DataflowCoreNodes.h b/Engine/Source/Runtime/Experimental/Dataflow/Core/Public/Dataflow/DataflowCoreNodes.h
--- a/Engine/Source/Runtime/Experimental/Dataflow/Core/Public/Dataflow/DataflowCoreNodes.h
+++ b/Engine/Source/Runtime/Experimental/Dataflow/Core/Public/Dataflow/DataflowCoreNodes.h
@@ -28,7 +28,7 @@ public:
 	FDataflowAnyType Value;
 };
 
-USTRUCT()
+USTRUCT(meta=(Icon="GraphEditor.Branch_16x"))
 struct FDataflowBranchNode : public FDataflowNode
 {
 	GENERATED_USTRUCT_BODY()
@@ -46,7 +46,7 @@ public:
 	UPROPERTY(meta = (DataflowInput, DisplayName = "FalseValue"))
 	FDataflowAnyType FalseValue;
 
-	UPROPERTY(meta = (DataflowInput, DisplayName = "Condition"))
+	UPROPERTY(EditAnywhere, Category="Condition", meta = (DataflowInput, DisplayName = "Condition"))
 	bool bCondition = true;
 
 	UPROPERTY(meta = (DataflowOutput, DisplayName = "Result"))
diff --git a/Engine/Source/Runtime/Experimental/Dataflow/Engine/DataflowEngine.Build.cs b/Engine/Source/Runtime/Experimental/Dataflow/Engine/DataflowEngine.Build.cs
--- a/Engine/Source/Runtime/Experimental/Dataflow/Engine/DataflowEngine.Build.cs
+++ b/Engine/Source/Runtime/Experimental/Dataflow/Engine/DataflowEngine.Build.cs
@@ -12,6 +12,7 @@ namespace UnrealBuildTool.Rules
 				{
 					"Core",
 					"CoreUObject",
+					"SlateCore",
 					"Engine",
 					"DataflowCore",
 					"Chaos"
diff --git a/Engine/Source/Runtime/Experimental/Dataflow/Engine/Private/Dataflow/DataflowEdNode.cpp b/Engine/Source/Runtime/Experimental/Dataflow/Engine/Private/Dataflow/DataflowEdNode.cpp
--- a/Engine/Source/Runtime/Experimental/Dataflow/Engine/Private/Dataflow/DataflowEdNode.cpp
+++ b/Engine/Source/Runtime/Experimental/Dataflow/Engine/Private/Dataflow/DataflowEdNode.cpp
@@ -9,6 +9,8 @@
 #include "Dataflow/DataflowRenderingFactory.h"
 #include "GeometryCollection/Facades/CollectionRenderingFacade.h"
 #include "Logging/LogMacros.h"
+#include "Textures/SlateIcon.h"
+#include "Styling/AppStyle.h"
 
 #include UE_INLINE_GENERATED_CPP_BY_NAME(DataflowEdNode)
 
@@ -396,6 +398,25 @@ void UDataflowEdNode::Serialize(FArchive& Ar)
 
 #if WITH_EDITOR
 
+FSlateIcon UDataflowEdNode::GetIconAndTint(FLinearColor& OutColor) const
+{
+	FSlateIcon Icon;
+	if (TSharedPtr<const FDataflowNode> DataflowNode = GetDataflowNode())
+	{
+		if (const FString* IconName = DataflowNode->TypedScriptStruct()->FindMetaData("Icon"))
+		{
+			Icon = FSlateIcon(FAppStyle::GetAppStyleSetName(), FName(*IconName));
+		}
+	}
+
+	return Icon;
+}
+
+bool UDataflowEdNode::ShowPaletteIconOnNode() const
+{
+	return true;
+}
+
 FLinearColor UDataflowEdNode::GetNodeTitleColor() const
 {
 	if (DataflowGraph)
diff --git a/Engine/Source/Runtime/Experimental/Dataflow/Engine/Public/Dataflow/DataflowEdNode.h b/Engine/Source/Runtime/Experimental/Dataflow/Engine/Public/Dataflow/DataflowEdNode.h
--- a/Engine/Source/Runtime/Experimental/Dataflow/Engine/Public/Dataflow/DataflowEdNode.h
+++ b/Engine/Source/Runtime/Experimental/Dataflow/Engine/Public/Dataflow/DataflowEdNode.h
@@ -40,6 +40,8 @@ public:
 	DATAFLOWENGINE_API virtual void PinConnectionListChanged(UEdGraphPin* Pin) override;
 #endif // WITH_EDITOR && !UE_BUILD_SHIPPING
 #if WITH_EDITOR
+	DATAFLOWENGINE_API virtual FSlateIcon GetIconAndTint(FLinearColor& OutColor) const override;
+	DATAFLOWENGINE_API virtual bool ShowPaletteIconOnNode() const override;
 	DATAFLOWENGINE_API virtual FLinearColor GetNodeTitleColor() const override;
 	DATAFLOWENGINE_API virtual FLinearColor GetNodeBodyTintColor() const override;
 	DATAFLOWENGINE_API virtual FText GetTooltipText() const override;
