Chaos Debug Draw
- debug draw system supports out-of-context draws
- add debug draw timelines to RBAN
- debug draw context cleanup
- new debug draw command system (not used yet)

Debug Draw commands from a thread on which DD is not initialized will fall back to a "global frame" which is a never-ending debug draw frame that we pull debug draw commands from every time we render. Draws to the global frame suffer from the same flicker problems as the old DD backend but it will allow use to remove the old system. The ultimate goal is to give a context to all threads that use DD.

Debug Draw commands can now be more sophisticated and can defer decisions like the color and line thickness, or whether to display at all, to render time. This will be important when we can rewind the debug draw and want to change visualization settings after the commands have been queued. This can also make the draw queue much smaller (e.g, when drawing a mesh we previously had to enqueue a Line command for every edge, and now it is just an implicit object pointer)

diff --git a/Engine/Source/Runtime/AnimGraphRuntime/Private/BoneControllers/AnimNode_RigidBody.cpp b/Engine/Source/Runtime/AnimGraphRuntime/Private/BoneControllers/AnimNode_RigidBody.cpp
--- a/Engine/Source/Runtime/AnimGraphRuntime/Private/BoneControllers/AnimNode_RigidBody.cpp
+++ b/Engine/Source/Runtime/AnimGraphRuntime/Private/BoneControllers/AnimNode_RigidBody.cpp
@@ -8,6 +8,7 @@
 #include "ChaosVDRuntimeModule.h"
 #endif
 
+#include "ChaosDebugDraw/ChaosDDScene.h"
 #include "ClothCollisionSource.h"
 #include "Engine/OverlapResult.h"
 #include "Engine/SkeletalMesh.h"
@@ -19,10 +20,11 @@
 #include "PhysicsEngine/PhysicsConstraintTemplate.h"
 #include "PhysicsEngine/SkeletalBodySetup.h"
 #include "GameFramework/PawnMovementComponent.h"
-#include "Physics/PhysicsInterfaceCore.h"
+#include "Physics/Experimental/PhysScene_Chaos.h"
 #include "Physics/ImmediatePhysics/ImmediatePhysicsActorHandle.h"
 #include "Physics/ImmediatePhysics/ImmediatePhysicsSimulation.h"
 #include "Physics/ImmediatePhysics/ImmediatePhysicsStats.h"
+#include "Physics/PhysicsInterfaceCore.h"
 #include "PhysicsField/PhysicsFieldComponent.h"
 #include "PhysicsEngine/PhysicsSettings.h"
 #include "Logging/MessageLog.h"
@@ -1075,6 +1077,14 @@ void FAnimNode_RigidBody::InitPhysics(const UAnimInstance* InAnimInstance)
 		PhysicsSimulation->GetChaosVDContextData().Type = static_cast<int32>(EChaosVDContextType::Solver);
 #endif
 
+#if CHAOS_DEBUG_DRAW
+		if ((SkeletalMeshComp->GetWorld() != nullptr) && (SkeletalMeshComp->GetWorld()->GetPhysicsScene() != nullptr))
+		{
+			const FString DDName = FString::Format(TEXT("RBAN {0}"), { SkeletalMeshComp->GetName()});
+			PhysicsSimulation->SetDebugDrawScene(DDName, SkeletalMeshComp->GetWorld()->GetPhysicsScene()->GetDebugDrawScene());
+		}
+#endif
+
 		const int32 NumBodies = UsePhysicsAsset->SkeletalBodySetups.Num();
 		Bodies.Empty(NumBodies);
 		BodyAnimData.Reset(NumBodies);
diff --git a/Engine/Source/Runtime/Engine/Private/PhysicsEngine/ImmediatePhysics/ImmediatePhysicsChaos/ImmediatePhysicsSimulation_Chaos.cpp b/Engine/Source/Runtime/Engine/Private/PhysicsEngine/ImmediatePhysics/ImmediatePhysicsChaos/ImmediatePhysicsSimulation_Chaos.cpp
--- a/Engine/Source/Runtime/Engine/Private/PhysicsEngine/ImmediatePhysics/ImmediatePhysicsChaos/ImmediatePhysicsSimulation_Chaos.cpp
+++ b/Engine/Source/Runtime/Engine/Private/PhysicsEngine/ImmediatePhysics/ImmediatePhysicsChaos/ImmediatePhysicsSimulation_Chaos.cpp
@@ -16,6 +16,8 @@
 #include "Chaos/Joint/ChaosJointLog.h"
 #include "Chaos/MassConditioning.h"
 #include "Chaos/PBDJointConstraints.h"
+#include "ChaosDebugDraw/ChaosDDScene.h"
+#include "ChaosDebugDraw/ChaosDDTimeline.h"
 #include "ChaosVisualDebugger/ChaosVisualDebuggerTrace.h"
 #include "Stats/StatsTrace.h"
 
@@ -828,6 +830,10 @@ namespace ImmediatePhysics_Chaos
 
 	void FSimulation::DebugDraw()
 	{
+#if CHAOS_DEBUG_DRAW
+		ChaosDD::Private::FChaosDDScopeTimelineContext DDContext(DDSimulationTimeline, 0.0, 0.0);
+#endif
+
 		DebugDrawStaticParticles();
 		DebugDrawKinematicParticles();
 		DebugDrawDynamicParticles();
@@ -1152,4 +1158,14 @@ namespace ImmediatePhysics_Chaos
 		}
 #endif
 	}
+
+#if CHAOS_DEBUG_DRAW
+	void FSimulation::SetDebugDrawScene(const FString& SceneName, const ChaosDD::Private::FChaosDDScenePtr& InScene)
+	{
+		if (InScene.IsValid() && !DDSimulationTimeline.IsValid())
+		{
+			DDSimulationTimeline = InScene->CreateTimeline(SceneName);
+		}
+	}
+#endif
 }
diff --git a/Engine/Source/Runtime/Engine/Public/Physics/ImmediatePhysics/ImmediatePhysicsChaos/ImmediatePhysicsSimulation_Chaos.h b/Engine/Source/Runtime/Engine/Public/Physics/ImmediatePhysics/ImmediatePhysicsChaos/ImmediatePhysicsSimulation_Chaos.h
--- a/Engine/Source/Runtime/Engine/Public/Physics/ImmediatePhysics/ImmediatePhysicsChaos/ImmediatePhysicsSimulation_Chaos.h
+++ b/Engine/Source/Runtime/Engine/Public/Physics/ImmediatePhysics/ImmediatePhysicsChaos/ImmediatePhysicsSimulation_Chaos.h
@@ -2,6 +2,7 @@
 
 #pragma once
 
+#include "ChaosDebugDraw/ChaosDDTypes.h"
 #include "ChaosVisualDebugger/ChaosVDContextProvider.h"
 #include "Physics/ImmediatePhysics/ImmediatePhysicsChaos/ImmediatePhysicsCore_Chaos.h"
 
@@ -156,6 +157,15 @@ namespace ImmediatePhysics_Chaos
 			return CVDContextData;
 		};
 #endif
+
+#if CHAOS_DEBUG_DRAW
+	public:
+		ENGINE_API void SetDebugDrawScene(const FString& SceneName, const ChaosDD::Private::FChaosDDScenePtr& InScene);
+
+	private:
+		ChaosDD::Private::FChaosDDTimelinePtr DDSimulationTimeline;
+#endif
+
 	};
 
 }
diff --git a/Engine/Source/Runtime/Experimental/Chaos/Private/Chaos/DebugDrawQueue.cpp b/Engine/Source/Runtime/Experimental/Chaos/Private/Chaos/DebugDrawQueue.cpp
--- a/Engine/Source/Runtime/Experimental/Chaos/Private/Chaos/DebugDrawQueue.cpp
+++ b/Engine/Source/Runtime/Experimental/Chaos/Private/Chaos/DebugDrawQueue.cpp
@@ -7,6 +7,9 @@
 
 namespace Chaos
 {
+	bool bChaosDebugDraw_UseNewQueue = true;
+	FAutoConsoleVariableRef CVarChaos_DebugDraw_UseNewQueue(TEXT("p.Chaos.DebugDraw.UseNewQueue"), bChaosDebugDraw_UseNewQueue, TEXT(""));
+
 	bool bChaosDebugDraw_UseLegacyQueue = false;
 	FAutoConsoleVariableRef CVarChaos_DebugDraw_UseLegacyQueue(TEXT("p.Chaos.DebugDraw.UseLegacyQueue"), bChaosDebugDraw_UseLegacyQueue, TEXT(""));
 
diff --git a/Engine/Source/Runtime/Experimental/Chaos/Private/Chaos/Framework/Parallel.cpp b/Engine/Source/Runtime/Experimental/Chaos/Private/Chaos/Framework/Parallel.cpp
--- a/Engine/Source/Runtime/Experimental/Chaos/Private/Chaos/Framework/Parallel.cpp
+++ b/Engine/Source/Runtime/Experimental/Chaos/Private/Chaos/Framework/Parallel.cpp
@@ -3,6 +3,7 @@
 #include "Chaos/Framework/Parallel.h"
 #include "Async/ParallelFor.h"
 #include "Framework/Threading.h"
+#include "ChaosDebugDraw/ChaosDDContext.h"
 
 namespace Chaos
 {
@@ -65,12 +66,22 @@ void Chaos::PhysicsParallelFor(int32 InNum, TFunctionRef<void(int32)> InCallable
 	const bool bIsInGameThreadContext = false;
 #endif
 
-	auto PassThrough = [InCallable, bIsInPhysicsSimContext, bIsInGameThreadContext](int32 Idx)
+#if CHAOS_DEBUG_DRAW
+	const ChaosDD::Private::FChaosDDContext& ParentDDContext = ChaosDD::Private::FChaosDDContext::Get();
+#else
+	int32 ParentDDContext = 0;
+#endif
+
+	auto PassThrough = [InCallable, bIsInPhysicsSimContext, bIsInGameThreadContext, &ParentDDContext](int32 Idx)
 	{
 #if PHYSICS_THREAD_CONTEXT
 		FPhysicsThreadContextScope PTScope(bIsInPhysicsSimContext);
 		FGameThreadContextScope GTScope(bIsInGameThreadContext);
 #endif
+#if CHAOS_DEBUG_DRAW
+		ChaosDD::Private::FChaosDDScopeTaskContext DDTaskContext(ParentDDContext);
+#endif
+
 		InCallable(Idx);
 	};
 
diff --git a/Engine/Source/Runtime/Experimental/Chaos/Private/Chaos/PBDRigidsEvolutionGBF.cpp b/Engine/Source/Runtime/Experimental/Chaos/Private/Chaos/PBDRigidsEvolutionGBF.cpp
--- a/Engine/Source/Runtime/Experimental/Chaos/Private/Chaos/PBDRigidsEvolutionGBF.cpp
+++ b/Engine/Source/Runtime/Experimental/Chaos/Private/Chaos/PBDRigidsEvolutionGBF.cpp
@@ -363,8 +363,8 @@ void FPBDRigidsEvolutionGBF::Advance(const FReal Dt,const FReal MaxStepDt,const
 void FPBDRigidsEvolutionGBF::AdvanceOneTimeStep(const FReal Dt,const FSubStepInfo& SubStepInfo)
 {	
 #if CHAOS_DEBUG_DRAW
-	// @todo(chaos): need the time when we want to show historical debug draw
-	ChaosDD::Private::FChaosDDContext::BeginFrame(CDDTickTimeline, 0.0f, Dt);
+	// @todo(chaos): will need the time when historical debug draw is supported
+	ChaosDD::Private::FChaosDDScopeTimelineContext DDContext(CDDTickTimeline, 0.0f, Dt);
 #endif
 
 	PrepareTick();
@@ -372,20 +372,11 @@ void FPBDRigidsEvolutionGBF::AdvanceOneTimeStep(const FReal Dt,const FSubStepInf
 	AdvanceOneTimeStepImpl(Dt, SubStepInfo);
 
 	UnprepareTick();
-
-#if CHAOS_DEBUG_DRAW
-	ChaosDD::Private::FChaosDDContext::EndFrame();
-#endif
 }
 
 #if CHAOS_DEBUG_DRAW
 void FPBDRigidsEvolutionGBF::SetDebugDrawScene(const ChaosDD::Private::FChaosDDScenePtr& InCDDScene)
 {
-	if (CDDScene.IsValid() && CDDTickTimeline.IsValid())
-	{
-		CDDScene->ReleaseTimeline(CDDTickTimeline);
-	}
-
 	CDDScene = InCDDScene;
 	CDDTickTimeline.Reset();
 
diff --git a/Engine/Source/Runtime/Experimental/Chaos/Private/ChaosDebugDraw/ChaosDDContext.cpp b/Engine/Source/Runtime/Experimental/Chaos/Private/ChaosDebugDraw/ChaosDDContext.cpp
--- a/Engine/Source/Runtime/Experimental/Chaos/Private/ChaosDebugDraw/ChaosDDContext.cpp
+++ b/Engine/Source/Runtime/Experimental/Chaos/Private/ChaosDebugDraw/ChaosDDContext.cpp
@@ -3,8 +3,9 @@
 #include "ChaosDebugDraw/ChaosDDLog.h"
 #include "ChaosDebugDraw/ChaosDDFrame.h"
 #include "ChaosDebugDraw/ChaosDDTimeline.h"
-
 #include "ChaosLog.h"
+#include "HAL/IConsoleManager.h"
+
 
 DEFINE_LOG_CATEGORY(LogChaosDD);
 
@@ -12,52 +13,138 @@ DEFINE_LOG_CATEGORY(LogChaosDD);
 
 namespace ChaosDD::Private
 {
-	FChaosDDContext::FChaosDDContext()
+	bool bChaosDebugDraw_EnableGlobalQueue = true;
+	FAutoConsoleVariableRef CVarChaos_DebugDraw_EnableGlobalQueue(TEXT("p.Chaos.DebugDraw.EnableGlobalQueue"), bChaosDebugDraw_EnableGlobalQueue, TEXT(""));
+
+	FCriticalSection FChaosDDContext::GlobalFrameCS;
+	FChaosDDFramePtr FChaosDDContext::GlobalFrame;
+	int32 FChaosDDContext::GlobalCommandBudget = 20000;
+
+	//
+	//
+	//
+	//
+
+	FChaosDDScopeTimelineContext::FChaosDDScopeTimelineContext(const FChaosDDTimelinePtr& InTimeline, double InTime, double InDt)
 	{
-		TimelineStack.Push(FChaosDDTimelinePtr());
+		if (InTimeline.IsValid())
+		{
+			Timeline = InTimeline;
+			Timeline->BeginFrame(InTime, InDt);
+
+			FChaosDDContext& Context = FChaosDDContext::Get();
+			ParentFrame = Context.Frame;
+			Context.Frame = Timeline->GetActiveFrame();
+		}
 	}
 
-	void FChaosDDContext::BeginFrame(const FChaosDDTimelinePtr& InTimeline, double Time, double Dt)
+	FChaosDDScopeTimelineContext::~FChaosDDScopeTimelineContext()
 	{
-		FChaosDDContext& Context = Get();
+		if (Timeline.IsValid())
+		{
+			Timeline->EndFrame();
+			Timeline.Reset();
 
-		if (InTimeline.IsValid())
+			FChaosDDContext& Context = FChaosDDContext::Get();
+			Context.Frame = ParentFrame;
+			ParentFrame.Reset();
+		}
+	}
+
+	//
+	//
+	//
+	//
+
+	FChaosDDScopeTaskContext::FChaosDDScopeTaskContext(const FChaosDDContext& InParentDDContext)
+	{
+		if (InParentDDContext.Frame.IsValid())
 		{
-			Context.TimelineStack.Push(InTimeline);
+			FChaosDDContext& Context = FChaosDDContext::Get();
 
-			Context.GetTimeline()->BeginFrame(Time, Dt);
+			ParentFrame = Context.Frame;
+			Context.Frame = InParentDDContext.Frame;
 		}
+	}
 
-		Context.UpdateCachedFramePtr();
+	FChaosDDScopeTaskContext::~FChaosDDScopeTaskContext()
+	{
+		FChaosDDContext& Context = FChaosDDContext::Get();
 
+		Context.Frame = ParentFrame;
+		ParentFrame.Reset();
 	}
 
-	void FChaosDDContext::EndFrame()
+	//
+	//
+	//
+	//
+
+	FChaosDDContext::FChaosDDContext()
 	{
-		FChaosDDContext& Context = Get();
+	}
 
-		if (Context.GetTimeline().IsValid())
-		{
-			Context.GetTimeline()->EndFrame();
+	const FChaosDDFramePtr& FChaosDDContext::GetGlobalFrame()
+	{
+		FScopeLock Lock(&GlobalFrameCS);
 
-			Context.TimelineStack.Pop();
+		if (!GlobalFrame.IsValid())
+		{
+			CreateGlobalFrame();
 		}
 
-		Context.UpdateCachedFramePtr();
+		return GlobalFrame;
+	}
+
+	void FChaosDDContext::CreateGlobalFrame()
+	{
+		if (bChaosDebugDraw_EnableGlobalQueue)
+		{
+			FScopeLock Lock(&GlobalFrameCS);
+
+			GlobalFrame = MakeShared<FChaosDDGlobalFrame>(GlobalCommandBudget);
+		}
 	}
 
-	void FChaosDDContext::UpdateCachedFramePtr()
+	FChaosDDFramePtr FChaosDDContext::ExtractGlobalFrame()
 	{
-		if (GetTimeline().IsValid())
+		FScopeLock Lock(&GlobalFrameCS);
+
+		// Handle toggling the cvar in the runtime
+		if (!bChaosDebugDraw_EnableGlobalQueue && GlobalFrame.IsValid())
 		{
-			Frame = GetTimeline()->GetActiveFrame();
+			GlobalFrame.Reset();
 		}
-		else
+
+		if (GlobalFrame.IsValid())
 		{
-			Frame.Reset();
+			return GlobalFrame->ExtractFrame();
 		}
+
+		return {};
 	}
 
+	void FChaosDDContext::SetGlobalDrawRegion(const FSphere3d& InDrawRegion)
+	{
+		FScopeLock Lock(&GlobalFrameCS);
+
+		if (GlobalFrame.IsValid())
+		{
+			GlobalFrame->SetDrawRegion(InDrawRegion);
+		}
+	}
+
+	void FChaosDDContext::SetGlobalCommandBudget(int32 InCommandBudget)
+	{
+		FScopeLock Lock(&GlobalFrameCS);
+
+		GlobalCommandBudget = InCommandBudget;
+
+		if (GlobalFrame.IsValid())
+		{
+			GlobalFrame->SetCommandBudget(InCommandBudget);
+		}
+	}
 }
 
 #endif
\ No newline at end of file
diff --git a/Engine/Source/Runtime/Experimental/Chaos/Private/ChaosDebugDraw/ChaosDDScene.cpp b/Engine/Source/Runtime/Experimental/Chaos/Private/ChaosDebugDraw/ChaosDDScene.cpp
--- a/Engine/Source/Runtime/Experimental/Chaos/Private/ChaosDebugDraw/ChaosDDScene.cpp
+++ b/Engine/Source/Runtime/Experimental/Chaos/Private/ChaosDebugDraw/ChaosDDScene.cpp
@@ -1,5 +1,6 @@
 // Copyright Epic Games, Inc. All Rights Reserved.
 #include "ChaosDebugDraw/ChaosDDScene.h"
+#include "ChaosDebugDraw/ChaosDDContext.h"
 #include "ChaosDebugDraw/ChaosDDFrame.h"
 #include "ChaosDebugDraw/ChaosDDLog.h"
 #include "ChaosDebugDraw/ChaosDDRenderer.h"
@@ -9,6 +10,18 @@
 
 namespace ChaosDD::Private
 {
+	template<typename LambdaType>
+	void VisitTimelines(const TArray<FChaosDDTimelineWeakPtr>& Timelines, const LambdaType& Visitor)
+	{
+		for (const FChaosDDTimelineWeakPtr& TimelineWeak : Timelines)
+		{
+			if (const FChaosDDTimelinePtr& Timeline = TimelineWeak.Pin())
+			{
+				Visitor(Timeline);
+			}
+		}
+	}
+
 	FChaosDDScene::FChaosDDScene(const FString& InName)
 		: Name(InName)
 		, DrawRegion(FVector::Zero(), 0.0)
@@ -26,10 +39,15 @@ namespace ChaosDD::Private
 
 		DrawRegion = InDrawRegion;
 
-		for (const FChaosDDTimelinePtr& Timeline : Timelines)
-		{
-			Timeline->SetDrawRegion(InDrawRegion);
-		}
+		VisitTimelines(Timelines,
+			[&InDrawRegion](const FChaosDDTimelinePtr& Timeline)
+			{
+				Timeline->SetDrawRegion(InDrawRegion);
+			});
+
+		// The global timeline doesn't use a draw region because it could be shared
+		// among multiple viewports and they probably don't want the same region
+		//FChaosDDContext::SetGlobalDrawRegion(InDrawRegion);
 	}
 
 	void FChaosDDScene::SetCommandBudget(int32 InCommandBudget)
@@ -38,10 +56,13 @@ namespace ChaosDD::Private
 
 		CommandBudget = InCommandBudget;
 
-		for (const FChaosDDTimelinePtr& Timeline : Timelines)
-		{
-			Timeline->SetCommandBudget(InCommandBudget);
-		}
+		VisitTimelines(Timelines,
+			[InCommandBudget](const FChaosDDTimelinePtr& Timeline)
+			{
+				Timeline->SetCommandBudget(InCommandBudget);
+			});
+
+		FChaosDDContext::SetGlobalCommandBudget(InCommandBudget);
 	}
 
 	FChaosDDTimelinePtr FChaosDDScene::CreateTimeline(const FString& InName)
@@ -50,37 +71,40 @@ namespace ChaosDD::Private
 
 		FChaosDDTimelinePtr Timeline = MakeShared<FChaosDDTimeline>(InName, CommandBudget);
 
-		Timelines.Add(Timeline);
+		Timelines.Add(Timeline.ToWeakPtr());
 
 		return Timeline;
 	}
 
-	void FChaosDDScene::ReleaseTimeline(const FChaosDDTimelinePtr& InTimeline)
+	void FChaosDDScene::RenderLatestFrames(IChaosDDRenderer& Renderer, bool bIncludeGlobalFrame)
 	{
-		FScopeLock Lock(&TimelinesCS);
+		TArray<FChaosDDFramePtr> Frames = GetFrames();
 
-		for (int32 TimelineIndex = 0; TimelineIndex < Timelines.Num(); ++TimelineIndex)
+		if (bIncludeGlobalFrame)
 		{
-			if (Timelines[TimelineIndex].Get() == InTimeline.Get())
-			{
-				Timelines.RemoveAt(TimelineIndex, 1, EAllowShrinking::No);
-			}
+			Frames.Add(FChaosDDContext::ExtractGlobalFrame());
 		}
-	}
 
-	void FChaosDDScene::RenderLatestFrames(IChaosDDRenderer& Renderer)
-	{
-		TArray<FChaosDDFramePtr> Frames = GetFrames();
 		for (const FChaosDDFramePtr& Frame : Frames)
 		{
-			UE_LOG(LogChaosDD, VeryVerbose, TEXT("Render %s %d %d Commands"), *Frame->GetTimeline()->GetName(), Frame->GetFrameIndex(), Frame->GetNumCommands());
-
-			// Render the legacy commands
-			Frame->VisitCommands(
-				[&Renderer](const Chaos::FLatentDrawCommand& Command)
-				{
-					Renderer.DrawCommand(Command);
-				});
+			if (Frame.IsValid())
+			{
+				UE_LOG(LogChaosDD, VeryVerbose, TEXT("Render %s %d %d Commands"), *Frame->GetTimeline()->GetName(), Frame->GetFrameIndex(), Frame->GetNumCommands());
+
+				// Render the legacy commands
+				Frame->VisitLatentCommands(
+					[&Renderer](const Chaos::FLatentDrawCommand& Command)
+					{
+						Renderer.RenderLatentCommand(Command);
+					});
+
+				// Render the commands
+				Frame->VisitCommands(
+					[&Renderer](const FChaosDDCommand& Command)
+					{
+						Command(Renderer);
+					});
+			}
 		}
 	}
 
@@ -89,12 +113,27 @@ namespace ChaosDD::Private
 		FScopeLock Lock(&TimelinesCS);
 
 		TArray<FChaosDDFramePtr> Frames;
-		for (const FChaosDDTimelinePtr& Timeline : Timelines)
-		{
-			Timeline->GetFrames(Frames);
-		}
+
+		VisitTimelines(Timelines,
+			[&Frames](const FChaosDDTimelinePtr& Timeline)
+			{
+				Timeline->GetFrames(Frames);
+			});
+
 		return Frames;
 	}
+
+	void FChaosDDScene::PruneTimelines()
+	{
+		FScopeLock Lock(&TimelinesCS);
+
+		// Remove all timelines that have been deleted
+		Timelines.RemoveAll(
+			[](const FChaosDDTimelineWeakPtr& Timeline)
+			{
+				return !Timeline.IsValid();
+			});
+	}
 }
 
 #endif
diff --git a/Engine/Source/Runtime/Experimental/Chaos/Private/PBDRigidsSolver.cpp b/Engine/Source/Runtime/Experimental/Chaos/Private/PBDRigidsSolver.cpp
--- a/Engine/Source/Runtime/Experimental/Chaos/Private/PBDRigidsSolver.cpp
+++ b/Engine/Source/Runtime/Experimental/Chaos/Private/PBDRigidsSolver.cpp
@@ -1302,7 +1302,7 @@ namespace Chaos
 		const FReal StartSimTime = GetSolverTime();
 
 #if CHAOS_DEBUG_DRAW
-		ChaosDD::Private::FChaosDDContext::BeginFrame(CDDFrameTimeline, StartSimTime, GetLastDt());
+		ChaosDD::Private::FChaosDDScopeTimelineContext DDContext(CDDFrameTimeline, StartSimTime, GetLastDt());
 #endif
 
 #if !(UE_BUILD_SHIPPING || UE_BUILD_TEST)
@@ -1328,20 +1328,11 @@ namespace Chaos
 			//final step so we can destroy proxies
 			DestroyPendingProxies_Internal();
 		}
-
-#if CHAOS_DEBUG_DRAW
-		ChaosDD::Private::FChaosDDContext::EndFrame();
-#endif
 	}
 
 #if CHAOS_DEBUG_DRAW
 	void FPBDRigidsSolver::SetDebugDrawScene(const ChaosDD::Private::FChaosDDScenePtr& InCDDScene)
 	{
-		if (CDDScene.IsValid() && CDDFrameTimeline.IsValid())
-		{
-			CDDScene->ReleaseTimeline(CDDFrameTimeline);
-		}
-
 		CDDScene = InCDDScene;
 		CDDFrameTimeline.Reset();
 
diff --git a/Engine/Source/Runtime/Experimental/Chaos/Public/Chaos/DebugDrawQueue.h b/Engine/Source/Runtime/Experimental/Chaos/Public/Chaos/DebugDrawQueue.h
--- a/Engine/Source/Runtime/Experimental/Chaos/Public/Chaos/DebugDrawQueue.h
+++ b/Engine/Source/Runtime/Experimental/Chaos/Public/Chaos/DebugDrawQueue.h
@@ -21,6 +21,7 @@ class AActor;
 #if CHAOS_DEBUG_DRAW
 namespace Chaos
 {
+extern CHAOS_API bool bChaosDebugDraw_UseNewQueue;
 extern CHAOS_API bool bChaosDebugDraw_UseLegacyQueue;
 
 /** Thread-safe single-linked list (lock-free). (Taken from Light Mass, should probably just move into core) */
@@ -295,43 +296,47 @@ private:
 
 	bool AcceptCommand(int Cost, const FVec3& Position)
 	{
-		if (!bChaosDebugDraw_UseLegacyQueue)
+		if (bChaosDebugDraw_UseNewQueue)
 		{
-			const ChaosDD::Private::FChaosDDFramePtr& DDFrame = ChaosDD::Private::FChaosDDContext::Get().GetFrame();
-			if (DDFrame.IsValid())
+			ChaosDD::Private::FChaosDDFrameWriter DDWriter = ChaosDD::Private::FChaosDDContext::Get().GetWriter();
+			if (DDWriter.IsInDrawRegion(Position))
 			{
-				if (DDFrame->IsInDrawRegion(Position))
-				{
-					return DDFrame->AddToCost(Cost);
-				}
-				return false;
+				return DDWriter.AddToCost(Cost);
 			}
+			return false;
 		}
 
-		if (IsInRegionOfInterest(Position))
+		// If we get here...
+		// This thread has not been set up for the new debug draw system so default back to the old
+		// way that has issues with timing in async mode.
+		if (bChaosDebugDraw_UseLegacyQueue)
 		{
-			RequestedCommandCost += Cost;
-			return IsInBudget();
+			if (IsInRegionOfInterest(Position))
+			{
+				RequestedCommandCost += Cost;
+				return IsInBudget();
+			}
 		}
+
 		return false;
 	}
 
 	void AddCommand(const FLatentDrawCommand& Command)
 	{
 		// Try to add to the new debug draw system which queues commands per World and ticking thread
-		if (!bChaosDebugDraw_UseLegacyQueue)
+		if (bChaosDebugDraw_UseNewQueue)
 		{
-			const ChaosDD::Private::FChaosDDFramePtr& DDFrame = ChaosDD::Private::FChaosDDContext::Get().GetFrame();
-			if (DDFrame.IsValid())
-			{
-				DDFrame->EnqueueCommand(Command);
-				return;
-			}
+			ChaosDD::Private::FChaosDDContext::GetWriter().EnqueueLatentCommand(Command);
+			return;
 		}
 
+		// If we get here...
 		// This thread has not been set up for the new debug draw system so default back to the old
 		// way that has issues with timing in async mode.
-		CommandQueue.Add(Command);
+		if (bChaosDebugDraw_UseLegacyQueue)
+		{
+			CommandQueue.Add(Command);
+		}
 	}
 
 	bool IsInBudget() const
diff --git a/Engine/Source/Runtime/Experimental/Chaos/Public/Chaos/ParticleHandle.h b/Engine/Source/Runtime/Experimental/Chaos/Public/Chaos/ParticleHandle.h
--- a/Engine/Source/Runtime/Experimental/Chaos/Public/Chaos/ParticleHandle.h
+++ b/Engine/Source/Runtime/Experimental/Chaos/Public/Chaos/ParticleHandle.h
@@ -1715,6 +1715,11 @@ public:
 	FKinematicGeometryParticleHandle* CastToKinematicParticle() { return MHandle->CastToKinematicParticle(); }
 	const FPBDRigidParticleHandle* CastToRigidParticle() const { return MHandle->CastToRigidParticle(); }
 	FPBDRigidParticleHandle* CastToRigidParticle() { return MHandle->CastToRigidParticle(); }
+	const TPBDGeometryCollectionParticleHandleImp<FReal, 3, true>* CastToGeometryCollection() const { return MHandle->CastToGeometryCollection(); }
+	TPBDGeometryCollectionParticleHandleImp<FReal, 3, true>* CastToGeometryCollection() { return MHandle->CastToGeometryCollection(); }
+	const TPBDRigidClusteredParticleHandleImp<FReal, 3, true>* CastToClustered() const { return MHandle->CastToClustered(); }
+	TPBDRigidClusteredParticleHandleImp<FReal, 3, true>* CastToClustered() { return MHandle->CastToClustered(); }
+
 	const FGeometryParticleHandle* GeometryParticleHandle() const { return MHandle; }
 	FGeometryParticleHandle* GeometryParticleHandle() { return MHandle; }
 	//Needed for templated code to be the same
diff --git a/Engine/Source/Runtime/Experimental/Chaos/Public/ChaosDebugDraw/ChaosDDContext.h b/Engine/Source/Runtime/Experimental/Chaos/Public/ChaosDebugDraw/ChaosDDContext.h
--- a/Engine/Source/Runtime/Experimental/Chaos/Public/ChaosDebugDraw/ChaosDDContext.h
+++ b/Engine/Source/Runtime/Experimental/Chaos/Public/ChaosDebugDraw/ChaosDDContext.h
@@ -1,8 +1,11 @@
 // Copyright Epic Games, Inc. All Rights Reserved.
 #pragma once
 
+#include "ChaosDebugDraw/ChaosDDFrame.h"
 #include "ChaosDebugDraw/ChaosDDTypes.h"
+#include "HAL/CriticalSection.h"
 #include "HAL/ThreadSingleton.h"
+#include "Misc/ScopeLock.h"
 
 #if CHAOS_DEBUG_DRAW
 
@@ -17,31 +20,74 @@ namespace ChaosDD::Private
 	public:
 		FChaosDDContext();
 
-		// The frame we should be drawing into on this thread
-		const FChaosDDFramePtr& GetFrame() const
-		{ 
-			return Frame;
+		static FChaosDDFrameWriter GetWriter()
+		{
+			return FChaosDDFrameWriter(Get().GetFrame());
 		}
 
-		// The timeline that holds the frames from this thread
-		const FChaosDDTimelinePtr& GetTimeline() const
+	private:
+		friend class FChaosDDScene;
+		friend class FChaosDDScopeTimelineContext;
+		friend class FChaosDDScopeTaskContext;
+
+		const FChaosDDFramePtr& GetFrame() const
 		{
-			// NOTE: TimelineStack has null as its first element so that we can always return a shared pointer by ref
-			return TimelineStack.Top();
+			// The frame we should be drawing into on this thread
+			if (!Frame.IsValid())
+			{
+				// If there is no Context set up on this thread we fall back to a global frame that 
+				// is tied to the game thread. If debug draw commands are queued while the game thread
+				// is rendering the DDScene, the commands will be split across frames resulting in flicker
+				// @todo(chaos): ideally we don't have a global frame - try to get rid of it
+				return GetGlobalFrame();
+			}
+
+			return Frame;
 		}
 
-		// Allocate a new frame for the timeline and set up the context
-		static void BeginFrame(const FChaosDDTimelinePtr& InTimeline, double Time, double Dt);
+		// Global frame management
+		static const FChaosDDFramePtr& GetGlobalFrame();
+		static void CreateGlobalFrame();
+		static FChaosDDFramePtr ExtractGlobalFrame();
+		static void SetGlobalDrawRegion(const FSphere3d& InDrawRegion);
+		static void SetGlobalCommandBudget(int32 InCommandBudget);
+
+		// The frame to draw to on this thread (or null)
+		FChaosDDFramePtr Frame;
+
+		// Global frame: fallback for out-of-context debug draw
+		static FCriticalSection GlobalFrameCS;
+		static FChaosDDFramePtr GlobalFrame;
+		static int32 GlobalCommandBudget;
+	};
 
-		// Signal that the frame is complete and can be rendered
-		static void EndFrame();
+	//
+	// A scoped Debug Draw Context for use on the thread that owns the timeline
+	//
+	class CHAOS_API FChaosDDScopeTimelineContext
+	{
+	public:
+		FChaosDDScopeTimelineContext(const FChaosDDTimelinePtr& InTimeline, double InTime, double InDt);
+		~FChaosDDScopeTimelineContext();
 
 	private:
-		void UpdateCachedFramePtr();
+		FChaosDDTimelinePtr Timeline;
+		FChaosDDFramePtr ParentFrame;
+	};
 
-		// We need a stack of timelines to handle nested BeginFrame calls. E.g., for sub-steps in the solver.
-		TArray<FChaosDDTimelinePtr> TimelineStack;
-		FChaosDDFramePtr Frame;
+	//
+	// A scoped Debug Draw Context for use in tasks and parallel-for etc.
+	// Assumes that the task is kicked off from a thread that has an active debug draw context,
+	// which should be passed into this context.
+	//
+	class CHAOS_API FChaosDDScopeTaskContext
+	{
+	public:
+		FChaosDDScopeTaskContext(const FChaosDDContext& InParentDDContext);
+		~FChaosDDScopeTaskContext();
+
+	private:
+		FChaosDDFramePtr ParentFrame;
 	};
 }
 
diff --git a/Engine/Source/Runtime/Experimental/Chaos/Public/ChaosDebugDraw/ChaosDDFrame.h b/Engine/Source/Runtime/Experimental/Chaos/Public/ChaosDebugDraw/ChaosDDFrame.h
--- a/Engine/Source/Runtime/Experimental/Chaos/Public/ChaosDebugDraw/ChaosDDFrame.h
+++ b/Engine/Source/Runtime/Experimental/Chaos/Public/ChaosDebugDraw/ChaosDDFrame.h
@@ -14,9 +14,20 @@
 
 namespace ChaosDD::Private
 {
+	//
+	// A frame is a sequence of debug draw commands.
+	// A command is just a functor that uses a DD Renderer and can be a simple as drawing
+	// a line, or as complex as drawing a set of rigid bodies, constraints, etc.
+	// 
+	// E.g., See FChaosDDLine, FChaosDDSphere, FChaosDDParticle
+	//
+	using FChaosDDCommand = TFunction<void(const IChaosDDRenderer&)>;
+
 	//
 	// A single frame of debug draw data
 	//
+	// @todo(chaos): move commands to a per-thread buffer and eliminate write locks
+	//
 	class CHAOS_API FChaosDDFrame : public TSharedFromThis<FChaosDDFrame>
 	{
 	public:
@@ -31,11 +42,11 @@ namespace ChaosDD::Private
 		{
 			if (InCommandQueueLength > 0)
 			{
-				Commands.Reserve(InCommandQueueLength);
+				LatentCommands.Reserve(InCommandQueueLength);
 			}
 		}
 
-		~FChaosDDFrame()
+		virtual ~FChaosDDFrame()
 		{
 		}
 
@@ -59,7 +70,12 @@ namespace ChaosDD::Private
 			return Dt;
 		}
 
-		bool IsInDrawRegion(const FVector& InPos)
+		void SetDrawRegion(const FSphere3d& InRegion)
+		{
+			DrawRegion = InRegion;
+		}
+
+		bool IsInDrawRegion(const FVector& InPos) const
 		{
 			if (DrawRegion.W > 0.0)
 			{
@@ -68,7 +84,7 @@ namespace ChaosDD::Private
 			return true;
 		}
 
-		bool IsInDrawRegion(const FSphere3d& InSphere)
+		bool IsInDrawRegion(const FSphere3d& InSphere) const
 		{
 			if (DrawRegion.W > 0.0)
 			{
@@ -77,65 +93,272 @@ namespace ChaosDD::Private
 			return true;
 		}
 
-		bool IsInDrawRegion(const FBox3d& InBox)
+		bool IsInDrawRegion(const FBox3d& InBox) const
 		{
 			if (DrawRegion.W > 0.0)
 			{
-				const FBox3d DrawRegionBox = FBox3d(DrawRegion.Center - FVector(DrawRegion.W), DrawRegion.Center + FVector(DrawRegion.W));
-				return DrawRegionBox.Intersect(InBox);
+				const double BoxDistanceSq = InBox.ComputeSquaredDistanceToPoint(DrawRegion.Center);
+				return BoxDistanceSq < FMath::Square(DrawRegion.W);
 			}
 			return true;
 		}
 
+		void SetCommandBudget(int32 InCommandBudget)
+		{
+			CommandBudget = InCommandBudget;
+		}
+
 		bool AddToCost(int32 InCost)
 		{
 			FScopeLock Lock(&CommandsCS);
 
-			CommandCost += InCost;
-
 			// A budget of zero means infinite
-			return ((CommandBudget == 0) || (CommandCost <= CommandBudget));
+			if ((CommandCost + InCost <= CommandBudget) || (CommandBudget == 0))
+			{
+				CommandCost += InCost;
+				return true;
+			}
+
+			return false;
 		}
 
-		void EnqueueCommand(const Chaos::FLatentDrawCommand& InCommand)
+		void EnqueueCommand(FChaosDDCommand&& InCommand)
 		{
 			FScopeLock Lock(&CommandsCS);
 
-			Commands.Add(InCommand);
+			Commands.Emplace(MoveTemp(InCommand));
+		}
+
+		void EnqueueLatentCommand(const Chaos::FLatentDrawCommand& InCommand)
+		{
+			FScopeLock Lock(&CommandsCS);
+
+			LatentCommands.Add(InCommand);
 		}
 
 		int32 GetNumCommands() const
 		{
 			FScopeLock Lock(&CommandsCS);
 
-			return Commands.Num();
+			return LatentCommands.Num();
 		}
 
+
+		// VisitorType = void(const FChaosDDCommand& Command)
 		template<typename VisitorType>
 		void VisitCommands(const VisitorType& Visitor)
 		{
 			FScopeLock Lock(&CommandsCS);
 
-			for (const Chaos::FLatentDrawCommand& Command : Commands)
+			for (const FChaosDDCommand& Command : Commands)
 			{
 				Visitor(Command);
 			}
 		}
 
-	private:
+		// VisitorType = void(const Chaos::FLatentDrawCommand& Command)
+		template<typename VisitorType>
+		void VisitLatentCommands(const VisitorType& Visitor)
+		{
+			FScopeLock Lock(&CommandsCS);
+
+			for (const Chaos::FLatentDrawCommand& Command : LatentCommands)
+			{
+				Visitor(Command);
+			}
+		}
+
+		// Used by the global frame to prevent render while queuing commands
+		virtual void BeginWrite()
+		{
+		}
+
+		// Used by the global frame to prevent render while queuing commands
+		virtual void EndWrite()
+		{
+		}
+
+		// Used by the global frame to extract all debug draw commands so far into a new frame for rendering
+		virtual FChaosDDFramePtr ExtractFrame()
+		{
+			const int32 CommandQueueLength = Commands.Max();
+			const int32 LatentCommandQueueLength = LatentCommands.Max();
+
+			FChaosDDFrame* ExtractedFrame = new FChaosDDFrame(MoveTemp(*this));
+
+			Commands.Reset(CommandQueueLength);
+			LatentCommands.Reset(LatentCommandQueueLength);
+			CommandCost = 0;
+
+			return FChaosDDFramePtr(ExtractedFrame);
+		}
+
+	protected:
+		friend class FChaosDDGlobalFrame;
+
+		FChaosDDFrame(FChaosDDFrame&& Other)
+			: Timeline(Other.Timeline)
+			, FrameIndex(Other.FrameIndex)
+			, Time(Other.Time)
+			, Dt(Other.Dt)
+			, DrawRegion(Other.DrawRegion)
+			, CommandBudget(Other.CommandBudget)
+			, CommandCost(Other.CommandCost)
+			, Commands(MoveTemp(Other.Commands))
+			, LatentCommands(MoveTemp(Other.LatentCommands))
+		{
+		}
+
 		FChaosDDTimelineWeakPtr Timeline;
 		int64 FrameIndex;
 		double Time;
 		double Dt;
-		const FSphere3d& DrawRegion;
+		FSphere3d DrawRegion;
 		int32 CommandBudget;
 		int32 CommandCost;
 
-		// Legacy debug draw commands in a lock-protected queue
-		// @todo(chaos): we could move these to the per-thread buffer
-		TArray<Chaos::FLatentDrawCommand> Commands;
+		// Debug draw commands
+		TArray<FChaosDDCommand> Commands;
+
+		// Legacy debug draw commands (see FDebugDrawQueue)
+		TArray<Chaos::FLatentDrawCommand> LatentCommands;
+
 		mutable FCriticalSection CommandsCS;
 	};
+
+	//
+	// A special frame used for out-of-frame debug draw. All debug draw commands from a thread that does not
+	// have a context set up will use the global frame. This global frame suffers will be flickery because
+	// the render may occur while enqueueing a set of related debug draw commands.
+	// 
+	// @todo(chaos): eventually all threads that want debug draw should have a valid frame and this will be redundant.
+	//
+	class CHAOS_API FChaosDDGlobalFrame : public FChaosDDFrame
+	{
+	public:
+
+		FChaosDDGlobalFrame(int32 InCommandBudget)
+			: FChaosDDFrame(FChaosDDTimelinePtr(), 0, 0.0f, 0.0f, FSphere3d(FVector::Zero(), 0.0), InCommandBudget, 0)
+		{
+		}
+
+		virtual void BeginWrite() override
+		{
+			FrameWriteCS.Lock();
+		}
+
+		virtual void EndWrite() override
+		{
+			FrameWriteCS.Unlock();
+		}
+
+		// Create a new frame containing the accumulated commands and reset this frame
+		virtual FChaosDDFramePtr ExtractFrame() override
+		{
+			FScopeLock Lock(&FrameWriteCS);
+
+			return FChaosDDFrame::ExtractFrame();
+		}
+
+	protected:
+
+		mutable FCriticalSection FrameWriteCS;
+	};
+
+	//
+	// Used to write to a debug draw frame.
+	// Currently this writes to the Frame's draw buffer and holds a lock
+	// preventing the frame from being Ended. Eventually this will be a
+	// per-thread buffer to avoid the need for locks.
+	//
+	class CHAOS_API FChaosDDFrameWriter
+	{
+	public:
+		FChaosDDFrameWriter(const FChaosDDFramePtr& InFrame)
+			: Frame(InFrame)
+		{
+			if (Frame.IsValid())
+			{
+				Frame->BeginWrite();
+			}
+		}
+
+		~FChaosDDFrameWriter()
+		{
+			if (Frame.IsValid())
+			{
+				Frame->EndWrite();
+			}
+		}
+
+		bool IsInDrawRegion(const FVector& InPos) const
+		{
+			if (Frame.IsValid())
+			{
+				return Frame->IsInDrawRegion(InPos);
+			}
+			return false;
+		}
+
+		bool IsInDrawRegion(const FSphere3d& InSphere) const
+		{
+			if (Frame.IsValid())
+			{
+				return Frame->IsInDrawRegion(InSphere);
+			}
+			return false;
+		}
+
+		bool IsInDrawRegion(const FBox3d& InBox) const
+		{
+			if (Frame.IsValid())
+			{
+				return Frame->IsInDrawRegion(InBox);
+			}
+			return false;
+		}
+
+		bool AddToCost(int32 InCost)
+		{
+			if (Frame.IsValid())
+			{
+				return Frame->AddToCost(InCost);
+			}
+			return false;
+		}
+
+		bool TryEnqueueCommand(int32 InCost, const FBox3d& InBox, FChaosDDCommand&& InCommand)
+		{
+			if (Frame.IsValid())
+			{
+				if (Frame->IsInDrawRegion(InBox) && Frame->AddToCost(InCost))
+				{
+					Frame->EnqueueCommand(MoveTemp(InCommand));
+					return true;
+				}
+			}
+			return false;
+		}
+
+		void EnqueueCommand(FChaosDDCommand&& InCommand)
+		{
+			if (Frame.IsValid())
+			{
+				Frame->EnqueueCommand(MoveTemp(InCommand));
+			}
+		}
+
+		void EnqueueLatentCommand(const Chaos::FLatentDrawCommand& InCommand)
+		{
+			if (Frame.IsValid())
+			{
+				Frame->EnqueueLatentCommand(InCommand);
+			}
+		}
+
+	private:
+		FChaosDDFramePtr Frame;
+	};
 }
 
 #endif
diff --git a/Engine/Source/Runtime/Experimental/Chaos/Public/ChaosDebugDraw/ChaosDDRenderer.h b/Engine/Source/Runtime/Experimental/Chaos/Public/ChaosDebugDraw/ChaosDDRenderer.h
--- a/Engine/Source/Runtime/Experimental/Chaos/Public/ChaosDebugDraw/ChaosDDRenderer.h
+++ b/Engine/Source/Runtime/Experimental/Chaos/Public/ChaosDebugDraw/ChaosDDRenderer.h
@@ -22,7 +22,15 @@ namespace ChaosDD::Private
 		IChaosDDRenderer() {}
 		virtual ~IChaosDDRenderer() {}
 
-		virtual void DrawCommand(const Chaos::FLatentDrawCommand& Command) const = 0;
+		// Are we rendering a Server scene?
+		virtual bool IsServer() const = 0;
+
+		// Utility functions for use by Debug Draw commands (e.g., FChaosDDLine)
+		virtual void RenderLine(const FVector3d& A, const FVector3d& B, const FColor& Color, float LineThickness, float Lifetime) const = 0;
+		virtual void RenderBox(const FVector3d& Position, const FQuat4d& Rotation, const FVector3d& Size, const FColor& Color, float LineThickness, float Lifetime) const = 0;
+
+		// Render legacy debug draw command (See FChaosDDScene::RenderLatestFrames)
+		virtual void RenderLatentCommand(const Chaos::FLatentDrawCommand& Command) const = 0;
 	};
 }
 
diff --git a/Engine/Source/Runtime/Experimental/Chaos/Public/ChaosDebugDraw/ChaosDDScene.h b/Engine/Source/Runtime/Experimental/Chaos/Public/ChaosDebugDraw/ChaosDDScene.h
--- a/Engine/Source/Runtime/Experimental/Chaos/Public/ChaosDebugDraw/ChaosDDScene.h
+++ b/Engine/Source/Runtime/Experimental/Chaos/Public/ChaosDebugDraw/ChaosDDScene.h
@@ -17,7 +17,7 @@ namespace ChaosDD::Private
 	// Debug draw system for a world. In PIE there will be one of these for the server and each client.
 	//
 	// @todo(chaos): enable retention of debug draw frames and debug draw from a specific time
-	class CHAOS_API FChaosDDScene
+	class CHAOS_API FChaosDDScene : public TSharedFromThis<FChaosDDScene>
 	{
 	public:
 		FChaosDDScene(const FString& InName);
@@ -35,20 +35,19 @@ namespace ChaosDD::Private
 		void SetCommandBudget(int32 InCommandBudget);
 
 		// Create a new timeline. E.g., PT, GT, RBAN
+		// The caller must hold a shared pointer to the timeline to keep it alive.
 		FChaosDDTimelinePtr CreateTimeline(const FString& Name);
-		
-		// Release a timeline
-		void ReleaseTimeline(const FChaosDDTimelinePtr& Scene);
 
 		// Render the most recent frame from each timeline
-		void RenderLatestFrames(IChaosDDRenderer& Renderer);
+		void RenderLatestFrames(IChaosDDRenderer& Renderer, bool bIncludeGlobalFrame);
 
 	private:
 		TArray<FChaosDDFramePtr> GetFrames();
+		void PruneTimelines();
 
 		FString Name;
 		FCriticalSection TimelinesCS;
-		TArray<FChaosDDTimelinePtr> Timelines;
+		TArray<FChaosDDTimelineWeakPtr> Timelines;
 		FSphere3d DrawRegion;
 		int32 CommandBudget;
 	};
diff --git a/Engine/Source/Runtime/Experimental/Chaos/Public/ChaosDebugDraw/ChaosDDTimeline.h b/Engine/Source/Runtime/Experimental/Chaos/Public/ChaosDebugDraw/ChaosDDTimeline.h
--- a/Engine/Source/Runtime/Experimental/Chaos/Public/ChaosDebugDraw/ChaosDDTimeline.h
+++ b/Engine/Source/Runtime/Experimental/Chaos/Public/ChaosDebugDraw/ChaosDDTimeline.h
@@ -7,6 +7,7 @@
 #include "Containers/Map.h"
 #include "HAL/CriticalSection.h"
 #include "Math/Sphere.h"
+#include "Misc/ScopeLock.h"
 
 #if CHAOS_DEBUG_DRAW
 
@@ -38,11 +39,15 @@ namespace ChaosDD::Private
 
 		void SetDrawRegion(const FSphere3d& InDrawRegion)
 		{
+			FScopeLock Lock(&FramesCS);
+
 			DrawRegion = InDrawRegion;
 		}
 
 		void SetCommandBudget(int32 InCommandBudget)
 		{
+			FScopeLock Lock(&FramesCS);
+
 			CommandBudget = InCommandBudget;
 		}
 
diff --git a/Engine/Source/Runtime/Experimental/Chaos/Public/ChaosDebugDraw/ChaosDDTypes.h b/Engine/Source/Runtime/Experimental/Chaos/Public/ChaosDebugDraw/ChaosDDTypes.h
--- a/Engine/Source/Runtime/Experimental/Chaos/Public/ChaosDebugDraw/ChaosDDTypes.h
+++ b/Engine/Source/Runtime/Experimental/Chaos/Public/ChaosDebugDraw/ChaosDDTypes.h
@@ -2,22 +2,32 @@
 #pragma once
 
 #include "Chaos/Declares.h"
+#include "Math/MathFwd.h"
 #include "Templates/SharedPointer.h"
 #include "Templates/UniquePtr.h"
 
 #if CHAOS_DEBUG_DRAW
 
+namespace Chaos
+{
+	struct FLatentDrawCommand;
+}
+
+
 // @todo(chaos): Move to ChaosDD when API is decided
 namespace ChaosDD::Private
 {
 	class FChaosDDContext;
 	class FChaosDDFrame;
+	class FChaosDDFrameWriter;
+	class FChaosDDGlobalFrame;
 	class FChaosDDScene;
 	class FChaosDDTimeline;
 
 	class IChaosDDRenderer;
 
 	using FChaosDDFramePtr = TSharedPtr<FChaosDDFrame, ESPMode::ThreadSafe>;
+	using FChaosDDGlobalFramePtr = TSharedPtr<FChaosDDGlobalFrame, ESPMode::ThreadSafe>;
 	using FChaosDDScenePtr = TSharedPtr<FChaosDDScene, ESPMode::ThreadSafe>;
 	using FChaosDDTimelinePtr = TSharedPtr<FChaosDDTimeline, ESPMode::ThreadSafe>;
 	using FChaosDDTimelineWeakPtr = TWeakPtr<FChaosDDTimeline, ESPMode::ThreadSafe>;
diff --git a/Engine/Source/Runtime/Experimental/ChaosSolverEngine/Private/Chaos/ChaosDebugDrawComponent.cpp b/Engine/Source/Runtime/Experimental/ChaosSolverEngine/Private/Chaos/ChaosDebugDrawComponent.cpp
--- a/Engine/Source/Runtime/Experimental/ChaosSolverEngine/Private/Chaos/ChaosDebugDrawComponent.cpp
+++ b/Engine/Source/Runtime/Experimental/ChaosSolverEngine/Private/Chaos/ChaosDebugDrawComponent.cpp
@@ -228,7 +228,22 @@ public:
 		bIsPaused = false;
 	}
 
-	virtual void DrawCommand(const Chaos::FLatentDrawCommand& Command) const override final
+	virtual bool IsServer() const override final
+	{
+		return World->GetNetMode() == ENetMode::NM_DedicatedServer;
+	}
+
+	virtual void RenderLine(const FVector3d& A, const FVector3d& B, const FColor& Color, float LineThickness, float Lifetime) const override final
+	{
+		DrawDebugLine(World, A, B, Color, false, CommandLifeTime(Lifetime), DepthPriority, LineThickness);
+	}
+
+	virtual void RenderBox(const FVector3d& Position, const FQuat4d& Rotation, const FVector3d& Size, const FColor& Color, float LineThickness, float Lifetime) const override final
+	{
+		DrawDebugBox(World, Position, Size, Rotation, Color, false, CommandLifeTime(Lifetime), DepthPriority, LineThickness);
+	}
+
+	virtual void RenderLatentCommand(const Chaos::FLatentDrawCommand& Command) const override final
 	{
 		const bool bDrawUe = bChaosDebugDraw_DrawMode != 1;
 		if (bDrawUe)
@@ -402,7 +417,8 @@ void UChaosDebugDrawComponent::TickComponent(float DeltaTime, enum ELevelTick Ti
 				if (!bIsPaused)
 				{
 					FChaosDDRenderer CDDRenderer = FChaosDDRenderer(World, GetOwner());
-					CDDScene->RenderLatestFrames(CDDRenderer);
+					constexpr bool bRenderGlobalFrame = true;
+					CDDScene->RenderLatestFrames(CDDRenderer, bRenderGlobalFrame);
 				}
 			}
 		}
