GeometryCollection : API updates
- EnforceSingleRoot( const FString& ...)
- CHAOS_API is usually before the type

diff --git a/Engine/Source/Runtime/Experimental/Chaos/Private/GeometryCollection/Facades/CollectionTransformFacade.cpp b/Engine/Source/Runtime/Experimental/Chaos/Private/GeometryCollection/Facades/CollectionTransformFacade.cpp
--- a/Engine/Source/Runtime/Experimental/Chaos/Private/GeometryCollection/Facades/CollectionTransformFacade.cpp
+++ b/Engine/Source/Runtime/Experimental/Chaos/Private/GeometryCollection/Facades/CollectionTransformFacade.cpp
@@ -356,7 +356,7 @@ namespace GeometryCollection::Facades
 	}
 
 
-	CHAOS_API void FCollectionTransformFacade::EnforceSingleRoot(FString RootName)
+	CHAOS_API void FCollectionTransformFacade::EnforceSingleRoot(const FString & RootName)
 	{
 		TArray<int32> Roots = GetRootIndices();
 		if (Roots.Num() > 1)
@@ -367,7 +367,7 @@ namespace GeometryCollection::Facades
 			{
 				const TManagedArray<FString>& BoneName = BoneNameAttribute.Get();
 
-				auto MakeUnique = [&BoneName](FString& NewName)
+				auto MakeUnique = [&BoneName](FString NewName)
 				{
 					int32 CurrentIndex = 1;
 					FString TestName = NewName;
diff --git a/Engine/Source/Runtime/Experimental/Chaos/Public/GeometryCollection/Facades/CollectionTransformFacade.h b/Engine/Source/Runtime/Experimental/Chaos/Public/GeometryCollection/Facades/CollectionTransformFacade.h
--- a/Engine/Source/Runtime/Experimental/Chaos/Public/GeometryCollection/Facades/CollectionTransformFacade.h
+++ b/Engine/Source/Runtime/Experimental/Chaos/Public/GeometryCollection/Facades/CollectionTransformFacade.h
@@ -106,22 +106,22 @@ namespace GeometryCollection::Facades
 		CHAOS_API void SetBoneTransformToIdentity(int32 BoneIdx);
 
 		/** Does the transform heirarchy have a cycle*/
-		static bool CHAOS_API HasCycle(const TManagedArray<int32>& Parents, int32 Node);
+		static CHAOS_API bool HasCycle(const TManagedArray<int32>& Parents, int32 Node);
 
 		/** Does the transform heirarchy have a cycle*/
-		static bool CHAOS_API HasCycle(const TManagedArray<int32>& Parents, const TArray<int32>& SelectedBones);
+		static CHAOS_API bool HasCycle(const TManagedArray<int32>& Parents, const TArray<int32>& SelectedBones);
 
 		/** Parent a single transform */
-		void CHAOS_API ParentTransform(const int32 TransformIndex, const int32 ChildIndex);
+		CHAOS_API void ParentTransform(const int32 TransformIndex, const int32 ChildIndex);
 
 		/**  Parent the list of transforms to the selected index. */
-		void CHAOS_API ParentTransforms(const int32 TransformIndex, const TArray<int32>& SelectedBones);
+		CHAOS_API void ParentTransforms(const int32 TransformIndex, const TArray<int32>& SelectedBones);
 
 		/**  Unparent the child index from its parent */
-		void CHAOS_API UnparentTransform(const int32 ChildIndex);
+		CHAOS_API void UnparentTransform(const int32 ChildIndex);
 
 		/** Adds a Identity transform and nests all roots under the new transform */
-		CHAOS_API void EnforceSingleRoot(FString RootName);
+		CHAOS_API void EnforceSingleRoot(const FString & RootName);
 
 	private:
 		const FManagedArrayCollection& ConstCollection;
