Dataflow: Add ability to show/hide input pins.

Cloth Asset Editor: Make all existing optional input pins hidden by default.

diff --git a/Engine/Plugins/ChaosClothAssetEditor/Source/ChaosClothAssetDataflowNodes/Private/ChaosClothAsset/AddWeightMapNode.cpp b/Engine/Plugins/ChaosClothAssetEditor/Source/ChaosClothAssetDataflowNodes/Private/ChaosClothAsset/AddWeightMapNode.cpp
--- a/Engine/Plugins/ChaosClothAssetEditor/Source/ChaosClothAssetDataflowNodes/Private/ChaosClothAsset/AddWeightMapNode.cpp
+++ b/Engine/Plugins/ChaosClothAssetEditor/Source/ChaosClothAssetDataflowNodes/Private/ChaosClothAsset/AddWeightMapNode.cpp
@@ -147,8 +147,16 @@ FChaosClothAssetAddWeightMapNode::FChaosClothAssetAddWeightMapNode(const Dataflo
 	: FDataflowTerminalNode(InParam, InGuid)
 {
 	RegisterInputConnection(&Collection);
-	RegisterInputConnection(&InputName.StringValue, GET_MEMBER_NAME_CHECKED(FChaosClothAssetConnectableIStringValue, StringValue));
-	RegisterInputConnection(&TransferCollection);
+	if (FDataflowInput* const Input = RegisterInputConnection(&InputName.StringValue, GET_MEMBER_NAME_CHECKED(FChaosClothAssetConnectableIStringValue, StringValue)))
+	{
+		Input->SetCanHidePin(true);
+		Input->SetPinIsHidden(true);
+	}
+	if (FDataflowInput* const Input = RegisterInputConnection(&TransferCollection))
+	{
+		Input->SetCanHidePin(true);
+		Input->SetPinIsHidden(true);
+	}
 	RegisterOutputConnection(&Collection, &Collection);
 	RegisterOutputConnection(&Name);
 }
diff --git a/Engine/Plugins/ChaosClothAssetEditor/Source/ChaosClothAssetDataflowNodes/Private/ChaosClothAsset/SelectionNode.cpp b/Engine/Plugins/ChaosClothAssetEditor/Source/ChaosClothAssetDataflowNodes/Private/ChaosClothAsset/SelectionNode.cpp
--- a/Engine/Plugins/ChaosClothAssetEditor/Source/ChaosClothAssetDataflowNodes/Private/ChaosClothAsset/SelectionNode.cpp
+++ b/Engine/Plugins/ChaosClothAssetEditor/Source/ChaosClothAssetDataflowNodes/Private/ChaosClothAsset/SelectionNode.cpp
@@ -229,8 +229,16 @@ FChaosClothAssetSelectionNode::FChaosClothAssetSelectionNode(const Dataflow::FNo
 	: FDataflowTerminalNode(InParam, InGuid)
 {
 	RegisterInputConnection(&Collection);
-	RegisterInputConnection(&InputName.StringValue, GET_MEMBER_NAME_CHECKED(FChaosClothAssetConnectableIStringValue, StringValue));
-	RegisterInputConnection(&TransferCollection);
+	if (FDataflowInput* const Input = RegisterInputConnection(&InputName.StringValue, GET_MEMBER_NAME_CHECKED(FChaosClothAssetConnectableIStringValue, StringValue)))
+	{
+		Input->SetCanHidePin(true);
+		Input->SetPinIsHidden(true);
+	}
+	if (FDataflowInput* const Input = RegisterInputConnection(&TransferCollection))
+	{
+		Input->SetCanHidePin(true);
+		Input->SetPinIsHidden(true);
+	}
 	RegisterOutputConnection(&Collection, &Collection);
 	RegisterOutputConnection(&Name);
 }
diff --git a/Engine/Plugins/ChaosClothAssetEditor/Source/ChaosClothAssetDataflowNodes/Private/ChaosClothAsset/SimulationAerodynamicsConfigNode.cpp b/Engine/Plugins/ChaosClothAssetEditor/Source/ChaosClothAssetDataflowNodes/Private/ChaosClothAsset/SimulationAerodynamicsConfigNode.cpp
--- a/Engine/Plugins/ChaosClothAssetEditor/Source/ChaosClothAssetDataflowNodes/Private/ChaosClothAsset/SimulationAerodynamicsConfigNode.cpp
+++ b/Engine/Plugins/ChaosClothAssetEditor/Source/ChaosClothAssetDataflowNodes/Private/ChaosClothAsset/SimulationAerodynamicsConfigNode.cpp
@@ -10,8 +10,16 @@ FChaosClothAssetSimulationAerodynamicsConfigNode::FChaosClothAssetSimulationAero
 	: FChaosClothAssetSimulationBaseConfigNode(InParam, InGuid)
 {
 	RegisterCollectionConnections();
-	RegisterInputConnection(&Drag.WeightMap);
-	RegisterInputConnection(&Lift.WeightMap);
+	if (FDataflowInput* const Input = RegisterInputConnection(&Drag.WeightMap))
+	{
+		Input->SetCanHidePin(true);
+		Input->SetPinIsHidden(true);
+	}
+	if (FDataflowInput* const Input = RegisterInputConnection(&Lift.WeightMap))
+	{
+		Input->SetCanHidePin(true);
+		Input->SetPinIsHidden(true);
+	}
 }
 
 void FChaosClothAssetSimulationAerodynamicsConfigNode::AddProperties(FPropertyHelper& PropertyHelper) const
diff --git a/Engine/Plugins/ChaosClothAssetEditor/Source/ChaosClothAssetDataflowNodes/Private/ChaosClothAsset/SimulationAnimDriveConfigNode.cpp b/Engine/Plugins/ChaosClothAssetEditor/Source/ChaosClothAssetDataflowNodes/Private/ChaosClothAsset/SimulationAnimDriveConfigNode.cpp
--- a/Engine/Plugins/ChaosClothAssetEditor/Source/ChaosClothAssetDataflowNodes/Private/ChaosClothAsset/SimulationAnimDriveConfigNode.cpp
+++ b/Engine/Plugins/ChaosClothAssetEditor/Source/ChaosClothAssetDataflowNodes/Private/ChaosClothAsset/SimulationAnimDriveConfigNode.cpp
@@ -9,8 +9,16 @@ FChaosClothAssetSimulationAnimDriveConfigNode::FChaosClothAssetSimulationAnimDri
 	: FChaosClothAssetSimulationBaseConfigNode(InParam, InGuid)
 {
 	RegisterCollectionConnections();
-	RegisterInputConnection(&AnimDriveStiffness.WeightMap);
-	RegisterInputConnection(&AnimDriveDamping.WeightMap);
+	if (FDataflowInput* const Input = RegisterInputConnection(&AnimDriveStiffness.WeightMap))
+	{
+		Input->SetCanHidePin(true);
+		Input->SetPinIsHidden(true);
+	}
+	if (FDataflowInput* const Input = RegisterInputConnection(&AnimDriveDamping.WeightMap))
+	{
+		Input->SetCanHidePin(true);
+		Input->SetPinIsHidden(true);
+	}
 }
 
 void FChaosClothAssetSimulationAnimDriveConfigNode::AddProperties(FPropertyHelper& PropertyHelper) const
diff --git a/Engine/Plugins/ChaosClothAssetEditor/Source/ChaosClothAssetDataflowNodes/Private/ChaosClothAsset/SimulationBackstopConfigNode.cpp b/Engine/Plugins/ChaosClothAssetEditor/Source/ChaosClothAssetDataflowNodes/Private/ChaosClothAsset/SimulationBackstopConfigNode.cpp
--- a/Engine/Plugins/ChaosClothAssetEditor/Source/ChaosClothAssetDataflowNodes/Private/ChaosClothAsset/SimulationBackstopConfigNode.cpp
+++ b/Engine/Plugins/ChaosClothAssetEditor/Source/ChaosClothAssetDataflowNodes/Private/ChaosClothAsset/SimulationBackstopConfigNode.cpp
@@ -9,8 +9,16 @@ FChaosClothAssetSimulationBackstopConfigNode::FChaosClothAssetSimulationBackstop
 	: FChaosClothAssetSimulationBaseConfigNode(InParam, InGuid)
 {
 	RegisterCollectionConnections();
-	RegisterInputConnection(&BackstopDistance.WeightMap);
-	RegisterInputConnection(&BackstopRadius.WeightMap);
+	if (FDataflowInput* const Input = RegisterInputConnection(&BackstopDistance.WeightMap))
+	{
+		Input->SetCanHidePin(true);
+		Input->SetPinIsHidden(true);
+	}
+	if (FDataflowInput* const Input = RegisterInputConnection(&BackstopRadius.WeightMap))
+	{
+		Input->SetCanHidePin(true);
+		Input->SetPinIsHidden(true);
+	}
 }
 
 void FChaosClothAssetSimulationBackstopConfigNode::AddProperties(FPropertyHelper& PropertyHelper) const
diff --git a/Engine/Plugins/ChaosClothAssetEditor/Source/ChaosClothAssetDataflowNodes/Private/ChaosClothAsset/SimulationBendingConfigNode.cpp b/Engine/Plugins/ChaosClothAssetEditor/Source/ChaosClothAssetDataflowNodes/Private/ChaosClothAsset/SimulationBendingConfigNode.cpp
--- a/Engine/Plugins/ChaosClothAssetEditor/Source/ChaosClothAssetDataflowNodes/Private/ChaosClothAsset/SimulationBendingConfigNode.cpp
+++ b/Engine/Plugins/ChaosClothAssetEditor/Source/ChaosClothAssetDataflowNodes/Private/ChaosClothAsset/SimulationBendingConfigNode.cpp
@@ -9,20 +9,76 @@ FChaosClothAssetSimulationBendingConfigNode::FChaosClothAssetSimulationBendingCo
 	: FChaosClothAssetSimulationBaseConfigNode(InParam, InGuid)
 {
 	RegisterCollectionConnections();
-	RegisterInputConnection(&FlatnessRatio.WeightMap);
-	RegisterInputConnection(&RestAngle.WeightMap);
-	RegisterInputConnection(&BendingStiffness.WeightMap);
-	RegisterInputConnection(&BendingStiffnessWarp.WeightMap);
-	RegisterInputConnection(&BendingStiffnessWeft.WeightMap);
-	RegisterInputConnection(&BendingStiffnessBias.WeightMap);
-	RegisterInputConnection(&BendingDamping.WeightMap);
-	RegisterInputConnection(&BendingAnisoDamping.WeightMap);
-	RegisterInputConnection(&AnisoBucklingRatio.WeightMap);
-	RegisterInputConnection(&BucklingStiffness.WeightMap);
-	RegisterInputConnection(&BucklingStiffnessWarp.WeightMap);
-	RegisterInputConnection(&BucklingStiffnessWeft.WeightMap);
-	RegisterInputConnection(&BucklingStiffnessBias.WeightMap);
-	RegisterInputConnection(&BucklingRatioWeighted.WeightMap);
+	if (FDataflowInput* const Input = RegisterInputConnection(&FlatnessRatio.WeightMap))
+	{
+		Input->SetCanHidePin(true);
+		Input->SetPinIsHidden(true);
+	}
+	if (FDataflowInput* const Input = RegisterInputConnection(&RestAngle.WeightMap))
+	{
+		Input->SetCanHidePin(true);
+		Input->SetPinIsHidden(true);
+	}
+	if (FDataflowInput* const Input = RegisterInputConnection(&BendingStiffness.WeightMap))
+	{
+		Input->SetCanHidePin(true);
+		Input->SetPinIsHidden(true);
+	}
+	if (FDataflowInput* const Input = RegisterInputConnection(&BendingStiffnessWarp.WeightMap))
+	{
+		Input->SetCanHidePin(true);
+		Input->SetPinIsHidden(true);
+	}
+	if (FDataflowInput* const Input = RegisterInputConnection(&BendingStiffnessWeft.WeightMap))
+	{
+		Input->SetCanHidePin(true);
+		Input->SetPinIsHidden(true);
+	}
+	if (FDataflowInput* const Input = RegisterInputConnection(&BendingStiffnessBias.WeightMap))
+	{
+		Input->SetCanHidePin(true);
+		Input->SetPinIsHidden(true);
+	}
+	if (FDataflowInput* const Input = RegisterInputConnection(&BendingDamping.WeightMap))
+	{
+		Input->SetCanHidePin(true);
+		Input->SetPinIsHidden(true);
+	}
+	if (FDataflowInput* const Input = RegisterInputConnection(&BendingAnisoDamping.WeightMap))
+	{
+		Input->SetCanHidePin(true);
+		Input->SetPinIsHidden(true);
+	}
+	if (FDataflowInput* const Input = RegisterInputConnection(&AnisoBucklingRatio.WeightMap))
+	{
+		Input->SetCanHidePin(true);
+		Input->SetPinIsHidden(true);
+	}
+	if (FDataflowInput* const Input = RegisterInputConnection(&BucklingStiffness.WeightMap))
+	{
+		Input->SetCanHidePin(true);
+		Input->SetPinIsHidden(true);
+	}
+	if (FDataflowInput* const Input = RegisterInputConnection(&BucklingStiffnessWarp.WeightMap))
+	{
+		Input->SetCanHidePin(true);
+		Input->SetPinIsHidden(true);
+	}
+	if (FDataflowInput* const Input = RegisterInputConnection(&BucklingStiffnessWeft.WeightMap))
+	{
+		Input->SetCanHidePin(true);
+		Input->SetPinIsHidden(true);
+	}
+	if (FDataflowInput* const Input = RegisterInputConnection(&BucklingStiffnessBias.WeightMap))
+	{
+		Input->SetCanHidePin(true);
+		Input->SetPinIsHidden(true);
+	}
+	if (FDataflowInput* const Input = RegisterInputConnection(&BucklingRatioWeighted.WeightMap))
+	{
+		Input->SetCanHidePin(true);
+		Input->SetPinIsHidden(true);
+	}
 }
 
 void FChaosClothAssetSimulationBendingConfigNode::AddProperties(FPropertyHelper& PropertyHelper) const
diff --git a/Engine/Plugins/ChaosClothAssetEditor/Source/ChaosClothAssetDataflowNodes/Private/ChaosClothAsset/SimulationDampingConfigNode.cpp b/Engine/Plugins/ChaosClothAssetEditor/Source/ChaosClothAssetDataflowNodes/Private/ChaosClothAsset/SimulationDampingConfigNode.cpp
--- a/Engine/Plugins/ChaosClothAssetEditor/Source/ChaosClothAssetDataflowNodes/Private/ChaosClothAsset/SimulationDampingConfigNode.cpp
+++ b/Engine/Plugins/ChaosClothAssetEditor/Source/ChaosClothAssetDataflowNodes/Private/ChaosClothAsset/SimulationDampingConfigNode.cpp
@@ -9,7 +9,11 @@ FChaosClothAssetSimulationDampingConfigNode::FChaosClothAssetSimulationDampingCo
 	: FChaosClothAssetSimulationBaseConfigNode(InParam, InGuid)
 {
 	RegisterCollectionConnections();
-	RegisterInputConnection(&DampingCoefficientWeighted.WeightMap);
+	if (FDataflowInput* const Input = RegisterInputConnection(&DampingCoefficientWeighted.WeightMap))
+	{
+		Input->SetCanHidePin(true);
+		Input->SetPinIsHidden(true);
+	}
 }
 
 void FChaosClothAssetSimulationDampingConfigNode::AddProperties(FPropertyHelper& PropertyHelper) const
diff --git a/Engine/Plugins/ChaosClothAssetEditor/Source/ChaosClothAssetDataflowNodes/Private/ChaosClothAsset/SimulationGravityConfigNode.cpp b/Engine/Plugins/ChaosClothAssetEditor/Source/ChaosClothAssetDataflowNodes/Private/ChaosClothAsset/SimulationGravityConfigNode.cpp
--- a/Engine/Plugins/ChaosClothAssetEditor/Source/ChaosClothAssetDataflowNodes/Private/ChaosClothAsset/SimulationGravityConfigNode.cpp
+++ b/Engine/Plugins/ChaosClothAssetEditor/Source/ChaosClothAssetDataflowNodes/Private/ChaosClothAsset/SimulationGravityConfigNode.cpp
@@ -11,7 +11,11 @@ FChaosClothAssetSimulationGravityConfigNode::FChaosClothAssetSimulationGravityCo
 	: FChaosClothAssetSimulationBaseConfigNode(InParam, InGuid)
 {
 	RegisterCollectionConnections();
-	RegisterInputConnection(&GravityScaleWeighted.WeightMap);
+	if (FDataflowInput* const Input = RegisterInputConnection(&GravityScaleWeighted.WeightMap))
+	{
+		Input->SetCanHidePin(true);
+		Input->SetPinIsHidden(true);
+	}
 }
 
 void FChaosClothAssetSimulationGravityConfigNode::AddProperties(FPropertyHelper& PropertyHelper) const
diff --git a/Engine/Plugins/ChaosClothAssetEditor/Source/ChaosClothAssetDataflowNodes/Private/ChaosClothAsset/SimulationLongRangeAttachmentConfigNode.cpp b/Engine/Plugins/ChaosClothAssetEditor/Source/ChaosClothAssetDataflowNodes/Private/ChaosClothAsset/SimulationLongRangeAttachmentConfigNode.cpp
--- a/Engine/Plugins/ChaosClothAssetEditor/Source/ChaosClothAssetDataflowNodes/Private/ChaosClothAsset/SimulationLongRangeAttachmentConfigNode.cpp
+++ b/Engine/Plugins/ChaosClothAssetEditor/Source/ChaosClothAssetDataflowNodes/Private/ChaosClothAsset/SimulationLongRangeAttachmentConfigNode.cpp
@@ -12,8 +12,16 @@ FChaosClothAssetSimulationLongRangeAttachmentConfigNode::FChaosClothAssetSimulat
 {
 	RegisterCollectionConnections();
 	RegisterInputConnection(&FixedEndWeightMap);
-	RegisterInputConnection(&TetherStiffness.WeightMap);
-	RegisterInputConnection(&TetherScale.WeightMap);
+	if (FDataflowInput* const Input = RegisterInputConnection(&TetherStiffness.WeightMap))
+	{
+		Input->SetCanHidePin(true);
+		Input->SetPinIsHidden(true);
+	}
+	if (FDataflowInput* const Input = RegisterInputConnection(&TetherScale.WeightMap))
+	{
+		Input->SetCanHidePin(true);
+		Input->SetPinIsHidden(true);
+	}
 }
 
 void FChaosClothAssetSimulationLongRangeAttachmentConfigNode::AddProperties(FPropertyHelper& PropertyHelper) const
diff --git a/Engine/Plugins/ChaosClothAssetEditor/Source/ChaosClothAssetDataflowNodes/Private/ChaosClothAsset/SimulationMassConfigNode.cpp b/Engine/Plugins/ChaosClothAssetEditor/Source/ChaosClothAssetDataflowNodes/Private/ChaosClothAsset/SimulationMassConfigNode.cpp
--- a/Engine/Plugins/ChaosClothAssetEditor/Source/ChaosClothAssetDataflowNodes/Private/ChaosClothAsset/SimulationMassConfigNode.cpp
+++ b/Engine/Plugins/ChaosClothAssetEditor/Source/ChaosClothAssetDataflowNodes/Private/ChaosClothAsset/SimulationMassConfigNode.cpp
@@ -12,8 +12,16 @@ FChaosClothAssetSimulationMassConfigNode::FChaosClothAssetSimulationMassConfigNo
 	: FChaosClothAssetSimulationBaseConfigNode(InParam, InGuid)
 {
 	RegisterCollectionConnections();
-	RegisterInputConnection(&UniformMassWeighted.WeightMap);
-	RegisterInputConnection(&DensityWeighted.WeightMap);
+	if (FDataflowInput* const Input = RegisterInputConnection(&UniformMassWeighted.WeightMap))
+	{
+		Input->SetCanHidePin(true);
+		Input->SetPinIsHidden(true);
+	}
+	if (FDataflowInput* const Input = RegisterInputConnection(&DensityWeighted.WeightMap))
+	{
+		Input->SetCanHidePin(true);
+		Input->SetPinIsHidden(true);
+	}
 }
 
 void FChaosClothAssetSimulationMassConfigNode::AddProperties(FPropertyHelper& PropertyHelper) const
diff --git a/Engine/Plugins/ChaosClothAssetEditor/Source/ChaosClothAssetDataflowNodes/Private/ChaosClothAsset/SimulationMaxDistanceConfigNode.cpp b/Engine/Plugins/ChaosClothAssetEditor/Source/ChaosClothAssetDataflowNodes/Private/ChaosClothAsset/SimulationMaxDistanceConfigNode.cpp
--- a/Engine/Plugins/ChaosClothAssetEditor/Source/ChaosClothAssetDataflowNodes/Private/ChaosClothAsset/SimulationMaxDistanceConfigNode.cpp
+++ b/Engine/Plugins/ChaosClothAssetEditor/Source/ChaosClothAssetDataflowNodes/Private/ChaosClothAsset/SimulationMaxDistanceConfigNode.cpp
@@ -9,7 +9,11 @@ FChaosClothAssetSimulationMaxDistanceConfigNode::FChaosClothAssetSimulationMaxDi
 	: FChaosClothAssetSimulationBaseConfigNode(InParam, InGuid)
 {
 	RegisterCollectionConnections();
-	RegisterInputConnection(&MaxDistance.WeightMap);
+	if (FDataflowInput* const Input = RegisterInputConnection(&MaxDistance.WeightMap))
+	{
+		Input->SetCanHidePin(true);
+		Input->SetPinIsHidden(true);
+	}
 }
 
 void FChaosClothAssetSimulationMaxDistanceConfigNode::AddProperties(FPropertyHelper& PropertyHelper) const
diff --git a/Engine/Plugins/ChaosClothAssetEditor/Source/ChaosClothAssetDataflowNodes/Private/ChaosClothAsset/SimulationMultiResConfigNode.cpp b/Engine/Plugins/ChaosClothAssetEditor/Source/ChaosClothAssetDataflowNodes/Private/ChaosClothAsset/SimulationMultiResConfigNode.cpp
--- a/Engine/Plugins/ChaosClothAssetEditor/Source/ChaosClothAssetDataflowNodes/Private/ChaosClothAsset/SimulationMultiResConfigNode.cpp
+++ b/Engine/Plugins/ChaosClothAssetEditor/Source/ChaosClothAssetDataflowNodes/Private/ChaosClothAsset/SimulationMultiResConfigNode.cpp
@@ -9,8 +9,16 @@ FChaosClothAssetSimulationMultiResConfigNode::FChaosClothAssetSimulationMultiRes
 	: FChaosClothAssetSimulationBaseConfigNode(InParam, InGuid)
 {
 	RegisterCollectionConnections();
-	RegisterInputConnection(&MultiResStiffness.WeightMap);
-	RegisterInputConnection(&MultiResVelocityTargetStiffness.WeightMap);
+	if (FDataflowInput* const Input = RegisterInputConnection(&MultiResStiffness.WeightMap))
+	{
+		Input->SetCanHidePin(true);
+		Input->SetPinIsHidden(true);
+	}
+	if (FDataflowInput* const Input = RegisterInputConnection(&MultiResVelocityTargetStiffness.WeightMap))
+	{
+		Input->SetCanHidePin(true);
+		Input->SetPinIsHidden(true);
+	}
 }
 
 void FChaosClothAssetSimulationMultiResConfigNode::AddProperties(FPropertyHelper& PropertyHelper) const
diff --git a/Engine/Plugins/ChaosClothAssetEditor/Source/ChaosClothAssetDataflowNodes/Private/ChaosClothAsset/SimulationPressureConfigNode.cpp b/Engine/Plugins/ChaosClothAssetEditor/Source/ChaosClothAssetDataflowNodes/Private/ChaosClothAsset/SimulationPressureConfigNode.cpp
--- a/Engine/Plugins/ChaosClothAssetEditor/Source/ChaosClothAssetDataflowNodes/Private/ChaosClothAsset/SimulationPressureConfigNode.cpp
+++ b/Engine/Plugins/ChaosClothAssetEditor/Source/ChaosClothAssetDataflowNodes/Private/ChaosClothAsset/SimulationPressureConfigNode.cpp
@@ -9,7 +9,11 @@ FChaosClothAssetSimulationPressureConfigNode::FChaosClothAssetSimulationPressure
 	: FChaosClothAssetSimulationBaseConfigNode(InParam, InGuid)
 {
 	RegisterCollectionConnections();
-	RegisterInputConnection(&Pressure.WeightMap);
+	if (FDataflowInput* const Input = RegisterInputConnection(&Pressure.WeightMap))
+	{
+		Input->SetCanHidePin(true);
+		Input->SetPinIsHidden(true);
+	}
 }
 
 void FChaosClothAssetSimulationPressureConfigNode::AddProperties(FPropertyHelper& PropertyHelper) const
diff --git a/Engine/Plugins/ChaosClothAssetEditor/Source/ChaosClothAssetDataflowNodes/Private/ChaosClothAsset/SimulationSelfCollisionConfigNode.cpp b/Engine/Plugins/ChaosClothAssetEditor/Source/ChaosClothAssetDataflowNodes/Private/ChaosClothAsset/SimulationSelfCollisionConfigNode.cpp
--- a/Engine/Plugins/ChaosClothAssetEditor/Source/ChaosClothAssetDataflowNodes/Private/ChaosClothAsset/SimulationSelfCollisionConfigNode.cpp
+++ b/Engine/Plugins/ChaosClothAssetEditor/Source/ChaosClothAssetDataflowNodes/Private/ChaosClothAsset/SimulationSelfCollisionConfigNode.cpp
@@ -11,11 +11,31 @@ FChaosClothAssetSimulationSelfCollisionConfigNode::FChaosClothAssetSimulationSel
 	: FChaosClothAssetSimulationBaseConfigNode(InParam, InGuid)
 {
 	RegisterCollectionConnections();
-	RegisterInputConnection(&SelfCollisionLayers.StringValue, GET_MEMBER_NAME_CHECKED(FChaosClothAssetConnectableIStringValue, StringValue));
-	RegisterInputConnection(&SelfCollisionDisabledFaces.StringValue, GET_MEMBER_NAME_CHECKED(FChaosClothAssetConnectableIStringValue, StringValue));
-	RegisterInputConnection(&SelfCollisionEnabledKinematicFaces.StringValue, GET_MEMBER_NAME_CHECKED(FChaosClothAssetConnectableIStringValue, StringValue));
-	RegisterInputConnection(&SelfCollisionKinematicColliderFrictionWeighted.WeightMap);
-	RegisterInputConnection(&SelfCollisionThicknessWeighted.WeightMap);
+	if (FDataflowInput* const Input = RegisterInputConnection(&SelfCollisionLayers.StringValue, GET_MEMBER_NAME_CHECKED(FChaosClothAssetConnectableIStringValue, StringValue)))
+	{
+		Input->SetCanHidePin(true);
+		Input->SetPinIsHidden(true);
+	}
+	if (FDataflowInput* const Input = RegisterInputConnection(&SelfCollisionDisabledFaces.StringValue, GET_MEMBER_NAME_CHECKED(FChaosClothAssetConnectableIStringValue, StringValue)))
+	{
+		Input->SetCanHidePin(true);
+		Input->SetPinIsHidden(true);
+	}
+	if (FDataflowInput* const Input = RegisterInputConnection(&SelfCollisionEnabledKinematicFaces.StringValue, GET_MEMBER_NAME_CHECKED(FChaosClothAssetConnectableIStringValue, StringValue)))
+	{
+		Input->SetCanHidePin(true);
+		Input->SetPinIsHidden(true);
+	}
+	if (FDataflowInput* const Input = RegisterInputConnection(&SelfCollisionKinematicColliderFrictionWeighted.WeightMap))
+	{
+		Input->SetCanHidePin(true);
+		Input->SetPinIsHidden(true);
+	}
+	if (FDataflowInput* const Input = RegisterInputConnection(&SelfCollisionThicknessWeighted.WeightMap))
+	{
+		Input->SetCanHidePin(true);
+		Input->SetPinIsHidden(true);
+	}
 }
 
 void FChaosClothAssetSimulationSelfCollisionConfigNode::AddProperties(FPropertyHelper& PropertyHelper) const
diff --git a/Engine/Plugins/ChaosClothAssetEditor/Source/ChaosClothAssetDataflowNodes/Private/ChaosClothAsset/SimulationStretchConfigNode.cpp b/Engine/Plugins/ChaosClothAssetEditor/Source/ChaosClothAssetDataflowNodes/Private/ChaosClothAsset/SimulationStretchConfigNode.cpp
--- a/Engine/Plugins/ChaosClothAssetEditor/Source/ChaosClothAssetDataflowNodes/Private/ChaosClothAsset/SimulationStretchConfigNode.cpp
+++ b/Engine/Plugins/ChaosClothAssetEditor/Source/ChaosClothAssetDataflowNodes/Private/ChaosClothAsset/SimulationStretchConfigNode.cpp
@@ -9,15 +9,51 @@ FChaosClothAssetSimulationStretchConfigNode::FChaosClothAssetSimulationStretchCo
 	: FChaosClothAssetSimulationBaseConfigNode(InParam, InGuid)
 {
 	RegisterCollectionConnections();
-	RegisterInputConnection(&StretchStiffness.WeightMap);
-	RegisterInputConnection(&StretchStiffnessWarp.WeightMap);
-	RegisterInputConnection(&StretchStiffnessWeft.WeightMap);
-	RegisterInputConnection(&StretchStiffnessBias.WeightMap);
-	RegisterInputConnection(&StretchDamping.WeightMap);
-	RegisterInputConnection(&StretchAnisoDamping.WeightMap);
-	RegisterInputConnection(&StretchWarpScale.WeightMap);
-	RegisterInputConnection(&StretchWeftScale.WeightMap);
-	RegisterInputConnection(&AreaStiffness.WeightMap);
+	if (FDataflowInput* const Input = RegisterInputConnection(&StretchStiffness.WeightMap))
+	{
+		Input->SetCanHidePin(true);
+		Input->SetPinIsHidden(true);
+	}
+	if (FDataflowInput* const Input = RegisterInputConnection(&StretchStiffnessWarp.WeightMap))
+	{
+		Input->SetCanHidePin(true);
+		Input->SetPinIsHidden(true);
+	}
+	if (FDataflowInput* const Input = RegisterInputConnection(&StretchStiffnessWeft.WeightMap))
+	{
+		Input->SetCanHidePin(true);
+		Input->SetPinIsHidden(true);
+	}
+	if (FDataflowInput* const Input = RegisterInputConnection(&StretchStiffnessBias.WeightMap))
+	{
+		Input->SetCanHidePin(true);
+		Input->SetPinIsHidden(true);
+	}
+	if (FDataflowInput* const Input = RegisterInputConnection(&StretchDamping.WeightMap))
+	{
+		Input->SetCanHidePin(true);
+		Input->SetPinIsHidden(true);
+	}
+	if (FDataflowInput* const Input = RegisterInputConnection(&StretchAnisoDamping.WeightMap))
+	{
+		Input->SetCanHidePin(true);
+		Input->SetPinIsHidden(true);
+	}
+	if (FDataflowInput* const Input = RegisterInputConnection(&StretchWarpScale.WeightMap))
+	{
+		Input->SetCanHidePin(true);
+		Input->SetPinIsHidden(true);
+	}
+	if (FDataflowInput* const Input = RegisterInputConnection(&StretchWeftScale.WeightMap))
+	{
+		Input->SetCanHidePin(true);
+		Input->SetPinIsHidden(true);
+	}
+	if (FDataflowInput* const Input = RegisterInputConnection(&AreaStiffness.WeightMap))
+	{
+		Input->SetCanHidePin(true);
+		Input->SetPinIsHidden(true);
+	}
 }
 
 void FChaosClothAssetSimulationStretchConfigNode::AddProperties(FPropertyHelper& PropertyHelper) const
diff --git a/Engine/Plugins/ChaosClothAssetEditor/Source/ChaosClothAssetDataflowNodes/Private/ChaosClothAsset/TransferSkinWeightsNode.cpp b/Engine/Plugins/ChaosClothAssetEditor/Source/ChaosClothAssetDataflowNodes/Private/ChaosClothAsset/TransferSkinWeightsNode.cpp
--- a/Engine/Plugins/ChaosClothAssetEditor/Source/ChaosClothAssetDataflowNodes/Private/ChaosClothAsset/TransferSkinWeightsNode.cpp
+++ b/Engine/Plugins/ChaosClothAssetEditor/Source/ChaosClothAssetDataflowNodes/Private/ChaosClothAsset/TransferSkinWeightsNode.cpp
@@ -716,7 +716,11 @@ FChaosClothAssetTransferSkinWeightsNode::FChaosClothAssetTransferSkinWeightsNode
 {
 	RegisterInputConnection(&Collection);
 	RegisterOutputConnection(&Collection, &Collection);
-	RegisterInputConnection(&InpaintMask.WeightMap, GET_MEMBER_NAME_CHECKED(FChaosClothAssetWeightedValueNonAnimatableNoLowHighRange, WeightMap));
+	if (FDataflowInput* const Input = RegisterInputConnection(&InpaintMask.WeightMap, GET_MEMBER_NAME_CHECKED(FChaosClothAssetWeightedValueNonAnimatableNoLowHighRange, WeightMap)))
+	{
+		Input->SetCanHidePin(true);
+		Input->SetPinIsHidden(true);
+	}
 }
 
 void FChaosClothAssetTransferSkinWeightsNode::Evaluate(Dataflow::FContext& Context, const FDataflowOutput* Out) const
diff --git a/Engine/Plugins/Experimental/Dataflow/Source/DataflowEditor/Private/Dataflow/DataflowSNode.cpp b/Engine/Plugins/Experimental/Dataflow/Source/DataflowEditor/Private/Dataflow/DataflowSNode.cpp
--- a/Engine/Plugins/Experimental/Dataflow/Source/DataflowEditor/Private/Dataflow/DataflowSNode.cpp
+++ b/Engine/Plugins/Experimental/Dataflow/Source/DataflowEditor/Private/Dataflow/DataflowSNode.cpp
@@ -8,6 +8,7 @@
 #include "Dataflow/DataflowGraphEditor.h"
 #include "Dataflow/DataflowNodeFactory.h"
 #include "Dataflow/DataflowObject.h"
+#include "Framework/Application/SlateApplication.h"
 #include "Logging/LogMacros.h"
 #include "SourceCodeNavigation.h"
 #include "Styling/SlateTypes.h"
@@ -15,7 +16,9 @@
 #include "Styling/AppStyle.h"
 #include "Widgets/Layout/SBorder.h"
 #include "Widgets/Input/SCheckBox.h"
+#include "Widgets/SBoxPanel.h"
 #include "Editor/Transactor.h"
+#include "GraphEditorSettings.h"
 
 #include UE_INLINE_GENERATED_CPP_BY_NAME(DataflowSNode)
 
@@ -185,6 +188,91 @@ void SDataflowEdNode::AddReferencedObjects(FReferenceCollector& Collector)
 	}
 }
 
+void SDataflowEdNode::CreateInputSideAddButton(TSharedPtr<SVerticalBox> InputBox)
+{
+	TSharedRef<SWidget> AddPinButton = AddPinButtonContent(
+		LOCTEXT("AddPinInputButton", "Show/Hide Inputs"),
+		LOCTEXT("AddPinInputButton_Tooltip", "Show/Hide input pins."),
+		false
+	);
+
+	FMargin AddPinPadding = Settings->GetOutputPinPadding();
+	AddPinPadding.Top += 6.0f;
+
+	InputBox->AddSlot()
+		.AutoHeight()
+		.VAlign(VAlign_Center)
+		.Padding(AddPinPadding)
+		[
+			AddPinButton
+		];
+}
+
+FReply SDataflowEdNode::OnAddPin()
+{
+	if (DataflowGraphNode)
+	{
+		FMenuBuilder MenuBuilder(false, nullptr);
+		if (TSharedPtr<FDataflowNode> DataflowNode = DataflowGraphNode->GetDataflowNode())
+		{
+			if (DataflowNode->HasHideableInputs())
+			{
+				MenuBuilder.AddMenuEntry(LOCTEXT("HideAllInputs", "Hide all"), LOCTEXT("HideAllInputsTooltip", "Hide all hideable input pins"), FSlateIcon(),
+					FUIAction(
+						FExecuteAction::CreateUObject(DataflowGraphNode, &UDataflowEdNode::HideAllInputPins)));
+				MenuBuilder.AddMenuEntry(LOCTEXT("UnhideAllInputs", "Unhide all"), LOCTEXT("UnhideAllInputsTooltip", "Unhide all hideable input pins"), FSlateIcon(),
+					FUIAction(
+						FExecuteAction::CreateUObject(DataflowGraphNode, &UDataflowEdNode::ShowAllInputPins)));
+
+
+				TArray<FDataflowInput*> Inputs = DataflowNode->GetInputs();
+				for (FDataflowInput* Input : Inputs)
+				{
+					if (Input->GetCanHidePin())
+					{
+						MenuBuilder.AddMenuEntry(FText::FromName(Input->GetName()), LOCTEXT("UnhidePinTooltip", "Show/Hide pin"), FSlateIcon(),
+							FUIAction(
+								FExecuteAction::CreateUObject(DataflowGraphNode, &UDataflowEdNode::ToggleHideInputPin, Input->GetName()),
+								FCanExecuteAction::CreateUObject(DataflowGraphNode, &UDataflowEdNode::CanToggleHideInputPin, Input->GetName()),
+								FIsActionChecked::CreateUObject(DataflowGraphNode, &UDataflowEdNode::IsInputPinShown, Input->GetName())),
+							NAME_None, EUserInterfaceActionType::ToggleButton);
+					}
+				}
+			}
+		}
+		FSlateApplication::Get().PushMenu(AsShared(),
+			FWidgetPath(),
+			MenuBuilder.MakeWidget(),
+			FSlateApplication::Get().GetCursorPos(),
+			FPopupTransitionEffect(FPopupTransitionEffect::ContextMenu)
+		);
+		return FReply::Handled();
+	}
+	return FReply::Unhandled();
+}
+
+EVisibility SDataflowEdNode::IsAddPinButtonVisible() const
+{
+	EVisibility Visibility = Super::IsAddPinButtonVisible();
+	if (Visibility == EVisibility::Collapsed)
+	{
+		return Visibility;
+	}
+
+	if (DataflowGraphNode)
+	{
+		if (const TSharedPtr<FDataflowNode> DataflowNode = DataflowGraphNode->GetDataflowNode())
+		{
+			if (DataflowNode->HasHideableInputs())
+			{
+				return Visibility;
+			}
+		}
+	}
+
+	return EVisibility::Collapsed;
+}
+
 //
 // Add a menu option to create a graph node.
 //
diff --git a/Engine/Plugins/Experimental/Dataflow/Source/DataflowEditor/Public/Dataflow/DataflowSNode.h b/Engine/Plugins/Experimental/Dataflow/Source/DataflowEditor/Public/Dataflow/DataflowSNode.h
--- a/Engine/Plugins/Experimental/Dataflow/Source/DataflowEditor/Public/Dataflow/DataflowSNode.h
+++ b/Engine/Plugins/Experimental/Dataflow/Source/DataflowEditor/Public/Dataflow/DataflowSNode.h
@@ -47,7 +47,18 @@ public:
 	}
 	//~ End FGCObject interface
 
+
 private:
+	// SGraphNode interface
+	/** Override this to create a button to add pins on the input side of the node */
+	virtual void CreateInputSideAddButton(TSharedPtr<SVerticalBox> InputBox) override;
+	/** Callback function executed when Add pin button is clicked */
+	virtual FReply OnAddPin() override;
+	/** Checks whether Add pin button should currently be visible */
+	virtual EVisibility IsAddPinButtonVisible() const;
+
+	
+
 	TObjectPtr<UDataflowEdNode> DataflowGraphNode = nullptr;	
 
 	FCheckBoxStyle CheckBoxStyle;
diff --git a/Engine/Source/Runtime/Core/Public/UObject/FortniteMainBranchObjectVersion.h b/Engine/Source/Runtime/Core/Public/UObject/FortniteMainBranchObjectVersion.h
--- a/Engine/Source/Runtime/Core/Public/UObject/FortniteMainBranchObjectVersion.h
+++ b/Engine/Source/Runtime/Core/Public/UObject/FortniteMainBranchObjectVersion.h
@@ -474,6 +474,9 @@ struct FFortniteMainBranchObjectVersion
 		// Added a more stable pixel depth offset mode.
 		MaterialPixelDepthOffsetMode,
 
+		// Added hideable pins to dataflow
+		DataflowHideablePins,
+
 		// -----<new versions can be added above this line>-------------------------------------------------
 		VersionPlusOne,
 		LatestVersion = VersionPlusOne - 1
diff --git a/Engine/Source/Runtime/Experimental/Dataflow/Core/Private/Dataflow/DataflowGraph.cpp b/Engine/Source/Runtime/Experimental/Dataflow/Core/Private/Dataflow/DataflowGraph.cpp
--- a/Engine/Source/Runtime/Experimental/Dataflow/Core/Private/Dataflow/DataflowGraph.cpp
+++ b/Engine/Source/Runtime/Experimental/Dataflow/Core/Private/Dataflow/DataflowGraph.cpp
@@ -176,6 +176,8 @@ namespace Dataflow
 
 					bool bIsAnytype = Output->IsAnyType();
 					Ar << bIsAnytype;
+					bool bIsHidden = Output->GetPinIsHidden();
+					Ar << bIsHidden;
 				}
 
 				int32 ArNumInputs = Node->GetInputs().Num();
@@ -189,6 +191,8 @@ namespace Dataflow
 
 					bool bIsAnytype = Input->IsAnyType();
 					Ar << bIsAnytype;
+					bool bIsHidden = Input->GetPinIsHidden();
+					Ar << bIsHidden;
 				}
 			}
 			DATAFLOW_OPTIONAL_BLOCK_WRITE_END();
@@ -240,6 +244,8 @@ namespace Dataflow
 				ensure(!NodeGuidMap.Contains(ArGuid));
 				NodeGuidMap.Add(ArGuid, Node);
 
+				const bool bDataflowHideablePinSupport = (Ar.CustomVer(FFortniteMainBranchObjectVersion::GUID) >= FFortniteMainBranchObjectVersion::DataflowHideablePins);
+
 				if (!bDataflowSeparateInputOutputSerialization)
 				{
 
@@ -302,6 +308,7 @@ namespace Dataflow
 					Node->SerializeInternal(Ar);
 
 					bool bIsAnyType = false;
+					bool bIsHidden = true;
 					// Outputs deserialization
 					{
 						int32 ArNumOutputs;
@@ -314,6 +321,10 @@ namespace Dataflow
 							{
 								Ar << bIsAnyType;
 							}
+							if (bDataflowHideablePinSupport)
+							{
+								Ar << bIsHidden;
+							}
 
 							if (FDataflowOutput* Output = Node->FindOutput(ArName))
 							{
@@ -325,6 +336,7 @@ namespace Dataflow
 								{
 									check(Output->GetType() == ArType || bIsAnyType);
 								}
+								Output->SetPinIsHidden(bIsHidden);
 								Output->SetGuid(ArGuid);
 								ensure(!ConnectionGuidMap.Contains(ArGuid));
 								ConnectionGuidMap.Add(ArGuid, Output);
@@ -351,6 +363,10 @@ namespace Dataflow
 							{
 								Ar << bIsAnyType;
 							}
+							if (bDataflowHideablePinSupport)
+							{
+								Ar << bIsHidden;
+							}
 
 							if (FDataflowInput* Input = Node->FindInput(ArName))
 							{
@@ -362,6 +378,7 @@ namespace Dataflow
 								{
 									check(Input->GetType() == ArType || bIsAnyType);
 								}
+								Input->SetPinIsHidden(bIsHidden);
 								Input->SetGuid(ArGuid);
 								ensure(!ConnectionGuidMap.Contains(ArGuid));
 								ConnectionGuidMap.Add(ArGuid, Input);
diff --git a/Engine/Source/Runtime/Experimental/Dataflow/Core/Private/Dataflow/DataflowNode.cpp b/Engine/Source/Runtime/Experimental/Dataflow/Core/Private/Dataflow/DataflowNode.cpp
--- a/Engine/Source/Runtime/Experimental/Dataflow/Core/Private/Dataflow/DataflowNode.cpp
+++ b/Engine/Source/Runtime/Experimental/Dataflow/Core/Private/Dataflow/DataflowNode.cpp
@@ -30,9 +30,9 @@ void FDataflowNode::AddInput(FDataflowInput* InPtr)
 {
 	if (InPtr)
 	{
-		for (TPair<uint32, FDataflowInput*> Elem : Inputs)
+		for (const TPair<int32, FDataflowInput*>& Elem : Inputs)
 		{
-			FDataflowInput* In = Elem.Value;
+			const FDataflowInput* const In = Elem.Value;
 			ensureMsgf(!In->GetName().IsEqual(InPtr->GetName()), TEXT("Add Input Failed: Existing Node input already defined with name (%s)"), *InPtr->GetName().ToString());
 		}
 
@@ -50,9 +50,9 @@ void FDataflowNode::AddInput(FDataflowInput* InPtr)
 
 FDataflowInput* FDataflowNode::FindInput(FName InName)
 {
-	for (TPair<uint32, FDataflowInput*> Elem : Inputs)
+	for (TPair<int32, FDataflowInput*>& Elem : Inputs)
 	{
-		FDataflowInput* Con = Elem.Value;
+		FDataflowInput* const Con = Elem.Value;
 		if (Con->GetName().IsEqual(InName))
 		{
 			return Con;
@@ -62,11 +62,24 @@ FDataflowInput* FDataflowNode::FindInput(FName InName)
 }
 
 
+const FDataflowInput* FDataflowNode::FindInput(FName InName) const
+{
+	for (const TPair<int32, FDataflowInput*>& Elem : Inputs)
+	{
+		const FDataflowInput* const Con = Elem.Value;
+		if (Con->GetName().IsEqual(InName))
+		{
+			return Con;
+		}
+	}
+	return nullptr;
+}
+
 const FDataflowInput* FDataflowNode::FindInput(const void* Reference) const
 {
-	for (TPair<uint32, FDataflowInput*> Elem : Inputs)
+	for (const TPair<int32, FDataflowInput*>& Elem : Inputs)
 	{
-		FDataflowInput* Con = Elem.Value;
+		const FDataflowInput* const Con = Elem.Value;
 		if (Con->RealAddress() == Reference)
 		{
 			return Con;
@@ -77,9 +90,9 @@ const FDataflowInput* FDataflowNode::FindInput(const void* Reference) const
 
 FDataflowInput* FDataflowNode::FindInput(void* Reference)
 {
-	for (TPair<uint32, FDataflowInput*> Elem : Inputs)
+	for (TPair<int32, FDataflowInput*>& Elem : Inputs)
 	{
-		FDataflowInput* Con = Elem.Value;
+		FDataflowInput* const Con = Elem.Value;
 		if (Con->RealAddress() == Reference)
 		{
 			return Con;
@@ -90,9 +103,9 @@ FDataflowInput* FDataflowNode::FindInput(void* Reference)
 
 const FDataflowInput* FDataflowNode::FindInput(const FGuid& InGuid) const
 {
-	for (TPair<uint32, FDataflowInput*> Elem : Inputs)
+	for (const TPair<int32, FDataflowInput*>& Elem : Inputs)
 	{
-		FDataflowInput* Con = Elem.Value;
+		const FDataflowInput* const Con = Elem.Value;
 		if (Con->GetGuid() == InGuid)
 		{
 			return Con;
@@ -110,14 +123,39 @@ TArray< FDataflowInput* > FDataflowNode::GetInputs() const
 
 void FDataflowNode::ClearInputs()
 {
-	for (TPair<uint32, FDataflowInput*> Elem : Inputs)
+	for (TPair<int32, FDataflowInput*>& Elem : Inputs)
 	{
-		FDataflowInput* Con = Elem.Value;
+		FDataflowInput* const Con = Elem.Value;
 		delete Con;
 	}
 	Inputs.Reset();
 }
 
+bool FDataflowNode::HasHideableInputs() const
+{
+	for (const TPair<int32, FDataflowInput*>& Elem : Inputs)
+	{
+		const FDataflowInput* const Con = Elem.Value;
+		if (Con->GetCanHidePin())
+		{
+			return true;
+		}
+	}
+	return false;
+}
+
+bool FDataflowNode::HasHiddenInputs() const
+{
+	for (const TPair<int32, FDataflowInput*>& Elem : Inputs)
+	{
+		const FDataflowInput* const Con = Elem.Value;
+		if (Con->GetPinIsHidden())
+		{
+			return true;
+		}
+	}
+	return false;
+}
 
 //
 // Outputs
@@ -128,9 +166,9 @@ void FDataflowNode::AddOutput(FDataflowOutput* InPtr)
 {
 	if (InPtr)
 	{
-		for (TPair<uint32, FDataflowOutput*> Elem : Outputs)
+		for (const TPair<int32, FDataflowOutput*>& Elem : Outputs)
 		{
-			FDataflowOutput* Out = Elem.Value;
+			const FDataflowOutput* const Out = Elem.Value;
 			ensureMsgf(!Out->GetName().IsEqual(InPtr->GetName()), TEXT("Add Output Failed: Existing Node output already defined with name (%s)"), *InPtr->GetName().ToString());
 		}
 
@@ -148,12 +186,12 @@ void FDataflowNode::AddOutput(FDataflowOutput* InPtr)
 
 FDataflowOutput* FDataflowNode::FindOutput(uint32 InGuidHash)
 {
-	for (TPair<uint32, FDataflowOutput*> Elem : Outputs)
+	for (TPair<int32, FDataflowOutput*>& Elem : Outputs)
 	{
-		FDataflowOutput* Con = Elem.Value;
+		FDataflowOutput* const Con = Elem.Value;
 		if (GetTypeHash(Con->GetGuid()) == InGuidHash)
 		{
-			return (FDataflowOutput*)Con;
+			return Con;
 		}
 	}
 	return nullptr;
@@ -162,12 +200,12 @@ FDataflowOutput* FDataflowNode::FindOutput(uint32 InGuidHash)
 
 FDataflowOutput* FDataflowNode::FindOutput(FName InName)
 {
-	for (TPair<uint32, FDataflowOutput*> Elem : Outputs)
+	for (TPair<int32, FDataflowOutput*>& Elem : Outputs)
 	{
-		FDataflowOutput* Con = Elem.Value;
+		FDataflowOutput* const Con = Elem.Value;
 		if (Con->GetName().IsEqual(InName))
 		{
-			return (FDataflowOutput*)Con;
+			return Con;
 		}
 	}
 	return nullptr;
@@ -175,12 +213,12 @@ FDataflowOutput* FDataflowNode::FindOutput(FName InName)
 
 const FDataflowOutput* FDataflowNode::FindOutput(FName InName) const
 {
-	for (TPair<uint32, FDataflowOutput*> Elem : Outputs)
+	for (const TPair<int32, FDataflowOutput*>& Elem : Outputs)
 	{
-		FDataflowOutput* Con = Elem.Value;
+		const FDataflowOutput* const Con = Elem.Value;
 		if (Con->GetName().IsEqual(InName))
 		{
-			return (FDataflowOutput*)Con;
+			return Con;
 		}
 	}
 	return nullptr;
@@ -188,12 +226,12 @@ const FDataflowOutput* FDataflowNode::FindOutput(FName InName) const
 
 const FDataflowOutput* FDataflowNode::FindOutput(uint32 InGuidHash) const
 {
-	for (TPair<uint32, FDataflowOutput*> Elem : Outputs)
+	for (const TPair<int32, FDataflowOutput*>& Elem : Outputs)
 	{
-		FDataflowOutput* Con = Elem.Value;
+		const FDataflowOutput* const Con = Elem.Value;
 		if (GetTypeHash(Con->GetGuid()) == InGuidHash)
 		{
-			return (FDataflowOutput*)Con;
+			return Con;
 		}
 	}
 	return nullptr;
@@ -201,12 +239,12 @@ const FDataflowOutput* FDataflowNode::FindOutput(uint32 InGuidHash) const
 
 const FDataflowOutput* FDataflowNode::FindOutput(const void* Reference) const
 {
-	for (TPair<uint32, FDataflowOutput*> Elem : Outputs)
+	for (const TPair<int32, FDataflowOutput*>& Elem : Outputs)
 	{
-		FDataflowOutput* Con = Elem.Value;
+		const FDataflowOutput* const Con = Elem.Value;
 		if (Con->RealAddress() == Reference)
 		{
-			return (FDataflowOutput*)Con;
+			return Con;
 		}
 	}
 	return nullptr;
@@ -214,12 +252,12 @@ const FDataflowOutput* FDataflowNode::FindOutput(const void* Reference) const
 
 FDataflowOutput* FDataflowNode::FindOutput(void* Reference)
 {
-	for (TPair<uint32, FDataflowOutput*> Elem : Outputs)
+	for (TPair<int32, FDataflowOutput*>& Elem : Outputs)
 	{
-		FDataflowOutput* Con = Elem.Value;
+		FDataflowOutput* const Con = Elem.Value;
 		if (Con->RealAddress() == Reference)
 		{
-			return (FDataflowOutput*)Con;
+			return Con;
 		}
 	}
 	return nullptr;
@@ -227,9 +265,9 @@ FDataflowOutput* FDataflowNode::FindOutput(void* Reference)
 
 const FDataflowOutput* FDataflowNode::FindOutput(const FGuid& InGuid) const
 {
-	for (TPair<uint32, FDataflowOutput*> Elem : Outputs)
+	for (const TPair<int32, FDataflowOutput*>& Elem : Outputs)
 	{
-		FDataflowOutput* Con = Elem.Value;
+		const FDataflowOutput* const Con = Elem.Value;
 		if (Con->GetGuid() == InGuid)
 		{
 			return Con;
@@ -254,27 +292,52 @@ TArray< FDataflowOutput* > FDataflowNode::GetOutputs() const
 
 void FDataflowNode::ClearOutputs()
 {
-	for (TPair<uint32, FDataflowOutput*> Elem : Outputs)
+	for (TPair<int32, FDataflowOutput*>& Elem : Outputs)
 	{
-		FDataflowOutput* Con = Elem.Value;
+		FDataflowOutput* const Con = Elem.Value;
 		delete Con;
 	}
 	Outputs.Reset();
 }
 
+bool FDataflowNode::HasHideableOutputs() const
+{
+	for (const TPair<int32, FDataflowOutput*>& Elem : Outputs)
+	{
+		const FDataflowOutput* const Con = Elem.Value;
+		if (Con->GetCanHidePin())
+		{
+			return true;
+		}
+	}
+	return false;
+}
+
+bool FDataflowNode::HasHiddenOutputs() const
+{
+	for (const TPair<int32, FDataflowOutput*>& Elem : Outputs)
+	{
+		const FDataflowOutput* const Con = Elem.Value;
+		if (Con->GetPinIsHidden())
+		{
+			return true;
+		}
+	}
+	return false;
+}
 
 TArray<Dataflow::FPin> FDataflowNode::GetPins() const
 {
 	TArray<Dataflow::FPin> RetVal;
-	for (TPair<uint32, FDataflowInput*> Elem : Inputs)
+	for (const TPair<int32, FDataflowInput*>& Elem : Inputs)
 	{
-		FDataflowInput* Con = Elem.Value;
-		RetVal.Add({ Dataflow::FPin::EDirection::INPUT,Con->GetType(), Con->GetName() });
+		const FDataflowInput* const Con = Elem.Value;
+		RetVal.Add({ Dataflow::FPin::EDirection::INPUT,Con->GetType(), Con->GetName(), Con->GetPinIsHidden()});
 	}
-	for (TPair<uint32, FDataflowOutput*> Elem : Outputs)
+	for (const TPair<int32, FDataflowOutput*>& Elem : Outputs)
 	{
-		FDataflowOutput* Con = Elem.Value;
-		RetVal.Add({ Dataflow::FPin::EDirection::OUTPUT,Con->GetType(), Con->GetName() });
+		const FDataflowOutput* const Con = Elem.Value;
+		RetVal.Add({ Dataflow::FPin::EDirection::OUTPUT,Con->GetType(), Con->GetName(), Con->GetPinIsHidden() });
 	}
 	return RetVal;
 }
@@ -341,9 +404,9 @@ void FDataflowNode::Invalidate(const Dataflow::FTimestamp& InModifiedTimestamp)
 	{
 		LastModifiedTimestamp = InModifiedTimestamp;
 
-		for (TPair<uint32, FDataflowOutput*> Elem : Outputs)
+		for (TPair<int32, FDataflowOutput*>& Elem : Outputs)
 		{
-			FDataflowOutput* Con = Elem.Value;
+			FDataflowOutput* const Con = Elem.Value;
 			Con->Invalidate(InModifiedTimestamp);
 		}
 
@@ -453,7 +516,7 @@ FText FDataflowNode::GetPropertyDisplayNameText(const TArray<const FProperty*>&
 #endif
 }
 
-void FDataflowNode::RegisterInputConnection(const void* InProperty, const FName& PropertyName)
+FDataflowInput* FDataflowNode::RegisterInputConnection(const void* InProperty, const FName& PropertyName)
 {
 	if (TUniquePtr<FStructOnScope> ScriptOnStruct = TUniquePtr<FStructOnScope>(NewStructOnScope()))
 	{
@@ -467,10 +530,13 @@ void FDataflowNode::RegisterInputConnection(const void* InProperty, const FName&
 				FString ExtendedType;
 				const FString CPPType = Property->GetCPPType(&ExtendedType);
 				const FName PropType(CPPType + ExtendedType);
-				AddInput(new FDataflowInput({ PropType, PropName, this, Property }));
+				FDataflowInput* const Input = new FDataflowInput({ PropType, PropName, this, Property });
+				AddInput(Input);
+				return Input;
 			}
 		}
 	}
+	return nullptr;
 }
 
 void FDataflowNode::UnregisterInputConnection(const void* InProperty, const FName& PropertyName)
@@ -498,7 +564,7 @@ void FDataflowNode::UnregisterInputConnection(const void* InProperty, const FNam
 	}
 }
 
-void FDataflowNode::RegisterOutputConnection(const void* InProperty, const void* Passthrough, const FName& PropertyName, const FName& PassthroughName)
+FDataflowOutput* FDataflowNode::RegisterOutputConnection(const void* InProperty, const void* Passthrough, const FName& PropertyName, const FName& PassthroughName)
 {
 	if (TUniquePtr<FStructOnScope> ScriptOnStruct = TUniquePtr<FStructOnScope>(NewStructOnScope()))
 	{
@@ -524,8 +590,10 @@ void FDataflowNode::RegisterOutputConnection(const void* InProperty, const void*
 				}
 				AddOutput(OutputConnection);
 			}
+			return OutputConnection;
 		}
 	}
+	return nullptr;
 }
 
 
diff --git a/Engine/Source/Runtime/Experimental/Dataflow/Core/Public/Dataflow/DataflowConnection.h b/Engine/Source/Runtime/Experimental/Dataflow/Core/Public/Dataflow/DataflowConnection.h
--- a/Engine/Source/Runtime/Experimental/Dataflow/Core/Public/Dataflow/DataflowConnection.h
+++ b/Engine/Source/Runtime/Experimental/Dataflow/Core/Public/Dataflow/DataflowConnection.h
@@ -32,6 +32,7 @@ namespace Dataflow
 		EDirection Direction;
 		FName Type;
 		FName Name;
+		bool bHidden = false;
 	};
 
 	class FGraph;
@@ -53,6 +54,8 @@ protected:
 	const FProperty* Property = nullptr;
 	FGuid  Guid;
 	bool bIsAnyType = false;
+	bool bCanHidePin = false;
+	bool bPinIsHidden = false;
 
 	friend struct FDataflowNode;
 	friend class Dataflow::FGraph;
@@ -109,4 +112,9 @@ public:
 	}
 
 	virtual void Invalidate(const Dataflow::FTimestamp& ModifiedTimestamp = Dataflow::FTimestamp::Current()) {};
+
+	bool GetCanHidePin() const { return bCanHidePin; }
+	bool GetPinIsHidden() const { return bCanHidePin && bPinIsHidden; }
+	void SetCanHidePin(bool bInCanHidePin) { bCanHidePin = bInCanHidePin; }
+	void SetPinIsHidden(bool bInPinIsHidden) { bPinIsHidden = bInPinIsHidden; }
 };
diff --git a/Engine/Source/Runtime/Experimental/Dataflow/Core/Public/Dataflow/DataflowInputOutput.h b/Engine/Source/Runtime/Experimental/Dataflow/Core/Public/Dataflow/DataflowInputOutput.h
--- a/Engine/Source/Runtime/Experimental/Dataflow/Core/Public/Dataflow/DataflowInputOutput.h
+++ b/Engine/Source/Runtime/Experimental/Dataflow/Core/Public/Dataflow/DataflowInputOutput.h
@@ -55,6 +55,7 @@ public:
 
 	FDataflowOutput* GetConnection() { return Connection; }
 	const FDataflowOutput* GetConnection() const { return Connection; }
+	bool HasAnyConnections() const { return Connection != nullptr; }
 
 	virtual TArray< FDataflowOutput* > GetConnectedOutputs();
 	virtual const TArray< const FDataflowOutput* > GetConnectedOutputs() const;
@@ -121,6 +122,7 @@ public:
 
 	DATAFLOWCORE_API TArray<FDataflowInput*>& GetConnections();
 	DATAFLOWCORE_API const TArray<FDataflowInput*>& GetConnections() const;
+	bool HasAnyConnections() const { return !Connections.IsEmpty(); }
 
 	DATAFLOWCORE_API virtual TArray<FDataflowInput*> GetConnectedInputs();
 	DATAFLOWCORE_API virtual const TArray<const FDataflowInput*> GetConnectedInputs() const;
diff --git a/Engine/Source/Runtime/Experimental/Dataflow/Core/Public/Dataflow/DataflowNode.h b/Engine/Source/Runtime/Experimental/Dataflow/Core/Public/Dataflow/DataflowNode.h
--- a/Engine/Source/Runtime/Experimental/Dataflow/Core/Public/Dataflow/DataflowNode.h
+++ b/Engine/Source/Runtime/Experimental/Dataflow/Core/Public/Dataflow/DataflowNode.h
@@ -115,9 +115,12 @@ struct FDataflowNode
 	DATAFLOWCORE_API virtual void AddInput(FDataflowInput* InPtr);
 	DATAFLOWCORE_API TArray< FDataflowInput* > GetInputs() const;
 	DATAFLOWCORE_API void ClearInputs();
+	DATAFLOWCORE_API bool HasHideableInputs() const;
+	DATAFLOWCORE_API bool HasHiddenInputs() const;
 
 	DATAFLOWCORE_API FDataflowInput* FindInput(FName Name);
 	DATAFLOWCORE_API FDataflowInput* FindInput(void* Reference);
+	DATAFLOWCORE_API const FDataflowInput* FindInput(FName Name) const;
 	DATAFLOWCORE_API const FDataflowInput* FindInput(const void* Reference) const;
 	DATAFLOWCORE_API const FDataflowInput* FindInput(const FGuid& InGuid) const;
 
@@ -125,6 +128,8 @@ struct FDataflowNode
 	DATAFLOWCORE_API int NumOutputs() const;
 	DATAFLOWCORE_API TArray< FDataflowOutput* > GetOutputs() const;
 	DATAFLOWCORE_API void ClearOutputs();
+	DATAFLOWCORE_API bool HasHideableOutputs() const;
+	DATAFLOWCORE_API bool HasHiddenOutputs() const;
 
 	DATAFLOWCORE_API FDataflowOutput* FindOutput(FName Name);
 	DATAFLOWCORE_API FDataflowOutput* FindOutput(uint32 GuidHash);
@@ -162,8 +167,8 @@ struct FDataflowNode
 	virtual const UScriptStruct* TypedScriptStruct() const { return nullptr; }
 
 	/** Register the Input and Outputs after the creation in the factory. Use PropertyName to disambiguate a struct name from its first property. */
-	DATAFLOWCORE_API void RegisterInputConnection(const void* Property, const FName& PropertyName = NAME_None);
-	DATAFLOWCORE_API void RegisterOutputConnection(
+	DATAFLOWCORE_API FDataflowInput* RegisterInputConnection(const void* Property, const FName& PropertyName = NAME_None);
+	DATAFLOWCORE_API FDataflowOutput* RegisterOutputConnection(
 		const void* Property,
 		const void* Passthrough = nullptr,
 		const FName& PropertyName = NAME_None,
diff --git a/Engine/Source/Runtime/Experimental/Dataflow/Engine/Private/Dataflow/DataflowEdNode.cpp b/Engine/Source/Runtime/Experimental/Dataflow/Engine/Private/Dataflow/DataflowEdNode.cpp
--- a/Engine/Source/Runtime/Experimental/Dataflow/Engine/Private/Dataflow/DataflowEdNode.cpp
+++ b/Engine/Source/Runtime/Experimental/Dataflow/Engine/Private/Dataflow/DataflowEdNode.cpp
@@ -17,6 +17,8 @@
 #if WITH_EDITOR
 #include "EdGraph/EdGraphPin.h"
 #include "EdGraph/EdGraphSchema.h"
+#include "Framework/MultiBox/MultiBoxBuilder.h"
+#include "Framework/Commands/UIAction.h"
 #endif
 
 #define LOCTEXT_NAMESPACE "DataflowEdNode"
@@ -84,11 +86,13 @@ void UDataflowEdNode::AllocateDefaultPins()
 				{
 					if (Pin.Direction == Dataflow::FPin::EDirection::INPUT)
 					{
-						CreatePin(EEdGraphPinDirection::EGPD_Input, Pin.Type, Pin.Name);
+						UEdGraphPin* const EdPin = CreatePin(EEdGraphPinDirection::EGPD_Input, Pin.Type, Pin.Name);
+						EdPin->bHidden = Pin.bHidden;
 					}					
 					if (Pin.Direction == Dataflow::FPin::EDirection::OUTPUT)
 					{
-						CreatePin(EEdGraphPinDirection::EGPD_Output, Pin.Type, Pin.Name);
+						UEdGraphPin* const EdPin = CreatePin(EEdGraphPinDirection::EGPD_Output, Pin.Type, Pin.Name);
+						EdPin->bHidden = Pin.bHidden;
 					}
 				}
 			}
@@ -161,17 +165,21 @@ void UDataflowEdNode::UpdatePinsFromDataflowNode()
 				{
 					if (Pin.Direction == Dataflow::FPin::EDirection::INPUT)
 					{
-						if (!FindPin(Pin.Name, EEdGraphPinDirection::EGPD_Input))
+						UEdGraphPin* EdPin = FindPin(Pin.Name, EEdGraphPinDirection::EGPD_Input);
+						if(!EdPin)
 						{
-							CreatePin(EEdGraphPinDirection::EGPD_Input, Pin.Type, Pin.Name);
+							EdPin = CreatePin(EEdGraphPinDirection::EGPD_Input, Pin.Type, Pin.Name);
 						}
+						EdPin->bHidden = Pin.bHidden;
 					}
 					if (Pin.Direction == Dataflow::FPin::EDirection::OUTPUT)
 					{
-						if (!FindPin(Pin.Name, EEdGraphPinDirection::EGPD_Output))
+						UEdGraphPin* EdPin = FindPin(Pin.Name, EEdGraphPinDirection::EGPD_Output);
+						if(!EdPin)
 						{
-							CreatePin(EEdGraphPinDirection::EGPD_Output, Pin.Type, Pin.Name);
+							EdPin = CreatePin(EEdGraphPinDirection::EGPD_Output, Pin.Type, Pin.Name);
 						}
+						EdPin->bHidden = Pin.bHidden;
 					}
 				}
 			}
@@ -624,6 +632,113 @@ bool UDataflowEdNode::ShouldDrawNodeAsControlPointOnly(int32& OutInputPinIndex,
 	return false;
 }
 
+void UDataflowEdNode::HideAllInputPins()
+{
+	bool bAnyHidden = false;
+	if (TSharedPtr<FDataflowNode> DataflowNode = GetDataflowNode())
+	{
+		TArray<FDataflowInput*> Inputs = DataflowNode->GetInputs();
+		for (FDataflowInput* const Input : Inputs)
+		{
+			if (Input->GetCanHidePin() && !Input->GetPinIsHidden())
+			{
+				Input->SetPinIsHidden(true);
+				if (!bAnyHidden)
+				{
+					Modify();
+					bAnyHidden = true;
+				}
+				UEdGraphPin* const EdPin = FindPin(Input->GetName(), EEdGraphPinDirection::EGPD_Input);
+				check(EdPin);
+				EdPin->Modify();
+				EdPin->bHidden = true;
+			}
+		}
+	}
+
+	if (bAnyHidden)
+	{
+		GetGraph()->NotifyGraphChanged();
+	}
+}
+
+void UDataflowEdNode::ShowAllInputPins()
+{
+	bool bAnyUnhidden = false;
+	if (TSharedPtr<FDataflowNode> DataflowNode = GetDataflowNode())
+	{
+		TArray<FDataflowInput*> Inputs = DataflowNode->GetInputs();
+		for (FDataflowInput* const Input : Inputs)
+		{
+			if (Input->GetCanHidePin() && Input->GetPinIsHidden())
+			{
+				Input->SetPinIsHidden(false);
+				if (!bAnyUnhidden)
+				{
+					Modify();
+					bAnyUnhidden = true;
+				}
+				UEdGraphPin* const EdPin = FindPin(Input->GetName(), EEdGraphPinDirection::EGPD_Input);
+				check(EdPin);
+				EdPin->Modify();
+				EdPin->bHidden = false;
+			}
+		}
+	}
+
+	if (bAnyUnhidden)
+	{
+		GetGraph()->NotifyGraphChanged();
+	}
+}
+
+void UDataflowEdNode::ToggleHideInputPin(FName PinName)
+{
+	if (TSharedPtr<FDataflowNode> DataflowNode = GetDataflowNode())
+	{
+		if (FDataflowInput* const Input = DataflowNode->FindInput(PinName))
+		{
+			if (ensure(Input->GetCanHidePin()))
+			{
+				const bool bWasHidden = Input->GetPinIsHidden();
+				Input->SetPinIsHidden(!bWasHidden);
+				Modify();
+				UEdGraphPin* const EdPin = FindPin(Input->GetName(), EEdGraphPinDirection::EGPD_Input);
+				check(EdPin);
+				EdPin->Modify();
+				EdPin->bHidden = !bWasHidden;
+				GetGraph()->NotifyGraphChanged();
+			}
+		}
+	}
+}
+
+bool UDataflowEdNode::CanToggleHideInputPin(FName PinName) const
+{
+	if (TSharedPtr<const FDataflowNode> DataflowNode = GetDataflowNode())
+	{
+		if (const FDataflowInput* const Input = DataflowNode->FindInput(PinName))
+		{
+			if (Input->GetCanHidePin() && !Input->HasAnyConnections())
+			{
+				return true;
+			}
+		}
+	}
+	return false;
+}
+
+bool UDataflowEdNode::IsInputPinShown(FName PinName) const
+{
+	if (TSharedPtr<const FDataflowNode> DataflowNode = GetDataflowNode())
+	{
+		if (const FDataflowInput* const Input = DataflowNode->FindInput(PinName))
+		{
+			return !Input->GetPinIsHidden();
+		}
+	}
+	return false;
+}
 #endif //WITH_EDITOR
 
 
diff --git a/Engine/Source/Runtime/Experimental/Dataflow/Engine/Public/Dataflow/DataflowEdNode.h b/Engine/Source/Runtime/Experimental/Dataflow/Engine/Public/Dataflow/DataflowEdNode.h
--- a/Engine/Source/Runtime/Experimental/Dataflow/Engine/Public/Dataflow/DataflowEdNode.h
+++ b/Engine/Source/Runtime/Experimental/Dataflow/Engine/Public/Dataflow/DataflowEdNode.h
@@ -76,6 +76,15 @@ public:
 	/** Remove an option pin if the underlying Dataflow node RemovePin member is overriden. */
 	DATAFLOWENGINE_API void RemoveOptionPin();
 
+#if WITH_EDITOR
+	// Pin hiding
+	DATAFLOWENGINE_API void HideAllInputPins();
+	DATAFLOWENGINE_API void ShowAllInputPins();
+	DATAFLOWENGINE_API void ToggleHideInputPin(FName PinName);
+	DATAFLOWENGINE_API bool CanToggleHideInputPin(FName PinName) const;
+	DATAFLOWENGINE_API bool IsInputPinShown(FName PinName) const;
+#endif
+
 	//
 	// Node Rendering
 	//
