ChaosFlesh: Fix caching to update cached particles as kinematic. Bug fix in triangle mesh simulation code. Add spatial hash building for kinematic triangles.

diff --git a/Engine/Plugins/Experimental/ChaosFlesh/Source/ChaosFleshEngine/Private/ChaosCache/FleshComponentCacheAdapter.cpp b/Engine/Plugins/Experimental/ChaosFlesh/Source/ChaosFleshEngine/Private/ChaosCache/FleshComponentCacheAdapter.cpp
--- a/Engine/Plugins/Experimental/ChaosFlesh/Source/ChaosFleshEngine/Private/ChaosCache/FleshComponentCacheAdapter.cpp
+++ b/Engine/Plugins/Experimental/ChaosFlesh/Source/ChaosFleshEngine/Private/ChaosCache/FleshComponentCacheAdapter.cpp
@@ -210,53 +210,65 @@ namespace Chaos
 					const int32 NumParticles = Particles.Size();
 					Softs::FSolverVec3* ParticleXs = &Particles.X(0);
 					Softs::FSolverVec3* ParticleVs = &Particles.V(0);
+					Softs::FSolverReal* ParticleInvMs = &Particles.InvM(0);
+					Softs::FPAndInvM* ParticlePAndInvMs = &Particles.PAndInvM(0);
 
-					int32 NumCachedParticles = static_cast<int32>(Points0.size());
-					if (NumCachedParticles > NumParticles)
+					if (UFleshComponent* FleshComp = CastChecked<UFleshComponent>(InComponent))
 					{
-						// Cached particles doesn't match solver particles.  Truncate.
-						NumCachedParticles = NumParticles;
-					}
-
-					// < time range start, > time range end, or exact hit
-					if (FMath::IsNearlyEqual(Prev, Next))
-					{
-						// Directly set the result of the cache into the solver particles
-						for (int32 CachedIndex = 0; CachedIndex < NumCachedParticles; ++CachedIndex)
+						if (Chaos::Softs::FFleshThreadingProxy* FleshProxy = FleshComp->GetPhysicsProxy()->As<Chaos::Softs::FFleshThreadingProxy>())
 						{
-							// Note that VtArray::operator[] is non-const access and will cause trigger 
-							// the copy-on-write memcopy!  VtArray::cdata() avoids that.
-							const pxr::GfVec3f& P0 = Points0.cdata()[CachedIndex];
-							const pxr::GfVec3f& V0 = Vels0.cdata()[CachedIndex];
-							ParticleXs[CachedIndex].Set(P0[0], P0[1], P0[2]);
-							ParticleVs[CachedIndex].Set(V0[0], V0[1], V0[2]);
-						}
-						return;
-					}
+							FIntVector2 ParticleRange = FleshProxy->GetSolverParticleRange();
 
-					if (!UE::ChaosCachingUSD::ReadPoints(MonolithStage, PrimPath, Next, Points1, Vels1) ||
-						Points1.size() != Vels1.size() ||
-						Points0.size() != Points1.size())
-					{
-						UE_LOG(LogChaosFleshCache, Error,
-							TEXT("Failed to read points '%s' at time %g from file: '%s'"),
-							*PrimPath, Next, *MonolithStage.GetRootLayer().GetDisplayName());
-						return;
-					}
-					double Duration = Next - Prev;
-					double Alpha = Duration > UE_SMALL_NUMBER ? (TargetTime - Prev) / Duration : 0.5;
-					for (int32 CachedIndex = 0; CachedIndex < NumCachedParticles; ++CachedIndex)
-					{
-						// Note that VtArray::operator[] is non-const access and will cause trigger 
-						// the copy-on-write memcopy!  VtArray::cdata() avoids that.
-						const pxr::GfVec3f& P0 = Points0.cdata()[CachedIndex];
-						const pxr::GfVec3f& P1 = Points1.cdata()[CachedIndex];
-						const pxr::GfVec3f& V0 = Vels0.cdata()[CachedIndex];
-						const pxr::GfVec3f& V1 = Vels1.cdata()[CachedIndex];
-						pxr::GfVec3f Pos = (1.0 - Alpha) * P0 + Alpha * P1;
-						pxr::GfVec3f Vel = (1.0 - Alpha) * V0 + Alpha * V1;
-						ParticleXs[CachedIndex].Set(Pos[0], Pos[1], Pos[2]);
-						ParticleVs[CachedIndex].Set(Vel[0], Vel[1], Vel[2]);
+							int32 NumCachedParticles = static_cast<int32>(Points0.size());
+							if (NumCachedParticles > ParticleRange[1])
+							{
+								// Cached particles doesn't match solver particles.  Truncate.
+								NumCachedParticles = ParticleRange[1];
+							}
+							// < time range start, > time range end, or exact hit
+							if (FMath::IsNearlyEqual(Prev, Next))
+							{
+								// Directly set the result of the cache into the solver particles
+								for (int32 CachedIndex = ParticleRange[0]; CachedIndex < ParticleRange[0] + ParticleRange[1]; ++CachedIndex)
+								{
+									// Note that VtArray::operator[] is non-const access and will cause trigger 
+									// the copy-on-write memcopy!  VtArray::cdata() avoids that.
+									const pxr::GfVec3f& P0 = Points0.cdata()[CachedIndex - ParticleRange[0]];
+									const pxr::GfVec3f& V0 = Vels0.cdata()[CachedIndex - ParticleRange[0]];
+									ParticleXs[CachedIndex].Set(P0[0], P0[1], P0[2]);
+									ParticleVs[CachedIndex].Set(V0[0], V0[1], V0[2]);
+									ParticleInvMs[CachedIndex] = Softs::FSolverReal(0);
+									ParticlePAndInvMs[CachedIndex].InvM = Softs::FSolverReal(0);
+									ParticlePAndInvMs[CachedIndex].P = ParticleXs[CachedIndex];
+								}
+								return;
+							}
+
+							if (!UE::ChaosCachingUSD::ReadPoints(MonolithStage, PrimPath, Next, Points1, Vels1) ||
+								Points1.size() != Vels1.size() ||
+								Points0.size() != Points1.size())
+							{
+								UE_LOG(LogChaosFleshCache, Error,
+									TEXT("Failed to read points '%s' at time %g from file: '%s'"),
+									*PrimPath, Next, *MonolithStage.GetRootLayer().GetDisplayName());
+								return;
+							}
+							double Duration = Next - Prev;
+							double Alpha = Duration > UE_SMALL_NUMBER ? (TargetTime - Prev) / Duration : 0.5;
+							for (int32 CachedIndex = ParticleRange[0]; CachedIndex < ParticleRange[0] + ParticleRange[1]; ++CachedIndex)
+							{
+								// Note that VtArray::operator[] is non-const access and will cause trigger 
+								// the copy-on-write memcopy!  VtArray::cdata() avoids that.
+								const pxr::GfVec3f& P0 = Points0.cdata()[CachedIndex - ParticleRange[0]];
+								const pxr::GfVec3f& P1 = Points1.cdata()[CachedIndex - ParticleRange[0]];
+								const pxr::GfVec3f& V0 = Vels0.cdata()[CachedIndex - ParticleRange[0]];
+								const pxr::GfVec3f& V1 = Vels1.cdata()[CachedIndex - ParticleRange[0]];
+								pxr::GfVec3f Pos = (1.0 - Alpha) * P0 + Alpha * P1;
+								pxr::GfVec3f Vel = (1.0 - Alpha) * V0 + Alpha * V1;
+								ParticleXs[CachedIndex].Set(Pos[0], Pos[1], Pos[2]);
+								ParticleVs[CachedIndex].Set(Vel[0], Vel[1], Vel[2]);
+							}
+						}
 					}
 				}
 #else // USE_USD_SDK && DO_USD_CACHING
@@ -660,11 +672,14 @@ namespace Chaos
 		if (FDeformableSolver* Solver = GetDeformableSolver(InComponent))
 		{
 			FDeformableSolver::FGameThreadAccess GameThreadAccess(Solver, Softs::FGameThreadAccessor());
-			GameThreadAccess.SetEnableSolver(false);
 			if (UDeformableTetrahedralComponent* FleshComp = CastChecked<UDeformableTetrahedralComponent>(InComponent))
 			{
 				FleshComp->ResetDynamicCollection();
 
+				if (FleshComp->GetPhysicsProxy()->As<Chaos::Softs::FFleshThreadingProxy>())
+				{
+					FleshComp->GetPhysicsProxy()->As<Chaos::Softs::FFleshThreadingProxy>()->SetIsCached(true);
+				}
 #if USE_USD_SDK && DO_USD_CACHING
 				//
 				// USD caching
diff --git a/Engine/Plugins/Experimental/ChaosFlesh/Source/ChaosFleshEngine/Private/ChaosFlesh/ChaosDeformableSolverComponent.cpp b/Engine/Plugins/Experimental/ChaosFlesh/Source/ChaosFleshEngine/Private/ChaosFlesh/ChaosDeformableSolverComponent.cpp
--- a/Engine/Plugins/Experimental/ChaosFlesh/Source/ChaosFleshEngine/Private/ChaosFlesh/ChaosDeformableSolverComponent.cpp
+++ b/Engine/Plugins/Experimental/ChaosFlesh/Source/ChaosFleshEngine/Private/ChaosFlesh/ChaosDeformableSolverComponent.cpp
@@ -223,6 +223,7 @@ void UDeformableSolverComponent::Reset()
 			, SolverConstraints.GaussSeidelConstraints.CollisionSpring.CollisionSpringStiffness
 			, SolverConstraints.GaussSeidelConstraints.CollisionSpring.bAllowSliding
 			, SolverMuscleActivation.bDoMuscleActivation
+			, SolverConstraints.GaussSeidelConstraints.CollisionSpring.bCollideWithFullmesh
 		}));
 
 		for (TObjectPtr<UDeformablePhysicsComponent>& DeformableComponent : ConnectedObjects.DeformableComponents)
diff --git a/Engine/Plugins/Experimental/ChaosFlesh/Source/ChaosFleshEngine/Public/ChaosFlesh/ChaosDeformableSolverComponent.h b/Engine/Plugins/Experimental/ChaosFlesh/Source/ChaosFleshEngine/Public/ChaosFlesh/ChaosDeformableSolverComponent.h
--- a/Engine/Plugins/Experimental/ChaosFlesh/Source/ChaosFleshEngine/Public/ChaosFlesh/ChaosDeformableSolverComponent.h
+++ b/Engine/Plugins/Experimental/ChaosFlesh/Source/ChaosFleshEngine/Public/ChaosFlesh/ChaosDeformableSolverComponent.h
@@ -128,6 +128,11 @@ struct FCollisionSpringGroup
 	*/
 	UPROPERTY(EditAnywhere, Category = "CollisionSpring")
 	bool bAllowSliding = true;
+	/**
+	* Do self collision with kinematic triangles as well
+	*/
+	UPROPERTY(EditAnywhere, Category = "CollisionSpring")
+	bool bCollideWithFullmesh = true;
 };
 
 USTRUCT(BlueprintType)
diff --git a/Engine/Source/Runtime/Experimental/Chaos/Private/Chaos/Deformable/ChaosDeformableSolver.cpp b/Engine/Source/Runtime/Experimental/Chaos/Private/Chaos/Deformable/ChaosDeformableSolver.cpp
--- a/Engine/Source/Runtime/Experimental/Chaos/Private/Chaos/Deformable/ChaosDeformableSolver.cpp
+++ b/Engine/Source/Runtime/Experimental/Chaos/Private/Chaos/Deformable/ChaosDeformableSolver.cpp
@@ -768,12 +768,40 @@ namespace Chaos::Softs
 		auto ChaosTet = [](FIntVector4 V, int32 dp) { return Chaos::TVec4<int32>(dp + V.X, dp + V.Y, dp + V.Z, dp + V.W); };
 
 		TArray<FSolverReal> StiffnessWithMultiplier;
+		const FIntVector2& Range = Proxy.GetSolverParticleRange();
+		if (Rest.HasAttributes({ FManagedArrayCollection::TManagedType<FSolverReal>("Stiffness", FGeometryCollection::VerticesGroup) }))
+		{
+			uint32 NumParticles = Rest.NumElements(FGeometryCollection::VerticesGroup);
+
+			StiffnessWithMultiplier.Init(0.f, NumParticles);
+			FSolverReal StiffnessMultiplier = 1.f;
+			FSolverReal IncompressibilityMultiplier = 1.f;
+			FSolverReal InflationMultiplier = 1.f;
+
+			if (const UObject* Owner = this->MObjects[Range[0]]) {
+				FFleshThreadingProxy::FFleshInputBuffer* FleshInputBuffer = nullptr;
+				if (this->CurrentInputPackage->ObjectMap.Contains(Owner))
+				{
+					FleshInputBuffer = this->CurrentInputPackage->ObjectMap[Owner]->As<FFleshThreadingProxy::FFleshInputBuffer>();
+					if (FleshInputBuffer)
+					{
+						StiffnessMultiplier = FleshInputBuffer->StiffnessMultiplier;
+					}
+				}
+			}
+			const TManagedArray<FSolverReal>* StiffnessArray = Rest.FindAttribute<FSolverReal>("Stiffness", FGeometryCollection::VerticesGroup);
+			if (StiffnessArray)
+			{
+				for (uint32 vdx = 0; vdx < NumParticles; ++vdx)
+				{
+					StiffnessWithMultiplier[vdx] = (*StiffnessArray)[vdx] * StiffnessMultiplier;
+				}
+			}
+		}
 
 		const TManagedArray<FIntVector4>& Tetrahedron = Rest.GetAttribute<FIntVector4>("Tetrahedron", "Tetrahedral");
 		if (uint32 NumElements = Tetrahedron.Num())
 		{
-			const FIntVector2& Range = Proxy.GetSolverParticleRange();
-
 			// Add Tetrahedral Elements Node
 			TArray<Chaos::TVec4<int32>> Elements;
 			Elements.SetNum(NumElements);
@@ -797,7 +825,7 @@ namespace Chaos::Softs
 			{
 				uint32 NumParticles = Rest.NumElements(FGeometryCollection::VerticesGroup);
 
-				StiffnessWithMultiplier.Init(0.f, NumParticles);
+				//StiffnessWithMultiplier.Init(0.f, NumParticles);
 				FSolverReal StiffnessMultiplier = 1.f;
 				FSolverReal IncompressibilityMultiplier = 1.f;
 				FSolverReal InflationMultiplier = 1.f;
@@ -809,7 +837,6 @@ namespace Chaos::Softs
 						FleshInputBuffer = this->CurrentInputPackage->ObjectMap[Owner]->As<FFleshThreadingProxy::FFleshInputBuffer>();
 						if (FleshInputBuffer)
 						{
-							StiffnessMultiplier = FleshInputBuffer->StiffnessMultiplier;
 							IncompressibilityMultiplier = FleshInputBuffer->IncompressibilityMultiplier;
 							InflationMultiplier = FleshInputBuffer->InflationMultiplier;
 						}
@@ -820,10 +847,6 @@ namespace Chaos::Softs
 				TetStiffness.Init(Property.EMesh, Elements.Num());
 				if (StiffnessArray)
 				{
-					for (uint32 vdx = 0; vdx < NumParticles; ++vdx)
-					{
-						StiffnessWithMultiplier[vdx] = (*StiffnessArray)[vdx] * StiffnessMultiplier;
-					}
 					for (int32 edx = 0; edx < Elements.Num(); edx++)
 					{
 						TetStiffness[edx] = (StiffnessWithMultiplier[Tetrahedron[edx].X] + StiffnessWithMultiplier[Tetrahedron[edx].Y]
@@ -935,7 +958,8 @@ namespace Chaos::Softs
 
 		}
 
-		const FIntVector2& Range = Proxy.GetSolverParticleRange();
+		FIntVector2 AllUnconstrainedSurfaceElementsSkinRange(INDEX_NONE), AllUnconstrainedSurfaceElementsCorotatedCodRange(INDEX_NONE);
+
 		if (const TManagedArray<int32>* TriangleMeshIndices = Rest.FindAttribute<int32>("ObjectIndices", "TriangleMesh"))
 		{
 			if (const TManagedArray<FIntVector>* Indices = Rest.FindAttribute<FIntVector>("Indices", FGeometryCollection::FacesGroup))
@@ -948,6 +972,10 @@ namespace Chaos::Softs
 						{
 							if (const TManagedArray<int32>* VertexCounts = Rest.FindAttribute<int32>("VertexCount", FGeometryCollection::GeometryGroup))
 							{
+								AllUnconstrainedSurfaceElementsCorotatedCodRange[0] =  AllUnconstrainedSurfaceElementsCorotatedCod->Num();
+								AllUnconstrainedSurfaceElementsSkinRange[0] =  AllUnconstrainedSurfaceElementsSkin->Num();
+								AllUnconstrainedSurfaceElementsCorotatedCodRange[1] = 0;
+								AllUnconstrainedSurfaceElementsSkinRange[1] = 0;
 								if (const TManagedArray<bool>* bUseSkinConstraints = Rest.FindAttribute<bool>("SkinConstraints", "TriangleMesh"))
 								{
 									for (int32 i = 0; i < TriangleMeshIndices->Num(); i++)
@@ -963,9 +991,10 @@ namespace Chaos::Softs
 											{
 												for (int32 j = 0; j < 3; j++)
 												{
-													(*AllUnconstrainedSurfaceElementsSkin)[e - FaceStartIndex + SurfaceOffset][j] = (*Indices)[e][j];
+													(*AllUnconstrainedSurfaceElementsSkin)[e - FaceStartIndex + SurfaceOffset][j] = (*Indices)[e][j] + Range[0];
 												}
 											}
+											AllUnconstrainedSurfaceElementsSkinRange[1] += FaceNum;
 										}
 										else
 										{
@@ -975,9 +1004,10 @@ namespace Chaos::Softs
 											{
 												for (int32 j = 0; j < 3; j++)
 												{
-													(*AllUnconstrainedSurfaceElementsCorotatedCod)[e - FaceStartIndex + SurfaceOffset][j] = (*Indices)[e][j];
+													(*AllUnconstrainedSurfaceElementsCorotatedCod)[e - FaceStartIndex + SurfaceOffset][j] = (*Indices)[e][j] + Range[0];
 												}
 											}
+											AllUnconstrainedSurfaceElementsCorotatedCodRange[1] += FaceNum;
 										}
 									}
 								}
@@ -994,34 +1024,57 @@ namespace Chaos::Softs
 										{
 											for (int32 j = 0; j < 3; j++)
 											{
-												(*AllUnconstrainedSurfaceElementsCorotatedCod)[e - FaceStartIndex + SurfaceOffset][j] = (*Indices)[e][j];
+												(*AllUnconstrainedSurfaceElementsCorotatedCod)[e - FaceStartIndex + SurfaceOffset][j] = (*Indices)[e][j] + Range[0];
 											}
 										}
+										AllUnconstrainedSurfaceElementsCorotatedCodRange[1] += FaceNum;
 									}
 								}
 							}
-						}
-							
+						}	
 					}
-					
 				}
 			}
-		
 		}
-
-		AllCorotatedCodEMeshArray->Init(Property.EMesh, AllUnconstrainedSurfaceElementsCorotatedCod->Num());
-		AllSkinEMeshArray->Init(Property.EMesh, AllUnconstrainedSurfaceElementsSkin->Num());
+		if (AllCorotatedCodEMeshArray)
+		{
+			if (AllUnconstrainedSurfaceElementsCorotatedCodRange[1] > 0 )
+			{
+				AllCorotatedCodEMeshArray->SetNum(AllCorotatedCodEMeshArray->Num() + AllUnconstrainedSurfaceElementsCorotatedCodRange[1]);
+				for (int32 i = AllUnconstrainedSurfaceElementsCorotatedCodRange[0]; i < AllUnconstrainedSurfaceElementsCorotatedCodRange[0] + AllUnconstrainedSurfaceElementsCorotatedCodRange[1]; i++)
+				{
+					(*AllCorotatedCodEMeshArray)[i] = Property.EMesh;
+				}
+			}
+		}
+		if (AllSkinEMeshArray)
+		{
+			if (AllUnconstrainedSurfaceElementsSkinRange[1] > 0)
+			{
+				AllSkinEMeshArray->SetNum(AllSkinEMeshArray->Num() + AllUnconstrainedSurfaceElementsSkinRange[1]);
+				for (int32 i = AllUnconstrainedSurfaceElementsSkinRange[0]; i < AllUnconstrainedSurfaceElementsSkinRange[0] + AllUnconstrainedSurfaceElementsSkinRange[1]; i++)
+				{
+					(*AllSkinEMeshArray)[i] = Property.EMesh;
+				}
+			}
+		}
 		if (StiffnessWithMultiplier.Num() > 0)
 		{
-			for (int32 i = 0; i < AllUnconstrainedSurfaceElementsCorotatedCod->Num(); i++)
+			if (AllUnconstrainedSurfaceElementsCorotatedCodRange[1] > 0 && AllCorotatedCodEMeshArray)
 			{
-				(*AllCorotatedCodEMeshArray)[i] = (StiffnessWithMultiplier[(*AllUnconstrainedSurfaceElementsCorotatedCod)[i][0]] + StiffnessWithMultiplier[(*AllUnconstrainedSurfaceElementsCorotatedCod)[i][1]]
-								+ StiffnessWithMultiplier[(*AllUnconstrainedSurfaceElementsCorotatedCod)[i][2]]) / 3.f;
+				for (int32 i = AllUnconstrainedSurfaceElementsCorotatedCodRange[0]; i < AllUnconstrainedSurfaceElementsCorotatedCodRange[0] + AllUnconstrainedSurfaceElementsCorotatedCodRange[1]; i++)
+				{
+					(*AllCorotatedCodEMeshArray)[i] = (StiffnessWithMultiplier[(*AllUnconstrainedSurfaceElementsCorotatedCod)[i][0] - Range[0]] + StiffnessWithMultiplier[(*AllUnconstrainedSurfaceElementsCorotatedCod)[i][1] - Range[0]]
+									+ StiffnessWithMultiplier[(*AllUnconstrainedSurfaceElementsCorotatedCod)[i][2] - Range[0]]) / 3.f;
+				}
 			}
-			for (int32 i = 0; i < AllSkinEMeshArray->Num(); i++)
+			if (AllSkinEMeshArray && AllUnconstrainedSurfaceElementsSkinRange[1] > 0)
 			{
-				(*AllSkinEMeshArray)[i] = (StiffnessWithMultiplier[(*AllUnconstrainedSurfaceElementsSkin)[i][0]] + StiffnessWithMultiplier[(*AllUnconstrainedSurfaceElementsSkin)[i][1]]
-								+ StiffnessWithMultiplier[(*AllUnconstrainedSurfaceElementsSkin)[i][2]]) / 3.f;
+				for (int32 i = AllUnconstrainedSurfaceElementsSkinRange[0]; i < AllUnconstrainedSurfaceElementsSkinRange[0] + AllUnconstrainedSurfaceElementsSkinRange[1]; i++)
+				{
+					(*AllSkinEMeshArray)[i] = (StiffnessWithMultiplier[(*AllUnconstrainedSurfaceElementsSkin)[i][0] - Range[0]] + StiffnessWithMultiplier[(*AllUnconstrainedSurfaceElementsSkin)[i][1] - Range[0]]
+									+ StiffnessWithMultiplier[(*AllUnconstrainedSurfaceElementsSkin)[i][2] - Range[0]]) / 3.f;
+				}
 			}
 		}
 	}
@@ -1105,6 +1158,7 @@ namespace Chaos::Softs
 
 			if (0 <= Index && Index < this->MObjects.Num())
 			{
+
 				if (TransientConstraintBuffer.Contains(Index))
 				{
 					return;
@@ -1114,109 +1168,106 @@ namespace Chaos::Softs
 				{
 					if (const FFleshThreadingProxy* Proxy = Proxies[Owner]->As<FFleshThreadingProxy>())
 					{
-						FTransform GlobalTransform = Proxy->GetCurrentPointsTransform();
-						const FIntVector2& Range = Proxy->GetSolverParticleRange();
-						const FManagedArrayCollection& Rest = Proxy->GetRestCollection();
-
-						if (Rest.FindAttributeTyped<FVector3f>("Vertex", FGeometryCollection::VerticesGroup))
+						if (!Proxy->GetIsCached())
 						{
-							const TManagedArray<FVector3f>& Vertex = Rest.GetAttribute<FVector3f>("Vertex", FGeometryCollection::VerticesGroup);
-							// @todo(chaos) : reduce conversions
-							auto ChaosVert = [](FVector3f V) { return Chaos::FVec3(V.X, V.Y, V.Z); };
-							auto ChaosVertfloat = [](FVector3f V) { return Chaos::TVector<FSolverReal, 3>(V.X, V.Y, V.Z); };
-							auto SolverParticleToObjectVertexIndex = [&](int32 SolverParticleIndex) {return SolverParticleIndex - Range[0]; };
-
-							FFleshThreadingProxy::FFleshInputBuffer* FleshInputBuffer = nullptr;
-							if (this->CurrentInputPackage->ObjectMap.Contains(Owner))
-							{
-								FleshInputBuffer = this->CurrentInputPackage->ObjectMap[Owner]->As<FFleshThreadingProxy::FFleshInputBuffer>();
-							}
+							FTransform GlobalTransform = Proxy->GetCurrentPointsTransform();
+							const FIntVector2& Range = Proxy->GetSolverParticleRange();
+							const FManagedArrayCollection& Rest = Proxy->GetRestCollection();
 
-							typedef GeometryCollection::Facades::FVertexBoneWeightsFacade FWeightsFacade;
-							bool bParticleTouched = false;
-							FWeightsFacade WeightsFacade(Rest);
-							if (WeightsFacade.IsValid())
+							if (Rest.FindAttributeTyped<FVector3f>("Vertex", FGeometryCollection::VerticesGroup))
 							{
-								int32 NumObjectVertices = Rest.NumElements(FGeometryCollection::VerticesGroup);
-								int32 ObjectVertexIndex = SolverParticleToObjectVertexIndex(Index);
-								if (ensure(0 <= ObjectVertexIndex && ObjectVertexIndex < NumObjectVertices))
+								const TManagedArray<FVector3f>& Vertex = Rest.GetAttribute<FVector3f>("Vertex", FGeometryCollection::VerticesGroup);
+								// @todo(chaos) : reduce conversions
+								auto ChaosVert = [](FVector3f V) { return Chaos::FVec3(V.X, V.Y, V.Z); };
+								auto ChaosVertfloat = [](FVector3f V) { return Chaos::TVector<FSolverReal, 3>(V.X, V.Y, V.Z); };
+								auto SolverParticleToObjectVertexIndex = [&](int32 SolverParticleIndex) {return SolverParticleIndex - Range[0]; };
+
+								FFleshThreadingProxy::FFleshInputBuffer* FleshInputBuffer = nullptr;
+								if (this->CurrentInputPackage->ObjectMap.Contains(Owner))
 								{
-									if (FleshInputBuffer)
-									{
-										TArray<int32> BoneIndices = WeightsFacade.GetBoneIndices()[ObjectVertexIndex];
-										TArray<float> BoneWeights = WeightsFacade.GetBoneWeights()[ObjectVertexIndex];
+									FleshInputBuffer = this->CurrentInputPackage->ObjectMap[Owner]->As<FFleshThreadingProxy::FFleshInputBuffer>();
+								}
 
-										FFleshThreadingProxy::FFleshInputBuffer* PreviousFleshBuffer = nullptr;
-										if (this->PreviousInputPackage && this->PreviousInputPackage->ObjectMap.Contains(Owner))
+								typedef GeometryCollection::Facades::FVertexBoneWeightsFacade FWeightsFacade;
+								bool bParticleTouched = false;
+								FWeightsFacade WeightsFacade(Rest);
+								if (WeightsFacade.IsValid())
+								{
+									int32 NumObjectVertices = Rest.NumElements(FGeometryCollection::VerticesGroup);
+									int32 ObjectVertexIndex = SolverParticleToObjectVertexIndex(Index);
+									if (ensure(0 <= ObjectVertexIndex && ObjectVertexIndex < NumObjectVertices))
+									{
+										if (FleshInputBuffer)
 										{
-											PreviousFleshBuffer = this->PreviousInputPackage->ObjectMap[Owner]->As<FFleshThreadingProxy::FFleshInputBuffer>();
-										}
+											TArray<int32> BoneIndices = WeightsFacade.GetBoneIndices()[ObjectVertexIndex];
+											TArray<float> BoneWeights = WeightsFacade.GetBoneWeights()[ObjectVertexIndex];
 
-										MParticles.SetX(Index, Chaos::TVector<FSolverReal, 3>((FSolverReal)0.));
-										TVector<FSolverReal, 3> TargetPos((FSolverReal)0.);
-										FSolverReal CurrentRatio = FSolverReal(this->Iteration) / FSolverReal(this->Property.NumSolverSubSteps);
+											FFleshThreadingProxy::FFleshInputBuffer* PreviousFleshBuffer = nullptr;
+											if (this->PreviousInputPackage && this->PreviousInputPackage->ObjectMap.Contains(Owner))
+											{
+												PreviousFleshBuffer = this->PreviousInputPackage->ObjectMap[Owner]->As<FFleshThreadingProxy::FFleshInputBuffer>();
+											}
 
-										int32 RestNum = FleshInputBuffer->RestTransforms.Num();
-										int32 TransformNum = FleshInputBuffer->Transforms.Num();
-										if (RestNum > 0 && TransformNum > 0)
-										{
+											MParticles.SetX(Index, Chaos::TVector<FSolverReal, 3>((FSolverReal)0.));
+											TVector<FSolverReal, 3> TargetPos((FSolverReal)0.);
+											FSolverReal CurrentRatio = FSolverReal(this->Iteration) / FSolverReal(this->Property.NumSolverSubSteps);
 
-											for (int32 i = 0; i < BoneIndices.Num(); i++)
+											int32 RestNum = FleshInputBuffer->RestTransforms.Num();
+											int32 TransformNum = FleshInputBuffer->Transforms.Num();
+											if (RestNum > 0 && TransformNum > 0)
 											{
-												if (BoneIndices[i] > INDEX_NONE && BoneIndices[i] < RestNum && BoneIndices[i] < TransformNum)
+
+												for (int32 i = 0; i < BoneIndices.Num(); i++)
 												{
+													if (BoneIndices[i] > INDEX_NONE && BoneIndices[i] < RestNum && BoneIndices[i] < TransformNum)
+													{
 
-													// @todo(flesh) : Add the pre-cached component space rest transforms to the rest collection. 
-													// see  UFleshComponent::NewDeformableData for how its pulled from the SkeletalMesh
-													FVec3 LocalPoint = FleshInputBuffer->RestTransforms[BoneIndices[i]].InverseTransformPosition(ChaosVert(Vertex[Index - Range[0]]));
-													FVec3 ComponentPointAtT = FleshInputBuffer->Transforms[BoneIndices[i]].TransformPosition(LocalPoint);
+														// @todo(flesh) : Add the pre-cached component space rest transforms to the rest collection. 
+														// see  UFleshComponent::NewDeformableData for how its pulled from the SkeletalMesh
+														FVec3 LocalPoint = FleshInputBuffer->RestTransforms[BoneIndices[i]].InverseTransformPosition(ChaosVert(Vertex[Index - Range[0]]));
+														FVec3 ComponentPointAtT = FleshInputBuffer->Transforms[BoneIndices[i]].TransformPosition(LocalPoint);
 
-													if (PreviousFleshBuffer)
-													{
-														FTransform BonePreviousTransform = PreviousFleshBuffer->Transforms[BoneIndices[i]];
-														ComponentPointAtT = ComponentPointAtT * CurrentRatio + BonePreviousTransform.TransformPosition(LocalPoint) * ((FSolverReal)1. - CurrentRatio);
-													}
+														if (PreviousFleshBuffer)
+														{
+															FTransform BonePreviousTransform = PreviousFleshBuffer->Transforms[BoneIndices[i]];
+															ComponentPointAtT = ComponentPointAtT * CurrentRatio + BonePreviousTransform.TransformPosition(LocalPoint) * ((FSolverReal)1. - CurrentRatio);
+														}
 
-													MParticles.SetX(Index, MParticles.GetX(Index) + GlobalTransform.TransformPosition(ComponentPointAtT) * BoneWeights[i]);
+														MParticles.SetX(Index, MParticles.GetX(Index) + GlobalTransform.TransformPosition(ComponentPointAtT) * BoneWeights[i]);
 
-													bParticleTouched = true;
+														bParticleTouched = true;
+													}
 												}
+
 											}
-#if WITH_EDITOR
-											//debug draw
-											//p.Chaos.DebugDraw.Enabled 1
-											//p.Chaos.DebugDraw.Deformable.KinematicParticle 1
-											if (GDeformableDebugParams.IsDebugDrawingEnabled() && GDeformableDebugParams.bDoDrawKinematicParticles)
-											{
-												auto DoubleVert = [](FVector3f V) { return FVector3d(V.X, V.Y, V.Z); };
-												Chaos::FDebugDrawQueue::GetInstance().DrawDebugPoint(DoubleVert(MParticles.GetX(Index)), FColor::Red, false, -1.0f, 0, 5);
-											}
-#endif
+											MParticles.PAndInvM(Index).P = MParticles.GetX(Index);
 										}
-										MParticles.PAndInvM(Index).P = MParticles.GetX(Index);
 									}
 								}
-							}
-							if (!bParticleTouched)
-							{
-								MParticles.SetX(Index, GlobalTransform.TransformPosition(ChaosVert(Vertex[Index - Range[0]])));
-								MParticles.PAndInvM(Index).P = MParticles.GetX(Index);
-
-#if WITH_EDITOR
-								//debug draw
-								//p.Chaos.DebugDraw.Enabled 1
-								//p.Chaos.DebugDraw.Deformable.KinematicParticle 1
-								if (GDeformableDebugParams.IsDebugDrawingEnabled() && GDeformableDebugParams.bDoDrawKinematicParticles)
+								if (!bParticleTouched)
 								{
-									auto DoubleVert = [](FVector3f V) { return FVector3d(V.X, V.Y, V.Z); };
-									Chaos::FDebugDrawQueue::GetInstance().DrawDebugPoint(DoubleVert(MParticles.GetX(Index)), FColor::Red, false, -1.0f, 0, 5);
+									MParticles.SetX(Index, GlobalTransform.TransformPosition(ChaosVert(Vertex[Index - Range[0]])));
+									MParticles.PAndInvM(Index).P = MParticles.GetX(Index);
 								}
-#endif
 							}
 						}
+
+#if WITH_EDITOR
+						//debug draw
+						//p.Chaos.DebugDraw.Enabled 1
+						//p.Chaos.DebugDraw.Deformable.KinematicParticle 1
+						if (GDeformableDebugParams.IsDebugDrawingEnabled() && GDeformableDebugParams.bDoDrawKinematicParticles )
+						{
+							auto DoubleVert = [](FVector3f V) { return FVector3d(V.X, V.Y, V.Z); };
+							Chaos::FDebugDrawQueue::GetInstance().DrawDebugPoint(DoubleVert(MParticles.GetX(Index)), FColor::Red, false, -1.0f, 0, 5);
+						}
+#endif
+
 					}
 				}
 			}
+
+
 		};
 		Evolution->SetKinematicUpdateFunction(MKineticUpdate);
 	}
@@ -1589,7 +1640,7 @@ namespace Chaos::Softs
 			Evolution->ConstraintInits()[InitIndex] =
 				[this](FSolverParticles& InParticles, const FSolverReal Dt)
 				{
-					this->TriangleMeshCollisions->Init(InParticles, Property.CollisionSearchRadius);
+					this->TriangleMeshCollisions->InitFlesh(InParticles, Property.CollisionSearchRadius, this->Property.bCollideWithFullMesh);
 					TArray<FPBDTriangleMeshCollisions::FGIAColor> EmptyGIAColors;
 					if (Property.bDoInComponentSelfCollision)
 					{
diff --git a/Engine/Source/Runtime/Experimental/Chaos/Private/Chaos/PBDTriangleMeshCollisions.cpp b/Engine/Source/Runtime/Experimental/Chaos/Private/Chaos/PBDTriangleMeshCollisions.cpp
--- a/Engine/Source/Runtime/Experimental/Chaos/Private/Chaos/PBDTriangleMeshCollisions.cpp
+++ b/Engine/Source/Runtime/Experimental/Chaos/Private/Chaos/PBDTriangleMeshCollisions.cpp
@@ -1467,6 +1467,55 @@ void FPBDTriangleMeshCollisions::Init(const SolverParticlesOrRange& Particles, c
 template CHAOS_API void FPBDTriangleMeshCollisions::Init(const FSolverParticles& Particles, const FPBDFlatWeightMap& ThicknessMap);
 template CHAOS_API void FPBDTriangleMeshCollisions::Init(const FSolverParticlesRange& Particles, const FPBDFlatWeightMap& ThicknessMap);
 
+template<typename SolverParticlesOrRange>
+void FPBDTriangleMeshCollisions::InitFlesh(const SolverParticlesOrRange& Particles, const FPBDFlatWeightMap& ThicknessMap, const bool bUseFullMesh)
+{
+	const bool bDoSelfIntersections = bGlobalIntersectionAnalysis || bContourMinimization;
+	if (bCollidableSubMeshDirty)
+	{		
+		CollidableSubMesh.Init(Particles, DisabledFaces, bSelfCollideAgainstAllKinematicVertices, EnabledKinematicFaces, bOnlyCollideWithKinematics);
+		bCollidableSubMeshDirty = false;
+	}
+
+	if (CollidableSubMesh.GetDynamicSubMesh().GetNumElements() == 0 && CollidableSubMesh.GetFullMesh().GetNumElements() == 0)
+	{
+		return;
+	}
+
+	{
+		TRACE_CPUPROFILER_EVENT_SCOPE(ChaosFPBDTriangleMeshCollisions_BuildSpatialHash);
+
+		constexpr FSolverReal RadiusToLodSizeMultiplier = 2.; // Radius to Diameter
+		FSolverReal MinProximityQueryRadius;
+		if (!ThicknessMap.HasWeightMap())
+		{	
+			MinProximityQueryRadius = ((FSolverReal)ThicknessMap) * 2.f; // When thickness is constant, we don't put any thickness in the spatial hash and instead query with particles with 2*thickness radius.
+			if (bUseFullMesh)
+			{
+				CollidableSubMesh.GetFullMesh().BuildSpatialHash(static_cast<const TArrayView<const FSolverVec3>&>(Particles.XArray()), DynamicSubMeshSpatialHash, RadiusToLodSizeMultiplier * MinProximityQueryRadius);
+			}
+			else
+			{
+				CollidableSubMesh.GetDynamicSubMesh().BuildSpatialHash(static_cast<const TArrayView<const FSolverVec3>&>(Particles.XArray()), DynamicSubMeshSpatialHash, RadiusToLodSizeMultiplier * MinProximityQueryRadius);
+			}
+		}
+		else
+		{
+			MinProximityQueryRadius = FMath::Max(ThicknessMap.GetLow(), ThicknessMap.GetHigh());
+			if (bUseFullMesh)
+			{
+				CollidableSubMesh.GetFullMesh().BuildSpatialHash(static_cast<const TArrayView<const FSolverVec3>&>(Particles.XArray()), DynamicSubMeshSpatialHash, ThicknessMap, Offset, RadiusToLodSizeMultiplier * MinProximityQueryRadius);
+			}
+			else
+			{
+				CollidableSubMesh.GetDynamicSubMesh().BuildSpatialHash(static_cast<const TArrayView<const FSolverVec3>&>(Particles.XArray()), DynamicSubMeshSpatialHash, ThicknessMap, Offset, RadiusToLodSizeMultiplier * MinProximityQueryRadius);
+			}
+		}
+	}
+}
+template CHAOS_API void FPBDTriangleMeshCollisions::InitFlesh(const FSolverParticles& Particles, const FPBDFlatWeightMap& ThicknessMap, const bool bUseFullMesh);
+template CHAOS_API void FPBDTriangleMeshCollisions::InitFlesh(const FSolverParticlesRange& Particles, const FPBDFlatWeightMap& ThicknessMap, const bool bUseFullMesh);
+
 template<typename SolverParticlesOrRange>
 void FPBDTriangleMeshCollisions::PostStepInit(const SolverParticlesOrRange& Particles)
 {
diff --git a/Engine/Source/Runtime/Experimental/Chaos/Public/Chaos/Deformable/ChaosDeformableSolverProxy.h b/Engine/Source/Runtime/Experimental/Chaos/Public/Chaos/Deformable/ChaosDeformableSolverProxy.h
--- a/Engine/Source/Runtime/Experimental/Chaos/Public/Chaos/Deformable/ChaosDeformableSolverProxy.h
+++ b/Engine/Source/Runtime/Experimental/Chaos/Public/Chaos/Deformable/ChaosDeformableSolverProxy.h
@@ -88,6 +88,7 @@ namespace Chaos::Softs
 		FManagedArrayCollection Dynamic;
 		FIntVector2 SolverParticleRange = FIntVector2(0, 0);
 		ChaosDeformableSimSpace SimSpace = ChaosDeformableSimSpace::World;
+		bool bIsCached = false;
 
 	public:
 
@@ -125,6 +126,9 @@ namespace Chaos::Softs
 			return SimSpace == ChaosDeformableSimSpace::Bone;
 		}
 
+		void SetIsCached(bool bInIsCached) { bIsCached = bInIsCached; }
+		bool GetIsCached() const { return bIsCached; }
+
 		//! Update the component and bone transforms with current data from the scene.
 		void UpdateSimSpace(const FTransform& InWorldToComponentXf, const FTransform& InComponentToBoneXf)
 		{
diff --git a/Engine/Source/Runtime/Experimental/Chaos/Public/Chaos/Deformable/ChaosDeformableSolverTypes.h b/Engine/Source/Runtime/Experimental/Chaos/Public/Chaos/Deformable/ChaosDeformableSolverTypes.h
--- a/Engine/Source/Runtime/Experimental/Chaos/Public/Chaos/Deformable/ChaosDeformableSolverTypes.h
+++ b/Engine/Source/Runtime/Experimental/Chaos/Public/Chaos/Deformable/ChaosDeformableSolverTypes.h
@@ -43,7 +43,8 @@ namespace Chaos::Softs
 			FSolverReal InCollisionSearchRadius = (FSolverReal)0,
 			FSolverReal InCollisionSpringStiffness = (FSolverReal)500.,
 			bool InbAllowSliding = true,
-			bool InbDoMuscleActivation = false)
+			bool InbDoMuscleActivation = false,
+			bool InbCollideWithFullMesh = false)
 			: NumSolverSubSteps(InNumSolverSubSteps)
 			, NumSolverIterations(InNumSolverIterations)
 			, FixTimeStep(InFixTimeStep)
@@ -72,6 +73,7 @@ namespace Chaos::Softs
 			, CollisionSpringStiffness(InCollisionSpringStiffness)
 			, bAllowSliding(InbAllowSliding)
 			, bDoMuscleActivation(InbDoMuscleActivation)
+			, bCollideWithFullMesh(InbCollideWithFullMesh)
 		{}
 
 		int32 NumSolverSubSteps = 5;
@@ -102,6 +104,7 @@ namespace Chaos::Softs
 		FSolverReal CollisionSpringStiffness = (FSolverReal)500.;
 		bool bAllowSliding = true;
 		bool bDoMuscleActivation = false;
+		bool bCollideWithFullMesh = false;
 	};
 
 
diff --git a/Engine/Source/Runtime/Experimental/Chaos/Public/Chaos/PBDTriangleMeshCollisions.h b/Engine/Source/Runtime/Experimental/Chaos/Public/Chaos/PBDTriangleMeshCollisions.h
--- a/Engine/Source/Runtime/Experimental/Chaos/Public/Chaos/PBDTriangleMeshCollisions.h
+++ b/Engine/Source/Runtime/Experimental/Chaos/Public/Chaos/PBDTriangleMeshCollisions.h
@@ -366,6 +366,15 @@ public:
 	template<typename SolverParticlesOrRange>
 	CHAOS_API void Init(const SolverParticlesOrRange& Particles, const FPBDFlatWeightMap& ThicknessMap);
 
+	template<typename SolverParticlesOrRange>
+	void InitFlesh(const SolverParticlesOrRange& Particles, const FSolverReal MinProximityQueryRadius = (FSolverReal)0., const bool bUseFullMesh = false)
+	{
+		InitFlesh(Particles, FPBDFlatWeightMap(FSolverVec2(MinProximityQueryRadius * (FSolverReal).5f)), bUseFullMesh);
+	}
+
+	template<typename SolverParticlesOrRange>
+	CHAOS_API void InitFlesh(const SolverParticlesOrRange& Particles, const FPBDFlatWeightMap& ThicknessMap, const bool bUseFullMesh = false);
+
 	template<typename SolverParticlesOrRange>
 	CHAOS_API void PostStepInit(const SolverParticlesOrRange& Particles);
 
