Unmount on-demand container from package store

diff --git a/Engine/Source/Runtime/Experimental/IoStore/OnDemand/Private/OnDemandIoStore.cpp b/Engine/Source/Runtime/Experimental/IoStore/OnDemand/Private/OnDemandIoStore.cpp
--- a/Engine/Source/Runtime/Experimental/IoStore/OnDemand/Private/OnDemandIoStore.cpp
+++ b/Engine/Source/Runtime/Experimental/IoStore/OnDemand/Private/OnDemandIoStore.cpp
@@ -322,12 +322,18 @@ FIoStatus FOnDemandIoStore::Unmount(FStringView MountId)
 	{
 		TUniqueLock Lock(ContainerMutex);
 
-		Containers.SetNum(Algo::RemoveIf(Containers, [&MountId](const FSharedOnDemandContainer& Container)
+		Containers.SetNum(Algo::RemoveIf(Containers, [this, &MountId](const FSharedOnDemandContainer& Container)
 		{
 			if (Container->MountId == MountId)
 			{
 				UE_LOG(LogIoStoreOnDemand, Log, TEXT("Unmounting container, ContainerName='%s', MountId='%s'"),
 					*WriteToString<128>(Container->Name), *WriteToString<128>(Container->MountId));
+
+				if (PackageStoreBackend.IsValid())
+				{
+					PackageStoreBackend->Unmount(Container->UniqueName());
+				}
+
 				return true;
 			}
 
