Preparing for long running on-demand mount/install requests

diff --git a/Engine/Source/Runtime/Experimental/IoStore/OnDemand/Private/OnDemandIoStore.cpp b/Engine/Source/Runtime/Experimental/IoStore/OnDemand/Private/OnDemandIoStore.cpp
--- a/Engine/Source/Runtime/Experimental/IoStore/OnDemand/Private/OnDemandIoStore.cpp
+++ b/Engine/Source/Runtime/Experimental/IoStore/OnDemand/Private/OnDemandIoStore.cpp
@@ -374,18 +374,30 @@ FOnDemandChunkInfo FOnDemandIoStore::GetChunkInfo(const FIoChunkId& ChunkId, EOn
 
 void FOnDemandIoStore::TryEnterTickLoop()
 {
-	bool bExpected = false;
-	if (bTicking.compare_exchange_strong(bExpected, true))
+	bool bEnterTickLoop = false;
 	{
-		if (FPlatformProcess::SupportsMultithreading() && GIOThreadPool != nullptr)
-		{
-			TickFuture = AsyncPool(*GIOThreadPool, [this] { TickLoop(); }, nullptr, EQueuedWorkPriority::Low);
-		}
-		else
+		UE::TUniqueLock Lock(MountRequestMutex);
+		bTickRequested = true;
+		if (bTicking == false)
 		{
-			TickLoop();
+			bTicking = bEnterTickLoop = true;
 		}
 	}
+
+	if (bEnterTickLoop == false)
+	{
+		UE_LOG(LogIoStoreOnDemand, Verbose, TEXT("I/O store already ticking"));
+		return;
+	}
+
+	if (FPlatformProcess::SupportsMultithreading() && GIOThreadPool != nullptr)
+	{
+		TickFuture = AsyncPool(*GIOThreadPool, [this] { TickLoop(); }, nullptr, EQueuedWorkPriority::Low);
+	}
+	else
+	{
+		TickLoop();
+	}
 }
 
 void FOnDemandIoStore::TickLoop()
@@ -395,20 +407,21 @@ void FOnDemandIoStore::TickLoop()
 	UE_LOG(LogIoStoreOnDemand, Verbose, TEXT("Entering I/O store tick loop"));
 	for (;;)
 	{
+		const bool bTicked = Tick();
+		if (bTicked == false)
 		{
 			UE::TUniqueLock Lock(MountRequestMutex);
-			if (MountRequests.IsEmpty())
+			if (bTickRequested == false)
 			{
 				bTicking = false;
 				break;
 			}
+			bTickRequested = false;
 		}
-
-		Tick();
 	}
 }
 
-void FOnDemandIoStore::Tick()
+bool FOnDemandIoStore::Tick()
 {
 	TArray<FSharedMountRequest> Requests;
 	{
@@ -421,12 +434,12 @@ void FOnDemandIoStore::Tick()
 
 	if (Requests.IsEmpty())
 	{
-		return;
+		return false;
 	}
 
 	for (FSharedMountRequest& Request : Requests)
 	{
-		FIoStatus Status = ProcessMountRequest(*Request);
+		FIoStatus Status = TickMountRequest(*Request);
 
 		{
 			UE::TUniqueLock Lock(MountRequestMutex);
@@ -490,11 +503,13 @@ void FOnDemandIoStore::Tick()
 		FOnDemandMountCompleted OnCompleted = MoveTemp(Request->OnCompleted);
 		OnCompleted(MountStatus);
 	}
+
+	return true;
 }
 
-FIoStatus FOnDemandIoStore::ProcessMountRequest(FMountRequest& MountRequest)
+FIoStatus FOnDemandIoStore::TickMountRequest(FMountRequest& MountRequest)
 {
-	UE_LOG(LogIoStoreOnDemand, Verbose, TEXT("Processing mount request, MountId='%s'"), *MountRequest.MountArgs.MountId);
+	UE_LOG(LogIoStoreOnDemand, Verbose, TEXT("Ticking mount request, MountId='%s'"), *MountRequest.MountArgs.MountId);
 
 	FOnDemandMountArgs& Args = MountRequest.MountArgs;
 	FStringView Host, TocRelUrl;
diff --git a/Engine/Source/Runtime/Experimental/IoStore/OnDemand/Private/OnDemandIoStore.h b/Engine/Source/Runtime/Experimental/IoStore/OnDemand/Private/OnDemandIoStore.h
--- a/Engine/Source/Runtime/Experimental/IoStore/OnDemand/Private/OnDemandIoStore.h
+++ b/Engine/Source/Runtime/Experimental/IoStore/OnDemand/Private/OnDemandIoStore.h
@@ -12,8 +12,6 @@
 #include "Misc/AES.h"
 #include "Misc/EnumClassFlags.h"
 
-#include <atomic>
-
 namespace UE::IoStore
 {
 
@@ -150,8 +148,8 @@ private:
 	FOnDemandChunkInfo		GetChunkInfo(const FIoChunkId& ChunkId, EOnDemandContainerFlags ContainerFlags);
 	void					TryEnterTickLoop();
 	void					TickLoop();
-	void					Tick();
-	FIoStatus				ProcessMountRequest(FMountRequest& MountRequest);
+	bool					Tick();
+	FIoStatus				TickMountRequest(FMountRequest& MountRequest);
 	void					ConditionallyStartTicking();
 	void					OnEncryptionKeyAdded(const FGuid& Id, const FAES::FAESKey& Key);
 	static void				CreateContainersFromToc(
@@ -174,7 +172,8 @@ private:
 	FMountRequestMap					MountRequests;
 	UE::FMutex							MountRequestMutex;
 
-	std::atomic_bool					bTicking{false};
+	bool								bTicking = false;
+	bool								bTickRequested = false;
 	TFuture<void>						TickFuture;
 };
 
