Moved on-demand TOC loading from I/O store to TOC.

diff --git a/Engine/Source/Runtime/Experimental/IoStore/OnDemand/Private/IoStoreOnDemand.cpp b/Engine/Source/Runtime/Experimental/IoStore/OnDemand/Private/IoStoreOnDemand.cpp
--- a/Engine/Source/Runtime/Experimental/IoStore/OnDemand/Private/IoStoreOnDemand.cpp
+++ b/Engine/Source/Runtime/Experimental/IoStore/OnDemand/Private/IoStoreOnDemand.cpp
@@ -3,6 +3,7 @@
 #include "IO/IoStoreOnDemand.h"
 
 #include "HAL/FileManager.h"
+#include "HAL/FileManagerGeneric.h"
 #include "HAL/LowLevelMemTracker.h"
 #include "HAL/PlatformMisc.h"
 #include "HAL/PlatformTime.h"
@@ -320,6 +321,30 @@ static bool SaveArrayToFile(TArrayView64<uint8> Data, const TCHAR* Filename, uin
 	return !Ar->IsError() && !Ar->IsCriticalError();
 }
 
+////////////////////////////////////////////////////////////////////////////////
+/**
+ * Utility to create a FArchive capable of reading from disk using the exact same pathing
+ * rules as FPlatformMisc::LoadTextFileFromPlatformPackage but without forcing the entire
+ * file to be loaded at once.
+ */
+static TUniquePtr<FArchive> CreateReaderFromPlatformPackage(const FString& RelPath)
+{
+	const FString AbsPath = FPaths::Combine(FGenericPlatformMisc::RootDir(), RelPath);
+	if (TUniquePtr<IFileHandle> File(IPlatformFile::GetPlatformPhysical().OpenRead(*AbsPath)); File.IsValid())
+	{
+#if PLATFORM_ANDROID
+		// This is a handle to an asset so we need to call Seek(0) to move the internal
+		// offset to the start of the asset file.
+		File->Seek(0);
+#endif //PLATFORM_ANDROID
+		const uint32 ReadBufferSize = 256 * 1024;
+		const int64 FileSize = File->Size();
+		return MakeUnique<FArchiveFileReaderGeneric>(File.Release(), *AbsPath, FileSize, ReadBufferSize);
+	}
+
+	return TUniquePtr<FArchive>();
+}
+
 ////////////////////////////////////////////////////////////////////////////////
 FArchive& operator<<(FArchive& Ar, FTocMeta& Meta)
 {
@@ -736,58 +761,90 @@ bool LoadFromCompactBinary(FCbFieldView Field, FOnDemandToc& OutToc)
 	return false;
 }
 
-TIoStatusOr<FOnDemandToc> LoadTocFromUrl(const FString& ServiceUrl, const FString& TocPath, int32 RetryCount)
+////////////////////////////////////////////////////////////////////////////////
+TIoStatusOr<FOnDemandToc> FOnDemandToc::LoadFromFile(const FString& FilePath, bool bValidate)
 {
-	TRACE_CPUPROFILER_EVENT_SCOPE(LoadTocFromUrl);
-
-	FString ErrorMsg;
+	TUniquePtr<FArchive> Ar;
+	if (FPlatformMisc::FileExistsInPlatformPackage(FilePath))
+	{
+		Ar = CreateReaderFromPlatformPackage(FilePath);
+	}
+	else
+	{
+		Ar.Reset(IFileManager::Get().CreateFileReader(*FilePath));
+	}
 
-	for (int32 Attempt = 0; Attempt <= RetryCount; ++Attempt)
+	if (Ar.IsValid() == false)
 	{
-		TUniquePtr<FHttpClient> HttpClient = FHttpClient::Create(ServiceUrl);
-		TAnsiStringBuilder<256> Url;
+		FIoStatus Status = FIoStatusBuilder(EIoErrorCode::FileNotOpen) << TEXT("Failed to open '") << FilePath << TEXT("'");
+		return Status;
+	}
 
-		Url << "/" << TocPath;
+	if (bValidate)
+	{
+		const int64 SentinelPos = Ar->TotalSize() - FOnDemandTocSentinel::SentinelSize;
 
-		UE_LOG(LogIas, Log, TEXT("Fetching TOC '%s/%s' (#%d/%d)"), *ServiceUrl, *TocPath, Attempt + 1, RetryCount);
+		if (SentinelPos < 0)
+		{
+			FIoStatus Status = FIoStatusBuilder(EIoErrorCode::CorruptToc) << TEXT("Unexpected file size");
+			return Status;
+		}
 
-		TIoStatusOr<FOnDemandToc> Toc;
-		HttpClient->Get(Url.ToView(), [&Toc, &ErrorMsg](TIoStatusOr<FIoBuffer> Response, uint64 DurationMs)
-			{
-				if (Response.IsOk())
-				{
-					TRACE_CPUPROFILER_EVENT_SCOPE(LoadTocFromEndpoint);
-					FIoBuffer Buffer = Response.ConsumeValueOrDie();
-					FOnDemandToc NewToc;
+		Ar->Seek(SentinelPos);
 
-					FMemoryReaderView Ar(Buffer.GetView());
-					Ar << NewToc;
-					if (!Ar.IsError())
-					{
-						Toc = TIoStatusOr<FOnDemandToc>(MoveTemp(NewToc));
-					}
-					else
-					{
-						ErrorMsg = TEXT("Failed loading on demand TOC from compact binary");
-					}
-				}
-				else
-				{
-					ErrorMsg = FString::Printf(TEXT("Failed fetching TOC, reason '%s'"), *Response.Status().ToString());
-				}
-			});
-
-		while (HttpClient->Tick());
+		FOnDemandTocSentinel Sentinel;
+		*Ar << Sentinel;
 
-		if (Toc.IsOk())
+		if (!Sentinel.IsValid())
 		{
-			return Toc;
+			return FIoStatus(EIoErrorCode::CorruptToc);
 		}
+
+		Ar->Seek(0);
+	}
+
+	FOnDemandToc Toc;
+	*Ar << Toc;
+
+	if (Ar->IsError() || Ar->IsCriticalError())
+	{
+		FIoStatus Status = FIoStatusBuilder(EIoErrorCode::FileNotOpen) << TEXT("Failed to serialize TOC file");
+		return Status;
 	}
 
-	UE_LOG(LogIas, Error, TEXT("%s"), *ErrorMsg);
+	return Toc; 
+}
+
+////////////////////////////////////////////////////////////////////////////////
+TIoStatusOr<FOnDemandToc> FOnDemandToc::LoadFromUrl(FAnsiStringView Url, uint32 RetryCount, bool bFollowRedirects)
+{
+	const EHttpRedirects Redirects = bFollowRedirects ? EHttpRedirects::Follow : EHttpRedirects::Disabled;
+	TIoStatusOr<FIoBuffer> Response = FHttpClient::Get(Url, RetryCount, Redirects); 
 
-	return TIoStatusOr<FOnDemandToc>(FIoStatus(EIoErrorCode::NotFound));
+	if (!Response.IsOk())
+	{
+		FIoStatus Status = FIoStatusBuilder(EIoErrorCode::ReadError) << TEXT("Failed to fetch TOC from URL");
+		return Status;
+	}
+
+	FMemoryReaderView Ar(Response.ValueOrDie().GetView());
+	FOnDemandToc Toc;
+	Ar << Toc;
+
+	if (Ar.IsError() || Ar.IsCriticalError())
+	{
+		FIoStatus Status = FIoStatusBuilder(EIoErrorCode::ReadError) << TEXT("Failed to serialize TOC from HTTP response");
+		return Status;
+	}
+
+	return Toc; 
+}
+
+////////////////////////////////////////////////////////////////////////////////
+TIoStatusOr<FOnDemandToc> FOnDemandToc::LoadFromUrl(FStringView Url, uint32 RetryCount, bool bFollowRedirects)
+{
+	auto AnsiUrl = StringCast<ANSICHAR>(Url.GetData(), Url.Len());
+	return LoadFromUrl(AnsiUrl, RetryCount, bFollowRedirects);
 }
 
 ////////////////////////////////////////////////////////////////////////////////
diff --git a/Engine/Source/Runtime/Experimental/IoStore/OnDemand/Private/OnDemandIoStore.cpp b/Engine/Source/Runtime/Experimental/IoStore/OnDemand/Private/OnDemandIoStore.cpp
--- a/Engine/Source/Runtime/Experimental/IoStore/OnDemand/Private/OnDemandIoStore.cpp
+++ b/Engine/Source/Runtime/Experimental/IoStore/OnDemand/Private/OnDemandIoStore.cpp
@@ -12,7 +12,6 @@
 #include "Async/UniqueLock.h"
 #include "Containers/StringConv.h"
 #include "HAL/FileManager.h"
-#include "HAL/FileManagerGeneric.h"
 #include "HAL/PlatformFile.h"
 #include "HAL/PlatformFileManager.h"
 #include "HAL/PlatformTime.h"
@@ -66,34 +65,6 @@ static void SplitHostUrl(const FStringView& Url, FStringView& OutHost, FStringVi
 	OutRemainder = Url.RightChop(OutHost.Len());
 }
 
-/**
- * Utility to create a FArchive capable of reading from disk using the exact same pathing
- * rules as FPlatformMisc::LoadTextFileFromPlatformPackage but without forcing the entire
- * file to be loaded at once.
- */
-static TUniquePtr<FArchive> CreateReaderFromPlatformPackage(const FString& RelPath)
-{
-	TRACE_CPUPROFILER_EVENT_SCOPE(IasBackend::CreateReaderFromPlatformPackage);
-
-	const FString AbsPath = FPaths::Combine(FGenericPlatformMisc::RootDir(), RelPath);
-
-	IFileHandle* File = IPlatformFile::GetPlatformPhysical().OpenRead(*AbsPath);
-	if (File)
-	{
-#if PLATFORM_ANDROID
-		// This is a handle to an asset so we need to call Seek(0) to move the internal
-		// offset to the start of the asset file.
-		File->Seek(0);
-#endif //PLATFORM_ANDROID
-		const uint32 ReadBufferSize = 256 * 1024;
-		return MakeUnique<FArchiveFileReaderGeneric>(File, *AbsPath, File->Size(), ReadBufferSize);
-	}
-	else
-	{
-		return TUniquePtr<FArchive>();
-	}	
-}
-
 } // namespace UE::IoStore::Private
 
 ///////////////////////////////////////////////////////////////////////////////
@@ -524,85 +495,31 @@ FIoStatus FOnDemandIoStore::TickMountRequest(FMountRequest& MountRequest)
 	{
 		UE_LOG(LogIoStoreOnDemand, Log, TEXT("Loading TOC from file '%s'"), *Args.FilePath);
 
-		TUniquePtr<FArchive> Ar;
-		if (FPlatformMisc::FileExistsInPlatformPackage(Args.FilePath))
+		// TODO: Enable validation when the sentinal is included in all serialization paths
+		const bool bValidate = false;
+		TIoStatusOr<FOnDemandToc> TocStatus = FOnDemandToc::LoadFromFile(Args.FilePath, bValidate);
+		if (!TocStatus.IsOk())
 		{
-			Ar = Private::CreateReaderFromPlatformPackage(Args.FilePath);
-		}
-		else
-		{
-			Ar.Reset(IFileManager::Get().CreateFileReader(*Args.FilePath));
-		}
-
-		if (Ar.IsValid() == false)
-		{
-			return FIoStatusBuilder(EIoErrorCode::FileNotOpen)
-				<< TEXT("Failed to open '")
-				<< Args.FilePath
-				<< TEXT("'");
-		}
-
-		// TODO: Enable validation when the sentinal is included in all serialization paths 
-		//const bool bIncludeDot = true;	
-		const bool bValidate = false; //FPathViews::GetExtension(Args.FilePath, bIncludeDot) == GIasOnDemandTocExt;
-		if (bValidate)
-		{
-			const int64 SentinelPos = Ar->TotalSize() - FOnDemandTocSentinel::SentinelSize;
-
-			if (SentinelPos < 0)
-			{
-				UE_LOG(LogIoStoreOnDemand, Error, TEXT("The file '%s' is smaller than expected and quite possible corrupted"), *Args.FilePath);
-				return EIoErrorCode::CorruptToc;
-			}
-
-			Ar->Seek(SentinelPos);
-
-			FOnDemandTocSentinel Sentinel;
-			(*Ar) << Sentinel;
-
-			if (!Sentinel.IsValid())
-			{
-				UE_LOG(LogIas, Error, TEXT("File corruption detected when serializing '%s'"), *Args.FilePath);
-				return EIoErrorCode::CorruptToc;
-			}
-
-			Ar->Seek(0);
-		}
-
-		FOnDemandToc Toc;
-		*Ar << Toc;
-
-		if (Ar->IsError() || Ar->IsCriticalError())
-		{
-			return FIoStatusBuilder(EIoErrorCode::FileNotOpen) << TEXT("Failed to serialize TOC file");
+			return TocStatus.Status();
 		}
 
+		FOnDemandToc Toc = TocStatus.ConsumeValueOrDie();
 		CreateContainersFromToc(Args.MountId, TocPath, Toc, MountRequest.Containers);
 	}
 	else if (Args.Url.IsEmpty() == false)
 	{
 		UE_LOG(LogIoStoreOnDemand, Log, TEXT("Loading TOC from URL '%s'"), *Args.Url);
 
-		if (Host.IsEmpty())
-		{
-			return FIoStatusBuilder(EIoErrorCode::InvalidParameter) << TEXT("No valid host URL");
-		}
-
-		TIoStatusOr<FIoBuffer> Response = FHttpClient::Get(Args.Url, 2, EHttpRedirects::Follow);
-		if (Response.IsOk() == false)
-		{
-			return FIoStatusBuilder(EIoErrorCode::InvalidCode) << TEXT("Failed to fetch TOC from URL");
-		}
-
-		FOnDemandToc Toc;
-		FMemoryReaderView Ar(Response.ValueOrDie().GetView());
-		Ar << Toc;
+		const uint32 RetryCount				= 2;
+		const bool bFollowRedirects			= true;
+		TIoStatusOr<FOnDemandToc> TocStatus	= FOnDemandToc::LoadFromUrl(Args.Url, 2, bFollowRedirects);
 
-		if (Ar.IsError() || Ar.IsCriticalError())
+		if (!TocStatus.IsOk())
 		{
-			return FIoStatusBuilder(EIoErrorCode::FileNotOpen) << TEXT("Failed to serialize TOC from HTTP response");
+			return TocStatus.Status();
 		}
 
+		FOnDemandToc Toc = TocStatus.ConsumeValueOrDie();
 		CreateContainersFromToc(Args.MountId, TocPath, Toc, MountRequest.Containers);
 	}
 
diff --git a/Engine/Source/Runtime/Experimental/IoStore/OnDemand/Public/IO/IoStoreOnDemand.h b/Engine/Source/Runtime/Experimental/IoStore/OnDemand/Public/IO/IoStoreOnDemand.h
--- a/Engine/Source/Runtime/Experimental/IoStore/OnDemand/Public/IO/IoStoreOnDemand.h
+++ b/Engine/Source/Runtime/Experimental/IoStore/OnDemand/Public/IO/IoStoreOnDemand.h
@@ -182,12 +182,14 @@ struct FOnDemandToc
 	UE_API friend FCbWriter& operator<<(FCbWriter& Writer, const FOnDemandToc& Toc);
 
 	static FGuid VersionGuid;
+
+	static TIoStatusOr<FOnDemandToc> LoadFromFile(const FString& FilePath, bool bValidate);
+	static TIoStatusOr<FOnDemandToc> LoadFromUrl(FAnsiStringView Url, uint32 RetryCount = 0, bool bFollowRedirects = false);
+	static TIoStatusOr<FOnDemandToc> LoadFromUrl(FStringView Url, uint32 RetryCount = 0, bool bFollowRedirects = false);
 };
 
 UE_API bool LoadFromCompactBinary(FCbFieldView Field, FOnDemandToc& OutToc);
 
-TIoStatusOr<FOnDemandToc> LoadTocFromUrl(const FString& ServiceURL, const FString& TocPath, int32 RetryCount);
-
 #if (IS_PROGRAM || WITH_EDITOR)
 
 ////////////////////////////////////////////////////////////////////////////////
