Iris - Fix proper cleanup of ReplicationState-data and InstanceProtocols

diff --git a/Engine/Source/Runtime/Experimental/Iris/Core/Private/Iris/ReplicationSystem/DeltaCompression/DeltaCompressionBaselineManager.cpp b/Engine/Source/Runtime/Experimental/Iris/Core/Private/Iris/ReplicationSystem/DeltaCompression/DeltaCompressionBaselineManager.cpp
--- a/Engine/Source/Runtime/Experimental/Iris/Core/Private/Iris/ReplicationSystem/DeltaCompression/DeltaCompressionBaselineManager.cpp
+++ b/Engine/Source/Runtime/Experimental/Iris/Core/Private/Iris/ReplicationSystem/DeltaCompression/DeltaCompressionBaselineManager.cpp
@@ -60,7 +60,6 @@ FDeltaCompressionBaselineManager::FDeltaCompressionBaselineManager()
 
 FDeltaCompressionBaselineManager::~FDeltaCompressionBaselineManager()
 {
-	Deinit();
 }
 
 void FDeltaCompressionBaselineManager::Init(FDeltaCompressionBaselineManagerInitParams& InitParams)
@@ -102,6 +101,7 @@ void FDeltaCompressionBaselineManager::Deinit()
 {
 	// The BaselineStateManager will release all baselines in its destructor.
 	FreeAllPerObjectInfos();
+	BaselineStorage.Deinit();
 }
 
 void FDeltaCompressionBaselineManager::PreSendUpdate(FDeltaCompressionBaselineManagerPreSendUpdateParams& UpdateParams)
diff --git a/Engine/Source/Runtime/Experimental/Iris/Core/Private/Iris/ReplicationSystem/DeltaCompression/DeltaCompressionBaselineStorage.cpp b/Engine/Source/Runtime/Experimental/Iris/Core/Private/Iris/ReplicationSystem/DeltaCompression/DeltaCompressionBaselineStorage.cpp
--- a/Engine/Source/Runtime/Experimental/Iris/Core/Private/Iris/ReplicationSystem/DeltaCompression/DeltaCompressionBaselineStorage.cpp
+++ b/Engine/Source/Runtime/Experimental/Iris/Core/Private/Iris/ReplicationSystem/DeltaCompression/DeltaCompressionBaselineStorage.cpp
@@ -22,7 +22,6 @@ FDeltaCompressionBaselineStorage::FDeltaCompressionBaselineStorage()
 
 FDeltaCompressionBaselineStorage::~FDeltaCompressionBaselineStorage()
 {
-	Deinit();
 }
 
 void FDeltaCompressionBaselineStorage::Init(FDeltaCompressionBaselineStorageInitParams& InitParams)
diff --git a/Engine/Source/Runtime/Experimental/Iris/Core/Private/Iris/ReplicationSystem/NetRefHandleManager.cpp b/Engine/Source/Runtime/Experimental/Iris/Core/Private/Iris/ReplicationSystem/NetRefHandleManager.cpp
--- a/Engine/Source/Runtime/Experimental/Iris/Core/Private/Iris/ReplicationSystem/NetRefHandleManager.cpp
+++ b/Engine/Source/Runtime/Experimental/Iris/Core/Private/Iris/ReplicationSystem/NetRefHandleManager.cpp
@@ -82,7 +82,11 @@ void FNetRefHandleManager::Init(const FInitParams& InitParams)
 
 void FNetRefHandleManager::Deinit()
 {
-
+	// We can clear out the reserved index for invalid handles now.
+	AssignedInternalIndices.ClearBit(0);
+	AssignedInternalIndices.ForAllSetBits([this](uint32 InternalIndex) { 
+		this->InternalDestroyNetObject(InternalIndex); 
+	});
 }
 
 void FNetRefHandleManager::InitNetBitArray(FNetBitArray* NetBitArray)
diff --git a/Engine/Source/Runtime/Experimental/Iris/Core/Private/Iris/ReplicationSystem/ReplicationBridge.cpp b/Engine/Source/Runtime/Experimental/Iris/Core/Private/Iris/ReplicationSystem/ReplicationBridge.cpp
--- a/Engine/Source/Runtime/Experimental/Iris/Core/Private/Iris/ReplicationSystem/ReplicationBridge.cpp
+++ b/Engine/Source/Runtime/Experimental/Iris/Core/Private/Iris/ReplicationSystem/ReplicationBridge.cpp
@@ -347,6 +347,34 @@ void UReplicationBridge::Initialize(UReplicationSystem* InReplicationSystem)
 
 void UReplicationBridge::Deinitialize()
 {
+	using namespace UE::Net;
+	using namespace UE::Net::Private;
+
+	// Detach all replicated instances that have not yet been destroyed as part of shutting down the rest of the game.
+	NetRefHandleManager->GetAssignedInternalIndices().ForAllSetBits([this](uint32 InternalObjectIndex) 
+	{
+		if (InternalObjectIndex == FNetRefHandleManager::InvalidInternalIndex)
+		{
+			return;
+		}
+
+		FNetRefHandleManager::FReplicatedObjectData& ObjectData = NetRefHandleManager->GetReplicatedObjectDataNoCheck(InternalObjectIndex);
+		if (ObjectData.InstanceProtocol)
+		{
+			const FNetRefHandle RefHandle = ObjectData.RefHandle;
+
+			// Clear out, tracking data or not?, currently we only have a single replication-system so it should not be a big issue if we do.
+			// Note: Currently we opted to leave it as is.
+			// -- This only applies to server. If it is a restart of ReplicationSystem, the same actors will be re-registered and global handle will be destroyed later, otherwise it does not matter what we do.
+			// DestroyGlobalNetHandle(InternalObjectIndex);
+			// ClearNetPushIds(InternalObjectIndex);
+
+			// Detach and destroy instance protocol
+			ObjectData.bPendingEndReplication = 1U;
+			InternalDetachInstanceFromNetRefHandle(ObjectData.RefHandle);
+		}
+	});
+
 	// Just set the protocol to null, it will be destroyed with the protocolmanager
 	DestructionInfoProtocol = nullptr;
 
diff --git a/Engine/Source/Runtime/Experimental/Iris/Core/Private/Iris/ReplicationSystem/ReplicationSystem.cpp b/Engine/Source/Runtime/Experimental/Iris/Core/Private/Iris/ReplicationSystem/ReplicationSystem.cpp
--- a/Engine/Source/Runtime/Experimental/Iris/Core/Private/Iris/ReplicationSystem/ReplicationSystem.cpp
+++ b/Engine/Source/Runtime/Experimental/Iris/Core/Private/Iris/ReplicationSystem/ReplicationSystem.cpp
@@ -273,6 +273,7 @@ public:
 			ReplicationSystemInternal.SetIrisObjectReferencePackageMap(static_cast<UIrisObjectReferencePackageMap*>(nullptr));
 		}
 
+		ReplicationSystemInternal.GetDeltaCompressionBaselineManager().Deinit();
 		ReplicationSystemInternal.GetNetRefHandleManager().Deinit();
 	}
 
