Fix broken chaos cache replay
- symptoms : Cache would fail to set the geometry collection to be kinematic when replaying , causing also the pose to be reset regardless of the start time
- fix : make sure to note recreate the physics state if it was already created, if it was , make sure to recreate the physics proxy as well

diff --git a/Engine/Source/Runtime/Experimental/GeometryCollectionEngine/Private/GeometryCollection/GeometryCollectionComponent.cpp b/Engine/Source/Runtime/Experimental/GeometryCollectionEngine/Private/GeometryCollection/GeometryCollectionComponent.cpp
--- a/Engine/Source/Runtime/Experimental/GeometryCollectionEngine/Private/GeometryCollection/GeometryCollectionComponent.cpp
+++ b/Engine/Source/Runtime/Experimental/GeometryCollectionEngine/Private/GeometryCollection/GeometryCollectionComponent.cpp
@@ -1262,7 +1262,11 @@ void UGeometryCollectionComponent::SetSimulatePhysics(bool bEnabled)
 
 	if (bEnabled && !PhysicsProxy && RestCollection)
 	{
-		RecreatePhysicsState();
+		EnsurePhysicsStateCreated();
+		if (!PhysicsProxy)
+		{
+			RegisterAndInitializePhysicsProxy();
+		}
 	}
 }
 
