Fix for a replication descriptors memory leak.

diff --git a/Engine/Source/Runtime/Experimental/Iris/Core/Private/Iris/Serialization/InternalNetSerializers.cpp b/Engine/Source/Runtime/Experimental/Iris/Core/Private/Iris/Serialization/InternalNetSerializers.cpp
--- a/Engine/Source/Runtime/Experimental/Iris/Core/Private/Iris/Serialization/InternalNetSerializers.cpp
+++ b/Engine/Source/Runtime/Experimental/Iris/Core/Private/Iris/Serialization/InternalNetSerializers.cpp
@@ -3,5 +3,10 @@
 #include "Iris/Serialization/InternalNetSerializers.h"
 #include "Iris/ReplicationState/ReplicationStateDescriptor.h"
 
-FArrayPropertyNetSerializerConfig::FArrayPropertyNetSerializerConfig() = default;
+FArrayPropertyNetSerializerConfig::FArrayPropertyNetSerializerConfig()
+: FNetSerializerConfig()
+{
+	ConfigTraits = ENetSerializerConfigTraits::NeedDestruction;
+}
+
 FArrayPropertyNetSerializerConfig::~FArrayPropertyNetSerializerConfig() = default;
