Fix partial destruction

diff --git a/Engine/Source/Runtime/Experimental/Chaos/Private/Chaos/ClusterUnionManager.cpp b/Engine/Source/Runtime/Experimental/Chaos/Private/Chaos/ClusterUnionManager.cpp
--- a/Engine/Source/Runtime/Experimental/Chaos/Private/Chaos/ClusterUnionManager.cpp
+++ b/Engine/Source/Runtime/Experimental/Chaos/Private/Chaos/ClusterUnionManager.cpp
@@ -32,6 +32,14 @@ namespace Chaos
 			TEXT("Gating a risky bug fix.")
 		);
 
+		// @tmp: To be removed
+		bool bFixPartialDestruction = true;
+		FAutoConsoleVariableRef CVarChaosFixPartialDestruction(
+			TEXT("p.Chaos.ClusterUnion.FixPartialDestruction"),
+			bFixPartialDestruction,
+			TEXT("Fix the partial destruction.")
+		);
+
 		FRigidTransform3 GetParticleRigidFrameInClusterUnion(FPBDRigidParticleHandle* Child, const FRigidTransform3& ClusterWorldTM)
 		{
 			FRigidTransform3 Frame = FRigidTransform3::Identity;
@@ -1085,7 +1093,7 @@ namespace Chaos
 		{
 			// Only generate intercluster edges for main particles. Auxiliary particles that are just bits and pieces of geometry collections
 			// shouldn't also generate intercluster edges.
-			if (!Properties || !Properties->bIsAuxiliaryParticle)
+			if (bFixPartialDestruction || !Properties || !Properties->bIsAuxiliaryParticle)
 			{
 				const TArray<FPBDRigidParticleHandle*>& ParticleChildren = MClustering.GetChildrenMap().FindRef(Particle->CastToClustered());
 				const TArray<FPBDRigidParticleHandle*>& OtherChildren = MClustering.GetChildrenMap().FindRef(OtherParticle->CastToClustered());
