Chaos Debug Draw
- add debug draw subsystem
- add DD timeline to game thread
- move all new DD stuff out of ChaosDebugDrawComponent (it can be deleted later)

diff --git a/Engine/Source/Runtime/Experimental/Chaos/Private/Chaos/PBDRigidsEvolutionGBF.cpp b/Engine/Source/Runtime/Experimental/Chaos/Private/Chaos/PBDRigidsEvolutionGBF.cpp
--- a/Engine/Source/Runtime/Experimental/Chaos/Private/Chaos/PBDRigidsEvolutionGBF.cpp
+++ b/Engine/Source/Runtime/Experimental/Chaos/Private/Chaos/PBDRigidsEvolutionGBF.cpp
@@ -382,7 +382,7 @@ void FPBDRigidsEvolutionGBF::SetDebugDrawScene(const ChaosDD::Private::FChaosDDS
 
 	if (CDDScene.IsValid())
 	{
-		CDDTickTimeline = CDDScene->CreateTimeline(FString::Format(TEXT("{0} {1}"), { CDDScene->GetName(), "Tick" }));
+		CDDTickTimeline = CDDScene->CreateTimeline(FString::Format(TEXT("{0} {1}"), { CDDScene->GetName(), "Physics Tick" }));
 	}
 }
 #endif
diff --git a/Engine/Source/Runtime/Experimental/Chaos/Private/ChaosDebugDraw/ChaosDDContext.cpp b/Engine/Source/Runtime/Experimental/Chaos/Private/ChaosDebugDraw/ChaosDDContext.cpp
--- a/Engine/Source/Runtime/Experimental/Chaos/Private/ChaosDebugDraw/ChaosDDContext.cpp
+++ b/Engine/Source/Runtime/Experimental/Chaos/Private/ChaosDebugDraw/ChaosDDContext.cpp
@@ -22,61 +22,91 @@ namespace ChaosDD::Private
 
 	//
 	//
+	// Timeline Context
 	//
 	//
 
-	FChaosDDScopeTimelineContext::FChaosDDScopeTimelineContext(const FChaosDDTimelinePtr& InTimeline, double InTime, double InDt)
+	void FChaosDDTimelineContext::BeginFrame(const FChaosDDTimelinePtr& InTimeline, double InTime, double InDt)
 	{
+		FChaosDDContext& Context = FChaosDDContext::Get();
+		ParentFrame = Context.Frame;
+
 		if (InTimeline.IsValid())
 		{
 			Timeline = InTimeline;
 			Timeline->BeginFrame(InTime, InDt);
-
-			FChaosDDContext& Context = FChaosDDContext::Get();
-			ParentFrame = Context.Frame;
 			Context.Frame = Timeline->GetActiveFrame();
 		}
+		else
+		{
+			Context.Frame.Reset();
+		}
 	}
 
-	FChaosDDScopeTimelineContext::~FChaosDDScopeTimelineContext()
+	void FChaosDDTimelineContext::EndFrame()
 	{
 		if (Timeline.IsValid())
 		{
 			Timeline->EndFrame();
 			Timeline.Reset();
-
-			FChaosDDContext& Context = FChaosDDContext::Get();
-			Context.Frame = ParentFrame;
-			ParentFrame.Reset();
 		}
+
+		FChaosDDContext& Context = FChaosDDContext::Get();
+		Context.Frame = ParentFrame;
+		ParentFrame.Reset();
+	}
+
+	FChaosDDScopeTimelineContext::FChaosDDScopeTimelineContext(const FChaosDDTimelinePtr& InTimeline, double InTime, double InDt)
+	{
+		Context.BeginFrame(InTimeline, InTime, InDt);
+	}
+
+	FChaosDDScopeTimelineContext::~FChaosDDScopeTimelineContext()
+	{
+		Context.EndFrame();
 	}
 
 	//
 	//
+	// Task Context
 	//
 	//
 
-	FChaosDDScopeTaskContext::FChaosDDScopeTaskContext(const FChaosDDContext& InParentDDContext)
+	void FChaosDDTaskContext::BeginThread(const FChaosDDContext& InParentDDContext)
 	{
+		FChaosDDContext& Context = FChaosDDContext::Get();
+		ParentFrame = Context.Frame;
+
 		if (InParentDDContext.Frame.IsValid())
 		{
-			FChaosDDContext& Context = FChaosDDContext::Get();
-
-			ParentFrame = Context.Frame;
 			Context.Frame = InParentDDContext.Frame;
 		}
+		else
+		{
+			Context.Frame.Reset();
+		}
 	}
 
-	FChaosDDScopeTaskContext::~FChaosDDScopeTaskContext()
+	void FChaosDDTaskContext::EndThread()
 	{
 		FChaosDDContext& Context = FChaosDDContext::Get();
-
 		Context.Frame = ParentFrame;
 		ParentFrame.Reset();
 	}
 
+	FChaosDDScopeTaskContext::FChaosDDScopeTaskContext(const FChaosDDContext& InParentDDContext)
+	{
+		Context.BeginThread(InParentDDContext);
+	}
+
+	FChaosDDScopeTaskContext::~FChaosDDScopeTaskContext()
+	{
+		Context.EndThread();
+	}
+
 	//
 	//
+	// Thread Local Context
 	//
 	//
 
diff --git a/Engine/Source/Runtime/Experimental/Chaos/Private/ChaosDebugDraw/ChaosDDScene.cpp b/Engine/Source/Runtime/Experimental/Chaos/Private/ChaosDebugDraw/ChaosDDScene.cpp
--- a/Engine/Source/Runtime/Experimental/Chaos/Private/ChaosDebugDraw/ChaosDDScene.cpp
+++ b/Engine/Source/Runtime/Experimental/Chaos/Private/ChaosDebugDraw/ChaosDDScene.cpp
@@ -22,10 +22,12 @@ namespace ChaosDD::Private
 		}
 	}
 
-	FChaosDDScene::FChaosDDScene(const FString& InName)
+	FChaosDDScene::FChaosDDScene(const FString& InName, bool bInIsServer)
 		: Name(InName)
 		, DrawRegion(FVector::Zero(), 0.0)
 		, CommandBudget(20000)
+		, bIsServer(bInIsServer)
+		, bRenderEnabled(true)
 	{
 	}
 
@@ -33,6 +35,25 @@ namespace ChaosDD::Private
 	{
 	}
 
+	bool FChaosDDScene::IsServer() const
+	{
+		return bIsServer;
+	}
+
+	void FChaosDDScene::SetRenderEnabled(bool bInRenderEnabled)
+	{
+		FScopeLock Lock(&TimelinesCS);
+
+		bRenderEnabled = bInRenderEnabled;
+	}
+
+	bool FChaosDDScene::IsRenderEnabled() const
+	{
+		FScopeLock Lock(&TimelinesCS);
+
+		return bRenderEnabled;
+	}
+
 	void FChaosDDScene::SetDrawRegion(const FSphere3d& InDrawRegion)
 	{
 		FScopeLock Lock(&TimelinesCS);
@@ -50,6 +71,13 @@ namespace ChaosDD::Private
 		//FChaosDDContext::SetGlobalDrawRegion(InDrawRegion);
 	}
 
+	const FSphere3d& FChaosDDScene::GetDrawRegion() const
+	{
+		FScopeLock Lock(&TimelinesCS);
+
+		return DrawRegion;
+	}
+
 	void FChaosDDScene::SetCommandBudget(int32 InCommandBudget)
 	{
 		FScopeLock Lock(&TimelinesCS);
@@ -65,6 +93,13 @@ namespace ChaosDD::Private
 		FChaosDDContext::SetGlobalCommandBudget(InCommandBudget);
 	}
 
+	int32 FChaosDDScene::GetCommandBudget() const
+	{
+		FScopeLock Lock(&TimelinesCS);
+
+		return CommandBudget;
+	}
+
 	FChaosDDTimelinePtr FChaosDDScene::CreateTimeline(const FString& InName)
 	{
 		FScopeLock Lock(&TimelinesCS);
@@ -76,36 +111,9 @@ namespace ChaosDD::Private
 		return Timeline;
 	}
 
-	void FChaosDDScene::RenderLatestFrames(IChaosDDRenderer& Renderer, bool bIncludeGlobalFrame)
+	TArray<FChaosDDFramePtr> FChaosDDScene::GetLatestFrames()
 	{
-		TArray<FChaosDDFramePtr> Frames = GetFrames();
-
-		if (bIncludeGlobalFrame)
-		{
-			Frames.Add(FChaosDDContext::ExtractGlobalFrame());
-		}
-
-		for (const FChaosDDFramePtr& Frame : Frames)
-		{
-			if (Frame.IsValid())
-			{
-				UE_LOG(LogChaosDD, VeryVerbose, TEXT("Render %s %d %d Commands"), *Frame->GetTimeline()->GetName(), Frame->GetFrameIndex(), Frame->GetNumCommands());
-
-				// Render the legacy commands
-				Frame->VisitLatentCommands(
-					[&Renderer](const Chaos::FLatentDrawCommand& Command)
-					{
-						Renderer.RenderLatentCommand(Command);
-					});
-
-				// Render the commands
-				Frame->VisitCommands(
-					[&Renderer](const FChaosDDCommand& Command)
-					{
-						Command(Renderer);
-					});
-			}
-		}
+		return GetFrames();
 	}
 
 	TArray<FChaosDDFramePtr> FChaosDDScene::GetFrames()
diff --git a/Engine/Source/Runtime/Experimental/Chaos/Private/PBDRigidsSolver.cpp b/Engine/Source/Runtime/Experimental/Chaos/Private/PBDRigidsSolver.cpp
--- a/Engine/Source/Runtime/Experimental/Chaos/Private/PBDRigidsSolver.cpp
+++ b/Engine/Source/Runtime/Experimental/Chaos/Private/PBDRigidsSolver.cpp
@@ -1338,7 +1338,7 @@ namespace Chaos
 
 		if (CDDScene.IsValid())
 		{
-			CDDFrameTimeline = CDDScene->CreateTimeline(FString::Format(TEXT("{0} {1}"), { CDDScene->GetName(), "Frame" }));
+			CDDFrameTimeline = CDDScene->CreateTimeline(FString::Format(TEXT("{0} {1}"), { CDDScene->GetName(), "Physics Frame" }));
 		}
 
 		GetEvolution()->SetDebugDrawScene(InCDDScene);
@@ -2467,6 +2467,13 @@ TRACE_COUNTER_SET(ChaosTraceCounter_##Name, Value)
 #if CHAOS_DEBUG_DRAW
 		QUICK_SCOPE_CYCLE_COUNTER(SolverDebugDraw);
 
+		if (CDDScene.IsValid())
+		{
+			// @todo(chaos): this logic should be in UChaosDebugDrawSubsystem but the CVars are here. Move them and this...
+			const bool bRenderEnabled = (CDDScene->IsServer() && ChaosSolverDebugDrawShowServer) || (!CDDScene->IsServer() && ChaosSolverDebugDrawShowClient);
+			CDDScene->SetRenderEnabled(bRenderEnabled);
+		}
+
 		const bool bIsServer = GetDebugName().ToString().StartsWith(TEXT("Server"));
 		if (bIsServer && !ChaosSolverDebugDrawShowServer)
 		{
diff --git a/Engine/Source/Runtime/Experimental/Chaos/Public/Chaos/DebugDrawQueue.h b/Engine/Source/Runtime/Experimental/Chaos/Public/Chaos/DebugDrawQueue.h
--- a/Engine/Source/Runtime/Experimental/Chaos/Public/Chaos/DebugDrawQueue.h
+++ b/Engine/Source/Runtime/Experimental/Chaos/Public/Chaos/DebugDrawQueue.h
@@ -249,29 +249,29 @@ public:
 		RadiusOfInterest = InRadius;
 	}
 
-	const FVector& GetCenterOfInterest() const
+	FVector GetCenterOfInterest() const
 	{
-		return CenterOfInterest;
+		return GetCenterOfInterestImpl();
 	}
 
 	FReal GetRadiusOfInterest() const
 	{
-		return RadiusOfInterest;
+		return GetRadiusOfInterestImpl();
 	}
 
 	bool IsInRegionOfInterest(FVector Pos) const
 	{
-		return (RadiusOfInterest <= 0.0f) || ((Pos - CenterOfInterest).SizeSquared() < RadiusOfInterest * RadiusOfInterest);
+		return IsInRegionOfInterestImpl(Pos);
 	}
 
 	bool IsInRegionOfInterest(FVector Pos, FReal Radius) const
 	{
-		return (RadiusOfInterest <= 0.0f) || ((Pos - CenterOfInterest).SizeSquared() < (RadiusOfInterest + Radius) * (RadiusOfInterest + Radius));
+		return IsInRegionOfInterestImpl(Pos, Radius);
 	}
 
 	bool IsInRegionOfInterest(FAABB3 Bounds) const
 	{
-		return Bounds.ThickenSymmetrically(FVec3(RadiusOfInterest)).Contains(CenterOfInterest);
+		return IsInRegionOfInterestImpl(Bounds);
 	}
 
 	CHAOS_API void SetConsumerActive(void* Consumer, bool bConsumerActive);
@@ -294,6 +294,60 @@ private:
 	{}
 	~FDebugDrawQueue() {}
 
+	FVector GetCenterOfInterestImpl() const
+	{
+		if (bChaosDebugDraw_UseNewQueue)
+		{
+			ChaosDD::Private::FChaosDDFrameWriter DDWriter = ChaosDD::Private::FChaosDDContext::Get().GetWriter();
+			return DDWriter.GetDrawRegion().Center;
+		}
+
+		return CenterOfInterest;
+	}
+
+	FReal GetRadiusOfInterestImpl() const
+	{
+		if (bChaosDebugDraw_UseNewQueue)
+		{
+			ChaosDD::Private::FChaosDDFrameWriter DDWriter = ChaosDD::Private::FChaosDDContext::Get().GetWriter();
+			return DDWriter.GetDrawRegion().W;
+		}
+
+		return RadiusOfInterest;
+	}
+
+	bool IsInRegionOfInterestImpl(FVector Pos) const
+	{
+		if (bChaosDebugDraw_UseNewQueue)
+		{
+			ChaosDD::Private::FChaosDDFrameWriter DDWriter = ChaosDD::Private::FChaosDDContext::Get().GetWriter();
+			return DDWriter.IsInDrawRegion(Pos);
+		}
+
+		return (RadiusOfInterest <= 0.0f) || ((Pos - CenterOfInterest).SizeSquared() < RadiusOfInterest * RadiusOfInterest);
+	}
+
+	bool IsInRegionOfInterestImpl(FVector Pos, FReal Radius) const
+	{
+		if (bChaosDebugDraw_UseNewQueue)
+		{
+			ChaosDD::Private::FChaosDDFrameWriter DDWriter = ChaosDD::Private::FChaosDDContext::Get().GetWriter();
+			return DDWriter.IsInDrawRegion(FSphere3d(Pos, Radius));
+		}
+
+		return (RadiusOfInterest <= 0.0f) || ((Pos - CenterOfInterest).SizeSquared() < (RadiusOfInterest + Radius) * (RadiusOfInterest + Radius));
+	}
+
+	bool IsInRegionOfInterestImpl(FAABB3 Bounds) const
+	{
+		if (bChaosDebugDraw_UseNewQueue)
+		{
+			ChaosDD::Private::FChaosDDFrameWriter DDWriter = ChaosDD::Private::FChaosDDContext::Get().GetWriter();
+			return DDWriter.IsInDrawRegion(FBox3d(Bounds.Min(), Bounds.Max()));
+		}
+		return Bounds.ThickenSymmetrically(FVec3(RadiusOfInterest)).Contains(CenterOfInterest);
+	}
+
 	bool AcceptCommand(int Cost, const FVec3& Position)
 	{
 		if (bChaosDebugDraw_UseNewQueue)
diff --git a/Engine/Source/Runtime/Experimental/Chaos/Public/ChaosDebugDraw/ChaosDDContext.h b/Engine/Source/Runtime/Experimental/Chaos/Public/ChaosDebugDraw/ChaosDDContext.h
--- a/Engine/Source/Runtime/Experimental/Chaos/Public/ChaosDebugDraw/ChaosDDContext.h
+++ b/Engine/Source/Runtime/Experimental/Chaos/Public/ChaosDebugDraw/ChaosDDContext.h
@@ -25,10 +25,12 @@ namespace ChaosDD::Private
 			return FChaosDDFrameWriter(Get().GetFrame());
 		}
 
+		static FChaosDDFramePtr ExtractGlobalFrame();
+
 	private:
 		friend class FChaosDDScene;
-		friend class FChaosDDScopeTimelineContext;
-		friend class FChaosDDScopeTaskContext;
+		friend class FChaosDDTimelineContext;
+		friend class FChaosDDTaskContext;
 
 		const FChaosDDFramePtr& GetFrame() const
 		{
@@ -48,7 +50,6 @@ namespace ChaosDD::Private
 		// Global frame management
 		static const FChaosDDFramePtr& GetGlobalFrame();
 		static void CreateGlobalFrame();
-		static FChaosDDFramePtr ExtractGlobalFrame();
 		static void SetGlobalDrawRegion(const FSphere3d& InDrawRegion);
 		static void SetGlobalCommandBudget(int32 InCommandBudget);
 
@@ -62,7 +63,21 @@ namespace ChaosDD::Private
 	};
 
 	//
-	// A scoped Debug Draw Context for use on the thread that owns the timeline
+	// Initializes the FChaosDDContext for a thread that owns a timeline
+	//
+	class CHAOS_API FChaosDDTimelineContext
+	{
+	public:
+		void BeginFrame(const FChaosDDTimelinePtr& InTimeline, double InTime, double InDt);
+		void EndFrame();
+
+	private:
+		FChaosDDTimelinePtr Timeline;
+		FChaosDDFramePtr ParentFrame;
+	};
+
+	//
+	// A scoped wrapper for FChaosDDTimelineContext
 	//
 	class CHAOS_API FChaosDDScopeTimelineContext
 	{
@@ -71,14 +86,27 @@ namespace ChaosDD::Private
 		~FChaosDDScopeTimelineContext();
 
 	private:
-		FChaosDDTimelinePtr Timeline;
-		FChaosDDFramePtr ParentFrame;
+		FChaosDDTimelineContext Context;
 	};
 
 	//
-	// A scoped Debug Draw Context for use in tasks and parallel-for etc.
+	// Initializes the FChaosDDContext for a task.
 	// Assumes that the task is kicked off from a thread that has an active debug draw context,
-	// which should be passed into this context.
+	// which should be passed into this context. Any debug draws from the task will go to the 
+	// same frame as the parent context.
+	//
+	class CHAOS_API FChaosDDTaskContext
+	{
+	public:
+		void BeginThread(const FChaosDDContext& InParentDDContext);
+		void EndThread();
+
+	private:
+		FChaosDDFramePtr ParentFrame;
+	};
+
+	//
+	// A scoped wrapper for FChaosDDTaskContext
 	//
 	class CHAOS_API FChaosDDScopeTaskContext
 	{
@@ -87,7 +115,7 @@ namespace ChaosDD::Private
 		~FChaosDDScopeTaskContext();
 
 	private:
-		FChaosDDFramePtr ParentFrame;
+		FChaosDDTaskContext Context;
 	};
 }
 
diff --git a/Engine/Source/Runtime/Experimental/Chaos/Public/ChaosDebugDraw/ChaosDDFrame.h b/Engine/Source/Runtime/Experimental/Chaos/Public/ChaosDebugDraw/ChaosDDFrame.h
--- a/Engine/Source/Runtime/Experimental/Chaos/Public/ChaosDebugDraw/ChaosDDFrame.h
+++ b/Engine/Source/Runtime/Experimental/Chaos/Public/ChaosDebugDraw/ChaosDDFrame.h
@@ -21,7 +21,7 @@ namespace ChaosDD::Private
 	// 
 	// E.g., See FChaosDDLine, FChaosDDSphere, FChaosDDParticle
 	//
-	using FChaosDDCommand = TFunction<void(const IChaosDDRenderer&)>;
+	using FChaosDDCommand = TFunction<void(IChaosDDRenderer&)>;
 
 	//
 	// A single frame of debug draw data
@@ -42,6 +42,7 @@ namespace ChaosDD::Private
 		{
 			if (InCommandQueueLength > 0)
 			{
+				Commands.Reserve(InCommandQueueLength);
 				LatentCommands.Reserve(InCommandQueueLength);
 			}
 		}
@@ -75,6 +76,11 @@ namespace ChaosDD::Private
 			DrawRegion = InRegion;
 		}
 
+		const FSphere3d& GetDrawRegion() const
+		{
+			return DrawRegion;
+		}
+
 		bool IsInDrawRegion(const FVector& InPos) const
 		{
 			if (DrawRegion.W > 0.0)
@@ -108,14 +114,30 @@ namespace ChaosDD::Private
 			CommandBudget = InCommandBudget;
 		}
 
+		int32 GetCommandBudget() const
+		{
+			return CommandBudget;
+		}
+
+		int32 GetCommandCost() const
+		{
+			return CommandCost;
+		}
+
+		bool WasCommandBudgetExceeded() const
+		{
+			return CommandCost > CommandBudget;
+		}
+
 		bool AddToCost(int32 InCost)
 		{
 			FScopeLock Lock(&CommandsCS);
 
+			CommandCost += InCost;
+
 			// A budget of zero means infinite
-			if ((CommandCost + InCost <= CommandBudget) || (CommandBudget == 0))
+			if ((CommandCost <= CommandBudget) || (CommandBudget == 0))
 			{
-				CommandCost += InCost;
 				return true;
 			}
 
@@ -140,9 +162,15 @@ namespace ChaosDD::Private
 		{
 			FScopeLock Lock(&CommandsCS);
 
-			return LatentCommands.Num();
+			return Commands.Num();
 		}
 
+		int32 GetNumLatentCommands() const
+		{
+			FScopeLock Lock(&CommandsCS);
+
+			return LatentCommands.Num();
+		}
 
 		// VisitorType = void(const FChaosDDCommand& Command)
 		template<typename VisitorType>
@@ -291,6 +319,15 @@ namespace ChaosDD::Private
 			}
 		}
 
+		FSphere3d GetDrawRegion() const
+		{
+			if (Frame.IsValid())
+			{
+				return Frame->GetDrawRegion();
+			}
+			return FSphere3d(FVector3d(0), 0);
+		}
+
 		bool IsInDrawRegion(const FVector& InPos) const
 		{
 			if (Frame.IsValid())
diff --git a/Engine/Source/Runtime/Experimental/Chaos/Public/ChaosDebugDraw/ChaosDDRenderer.h b/Engine/Source/Runtime/Experimental/Chaos/Public/ChaosDebugDraw/ChaosDDRenderer.h
--- a/Engine/Source/Runtime/Experimental/Chaos/Public/ChaosDebugDraw/ChaosDDRenderer.h
+++ b/Engine/Source/Runtime/Experimental/Chaos/Public/ChaosDebugDraw/ChaosDDRenderer.h
@@ -25,12 +25,18 @@ namespace ChaosDD::Private
 		// Are we rendering a Server scene?
 		virtual bool IsServer() const = 0;
 
+		// The region of interest
+		virtual FSphere3d GetDrawRegion() const = 0;
+
 		// Utility functions for use by Debug Draw commands (e.g., FChaosDDLine)
-		virtual void RenderLine(const FVector3d& A, const FVector3d& B, const FColor& Color, float LineThickness, float Lifetime) const = 0;
-		virtual void RenderBox(const FVector3d& Position, const FQuat4d& Rotation, const FVector3d& Size, const FColor& Color, float LineThickness, float Lifetime) const = 0;
+		virtual void RenderLine(const FVector3d& A, const FVector3d& B, const FColor& Color, float LineThickness, float Lifetime) = 0;
+		virtual void RenderSphere(const FVector3d& Center, float Radius, const FColor& Color, float LineThickness, float Lifetime) = 0;
+		virtual void RenderCapsule(const FVector3d& Center, const FQuat4d& Rotation, float HalfHeight, float Radius, const FColor& Color, float LineThickness, float Lifetime) = 0;
+		virtual void RenderBox(const FVector3d& Position, const FQuat4d& Rotation, const FVector3d& Size, const FColor& Color, float LineThickness, float Lifetime) = 0;
+		virtual void RenderTriangle(const FVector3d& A, const FVector3d& B, const FVector3d& C, const FColor& Color, float LineThickness, float Lifetime) = 0;
 
 		// Render legacy debug draw command (See FChaosDDScene::RenderLatestFrames)
-		virtual void RenderLatentCommand(const Chaos::FLatentDrawCommand& Command) const = 0;
+		virtual void RenderLatentCommand(const Chaos::FLatentDrawCommand& Command) = 0;
 	};
 }
 
diff --git a/Engine/Source/Runtime/Experimental/Chaos/Public/ChaosDebugDraw/ChaosDDScene.h b/Engine/Source/Runtime/Experimental/Chaos/Public/ChaosDebugDraw/ChaosDDScene.h
--- a/Engine/Source/Runtime/Experimental/Chaos/Public/ChaosDebugDraw/ChaosDDScene.h
+++ b/Engine/Source/Runtime/Experimental/Chaos/Public/ChaosDebugDraw/ChaosDDScene.h
@@ -20,7 +20,7 @@ namespace ChaosDD::Private
 	class CHAOS_API FChaosDDScene : public TSharedFromThis<FChaosDDScene>
 	{
 	public:
-		FChaosDDScene(const FString& InName);
+		FChaosDDScene(const FString& InName, bool bIsServer);
 		~FChaosDDScene();
 
 		const FString& GetName() const
@@ -28,28 +28,41 @@ namespace ChaosDD::Private
 			return Name;
 		}
 
+		bool IsServer() const;
+		void SetRenderEnabled(bool bInRenderEnabled);
+		bool IsRenderEnabled() const;
+
 		// Specify the region of in which we wish to enable debug draw. A radius of zero means everywhere.
 		void SetDrawRegion(const FSphere3d& InDrawRegion);
 
+		// The region of interest
+		const FSphere3d& GetDrawRegion() const;
+
 		// Set the line budget for debug draw
 		void SetCommandBudget(int32 InCommandBudget);
 
+		// The number of commands we can draw (also max number of lines for now)
+		int32 GetCommandBudget() const;
+
 		// Create a new timeline. E.g., PT, GT, RBAN
 		// The caller must hold a shared pointer to the timeline to keep it alive.
 		FChaosDDTimelinePtr CreateTimeline(const FString& Name);
 
-		// Render the most recent frame from each timeline
-		void RenderLatestFrames(IChaosDDRenderer& Renderer, bool bIncludeGlobalFrame);
+		// Collect all the latest complete frames for rendering
+		TArray<FChaosDDFramePtr> GetLatestFrames();
 
 	private:
 		TArray<FChaosDDFramePtr> GetFrames();
 		void PruneTimelines();
 
+		mutable FCriticalSection TimelinesCS;
+
 		FString Name;
-		FCriticalSection TimelinesCS;
 		TArray<FChaosDDTimelineWeakPtr> Timelines;
 		FSphere3d DrawRegion;
 		int32 CommandBudget;
+		bool bIsServer;
+		bool bRenderEnabled;
 	};
 }
 
diff --git a/Engine/Source/Runtime/Experimental/Chaos/Public/ChaosDebugDraw/ChaosDDTypes.h b/Engine/Source/Runtime/Experimental/Chaos/Public/ChaosDebugDraw/ChaosDDTypes.h
--- a/Engine/Source/Runtime/Experimental/Chaos/Public/ChaosDebugDraw/ChaosDDTypes.h
+++ b/Engine/Source/Runtime/Experimental/Chaos/Public/ChaosDebugDraw/ChaosDDTypes.h
@@ -23,12 +23,15 @@ namespace ChaosDD::Private
 	class FChaosDDGlobalFrame;
 	class FChaosDDScene;
 	class FChaosDDTimeline;
+	class FChaosDDTimelineContext;
+	class FChaosDDTaskContext;
 
 	class IChaosDDRenderer;
 
 	using FChaosDDFramePtr = TSharedPtr<FChaosDDFrame, ESPMode::ThreadSafe>;
 	using FChaosDDGlobalFramePtr = TSharedPtr<FChaosDDGlobalFrame, ESPMode::ThreadSafe>;
 	using FChaosDDScenePtr = TSharedPtr<FChaosDDScene, ESPMode::ThreadSafe>;
+	using FChaosDDSceneWeakPtr = TWeakPtr<FChaosDDScene, ESPMode::ThreadSafe>;
 	using FChaosDDTimelinePtr = TSharedPtr<FChaosDDTimeline, ESPMode::ThreadSafe>;
 	using FChaosDDTimelineWeakPtr = TWeakPtr<FChaosDDTimeline, ESPMode::ThreadSafe>;
 }
diff --git a/Engine/Source/Runtime/Experimental/ChaosSolverEngine/Private/Chaos/ChaosDebugDrawComponent.cpp b/Engine/Source/Runtime/Experimental/ChaosSolverEngine/Private/Chaos/ChaosDebugDrawComponent.cpp
--- a/Engine/Source/Runtime/Experimental/ChaosSolverEngine/Private/Chaos/ChaosDebugDrawComponent.cpp
+++ b/Engine/Source/Runtime/Experimental/ChaosSolverEngine/Private/Chaos/ChaosDebugDrawComponent.cpp
@@ -2,8 +2,6 @@
 
 #include "Chaos/ChaosDebugDrawComponent.h"
 #include "Chaos/DebugDrawQueue.h"
-#include "ChaosDebugDraw/ChaosDDRenderer.h"
-#include "ChaosDebugDraw/ChaosDDScene.h"
 #include "ChaosLog.h"
 #include "Debug/DebugDrawService.h"
 #include "DrawDebugHelpers.h"
@@ -215,58 +213,6 @@ void DebugDrawChaos(const AActor* DebugDrawActor, const TArray<Chaos::FLatentDra
 }
 #endif
 
-#if CHAOS_DEBUG_DRAW
-
-class FChaosDDRenderer : public ChaosDD::Private::IChaosDDRenderer
-{
-public:
-	FChaosDDRenderer(UWorld* InWorld, AActor* InDebugDrawActor)
-		: World(InWorld)
-		, DebugDrawActor(InDebugDrawActor)
-	{
-		DepthPriority = 10;
-		bIsPaused = false;
-	}
-
-	virtual bool IsServer() const override final
-	{
-		return World->GetNetMode() == ENetMode::NM_DedicatedServer;
-	}
-
-	virtual void RenderLine(const FVector3d& A, const FVector3d& B, const FColor& Color, float LineThickness, float Lifetime) const override final
-	{
-		DrawDebugLine(World, A, B, Color, false, CommandLifeTime(Lifetime), DepthPriority, LineThickness);
-	}
-
-	virtual void RenderBox(const FVector3d& Position, const FQuat4d& Rotation, const FVector3d& Size, const FColor& Color, float LineThickness, float Lifetime) const override final
-	{
-		DrawDebugBox(World, Position, Size, Rotation, Color, false, CommandLifeTime(Lifetime), DepthPriority, LineThickness);
-	}
-
-	virtual void RenderLatentCommand(const Chaos::FLatentDrawCommand& Command) const override final
-	{
-		const bool bDrawUe = bChaosDebugDraw_DrawMode != 1;
-		if (bDrawUe)
-		{
-			DebugDrawChaosCommand(World, Command);
-		}
-
-		const bool bDrawVisLog = bChaosDebugDraw_DrawMode != 0;
-		if (bDrawVisLog)
-		{
-			const AActor* Actor = (Command.TestBaseActor) ? Command.TestBaseActor : DebugDrawActor;
-			VisLogChaosCommand(Actor, Command);
-		}
-	}
-
-private:
-	UWorld* World;
-	AActor* DebugDrawActor;
-	int8 DepthPriority;
-	bool bIsPaused;
-};
-
-#endif
 
 
 UChaosDebugDrawComponent::UChaosDebugDrawComponent()
@@ -405,23 +351,6 @@ void UChaosDebugDrawComponent::TickComponent(float DeltaTime, enum ELevelTick Ti
 
 			DebugDrawChaos(GetOwner(), DrawCommands, World->IsPaused());
 		}
-
-		if (World->GetPhysicsScene() != nullptr)
-		{
-			const ChaosDD::Private::FChaosDDScenePtr& CDDScene = World->GetPhysicsScene()->GetDebugDrawScene();
-			if (CDDScene.IsValid())
-			{
-				CDDScene->SetDrawRegion(FSphere3d(RegionOfInterestOrigin, RegionOfInterestRadius));
-				CDDScene->SetCommandBudget(ChaosDebugDraw_MaxElements);
-
-				if (!bIsPaused)
-				{
-					FChaosDDRenderer CDDRenderer = FChaosDDRenderer(World, GetOwner());
-					constexpr bool bRenderGlobalFrame = true;
-					CDDScene->RenderLatestFrames(CDDRenderer, bRenderGlobalFrame);
-				}
-			}
-		}
 	}
 #endif
 }
diff --git a/Engine/Source/Runtime/Experimental/ChaosSolverEngine/Private/Chaos/ChaosDebugDrawSubsystem.cpp b/Engine/Source/Runtime/Experimental/ChaosSolverEngine/Private/Chaos/ChaosDebugDrawSubsystem.cpp
new file mode 100644
--- /dev/null
+++ b/Engine/Source/Runtime/Experimental/ChaosSolverEngine/Private/Chaos/ChaosDebugDrawSubsystem.cpp
@@ -0,0 +1,395 @@
+// Copyright Epic Games, Inc. All Rights Reserved.
+#include "Chaos/ChaosDebugDrawSubsystem.h"
+#include "Chaos/DebugDrawQueue.h"
+#include "ChaosDebugDraw/ChaosDDContext.h"
+#include "ChaosDebugDraw/ChaosDDLog.h"
+#include "ChaosDebugDraw/ChaosDDRenderer.h"
+#include "ChaosDebugDraw/ChaosDDScene.h"
+#include "ChaosDebugDraw/ChaosDDTimeline.h"
+#include "DrawDebugHelpers.h"
+#include "Engine/Engine.h"
+#include "Engine/World.h"
+#include "HAL/CriticalSection.h"
+#include "Physics/Experimental/PhysScene_Chaos.h"
+#include "VisualLogger/VisualLogger.h"
+
+FDelegateHandle UChaosDebugDrawSubsystem::OnTickWorldStartDelegate;
+FDelegateHandle UChaosDebugDrawSubsystem::OnTickWorldEndDelegate;
+
+#if CHAOS_DEBUG_DRAW
+
+// @todo(chaos): Stuff in ChaosDebugDrawComponent that should move here when we drop that file
+extern float ChaosDebugDraw_Radius;
+extern int32 bChaosDebugDraw_DrawMode;
+extern int32 ChaosDebugDraw_MaxElements;
+extern float CommandLifeTime(const float LifeTime);
+extern void DebugDrawChaosCommand(const UWorld* World, const Chaos::FLatentDrawCommand& Command);
+extern void VisLogChaosCommand(const AActor* Actor, const Chaos::FLatentDrawCommand& Command);
+
+namespace Chaos::CVars
+{
+	extern int32 ChaosSolverDebugDrawShowServer;
+	extern int32 ChaosSolverDebugDrawShowClient;
+}
+
+class FChaosDDRenderer : public ChaosDD::Private::IChaosDDRenderer
+{
+public:
+	FChaosDDRenderer(UWorld* InWorld, const FSphere3d& InDrawRegion, const int32 InRenderBudget)
+		: World(InWorld)
+		, DrawRegion(InDrawRegion)
+		, RenderBudget(InRenderBudget)
+		, SphereSegments(8)
+		, DepthPriority(10)
+	{
+	}
+
+	virtual bool IsServer() const override final
+	{
+		return World->GetNetMode() == ENetMode::NM_DedicatedServer;
+	}
+
+	virtual FSphere3d GetDrawRegion() const override final
+	{
+		return DrawRegion;
+	}
+
+	virtual void RenderLine(const FVector3d& A, const FVector3d& B, const FColor& Color, float LineThickness, float Lifetime) override final
+	{
+		constexpr int32 Cost = 1;
+		if (TryAddToCost(Cost))
+		{
+			DrawDebugLine(World, A, B, Color, false, CommandLifeTime(Lifetime), DepthPriority, LineThickness);
+		}
+	}
+
+	virtual void RenderSphere(const FVector3d& Center, float Radius, const FColor& Color, float LineThickness, float Lifetime) override final
+	{
+		constexpr int32 Cost = 16;
+		if (TryAddToCost(Cost))
+		{
+			DrawDebugSphere(World, Center, Radius, SphereSegments, Color, false, CommandLifeTime(Lifetime), DepthPriority, LineThickness);
+		}
+	}
+
+	virtual void RenderCapsule(const FVector3d& Center, const FQuat4d& Rotation, float HalfHeight, float Radius, const FColor& Color, float LineThickness, float Lifetime) override final
+	{
+		constexpr int32 Cost = 16;
+		if (TryAddToCost(Cost))
+		{
+			DrawDebugCapsule(World, Center, HalfHeight, Radius, Rotation, Color, false, CommandLifeTime(Lifetime), DepthPriority, LineThickness);
+		}
+	}
+
+	virtual void RenderBox(const FVector3d& Position, const FQuat4d& Rotation, const FVector3d& Size, const FColor& Color, float LineThickness, float Lifetime) override final
+	{
+		constexpr int32 Cost = 12;
+		if (TryAddToCost(Cost))
+		{
+			DrawDebugBox(World, Position, Size, Rotation, Color, false, CommandLifeTime(Lifetime), DepthPriority, LineThickness);
+		}
+	}
+
+	virtual void RenderTriangle(const FVector3d& A, const FVector3d& B, const FVector3d& C, const FColor& Color, float LineThickness, float Lifetime) override final
+	{
+		constexpr int32 Cost = 3;
+		if (TryAddToCost(Cost))
+		{
+			DrawDebugLine(World, A, B, Color, false, CommandLifeTime(Lifetime), DepthPriority, LineThickness);
+			DrawDebugLine(World, B, C, Color, false, CommandLifeTime(Lifetime), DepthPriority, LineThickness);
+			DrawDebugLine(World, C, A, Color, false, CommandLifeTime(Lifetime), DepthPriority, LineThickness);
+		}
+	}
+
+	virtual void RenderLatentCommand(const Chaos::FLatentDrawCommand& Command) override final
+	{
+		// We don't check the cost of the legacy commands because the budget is applied at capture time for them
+		const bool bDrawUe = bChaosDebugDraw_DrawMode != 1;
+		if (bDrawUe)
+		{
+			DebugDrawChaosCommand(World, Command);
+		}
+
+		const bool bDrawVisLog = bChaosDebugDraw_DrawMode != 0;
+		if (bDrawVisLog)
+		{
+			VisLogChaosCommand(Command.TestBaseActor, Command);
+		}
+	}
+
+private:
+	bool TryAddToCost(int32 InCost)
+	{
+		RenderCost += InCost;
+
+		// A budget of zero means infinity
+		if ((RenderCost <= RenderBudget) || (RenderBudget == 0))
+		{
+			RenderCost += InCost;
+			return true;
+		}
+
+		return false;
+	}
+
+	UWorld* World;
+	FSphere3d DrawRegion;
+	int32 RenderBudget;
+	int32 RenderCost;
+
+	// Draw Settings
+	int32 SphereSegments;
+	int8 DepthPriority;
+};
+
+class FChaosDDWorldManager
+{
+public:
+	static FChaosDDWorldManager& Get()
+	{
+		static FChaosDDWorldManager SManager;
+		return SManager;
+	}
+
+	void SetServerDebugDrawScene(const ChaosDD::Private::FChaosDDScenePtr& InServerScene)
+	{
+		FScopeLock Lock(&SceneCS);
+		CDDServerScene = InServerScene;
+	}
+
+	const ChaosDD::Private::FChaosDDScenePtr& GetServerDebugDrawScene() const
+	{
+		FScopeLock Lock(&SceneCS);
+		return CDDServerScene;
+	}
+
+private:
+
+	mutable FCriticalSection SceneCS;
+
+	ChaosDD::Private::FChaosDDScenePtr CDDServerScene;
+};
+
+#endif
+
+
+bool UChaosDebugDrawSubsystem::ShouldCreateSubsystem(UObject* Outer) const
+{
+#if CHAOS_DEBUG_DRAW
+	if (UWorld* World = Cast<UWorld>(Outer))
+	{
+		const bool bCreateDebugDraw = !IsRunningCommandlet() && !World->IsPreviewWorld();
+
+		UE_LOG(LogChaosDD, Log, TEXT("%s Chaos Debug Draw Scene for world %s"), bCreateDebugDraw ? TEXT("Creating") : TEXT("Not creating"), *World->GetName());
+		
+		return bCreateDebugDraw;
+	}
+	return Super::ShouldCreateSubsystem(Outer);
+#else
+	return false;
+#endif
+}
+
+void UChaosDebugDrawSubsystem::PostInitialize()
+{
+	Super::PostInitialize();
+
+#if CHAOS_DEBUG_DRAW
+	if (UWorld* World = GetWorld())
+	{
+		const bool bIsServer = World->IsNetMode(NM_DedicatedServer);
+
+		// Create the debug draw scene which holds all debug draw data for this world
+		const FString DDSceneName = FString::Format(TEXT("{0} {1}"), { World->GetName(), bIsServer ? TEXT("Server") : TEXT("Client") });
+		CDDScene = MakeShared<ChaosDD::Private::FChaosDDScene>(DDSceneName, bIsServer);
+
+		// Create the debug draw timeline for the game thread
+		const FString DDTimelineName = FString::Format(TEXT("{0} {1} {2}"), { World->GetName(), bIsServer ? TEXT("Server") : TEXT("Client"), TEXT("Game Frame") });
+		CDDWorldTimeline = CDDScene->CreateTimeline(DDTimelineName);
+
+		// Tell the physics scene about the debug draw - it is async and will create its own timeline(s)
+		if (World->GetPhysicsScene() != nullptr)
+		{
+			World->GetPhysicsScene()->SetDebugDrawScene(CDDScene);
+		}
+
+		if (bIsServer)
+		{
+			FChaosDDWorldManager::Get().SetServerDebugDrawScene(CDDScene);
+		}
+	}
+#endif
+}
+
+void UChaosDebugDrawSubsystem::Deinitialize()
+{
+#if CHAOS_DEBUG_DRAW
+	if (UWorld* World = GetWorld())
+	{
+		if (CDDScene.IsValid() && CDDScene->IsServer())
+		{
+			FChaosDDWorldManager::Get().SetServerDebugDrawScene({});
+		}
+	}
+#endif
+
+	Super::Deinitialize();
+}
+
+void UChaosDebugDrawSubsystem::OnWorldTickStart(ELevelTick TickType, float Dt)
+{
+#if CHAOS_DEBUG_DRAW
+	if (UWorld* World = GetWorld())
+	{
+		CDDWorldTimelineContext.BeginFrame(CDDWorldTimeline, World->GetTimeSeconds(), Dt);
+	}
+#endif
+}
+
+void UChaosDebugDrawSubsystem::OnWorldTickEnd(ELevelTick TickType, float Dt)
+{
+#if CHAOS_DEBUG_DRAW
+	if (UWorld* World = GetWorld())
+	{
+		CDDWorldTimelineContext.EndFrame();
+
+		if (CDDScene.IsValid())
+		{
+			FSphere3d DrawRegion = CDDScene->GetDrawRegion();
+			DrawRegion.W = ChaosDebugDraw_Radius;
+			if (World->ViewLocationsRenderedLastFrame.Num() > 0)
+			{
+				DrawRegion.Center = World->ViewLocationsRenderedLastFrame[0];
+			}
+
+			CDDScene->SetDrawRegion(DrawRegion);
+			CDDScene->SetCommandBudget(ChaosDebugDraw_MaxElements);
+
+			// @todo(chaos): This is a problem in multi-client PIE - the draw region on the server can flipflop between
+			// the client positions depending on which one renders last. Could support multiple draw regions maybe...
+			ChaosDD::Private::FChaosDDScenePtr ServerScene = FChaosDDWorldManager::Get().GetServerDebugDrawScene();
+			if (ServerScene.IsValid())
+			{
+				ServerScene->SetDrawRegion(DrawRegion);
+				ServerScene->SetCommandBudget(ChaosDebugDraw_MaxElements);
+			}
+		}
+
+		RenderScene();
+	}
+#endif
+}
+
+void UChaosDebugDrawSubsystem::RenderScene()
+{
+#if CHAOS_DEBUG_DRAW
+	if (UWorld* World = GetWorld())
+	{
+		if (CDDScene.IsValid() && !CDDScene->IsServer())
+		{
+			if (!World->IsPaused())
+			{
+				// @todo(chaos): 
+				FChaosDDRenderer Renderer = FChaosDDRenderer(World, CDDScene->GetDrawRegion(), CDDScene->GetCommandBudget());
+
+				// Render all of the out-of frame commands
+				// @todo(chaos): extracting the global frame hereis a problem if we have multiple 
+				// PIE clients, but the global scene is a hack anyway so won't worry about that
+				RenderFrame(Renderer, ChaosDD::Private::FChaosDDContext::ExtractGlobalFrame());
+
+				// Render the commands from the server on every client (in PIE)
+				RenderScene(Renderer, FChaosDDWorldManager::Get().GetServerDebugDrawScene());
+
+				// Render the commands from this world
+				RenderScene(Renderer, CDDScene);
+			}
+		}
+	}
+#endif
+}
+
+void UChaosDebugDrawSubsystem::RenderScene(FChaosDDRenderer& Renderer, const ChaosDD::Private::FChaosDDScenePtr& Scene)
+{
+#if CHAOS_DEBUG_DRAW
+	if (Scene.IsValid() && Scene->IsRenderEnabled())
+	{
+		TArray<ChaosDD::Private::FChaosDDFramePtr> Frames = Scene->GetLatestFrames();
+		for (const ChaosDD::Private::FChaosDDFramePtr& Frame : Frames)
+		{
+			RenderFrame(Renderer, Frame);
+
+			if (Frame->WasCommandBudgetExceeded())
+			{
+				constexpr int32 MsgId = 86421357;
+				const FString Msg = FString::Format(TEXT("Debug Draw Capture Budget Exceeded for {0} [{1} / {2}]"), { Frame->GetTimeline()->GetName(), Frame->GetCommandCost(), Frame->GetCommandBudget() });
+				GEngine->AddOnScreenDebugMessage(MsgId, 1.0f, FColor::Red, *Msg);
+			}
+		}
+	}
+#endif
+}
+
+
+void UChaosDebugDrawSubsystem::RenderFrame(FChaosDDRenderer& Renderer, const ChaosDD::Private::FChaosDDFramePtr& Frame)
+{
+#if CHAOS_DEBUG_DRAW
+	if (Frame.IsValid() && (Frame->GetNumCommands() + Frame->GetNumLatentCommands() > 0))
+	{
+		UE_LOG(LogChaosDD, VeryVerbose, TEXT("Render %s %d %d+%d Commands"), *Frame->GetTimeline()->GetName(), Frame->GetFrameIndex(), Frame->GetNumCommands(), Frame->GetNumLatentCommands());
+
+		// Render the legacy commands
+		Frame->VisitLatentCommands(
+			[&Renderer](const Chaos::FLatentDrawCommand& Command)
+			{
+				Renderer.RenderLatentCommand(Command);
+			});
+
+		// Render the commands
+		Frame->VisitCommands(
+			[&Renderer](const ChaosDD::Private::FChaosDDCommand& Command)
+			{
+				Command(Renderer);
+			});
+	}
+#endif
+}
+
+void UChaosDebugDrawSubsystem::Startup()
+{
+#if CHAOS_DEBUG_DRAW
+	UE_LOG(LogChaosDD, Log, TEXT("Chaos Debug Draw Startup"));
+
+	OnTickWorldStartDelegate = FWorldDelegates::OnWorldPreActorTick.AddStatic(&StaticOnWorldTickStart);
+	OnTickWorldEndDelegate = FWorldDelegates::OnWorldPostActorTick.AddStatic(&StaticOnWorldTickEnd);
+#endif
+}
+
+void UChaosDebugDrawSubsystem::Shutdown()
+{
+#if CHAOS_DEBUG_DRAW
+	UE_LOG(LogChaosDD, Log, TEXT("Chaos Debug Draw Shutdown"));
+
+	FWorldDelegates::OnWorldPreActorTick.Remove(OnTickWorldStartDelegate);
+	FWorldDelegates::OnWorldPostActorTick.Remove(OnTickWorldEndDelegate);
+#endif
+}
+
+void UChaosDebugDrawSubsystem::StaticOnWorldTickStart(UWorld* World, ELevelTick TickType, float Dt)
+{
+#if CHAOS_DEBUG_DRAW
+	if (UChaosDebugDrawSubsystem* CDDSystem = World->GetSubsystem<UChaosDebugDrawSubsystem>())
+	{
+		CDDSystem->OnWorldTickStart(TickType, Dt);
+	}
+#endif
+}
+
+void UChaosDebugDrawSubsystem::StaticOnWorldTickEnd(UWorld* World, ELevelTick TickType, float Dt)
+{
+#if CHAOS_DEBUG_DRAW
+	if (UChaosDebugDrawSubsystem* CDDSystem = World->GetSubsystem<UChaosDebugDrawSubsystem>())
+	{
+		CDDSystem->OnWorldTickEnd(TickType, Dt);
+	}
+#endif
+}
diff --git a/Engine/Source/Runtime/Experimental/ChaosSolverEngine/Private/Chaos/ChaosSolverEnginePlugin.cpp b/Engine/Source/Runtime/Experimental/ChaosSolverEngine/Private/Chaos/ChaosSolverEnginePlugin.cpp
--- a/Engine/Source/Runtime/Experimental/ChaosSolverEngine/Private/Chaos/ChaosSolverEnginePlugin.cpp
+++ b/Engine/Source/Runtime/Experimental/ChaosSolverEngine/Private/Chaos/ChaosSolverEnginePlugin.cpp
@@ -7,6 +7,7 @@
 #include "Engine/Engine.h"
 #include "Engine/World.h"
 #include "Chaos/ChaosDebugDrawComponent.h"
+#include "Chaos/ChaosDebugDrawSubsystem.h"
 #include "Chaos/ChaosSolverActor.h"
 #include "ChaosSolversModule.h"
 #include "ChaosVDRecordingStateScreenMessageHandler.h"
@@ -28,6 +29,10 @@ void FChaosSolverEnginePlugin::StartupModule()
 
 	UChaosDebugDrawComponent::BindWorldDelegates();
 
+#if CHAOS_DEBUG_DRAW
+	UChaosDebugDrawSubsystem::Startup();
+#endif
+
 #if WITH_CHAOS_VISUAL_DEBUGGER
 	FChaosVDRecordingStateScreenMessageHandler::Get().Initialize();
 #endif
@@ -36,6 +41,9 @@ void FChaosSolverEnginePlugin::StartupModule()
 
 void FChaosSolverEnginePlugin::ShutdownModule()
 {
+#if CHAOS_DEBUG_DRAW
+	UChaosDebugDrawSubsystem::Shutdown();
+#endif
 
 #if WITH_CHAOS_VISUAL_DEBUGGER
 	FChaosVDRecordingStateScreenMessageHandler::Get().TearDown();
diff --git a/Engine/Source/Runtime/Experimental/ChaosSolverEngine/Public/Chaos/ChaosDebugDrawSubsystem.h b/Engine/Source/Runtime/Experimental/ChaosSolverEngine/Public/Chaos/ChaosDebugDrawSubsystem.h
new file mode 100644
--- /dev/null
+++ b/Engine/Source/Runtime/Experimental/ChaosSolverEngine/Public/Chaos/ChaosDebugDrawSubsystem.h
@@ -0,0 +1,44 @@
+// Copyright Epic Games, Inc. All Rights Reserved.
+#pragma once
+
+#include "ChaosDebugDraw/ChaosDDContext.h"
+#include "ChaosDebugDraw/ChaosDDTypes.h"
+#include "Engine/EngineBaseTypes.h"
+#include "Subsystems/WorldSubsystem.h"
+
+#include "ChaosDebugDrawSubsystem.generated.h"
+
+class UWorld;
+class FChaosDDRenderer;
+
+UCLASS(MinimalAPI)
+class UChaosDebugDrawSubsystem : public UWorldSubsystem
+{
+	GENERATED_BODY()
+
+public:
+	CHAOSSOLVERENGINE_API const ChaosDD::Private::FChaosDDScenePtr& GetDebugDrawScene() const { return CDDScene; }
+
+	virtual bool ShouldCreateSubsystem(UObject* Outer) const override final;
+	virtual void PostInitialize() override final;
+	virtual void Deinitialize() override final;
+
+	static void Startup();
+	static void Shutdown();
+
+protected:
+	void OnWorldTickStart(ELevelTick TickType, float Dt);
+	void OnWorldTickEnd(ELevelTick TickType, float Dt);
+	void RenderScene();
+	void RenderScene(FChaosDDRenderer& Renderer, const ChaosDD::Private::FChaosDDScenePtr& Scene);
+	void RenderFrame(FChaosDDRenderer& Renderer, const ChaosDD::Private::FChaosDDFramePtr& Frame);
+
+	ChaosDD::Private::FChaosDDScenePtr CDDScene;
+	ChaosDD::Private::FChaosDDTimelinePtr CDDWorldTimeline;
+	ChaosDD::Private::FChaosDDTimelineContext CDDWorldTimelineContext;
+
+	static void StaticOnWorldTickStart(UWorld* World, ELevelTick TickType, float Dt);
+	static void StaticOnWorldTickEnd(UWorld* World, ELevelTick TickType, float Dt);
+	static FDelegateHandle OnTickWorldStartDelegate;
+	static FDelegateHandle OnTickWorldEndDelegate;
+};
diff --git a/Engine/Source/Runtime/PhysicsCore/Private/ChaosScene.cpp b/Engine/Source/Runtime/PhysicsCore/Private/ChaosScene.cpp
--- a/Engine/Source/Runtime/PhysicsCore/Private/ChaosScene.cpp
+++ b/Engine/Source/Runtime/PhysicsCore/Private/ChaosScene.cpp
@@ -82,11 +82,6 @@ FChaosScene::FChaosScene(
 	SceneSolver->GetChaosVDContextData().Type = static_cast<int32>(EChaosVDContextType::Solver);
 #endif
 
-#if CHAOS_DEBUG_DRAW
-	CDDScene = MakeShared<ChaosDD::Private::FChaosDDScene>(DebugName.ToString());
-	SceneSolver->SetDebugDrawScene(CDDScene);
-#endif
-
 	SceneSolver->PhysSceneHack = this;
 	SimCallback = SceneSolver->CreateAndRegisterSimCallbackObject_External<FChaosSceneSimCallback>();
 
@@ -581,3 +576,16 @@ FGraphEventArray FChaosScene::GetCompletionEvents()
 {
 	return CompletionEvents;
 }
+
+#if CHAOS_DEBUG_DRAW
+void FChaosScene::SetDebugDrawScene(const ChaosDD::Private::FChaosDDScenePtr& InCDDScene)
+{
+	CDDScene = InCDDScene;
+
+	SceneSolver->EnqueueCommandImmediate(
+		[this]()
+		{
+			SceneSolver->SetDebugDrawScene(CDDScene);
+		});
+}
+#endif
diff --git a/Engine/Source/Runtime/PhysicsCore/Public/Chaos/ChaosScene.h b/Engine/Source/Runtime/PhysicsCore/Public/Chaos/ChaosScene.h
--- a/Engine/Source/Runtime/PhysicsCore/Public/Chaos/ChaosScene.h
+++ b/Engine/Source/Runtime/PhysicsCore/Public/Chaos/ChaosScene.h
@@ -213,6 +213,8 @@ public:
 		return CDDScene;
 	}
 
+	PHYSICSCORE_API void SetDebugDrawScene(const ChaosDD::Private::FChaosDDScenePtr& CDDScene);
+
 private:
 	ChaosDD::Private::FChaosDDScenePtr CDDScene;
 #endif
