[Backout] - CL33567068
Original CL Desc
-----------------------------------------------------------------
Chaos Debug Draw
- add debug draw subsystem
- add DD timeline to game thread
- move all new DD stuff out of ChaosDebugDrawComponent (it can be deleted later)

diff --git a/Engine/Source/Runtime/Experimental/Chaos/Private/Chaos/PBDRigidsEvolutionGBF.cpp b/Engine/Source/Runtime/Experimental/Chaos/Private/Chaos/PBDRigidsEvolutionGBF.cpp
--- a/Engine/Source/Runtime/Experimental/Chaos/Private/Chaos/PBDRigidsEvolutionGBF.cpp
+++ b/Engine/Source/Runtime/Experimental/Chaos/Private/Chaos/PBDRigidsEvolutionGBF.cpp
@@ -382,7 +382,7 @@ void FPBDRigidsEvolutionGBF::SetDebugDrawScene(const ChaosDD::Private::FChaosDDS
 
 	if (CDDScene.IsValid())
 	{
-		CDDTickTimeline = CDDScene->CreateTimeline(FString::Format(TEXT("{0} {1}"), { CDDScene->GetName(), "Physics Tick" }));
+		CDDTickTimeline = CDDScene->CreateTimeline(FString::Format(TEXT("{0} {1}"), { CDDScene->GetName(), "Tick" }));
 	}
 }
 #endif
diff --git a/Engine/Source/Runtime/Experimental/Chaos/Private/ChaosDebugDraw/ChaosDDContext.cpp b/Engine/Source/Runtime/Experimental/Chaos/Private/ChaosDebugDraw/ChaosDDContext.cpp
--- a/Engine/Source/Runtime/Experimental/Chaos/Private/ChaosDebugDraw/ChaosDDContext.cpp
+++ b/Engine/Source/Runtime/Experimental/Chaos/Private/ChaosDebugDraw/ChaosDDContext.cpp
@@ -22,91 +22,61 @@ namespace ChaosDD::Private
 
 	//
 	//
-	// Timeline Context
 	//
 	//
 
-	void FChaosDDTimelineContext::BeginFrame(const FChaosDDTimelinePtr& InTimeline, double InTime, double InDt)
+	FChaosDDScopeTimelineContext::FChaosDDScopeTimelineContext(const FChaosDDTimelinePtr& InTimeline, double InTime, double InDt)
 	{
-		FChaosDDContext& Context = FChaosDDContext::Get();
-		ParentFrame = Context.Frame;
-
 		if (InTimeline.IsValid())
 		{
 			Timeline = InTimeline;
 			Timeline->BeginFrame(InTime, InDt);
+
+			FChaosDDContext& Context = FChaosDDContext::Get();
+			ParentFrame = Context.Frame;
 			Context.Frame = Timeline->GetActiveFrame();
 		}
-		else
-		{
-			Context.Frame.Reset();
-		}
 	}
 
-	void FChaosDDTimelineContext::EndFrame()
+	FChaosDDScopeTimelineContext::~FChaosDDScopeTimelineContext()
 	{
 		if (Timeline.IsValid())
 		{
 			Timeline->EndFrame();
 			Timeline.Reset();
-		}
-
-		FChaosDDContext& Context = FChaosDDContext::Get();
-		Context.Frame = ParentFrame;
-		ParentFrame.Reset();
-	}
 
-	FChaosDDScopeTimelineContext::FChaosDDScopeTimelineContext(const FChaosDDTimelinePtr& InTimeline, double InTime, double InDt)
-	{
-		Context.BeginFrame(InTimeline, InTime, InDt);
-	}
-
-	FChaosDDScopeTimelineContext::~FChaosDDScopeTimelineContext()
-	{
-		Context.EndFrame();
+			FChaosDDContext& Context = FChaosDDContext::Get();
+			Context.Frame = ParentFrame;
+			ParentFrame.Reset();
+		}
 	}
 
 	//
 	//
-	// Task Context
 	//
 	//
 
-	void FChaosDDTaskContext::BeginThread(const FChaosDDContext& InParentDDContext)
+	FChaosDDScopeTaskContext::FChaosDDScopeTaskContext(const FChaosDDContext& InParentDDContext)
 	{
-		FChaosDDContext& Context = FChaosDDContext::Get();
-		ParentFrame = Context.Frame;
-
 		if (InParentDDContext.Frame.IsValid())
 		{
+			FChaosDDContext& Context = FChaosDDContext::Get();
+
+			ParentFrame = Context.Frame;
 			Context.Frame = InParentDDContext.Frame;
 		}
-		else
-		{
-			Context.Frame.Reset();
-		}
 	}
 
-	void FChaosDDTaskContext::EndThread()
+	FChaosDDScopeTaskContext::~FChaosDDScopeTaskContext()
 	{
 		FChaosDDContext& Context = FChaosDDContext::Get();
+
 		Context.Frame = ParentFrame;
 		ParentFrame.Reset();
 	}
 
-	FChaosDDScopeTaskContext::FChaosDDScopeTaskContext(const FChaosDDContext& InParentDDContext)
-	{
-		Context.BeginThread(InParentDDContext);
-	}
-
-	FChaosDDScopeTaskContext::~FChaosDDScopeTaskContext()
-	{
-		Context.EndThread();
-	}
-
 	//
 	//
-	// Thread Local Context
 	//
 	//
 
diff --git a/Engine/Source/Runtime/Experimental/Chaos/Private/ChaosDebugDraw/ChaosDDScene.cpp b/Engine/Source/Runtime/Experimental/Chaos/Private/ChaosDebugDraw/ChaosDDScene.cpp
--- a/Engine/Source/Runtime/Experimental/Chaos/Private/ChaosDebugDraw/ChaosDDScene.cpp
+++ b/Engine/Source/Runtime/Experimental/Chaos/Private/ChaosDebugDraw/ChaosDDScene.cpp
@@ -22,12 +22,10 @@ namespace ChaosDD::Private
 		}
 	}
 
-	FChaosDDScene::FChaosDDScene(const FString& InName, bool bInIsServer)
+	FChaosDDScene::FChaosDDScene(const FString& InName)
 		: Name(InName)
 		, DrawRegion(FVector::Zero(), 0.0)
 		, CommandBudget(20000)
-		, bIsServer(bInIsServer)
-		, bRenderEnabled(true)
 	{
 	}
 
@@ -35,25 +33,6 @@ namespace ChaosDD::Private
 	{
 	}
 
-	bool FChaosDDScene::IsServer() const
-	{
-		return bIsServer;
-	}
-
-	void FChaosDDScene::SetRenderEnabled(bool bInRenderEnabled)
-	{
-		FScopeLock Lock(&TimelinesCS);
-
-		bRenderEnabled = bInRenderEnabled;
-	}
-
-	bool FChaosDDScene::IsRenderEnabled() const
-	{
-		FScopeLock Lock(&TimelinesCS);
-
-		return bRenderEnabled;
-	}
-
 	void FChaosDDScene::SetDrawRegion(const FSphere3d& InDrawRegion)
 	{
 		FScopeLock Lock(&TimelinesCS);
@@ -71,13 +50,6 @@ namespace ChaosDD::Private
 		//FChaosDDContext::SetGlobalDrawRegion(InDrawRegion);
 	}
 
-	const FSphere3d& FChaosDDScene::GetDrawRegion() const
-	{
-		FScopeLock Lock(&TimelinesCS);
-
-		return DrawRegion;
-	}
-
 	void FChaosDDScene::SetCommandBudget(int32 InCommandBudget)
 	{
 		FScopeLock Lock(&TimelinesCS);
@@ -93,13 +65,6 @@ namespace ChaosDD::Private
 		FChaosDDContext::SetGlobalCommandBudget(InCommandBudget);
 	}
 
-	int32 FChaosDDScene::GetCommandBudget() const
-	{
-		FScopeLock Lock(&TimelinesCS);
-
-		return CommandBudget;
-	}
-
 	FChaosDDTimelinePtr FChaosDDScene::CreateTimeline(const FString& InName)
 	{
 		FScopeLock Lock(&TimelinesCS);
@@ -111,9 +76,36 @@ namespace ChaosDD::Private
 		return Timeline;
 	}
 
-	TArray<FChaosDDFramePtr> FChaosDDScene::GetLatestFrames()
+	void FChaosDDScene::RenderLatestFrames(IChaosDDRenderer& Renderer, bool bIncludeGlobalFrame)
 	{
-		return GetFrames();
+		TArray<FChaosDDFramePtr> Frames = GetFrames();
+
+		if (bIncludeGlobalFrame)
+		{
+			Frames.Add(FChaosDDContext::ExtractGlobalFrame());
+		}
+
+		for (const FChaosDDFramePtr& Frame : Frames)
+		{
+			if (Frame.IsValid())
+			{
+				UE_LOG(LogChaosDD, VeryVerbose, TEXT("Render %s %d %d Commands"), *Frame->GetTimeline()->GetName(), Frame->GetFrameIndex(), Frame->GetNumCommands());
+
+				// Render the legacy commands
+				Frame->VisitLatentCommands(
+					[&Renderer](const Chaos::FLatentDrawCommand& Command)
+					{
+						Renderer.RenderLatentCommand(Command);
+					});
+
+				// Render the commands
+				Frame->VisitCommands(
+					[&Renderer](const FChaosDDCommand& Command)
+					{
+						Command(Renderer);
+					});
+			}
+		}
 	}
 
 	TArray<FChaosDDFramePtr> FChaosDDScene::GetFrames()
diff --git a/Engine/Source/Runtime/Experimental/Chaos/Private/PBDRigidsSolver.cpp b/Engine/Source/Runtime/Experimental/Chaos/Private/PBDRigidsSolver.cpp
--- a/Engine/Source/Runtime/Experimental/Chaos/Private/PBDRigidsSolver.cpp
+++ b/Engine/Source/Runtime/Experimental/Chaos/Private/PBDRigidsSolver.cpp
@@ -1338,7 +1338,7 @@ namespace Chaos
 
 		if (CDDScene.IsValid())
 		{
-			CDDFrameTimeline = CDDScene->CreateTimeline(FString::Format(TEXT("{0} {1}"), { CDDScene->GetName(), "Physics Frame" }));
+			CDDFrameTimeline = CDDScene->CreateTimeline(FString::Format(TEXT("{0} {1}"), { CDDScene->GetName(), "Frame" }));
 		}
 
 		GetEvolution()->SetDebugDrawScene(InCDDScene);
@@ -2467,13 +2467,6 @@ TRACE_COUNTER_SET(ChaosTraceCounter_##Name, Value)
 #if CHAOS_DEBUG_DRAW
 		QUICK_SCOPE_CYCLE_COUNTER(SolverDebugDraw);
 
-		if (CDDScene.IsValid())
-		{
-			// @todo(chaos): this logic should be in UChaosDebugDrawSubsystem but the CVars are here. Move them and this...
-			const bool bRenderEnabled = (CDDScene->IsServer() && ChaosSolverDebugDrawShowServer) || (!CDDScene->IsServer() && ChaosSolverDebugDrawShowClient);
-			CDDScene->SetRenderEnabled(bRenderEnabled);
-		}
-
 		const bool bIsServer = GetDebugName().ToString().StartsWith(TEXT("Server"));
 		if (bIsServer && !ChaosSolverDebugDrawShowServer)
 		{
diff --git a/Engine/Source/Runtime/Experimental/Chaos/Public/Chaos/DebugDrawQueue.h b/Engine/Source/Runtime/Experimental/Chaos/Public/Chaos/DebugDrawQueue.h
--- a/Engine/Source/Runtime/Experimental/Chaos/Public/Chaos/DebugDrawQueue.h
+++ b/Engine/Source/Runtime/Experimental/Chaos/Public/Chaos/DebugDrawQueue.h
@@ -249,29 +249,29 @@ public:
 		RadiusOfInterest = InRadius;
 	}
 
-	FVector GetCenterOfInterest() const
+	const FVector& GetCenterOfInterest() const
 	{
-		return GetCenterOfInterestImpl();
+		return CenterOfInterest;
 	}
 
 	FReal GetRadiusOfInterest() const
 	{
-		return GetRadiusOfInterestImpl();
+		return RadiusOfInterest;
 	}
 
 	bool IsInRegionOfInterest(FVector Pos) const
 	{
-		return IsInRegionOfInterestImpl(Pos);
+		return (RadiusOfInterest <= 0.0f) || ((Pos - CenterOfInterest).SizeSquared() < RadiusOfInterest * RadiusOfInterest);
 	}
 
 	bool IsInRegionOfInterest(FVector Pos, FReal Radius) const
 	{
-		return IsInRegionOfInterestImpl(Pos, Radius);
+		return (RadiusOfInterest <= 0.0f) || ((Pos - CenterOfInterest).SizeSquared() < (RadiusOfInterest + Radius) * (RadiusOfInterest + Radius));
 	}
 
 	bool IsInRegionOfInterest(FAABB3 Bounds) const
 	{
-		return IsInRegionOfInterestImpl(Bounds);
+		return Bounds.ThickenSymmetrically(FVec3(RadiusOfInterest)).Contains(CenterOfInterest);
 	}
 
 	CHAOS_API void SetConsumerActive(void* Consumer, bool bConsumerActive);
@@ -294,60 +294,6 @@ private:
 	{}
 	~FDebugDrawQueue() {}
 
-	FVector GetCenterOfInterestImpl() const
-	{
-		if (bChaosDebugDraw_UseNewQueue)
-		{
-			ChaosDD::Private::FChaosDDFrameWriter DDWriter = ChaosDD::Private::FChaosDDContext::Get().GetWriter();
-			return DDWriter.GetDrawRegion().Center;
-		}
-
-		return CenterOfInterest;
-	}
-
-	FReal GetRadiusOfInterestImpl() const
-	{
-		if (bChaosDebugDraw_UseNewQueue)
-		{
-			ChaosDD::Private::FChaosDDFrameWriter DDWriter = ChaosDD::Private::FChaosDDContext::Get().GetWriter();
-			return DDWriter.GetDrawRegion().W;
-		}
-
-		return RadiusOfInterest;
-	}
-
-	bool IsInRegionOfInterestImpl(FVector Pos) const
-	{
-		if (bChaosDebugDraw_UseNewQueue)
-		{
-			ChaosDD::Private::FChaosDDFrameWriter DDWriter = ChaosDD::Private::FChaosDDContext::Get().GetWriter();
-			return DDWriter.IsInDrawRegion(Pos);
-		}
-
-		return (RadiusOfInterest <= 0.0f) || ((Pos - CenterOfInterest).SizeSquared() < RadiusOfInterest * RadiusOfInterest);
-	}
-
-	bool IsInRegionOfInterestImpl(FVector Pos, FReal Radius) const
-	{
-		if (bChaosDebugDraw_UseNewQueue)
-		{
-			ChaosDD::Private::FChaosDDFrameWriter DDWriter = ChaosDD::Private::FChaosDDContext::Get().GetWriter();
-			return DDWriter.IsInDrawRegion(FSphere3d(Pos, Radius));
-		}
-
-		return (RadiusOfInterest <= 0.0f) || ((Pos - CenterOfInterest).SizeSquared() < (RadiusOfInterest + Radius) * (RadiusOfInterest + Radius));
-	}
-
-	bool IsInRegionOfInterestImpl(FAABB3 Bounds) const
-	{
-		if (bChaosDebugDraw_UseNewQueue)
-		{
-			ChaosDD::Private::FChaosDDFrameWriter DDWriter = ChaosDD::Private::FChaosDDContext::Get().GetWriter();
-			return DDWriter.IsInDrawRegion(FBox3d(Bounds.Min(), Bounds.Max()));
-		}
-		return Bounds.ThickenSymmetrically(FVec3(RadiusOfInterest)).Contains(CenterOfInterest);
-	}
-
 	bool AcceptCommand(int Cost, const FVec3& Position)
 	{
 		if (bChaosDebugDraw_UseNewQueue)
diff --git a/Engine/Source/Runtime/Experimental/Chaos/Public/ChaosDebugDraw/ChaosDDContext.h b/Engine/Source/Runtime/Experimental/Chaos/Public/ChaosDebugDraw/ChaosDDContext.h
--- a/Engine/Source/Runtime/Experimental/Chaos/Public/ChaosDebugDraw/ChaosDDContext.h
+++ b/Engine/Source/Runtime/Experimental/Chaos/Public/ChaosDebugDraw/ChaosDDContext.h
@@ -25,12 +25,10 @@ namespace ChaosDD::Private
 			return FChaosDDFrameWriter(Get().GetFrame());
 		}
 
-		static FChaosDDFramePtr ExtractGlobalFrame();
-
 	private:
 		friend class FChaosDDScene;
-		friend class FChaosDDTimelineContext;
-		friend class FChaosDDTaskContext;
+		friend class FChaosDDScopeTimelineContext;
+		friend class FChaosDDScopeTaskContext;
 
 		const FChaosDDFramePtr& GetFrame() const
 		{
@@ -50,6 +48,7 @@ namespace ChaosDD::Private
 		// Global frame management
 		static const FChaosDDFramePtr& GetGlobalFrame();
 		static void CreateGlobalFrame();
+		static FChaosDDFramePtr ExtractGlobalFrame();
 		static void SetGlobalDrawRegion(const FSphere3d& InDrawRegion);
 		static void SetGlobalCommandBudget(int32 InCommandBudget);
 
@@ -63,21 +62,7 @@ namespace ChaosDD::Private
 	};
 
 	//
-	// Initializes the FChaosDDContext for a thread that owns a timeline
-	//
-	class CHAOS_API FChaosDDTimelineContext
-	{
-	public:
-		void BeginFrame(const FChaosDDTimelinePtr& InTimeline, double InTime, double InDt);
-		void EndFrame();
-
-	private:
-		FChaosDDTimelinePtr Timeline;
-		FChaosDDFramePtr ParentFrame;
-	};
-
-	//
-	// A scoped wrapper for FChaosDDTimelineContext
+	// A scoped Debug Draw Context for use on the thread that owns the timeline
 	//
 	class CHAOS_API FChaosDDScopeTimelineContext
 	{
@@ -86,27 +71,14 @@ namespace ChaosDD::Private
 		~FChaosDDScopeTimelineContext();
 
 	private:
-		FChaosDDTimelineContext Context;
-	};
-
-	//
-	// Initializes the FChaosDDContext for a task.
-	// Assumes that the task is kicked off from a thread that has an active debug draw context,
-	// which should be passed into this context. Any debug draws from the task will go to the 
-	// same frame as the parent context.
-	//
-	class CHAOS_API FChaosDDTaskContext
-	{
-	public:
-		void BeginThread(const FChaosDDContext& InParentDDContext);
-		void EndThread();
-
-	private:
+		FChaosDDTimelinePtr Timeline;
 		FChaosDDFramePtr ParentFrame;
 	};
 
 	//
-	// A scoped wrapper for FChaosDDTaskContext
+	// A scoped Debug Draw Context for use in tasks and parallel-for etc.
+	// Assumes that the task is kicked off from a thread that has an active debug draw context,
+	// which should be passed into this context.
 	//
 	class CHAOS_API FChaosDDScopeTaskContext
 	{
@@ -115,7 +87,7 @@ namespace ChaosDD::Private
 		~FChaosDDScopeTaskContext();
 
 	private:
-		FChaosDDTaskContext Context;
+		FChaosDDFramePtr ParentFrame;
 	};
 }
 
diff --git a/Engine/Source/Runtime/Experimental/Chaos/Public/ChaosDebugDraw/ChaosDDFrame.h b/Engine/Source/Runtime/Experimental/Chaos/Public/ChaosDebugDraw/ChaosDDFrame.h
--- a/Engine/Source/Runtime/Experimental/Chaos/Public/ChaosDebugDraw/ChaosDDFrame.h
+++ b/Engine/Source/Runtime/Experimental/Chaos/Public/ChaosDebugDraw/ChaosDDFrame.h
@@ -21,7 +21,7 @@ namespace ChaosDD::Private
 	// 
 	// E.g., See FChaosDDLine, FChaosDDSphere, FChaosDDParticle
 	//
-	using FChaosDDCommand = TFunction<void(IChaosDDRenderer&)>;
+	using FChaosDDCommand = TFunction<void(const IChaosDDRenderer&)>;
 
 	//
 	// A single frame of debug draw data
@@ -42,7 +42,6 @@ namespace ChaosDD::Private
 		{
 			if (InCommandQueueLength > 0)
 			{
-				Commands.Reserve(InCommandQueueLength);
 				LatentCommands.Reserve(InCommandQueueLength);
 			}
 		}
@@ -76,11 +75,6 @@ namespace ChaosDD::Private
 			DrawRegion = InRegion;
 		}
 
-		const FSphere3d& GetDrawRegion() const
-		{
-			return DrawRegion;
-		}
-
 		bool IsInDrawRegion(const FVector& InPos) const
 		{
 			if (DrawRegion.W > 0.0)
@@ -114,30 +108,14 @@ namespace ChaosDD::Private
 			CommandBudget = InCommandBudget;
 		}
 
-		int32 GetCommandBudget() const
-		{
-			return CommandBudget;
-		}
-
-		int32 GetCommandCost() const
-		{
-			return CommandCost;
-		}
-
-		bool WasCommandBudgetExceeded() const
-		{
-			return CommandCost > CommandBudget;
-		}
-
 		bool AddToCost(int32 InCost)
 		{
 			FScopeLock Lock(&CommandsCS);
 
-			CommandCost += InCost;
-
 			// A budget of zero means infinite
-			if ((CommandCost <= CommandBudget) || (CommandBudget == 0))
+			if ((CommandCost + InCost <= CommandBudget) || (CommandBudget == 0))
 			{
+				CommandCost += InCost;
 				return true;
 			}
 
@@ -162,16 +140,10 @@ namespace ChaosDD::Private
 		{
 			FScopeLock Lock(&CommandsCS);
 
-			return Commands.Num();
-		}
-
-		int32 GetNumLatentCommands() const
-		{
-			FScopeLock Lock(&CommandsCS);
-
 			return LatentCommands.Num();
 		}
 
+
 		// VisitorType = void(const FChaosDDCommand& Command)
 		template<typename VisitorType>
 		void VisitCommands(const VisitorType& Visitor)
@@ -319,15 +291,6 @@ namespace ChaosDD::Private
 			}
 		}
 
-		FSphere3d GetDrawRegion() const
-		{
-			if (Frame.IsValid())
-			{
-				return Frame->GetDrawRegion();
-			}
-			return FSphere3d(FVector3d(0), 0);
-		}
-
 		bool IsInDrawRegion(const FVector& InPos) const
 		{
 			if (Frame.IsValid())
diff --git a/Engine/Source/Runtime/Experimental/Chaos/Public/ChaosDebugDraw/ChaosDDRenderer.h b/Engine/Source/Runtime/Experimental/Chaos/Public/ChaosDebugDraw/ChaosDDRenderer.h
--- a/Engine/Source/Runtime/Experimental/Chaos/Public/ChaosDebugDraw/ChaosDDRenderer.h
+++ b/Engine/Source/Runtime/Experimental/Chaos/Public/ChaosDebugDraw/ChaosDDRenderer.h
@@ -25,18 +25,12 @@ namespace ChaosDD::Private
 		// Are we rendering a Server scene?
 		virtual bool IsServer() const = 0;
 
-		// The region of interest
-		virtual FSphere3d GetDrawRegion() const = 0;
-
 		// Utility functions for use by Debug Draw commands (e.g., FChaosDDLine)
-		virtual void RenderLine(const FVector3d& A, const FVector3d& B, const FColor& Color, float LineThickness, float Lifetime) = 0;
-		virtual void RenderSphere(const FVector3d& Center, float Radius, const FColor& Color, float LineThickness, float Lifetime) = 0;
-		virtual void RenderCapsule(const FVector3d& Center, const FQuat4d& Rotation, float HalfHeight, float Radius, const FColor& Color, float LineThickness, float Lifetime) = 0;
-		virtual void RenderBox(const FVector3d& Position, const FQuat4d& Rotation, const FVector3d& Size, const FColor& Color, float LineThickness, float Lifetime) = 0;
-		virtual void RenderTriangle(const FVector3d& A, const FVector3d& B, const FVector3d& C, const FColor& Color, float LineThickness, float Lifetime) = 0;
+		virtual void RenderLine(const FVector3d& A, const FVector3d& B, const FColor& Color, float LineThickness, float Lifetime) const = 0;
+		virtual void RenderBox(const FVector3d& Position, const FQuat4d& Rotation, const FVector3d& Size, const FColor& Color, float LineThickness, float Lifetime) const = 0;
 
 		// Render legacy debug draw command (See FChaosDDScene::RenderLatestFrames)
-		virtual void RenderLatentCommand(const Chaos::FLatentDrawCommand& Command) = 0;
+		virtual void RenderLatentCommand(const Chaos::FLatentDrawCommand& Command) const = 0;
 	};
 }
 
diff --git a/Engine/Source/Runtime/Experimental/Chaos/Public/ChaosDebugDraw/ChaosDDScene.h b/Engine/Source/Runtime/Experimental/Chaos/Public/ChaosDebugDraw/ChaosDDScene.h
--- a/Engine/Source/Runtime/Experimental/Chaos/Public/ChaosDebugDraw/ChaosDDScene.h
+++ b/Engine/Source/Runtime/Experimental/Chaos/Public/ChaosDebugDraw/ChaosDDScene.h
@@ -20,7 +20,7 @@ namespace ChaosDD::Private
 	class CHAOS_API FChaosDDScene : public TSharedFromThis<FChaosDDScene>
 	{
 	public:
-		FChaosDDScene(const FString& InName, bool bIsServer);
+		FChaosDDScene(const FString& InName);
 		~FChaosDDScene();
 
 		const FString& GetName() const
@@ -28,41 +28,28 @@ namespace ChaosDD::Private
 			return Name;
 		}
 
-		bool IsServer() const;
-		void SetRenderEnabled(bool bInRenderEnabled);
-		bool IsRenderEnabled() const;
-
 		// Specify the region of in which we wish to enable debug draw. A radius of zero means everywhere.
 		void SetDrawRegion(const FSphere3d& InDrawRegion);
 
-		// The region of interest
-		const FSphere3d& GetDrawRegion() const;
-
 		// Set the line budget for debug draw
 		void SetCommandBudget(int32 InCommandBudget);
 
-		// The number of commands we can draw (also max number of lines for now)
-		int32 GetCommandBudget() const;
-
 		// Create a new timeline. E.g., PT, GT, RBAN
 		// The caller must hold a shared pointer to the timeline to keep it alive.
 		FChaosDDTimelinePtr CreateTimeline(const FString& Name);
 
-		// Collect all the latest complete frames for rendering
-		TArray<FChaosDDFramePtr> GetLatestFrames();
+		// Render the most recent frame from each timeline
+		void RenderLatestFrames(IChaosDDRenderer& Renderer, bool bIncludeGlobalFrame);
 
 	private:
 		TArray<FChaosDDFramePtr> GetFrames();
 		void PruneTimelines();
 
-		mutable FCriticalSection TimelinesCS;
-
 		FString Name;
+		FCriticalSection TimelinesCS;
 		TArray<FChaosDDTimelineWeakPtr> Timelines;
 		FSphere3d DrawRegion;
 		int32 CommandBudget;
-		bool bIsServer;
-		bool bRenderEnabled;
 	};
 }
 
diff --git a/Engine/Source/Runtime/Experimental/Chaos/Public/ChaosDebugDraw/ChaosDDTypes.h b/Engine/Source/Runtime/Experimental/Chaos/Public/ChaosDebugDraw/ChaosDDTypes.h
--- a/Engine/Source/Runtime/Experimental/Chaos/Public/ChaosDebugDraw/ChaosDDTypes.h
+++ b/Engine/Source/Runtime/Experimental/Chaos/Public/ChaosDebugDraw/ChaosDDTypes.h
@@ -23,15 +23,12 @@ namespace ChaosDD::Private
 	class FChaosDDGlobalFrame;
 	class FChaosDDScene;
 	class FChaosDDTimeline;
-	class FChaosDDTimelineContext;
-	class FChaosDDTaskContext;
 
 	class IChaosDDRenderer;
 
 	using FChaosDDFramePtr = TSharedPtr<FChaosDDFrame, ESPMode::ThreadSafe>;
 	using FChaosDDGlobalFramePtr = TSharedPtr<FChaosDDGlobalFrame, ESPMode::ThreadSafe>;
 	using FChaosDDScenePtr = TSharedPtr<FChaosDDScene, ESPMode::ThreadSafe>;
-	using FChaosDDSceneWeakPtr = TWeakPtr<FChaosDDScene, ESPMode::ThreadSafe>;
 	using FChaosDDTimelinePtr = TSharedPtr<FChaosDDTimeline, ESPMode::ThreadSafe>;
 	using FChaosDDTimelineWeakPtr = TWeakPtr<FChaosDDTimeline, ESPMode::ThreadSafe>;
 }
diff --git a/Engine/Source/Runtime/Experimental/ChaosSolverEngine/Private/Chaos/ChaosDebugDrawComponent.cpp b/Engine/Source/Runtime/Experimental/ChaosSolverEngine/Private/Chaos/ChaosDebugDrawComponent.cpp
--- a/Engine/Source/Runtime/Experimental/ChaosSolverEngine/Private/Chaos/ChaosDebugDrawComponent.cpp
+++ b/Engine/Source/Runtime/Experimental/ChaosSolverEngine/Private/Chaos/ChaosDebugDrawComponent.cpp
@@ -2,6 +2,8 @@
 
 #include "Chaos/ChaosDebugDrawComponent.h"
 #include "Chaos/DebugDrawQueue.h"
+#include "ChaosDebugDraw/ChaosDDRenderer.h"
+#include "ChaosDebugDraw/ChaosDDScene.h"
 #include "ChaosLog.h"
 #include "Debug/DebugDrawService.h"
 #include "DrawDebugHelpers.h"
@@ -213,6 +215,58 @@ void DebugDrawChaos(const AActor* DebugDrawActor, const TArray<Chaos::FLatentDra
 }
 #endif
 
+#if CHAOS_DEBUG_DRAW
+
+class FChaosDDRenderer : public ChaosDD::Private::IChaosDDRenderer
+{
+public:
+	FChaosDDRenderer(UWorld* InWorld, AActor* InDebugDrawActor)
+		: World(InWorld)
+		, DebugDrawActor(InDebugDrawActor)
+	{
+		DepthPriority = 10;
+		bIsPaused = false;
+	}
+
+	virtual bool IsServer() const override final
+	{
+		return World->GetNetMode() == ENetMode::NM_DedicatedServer;
+	}
+
+	virtual void RenderLine(const FVector3d& A, const FVector3d& B, const FColor& Color, float LineThickness, float Lifetime) const override final
+	{
+		DrawDebugLine(World, A, B, Color, false, CommandLifeTime(Lifetime), DepthPriority, LineThickness);
+	}
+
+	virtual void RenderBox(const FVector3d& Position, const FQuat4d& Rotation, const FVector3d& Size, const FColor& Color, float LineThickness, float Lifetime) const override final
+	{
+		DrawDebugBox(World, Position, Size, Rotation, Color, false, CommandLifeTime(Lifetime), DepthPriority, LineThickness);
+	}
+
+	virtual void RenderLatentCommand(const Chaos::FLatentDrawCommand& Command) const override final
+	{
+		const bool bDrawUe = bChaosDebugDraw_DrawMode != 1;
+		if (bDrawUe)
+		{
+			DebugDrawChaosCommand(World, Command);
+		}
+
+		const bool bDrawVisLog = bChaosDebugDraw_DrawMode != 0;
+		if (bDrawVisLog)
+		{
+			const AActor* Actor = (Command.TestBaseActor) ? Command.TestBaseActor : DebugDrawActor;
+			VisLogChaosCommand(Actor, Command);
+		}
+	}
+
+private:
+	UWorld* World;
+	AActor* DebugDrawActor;
+	int8 DepthPriority;
+	bool bIsPaused;
+};
+
+#endif
 
 
 UChaosDebugDrawComponent::UChaosDebugDrawComponent()
@@ -351,6 +405,23 @@ void UChaosDebugDrawComponent::TickComponent(float DeltaTime, enum ELevelTick Ti
 
 			DebugDrawChaos(GetOwner(), DrawCommands, World->IsPaused());
 		}
+
+		if (World->GetPhysicsScene() != nullptr)
+		{
+			const ChaosDD::Private::FChaosDDScenePtr& CDDScene = World->GetPhysicsScene()->GetDebugDrawScene();
+			if (CDDScene.IsValid())
+			{
+				CDDScene->SetDrawRegion(FSphere3d(RegionOfInterestOrigin, RegionOfInterestRadius));
+				CDDScene->SetCommandBudget(ChaosDebugDraw_MaxElements);
+
+				if (!bIsPaused)
+				{
+					FChaosDDRenderer CDDRenderer = FChaosDDRenderer(World, GetOwner());
+					constexpr bool bRenderGlobalFrame = true;
+					CDDScene->RenderLatestFrames(CDDRenderer, bRenderGlobalFrame);
+				}
+			}
+		}
 	}
 #endif
 }
diff --git a/Engine/Source/Runtime/Experimental/ChaosSolverEngine/Private/Chaos/ChaosDebugDrawSubsystem.cpp b/Engine/Source/Runtime/Experimental/ChaosSolverEngine/Private/Chaos/ChaosDebugDrawSubsystem.cpp
deleted file mode 100644
diff --git a/Engine/Source/Runtime/Experimental/ChaosSolverEngine/Private/Chaos/ChaosSolverEnginePlugin.cpp b/Engine/Source/Runtime/Experimental/ChaosSolverEngine/Private/Chaos/ChaosSolverEnginePlugin.cpp
--- a/Engine/Source/Runtime/Experimental/ChaosSolverEngine/Private/Chaos/ChaosSolverEnginePlugin.cpp
+++ b/Engine/Source/Runtime/Experimental/ChaosSolverEngine/Private/Chaos/ChaosSolverEnginePlugin.cpp
@@ -7,7 +7,6 @@
 #include "Engine/Engine.h"
 #include "Engine/World.h"
 #include "Chaos/ChaosDebugDrawComponent.h"
-#include "Chaos/ChaosDebugDrawSubsystem.h"
 #include "Chaos/ChaosSolverActor.h"
 #include "ChaosSolversModule.h"
 #include "ChaosVDRecordingStateScreenMessageHandler.h"
@@ -29,10 +28,6 @@ void FChaosSolverEnginePlugin::StartupModule()
 
 	UChaosDebugDrawComponent::BindWorldDelegates();
 
-#if CHAOS_DEBUG_DRAW
-	UChaosDebugDrawSubsystem::Startup();
-#endif
-
 #if WITH_CHAOS_VISUAL_DEBUGGER
 	FChaosVDRecordingStateScreenMessageHandler::Get().Initialize();
 #endif
@@ -41,9 +36,6 @@ void FChaosSolverEnginePlugin::StartupModule()
 
 void FChaosSolverEnginePlugin::ShutdownModule()
 {
-#if CHAOS_DEBUG_DRAW
-	UChaosDebugDrawSubsystem::Shutdown();
-#endif
 
 #if WITH_CHAOS_VISUAL_DEBUGGER
 	FChaosVDRecordingStateScreenMessageHandler::Get().TearDown();
diff --git a/Engine/Source/Runtime/Experimental/ChaosSolverEngine/Public/Chaos/ChaosDebugDrawSubsystem.h b/Engine/Source/Runtime/Experimental/ChaosSolverEngine/Public/Chaos/ChaosDebugDrawSubsystem.h
deleted file mode 100644
diff --git a/Engine/Source/Runtime/PhysicsCore/Private/ChaosScene.cpp b/Engine/Source/Runtime/PhysicsCore/Private/ChaosScene.cpp
--- a/Engine/Source/Runtime/PhysicsCore/Private/ChaosScene.cpp
+++ b/Engine/Source/Runtime/PhysicsCore/Private/ChaosScene.cpp
@@ -82,6 +82,11 @@ FChaosScene::FChaosScene(
 	SceneSolver->GetChaosVDContextData().Type = static_cast<int32>(EChaosVDContextType::Solver);
 #endif
 
+#if CHAOS_DEBUG_DRAW
+	CDDScene = MakeShared<ChaosDD::Private::FChaosDDScene>(DebugName.ToString());
+	SceneSolver->SetDebugDrawScene(CDDScene);
+#endif
+
 	SceneSolver->PhysSceneHack = this;
 	SimCallback = SceneSolver->CreateAndRegisterSimCallbackObject_External<FChaosSceneSimCallback>();
 
@@ -576,16 +581,3 @@ FGraphEventArray FChaosScene::GetCompletionEvents()
 {
 	return CompletionEvents;
 }
-
-#if CHAOS_DEBUG_DRAW
-void FChaosScene::SetDebugDrawScene(const ChaosDD::Private::FChaosDDScenePtr& InCDDScene)
-{
-	CDDScene = InCDDScene;
-
-	SceneSolver->EnqueueCommandImmediate(
-		[this]()
-		{
-			SceneSolver->SetDebugDrawScene(CDDScene);
-		});
-}
-#endif
diff --git a/Engine/Source/Runtime/PhysicsCore/Public/Chaos/ChaosScene.h b/Engine/Source/Runtime/PhysicsCore/Public/Chaos/ChaosScene.h
--- a/Engine/Source/Runtime/PhysicsCore/Public/Chaos/ChaosScene.h
+++ b/Engine/Source/Runtime/PhysicsCore/Public/Chaos/ChaosScene.h
@@ -213,8 +213,6 @@ public:
 		return CDDScene;
 	}
 
-	PHYSICSCORE_API void SetDebugDrawScene(const ChaosDD::Private::FChaosDDScenePtr& CDDScene);
-
 private:
 	ChaosDD::Private::FChaosDDScenePtr CDDScene;
 #endif
