[Chaos] Changing the tolerance used when we calculate the normal of a traingle to 1.e-8f (it was using 1.e-4f ).

During cook we are accepting triangles with a squared area less smaller than 1.e-4f. I added a local visualizer in CVD to see if these small triangles still generate valid normals and they do, so we decided to change the tolerance used at runtime to be 1.e-8f

diff --git a/Engine/Source/Runtime/Experimental/Chaos/Private/Chaos/TriangleMeshImplicitObject.cpp b/Engine/Source/Runtime/Experimental/Chaos/Private/Chaos/TriangleMeshImplicitObject.cpp
--- a/Engine/Source/Runtime/Experimental/Chaos/Private/Chaos/TriangleMeshImplicitObject.cpp
+++ b/Engine/Source/Runtime/Experimental/Chaos/Private/Chaos/TriangleMeshImplicitObject.cpp
@@ -143,7 +143,7 @@ struct FTriangleMeshRaycastVisitor
 	 */
 	bool VisitRaycast(TSpatialVisitorData<int32> TriIdx, FRealSingle& CurDataLength)
 	{
-		constexpr FReal Epsilon = 1e-4f;
+		constexpr FReal Epsilon = UE_SMALL_NUMBER;
 
 		const int32 FaceIndex = TriIdx.Payload;
 		const FVec3& A = Particles.GetX(Elements[FaceIndex][0]);
@@ -182,7 +182,7 @@ struct FTriangleMeshRaycastVisitor
 	
 	bool VisitSweep(TSpatialVisitorData<int32> TriIdx, FRealSingle& CurDataLength)
 	{
-		constexpr FReal Epsilon = 1e-4f;
+		constexpr FReal Epsilon = UE_SMALL_NUMBER;
 		constexpr FReal Epsilon2 = Epsilon * Epsilon;
 		
 		const FReal Thickness2 = Thickness * Thickness;
@@ -200,7 +200,7 @@ struct FTriangleMeshRaycastVisitor
 		const FVec3 AB = B - A;
 		const FVec3 AC = C - A;
 		FVec3 TriNormal = FVec3::CrossProduct(AB, AC);
-		const FReal NormalLength = TriNormal.SafeNormalize();
+		const FReal NormalLength = TriNormal.SafeNormalize(UE_SMALL_NUMBER);
 		if (!CHAOS_ENSURE(NormalLength > Epsilon))
 		{
 			//hitting degenerate triangle so keep searching - should be fixed before we get to this stage
@@ -1037,7 +1037,7 @@ bool FTriangleMeshImplicitObject::OverlapImp(const TArray<TVec3<IdxType>>& Eleme
 	QueryBounds.Thicken(Thickness);
 	const TArray<int32> PotentialIntersections = FastBVH.FindAllIntersections(QueryBounds);
 
-	const FReal Epsilon = 1e-4f;
+	constexpr FReal Epsilon = UE_SMALL_NUMBER;
 	//ensure(Thickness > Epsilon);	//There's no hope for this to work unless thickness is large (really a sphere overlap test)
 	//todo: turn ensure back on, off until some other bug is fixed
 
@@ -1050,7 +1050,7 @@ bool FTriangleMeshImplicitObject::OverlapImp(const TArray<TVec3<IdxType>>& Eleme
 		const FVec3 AB = B - A;
 		const FVec3 AC = C - A;
 		FVec3 Normal = FVec3::CrossProduct(AB, AC);
-		const FReal NormalLength = Normal.SafeNormalize();
+		const FReal NormalLength = Normal.SafeNormalize(UE_SMALL_NUMBER);
 		if (!CHAOS_ENSURE(NormalLength > Epsilon))
 		{
 			//hitting degenerate triangle - should be fixed before we get to this stage
@@ -1781,7 +1781,7 @@ int32 FTriangleMeshImplicitObject::FindMostOpposingFace(const TArray<TVec3<IdxTy
 	const FAABB3 QueryBounds(ScaledPosition - AbsScaledSearchDist, ScaledPosition + AbsScaledSearchDist);
 
 	const TArray<int32> PotentialIntersections = FastBVH.FindAllIntersections(QueryBounds);
-	const FReal Epsilon = 1e-4f;
+	constexpr FReal Epsilon = UE_SMALL_NUMBER;
 
 	FReal MostOpposingDot = TNumericLimits<FReal>::Max();
 	int32 MostOpposingFace = HintFaceIndex;
@@ -1795,7 +1795,7 @@ int32 FTriangleMeshImplicitObject::FindMostOpposingFace(const TArray<TVec3<IdxTy
 		const FVec3 AB = B - A;
 		const FVec3 AC = C - A;
 		FVec3 Normal = FVec3::CrossProduct(AB, AC);
-		const FReal NormalLength = Normal.SafeNormalize();
+		const FReal NormalLength = Normal.SafeNormalize(UE_SMALL_NUMBER);
 		if (!CHAOS_ENSURE(NormalLength > Epsilon))
 		{
 			//hitting degenerate triangle - should be fixed before we get to this stage
@@ -1946,7 +1946,7 @@ FVec3 FTriangleMeshImplicitObject::GetFaceNormal(const int32 FaceIdx) const
 			const ParticleVecType AC = C - A;
 			ParticleVecType Normal = ParticleVecType::CrossProduct(AB, AC);
 			
-			if(Normal.SafeNormalize() < UE_SMALL_NUMBER)
+			if(Normal.SafeNormalize(UE_SMALL_NUMBER) < UE_SMALL_NUMBER)
 			{
 				UE_LOG(LogChaos, Warning, TEXT("Degenerate triangle %d: (%f %f %f) (%f %f %f) (%f %f %f)"), FaceIdx, A.X, A.Y, A.Z, B.X, B.Y, B.Z, C.X, C.Y, C.Z);
 				CHAOS_ENSURE(false);
