Chaos Debug Draw - custom debug draw objects for particles, implicits, primitives

diff --git a/Engine/Source/Runtime/Experimental/Chaos/Private/Chaos/DebugDraw/DebugDrawImplicitObject.cpp b/Engine/Source/Runtime/Experimental/Chaos/Private/Chaos/DebugDraw/DebugDrawImplicitObject.cpp
new file mode 100644
--- /dev/null
+++ b/Engine/Source/Runtime/Experimental/Chaos/Private/Chaos/DebugDraw/DebugDrawImplicitObject.cpp
@@ -0,0 +1,183 @@
+// Copyright Epic Games, Inc. All Rights Reserved.
+#include "Chaos/DebugDraw/DebugDrawImplicitObject.h"
+#include "Chaos/Box.h"
+#include "Chaos/Collision/ContactTriangles.h"
+#include "Chaos/Capsule.h"
+#include "Chaos/Convex.h"
+#include "Chaos/DebugDraw/DebugDrawPrimitives.h"
+#include "Chaos/HeightField.h"
+#include "Chaos/ImplicitObject.h"
+#include "Chaos/ImplicitObjectScaled.h"
+#include "Chaos/ImplicitObjectTransformed.h"
+#include "Chaos/Sphere.h"
+#include "Chaos/TriangleMeshImplicitObject.h"
+#include "ChaosDebugDraw/ChaosDDContext.h"
+#include "ChaosDebugDraw/ChaosDDFrame.h"
+#include "ChaosDebugDraw/ChaosDDRenderer.h"
+
+#if CHAOS_DEBUG_DRAW
+
+namespace Chaos
+{
+	void FChaosDDImplicitObject::Draw(
+		const FConstImplicitObjectPtr& Implicit,
+		const FRigidTransform3& Transform,
+		const FColor& Color,
+		float LineThickness, 
+		float Duration)
+	{
+		ChaosDD::Private::FChaosDDContext::GetWriter().EnqueueCommand(
+			[=](ChaosDD::Private::IChaosDDRenderer& Renderer)
+			{
+				Chaos::Private::ChaosDDRenderImplicitObject(Renderer, Implicit, Transform, Color, LineThickness, Duration);
+			});
+	}
+}
+
+namespace Chaos::Private
+{
+	void ChaosDDRenderSphere(
+		ChaosDD::Private::IChaosDDRenderer& Renderer,
+		const FImplicitSphere3& Sphere,
+		const FRigidTransform3& Transform,
+		const FColor& Color,
+		float LineThickness,
+		float Duration)
+	{
+		const FVec3 Scale = Transform.GetScale3D();
+		Renderer.RenderSphere(Transform.TransformPosition(Sphere.GetCenter()), float(Scale.Z * Sphere.GetRadius()), Color, LineThickness, Duration);
+	}
+
+	void ChaosDDRenderCapsule(
+		ChaosDD::Private::IChaosDDRenderer& Renderer,
+		const FImplicitCapsule3& Capsule,
+		const FRigidTransform3& Transform,
+		const FColor& Color,
+		float LineThickness,
+		float Duration)
+	{
+		const FVec3 Scale = Transform.GetScale3D();
+		const FVec3 P = Transform.TransformPosition(Capsule.GetCenter());
+		const FRotation3 Q = Transform.GetRotation() * FRotationMatrix::MakeFromZ(Capsule.GetAxis());
+		Renderer.RenderCapsule(P, Q, float(Scale.Z * (0.5 * Capsule.GetHeight() + Capsule.GetRadius())), float(Scale.X * Capsule.GetRadius()), Color, LineThickness, Duration);
+	}
+
+	void ChaosDDRenderBox(
+		ChaosDD::Private::IChaosDDRenderer& Renderer,
+		const FImplicitBox3& Box,
+		const FRigidTransform3& Transform,
+		const FColor& Color,
+		float LineThickness,
+		float Duration)
+	{
+		const FVector3d Size = 0.5 * Transform.GetScale3D() * Box.Extents();
+		const FVector3d Center = Transform.TransformPosition(Box.Center());
+
+		Renderer.RenderBox(Center, Transform.GetRotation(), Size, Color, LineThickness, Duration);
+	}
+
+	void ChaosDDRenderConvex(
+		ChaosDD::Private::IChaosDDRenderer& Renderer,
+		const FImplicitConvex3& Convex,
+		const FRigidTransform3& Transform,
+		const FColor& Color,
+		float LineThickness,
+		float Duration)
+	{
+		for (int32 EdgeIndex = 0; EdgeIndex < Convex.NumEdges(); ++EdgeIndex)
+		{
+			const int32 EdgeVertexIndex0 = Convex.GetEdgeVertex(EdgeIndex, 0);
+			const int32 EdgeVertexIndex1 = Convex.GetEdgeVertex(EdgeIndex, 1);
+			const FVec3 EdgeVertex0 = Transform.TransformPosition(FVec3(Convex.GetVertex(EdgeVertexIndex0)));
+			const FVec3 EdgeVertex1 = Transform.TransformPosition(FVec3(Convex.GetVertex(EdgeVertexIndex1)));
+
+			Renderer.RenderLine(EdgeVertex0, EdgeVertex1, Color, LineThickness, Duration);
+		}
+	}
+
+	template<typename MeshType>
+	void ChaosDDRenderMesh(
+		ChaosDD::Private::IChaosDDRenderer& Renderer,
+		const MeshType& Mesh,
+		const FRigidTransform3& Transform,
+		const FColor& Color,
+		float LineThickness,
+		float Duration)
+	{
+		const FSphere3d RegionOfInterest = Renderer.GetDrawRegion();
+		const FAABB3 WorldQueryBounds = FAABB3(RegionOfInterest.Center - FVec3(RegionOfInterest.W), RegionOfInterest.Center + FVec3(RegionOfInterest.W));
+		const FAABB3 LocalQueryBounds = WorldQueryBounds.InverseTransformedAABB(Transform);
+
+		Mesh.VisitTriangles(LocalQueryBounds, Transform, 
+			[&Renderer, &Color, LineThickness, Duration, &RegionOfInterest](const FTriangle& Tri, const int32 TriangleIndex, const int32 VertexIndex0, const int32 VertexIndex1, const int32 VertexIndex2)
+			{
+				const FReal RegionDistanceSq = (Tri.GetCentroid() - FVec3(RegionOfInterest.Center)).SizeSquared();
+				if (RegionDistanceSq < FMath::Square(RegionOfInterest.W))
+				{
+					// @todo(chaos): could keep track of edges drawn so we don't add every shared edge twice...
+					Renderer.RenderTriangle(Tri[0], Tri[1], Tri[2], Color, LineThickness, Duration);
+				}
+			});
+	}
+
+	void ChaosDDRenderLeafImplicitObject(
+		ChaosDD::Private::IChaosDDRenderer& Renderer,
+		const FImplicitObject* Implicit,
+		const FRigidTransform3& Transform,
+		const FColor& Color,
+		float LineThickness,
+		float Duration)
+	{
+		if (const FImplicitSphere3* Sphere = Implicit->AsA<FImplicitSphere3>())
+		{
+			ChaosDDRenderSphere(Renderer, *Sphere, Transform, Color, LineThickness, Duration);
+		}
+		else if (const FImplicitCapsule3* Capsule = Implicit->AsA<FImplicitCapsule3>())
+		{
+			ChaosDDRenderCapsule(Renderer, *Capsule, Transform, Color, LineThickness, Duration);
+		}
+		else if (const FImplicitBox3* Box = Implicit->AsA<FImplicitBox3>())
+		{
+			ChaosDDRenderBox(Renderer, *Box, Transform, Color, LineThickness, Duration);
+		}
+		else if (const FImplicitConvex3* Convex = Implicit->AsA<FImplicitConvex3>())
+		{
+			ChaosDDRenderConvex(Renderer, *Convex, Transform, Color, LineThickness, Duration);
+		}
+		else if (const FTriangleMeshImplicitObject* TriMesh = Implicit->AsA<FTriangleMeshImplicitObject>())
+		{
+			ChaosDDRenderMesh(Renderer, *TriMesh, Transform, Color, LineThickness, Duration);
+		}
+		else if (const FHeightField* HeightField = Implicit->AsA<FHeightField>())
+		{
+			ChaosDDRenderMesh(Renderer, *HeightField, Transform, Color, LineThickness, Duration);
+		}
+		else if (const FImplicitObjectInstanced* Instanced = Implicit->AsA<FImplicitObjectInstanced>())
+		{
+			ChaosDDRenderLeafImplicitObject(Renderer, Instanced->GetInnerObject().Get(), Transform, Color, LineThickness, Duration);
+		}
+		else if (const FImplicitObjectScaled* Scaled = Implicit->AsA<FImplicitObjectScaled>())
+		{
+			const FRigidTransform3 ScaledTransform = FRigidTransform3(Transform.GetTranslation(), Transform.GetRotation(), Transform.GetScale3D() * Scaled->GetScale());
+			ChaosDDRenderLeafImplicitObject(Renderer, Scaled->GetInnerObject().Get(), ScaledTransform, Color, LineThickness, Duration);
+		}
+	}
+
+	void ChaosDDRenderImplicitObject(
+		ChaosDD::Private::IChaosDDRenderer& Renderer,
+		const FConstImplicitObjectPtr& RootImplicit,
+		const FRigidTransform3& RootTransform,
+		const FColor& Color,
+		float LineThickness,
+		float Duration)
+	{
+		RootImplicit->VisitLeafObjects(
+			[&](const FImplicitObject* LeafImplicitObject, const FRigidTransform3& LeafRelativeTransform, const int32 UnusedRootObjectIndex, const int32 UnusedObjectIndex, const int32 UnusedLeafObjectIndex)
+			{
+				const FRigidTransform3 LeafTransform = LeafRelativeTransform * RootTransform;
+				ChaosDDRenderLeafImplicitObject(Renderer, LeafImplicitObject, LeafTransform, Color, LineThickness, Duration);
+			});
+	}
+}
+
+#endif
\ No newline at end of file
diff --git a/Engine/Source/Runtime/Experimental/Chaos/Private/Chaos/DebugDraw/DebugDrawParticle.cpp b/Engine/Source/Runtime/Experimental/Chaos/Private/Chaos/DebugDraw/DebugDrawParticle.cpp
new file mode 100644
--- /dev/null
+++ b/Engine/Source/Runtime/Experimental/Chaos/Private/Chaos/DebugDraw/DebugDrawParticle.cpp
@@ -0,0 +1,351 @@
+// Copyright Epic Games, Inc. All Rights Reserved.
+#include "Chaos/DebugDraw/DebugDrawParticle.h"
+#include "Chaos/ConvexOptimizer.h"
+#include "Chaos/DebugDraw/DebugDrawImplicitObject.h"
+#include "Chaos/GeometryParticles.h"
+#include "Chaos/Particle/ObjectState.h"
+#include "Chaos/ParticleHandle.h"
+#include "Chaos/ParticleIterator.h"
+#include "Chaos/PBDRigidParticles.h"
+#include "Chaos/ShapeInstance.h"
+#include "Chaos/Island/IslandManager.h"
+#include "ChaosDebugDraw/ChaosDDContext.h"
+#include "ChaosDebugDraw/ChaosDDFrame.h"
+#include "ChaosDebugDraw/ChaosDDRenderer.h"
+
+// For FChaosDebugDrawSettings - move that class when the old debug draw is no longer used
+#include "Chaos/ChaosDebugDraw.h"
+
+#if CHAOS_DEBUG_DRAW
+
+namespace Chaos::DebugDraw
+{
+	extern bool bChaosDebugDebugDrawColorShapesByShapeType;
+	extern bool bChaosDebugDebugDrawColorShapesByIsland;
+	extern bool bChaosDebugDebugDrawColorShapesByInternalCluster;
+	extern bool bChaosDebugDebugDrawColorShapesBySimQueryType;
+	extern bool bChaosDebugDebugDrawColorShapesByConvexType;
+	extern bool bChaosDebugDebugDrawColorShapesByClusterUnion;
+	
+	extern bool bChaosDebugDebugDrawColorBoundsByShapeType;
+
+	extern bool bChaosDebugDebugDrawShowQueryOnlyShapes;
+	extern bool bChaosDebugDebugDrawShowSimOnlyShapes;
+	extern bool bChaosDebugDebugDrawShowProbeOnlyShapes;
+
+}
+
+namespace Chaos::CVars
+{
+	extern int32 ChaosSolverDrawShapesShowStatic;
+	extern int32 ChaosSolverDrawShapesShowKinematic;
+	extern int32 ChaosSolverDrawShapesShowDynamic;
+
+	extern Chaos::DebugDraw::FChaosDebugDrawSettings ChaosSolverDebugDebugDrawSettings;
+}
+
+namespace Chaos
+{
+	FColor GetIndexColor(const int32 Index)
+	{
+		static const FColor Colors[] =
+		{
+			FColor::Red,
+			FColor::Orange,
+			FColor::Yellow,
+			FColor::Green,
+			FColor::Emerald,
+			FColor::Cyan,
+			FColor::Turquoise,
+			FColor::Blue,
+			FColor::Magenta,
+			FColor::Purple,
+		};
+		const int32 NumColors = UE_ARRAY_COUNT(Colors);
+
+		return Colors[Index % NumColors];
+	};
+
+
+	FColor GetIslandColor(const int32 IslandIndex, const bool bIsAwake)
+	{
+		static const FColor SleepingColor = FColor::Black;
+		static const FColor NullColor = FColor::White;
+
+		if (IslandIndex == INDEX_NONE)
+		{
+			return NullColor;
+		}
+
+		if (!bIsAwake)
+		{
+			return SleepingColor;
+		}
+
+		return GetIndexColor(IslandIndex);
+	};
+
+	struct FChaosDDParticleData
+	{
+		int32 IslandId;
+		int32 ClusterId;
+		EObjectStateType ObjectState;
+		uint8 bIsClusterUnion : 1;
+		uint8 bIsInternalCluster : 1;
+		uint8 bIsOneWay : 1;
+
+		FChaosDDParticleData(
+			const FConstGenericParticleHandle& InParticle)
+		{
+			IslandId = (InParticle->GetConstraintGraphNode() != nullptr) ? InParticle->GetConstraintGraphNode()->GetIslandId() : INDEX_NONE;
+			ClusterId = (InParticle->CastToClustered() != nullptr) ? FMath::Abs(InParticle->CastToClustered()->ClusterGroupIndex()) : INDEX_NONE;
+			ObjectState = InParticle->ObjectState();
+			bIsClusterUnion = (InParticle->CastToClustered() != nullptr) ? (InParticle->CastToClustered()->PhysicsProxy()->GetType() == EPhysicsProxyType::ClusterUnionProxy) : false;
+			bIsInternalCluster = (InParticle->CastToClustered() != nullptr) ? InParticle->CastToClustered()->InternalCluster() : false;
+			bIsOneWay = InParticle->OneWayInteraction();
+		}
+	};
+
+	struct FChaosDDShapeData
+	{
+		EChaosCollisionTraceFlag CollisionTraceFlag;
+		uint8 bIsQuery: 1;
+		uint8 bIsSim: 1;
+		uint8 bIsProbe: 1;
+		uint8 bIsOptimized: 1;
+
+		FChaosDDShapeData(
+			const FShapeInstancePtr& InShapeInstance, 
+			bool bInIsOptimized)
+		{
+			CollisionTraceFlag = InShapeInstance->GetCollisionTraceType();
+			bIsQuery = InShapeInstance->GetQueryEnabled();
+			bIsSim = InShapeInstance->GetSimEnabled();
+			bIsProbe = InShapeInstance->GetIsProbe();
+			bIsOptimized = bInIsOptimized;
+		}
+	};
+
+	// Helper for collecting the data used to render particle shapes
+	struct FChaosDDParticleShape
+	{
+	public:
+		// Determine whether we should render based on the flags
+		static bool ShouldRender(const FConstImplicitObjectPtr& Implicit, FChaosDDParticleData ParticleData, FChaosDDShapeData ShapeData, bool bIsServer)
+		{
+			const Chaos::DebugDraw::FChaosDebugDrawSettings& Settings = Chaos::CVars::ChaosSolverDebugDebugDrawSettings;
+
+			if (!Chaos::CVars::ChaosSolverDrawShapesShowStatic && ParticleData.ObjectState == EObjectStateType::Static)
+			{
+				return false;
+			}
+			if (!Chaos::CVars::ChaosSolverDrawShapesShowKinematic && ParticleData.ObjectState == EObjectStateType::Kinematic)
+			{
+				return false;
+			}
+			if (!Chaos::CVars::ChaosSolverDrawShapesShowDynamic && ParticleData.ObjectState == EObjectStateType::Dynamic)
+			{
+				return false;
+			}
+
+			if (!Chaos::DebugDraw::bChaosDebugDebugDrawShowQueryOnlyShapes && ShapeData.bIsQuery && !ShapeData.bIsSim && !ShapeData.bIsProbe)
+			{
+				return false;
+			}
+			if (!Chaos::DebugDraw::bChaosDebugDebugDrawShowSimOnlyShapes && !ShapeData.bIsQuery && ShapeData.bIsSim && !ShapeData.bIsProbe)
+			{
+				return false;
+			}
+			if (!Chaos::DebugDraw::bChaosDebugDebugDrawShowProbeOnlyShapes && !ShapeData.bIsQuery && !ShapeData.bIsSim && ShapeData.bIsProbe)
+			{
+				return false;
+			}
+
+			// Depending on the shape settings, we may not show the simple or complex shape
+			const EImplicitObjectType InnerType = GetInnerType(Implicit->GetType());
+			const bool bIsMesh = (InnerType == ImplicitObjectType::TriangleMesh);
+			const bool bShowMeshes = (Settings.bShowComplexCollision && (ShapeData.CollisionTraceFlag != EChaosCollisionTraceFlag::Chaos_CTF_UseSimpleAsComplex))
+				|| (Settings.bShowSimpleCollision && (ShapeData.CollisionTraceFlag == EChaosCollisionTraceFlag::Chaos_CTF_UseComplexAsSimple));
+			const bool bShowNonMeshes = (Settings.bShowSimpleCollision && (ShapeData.CollisionTraceFlag != EChaosCollisionTraceFlag::Chaos_CTF_UseComplexAsSimple))
+				|| (Settings.bShowComplexCollision && (ShapeData.CollisionTraceFlag == EChaosCollisionTraceFlag::Chaos_CTF_UseSimpleAsComplex));
+			if (bIsMesh && !bShowMeshes)
+			{
+				return false;
+			}
+			else if (!bIsMesh && !bShowNonMeshes)
+			{
+				return false;
+			}
+
+			return true;
+		}
+
+		// Determine the color from the flags
+		static FColor GetRenderColor(const FConstImplicitObjectPtr& Implicit, FChaosDDParticleData ParticleData, FChaosDDShapeData ShapeData, bool bIsServer)
+		{
+			const Chaos::DebugDraw::FChaosDebugDrawSettings& Settings = Chaos::CVars::ChaosSolverDebugDebugDrawSettings;
+
+			const EImplicitObjectType InnerType = GetInnerType(Implicit->GetType());
+			if (Chaos::DebugDraw::bChaosDebugDebugDrawColorShapesByShapeType)
+			{
+				return Settings.ShapesColorsPerShapeType.GetColorFromShapeType(InnerType);
+			}
+			if (Chaos::DebugDraw::bChaosDebugDebugDrawColorShapesByIsland)
+			{
+				return GetIslandColor(ParticleData.IslandId, true);
+			}
+			if (Chaos::DebugDraw::bChaosDebugDebugDrawColorShapesByInternalCluster)
+			{
+				if (ParticleData.bIsClusterUnion)
+				{
+					if (Chaos::DebugDraw::bChaosDebugDebugDrawColorShapesByClusterUnion)
+					{
+						return GetIndexColor(ParticleData.ClusterId);
+					}
+					else
+					{
+						if (ParticleData.bIsInternalCluster)
+						{
+							return FColor::Purple;
+						}
+					}
+				}
+				return FColor::Black;
+			}
+			if (Chaos::DebugDraw::bChaosDebugDebugDrawColorShapesByConvexType)
+			{
+				if (InnerType == ImplicitObjectType::Convex)
+				{
+					if (ShapeData.bIsOptimized)
+					{
+						return FColor::Green;
+					}
+					return FColor::Orange;
+				}
+			}
+
+			return Settings.ShapesColorsPerState.GetColorFromState(ParticleData.ObjectState, ParticleData.bIsOneWay);
+		}
+
+		static float GetLineThickness()
+		{
+			const Chaos::DebugDraw::FChaosDebugDrawSettings& Settings = Chaos::CVars::ChaosSolverDebugDebugDrawSettings;
+
+			return Settings.LineThickness;
+		}
+
+		static float GetDuration()
+		{
+			return 0.0f;
+		}
+
+		static int32 GetParticleCommandCost(const FConstGenericParticleHandle& InParticle)
+		{
+			return 1;
+		}
+
+		static void Draw(
+			const FRigidTransform3& SpaceTransform, 
+			const FConstGenericParticleHandle& InParticle,
+			const FShapeInstancePtr& InShapeInstance,
+			bool bInIsOptimized,
+			bool bAutoColor,
+			const FColor& InColor = FColor::Purple)
+		{
+			const int32 Cost = GetParticleCommandCost(InParticle);
+			const FBox3d Bounds = FBox3d((InParticle->WorldSpaceInflatedBounds().Min()), InParticle->WorldSpaceInflatedBounds().Max()).TransformBy(SpaceTransform);
+
+			ChaosDD::Private::FChaosDDFrameWriter Writer = ChaosDD::Private::FChaosDDContext::GetWriter();
+
+			if (!Writer.IsInDrawRegion(Bounds))
+			{
+				return;
+			}
+
+			if (!Writer.AddToCost(Cost))
+			{
+				return;
+			}
+
+			const FRigidTransform3 ParticleTransform = InParticle->GetTransformPQ() * SpaceTransform;
+			const FChaosDDParticleData ParticleData = FChaosDDParticleData(InParticle);
+			const FChaosDDShapeData ShapeData = FChaosDDShapeData(InShapeInstance, bInIsOptimized);
+
+			InShapeInstance->GetGeometry()->VisitLeafObjects(
+				[&Writer, &ParticleTransform, &ParticleData, &ShapeData, &InColor, &bAutoColor](const FImplicitObject* LeafImplicitObject, const FRigidTransform3& LeafRelativeTransform, const int32 UnusedRootObjectIndex, const int32 UnusedObjectIndex, const int32 UnusedLeafObjectIndex)
+				{
+					Writer.EnqueueCommand(
+						[
+							Transform = LeafRelativeTransform * ParticleTransform,
+							ImplicitObject = Chaos::FConstImplicitObjectPtr(LeafImplicitObject, /*bAddRef*/true),
+							ParticleData,
+							ShapeData,
+							Color = InColor,
+							bAutoColor
+						]
+						(ChaosDD::Private::IChaosDDRenderer& Renderer)
+						{
+							if (ShouldRender(ImplicitObject, ParticleData, ShapeData, Renderer.IsServer()))
+							{
+								Chaos::Private::ChaosDDRenderImplicitObject(
+									Renderer,
+									ImplicitObject,
+									Transform,
+									bAutoColor ? GetRenderColor(ImplicitObject, ParticleData, ShapeData, Renderer.IsServer()) : Color,
+									GetLineThickness(),
+									GetDuration());
+							}
+						});
+
+				});
+
+		}
+	};
+
+	void FChaosDDParticle::DrawShapes(
+		const FRigidTransform3& SpaceTransform,
+		const FGeometryParticleHandle* InParticleHandle)
+	{
+		for (const Chaos::FShapeInstancePtr& ShapeInstance : InParticleHandle->ShapeInstances())
+		{
+			FChaosDDParticleShape::Draw(SpaceTransform, InParticleHandle, ShapeInstance, /*bIsOptimized*/false, /*bAutoColor*/true);
+		}
+	}
+
+	void FChaosDDParticle::DrawShapes(
+		const FGeometryParticleHandle* InParticleHandle)
+	{
+		for (const Chaos::FShapeInstancePtr& ShapeInstance : InParticleHandle->ShapeInstances())
+		{
+			FChaosDDParticleShape::Draw(FRigidTransform3::Identity, InParticleHandle, ShapeInstance, /*bIsOptimized*/false, /*bAutoColor*/true);
+		}
+	}
+
+	void FChaosDDParticle::DrawShapes(
+		const FGeometryParticleHandle* InParticleHandle,
+		const FColor& Color)
+	{
+		for (const Chaos::FShapeInstancePtr& ShapeInstance : InParticleHandle->ShapeInstances())
+		{
+			FChaosDDParticleShape::Draw(FRigidTransform3::Identity, InParticleHandle, ShapeInstance, /*bIsOptimized*/false, /*bAutoColor*/false, Color);
+		}
+	}
+
+	bool FChaosDDParticle::DrawOptimizedShapes(
+		const FGeometryParticleHandle* InParticleHandle)
+	{
+		const Private::FConvexOptimizer* ConvexOptimizer = InParticleHandle->CastToClustered() ? InParticleHandle->CastToClustered()->ConvexOptimizer().Get() : nullptr;
+		if (ConvexOptimizer != nullptr)
+		{
+			for (const Chaos::FShapeInstancePtr& ShapeInstance : ConvexOptimizer->GetShapeInstances())
+			{
+				FChaosDDParticleShape::Draw(FRigidTransform3::Identity, InParticleHandle, ShapeInstance, /*bIsOptimized*/true, /*bAutoColor*/true);
+			}
+			return true;
+		}
+		return false;
+	}
+
+}
+
+#endif
\ No newline at end of file
diff --git a/Engine/Source/Runtime/Experimental/Chaos/Private/Chaos/DebugDraw/DebugDrawPrimitives.cpp b/Engine/Source/Runtime/Experimental/Chaos/Private/Chaos/DebugDraw/DebugDrawPrimitives.cpp
new file mode 100644
--- /dev/null
+++ b/Engine/Source/Runtime/Experimental/Chaos/Private/Chaos/DebugDraw/DebugDrawPrimitives.cpp
@@ -0,0 +1,76 @@
+// Copyright Epic Games, Inc. All Rights Reserved.
+#include "Chaos/DebugDraw/DebugDrawPrimitives.h"
+#include "ChaosDebugDraw/ChaosDDContext.h"
+#include "ChaosDebugDraw/ChaosDDFrame.h"
+#include "ChaosDebugDraw/ChaosDDRenderer.h"
+
+#if CHAOS_DEBUG_DRAW
+
+namespace Chaos
+{
+	void FChaosDDPrimitives::DrawLine(const FVector3d& A, const FVector3d& B, const FColor& Color, float LineThickness, float Duration)
+	{
+		constexpr int32 Cost = 1;
+		const FBox3d Bounds = FBox3d(FVector3d::Min(A, B), FVector3d::Max(A, B));
+
+		ChaosDD::Private::FChaosDDContext::GetWriter().TryEnqueueCommand(Cost, Bounds,
+			[=](ChaosDD::Private::IChaosDDRenderer& Renderer)
+			{
+				Renderer.RenderLine(A, B, Color, LineThickness, Duration);
+			});
+	}
+
+	
+	void FChaosDDPrimitives::DrawSphere(const FVector3d& Center, const float Radius, const FColor& Color, float LineThickness, float Duration)
+	{
+		constexpr int32 Cost = 64;
+		const FBox3d Bounds = FBox3d(Center - FVector3d(Radius), Center + FVector3d(Radius));
+
+		ChaosDD::Private::FChaosDDContext::GetWriter().TryEnqueueCommand(Cost, Bounds,
+			[=](ChaosDD::Private::IChaosDDRenderer& Renderer)
+			{
+				Renderer.RenderSphere(Center, Radius, Color, LineThickness, Duration);
+			});
+	}
+
+	void FChaosDDPrimitives::DrawCapsule(const FVector3d& Center, const FQuat4d& Rotation, float HalfHeight, float Radius, const FColor& Color, float LineThickness, float Duration)
+	{
+		constexpr int32 Cost = 16;
+		const FVector3d EndOffset = HalfHeight * (Rotation * FVector3d::UnitZ());
+		const FVector3d A = Center - EndOffset;
+		const FVector3d B = Center + EndOffset;
+		const FBox3d Bounds = FBox3d(FVector3d::Min(A, B), FVector3d::Max(A, B));
+
+		ChaosDD::Private::FChaosDDContext::GetWriter().TryEnqueueCommand(Cost, Bounds,
+			[=](ChaosDD::Private::IChaosDDRenderer& Renderer)
+			{
+				Renderer.RenderCapsule(Center, Rotation, HalfHeight, Radius, Color, LineThickness, Duration);
+			});
+	}
+
+	void FChaosDDPrimitives::DrawBox(const FVector3d& Center, const FQuat4d& Rotation, const FVector3d& Size, const FColor& Color, float LineThickness, float Duration)
+	{
+		constexpr int32 Cost = 12;
+		const FBox3d Bounds = FBox3d(-0.5 * Size, 0.5 * Size).TransformBy(FTransform(Rotation, Center));
+
+		ChaosDD::Private::FChaosDDContext::GetWriter().TryEnqueueCommand(Cost, Bounds,
+			[=](ChaosDD::Private::IChaosDDRenderer& Renderer)
+			{
+				Renderer.RenderBox(Center, Rotation, Size, Color, LineThickness, Duration);
+			});
+	}
+
+	void FChaosDDPrimitives::DrawTriangle(const FVector3d& A, const FVector3d& B, const FVector3d& C, const FColor& Color, float LineThickness, float Duration)
+	{
+		constexpr int32 Cost = 3;
+		const FBox3d Bounds = FBox3d(FVector3d::Min3(A, B, C), FVector3d::Max3(A, B, C));
+
+		ChaosDD::Private::FChaosDDContext::GetWriter().TryEnqueueCommand(Cost, Bounds,
+			[=](ChaosDD::Private::IChaosDDRenderer& Renderer)
+			{
+				Renderer.RenderTriangle(A, B, C, Color, LineThickness, Duration);
+			});
+	}
+}
+
+#endif
\ No newline at end of file
diff --git a/Engine/Source/Runtime/Experimental/Chaos/Public/Chaos/DebugDraw/DebugDrawImplicitObject.h b/Engine/Source/Runtime/Experimental/Chaos/Public/Chaos/DebugDraw/DebugDrawImplicitObject.h
new file mode 100644
--- /dev/null
+++ b/Engine/Source/Runtime/Experimental/Chaos/Public/Chaos/DebugDraw/DebugDrawImplicitObject.h
@@ -0,0 +1,36 @@
+// Copyright Epic Games, Inc. All Rights Reserved.
+#pragma once
+
+#include "Chaos/Core.h"
+#include "Chaos/Declares.h"
+#include "Chaos/ImplicitFwd.h"
+#include "ChaosDebugDraw/ChaosDDTypes.h"
+#include "Math/MathFwd.h"
+
+#if CHAOS_DEBUG_DRAW
+
+namespace Chaos
+{
+	struct CHAOS_API FChaosDDImplicitObject
+	{
+		static void Draw(
+			const FConstImplicitObjectPtr& Implicit,
+			const FRigidTransform3& Transform, 
+			const FColor& Color, 
+			float LineThickness, 
+			float Duration);
+	};
+}
+
+namespace Chaos::Private
+{
+	void ChaosDDRenderImplicitObject(
+		ChaosDD::Private::IChaosDDRenderer& Renderer, 
+		const FConstImplicitObjectPtr& Implicit, 
+		const FRigidTransform3& Transform,
+		const FColor& Color,
+		float LineThickness, 
+		float Duration);
+}
+
+#endif
\ No newline at end of file
diff --git a/Engine/Source/Runtime/Experimental/Chaos/Public/Chaos/DebugDraw/DebugDrawParticle.h b/Engine/Source/Runtime/Experimental/Chaos/Public/Chaos/DebugDraw/DebugDrawParticle.h
new file mode 100644
--- /dev/null
+++ b/Engine/Source/Runtime/Experimental/Chaos/Public/Chaos/DebugDraw/DebugDrawParticle.h
@@ -0,0 +1,35 @@
+// Copyright Epic Games, Inc. All Rights Reserved.
+#pragma once
+
+#include "Chaos/Core.h"
+#include "Chaos/Declares.h"
+#include "Chaos/ParticleHandleFwd.h"
+#include "Chaos/ShapeInstanceFwd.h"
+
+#if CHAOS_DEBUG_DRAW
+
+namespace Chaos
+{
+	struct CHAOS_API FChaosDDParticle
+	{
+		// Draw a particle's shapes with auto-coloring. SpaceTransform is simulation space to world transform (for RBAN)
+		static void DrawShapes(
+			const FRigidTransform3& SpaceTransform,
+			const FGeometryParticleHandle* Particle);
+
+		// Draw a particle's shapes with auto-coloring
+		static void DrawShapes(
+			const FGeometryParticleHandle* Particle);
+
+		// Draw a particle's shapes with manual coloring
+		static void DrawShapes(
+			const FGeometryParticleHandle* Particle,
+			const FColor& Color);
+
+		// Draw a particle's optimized shapes with auto-coloring if it has any
+		static bool DrawOptimizedShapes(
+			const FGeometryParticleHandle* Particle);
+	};
+}
+
+#endif
\ No newline at end of file
diff --git a/Engine/Source/Runtime/Experimental/Chaos/Public/Chaos/DebugDraw/DebugDrawPrimitives.h b/Engine/Source/Runtime/Experimental/Chaos/Public/Chaos/DebugDraw/DebugDrawPrimitives.h
new file mode 100644
--- /dev/null
+++ b/Engine/Source/Runtime/Experimental/Chaos/Public/Chaos/DebugDraw/DebugDrawPrimitives.h
@@ -0,0 +1,39 @@
+// Copyright Epic Games, Inc. All Rights Reserved.
+#pragma once
+
+#include "Chaos/Core.h"
+#include "Chaos/Declares.h"
+#include "Math/MathFwd.h"
+
+#if CHAOS_DEBUG_DRAW
+
+namespace Chaos
+{
+	// A set of debug draw primitives.
+	// 
+	// NOTE: these are convenient for quick temporary debug draw but it is usually better to
+	// create a custom debug draw command with its own data and generating the primitives
+	// at render time, rather than capture time. Deferring the primitive generation to render
+	// time means that decisions about color, line width etc. can be deferred which is useful
+	// for rewinding and changing how objects are displayed.
+	//
+	// A debug draw command is a deferred lambda that captures the data required for 
+	// rendering and uses the IChaosDDRenderer interface to perform rendering. As the command
+	// is deferred, make sure all lambda parameters are captured by value, not reference!
+	// By convention the Draw function is a static member of a struct, but this is not required.
+	//
+	// See the DrawLine	implementation for a trivial example. For a less-trivial example 
+	// see FChaosDDCollisionConstraint, or for a complex example see FChaosDDParticle which 
+	// supports capture-time and various render-time coloring modes.
+	//
+	struct CHAOS_API FChaosDDPrimitives
+	{
+		static void DrawLine(const FVector3d& A, const FVector3d& B, const FColor& Color, float LineThickness, float Duration);
+		static void DrawSphere(const FVector3d& Center, const float Radius, const FColor& Color, float LineThickness, float Duration);
+		static void DrawCapsule(const FVector3d& Center, const FQuat4d& Rotation, float HalfHeight, float Radius, const FColor& Color, float LineThickness, float Duration);
+		static void DrawBox(const FVector3d& Center, const FQuat4d& Rotation, const FVector3d& Size, const FColor& Color, float LineThickness, float Duration);
+		static void DrawTriangle(const FVector3d& A, const FVector3d& B, const FVector3d& C, const FColor& Color, float LineThickness, float Duration);
+	};
+}
+
+#endif
\ No newline at end of file
