Corrected a member's type

#rnx

diff --git a/Engine/Source/Runtime/Experimental/IoStore/HttpClient/Private/Loop.inl b/Engine/Source/Runtime/Experimental/IoStore/HttpClient/Private/Loop.inl
--- a/Engine/Source/Runtime/Experimental/IoStore/HttpClient/Private/Loop.inl
+++ b/Engine/Source/Runtime/Experimental/IoStore/HttpClient/Private/Loop.inl
@@ -535,7 +535,7 @@ private:
 	FPeerType			Peer;
 	uint32				LastUseMs = 0;
 	uint8				IsKeepAlive = 0;
-	uint8				bNegotiating = 0;
+	bool				bNegotiating = false;
 	bool				bWaiting = false;
 
 	UE_NONCOPYABLE(FPeerGroup);
