Style fixes, plus don't set SNI if no name's given.

#rnx

diff --git a/Engine/Source/Runtime/Experimental/IoStore/HttpClient/Private/Peer.inl b/Engine/Source/Runtime/Experimental/IoStore/HttpClient/Private/Peer.inl
--- a/Engine/Source/Runtime/Experimental/IoStore/HttpClient/Private/Peer.inl
+++ b/Engine/Source/Runtime/Experimental/IoStore/HttpClient/Private/Peer.inl
@@ -152,11 +152,15 @@ FTlsPeer::FTlsPeer(FSocket InSocket, const FSslContext* Context)
 	// SSL_MODE_ENABLE_PARTIAL_WRITE ??!!!
 
 	Ssl = SSL_new(*Context);
-	SSL_set_tlsext_host_name(Ssl, Context->GetHostName());
 	SSL_set_connect_state(Ssl);
 	SSL_set0_rbio(Ssl, Bio);
 	SSL_set0_wbio(Ssl, Bio);
 	BIO_up_ref(Bio);
+
+	if (const char* HostName = Context->GetHostName(); HostName != nullptr)
+	{
+		SSL_set_tlsext_host_name(Ssl, HostName);
+	}
 }
 
 ////////////////////////////////////////////////////////////////////////////////
@@ -345,7 +349,9 @@ void FHttpPeer::AssignProto()
 	Proto = 1;
 
 	if (Ssl == nullptr)
+	{
 		return;
+	}
 
 	const char* AlpnProto = nullptr;
 	uint32 AlpnProtoLen;
