Specify host name to the peer instead of the more-general context

#rnx

diff --git a/Engine/Source/Runtime/Experimental/IoStore/HttpClient/Private/ConnectionPool.inl b/Engine/Source/Runtime/Experimental/IoStore/HttpClient/Private/ConnectionPool.inl
--- a/Engine/Source/Runtime/Experimental/IoStore/HttpClient/Private/ConnectionPool.inl
+++ b/Engine/Source/Runtime/Experimental/IoStore/HttpClient/Private/ConnectionPool.inl
@@ -160,8 +160,8 @@ FOutcome FHost::Connect(FTlsPeer& Peer)
 		return Outcome;
 	}
 
-	FSslContext SslContext(HostName, VerifyCert);
-	Peer = FTlsPeer(MoveTemp(Socket), &SslContext);
+	FSslContext SslContext(VerifyCert);
+	Peer = FTlsPeer(MoveTemp(Socket), &SslContext, HostName);
 	return Outcome;
 }
 
diff --git a/Engine/Source/Runtime/Experimental/IoStore/HttpClient/Private/Peer.inl b/Engine/Source/Runtime/Experimental/IoStore/HttpClient/Private/Peer.inl
--- a/Engine/Source/Runtime/Experimental/IoStore/HttpClient/Private/Peer.inl
+++ b/Engine/Source/Runtime/Experimental/IoStore/HttpClient/Private/Peer.inl
@@ -14,20 +14,17 @@ static_assert(OPENSSL_VERSION_NUMBER >= 0x10100000L, "version supporting autoini
 class FSslContext
 {
 public:
-				FSslContext(const char* InHostName, const FPemCert& CertPin={});
+				FSslContext(const FPemCert& CertPin={});
 				~FSslContext()					{ SSL_CTX_free(Context); }
 				operator SSL_CTX* () const		{ return Context; }
-	const char*	GetHostName() const				{ return HostName; }
 	bool		AddCert(const FPemCert& Cert);
 
 private:
 	SSL_CTX*	Context;
-	const char*	HostName;
 };
 
 ////////////////////////////////////////////////////////////////////////////////
-FSslContext::FSslContext(const char* InHostName, const FPemCert& CertPin)
-: HostName(InHostName)
+FSslContext::FSslContext(const FPemCert& CertPin)
 {
 	Context = SSL_CTX_new(TLS_client_method());
 
@@ -105,7 +102,7 @@ public:
 				~FTlsPeer();
 				FTlsPeer(FTlsPeer&& Rhs)				{ Move(MoveTemp(Rhs)); }
 				FTlsPeer& operator = (FTlsPeer&& Rhs)	{ return Move(MoveTemp(Rhs)); }
-				FTlsPeer(FSocket InSocket, const FSslContext* Context=nullptr);
+				FTlsPeer(FSocket InSocket, const FSslContext* Context=nullptr, const char* HostName=nullptr);
 	FTlsPeer&	Move(FTlsPeer&& Rhs);
 	FOutcome	Handshake();
 	FOutcome	Send(const char* Data, int32 Size);
@@ -122,7 +119,7 @@ private:
 };
 
 ////////////////////////////////////////////////////////////////////////////////
-FTlsPeer::FTlsPeer(FSocket InSocket, const FSslContext* Context)
+FTlsPeer::FTlsPeer(FSocket InSocket, const FSslContext* Context, const char* HostName)
 : FPeer(MoveTemp(InSocket))
 {
 	if (Context == nullptr)
@@ -157,7 +154,7 @@ FTlsPeer::FTlsPeer(FSocket InSocket, const FSslContext* Context)
 	SSL_set0_wbio(Ssl, Bio);
 	BIO_up_ref(Bio);
 
-	if (const char* HostName = Context->GetHostName(); HostName != nullptr)
+	if (HostName != nullptr)
 	{
 		SSL_set_tlsext_host_name(Ssl, HostName);
 	}
