Introduce an object for holding a set of trusted authorities against which to verify.

#rnx

diff --git a/Engine/Source/Runtime/Experimental/IoStore/HttpClient/Internal/IO/Http/Client.h b/Engine/Source/Runtime/Experimental/IoStore/HttpClient/Internal/IO/Http/Client.h
--- a/Engine/Source/Runtime/Experimental/IoStore/HttpClient/Internal/IO/Http/Client.h
+++ b/Engine/Source/Runtime/Experimental/IoStore/HttpClient/Internal/IO/Http/Client.h
@@ -49,6 +49,27 @@ using	FTicket		= uint64;
 using	FPemCert	= FMemoryView;
 struct	FActivity;
 
+////////////////////////////////////////////////////////////////////////////////
+class UE_API FCertRoots
+{
+public:
+							FCertRoots()					= default;
+							~FCertRoots();
+							FCertRoots(FMemoryView PemData);
+							FCertRoots(FCertRoots&& Rhs)	{ *this = MoveTemp(Rhs); }
+	FCertRoots&				operator = (FCertRoots&& Rhs)	{ Swap(Handle, Rhs.Handle); return *this; }
+	bool					IsValid() const					{ return Handle != 0; }
+	int32					Num() const;
+	static void				SetDefault(FCertRoots&& CertRoots);
+
+private:
+	UPTRINT					Handle = 0;
+
+private:
+							FCertRoots(const FCertRoots&)	= delete;
+	FCertRoots&				operator = (const FCertRoots&)	= delete;
+};
+
 ////////////////////////////////////////////////////////////////////////////////
 class UE_API FConnectionPool
 {
diff --git a/Engine/Source/Runtime/Experimental/IoStore/HttpClient/Private/Peer.inl b/Engine/Source/Runtime/Experimental/IoStore/HttpClient/Private/Peer.inl
--- a/Engine/Source/Runtime/Experimental/IoStore/HttpClient/Private/Peer.inl
+++ b/Engine/Source/Runtime/Experimental/IoStore/HttpClient/Private/Peer.inl
@@ -64,6 +64,81 @@ bool FSslContext::AddCert(const FPemCert& Cert)
 
 
 
+////////////////////////////////////////////////////////////////////////////////
+static FCertRoots GDefaultCertRoots;
+
+////////////////////////////////////////////////////////////////////////////////
+FCertRoots::~FCertRoots()
+{
+	if (Handle == 0)
+	{
+		return;
+	}
+
+	auto* Context = (SSL_CTX*)Handle;
+	SSL_CTX_free(Context);
+}
+
+////////////////////////////////////////////////////////////////////////////////
+FCertRoots::FCertRoots(FMemoryView PemData)
+{
+	SSL_CTX* Context = SSL_CTX_new(TLS_client_method());
+	SSL_CTX_set_options(Context, SSL_OP_NO_SSLv2|SSL_OP_NO_SSLv3);
+
+	const void* Data = PemData.GetData();
+	uint32 Size = uint32(PemData.GetSize());
+	BIO* Bio = BIO_new_mem_buf(Data, Size);
+
+	uint32 NumAdded = 0;
+	while (true)
+	{
+		X509* FiveOhNine = PEM_read_bio_X509(Bio, nullptr, 0, nullptr);
+		if (FiveOhNine == nullptr)
+		{
+			break;
+		}
+
+		X509_STORE* Store = SSL_CTX_get_cert_store(Context);
+		int32 Result = X509_STORE_add_cert(Store, FiveOhNine);
+		NumAdded += (Result == 1);
+
+		X509_free(FiveOhNine);
+	}
+
+	BIO_free(Bio);
+
+	if (NumAdded == 0)
+	{
+		SSL_CTX_free(Context);
+		return;
+	}
+
+	Handle = UPTRINT(Context);
+}
+
+////////////////////////////////////////////////////////////////////////////////
+int32 FCertRoots::Num() const
+{
+	if (Handle == 0)
+	{
+		return -1;
+	}
+
+	auto* Context = (SSL_CTX*)Handle;
+	X509_STORE* Store = SSL_CTX_get_cert_store(Context);
+	STACK_OF(X509_OBJECT)* Objects = X509_STORE_get0_objects(Store);
+	return sk_X509_OBJECT_num(Objects);
+}
+
+////////////////////////////////////////////////////////////////////////////////
+void FCertRoots::SetDefault(FCertRoots&& CertRoots)
+{
+	check(GDefaultCertRoots.IsValid() != CertRoots.IsValid());
+	GDefaultCertRoots = MoveTemp(CertRoots);
+}
+
+
+
 ////////////////////////////////////////////////////////////////////////////////
 class FPeer
 {
