Set container ID when uploading on-demand container files

diff --git a/Engine/Source/Runtime/Experimental/IoStore/OnDemand/Private/IoStoreOnDemand.cpp b/Engine/Source/Runtime/Experimental/IoStore/OnDemand/Private/IoStoreOnDemand.cpp
--- a/Engine/Source/Runtime/Experimental/IoStore/OnDemand/Private/IoStoreOnDemand.cpp
+++ b/Engine/Source/Runtime/Experimental/IoStore/OnDemand/Private/IoStoreOnDemand.cpp
@@ -527,6 +527,7 @@ FArchive& operator<<(FArchive& Ar, FOnDemandTocContainerEntry& ContainerEntry)
 FCbWriter& operator<<(FCbWriter& Writer, const FOnDemandTocContainerEntry& ContainerEntry)
 {
 	Writer.BeginObject();
+	Writer << UTF8TEXTVIEW("Id") << ContainerEntry.ContainerId;
 	Writer.AddString(UTF8TEXTVIEW("Name"), ContainerEntry.ContainerName);
 	Writer.AddString(UTF8TEXTVIEW("EncryptionKeyGuid"), ContainerEntry.EncryptionKeyGuid);
 
@@ -1511,6 +1512,8 @@ TIoStatusOr<FIoStoreUploadResult> UploadContainerFiles(
 		
 		FOnDemandTocContainerEntry& ContainerEntry = OnDemandToc.Containers.AddDefaulted_GetRef();
 		ContainerEntry.ContainerName = FPaths::GetBaseFilename(Path);
+		ContainerEntry.ContainerId = ContainerFileReader.GetContainerId();
+
 		if (EnumHasAnyFlags(ContainerFileReader.GetContainerFlags(), EIoContainerFlags::Encrypted))
 		{
 			ContainerEntry.EncryptionKeyGuid = LexToString(ContainerFileReader.GetEncryptionKeyGuid());
