Update GeometryCollection memory count to be more accurate in obj lists.

diff --git a/Engine/Source/Runtime/Experimental/Chaos/Private/GeometryCollection/ManagedArrayCollection.cpp b/Engine/Source/Runtime/Experimental/Chaos/Private/GeometryCollection/ManagedArrayCollection.cpp
--- a/Engine/Source/Runtime/Experimental/Chaos/Private/GeometryCollection/ManagedArrayCollection.cpp
+++ b/Engine/Source/Runtime/Experimental/Chaos/Private/GeometryCollection/ManagedArrayCollection.cpp
@@ -653,7 +653,12 @@ void FManagedArrayCollection::Serialize(Chaos::FChaosArchive& Ar)
 
 	Ar << Version;
 
-	if (Ar.IsLoading())
+	if (Ar.IsCountingMemory())
+	{
+		Ar << GroupInfo;
+		Ar << Map;
+	}
+	else if (Ar.IsLoading())
 	{
 		//We can't serialize entire tmap in place because we may have new groups. todo(ocohen): baked data should be simpler since all entries exist
 		TMap< FName, FGroupInfo> TmpGroupInfo;
diff --git a/Engine/Source/Runtime/Experimental/Chaos/Public/Chaos/ChaosArchive.h b/Engine/Source/Runtime/Experimental/Chaos/Public/Chaos/ChaosArchive.h
--- a/Engine/Source/Runtime/Experimental/Chaos/Public/Chaos/ChaosArchive.h
+++ b/Engine/Source/Runtime/Experimental/Chaos/Public/Chaos/ChaosArchive.h
@@ -379,6 +379,8 @@ private:
 template <typename T, typename TAllocator>
 FChaosArchive& operator<<(FChaosArchive& Ar, TArray<T, TAllocator>& Array)
 {
+	Array.CountBytes(Ar);
+
 	int32 ArrayNum = Array.Num();
 	Ar << ArrayNum;
 	Array.Reserve(ArrayNum);
diff --git a/Engine/Source/Runtime/Experimental/GeometryCollectionEngine/Private/GeometryCollection/GeometryCollectionObject.cpp b/Engine/Source/Runtime/Experimental/GeometryCollectionEngine/Private/GeometryCollection/GeometryCollectionObject.cpp
--- a/Engine/Source/Runtime/Experimental/GeometryCollectionEngine/Private/GeometryCollection/GeometryCollectionObject.cpp
+++ b/Engine/Source/Runtime/Experimental/GeometryCollectionEngine/Private/GeometryCollection/GeometryCollectionObject.cpp
@@ -1014,7 +1014,7 @@ void UGeometryCollection::Serialize(FArchive& Ar)
 			EnsureDataIsCooked(false /*bInitResources*/, Ar.IsTransacting(), Ar.IsPersistent(), false /*bAllowCopyFromDDC*/);
 		}
 #endif
-		if (Ar.IsLoading())
+		if (Ar.IsLoading() || Ar.IsCountingMemory())
 		{
 			GeometryCollection->Serialize(ChaosAr);
 		}
