Dataflow :  AnyType architecture refactor
- Added TypePolicy for connections  to decide what type they are compatible with
- Extended AnyType types ( like Numeric types )
- Nodes are now responsible to update their input/output types if an input or output type changes

diff --git a/Engine/Plugins/Experimental/Dataflow/Source/DataflowEditor/Private/Dataflow/DataflowSchema.cpp b/Engine/Plugins/Experimental/Dataflow/Source/DataflowEditor/Private/Dataflow/DataflowSchema.cpp
--- a/Engine/Plugins/Experimental/Dataflow/Source/DataflowEditor/Private/Dataflow/DataflowSchema.cpp
+++ b/Engine/Plugins/Experimental/Dataflow/Source/DataflowEditor/Private/Dataflow/DataflowSchema.cpp
@@ -167,25 +167,23 @@ const FPinConnectionResponse UDataflowSchema::CanCreateConnection(const UEdGraph
 		PinA = InPinB; PinB = InPinA;
 	}
 
-
 	if (PinA->Direction == EEdGraphPinDirection::EGPD_Output)
 	{
 		if (PinB->Direction == EEdGraphPinDirection::EGPD_Input)
 		{
 			// Make sure the pins are not on the same node
-			if (PinA->GetOwningNode() != PinB->GetOwningNode())
+			UDataflowEdNode* EdNodeA = Cast<UDataflowEdNode>(PinA->GetOwningNode());
+			UDataflowEdNode* EdNodeB = Cast<UDataflowEdNode>(PinB->GetOwningNode());
+
+			if (EdNodeA && EdNodeB && (EdNodeA != EdNodeB))
 			{
-				// Make sure types match. 
-				const bool bIsAnyTypeA = PinA->PinType.PinCategory == FDataflowAnyType::TypeName;
-				const bool bIsAnyTypeB = PinB->PinType.PinCategory == FDataflowAnyType::TypeName;
-				const bool bAreSameType = (PinA->PinType == PinB->PinType);
-				const bool bOnlyOneIsAnyType = (!bAreSameType && (bIsAnyTypeA != bIsAnyTypeB));
-				if (bAreSameType || bOnlyOneIsAnyType)
+				const bool AIsCompatibleWithB = EdNodeA->PinIsCompatibleWithType(*PinA, PinB->PinType);
+				const bool BIsCompatibleWithA = EdNodeB->PinIsCompatibleWithType(*PinB, PinA->PinType);
+				if (AIsCompatibleWithB || BIsCompatibleWithA)
 				{
 					// cycle checking on connect
-					if (!HasLoopIfConnected(PinA->GetOwningNode(), PinB->GetOwningNode()))
+					if (!HasLoopIfConnected(EdNodeA, EdNodeB))
 					{
-
 						if (PinB->LinkedTo.Num())
 						{
 							return (bSwapped) ?
diff --git a/Engine/Source/Runtime/Experimental/Dataflow/Core/Private/Dataflow/DataflowConnection.cpp b/Engine/Source/Runtime/Experimental/Dataflow/Core/Private/Dataflow/DataflowConnection.cpp
--- a/Engine/Source/Runtime/Experimental/Dataflow/Core/Private/Dataflow/DataflowConnection.cpp
+++ b/Engine/Source/Runtime/Experimental/Dataflow/Core/Private/Dataflow/DataflowConnection.cpp
@@ -8,14 +8,43 @@
 #include UE_INLINE_GENERATED_CPP_BY_NAME(DataflowConnection)
 
 FDataflowConnection::FDataflowConnection(Dataflow::FPin::EDirection InDirection, FName InType, FName InName, FDataflowNode* InOwningNode, const FProperty* InProperty, FGuid InGuid)
-	: Direction(InDirection)
-	, Type(InType)
+	: Type(InType)
 	, Name(InName)
 	, OwningNode(InOwningNode)
 	, Property(InProperty)
 	, Guid(InGuid)
+	, Direction(InDirection)
 {
-	bIsAnyType = FDataflowConnection::IsAnyType(Type);
+	InitFromType();
+}
+
+FDataflowConnection::FDataflowConnection(Dataflow::FPin::EDirection InDirection, const Dataflow::FConnectionParameters& Params)
+	: Type(Params.Type)
+	, Name(Params.Name)
+	, OwningNode(Params.Owner)
+	, Property(Params.Property)
+	, Guid(Params.Guid)
+	, Direction(InDirection)
+{
+	InitFromType();
+}
+
+void FDataflowConnection::InitFromType()
+{
+	bIsAnyType = false;
+	bHasConcreteType = true;
+	if (Property && Property->GetClass()->IsChildOf(FStructProperty::StaticClass()))
+	{
+		if (const FStructProperty* StructProperty = CastField<FStructProperty>(Property))
+		{
+			if (StructProperty->Struct->IsChildOf<FDataflowAnyType>())
+			{
+				Type = FDataflowAnyType::TypeName;
+				bIsAnyType = true;
+				bHasConcreteType = false;
+			}
+		}
+	}
 }
 
 uint32 FDataflowConnection::GetOffset() const
@@ -54,15 +83,46 @@ void FDataflowConnection::SetAsAnyType(bool bAnyType, const FName& ConcreteType)
 	if (bIsAnyType)
 	{
 		Type = ConcreteType;
+		bHasConcreteType = !IsAnyType(ConcreteType);
+	}
+}
+
+bool FDataflowConnection::SupportsType(FName InType) const
+{
+	if (IsAnyType(InType))
+	{
+		return false;
+	}
+	// resort to policy only if the concrete type is not defined ( case of anytype connection )
+	if (bIsAnyType && !bHasConcreteType)
+	{
+		return TypePolicy ? TypePolicy->SupportsType(InType) : true;
 	}
+	// todo : in the future we could also check for pointer compatibility
+	return (InType == GetType());
 }
 
-void FDataflowConnection::SetConcreteType(FName InType)
+bool FDataflowConnection::SetConcreteType(FName InType)
 {
 	// Can only change from AnyType to a concrete type
-	if (ensure(IsAnyType() && !IsAnyType(InType)))
+	if (Type != InType && ensure(!bHasConcreteType))
+	{
+		if (ensure(SupportsType(InType)))
+		{
+			Type = InType;
+			bHasConcreteType = true;
+			return true;
+		}
+	}
+	return false;
+}
+
+void FDataflowConnection::SetTypePolicy(IDataflowTypePolicy* InTypePolicy)
+{
+	// for now only allow setting it once
+	if (ensure(TypePolicy == nullptr))
 	{
-		Type = InType;
+		TypePolicy = InTypePolicy;
 	}
 }
 
@@ -70,6 +130,7 @@ void FDataflowConnection::ForceSimpleType(FName InType)
 {
 	check(Type.ToString().StartsWith(InType.ToString()));
 	Type = InType;
+	bHasConcreteType = true;
 }
 
 void FDataflowConnection::FixAndPropagateType()
diff --git a/Engine/Source/Runtime/Experimental/Dataflow/Core/Private/Dataflow/DataflowCoreNodes.cpp b/Engine/Source/Runtime/Experimental/Dataflow/Core/Private/Dataflow/DataflowCoreNodes.cpp
--- a/Engine/Source/Runtime/Experimental/Dataflow/Core/Private/Dataflow/DataflowCoreNodes.cpp
+++ b/Engine/Source/Runtime/Experimental/Dataflow/Core/Private/Dataflow/DataflowCoreNodes.cpp
@@ -16,6 +16,16 @@ void FDataflowReRouteNode::Evaluate(Dataflow::FContext& Context, const FDataflow
 	ForwardInput(Context, &Value, &Value);
 }
 
+bool FDataflowReRouteNode::OnInputTypeChanged(const FDataflowInput* Input)
+{
+	return SetOutputConcreteType(&Value, Input->GetType());
+}
+
+bool FDataflowReRouteNode::OnOutputTypeChanged(const FDataflowOutput* Input)
+{
+	return SetInputConcreteType(&Value, Input->GetType());
+}
+
 FDataflowBranchNode::FDataflowBranchNode(const Dataflow::FNodeParameters& Param, FGuid InGuid)
 	: Super(Param, InGuid)
 {
@@ -30,7 +40,7 @@ void FDataflowBranchNode::Evaluate(Dataflow::FContext& Context, const FDataflowO
 	if (Out->IsA(&Result))
 	{
 		const bool InCondition = GetValue<bool>(Context, &bCondition);
-		const FDataflowAnyType* SelectedInputReference = InCondition ? &TrueValue : &FalseValue;
+		const void* SelectedInputReference = InCondition ? &TrueValue : &FalseValue;
 		if (IsConnected(SelectedInputReference))
 		{
 			ForwardInput(Context, SelectedInputReference, &Result);
@@ -41,4 +51,21 @@ void FDataflowBranchNode::Evaluate(Dataflow::FContext& Context, const FDataflowO
 			// Context.Error(TEXT("Both True and False Inputs must be connected"));
 		}
 	}
+}
+
+bool FDataflowBranchNode::OnInputTypeChanged(const FDataflowInput* Input)
+{
+	return bool(
+		SetInputConcreteType(&TrueValue, Input->GetType())
+	  | SetInputConcreteType(&FalseValue, Input->GetType())
+	  | SetOutputConcreteType(&Result, Input->GetType())
+		);
+}
+
+bool FDataflowBranchNode::OnOutputTypeChanged(const FDataflowOutput* Input)
+{
+	return bool(
+		  SetInputConcreteType(&TrueValue, Input->GetType())
+		| SetInputConcreteType(&FalseValue, Input->GetType())
+		);
 }
\ No newline at end of file
diff --git a/Engine/Source/Runtime/Experimental/Dataflow/Core/Private/Dataflow/DataflowInputOutput.cpp b/Engine/Source/Runtime/Experimental/Dataflow/Core/Private/Dataflow/DataflowInputOutput.cpp
--- a/Engine/Source/Runtime/Experimental/Dataflow/Core/Private/Dataflow/DataflowInputOutput.cpp
+++ b/Engine/Source/Runtime/Experimental/Dataflow/Core/Private/Dataflow/DataflowInputOutput.cpp
@@ -5,6 +5,7 @@
 #include "ChaosLog.h"
 #include "Dataflow/DataflowNodeParameters.h"
 #include "Dataflow/DataflowNode.h"
+#include "Dataflow/DataflowCoreNodes.h"
 
 #include UE_INLINE_GENERATED_CPP_BY_NAME(DataflowInputOutput)
 
@@ -14,7 +15,14 @@ FDataflowOutput FDataflowOutput::NoOpOutput = FDataflowOutput();
 
 
 FDataflowInput::FDataflowInput(const Dataflow::FInputParameters& Param, FGuid InGuid)
-	: FDataflowConnection(Dataflow::FPin::EDirection::INPUT, Param.Type, Param.Name, Param.Owner, Param.Property, InGuid)
+	: FDataflowConnection(Dataflow::FPin::EDirection::INPUT, Param)
+	, Connection(nullptr)
+{
+	Guid = InGuid;
+}
+
+FDataflowInput::FDataflowInput(const Dataflow::FInputParameters& Param)
+	: FDataflowConnection(Dataflow::FPin::EDirection::INPUT, Param)
 	, Connection(nullptr)
 {
 }
@@ -81,17 +89,15 @@ void FDataflowInput::PullValue(Dataflow::FContext& Context) const
 void FDataflowInput::FixAndPropagateType(FName InType)
 {
 	check(InType.ToString().StartsWith(Type.ToString()));
-	Type = InType;
-
-	if (FDataflowOutput* Output = GetConnection())
+	const bool bSuccess = OwningNode->TrySetConnectionType(this, InType);
+	if (bSuccess)
 	{
-		if (Output && Output->GetType() != InType)
+		if (FDataflowReRouteNode* ReRouteNode = OwningNode->AsType<FDataflowReRouteNode>())
 		{
-			if (Output->IsAnyType())
+			if (FDataflowOutput* Output = GetConnection())
 			{
-				Output->GetOwningNode()->PropagateTypeToAllAnyTypeInputsAndOutputs(InType);
+				Output->FixAndPropagateType(InType);
 			}
-			Output->FixAndPropagateType(InType);
 		}
 	}
 }
@@ -105,7 +111,14 @@ void FDataflowInput::FixAndPropagateType(FName InType)
 
 
 FDataflowOutput::FDataflowOutput(const Dataflow::FOutputParameters& Param, FGuid InGuid)
-	: FDataflowConnection(Dataflow::FPin::EDirection::OUTPUT, Param.Type, Param.Name, Param.Owner, Param.Property, InGuid)
+	: FDataflowConnection(Dataflow::FPin::EDirection::OUTPUT, Param)
+{
+	Guid = InGuid;
+	OutputLock = MakeShared<FCriticalSection>();
+}
+
+FDataflowOutput::FDataflowOutput(const Dataflow::FOutputParameters& Param)
+	: FDataflowConnection(Dataflow::FPin::EDirection::OUTPUT, Param)
 {
 	OutputLock = MakeShared<FCriticalSection>();
 }
@@ -232,17 +245,18 @@ void FDataflowOutput::ForwardInput(const void* InputReference, Dataflow::FContex
 void FDataflowOutput::FixAndPropagateType(FName InType)
 {
 	check(InType.ToString().StartsWith(Type.ToString()));
-	Type = InType;
-
-	for (FDataflowInput* Input: Connections)
+	const bool bSuccess = OwningNode->TrySetConnectionType(this, InType);
+	if (bSuccess)
 	{
-		if (Input && Input->GetType() != InType)
+		if (FDataflowReRouteNode* ReRouteNode = OwningNode->AsType<FDataflowReRouteNode>())
 		{
-			if (Input->IsAnyType())
+			for (FDataflowInput* Input : Connections)
 			{
-				Input->GetOwningNode()->PropagateTypeToAllAnyTypeInputsAndOutputs(InType);
+				if (Input)
+				{
+					Input->FixAndPropagateType(InType);
+				}
 			}
-			Input->FixAndPropagateType(InType);
 		}
 	}
 }
\ No newline at end of file
diff --git a/Engine/Source/Runtime/Experimental/Dataflow/Core/Private/Dataflow/DataflowNode.cpp b/Engine/Source/Runtime/Experimental/Dataflow/Core/Private/Dataflow/DataflowNode.cpp
--- a/Engine/Source/Runtime/Experimental/Dataflow/Core/Private/Dataflow/DataflowNode.cpp
+++ b/Engine/Source/Runtime/Experimental/Dataflow/Core/Private/Dataflow/DataflowNode.cpp
@@ -26,6 +26,24 @@ const FName FDataflowAnyType::TypeName = TEXT("FDataflowAnyType");
 // Inputs
 //
 
+bool FDataflowNode::OutputSupportsType(FName InName, FName InType) const
+{
+	if (const FDataflowOutput* Output = FindOutput(InName))
+	{
+		return Output->SupportsType(InType);
+	}
+	return false;
+}
+
+bool FDataflowNode::InputSupportsType(FName InName, FName InType) const
+{
+	if (const FDataflowInput* Input = FindInput(InName))
+	{
+		return Input->SupportsType(InType);
+	}
+	return false;
+}
+
 void FDataflowNode::AddInput(FDataflowInput* InPtr)
 {
 	if (InPtr)
@@ -378,26 +396,6 @@ void FDataflowNode::UnregisterPinConnection(const Dataflow::FPin& Pin)
 	}
 }
 
-void FDataflowNode::PropagateTypeToAllAnyTypeInputsAndOutputs(FName ConcreteType)
-{
-	for (TPair<uint32, FDataflowOutput*> Elem : Outputs)
-	{
-		FDataflowOutput* Output = Elem.Value;
-		if (Output->IsAnyType())
-		{
-			Output->SetConcreteType(ConcreteType);
-		}
-	}
-	for (TPair<uint32, FDataflowInput*> Elem : Inputs)
-	{
-		FDataflowInput* Input = Elem.Value;
-		if (Input->IsAnyType())
-		{
-			Input->SetConcreteType(ConcreteType);
-		}
-	}
-}
-
 void FDataflowNode::Invalidate(const Dataflow::FTimestamp& InModifiedTimestamp)
 {
 	if (LastModifiedTimestamp < InModifiedTimestamp)
@@ -516,25 +514,37 @@ FText FDataflowNode::GetPropertyDisplayNameText(const TArray<const FProperty*>&
 #endif
 }
 
-FDataflowInput* FDataflowNode::RegisterInputConnection(const void* InProperty, const FName& PropertyName)
+bool FDataflowNode::InitConnectionParametersFromPropertyReference(const FStructOnScope& StructOnScope, const void* PropertyRef, const FName& PropertyName, Dataflow::FConnectionParameters& OutParams)
 {
-	if (TUniquePtr<FStructOnScope> ScriptOnStruct = TUniquePtr<FStructOnScope>(NewStructOnScope()))
-	{
-		if (const UStruct* Struct = ScriptOnStruct->GetStruct())
+		if (const UStruct* Struct = StructOnScope.GetStruct())
 		{
 			TArray<const FProperty*> PropertyChain;
-			const FProperty* const Property = FindProperty(Struct, InProperty, PropertyName, &PropertyChain);
+			const FProperty* const Property = FindProperty(Struct, PropertyRef, PropertyName, &PropertyChain);
 			if (ensure(Property && PropertyChain.Num()))
 			{
-				const FName PropName(GetPropertyFullName(PropertyChain));
 				FString ExtendedType;
 				const FString CPPType = Property->GetCPPType(&ExtendedType);
-				const FName PropType(CPPType + ExtendedType);
-				FDataflowInput* const Input = new FDataflowInput({ PropType, PropName, this, Property });
-				AddInput(Input);
-				return Input;
+				OutParams.Type = FName(CPPType + ExtendedType);
+				OutParams.Name = GetPropertyFullName(PropertyChain);
+				OutParams.Property = Property;
+				OutParams.Owner = this;
+				return true;
 			}
 		}
+	return false;
+}
+
+FDataflowInput* FDataflowNode::RegisterInputConnectionInternal(const void* InProperty, const FName& PropertyName)
+{
+	if (TUniquePtr<FStructOnScope> ScriptOnStruct = TUniquePtr<FStructOnScope>(NewStructOnScope()))
+	{
+		Dataflow::FInputParameters InputParams;
+		if (InitConnectionParametersFromPropertyReference(*ScriptOnStruct, InProperty, PropertyName, InputParams))
+		{
+			FDataflowInput* const Input = new FDataflowInput(InputParams);
+			AddInput(Input);
+			return Input;
+		}
 	}
 	return nullptr;
 }
@@ -564,39 +574,28 @@ void FDataflowNode::UnregisterInputConnection(const void* InProperty, const FNam
 	}
 }
 
-FDataflowOutput* FDataflowNode::RegisterOutputConnection(const void* InProperty, const void* Passthrough, const FName& PropertyName, const FName& PassthroughName)
+FDataflowOutput* FDataflowNode::RegisterOutputConnectionInternal(const void* InProperty, const void* Passthrough, const FName& PropertyName, const FName& PassthroughName)
 {
 	if (TUniquePtr<FStructOnScope> ScriptOnStruct = TUniquePtr<FStructOnScope>(NewStructOnScope()))
 	{
-		if (const UStruct* Struct = ScriptOnStruct->GetStruct())
+		Dataflow::FOutputParameters OutputParams;
+		if (InitConnectionParametersFromPropertyReference(*ScriptOnStruct, InProperty, PropertyName, OutputParams))
 		{
-			FDataflowOutput* OutputConnection = nullptr;
-			TArray<const FProperty*> PropertyChain;
-			const FProperty* const Property =
-				FindProperty(Struct, InProperty, PropertyName, &PropertyChain);
-			if (ensure(Property && PropertyChain.Num()))
-			{
-				const FName PropName(GetPropertyFullName(PropertyChain));
-				FString ExtendedType;
-				const FString CPPType = Property->GetCPPType(&ExtendedType);
-				const FName PropType(CPPType + ExtendedType);
-				OutputConnection = new FDataflowOutput({ PropType, PropName, this, Property });
+			FDataflowOutput* OutputConnection = new FDataflowOutput(OutputParams);
 
-				TArray<const FProperty*> PassthroughPropertyChain;
-				if (FindProperty(Struct, Passthrough, PassthroughName, &PassthroughPropertyChain))
-				{
-					const uint32 PassthroughOffset = GetPropertyOffset(PassthroughPropertyChain);
-					OutputConnection->SetPassthroughOffset(PassthroughOffset);
-				}
-				AddOutput(OutputConnection);
+			TArray<const FProperty*> PassthroughPropertyChain;
+			if (FindProperty(ScriptOnStruct->GetStruct(), Passthrough, PassthroughName, &PassthroughPropertyChain))
+			{
+				const uint32 PassthroughOffset = GetPropertyOffset(PassthroughPropertyChain);
+				OutputConnection->SetPassthroughOffset(PassthroughOffset);
 			}
-			return OutputConnection;
+
+			AddOutput(OutputConnection);
 		}
 	}
 	return nullptr;
 }
 
-
 uint32 FDataflowNode::GetValueHash()
 {
 	//UE_LOG(LogChaos, Warning, TEXT("%s"), *GetName().ToString())
@@ -886,6 +885,51 @@ void FDataflowNode::ForwardInput(Dataflow::FContext& Context, const void* InputR
 	}
 }
 
+bool FDataflowNode::TrySetConnectionType(FDataflowConnection* Connection, FName NewType)
+{
+	if (Connection)
+	{
+		if (Connection->IsAnyType() && Connection->GetType() != NewType && !FDataflowConnection::IsAnyType(NewType))
+		{
+			Connection->SetConcreteType(NewType);
+			if (Connection->GetDirection() == Dataflow::FPin::EDirection::INPUT)
+			{
+				OnInputTypeChanged((FDataflowInput*)Connection);
+			}
+			if (Connection->GetDirection() == Dataflow::FPin::EDirection::OUTPUT)
+			{
+				OnOutputTypeChanged((FDataflowOutput*)Connection);
+			}
+			return true;
+		}
+	}
+	return false;
+}
+
+bool FDataflowNode::SetInputConcreteType(void* InputReference, FName NewType)
+{
+	if (FDataflowInput* Input = FindInput(InputReference))
+	{
+		if (Input->GetType() != NewType)
+		{
+			return Input->SetConcreteType(NewType);
+		}
+	}
+	return false;
+}
+
+bool FDataflowNode::SetOutputConcreteType(void* OutputReference, FName NewType)
+{
+	if (FDataflowOutput* Output = FindOutput(OutputReference))
+	{
+		if (Output->GetType() != NewType)
+		{
+			return Output->SetConcreteType(NewType);
+		}
+	}
+	return false;
+}
+
 
 
 
diff --git a/Engine/Source/Runtime/Experimental/Dataflow/Core/Public/Dataflow/DataflowAnyType.h b/Engine/Source/Runtime/Experimental/Dataflow/Core/Public/Dataflow/DataflowAnyType.h
new file mode 100644
--- /dev/null
+++ b/Engine/Source/Runtime/Experimental/Dataflow/Core/Public/Dataflow/DataflowAnyType.h
@@ -0,0 +1,35 @@
+// Copyright Epic Games, Inc. All Rights Reserved.
+
+#pragma once
+
+#include "CoreMinimal.h"
+#include "Dataflow/DataflowTypePolicy.h"
+
+#include "DataflowAnyType.generated.h"
+
+USTRUCT()
+struct FDataflowAnyType
+{
+	GENERATED_USTRUCT_BODY()
+	DATAFLOWCORE_API static const FName TypeName;
+};
+
+USTRUCT()
+struct FDataflowAllTypes : public FDataflowAnyType
+{
+	using FPolicyType = FDataflowAllTypesPolicy;
+	using FStorageType = void;
+	GENERATED_USTRUCT_BODY()
+};
+
+USTRUCT()
+struct FDataflowNumericTypes : public FDataflowAnyType
+{
+	using FPolicyType = FDataflowNumericTypePolicy;
+	using FStorageType = double;
+
+	GENERATED_USTRUCT_BODY()
+
+	UPROPERTY(EditAnywhere, Category=Value)
+	double Value = 0.0;
+};
diff --git a/Engine/Source/Runtime/Experimental/Dataflow/Core/Public/Dataflow/DataflowConnection.h b/Engine/Source/Runtime/Experimental/Dataflow/Core/Public/Dataflow/DataflowConnection.h
--- a/Engine/Source/Runtime/Experimental/Dataflow/Core/Public/Dataflow/DataflowConnection.h
+++ b/Engine/Source/Runtime/Experimental/Dataflow/Core/Public/Dataflow/DataflowConnection.h
@@ -4,6 +4,7 @@
 
 #include "CoreMinimal.h"
 #include "Dataflow/DataflowNodeParameters.h"
+#include "Dataflow/DataflowTypePolicy.h"
 #include "HAL/PlatformMath.h"
 #include "Misc/AssertionMacros.h"
 #include "Misc/Guid.h"
@@ -15,7 +16,6 @@
 class FProperty;
 struct FDataflowNode;
 
-
 namespace Dataflow
 {
 
@@ -35,6 +35,23 @@ namespace Dataflow
 		bool bHidden = false;
 	};
 
+	struct FConnectionParameters
+	{
+		FConnectionParameters(FName InType = FName(""), FName InName = FName(""), FDataflowNode* InOwner = nullptr, const FProperty* InProperty = nullptr, FGuid InGuid = FGuid::NewGuid())
+			: Type(InType)
+			, Name(InName)
+			, Owner(InOwner)
+			, Property(InProperty)
+			, Guid(InGuid)
+		{}
+
+		FName Type;
+		FName Name;
+		FDataflowNode* Owner = nullptr;
+		const FProperty* Property = nullptr;
+		FGuid Guid;
+	};
+
 	class FGraph;
 }
 
@@ -47,23 +64,23 @@ struct FDataflowConnection
 	GENERATED_USTRUCT_BODY()
 
 protected:
-	Dataflow::FPin::EDirection Direction;
 	FName Type;
 	FName Name;
 	FDataflowNode* OwningNode = nullptr;
 	const FProperty* Property = nullptr;
 	FGuid  Guid;
-	bool bIsAnyType = false;
-	bool bCanHidePin = false;
-	bool bPinIsHidden = false;
+	IDataflowTypePolicy* TypePolicy = nullptr;
+	Dataflow::FPin::EDirection Direction;
+	bool bIsAnyType:1 = false;
+	bool bHasConcreteType : 1 = false;
+	bool bCanHidePin:1 = false;
+	bool bPinIsHidden:1 = false;
 
 	friend struct FDataflowNode;
 	friend class Dataflow::FGraph;
 
 protected:
 	DATAFLOWCORE_API bool IsOwningNodeEnabled() const;
-	DATAFLOWCORE_API FGuid GetOwningNodeGuid() const;
-	DATAFLOWCORE_API uint32 GetOwningNodeValueHash() const;
 
 	/** this should only be used for serialization */
 	DATAFLOWCORE_API void SetAsAnyType(bool bAnyType, const FName& ConcreteType);
@@ -73,14 +90,25 @@ protected:
 	DATAFLOWCORE_API void FixAndPropagateType();
 	DATAFLOWCORE_API virtual void FixAndPropagateType(FName InType) { ensure(false); }
 
+	/** 
+	* Set the concrete type of an anytype connection 
+	* @return true if the typewas effectivelly changed
+	*/
+	DATAFLOWCORE_API bool SetConcreteType(FName InType);
+
 public:
 	FDataflowConnection() {};
+	UE_DEPRECATED(5.5, "Deprecated constructor : use FConnectionParameters to pass parameters")
 	DATAFLOWCORE_API FDataflowConnection(Dataflow::FPin::EDirection Direction, FName InType, FName InName, FDataflowNode* OwningNode = nullptr, const FProperty* InProperty = nullptr, FGuid InGuid = FGuid::NewGuid());
+	DATAFLOWCORE_API FDataflowConnection(Dataflow::FPin::EDirection Direction, const Dataflow::FConnectionParameters& Params);
 	virtual ~FDataflowConnection() {};
 
 	FDataflowNode* GetOwningNode() { return OwningNode; }
 	const FDataflowNode* GetOwningNode() const { return OwningNode; }
 
+	DATAFLOWCORE_API FGuid GetOwningNodeGuid() const;
+	DATAFLOWCORE_API uint32 GetOwningNodeValueHash() const;
+
 	const FProperty* GetProperty() const { return Property; }
 
 	Dataflow::FPin::EDirection GetDirection() const { return Direction; }
@@ -101,9 +129,10 @@ public:
 	virtual bool RemoveConnection(FDataflowConnection* In) { return false; }
 
 	DATAFLOWCORE_API bool IsAnyType() const { return bIsAnyType; }
-	static bool IsAnyType(const FName& InType);
-
-	DATAFLOWCORE_API void SetConcreteType(FName InType);
+	DATAFLOWCORE_API static bool IsAnyType(const FName& InType);
+	DATAFLOWCORE_API bool HasConcreteType() const { return bHasConcreteType; }
+	DATAFLOWCORE_API void SetTypePolicy(IDataflowTypePolicy* InTypePolicy);
+	DATAFLOWCORE_API bool SupportsType(FName InType) const;
 
 	template<class T>
 	bool IsA(const T* InVar) const
@@ -117,4 +146,7 @@ public:
 	bool GetPinIsHidden() const { return bCanHidePin && bPinIsHidden; }
 	void SetCanHidePin(bool bInCanHidePin) { bCanHidePin = bInCanHidePin; }
 	void SetPinIsHidden(bool bInPinIsHidden) { bPinIsHidden = bInPinIsHidden; }
+
+private:
+	void InitFromType();
 };
diff --git a/Engine/Source/Runtime/Experimental/Dataflow/Core/Public/Dataflow/DataflowCoreNodes.h b/Engine/Source/Runtime/Experimental/Dataflow/Core/Public/Dataflow/DataflowCoreNodes.h
--- a/Engine/Source/Runtime/Experimental/Dataflow/Core/Public/Dataflow/DataflowCoreNodes.h
+++ b/Engine/Source/Runtime/Experimental/Dataflow/Core/Public/Dataflow/DataflowCoreNodes.h
@@ -23,9 +23,15 @@ public:
 	FDataflowReRouteNode(const Dataflow::FNodeParameters& Param, FGuid InGuid = FGuid::NewGuid());
 	virtual void Evaluate(Dataflow::FContext& Context, const FDataflowOutput* Out) const;
 
+
 public:
 	UPROPERTY(meta = (DataflowInput, DataflowOutput, DataflowPassthrough = "Value", DisplayName = "Value"))
-	FDataflowAnyType Value;
+	FDataflowAllTypes Value;
+
+protected:
+	virtual bool OnInputTypeChanged(const FDataflowInput* Input) override;
+	virtual bool OnOutputTypeChanged(const FDataflowOutput* Input) override;
+
 };
 
 USTRUCT(meta=(Icon="GraphEditor.Branch_16x"))
@@ -41,14 +47,18 @@ public:
 
 public:
 	UPROPERTY(meta = (DataflowInput, DisplayName = "TrueValue"))
-	FDataflowAnyType TrueValue;
+	FDataflowAllTypes TrueValue;
 
 	UPROPERTY(meta = (DataflowInput, DisplayName = "FalseValue"))
-	FDataflowAnyType FalseValue;
+	FDataflowAllTypes FalseValue;
 
 	UPROPERTY(EditAnywhere, Category="Condition", meta = (DataflowInput, DisplayName = "Condition"))
 	bool bCondition = true;
 
 	UPROPERTY(meta = (DataflowOutput, DisplayName = "Result"))
-	FDataflowAnyType Result;
+	FDataflowAllTypes Result;
+
+private:
+	virtual bool OnInputTypeChanged(const FDataflowInput* Input) override;
+	virtual bool OnOutputTypeChanged(const FDataflowOutput* Input) override;
 };
\ No newline at end of file
diff --git a/Engine/Source/Runtime/Experimental/Dataflow/Core/Public/Dataflow/DataflowInputOutput.h b/Engine/Source/Runtime/Experimental/Dataflow/Core/Public/Dataflow/DataflowInputOutput.h
--- a/Engine/Source/Runtime/Experimental/Dataflow/Core/Public/Dataflow/DataflowInputOutput.h
+++ b/Engine/Source/Runtime/Experimental/Dataflow/Core/Public/Dataflow/DataflowInputOutput.h
@@ -19,16 +19,11 @@ struct FDataflowOutput;
 //
 namespace Dataflow
 {
-	struct FInputParameters {
-		FInputParameters(FName InType = FName(""), FName InName = FName(""), FDataflowNode* InOwner = nullptr, const FProperty* InProperty = nullptr)
-			: Type(InType)
-			, Name(InName)
-			, Owner(InOwner)
-			, Property(InProperty){}
-		FName Type;
-		FName Name;
-		FDataflowNode* Owner = nullptr;
-		const FProperty* Property = nullptr;
+	struct FInputParameters : public FConnectionParameters
+	{
+		FInputParameters(FName InType = NAME_None, FName InName = NAME_None, FDataflowNode* InOwner = nullptr, const FProperty* InProperty = nullptr, FGuid InGuid = FGuid::NewGuid())
+			: FConnectionParameters(InType, InName, InOwner, InProperty, InGuid)
+		{}
 	};
 }
 
@@ -48,7 +43,10 @@ protected:
 	virtual void FixAndPropagateType(FName InType) override;
 
 public:
-	FDataflowInput(const Dataflow::FInputParameters& Param = {}, FGuid InGuid = FGuid::NewGuid());
+	UE_DEPRECATED(5.5, "Deprecated constructor : Guid is now passed through FInputParameters")
+	DATAFLOWCORE_API FDataflowInput(const Dataflow::FInputParameters& Param, FGuid InGuid);
+
+	DATAFLOWCORE_API FDataflowInput(const Dataflow::FInputParameters& Param = {});
 
 	virtual bool AddConnection(FDataflowConnection* InOutput) override;
 	virtual bool RemoveConnection(FDataflowConnection* InOutput) override;
@@ -67,11 +65,14 @@ public:
 	template<class T>
 	const T& GetValue(Dataflow::FContext& Context, const T& Default) const;
 
+	template<typename TAnyType>
+	typename TAnyType::FStorageType GetValueFromAnyType(Dataflow::FContext& Context, const typename TAnyType::FStorageType& Default) const;
+
 	/**
 	* pull the value from the upstream connections
 	* the upstream graph is evaluated if necessary and values are cached along the way 
 	*/
-	void PullValue(Dataflow::FContext& Context) const;
+	DATAFLOWCORE_API void PullValue(Dataflow::FContext& Context) const;
 
 	template<class T>
 	TFuture<const T&> GetValueParallel(Dataflow::FContext& Context, const T& Default) const;
@@ -84,18 +85,11 @@ public:
 //
 namespace Dataflow
 {
-	struct FOutputParameters
+	struct FOutputParameters: public FConnectionParameters
 	{
-		FOutputParameters(FName InType = FName(""), FName InName = FName(""), FDataflowNode* InOwner = nullptr, const FProperty* InProperty = nullptr)
-			: Type(InType)
-			, Name(InName)
-			, Owner(InOwner)
-			, Property(InProperty) {}
-
-		FName Type;
-		FName Name;
-		FDataflowNode* Owner = nullptr;
-		const FProperty* Property = nullptr;
+		FOutputParameters(FName InType = NAME_None, FName InName = NAME_None, FDataflowNode* InOwner = nullptr, const FProperty* InProperty = nullptr, FGuid InGuid = FGuid::NewGuid())
+			: FConnectionParameters(InType, InName, InOwner, InProperty, InGuid)
+		{}
 	};
 }
 USTRUCT()
@@ -118,7 +112,10 @@ public:
 	
 	mutable TSharedPtr<FCriticalSection> OutputLock;
 	
-	DATAFLOWCORE_API FDataflowOutput(const Dataflow::FOutputParameters& Param = {}, FGuid InGuid = FGuid::NewGuid());
+	UE_DEPRECATED(5.5, "Deprecated constructor : Guid is now passed through FOutputParameters")
+	DATAFLOWCORE_API FDataflowOutput(const Dataflow::FOutputParameters& Param, FGuid InGuid);
+
+	DATAFLOWCORE_API FDataflowOutput(const Dataflow::FOutputParameters& Param = {});
 
 	DATAFLOWCORE_API TArray<FDataflowInput*>& GetConnections();
 	DATAFLOWCORE_API const TArray<FDataflowInput*>& GetConnections() const;
@@ -154,6 +151,18 @@ public:
 		}
 	}
 
+	template<typename TAnyType>
+	void SetValueFromAnyType(const typename TAnyType::FStorageType& InVal, Dataflow::FContext& Context) const
+	{
+		TAnyType::FPolicyType::VisitPolicyByType(GetType(),
+			[this, &Context, &InVal](auto SingleTypePolicy)
+			{
+				using FSingleType = decltype(SingleTypePolicy)::FReturnType;
+				FSingleType ValueToSet(InVal);
+				Context.SetData(CacheKey(), GetProperty(), Forward<FSingleType>(ValueToSet), GetOwningNodeGuid(), GetOwningNodeValueHash(), Dataflow::FTimestamp::Current());
+			});
+	}
+
 	template<class T>
 	const T& GetValue(Dataflow::FContext& Context, const T& Default) const
 	{
@@ -185,27 +194,49 @@ public:
 private:
 	DATAFLOWCORE_API const FDataflowInput* GetPassthroughInput() const;
 };
- 
-template<class T>
+
+template<typename T>
 const T& FDataflowInput::GetValue(Dataflow::FContext& Context, const T& Default) const
 {
-	if (GetConnectedOutputs().Num())
+	if (const FDataflowOutput* ConnectionOut = GetConnection())
 	{
-		ensure(GetConnectedOutputs().Num() == 1);
-		if (const FDataflowOutput* ConnectionOut = GetConnection())
+		if (!ConnectionOut->Evaluate(Context))
 		{
-			if (!ConnectionOut->Evaluate(Context))
-			{
-				Context.SetData(ConnectionOut->CacheKey(), Property, Default, GetOwningNodeGuid(), GetOwningNodeValueHash(), Dataflow::FTimestamp::Current());
-			}
-			if (Context.HasData(ConnectionOut->CacheKey()))
+			Context.SetData(ConnectionOut->CacheKey(), Property, Default, GetOwningNodeGuid(), GetOwningNodeValueHash(), Dataflow::FTimestamp::Current());
+		}
+		if (Context.HasData(ConnectionOut->CacheKey()))
+		{
+			const T& data = Context.GetData(ConnectionOut->CacheKey(), Property, Default);
+			return data;
+		}
+	}
+	return Default;
+}
+
+template<typename TAnyType>
+typename TAnyType::FStorageType FDataflowInput::GetValueFromAnyType(Dataflow::FContext& Context, const typename TAnyType::FStorageType& Default) const
+{
+	typename TAnyType::FStorageType ReturnValue = Default;
+	if (const FDataflowOutput* ConnectionOut = GetConnection())
+	{
+		if (ConnectionOut->Evaluate(Context))
+		{
+			if (const TUniquePtr<Dataflow::FContextCacheElementBase>* CacheEntry = Context.GetDataImpl(ConnectionOut->CacheKey()))
 			{
-				const T& data = Context.GetData(ConnectionOut->CacheKey(), Property, Default);
-				return data;
+				if (*CacheEntry)
+				{
+					TAnyType::FPolicyType::VisitPolicyByType(GetType(),
+						[this, &Context, &CacheEntry, &ReturnValue](auto SingleTypePolicy)
+						{
+							using FSingleType = decltype(SingleTypePolicy)::FReturnType;
+							FSingleType Default{};
+							ReturnValue = (*CacheEntry)->GetTypedData<FSingleType>(Context, nullptr, Default);
+						});
+				}
 			}
 		}
 	}
-	return Default;
+	return ReturnValue;
 }
 
 template<class T>
diff --git a/Engine/Source/Runtime/Experimental/Dataflow/Core/Public/Dataflow/DataflowNode.h b/Engine/Source/Runtime/Experimental/Dataflow/Core/Public/Dataflow/DataflowNode.h
--- a/Engine/Source/Runtime/Experimental/Dataflow/Core/Public/Dataflow/DataflowNode.h
+++ b/Engine/Source/Runtime/Experimental/Dataflow/Core/Public/Dataflow/DataflowNode.h
@@ -7,6 +7,7 @@
 #include "Dataflow/DataflowConnection.h"
 #include "Dataflow/DataflowInputOutput.h"
 #include "Dataflow/DataflowNodeParameters.h"
+#include "Dataflow/DataflowAnyType.h"
 #include "UObject/StructOnScope.h"
 #include "Dataflow/DataflowSettings.h"
 
@@ -112,6 +113,9 @@ struct FDataflowNode
 	/** Override this function to add the RemoveOptionPin functionality to the node's context menu. */
 	virtual bool CanRemovePin() const { return false; }
 
+	DATAFLOWCORE_API bool InputSupportsType(FName Name, FName Type) const;
+	DATAFLOWCORE_API bool OutputSupportsType(FName Name, FName Type) const;
+
 	DATAFLOWCORE_API virtual void AddInput(FDataflowInput* InPtr);
 	DATAFLOWCORE_API TArray< FDataflowInput* > GetInputs() const;
 	DATAFLOWCORE_API void ClearInputs();
@@ -167,19 +171,40 @@ struct FDataflowNode
 	virtual const UScriptStruct* TypedScriptStruct() const { return nullptr; }
 
 	/** Register the Input and Outputs after the creation in the factory. Use PropertyName to disambiguate a struct name from its first property. */
-	DATAFLOWCORE_API FDataflowInput* RegisterInputConnection(const void* Property, const FName& PropertyName = NAME_None);
-	DATAFLOWCORE_API FDataflowOutput* RegisterOutputConnection(
-		const void* Property,
-		const void* Passthrough = nullptr,
-		const FName& PropertyName = NAME_None,
-		const FName& PassthroughName = NAME_None);
+	template <typename T>
+	FDataflowInput* RegisterInputConnection(const T* Reference, const FName& PropertyName = NAME_None)
+	{
+		FDataflowInput* Input = RegisterInputConnectionInternal(Reference, PropertyName);
+		if constexpr (std::is_base_of_v<FDataflowAnyType, T>)
+		{
+			if (Input)
+			{
+				Input->SetTypePolicy(T::FPolicyType::GetInterface());
+			}
+		}
+		return Input;
+	}
+
+	template <typename T>
+	FDataflowOutput* RegisterOutputConnection(const T* Reference, const T* Passthrough = nullptr, const FName& PropertyName = NAME_None, const FName& PassthroughName = NAME_None)
+	{
+		FDataflowOutput* Output = RegisterOutputConnectionInternal(Reference, Passthrough, PropertyName, PassthroughName);
+		if constexpr (std::is_base_of_v<FDataflowAnyType, T>)
+		{
+			if (Output)
+			{
+				Output->SetTypePolicy(T::FPolicyType::GetInterface());
+			}
+		}
+		return Output;
+	}
+
+
 	/** Unregister the input connection if one exists matching this property, and then invalidate the graph. */
 	DATAFLOWCORE_API void UnregisterInputConnection(const void* Property, const FName& PropertyName = NAME_None);
 	/** Unregister the connection if one exists matching this pin, then invalidate the graph. */
 	DATAFLOWCORE_API void UnregisterPinConnection(const Dataflow::FPin& Pin);
 
-	DATAFLOWCORE_API void PropagateTypeToAllAnyTypeInputsAndOutputs(FName ConcreteType);
-
 	//
 	// Evaluation
 	//
@@ -195,11 +220,19 @@ struct FDataflowNode
 	*   @param Reference : Pointer to a member of this node that corresponds with the output to set.
 	*						*Reference will be used as the default if the input is not connected. 
 	*/
-	template<class T> const T& GetValue(Dataflow::FContext& Context, const T* Reference) const
+	template<class T, typename = std::enable_if_t<!std::is_base_of_v<FDataflowAnyType, T>>>
+	const T& GetValue(Dataflow::FContext& Context, const T* Reference) const
 	{
 		checkSlow(FindInput(Reference));
 		return FindInput(Reference)->template GetValue<T>(Context, *Reference);
 	}
+
+	template<typename TAnyType, typename = std::enable_if_t<std::is_base_of_v<FDataflowAnyType, TAnyType>>>
+	TAnyType::FStorageType GetValue(Dataflow::FContext& Context, const TAnyType* Reference) const
+	{
+		checkSlow(Reference && FindInput(Reference));
+		return FindInput(Reference)->template GetValueFromAnyType<TAnyType>(Context, Reference->Value);
+	}
 	
 	template<class T>
 	TFuture<const T&> GetValueParallel(Dataflow::FContext& Context, const T* Reference) const
@@ -239,7 +272,8 @@ struct FDataflowNode
 	*   @param Value : The value to store in the contexts data store. 
 	*   @param Reference : Pointer to a member of this node that corresponds with the output to set. 
 	*/
-	template<class T> void SetValue(Dataflow::FContext& Context, T&& Value, const typename TDecay<T>::Type* Reference) const
+	template<class T, typename = std::enable_if_t<!std::is_base_of_v<FDataflowAnyType, T>>>
+	void SetValue(Dataflow::FContext& Context, T&& Value, const typename TDecay<T>::Type* Reference) const
 	{
 		if (const FDataflowOutput* Output = FindOutput(Reference))
 		{
@@ -251,6 +285,19 @@ struct FDataflowNode
 		}
 	}
 
+	template<typename TAnyType, typename = std::enable_if_t<std::is_base_of_v<FDataflowAnyType, TAnyType>>>
+	void SetValue(Dataflow::FContext& Context, const typename TAnyType::FStorageType& Value, const TAnyType* Reference) const
+	{
+		if (const FDataflowOutput* Output = FindOutput(Reference))
+		{
+			Output->template SetValueFromAnyType<TAnyType>(Value, Context);
+		}
+		else
+		{
+			checkfSlow(false, TEXT("This output could not be found within this node, check this has been properly registered in the node constructor"));
+		}
+	}
+
 	/**
 	*   ForwardInput(...)
 	*
@@ -318,7 +365,30 @@ struct FDataflowNode
 	DECLARE_MULTICAST_DELEGATE_OneParam(FOnNodeInvalidated, FDataflowNode*);
 	FOnNodeInvalidated& GetOnNodeInvalidatedDelegate() { return OnNodeInvalidatedDelegate; }
 
+	// returns true if the type was changed successfully
+	// only unset datatype connection will be set a new type 
+	DATAFLOWCORE_API bool TrySetConnectionType(FDataflowConnection* Connection, FName NewType);
+
+protected:
+	virtual bool OnInputTypeChanged(const FDataflowInput* Input) { return false; };
+	virtual bool OnOutputTypeChanged(const FDataflowOutput* Input) { return false; }
+
+	// returns true if the input type was changed successfully
+	DATAFLOWCORE_API bool SetInputConcreteType(void* OutputReference, FName NewType);
+
+	// returns true if the output type was changed successfully
+	DATAFLOWCORE_API bool SetOutputConcreteType(void* InputReference, FName NewType);
+
+	DATAFLOWCORE_API FDataflowInput* RegisterInputConnectionInternal(const void* Property, const FName& PropertyName = NAME_None);
+	DATAFLOWCORE_API FDataflowOutput* RegisterOutputConnectionInternal(
+		const void* Property,
+		const void* Passthrough = nullptr,
+		const FName& PropertyName = NAME_None,
+		const FName& PassthroughName = NAME_None);
+
+
 private:
+	bool InitConnectionParametersFromPropertyReference(const FStructOnScope& StructOnScope, const void* PropertyRef, const FName& PropertyName, Dataflow::FConnectionParameters& OutParams);
 	static FString GetPropertyFullNameString(const TConstArrayView<const FProperty*>& PropertyChain);
 	static FName GetPropertyFullName(const TArray<const FProperty*>& PropertyChain);
 	static FText GetPropertyDisplayNameText(const TArray<const FProperty*>& PropertyChain);
@@ -392,14 +462,4 @@ private:
 
 }
 
-USTRUCT()
-struct FDataflowAnyType
-{
-	GENERATED_USTRUCT_BODY()
-	void* Value = nullptr;
-
-	DATAFLOWCORE_API static const FName TypeName;
-};
-
-
 
diff --git a/Engine/Source/Runtime/Experimental/Dataflow/Core/Public/Dataflow/DataflowNodeParameters.h b/Engine/Source/Runtime/Experimental/Dataflow/Core/Public/Dataflow/DataflowNodeParameters.h
--- a/Engine/Source/Runtime/Experimental/Dataflow/Core/Public/Dataflow/DataflowNodeParameters.h
+++ b/Engine/Source/Runtime/Experimental/Dataflow/Core/Public/Dataflow/DataflowNodeParameters.h
@@ -131,11 +131,9 @@ namespace Dataflow
 	template<class T>
 	const T& FContextCacheElementBase::GetTypedData(FContext& Context, const FProperty* PropertyIn, const T& Default) const
 	{
-		check(PropertyIn);
 		// check(PropertyIn->IsA<T>()); // @todo(dataflow) compile error for non-class T; find alternatives
 		if (Type == EType::CacheElementTyped)
 		{
-			check(Property->SameType(PropertyIn));
 			return static_cast<const FContextCacheElement<T>&>(*this).GetData(Context, PropertyIn, Default);
 		}
 		if (Type == EType::CacheElementReference)
diff --git a/Engine/Source/Runtime/Experimental/Dataflow/Core/Public/Dataflow/DataflowTypePolicy.h b/Engine/Source/Runtime/Experimental/Dataflow/Core/Public/Dataflow/DataflowTypePolicy.h
new file mode 100644
--- /dev/null
+++ b/Engine/Source/Runtime/Experimental/Dataflow/Core/Public/Dataflow/DataflowTypePolicy.h
@@ -0,0 +1,122 @@
+// Copyright Epic Games, Inc. All Rights Reserved.
+
+#pragma once
+
+#include "CoreMinimal.h"
+#include "Misc/TVariant.h"
+
+struct IDataflowTypePolicy
+{
+	virtual bool SupportsType(FName Type) const = 0;
+};
+
+struct FDataflowAllTypesPolicy : public IDataflowTypePolicy
+{
+	virtual bool SupportsType(FName InType) const override
+	{
+		return true;
+	}
+
+	static IDataflowTypePolicy* GetInterface()
+	{
+		static FDataflowAllTypesPolicy Instance;
+		return &Instance;
+	}
+};
+
+template <typename T>
+struct TDataflowSingleTypePolicy : public IDataflowTypePolicy
+{
+	using FReturnType = T;
+	virtual bool SupportsType(FName InType) const override
+	{
+		return SupportsTypeStatic(InType);
+	}
+
+	static bool SupportsTypeStatic(FName InType)
+	{
+		return (InType == TypeName);
+	}
+
+	template <typename TVisitor>
+	static bool VisitPolicyByType(FName RequestedType, TVisitor Visitor)
+	{
+		if (RequestedType == TypeName)
+		{
+			TDataflowSingleTypePolicy<T> SingleTypePolicy;
+			Visitor(SingleTypePolicy);
+			return true;
+		}
+		return false;
+	}
+
+	static IDataflowTypePolicy* GetInterface()
+	{
+		static TDataflowSingleTypePolicy Instance;
+		return &Instance;
+	}
+
+	inline static const FName TypeName = FName(TNameOf<T>::GetName());
+};
+
+template <typename... TTypes>
+struct TDataflowMultiTypePolicy;
+
+template <>
+struct TDataflowMultiTypePolicy<>: public IDataflowTypePolicy
+{
+	virtual bool SupportsType(FName InType) const override
+	{
+		return false;
+	}
+
+	static bool SupportsTypeStatic(FName InType)
+	{
+		return false;
+	}
+
+	template <typename TVisitor>
+	static bool VisitPolicyByType(FName RequestedType, TVisitor Visitor)
+	{
+		return false;
+	}
+};
+
+template <typename T, typename... TTypes>
+struct TDataflowMultiTypePolicy<T, TTypes...>: public TDataflowMultiTypePolicy<TTypes...>
+{
+	using Super = TDataflowMultiTypePolicy<TTypes...>;
+	using FReturnType = T;
+
+	virtual bool SupportsType(FName InType) const override
+	{
+		return SupportsTypeStatic(InType);
+	}
+
+	static bool SupportsTypeStatic(FName InType)
+	{
+		return TDataflowSingleTypePolicy<T>::SupportsTypeStatic(InType)
+			|| Super::SupportsTypeStatic(InType);
+	}
+
+	template <typename TVisitor>
+	static bool VisitPolicyByType(FName RequestedType, TVisitor Visitor)
+	{
+		if (TDataflowSingleTypePolicy<T>::VisitPolicyByType(RequestedType, Visitor))
+		{
+			return true;
+		}
+		return Super::VisitPolicyByType(RequestedType, Visitor);
+	}
+
+	static IDataflowTypePolicy* GetInterface()
+	{
+		static TDataflowMultiTypePolicy<T, TTypes...> Instance;
+		return &Instance;
+	}
+};
+
+struct FDataflowNumericTypePolicy : 
+	public TDataflowMultiTypePolicy<double, float, int64, uint64, int32, uint32, int16, uint16, int8, uint8>
+{
+};
diff --git a/Engine/Source/Runtime/Experimental/Dataflow/Engine/Private/Dataflow/DataflowEdNode.cpp b/Engine/Source/Runtime/Experimental/Dataflow/Engine/Private/Dataflow/DataflowEdNode.cpp
--- a/Engine/Source/Runtime/Experimental/Dataflow/Engine/Private/Dataflow/DataflowEdNode.cpp
+++ b/Engine/Source/Runtime/Experimental/Dataflow/Engine/Private/Dataflow/DataflowEdNode.cpp
@@ -266,6 +266,25 @@ void UDataflowEdNode::RemoveOptionPin()
 #endif // WITH_EDITOR && !UE_BUILD_SHIPPING
 }
 
+bool UDataflowEdNode::PinIsCompatibleWithType(const UEdGraphPin& Pin, const FEdGraphPinType& PinType) const
+{
+#if WITH_EDITOR
+	check(Pin.GetOwningNode() == this);
+	if (TSharedPtr<const FDataflowNode> DataflowNode = GetDataflowNode())
+	{
+		if (Pin.Direction == EEdGraphPinDirection::EGPD_Input)
+		{
+			return DataflowNode->InputSupportsType(Pin.GetFName(), PinType.PinCategory);
+		}
+		if (Pin.Direction == EEdGraphPinDirection::EGPD_Output)
+		{
+			return DataflowNode->OutputSupportsType(Pin.GetFName(), PinType.PinCategory);
+		}
+	}
+#endif WITH_EDITOR
+	return false;
+}
+
 FText UDataflowEdNode::GetNodeTitle(ENodeTitleType::Type TitleType) const
 {
 	return FText::FromString(GetName());
@@ -293,14 +312,16 @@ void UDataflowEdNode::PinConnectionListChanged(UEdGraphPin* Pin)
 								{
 									if (FDataflowOutput* LinkedConOutput = LinkedDataflowNode->FindOutput(FName(LinkedCon->GetName())))
 									{
+										const bool bTypeChanged = DataflowNode->TrySetConnectionType(ConnectionInput, LinkedConOutput->GetType());
 										DataflowGraph->Connect(LinkedConOutput, ConnectionInput);
+
+										if (bTypeChanged)
+										{
+											UpdatePinsFromDataflowNode();
+										}
+
 									}
 								}
-								OnPinConnection(*Pin, LinkedCon->PinType);
-							}
-							else
-							{
-								OnPinDisconnection(*Pin);
 							}
 						}
 					}
@@ -321,14 +342,16 @@ void UDataflowEdNode::PinConnectionListChanged(UEdGraphPin* Pin)
 								{
 									if (FDataflowInput* LinkedConInput = LinkedDataflowNode->FindInput(FName(LinkedCon->GetName())))
 									{
+										const bool bTypeChanged = DataflowNode->TrySetConnectionType(ConnectionOutput, LinkedConInput->GetType());
 										DataflowGraph->Connect(ConnectionOutput, LinkedConInput);
+
+										if (bTypeChanged)
+										{
+											UpdatePinsFromDataflowNode();
+										}
 									}
 								}
-								OnPinConnection(*Pin, LinkedCon->PinType);
-							}
-							else
-							{
-								OnPinDisconnection(*Pin);
+								
 							}
 						}
 					}
@@ -342,62 +365,6 @@ void UDataflowEdNode::PinConnectionListChanged(UEdGraphPin* Pin)
 
 #endif // WITH_EDITOR && !UE_BUILD_SHIPPING
 
-void UDataflowEdNode::OnPinConnection(const UEdGraphPin& Pin, const FEdGraphPinType& Type)
-{
-#if WITH_EDITOR
-	if (Pin.PinType.PinCategory == FDataflowAnyType::TypeName)
-	{
-		SetAnyTypePinsToType(Type);
-	}
-#endif
-}
-
-void UDataflowEdNode::OnPinDisconnection(const UEdGraphPin& Pin)
-{
-}
-
-void UDataflowEdNode::SetAnyTypePinsToType(const FEdGraphPinType& Type)
-{
-#if WITH_EDITOR
-	if (const TSharedPtr<FDataflowNode> DataflowNode = this->GetDataflowNode())
-	{
-		for (UEdGraphPin* Pin : GetAllPins())
-		{
-			if (Pin)
-			{
-				switch (Pin->Direction)
-				{
-					case EEdGraphPinDirection::EGPD_Input:
-					{
-						if (FDataflowInput* Input = DataflowNode->FindInput(Pin->GetFName()))
-						{
-							if (Input->IsAnyType())
-							{
-								Pin->PinType = Type;
-								Input->SetConcreteType(Type.PinCategory);
-							}
-						}
-						break;
-					}
-					case EEdGraphPinDirection::EGPD_Output:
-					{
-						if (FDataflowOutput* Output = DataflowNode->FindOutput(Pin->GetFName()))
-						{
-							if (Output->IsAnyType())
-							{
-								Pin->PinType = Type;
-								Output->SetConcreteType(Type.PinCategory);
-							}
-						}
-						break;
-					}
-				}
-			}
-		}
-	}
-#endif
-}
-
 void UDataflowEdNode::Serialize(FArchive& Ar)
 {
 	Super::Serialize(Ar);
@@ -526,61 +493,50 @@ void UDataflowEdNode::GetPinHoverText(const UEdGraphPin& Pin, FString& HoverText
 
 void UDataflowEdNode::AutowireNewNode(UEdGraphPin* FromPin)
 {
-	if (DataflowGraph && FromPin)
+	const UEdGraph* EdGraph = this->GetGraph();
+	if (EdGraph == nullptr)
 	{
-		if (UEdGraphNode* FromGraphNode = FromPin->GetOwningNode())
+		return;
+	}
+
+	if (!DataflowGraph || !FromPin)
+	{
+		return;
+	}
+
+	if (UEdGraphNode* FromGraphNode = FromPin->GetOwningNode())
+	{
+		if (FromPin->Direction == EEdGraphPinDirection::EGPD_Output)
 		{
-			if (UDataflowEdNode* FromDataflowGraphNode = Cast<UDataflowEdNode>(FromPin->GetOwningNode()))
+			for (UEdGraphPin* InputPin : this->GetAllPins())
 			{
-				const TSharedPtr<FDataflowNode> FromDataFlowNode = FromDataflowGraphNode->GetDataflowNode();
-				if (FromDataFlowNode)
+				if (InputPin->Direction == EEdGraphPinDirection::EGPD_Input)
 				{
-					if (FDataflowOutput* FromOutput = FromDataFlowNode->FindOutput(FromPin->PinName))
+					if (this->PinIsCompatibleWithType(*InputPin, FromPin->PinType))
 					{
-						const TSharedPtr<FDataflowNode> ToDataFlowNode = this->GetDataflowNode();
-						for (UEdGraphPin* InputPin : this->GetAllPins())
+						if (EdGraph->GetSchema()->TryCreateConnection(FromPin, InputPin))
 						{
-							if (FDataflowInput* ToInput = ToDataFlowNode->FindInput(InputPin->PinName))
-							{
-								const bool bAreSameType = (ToInput->GetType() == FromOutput->GetType());
-								const bool bOnlyOneIsAnyType = (!bAreSameType && (ToInput->GetType() == FDataflowAnyType::TypeName) != (FromOutput->GetType() == FDataflowAnyType::TypeName));
-								if (bAreSameType || bOnlyOneIsAnyType)
-								{
-									if (const UEdGraph* EdGraph = this->GetGraph())
-									{
-										if (EdGraph->GetSchema()->TryCreateConnection(FromPin, InputPin))
-										{
-											FromGraphNode->NodeConnectionListChanged();
-											this->NodeConnectionListChanged();
-											break;
-										}
-									}
-								}
-							}
+							FromGraphNode->NodeConnectionListChanged();
+							this->NodeConnectionListChanged();
+							return;
 						}
 					}
-					else if(FDataflowInput* FromInput = FromDataFlowNode->FindInput(FromPin->PinName))
+				}
+			}
+		}
+		if (FromPin->Direction == EEdGraphPinDirection::EGPD_Input)
+		{
+			for (UEdGraphPin* OutputPin : this->GetAllPins())
+			{
+				if (OutputPin->Direction == EEdGraphPinDirection::EGPD_Output)
+				{
+					if (this->PinIsCompatibleWithType(*OutputPin, FromPin->PinType))
 					{
-						const TSharedPtr<FDataflowNode> ToDataFlowNode = this->GetDataflowNode();
-						for (UEdGraphPin* InputPin : this->GetAllPins())
+						if (EdGraph->GetSchema()->TryCreateConnection(FromPin, OutputPin))
 						{
-							if (FDataflowOutput* ToOutput = ToDataFlowNode->FindOutput(InputPin->PinName))
-							{
-								const bool bAreSameType = (ToOutput->GetType() == FromInput->GetType());
-								const bool bOnlyOneIsAnyType = (!bAreSameType && (ToOutput->GetType() == FDataflowAnyType::TypeName) != (FromInput->GetType() == FDataflowAnyType::TypeName));
-								if (bAreSameType || bOnlyOneIsAnyType)
-								{
-									if (const UEdGraph* EdGraph = this->GetGraph())
-									{
-										if (EdGraph->GetSchema()->TryCreateConnection(FromPin, InputPin))
-										{
-											FromGraphNode->NodeConnectionListChanged();
-											this->NodeConnectionListChanged();
-											break;
-										}
-									}
-								}
-							}
+							FromGraphNode->NodeConnectionListChanged();
+							this->NodeConnectionListChanged();
+							return;
 						}
 					}
 				}
diff --git a/Engine/Source/Runtime/Experimental/Dataflow/Engine/Public/Dataflow/DataflowEdNode.h b/Engine/Source/Runtime/Experimental/Dataflow/Engine/Public/Dataflow/DataflowEdNode.h
--- a/Engine/Source/Runtime/Experimental/Dataflow/Engine/Public/Dataflow/DataflowEdNode.h
+++ b/Engine/Source/Runtime/Experimental/Dataflow/Engine/Public/Dataflow/DataflowEdNode.h
@@ -76,6 +76,8 @@ public:
 	/** Remove an option pin if the underlying Dataflow node RemovePin member is overriden. */
 	DATAFLOWENGINE_API void RemoveOptionPin();
 
+	DATAFLOWENGINE_API bool PinIsCompatibleWithType(const UEdGraphPin& Pin, const FEdGraphPinType& PinType) const;
+
 #if WITH_EDITOR
 	// Pin hiding
 	DATAFLOWENGINE_API void HideAllInputPins();
@@ -95,11 +97,5 @@ public:
 
 	UPROPERTY()
 	bool bRenderInAssetEditor = false;
-
-private:
-	void OnPinConnection(const UEdGraphPin& Pin, const FEdGraphPinType& Type);
-	void OnPinDisconnection(const UEdGraphPin& Pin);
-	void SetAnyTypePinsToType(const FEdGraphPinType& Type);
-
 };
 
