Fix compilation error

diff --git a/Engine/Source/Runtime/Experimental/Dataflow/Engine/Private/Dataflow/DataflowEdNode.cpp b/Engine/Source/Runtime/Experimental/Dataflow/Engine/Private/Dataflow/DataflowEdNode.cpp
--- a/Engine/Source/Runtime/Experimental/Dataflow/Engine/Private/Dataflow/DataflowEdNode.cpp
+++ b/Engine/Source/Runtime/Experimental/Dataflow/Engine/Private/Dataflow/DataflowEdNode.cpp
@@ -281,7 +281,7 @@ bool UDataflowEdNode::PinIsCompatibleWithType(const UEdGraphPin& Pin, const FEdG
 			return DataflowNode->OutputSupportsType(Pin.GetFName(), PinType.PinCategory);
 		}
 	}
-#endif WITH_EDITOR
+#endif // WITH_EDITOR
 	return false;
 }
 
