Fix compilation error ( fallout from 33615268)

diff --git a/Engine/Source/Runtime/Experimental/Dataflow/Core/Public/Dataflow/DataflowInputOutput.h b/Engine/Source/Runtime/Experimental/Dataflow/Core/Public/Dataflow/DataflowInputOutput.h
--- a/Engine/Source/Runtime/Experimental/Dataflow/Core/Public/Dataflow/DataflowInputOutput.h
+++ b/Engine/Source/Runtime/Experimental/Dataflow/Core/Public/Dataflow/DataflowInputOutput.h
@@ -157,7 +157,7 @@ public:
 		TAnyType::FPolicyType::VisitPolicyByType(GetType(),
 			[this, &Context, &InVal](auto SingleTypePolicy)
 			{
-				using FSingleType = decltype(SingleTypePolicy)::FReturnType;
+				using FSingleType = typename decltype(SingleTypePolicy)::FReturnType;
 				FSingleType ValueToSet(InVal);
 				Context.SetData(CacheKey(), GetProperty(), Forward<FSingleType>(ValueToSet), GetOwningNodeGuid(), GetOwningNodeValueHash(), Dataflow::FTimestamp::Current());
 			});
@@ -228,7 +228,7 @@ typename TAnyType::FStorageType FDataflowInput::GetValueFromAnyType(Dataflow::FC
 					TAnyType::FPolicyType::VisitPolicyByType(GetType(),
 						[this, &Context, &CacheEntry, &ReturnValue](auto SingleTypePolicy)
 						{
-							using FSingleType = decltype(SingleTypePolicy)::FReturnType;
+							using FSingleType = typename decltype(SingleTypePolicy)::FReturnType;
 							FSingleType Default{};
 							ReturnValue = (*CacheEntry)->GetTypedData<FSingleType>(Context, nullptr, Default);
 						});
diff --git a/Engine/Source/Runtime/Experimental/Dataflow/Core/Public/Dataflow/DataflowNode.h b/Engine/Source/Runtime/Experimental/Dataflow/Core/Public/Dataflow/DataflowNode.h
--- a/Engine/Source/Runtime/Experimental/Dataflow/Core/Public/Dataflow/DataflowNode.h
+++ b/Engine/Source/Runtime/Experimental/Dataflow/Core/Public/Dataflow/DataflowNode.h
@@ -228,7 +228,7 @@ struct FDataflowNode
 	}
 
 	template<typename TAnyType, typename = std::enable_if_t<std::is_base_of_v<FDataflowAnyType, TAnyType>>>
-	TAnyType::FStorageType GetValue(Dataflow::FContext& Context, const TAnyType* Reference) const
+	typename TAnyType::FStorageType GetValue(Dataflow::FContext& Context, const TAnyType* Reference) const
 	{
 		checkSlow(Reference && FindInput(Reference));
 		return FindInput(Reference)->template GetValueFromAnyType<TAnyType>(Context, Reference->Value);
