Fix static analysis warning
in that case this was intentional so I'm disabling the warning using //-V792

diff --git a/Engine/Source/Runtime/Experimental/Dataflow/Core/Private/Dataflow/DataflowCoreNodes.cpp b/Engine/Source/Runtime/Experimental/Dataflow/Core/Private/Dataflow/DataflowCoreNodes.cpp
--- a/Engine/Source/Runtime/Experimental/Dataflow/Core/Private/Dataflow/DataflowCoreNodes.cpp
+++ b/Engine/Source/Runtime/Experimental/Dataflow/Core/Private/Dataflow/DataflowCoreNodes.cpp
@@ -55,17 +55,21 @@ void FDataflowBranchNode::Evaluate(Dataflow::FContext& Context, const FDataflowO
 
 bool FDataflowBranchNode::OnInputTypeChanged(const FDataflowInput* Input)
 {
+	// using single bitwise | operator to avoid skipping calls to SetInputConcreteType because of the shortcircuiting of ||
+	// need to disable the warning for static analysis ( V792 warning )
 	return bool(
 		SetInputConcreteType(&TrueValue, Input->GetType())
-	  | SetInputConcreteType(&FalseValue, Input->GetType())
-	  | SetOutputConcreteType(&Result, Input->GetType())
+	  | SetInputConcreteType(&FalseValue, Input->GetType()) //-V792
+	  | SetOutputConcreteType(&Result, Input->GetType()) //-V792
 		);
 }
 
 bool FDataflowBranchNode::OnOutputTypeChanged(const FDataflowOutput* Input)
 {
+	// using single bitwise | operator to avoid skipping calls to SetInputConcreteType because of the shortcircuiting of ||
+	// need to disable the warning for static analysis ( V792 warning )
 	return bool(
 		  SetInputConcreteType(&TrueValue, Input->GetType())
-		| SetInputConcreteType(&FalseValue, Input->GetType())
+		| SetInputConcreteType(&FalseValue, Input->GetType()) //-V792
 		);
 }
\ No newline at end of file
