Disable tick on ChaosDebugDrawComponent in test/shipping builds.

#tests Tested we dont tick the component in a test build

diff --git a/Engine/Source/Runtime/Experimental/ChaosSolverEngine/Private/Chaos/ChaosDebugDrawComponent.cpp b/Engine/Source/Runtime/Experimental/ChaosSolverEngine/Private/Chaos/ChaosDebugDrawComponent.cpp
--- a/Engine/Source/Runtime/Experimental/ChaosSolverEngine/Private/Chaos/ChaosDebugDrawComponent.cpp
+++ b/Engine/Source/Runtime/Experimental/ChaosSolverEngine/Private/Chaos/ChaosDebugDrawComponent.cpp
@@ -218,6 +218,7 @@ void DebugDrawChaos(const AActor* DebugDrawActor, const TArray<Chaos::FLatentDra
 UChaosDebugDrawComponent::UChaosDebugDrawComponent()
 	: bInPlay(false)
 {
+#if CHAOS_DEBUG_DRAW
 	// We must tick after anything that uses Chaos Debug Draw and also after the Line Batcher Component
 	PrimaryComponentTick.bAllowTickOnDedicatedServer = false;
 	PrimaryComponentTick.bCanEverTick = true;
@@ -225,6 +226,7 @@ UChaosDebugDrawComponent::UChaosDebugDrawComponent()
 	PrimaryComponentTick.TickGroup = TG_PostUpdateWork;
 
 	bTickInEditor = true;
+#endif
 }
 
 void UChaosDebugDrawComponent::BeginDestroy()
@@ -241,7 +243,9 @@ void UChaosDebugDrawComponent::BeginPlay()
 {
 	Super::BeginPlay();
 
+#if CHAOS_DEBUG_DRAW
 	SetTickableWhenPaused(true);
+#endif
 
 	bInPlay = true;
 
@@ -256,7 +260,9 @@ void UChaosDebugDrawComponent::EndPlay(const EEndPlayReason::Type EndPlayReason)
 {
 	Super::EndPlay(EndPlayReason);
 
+#if CHAOS_DEBUG_DRAW
 	SetTickableWhenPaused(false);
+#endif
 
 	bInPlay = false;
 
