Option to be explicit about time on systems where the epoch is quirky

#rnx

diff --git a/Engine/Source/Runtime/Experimental/IoStore/HttpClient/Private/Peer.inl b/Engine/Source/Runtime/Experimental/IoStore/HttpClient/Private/Peer.inl
--- a/Engine/Source/Runtime/Experimental/IoStore/HttpClient/Private/Peer.inl
+++ b/Engine/Source/Runtime/Experimental/IoStore/HttpClient/Private/Peer.inl
@@ -2,6 +2,15 @@
 
 #pragma once
 
+#if defined(IAS_HTTP_EXPLICIT_VERIFY_TIME)
+#	if !IAS_HTTP_EXPLICIT_VERIFY_TIME 
+#		error Either define this to >=1 or not at all
+#	endif
+#	include <HAL/PlatformTime.h>
+
+#	include <ctime>
+#endif
+
 namespace UE::IoStore::HTTP
 {
 
@@ -113,6 +122,27 @@ FCertRoots::FCertRoots(FMemoryView PemData)
 		return;
 	}
 
+#if defined(IAS_HTTP_EXPLICIT_VERIFY_TIME)
+	if (X509_VERIFY_PARAM* VerifyParam = SSL_CTX_get0_param(Context); VerifyParam != nullptr)
+	{
+		int32 AliasTown;
+		std::tm Utc = {};
+		FPlatformTime::UtcTime(
+			Utc.tm_year, Utc.tm_mon,
+			AliasTown,
+			Utc.tm_mday, Utc.tm_hour, Utc.tm_min,
+			AliasTown, AliasTown
+		);
+
+		Utc.tm_year -= 1900;
+		Utc.tm_mon -= 1;
+
+		time_t Now = std::mktime(&Utc);
+
+		X509_VERIFY_PARAM_set_time(VerifyParam, Now);
+	}
+#endif
+
 	Handle = UPTRINT(Context);
 }
 
