Match constructor signatures.

#rnx

diff --git a/Engine/Source/Runtime/Experimental/IoStore/HttpClient/Private/Peer.inl b/Engine/Source/Runtime/Experimental/IoStore/HttpClient/Private/Peer.inl
--- a/Engine/Source/Runtime/Experimental/IoStore/HttpClient/Private/Peer.inl
+++ b/Engine/Source/Runtime/Experimental/IoStore/HttpClient/Private/Peer.inl
@@ -409,7 +409,7 @@ class FHttpPeer
 {
 public:
 				FHttpPeer() = default;
-				FHttpPeer(FSocket InSocket, FSslContext* Context=nullptr);
+				FHttpPeer(FSocket InSocket, FSslContext* Context=nullptr, const char* HostName=nullptr);
 	FOutcome	Handshake();
 
 private:
@@ -418,8 +418,8 @@ private:
 };
 
 ////////////////////////////////////////////////////////////////////////////////
-FHttpPeer::FHttpPeer(FSocket InSocket, FSslContext* Context)
-: FTlsPeer(MoveTemp(InSocket), Context)
+FHttpPeer::FHttpPeer(FSocket InSocket, FSslContext* Context, const char* HostName)
+: FTlsPeer(MoveTemp(InSocket), Context, HostName)
 {
 	if (Ssl == nullptr)
 	{
