Iris
- Fixed missing referencecount on destructioninfoprotocol

diff --git a/Engine/Source/Runtime/Experimental/Iris/Core/Private/Iris/ReplicationSystem/ReplicationBridge.cpp b/Engine/Source/Runtime/Experimental/Iris/Core/Private/Iris/ReplicationSystem/ReplicationBridge.cpp
--- a/Engine/Source/Runtime/Experimental/Iris/Core/Private/Iris/ReplicationSystem/ReplicationBridge.cpp
+++ b/Engine/Source/Runtime/Experimental/Iris/Core/Private/Iris/ReplicationSystem/ReplicationBridge.cpp
@@ -375,7 +375,7 @@ void UReplicationBridge::Deinitialize()
 		}
 	});
 
-	// Just set the protocol to null, it will be destroyed with the protocolmanager
+	// Release destructioninfo protocol.
 	DestructionInfoProtocol = nullptr;
 
 	ReplicationSystem = nullptr;
diff --git a/Engine/Source/Runtime/Experimental/Iris/Core/Public/Iris/ReplicationSystem/ReplicationBridge.h b/Engine/Source/Runtime/Experimental/Iris/Core/Public/Iris/ReplicationSystem/ReplicationBridge.h
--- a/Engine/Source/Runtime/Experimental/Iris/Core/Public/Iris/ReplicationSystem/ReplicationBridge.h
+++ b/Engine/Source/Runtime/Experimental/Iris/Core/Public/Iris/ReplicationSystem/ReplicationBridge.h
@@ -10,6 +10,7 @@
 #include "Iris/ReplicationSystem/NetObjectGroupHandle.h"
 #include "Misc/EnumClassFlags.h"
 #include "Net/Core/NetHandle/NetHandle.h"
+#include "Templates/RefCounting.h"
 #include "UObject/ObjectKey.h"
 #include "UObject/ObjectMacros.h"
 
@@ -384,7 +385,7 @@ private:
 		UE::Net::Private::FInternalNetRefIndex InternalReplicationIndex;
 	};
 
-	const UE::Net::FReplicationProtocol* DestructionInfoProtocol;
+	TRefCountPtr<const UE::Net::FReplicationProtocol> DestructionInfoProtocol;
 	
 	// Need to track the objects with destruction infos so that we can clean them up properly
 	// We use this to be able ask remote to destroy static objects
