Move upload of on-demand container file(s) to I/O store tool

diff --git a/Engine/Source/Runtime/Experimental/IoStore/OnDemand/Private/IoStoreOnDemand.cpp b/Engine/Source/Runtime/Experimental/IoStore/OnDemand/Private/IoStoreOnDemand.cpp
--- a/Engine/Source/Runtime/Experimental/IoStore/OnDemand/Private/IoStoreOnDemand.cpp
+++ b/Engine/Source/Runtime/Experimental/IoStore/OnDemand/Private/IoStoreOnDemand.cpp
@@ -304,23 +304,6 @@ static FIasCacheConfig GetIasCacheConfig(const TCHAR* CommandLine)
 	return Ret;
 }
 
-////////////////////////////////////////////////////////////////////////////////
-
-/** Utility for saving data to disk, similar to FFileHelper::SaveArrayToFile but supporting larger sizes */
-static bool SaveArrayToFile(TArrayView64<uint8> Data, const TCHAR* Filename, uint32 WriteFlags = FILEWRITE_None)
-{
-	TUniquePtr<FArchive> Ar = TUniquePtr<FArchive>(IFileManager::Get().CreateFileWriter(Filename, WriteFlags));
-	if (!Ar)
-	{
-		return false;
-	}
-
-	Ar->Serialize(static_cast<void*>(Data.GetData()), Data.Num());
-	Ar->Close();
-
-	return !Ar->IsError() && !Ar->IsCriticalError();
-}
-
 ////////////////////////////////////////////////////////////////////////////////
 /**
  * Utility to create a FArchive capable of reading from disk using the exact same pathing
@@ -959,819 +942,6 @@ static void ToJson(FJsonWriter JsonWriter, const FOnDemandToc& Toc, EOnDemandToc
 	JsonWriter->WriteArrayEnd();
 }
 
-////////////////////////////////////////////////////////////////////////////////
-class FS3UploadQueue
-{
-public:
-	FS3UploadQueue(FS3Client& Client, const FString& Bucket, int32 ThreadCount);
-	bool Enqueue(FStringView Key, FIoBuffer Payload);
-	bool Flush();
-
-private:
-	void ThreadEntry();
-
-	struct FQueueEntry
-	{
-		FString Key;
-		FIoBuffer Payload;
-	};
-
-	FS3Client& Client;
-	TArray<TFuture<void>> Threads;
-	FString Bucket;
-	FCriticalSection CriticalSection;
-	TQueue<FQueueEntry> Queue;
-	FEventRef WakeUpEvent;
-	FEventRef UploadCompleteEvent;
-	std::atomic_int32_t ConcurrentUploads{0};
-	std::atomic_int32_t ActiveThreadCount{0};
-	std::atomic_int32_t ErrorCount{0};
-	std::atomic_bool bCompleteAdding{false};
-};
-
-FS3UploadQueue::FS3UploadQueue(FS3Client& InClient, const FString& InBucket, int32 ThreadCount)
-	: Client(InClient)
-	, Bucket(InBucket)
-{
-	ActiveThreadCount = ThreadCount;
-	for (int32 Idx = 0; Idx < ThreadCount; ++Idx)
-	{
-		Threads.Add(AsyncThread([this]()
-		{
-			ThreadEntry();
-		}));
-	}
-}
-
-bool FS3UploadQueue::Enqueue(FStringView Key, FIoBuffer Payload)
-{
-	if (ActiveThreadCount == 0)
-	{
-		return false;
-	}
-
-	for(;;)
-	{
-		bool bEnqueued = false;
-		{
-			FScopeLock _(&CriticalSection);
-			if (ConcurrentUploads < Threads.Num())
-			{
-				bEnqueued = Queue.Enqueue(FQueueEntry {FString(Key), Payload});
-			}
-		}
-
-		if (bEnqueued)
-		{
-			WakeUpEvent->Trigger();
-			break;
-		}
-
-		UploadCompleteEvent->Wait();
-	}
-
-	return true;
-}
-
-void FS3UploadQueue::ThreadEntry()
-{
-	for(;;)
-	{
-		FQueueEntry Entry;
-		bool bDequeued = false;
-		{
-			FScopeLock _(&CriticalSection);
-			bDequeued = Queue.Dequeue(Entry);
-		}
-
-		if (!bDequeued)
-		{
-			if (bCompleteAdding)
-			{
-				break;
-			}
-			WakeUpEvent->Wait();
-			continue;
-		}
-
-		ConcurrentUploads++;
-		const FS3PutObjectResponse Response = Client.TryPutObject(FS3PutObjectRequest{Bucket, Entry.Key, Entry.Payload.GetView()});
-		ConcurrentUploads--;
-		UploadCompleteEvent->Trigger();
-
-		if (Response.IsOk())
-		{
-			UE_LOG(LogIas, Log, TEXT("Uploaded chunk '%s/%s/%s'"), *Client.GetConfig().ServiceUrl, *Bucket, *Entry.Key);
-		}
-		else
-		{
-			TStringBuilder<256> ErrorResponse;
-			Response.GetErrorResponse(ErrorResponse);
-
-			UE_LOG(LogIas, Warning, TEXT("Failed to upload chunk '%s/%s/%s' (%s)"), *Client.GetConfig().ServiceUrl, *Bucket, *Entry.Key, ErrorResponse.ToString());
-			ErrorCount++;
-			break;
-		}
-	}
-
-	ActiveThreadCount--;
-}
-
-bool FS3UploadQueue::Flush()
-{
-	bCompleteAdding = true;
-	for (int32 Idx = 0; Idx < Threads.Num(); ++Idx)
-	{
-		WakeUpEvent->Trigger();
-	}
-
-	for (TFuture<void>& Thread : Threads)
-	{
-		Thread.Wait();
-	}
-
-	return ErrorCount == 0;
-}
-
-////////////////////////////////////////////////////////////////////////////////
-static void GetTocPath(
-	FStringView ServiceUrl,
-	FStringView Bucket,
-	FStringView BucketPrefix,
-	FString& OutServiceUrl,
-	FString& OutTocPath)
-{
-	// The configuration file should specify a service URL without any trailing
-	// host path, i.e. http://{host:port}/{host-path}. Add the trailing path
-	// to the TOC path to form the complete path the TOC from the host, i.e
-	// TocPath={host-path}/{bucket}/{bucket-prefix}/{toc-hash}.uchunktoc
-	
-	if (BucketPrefix.StartsWith(TEXT("/")))
-	{
-		BucketPrefix.RemovePrefix(1);
-	}
-	if (BucketPrefix.EndsWith(TEXT("/")))
-	{
-		BucketPrefix.RemoveSuffix(1);
-	}
-
-	if (ServiceUrl.IsEmpty())
-	{
-		// If the service URL is empty we assume uploading to AWS S3 using the Region parameter
-		// and that we don't need to prefix with the bucket name.
-		OutServiceUrl.Empty();
-		OutTocPath = BucketPrefix;
-	}
-	else
-	{
-		FStringView HostSuffix;
-		int32 Idx = INDEX_NONE;
-		ensure(ServiceUrl.FindChar(':', Idx));
-		const int32 SchemeEnd = Idx + 3;
-
-		Idx = INDEX_NONE;
-		if (ServiceUrl.RightChop(SchemeEnd).FindChar('/', Idx))
-		{
-			OutServiceUrl = ServiceUrl.Left(SchemeEnd + Idx);
-			HostSuffix = ServiceUrl.RightChop(OutServiceUrl.Len() + 1);
-			if (HostSuffix.EndsWith(TEXT("/")))
-			{
-				HostSuffix.RemoveSuffix(1);
-			}
-		}
-		else
-		{
-			OutServiceUrl = ServiceUrl;
-		}
-
-		TStringBuilder<256> Sb;
-		if (!HostSuffix.IsEmpty())
-		{
-			Sb << HostSuffix << TEXT("/");
-		}
-		Sb << Bucket << TEXT("/") << BucketPrefix;
-		OutTocPath = Sb;
-	}
-}
-
-static FIoStatus WriteContainerFiles(FOnDemandToc& OnDemandToc, const TMap<FIoHash, FString>& UTocPaths)
-{
-	UE_LOG(LogIas, Display, TEXT("Attempting to write out %d '%s' files"), OnDemandToc.Containers.Num(), *GIasOnDemandTocExt);
-
-	for (FOnDemandTocContainerEntry& Container : OnDemandToc.Containers)
-	{
-		const FString* UTocPath = UTocPaths.Find(Container.UTocHash);
-		if (UTocPath == nullptr)
-		{
-			return FIoStatus(EIoErrorCode::Unknown, FString::Printf(TEXT("Could not find the original path for ondemand container '%s'"), *Container.ContainerName));
-		}
-
-		// Create a new FOnDemandToc to be written out and give it ownership of the FOnDemandTocContainerEntry that we want
-		// associated with it.This way we don't need to write any specialized serialization code to split up the input FOnDemandToc
-		// into one per container.
-		FOnDemandToc ContainerToc;
-		ContainerToc.Header = OnDemandToc.Header;
-		ContainerToc.Meta = OnDemandToc.Meta;
-		ContainerToc.Containers.Emplace(MoveTemp(Container));
-
-		const FString OutputPath = FPathViews::ChangeExtension(*UTocPath, GIasOnDemandTocExt);
-
-		TUniquePtr<FArchive> Ar(IFileManager::Get().CreateFileWriter(*OutputPath));
-
-		if (!Ar.IsValid())
-		{
-			return FIoStatus(EIoErrorCode::FileOpenFailed, FString::Printf(TEXT("Failed to open '%s' for write"), *OutputPath));
-		}
-
-		// TODO: We should consider adding a hash of the FOnDemandToc that can be computed at runtime on the loaded structure
-		// (to avoid running over the file twice) to verify that nothing was corrupted.
-
-		(*Ar) << ContainerToc;
-
-		FOnDemandTocSentinel Sentinel;
-		(*Ar) << Sentinel;
-
-		if (Ar->IsError() || Ar->IsCriticalError())
-		{
-			return FIoStatus(EIoErrorCode::WriteError, FString::Printf(TEXT("Failed to write to '%s'"), *OutputPath));
-		}
-
-		// TODO: Should write an end of file sentinel here for a simple/easy file corruption check
-
-		UE_LOG(LogIas, Display, TEXT("Wrote ondemand container file '%s' to disk"), *OutputPath);
-
-		// Move the container entry back to OnDemandToc in case we want to use the data structure
-		// beyond this point in the future.
-		Container = MoveTemp(ContainerToc.Containers[0]);
-	}
-
-	return FIoStatus::Ok;
-}
-
-TIoStatusOr<FIoStoreUploadParams> FIoStoreUploadParams::Parse(const TCHAR* CommandLine)
-{
-	FIoStoreUploadParams Params;
-
-	if (!FParse::Value(CommandLine, TEXT("Bucket="), Params.Bucket))
-	{
-		FIoStatus(EIoErrorCode::InvalidParameter, TEXT("Invalid bucket name"));
-	}
-
-	FParse::Value(CommandLine, TEXT("BucketPrefix="), Params.BucketPrefix);
-	FParse::Value(CommandLine, TEXT("ServiceUrl="), Params.ServiceUrl);
-	FParse::Value(CommandLine, TEXT("DistributionUrl="), Params.DistributionUrl);
-	FParse::Value(CommandLine, TEXT("FallbackUrl="), Params.FallbackUrl);
-	FParse::Value(CommandLine, TEXT("Region="), Params.Region);
-	FParse::Value(CommandLine, TEXT("AccessKey="), Params.AccessKey);
-	FParse::Value(CommandLine, TEXT("SecretKey="), Params.SecretKey);
-	FParse::Value(CommandLine, TEXT("SessionToken="), Params.SessionToken);
-	FParse::Value(CommandLine, TEXT("CredentialsFile="), Params.CredentialsFile);
-	FParse::Value(CommandLine, TEXT("CredentialsFileKeyName="), Params.CredentialsFileKeyName);
-	FParse::Value(CommandLine, TEXT("BuildVersion="), Params.BuildVersion);
-	FParse::Value(CommandLine, TEXT("TargetPlatform="), Params.TargetPlatform);
-	FParse::Value(CommandLine, TEXT("EncryptionKeyName="), Params.EncryptionKeyName);
-	Params.bDeleteContainerFiles = !FParse::Param(CommandLine, TEXT("KeepContainerFiles"));
-	Params.bDeletePakFiles = !FParse::Param(CommandLine, TEXT("KeepPakFiles"));
-	Params.bWriteTocToDisk = FParse::Param(CommandLine, TEXT("WriteTocToDisk"));
-
-	// If we keep this feature we should allow the caller to set this path
-	// themselves rather than rely on the config file location
-	FString ConfigFilePath;
-	FParse::Value(FCommandLine::Get(), TEXT("ConfigFilePath="), ConfigFilePath);
-	Params.TocOutputDir = FPaths::GetPath(ConfigFilePath);
-
-	if (FIoStatus Validation = Params.Validate(); !Validation.IsOk())
-	{
-		return Validation;
-	}
-
-	return Params;
-}
-
-FIoStatus FIoStoreUploadParams::Validate() const
-{
-	if (bWriteTocToDisk && TocOutputDir.IsEmpty())
-	{
-		return FIoStatus(EIoErrorCode::InvalidParameter, TEXT("Cmdline param 'WriteToDisk' requires a valid 'ConfigFilePath' param as well"));
-	}
-
-	if (!AccessKey.IsEmpty() && SecretKey.IsEmpty())
-	{
-		return FIoStatus(EIoErrorCode::InvalidParameter, TEXT("Invalid secret key"));
-	}
-	else if (AccessKey.IsEmpty() && !SecretKey.IsEmpty())
-	{
-		return FIoStatus(EIoErrorCode::InvalidParameter, TEXT("Invalid access key"));
-	}
-
-	if (!CredentialsFile.IsEmpty() && CredentialsFileKeyName.IsEmpty())
-	{
-		return FIoStatus(EIoErrorCode::InvalidParameter, TEXT("Invalid credential file key name"));
-	}
-
-	if (ServiceUrl.IsEmpty() && Region.IsEmpty())
-	{
-		return FIoStatus(EIoErrorCode::InvalidParameter, TEXT("Service URL or AWS region needs to be specified"));
-	}
-
-	return FIoStatus::Ok;
-}
-
-static void WriteConfigFile(
-	const FIoStoreUploadParams& UploadParams,
-	const FIoStoreUploadResult& UploadResult,
-	const FKeyChain& KeyChain)
-{
-	FStringBuilderBase Sb;
-	Sb << TEXT("[Endpoint]") << TEXT("\r\n");
-
-	if (!UploadParams.DistributionUrl.IsEmpty())
-	{
-		Sb << TEXT("DistributionUrl=\"") << UploadParams.DistributionUrl << TEXT("\"\r\n");
-
-		if (!UploadParams.FallbackUrl.IsEmpty())
-		{
-			Sb << TEXT("FallbackUrl=\"") << UploadParams.FallbackUrl << TEXT("\"\r\n");
-		}
-	}
-	else
-	{
-		Sb << TEXT("ServiceUrl=\"") << UploadResult.ServiceUrl << TEXT("\"\r\n");
-	}
-
-	Sb << TEXT("TocPath=\"") << UploadResult.TocPath << TEXT("\"\r\n");
-
-	// Temporary solution to get replays working with encrypted on demand content
-	{
-		if (!UploadParams.EncryptionKeyName.IsEmpty())
-		{
-			const TCHAR* EncryptionKeyName = *(UploadParams.EncryptionKeyName);
-
-			TOptional<FNamedAESKey> EncryptionKey;
-			for (const TPair<FGuid, FNamedAESKey>& KeyPair: KeyChain.GetEncryptionKeys())
-			{
-				if (KeyPair.Value.Name.Compare(EncryptionKeyName, ESearchCase::IgnoreCase) == 0)
-				{
-					EncryptionKey.Emplace(KeyPair.Value);
-				}
-			}
-
-			if (EncryptionKey)
-			{
-				FString KeyString = FBase64::Encode(EncryptionKey.GetValue().Key.Key, FAES::FAESKey::KeySize);
-				Sb << TEXT("ContentKey=\"") << EncryptionKey.GetValue().Guid.ToString() << TEXT(":") << KeyString << TEXT("\"\r\n");
-			}
-			else
-			{
-				UE_LOG(LogIoStore, Warning, TEXT("Failed to encryption key '%s' in key chain"), EncryptionKeyName);
-			}
-		}
-	}
-
-	FString ConfigFilePath = UploadParams.TocOutputDir / TEXT("IoStoreOnDemand.ini");
-	UE_LOG(LogIoStore, Display, TEXT("Saving on demand config file '%s'"), *ConfigFilePath);
-	if (FFileHelper::SaveStringToFile(Sb.ToString(), *ConfigFilePath) == false)
-	{
-		UE_LOG(LogIoStore, Error, TEXT("Failed to save on demand config file '%s'"), *ConfigFilePath);
-	}
-}
-
-TIoStatusOr<FIoStoreUploadResult> UploadContainerFiles(
-	const FIoStoreUploadParams& UploadParams,
-	TConstArrayView<FString> ContainerFiles,
-	const FKeyChain& KeyChain)
-{
-	TMap<FGuid, FAES::FAESKey> EncryptionKeys;
-	for (const TPair<FGuid, FNamedAESKey>& KeyPair: KeyChain.GetEncryptionKeys())
-	{
-		EncryptionKeys.Add(KeyPair.Key, KeyPair.Value.Key);
-	}
-
-	struct FContainerStats
-	{
-		FString ContainerName;
-		uint64 ChunkCount = 0;
-		uint64 TotalBytes = 0;
-		uint64 UploadedChunkCount = 0;
-		uint64 UploadedBytes = 0;
-	};
-	TMap<FString, FContainerStats> ContainerSummary;
-
-	const double StartTime = FPlatformTime::Seconds();
-
-	FS3ClientConfig Config;
-	Config.ServiceUrl = UploadParams.ServiceUrl;
-	Config.Region = UploadParams.Region;
-	
-	FS3ClientCredentials Credentials;
-	if (UploadParams.CredentialsFile.IsEmpty() == false)
-	{
-		UE_LOG(LogIas, Display, TEXT("Loading credentials file '%s'"), *UploadParams.CredentialsFile);
-		FS3CredentialsProfileStore CredentialsStore = FS3CredentialsProfileStore::FromFile(UploadParams.CredentialsFile);
-		if (CredentialsStore.TryGetCredentials(UploadParams.CredentialsFileKeyName, Credentials) == false)
-		{
-			return FIoStatus(EIoErrorCode::InvalidParameter, TEXT("Failed to find valid credentials in credentials file"));
-		}
-		else
-		{
-			UE_LOG(LogIas, Display, TEXT("Found credentials for '%s'"), *UploadParams.CredentialsFileKeyName);
-		}
-	}
-	else
-	{
-		Credentials = FS3ClientCredentials(UploadParams.AccessKey, UploadParams.SecretKey, UploadParams.SessionToken);
-	}
-
-	FS3Client Client(Config, Credentials);
-	FS3UploadQueue UploadQueue(Client, UploadParams.Bucket, UploadParams.MaxConcurrentUploads);
-
-	if (ContainerFiles.Num() == 0)
-	{
-		return FIoStatus(EIoErrorCode::InvalidParameter, TEXT("No container file(s) specified"));
-	}
-
-	TSet<FIoHash> ExistingChunks;
-	uint64 TotalExistingTocs = 0;
-	uint64 TotalExistingBytes = 0;
-	{
-		TStringBuilder<256> TocsKey;
-		if (!UploadParams.BucketPrefix.IsEmpty())
-		{
-			TocsKey << UploadParams.BucketPrefix << "/";
-		}
-
-		UE_LOG(LogIas, Display, TEXT("	 '%s/%s/%s'"), *Client.GetConfig().ServiceUrl, *UploadParams.Bucket, TocsKey.ToString());
-		FS3ListObjectResponse Response = Client.ListObjects(FS3ListObjectsRequest
-		{
-			UploadParams.Bucket,
-			TocsKey.ToString(),
-			TEXT('/')
-		});
-
-		for (const FS3Object& TocInfo : Response.Objects)
-		{
-			if (TocInfo.Key.EndsWith(TEXT("iochunktoc")) == false)
-			{
-				continue;
-			}
-
-			UE_LOG(LogIas, Display, TEXT("Fetching TOC '%s/%s/%s'"), *Client.GetConfig().ServiceUrl, *UploadParams.Bucket, *TocInfo.Key);
-			FS3GetObjectResponse TocResponse = Client.GetObject(FS3GetObjectRequest
-			{
-				UploadParams.Bucket,
-				TocInfo.Key
-			});
-			
-			if (TocResponse.IsOk() == false)
-			{
-				UE_LOG(LogIas, Warning, TEXT("Failed to fetch TOC '%s/%s/%s'"), *Client.GetConfig().ServiceUrl, *UploadParams.Bucket, *TocInfo.Key);
-				continue;
-			}
-
-			FOnDemandToc Toc;
-			FMemoryReaderView Ar(TocResponse.GetBody().GetView());
-			Ar << Toc;
-
-			if (Ar.IsError()) 
-			{
-				Toc = FOnDemandToc{};
-				if (LoadFromCompactBinary(FCbFieldView(TocResponse.GetBody().GetData()), Toc) == false)
-				{
-					UE_LOG(LogIas, Display, TEXT("Failed to load TOC '%s/%s/%s'"), *Client.GetConfig().ServiceUrl, *UploadParams.Bucket, *TocInfo.Key);
-					continue;
-				}
-			}
-
-			for (const FOnDemandTocContainerEntry& ContainerEntry : Toc.Containers)
-			{
-				for (const FOnDemandTocEntry& TocEntry : ContainerEntry.Entries)
-				{
-					bool bIsAlreadyInSet;
-					ExistingChunks.Add(TocEntry.Hash, &bIsAlreadyInSet);
-					if (!bIsAlreadyInSet)
-					{
-						TotalExistingBytes += TocEntry.EncodedSize;
-					}
-				}
-			}
-
-			TotalExistingTocs++;
-		}
-	}
-
-	FString ChunksRelativePath = UploadParams.BucketPrefix.IsEmpty() ? TEXT("Chunks") : FString::Printf(TEXT("%s/Chunks"), *UploadParams.BucketPrefix);
-	ChunksRelativePath.ToLowerInline();
-
-	bool bPerContainerTocsConfigValue = false;
-	GConfig->GetBool(TEXT("Ias"), TEXT("CreatePerContainerTocs"), bPerContainerTocsConfigValue, GEngineIni);
-	const bool bWritePerContainerToc = bPerContainerTocsConfigValue || UploadParams.bPerContainerTocs;
-
-	uint64 TotalUploadedChunks = 0;
-	uint64 TotalUploadedBytes = 0;
-
-	FOnDemandToc OnDemandToc;
-	OnDemandToc.Containers.Reserve(ContainerFiles.Num());
-
-	// Map of the .utoc paths that we have created ondemand containers for, indexed by their hash so that the paths can
-	// be looked up later. We do not rely on the filename as we cannot be sure that there won't be duplicate file names
-	// stored in different directories.
-	TMap<FIoHash, FString> UTocPaths;
-
-	TArray<FString> FilesToDelete;
-	for (const FString& Path : ContainerFiles)
-	{
-		FIoStoreReader ContainerFileReader;
-		{
-			FIoStatus Status = ContainerFileReader.Initialize(*FPaths::ChangeExtension(Path, TEXT("")), EncryptionKeys);
-			if (!Status.IsOk())
-			{
-				UE_LOG(LogIas, Error, TEXT("Failed to open container '%s' for reading"), *Path);
-				continue;
-			}
-		}
-
-		if (EnumHasAnyFlags(ContainerFileReader.GetContainerFlags(), EIoContainerFlags::OnDemand) == false)
-		{
-			continue;
-		}
-		UE_LOG(LogIas, Display, TEXT("Uploading ondemand container '%s'"), *Path);
-
-		const uint32 BlockSize = ContainerFileReader.GetCompressionBlockSize();
-		if (OnDemandToc.Header.BlockSize == 0)
-		{
-			OnDemandToc.Header.BlockSize = ContainerFileReader.GetCompressionBlockSize();
-		}
-		check(OnDemandToc.Header.BlockSize == ContainerFileReader.GetCompressionBlockSize());
-
-		TArray<FIoStoreTocChunkInfo> ChunkInfos;
-		ContainerFileReader.EnumerateChunks([&ChunkInfos](FIoStoreTocChunkInfo&& Info)
-		{ 
-			ChunkInfos.Emplace(MoveTemp(Info));
-			return true;
-		});
-		
-		FOnDemandTocContainerEntry& ContainerEntry = OnDemandToc.Containers.AddDefaulted_GetRef();
-		ContainerEntry.ContainerName = FPaths::GetBaseFilename(Path);
-		ContainerEntry.ContainerId = ContainerFileReader.GetContainerId();
-
-		if (EnumHasAnyFlags(ContainerFileReader.GetContainerFlags(), EIoContainerFlags::Encrypted))
-		{
-			ContainerEntry.EncryptionKeyGuid = LexToString(ContainerFileReader.GetEncryptionKeyGuid());
-		}
-		FContainerStats& ContainerStats = ContainerSummary.FindOrAdd(ContainerEntry.ContainerName);
-		
-		ContainerEntry.Entries.Reserve(ChunkInfos.Num());
-
-		for (const FIoStoreTocChunkInfo& ChunkInfo : ChunkInfos)
-		{
-			const bool bDecrypt = false;
-			TIoStatusOr<FIoStoreCompressedReadResult> Status = ContainerFileReader.ReadCompressed(ChunkInfo.Id, FIoReadOptions(), bDecrypt);
-			if (!Status.IsOk())
-			{
-				return Status.Status();
-			}
-
-			FIoStoreCompressedReadResult ReadResult = Status.ConsumeValueOrDie();
-
-			const uint32 BlockOffset = ContainerEntry.BlockSizes.Num();
-			const uint32 BlockCount = ReadResult.Blocks.Num();
-			const FIoHash ChunkHash = FIoHash::HashBuffer(ReadResult.IoBuffer.GetView());
-
-			FMemoryView EncodedBlocks = ReadResult.IoBuffer.GetView();
-			uint64 RawChunkSize = 0;
-			uint64 EncodedChunkSize = 0;
-			for (const FIoStoreCompressedBlockInfo& BlockInfo : ReadResult.Blocks)
-			{
-				check(Align(BlockInfo.CompressedSize, FAES::AESBlockSize) == BlockInfo.AlignedSize);
-				const uint64 EncodedBlockSize = BlockInfo.AlignedSize;
-				ContainerEntry.BlockSizes.Add(uint32(BlockInfo.CompressedSize));
-
-				FMemoryView EncodedBlock = EncodedBlocks.Left(EncodedBlockSize);
-				EncodedBlocks += EncodedBlock.GetSize();
-				ContainerEntry.BlockHashes.Add(FIoChunkEncoding::HashBlock(EncodedBlock));
-
-				EncodedChunkSize += EncodedBlockSize;
-				RawChunkSize += BlockInfo.UncompressedSize;
-
-				if (OnDemandToc.Header.CompressionFormat.IsEmpty() && BlockInfo.CompressionMethod != NAME_None)
-				{
-					OnDemandToc.Header.CompressionFormat = BlockInfo.CompressionMethod.ToString();
-				}
-			}
-
-			if (EncodedChunkSize != ReadResult.IoBuffer.GetSize())
-			{
-				return FIoStatus(EIoErrorCode::ReadError, TEXT("Encoded chunk size does not match buffer"));
-			}
-
-			FOnDemandTocEntry& TocEntry = ContainerEntry.Entries.AddDefaulted_GetRef();
-			TocEntry.ChunkId = ChunkInfo.Id;
-			TocEntry.Hash = ChunkHash;
-			TocEntry.RawSize = RawChunkSize;
-			TocEntry.EncodedSize = EncodedChunkSize;
-			TocEntry.BlockOffset = BlockOffset;
-			TocEntry.BlockCount = BlockCount;
-
-			ContainerStats.ChunkCount++;
-			ContainerStats.TotalBytes += EncodedChunkSize;
-
-			if (ExistingChunks.Contains(TocEntry.Hash) == false)
-			{
-				const FString HashString = LexToString(ChunkHash);
-
-				TStringBuilder<256> Key;
-				Key << ChunksRelativePath
-					<< TEXT("/") << HashString.Left(2)
-					<< TEXT("/") << HashString
-					<< TEXT(".iochunk");
-
-				if (UploadQueue.Enqueue(Key, ReadResult.IoBuffer) == false)
-				{
-					return FIoStatus(EIoErrorCode::WriteError, TEXT("Failed to upload chunk"));
-				}
-
-				TotalUploadedChunks++;
-				TotalUploadedBytes += EncodedChunkSize;
-				ContainerStats.UploadedChunkCount++;
-				ContainerStats.UploadedBytes += EncodedChunkSize;
-			}
-		}
-
-		{
-			const FString UTocFilePath = FPaths::ChangeExtension(Path, TEXT(".utoc"));
-			TArray<uint8> Buffer;
-			if (!FFileHelper::LoadFileToArray(Buffer, *UTocFilePath))
-			{
-				return FIoStatus(EIoErrorCode::ReadError, TEXT("Failed to upload .utoc file"));
-			}
-
-			ContainerEntry.UTocHash = FIoHash::HashBuffer(Buffer.GetData(), Buffer.Num());
-			TStringBuilder<256> Key;
-			if (UploadParams.BucketPrefix.IsEmpty() == false)
-			{
-				Key << UploadParams.BucketPrefix.ToLower() << TEXT("/");
-			}
-			Key << LexToString(ContainerEntry.UTocHash) << TEXT(".utoc");
-
-			const FS3PutObjectResponse Response = Client.TryPutObject(
-				FS3PutObjectRequest{UploadParams.Bucket, Key.ToString(), MakeMemoryView(Buffer.GetData(), Buffer.Num())});
-
-			if (!Response.IsOk())
-			{
-				return FIoStatus(EIoErrorCode::WriteError, FString::Printf(TEXT("Failed to upload '%s' (%s)"), *UTocFilePath, *Response.GetErrorStatus()));
-			}
-			
-			UTocPaths.Add(ContainerEntry.UTocHash, UTocFilePath);
-
-			UE_LOG(LogIas, Display, TEXT("Uploaded '%s'"), *UTocFilePath);	
-		}
-
-		if (UploadParams.bDeleteContainerFiles)
-		{
-			FilesToDelete.Add(Path);
-			ContainerFileReader.GetContainerFilePaths(FilesToDelete);
-
-			// We need the pak files in order to mount OnDemand toc files!
-			if (UploadParams.bDeletePakFiles && !bWritePerContainerToc)
-			{
-				FilesToDelete.Add(FPaths::ChangeExtension(Path, TEXT(".pak")));
-				FilesToDelete.Add(FPaths::ChangeExtension(Path, TEXT(".sig")));
-			}
-		}
-	}
-
-	if (OnDemandToc.Containers.IsEmpty())
-	{
-		return FIoStatus(EIoErrorCode::InvalidParameter, TEXT("No container file(s) marked as on demand"));
-	}
-
-	if (!UploadQueue.Flush())
-	{
-		return FIoStatus(EIoErrorCode::WriteError, TEXT("Failed to upload chunk(s)"));
-	}
-
-	FString ServiceUrl;
-	GetTocPath(UploadParams.ServiceUrl, UploadParams.Bucket, UploadParams.BucketPrefix, ServiceUrl, OnDemandToc.Header.ChunksDirectory);
-
-	OnDemandToc.Meta.EpochTimestamp = FDateTime::Now().ToUnixTimestamp();
-	OnDemandToc.Meta.BuildVersion = UploadParams.BuildVersion;
-	OnDemandToc.Meta.TargetPlatform = UploadParams.TargetPlatform;
-
-	FIoStoreUploadResult UploadResult;
-	{
-		FLargeMemoryWriter Ar;
-		Ar << OnDemandToc;
-
-		UploadResult.TocHash = FIoHash::HashBuffer(Ar.GetView());
-		const FString TocFilename = LexToString(UploadResult.TocHash) + TEXT(".iochunktoc");
-
-		TStringBuilder<256> Key;
-		if (!UploadParams.BucketPrefix.IsEmpty())
-		{
-			Key << UploadParams.BucketPrefix.ToLower() << TEXT("/");
-		}
-		Key << TocFilename;
-
-		UploadResult.ServiceUrl = ServiceUrl;
-		UploadResult.TocPath = OnDemandToc.Header.ChunksDirectory / TocFilename;
-		UploadResult.TocSize = Ar.TotalSize();
-
-		const FS3PutObjectResponse Response = Client.TryPutObject(FS3PutObjectRequest{UploadParams.Bucket, Key.ToString(), Ar.GetView()});
-		if (Response.IsOk())
-		{
-			UE_LOG(LogIas, Display, TEXT("Uploaded on demand TOC '%s/%s/%s'"), *Client.GetConfig().ServiceUrl, *UploadParams.Bucket, Key.ToString());
-		}
-		else
-		{
-			UE_LOG(LogIas, Warning, TEXT("Failed to upload TOC '%s/%s/%s' (%s)"), *Client.GetConfig().ServiceUrl, *UploadParams.Bucket, Key.ToString(), *Response.GetErrorStatus());
-			return FIoStatus(EIoErrorCode::WriteError, TEXT("Failed to upload TOC"));
-		}
-
-		if (UploadParams.bWriteTocToDisk)
-		{
-			// Write a single .iochunktoc containing all on demand data for the current build
-			TStringBuilder<512> OnDemandTocFilePath;
-			FPathViews::Append(OnDemandTocFilePath, UploadParams.TocOutputDir, UploadResult.TocHash);
-			OnDemandTocFilePath << TEXT(".iochunktoc");
-
-			if (SaveArrayToFile(TArrayView64<uint8>(Ar.GetData(), Ar.TotalSize()), OnDemandTocFilePath.ToString()) == false)
-			{
-				UE_LOG(LogIoStore, Error, TEXT("Failed to save on demand toc file '%s'"), OnDemandTocFilePath.ToString());
-			}
-		}
-	}
-
-	if (bWritePerContainerToc)
-	{
-		// Write out separate .uondemandtoc files, one per.utoc containing ondemand data.
-		FIoStatus Result = WriteContainerFiles(OnDemandToc, UTocPaths);
-		if (!Result.IsOk())
-		{
-			return Result;
-		}
-	}
-
-	for (const FString& Path : FilesToDelete)
-	{
-		if (IFileManager::Get().FileExists(*Path))
-		{
-			UE_LOG(LogIas, Display, TEXT("Deleting '%s'"), *Path); 
-			IFileManager::Get().Delete(*Path);
-		}
-	}
-
-	if (!UploadParams.TocOutputDir.IsEmpty())
-	{
-		WriteConfigFile(UploadParams, UploadResult, KeyChain);
-	}
-
-	{
-		const double Duration = FPlatformTime::Seconds() - StartTime;
-		
-		UE_LOG(LogIas, Display, TEXT(""));
-		UE_LOG(LogIas, Display, TEXT("------------------------------------------------- Upload Summary -------------------------------------------------"));
-		UE_LOG(LogIas, Display, TEXT("%-15s: %s"), TEXT("Service URL"), *UploadParams.ServiceUrl);
-		UE_LOG(LogIas, Display, TEXT("%-15s: %s"), TEXT("Bucket"), *UploadParams.Bucket);
-		UE_LOG(LogIas, Display, TEXT("%-15s: %s"), TEXT("TargetPlatform"), *UploadParams.TargetPlatform);
-		UE_LOG(LogIas, Display, TEXT("%-15s: %s"), TEXT("BuildVersion"), *UploadParams.BuildVersion);
-		UE_LOG(LogIas, Display, TEXT("%-15s: %s"), TEXT("TOC path"), *UploadResult.TocPath);
-		UE_LOG(LogIas, Display, TEXT("%-15s: %.2lf KiB"), TEXT("TOC size"), double(UploadResult.TocSize) / 1024.0);
-		UE_LOG(LogIas, Display, TEXT("%-15s: %.2lf second(s)"), TEXT("Duration"), Duration);
-		UE_LOG(LogIas, Display, TEXT(""));
-		
-		UE_LOG(LogIas, Display, TEXT("%-50s %15s %15s %15s %15s"),
-			TEXT("Container"), TEXT("Chunk(s)"), TEXT("Size (MiB)"), TEXT("Uploaded"), TEXT("Uploaded (MiB)"));
-		UE_LOG(LogIas, Display, TEXT("-------------------------------------------------------------------------------------------------------------------"));
-
-		FContainerStats TotalStats;
-		for (const TTuple<FString, FContainerStats>& Kv : ContainerSummary)
-		{
-			UE_LOG(LogIas, Display, TEXT("%-50s %15llu %15.2lf %15llu %15.2lf"),
-				*Kv.Key, Kv.Value.ChunkCount, double(Kv.Value.TotalBytes) / 1024.0 / 1024.0, Kv.Value.UploadedChunkCount, double(Kv.Value.UploadedBytes) / 1024.0 / 1024.0);
-			
-			TotalStats.ChunkCount += Kv.Value.ChunkCount;
-			TotalStats.TotalBytes += Kv.Value.TotalBytes;
-			TotalStats.UploadedChunkCount += Kv.Value.UploadedChunkCount;
-			TotalStats.UploadedBytes += Kv.Value.UploadedBytes;
-		}
-		UE_LOG(LogIas, Display, TEXT("-------------------------------------------------------------------------------------------------------------------"));
-		UE_LOG(LogIas, Display, TEXT("%-50s %15llu %15.2lf %15llu %15.2lf"),
-			TEXT("Total"),TotalStats.ChunkCount, double(TotalStats.TotalBytes) / 1024.0 / 1024.0, TotalStats.UploadedChunkCount, double(TotalStats.UploadedBytes) / 1024.0 / 1024.0);
-		UE_LOG(LogIas, Display, TEXT(""));
-		
-		UE_LOG(LogIas, Display, TEXT("%-50s %15s %15s %15s"), TEXT("Bucket"), TEXT("TOC(s)"), TEXT("Chunk(s)"), TEXT("MiB"));
-		UE_LOG(LogIas, Display, TEXT("-------------------------------------------------------------------------------------------------------------------"));
-		UE_LOG(LogIas, Display, TEXT("%-50s %15llu %15d %15.2lf"), TEXT("Existing"), TotalExistingTocs, ExistingChunks.Num(), double(TotalExistingBytes) / 1024.0 / 1024.0);
-		UE_LOG(LogIas, Display, TEXT("%-50s %15llu %15d %15.2lf"), TEXT("Uploaded"), 1, TotalStats.UploadedChunkCount, double(TotalStats.UploadedBytes) / 1024.0 / 1024.0);
-		UE_LOG(LogIas, Display, TEXT("-------------------------------------------------------------------------------------------------------------------"));
-		UE_LOG(LogIas, Display, TEXT("%-50s %15llu %15d %15.2lf"), TEXT("Total"),
-			TotalExistingTocs + 1, ExistingChunks.Num() + TotalStats.UploadedChunkCount, double(TotalExistingBytes + TotalStats.UploadedBytes) / 1024.0 / 1024.0);
-		UE_LOG(LogIas, Display, TEXT(""));
-	}
-	
-	return UploadResult;
-}
-
 ////////////////////////////////////////////////////////////////////////////////
 TIoStatusOr<FIoStoreDownloadParams> FIoStoreDownloadParams::Parse(const TCHAR* CommandLine)
 {
diff --git a/Engine/Source/Runtime/Experimental/IoStore/OnDemand/Private/Tool/Upload.cpp b/Engine/Source/Runtime/Experimental/IoStore/OnDemand/Private/Tool/Upload.cpp
--- a/Engine/Source/Runtime/Experimental/IoStore/OnDemand/Private/Tool/Upload.cpp
+++ b/Engine/Source/Runtime/Experimental/IoStore/OnDemand/Private/Tool/Upload.cpp
@@ -4,35 +4,110 @@
 
 #include "Command.h"
 
-#include <HAL/FileManager.h>
-#include <IO/IoStoreOnDemand.h>
-#include <Misc/KeyChainUtilities.h>
-#include <Misc/Paths.h>
-#include <Misc/PathViews.h>
+#include "Async/Async.h"
+#include "Async/Mutex.h"
+#include "Containers/Queue.h"
+#include "HAL/FileManager.h"
+#include "IO/IoBuffer.h"
+#include "IO/IoChunkEncoding.h"
+#include "IO/IoStore.h"
+#include "IO/IoStoreOnDemand.h"
+#include "Misc/ConfigCacheIni.h"
+#include "Misc/FileHelper.h"
+#include "Misc/KeyChainUtilities.h"
+#include "Misc/Paths.h"
+#include "Misc/PathViews.h"
+#include "Serialization/CompactBinarySerialization.h"
+#include "Serialization/CompactBinaryWriter.h"
+#include "Serialization/MemoryReader.h"
+#include "Serialization/LargeMemoryWriter.h"
+#include "S3/S3Client.h"
+
+////////////////////////////////////////////////////////////////////////////////
+namespace UE::IoStore
+{
+extern FString GIasOnDemandTocExt;
+}
 
 namespace UE::IoStore::Tool
 {
 
 ////////////////////////////////////////////////////////////////////////////////
-static FKeyChain LoadKeys(const FContext& Context)
+struct FUploadParams
 {
-	FKeyChain Ret;
+	FString ServiceUrl;
+	FString DistributionUrl;
+	FString FallbackUrl;
+	FString Bucket;
+	FString BucketPrefix;
+	FString Region; 
+	FString AccessKey;
+	FString SecretKey;
+	FString SessionToken;
+	FString CredentialsFile;
+	FString CredentialsFileKeyName;
+	FString BuildVersion;
+	FString TargetPlatform;
+	FString EncryptionKeyName;
+	int32 MaxConcurrentUploads = 16;
+	bool bDeleteContainerFiles = true;
+	bool bDeletePakFiles = true;
+	bool bPerContainerTocs = false;
 
-	auto Path = Context.Get<FStringView>(TEXT("-CryptoKeys"));
-	if (Path.IsEmpty())
+	/** If we should write out the .iochunktoc to disk as well as uploading it. */
+	bool bWriteTocToDisk = false;
+	/** Where the .iochunktoc file should be written out. */
+	FString TocOutputDir;
+
+	FIoStatus Validate() const;
+};
+
+FIoStatus FUploadParams::Validate() const
+{
+	if (bWriteTocToDisk && TocOutputDir.IsEmpty())
 	{
-		return Ret;
+		return FIoStatus(EIoErrorCode::InvalidParameter, TEXT("Cmdline param 'WriteToDisk' requires a valid 'ConfigFilePath' param as well"));
 	}
 
-	// UE_LOG(LogIoStore, Display, TEXT("Parsing crypto keys from a crypto key cache file '%s'"), *CryptoKeysCacheFilename);
-	KeyChainUtilities::LoadKeyChainFromFile(Path.GetData(), Ret);
-	return Ret;
+	if (!AccessKey.IsEmpty() && SecretKey.IsEmpty())
+	{
+		return FIoStatus(EIoErrorCode::InvalidParameter, TEXT("Invalid secret key"));
+	}
+	else if (AccessKey.IsEmpty() && !SecretKey.IsEmpty())
+	{
+		return FIoStatus(EIoErrorCode::InvalidParameter, TEXT("Invalid access key"));
+	}
+
+	if (!CredentialsFile.IsEmpty() && CredentialsFileKeyName.IsEmpty())
+	{
+		return FIoStatus(EIoErrorCode::InvalidParameter, TEXT("Invalid credential file key name"));
+	}
+
+	if (ServiceUrl.IsEmpty() && Region.IsEmpty())
+	{
+		return FIoStatus(EIoErrorCode::InvalidParameter, TEXT("Service URL or AWS region needs to be specified"));
+	}
+
+	return FIoStatus::Ok;
 }
 
 ////////////////////////////////////////////////////////////////////////////////
-static FIoStoreUploadParams BuildUploadParams(const FContext& Context)
+struct FUploadResult
 {
-	FIoStoreUploadParams Ret;
+	/** Hash of the toc when written as a binary blob */
+	FIoHash TocHash;
+	/** Url of the service that the toc was uploaded too */
+	FString ServiceUrl;
+	/** Path of the toc on the service */
+	FString TocPath;
+	/** Size (in bytes) of the toc when written as a binary blob */
+	uint64 TocSize = 0;
+};
+
+////////////////////////////////////////////////////////////////////////////////
+static FUploadParams BuildUploadParams(const FContext& Context)
+{
+	FUploadParams Ret;
 
 	Ret.ServiceUrl				= Context.Get<FStringView>(TEXT("-ServiceUrl"),				Ret.ServiceUrl);
 	Ret.Bucket					= Context.Get<FStringView>(TEXT("-Bucket"),					Ret.Bucket);
@@ -70,6 +145,784 @@ static FIoStoreUploadParams BuildUploadParams(const FContext& Context)
 	return Ret;
 }
 
+////////////////////////////////////////////////////////////////////////////////
+static FKeyChain LoadKeys(const FContext& Context)
+{
+	FKeyChain Ret;
+
+	auto Path = Context.Get<FStringView>(TEXT("-CryptoKeys"));
+	if (Path.IsEmpty())
+	{
+		return Ret;
+	}
+
+	// UE_LOG(LogIoStore, Display, TEXT("Parsing crypto keys from a crypto key cache file '%s'"), *CryptoKeysCacheFilename);
+	KeyChainUtilities::LoadKeyChainFromFile(Path.GetData(), Ret);
+	return Ret;
+}
+
+////////////////////////////////////////////////////////////////////////////////
+static bool SaveArrayToFile(TArrayView64<uint8> Data, const TCHAR* Filename, uint32 WriteFlags = FILEWRITE_None)
+{
+	TUniquePtr<FArchive> Ar = TUniquePtr<FArchive>(IFileManager::Get().CreateFileWriter(Filename, WriteFlags));
+	if (!Ar)
+	{
+		return false;
+	}
+
+	Ar->Serialize(static_cast<void*>(Data.GetData()), Data.Num());
+	Ar->Close();
+
+	return !Ar->IsError() && !Ar->IsCriticalError();
+}
+
+////////////////////////////////////////////////////////////////////////////////
+static FIoStatus WriteContainerFiles(FOnDemandToc& OnDemandToc, const TMap<FIoHash, FString>& UTocPaths)
+{
+	UE_LOG(LogIas, Display, TEXT("Attempting to write out %d '%s' files"), OnDemandToc.Containers.Num(), *GIasOnDemandTocExt);
+
+	for (FOnDemandTocContainerEntry& Container : OnDemandToc.Containers)
+	{
+		const FString* UTocPath = UTocPaths.Find(Container.UTocHash);
+		if (UTocPath == nullptr)
+		{
+			return FIoStatus(EIoErrorCode::Unknown, FString::Printf(TEXT("Could not find the original path for ondemand container '%s'"), *Container.ContainerName));
+		}
+
+		// Create a new FOnDemandToc to be written out and give it ownership of the FOnDemandTocContainerEntry that we want
+		// associated with it.This way we don't need to write any specialized serialization code to split up the input FOnDemandToc
+		// into one per container.
+		FOnDemandToc ContainerToc;
+		ContainerToc.Header = OnDemandToc.Header;
+		ContainerToc.Meta = OnDemandToc.Meta;
+		ContainerToc.Containers.Emplace(MoveTemp(Container));
+
+		const FString OutputPath = FPathViews::ChangeExtension(*UTocPath, GIasOnDemandTocExt);
+
+		TUniquePtr<FArchive> Ar(IFileManager::Get().CreateFileWriter(*OutputPath));
+
+		if (!Ar.IsValid())
+		{
+			return FIoStatus(EIoErrorCode::FileOpenFailed, FString::Printf(TEXT("Failed to open '%s' for write"), *OutputPath));
+		}
+
+		// TODO: We should consider adding a hash of the FOnDemandToc that can be computed at runtime on the loaded structure
+		// (to avoid running over the file twice) to verify that nothing was corrupted.
+
+		(*Ar) << ContainerToc;
+
+		FOnDemandTocSentinel Sentinel;
+		(*Ar) << Sentinel;
+
+		if (Ar->IsError() || Ar->IsCriticalError())
+		{
+			return FIoStatus(EIoErrorCode::WriteError, FString::Printf(TEXT("Failed to write to '%s'"), *OutputPath));
+		}
+
+		// TODO: Should write an end of file sentinel here for a simple/easy file corruption check
+
+		UE_LOG(LogIas, Display, TEXT("Wrote ondemand container file '%s' to disk"), *OutputPath);
+
+		// Move the container entry back to OnDemandToc in case we want to use the data structure
+		// beyond this point in the future.
+		Container = MoveTemp(ContainerToc.Containers[0]);
+	}
+
+	return FIoStatus::Ok;
+}
+
+////////////////////////////////////////////////////////////////////////////////
+static void GetTocPath(
+	FStringView ServiceUrl,
+	FStringView Bucket,
+	FStringView BucketPrefix,
+	FString& OutServiceUrl,
+	FString& OutTocPath)
+{
+	// The configuration file should specify a service URL without any trailing
+	// host path, i.e. http://{host:port}/{host-path}. Add the trailing path
+	// to the TOC path to form the complete path the TOC from the host, i.e
+	// TocPath={host-path}/{bucket}/{bucket-prefix}/{toc-hash}.uchunktoc
+	
+	if (BucketPrefix.StartsWith(TEXT("/")))
+	{
+		BucketPrefix.RemovePrefix(1);
+	}
+	if (BucketPrefix.EndsWith(TEXT("/")))
+	{
+		BucketPrefix.RemoveSuffix(1);
+	}
+
+	if (ServiceUrl.IsEmpty())
+	{
+		// If the service URL is empty we assume uploading to AWS S3 using the Region parameter
+		// and that we don't need to prefix with the bucket name.
+		OutServiceUrl.Empty();
+		OutTocPath = BucketPrefix;
+	}
+	else
+	{
+		FStringView HostSuffix;
+		int32 Idx = INDEX_NONE;
+		ensure(ServiceUrl.FindChar(':', Idx));
+		const int32 SchemeEnd = Idx + 3;
+
+		Idx = INDEX_NONE;
+		if (ServiceUrl.RightChop(SchemeEnd).FindChar('/', Idx))
+		{
+			OutServiceUrl = ServiceUrl.Left(SchemeEnd + Idx);
+			HostSuffix = ServiceUrl.RightChop(OutServiceUrl.Len() + 1);
+			if (HostSuffix.EndsWith(TEXT("/")))
+			{
+				HostSuffix.RemoveSuffix(1);
+			}
+		}
+		else
+		{
+			OutServiceUrl = ServiceUrl;
+		}
+
+		TStringBuilder<256> Sb;
+		if (!HostSuffix.IsEmpty())
+		{
+			Sb << HostSuffix << TEXT("/");
+		}
+		Sb << Bucket << TEXT("/") << BucketPrefix;
+		OutTocPath = Sb;
+	}
+}
+
+////////////////////////////////////////////////////////////////////////////////
+static void WriteConfigFile(
+	const FUploadParams& UploadParams,
+	const FUploadResult& UploadResult,
+	const FKeyChain& KeyChain)
+{
+	FStringBuilderBase Sb;
+	Sb << TEXT("[Endpoint]") << TEXT("\r\n");
+
+	if (!UploadParams.DistributionUrl.IsEmpty())
+	{
+		Sb << TEXT("DistributionUrl=\"") << UploadParams.DistributionUrl << TEXT("\"\r\n");
+
+		if (!UploadParams.FallbackUrl.IsEmpty())
+		{
+			Sb << TEXT("FallbackUrl=\"") << UploadParams.FallbackUrl << TEXT("\"\r\n");
+		}
+	}
+	else
+	{
+		Sb << TEXT("ServiceUrl=\"") << UploadResult.ServiceUrl << TEXT("\"\r\n");
+	}
+
+	Sb << TEXT("TocPath=\"") << UploadResult.TocPath << TEXT("\"\r\n");
+
+	// Temporary solution to get replays working with encrypted on demand content
+	{
+		if (!UploadParams.EncryptionKeyName.IsEmpty())
+		{
+			const TCHAR* EncryptionKeyName = *(UploadParams.EncryptionKeyName);
+
+			TOptional<FNamedAESKey> EncryptionKey;
+			for (const TPair<FGuid, FNamedAESKey>& KeyPair: KeyChain.GetEncryptionKeys())
+			{
+				if (KeyPair.Value.Name.Compare(EncryptionKeyName, ESearchCase::IgnoreCase) == 0)
+				{
+					EncryptionKey.Emplace(KeyPair.Value);
+				}
+			}
+
+			if (EncryptionKey)
+			{
+				FString KeyString = FBase64::Encode(EncryptionKey.GetValue().Key.Key, FAES::FAESKey::KeySize);
+				Sb << TEXT("ContentKey=\"") << EncryptionKey.GetValue().Guid.ToString() << TEXT(":") << KeyString << TEXT("\"\r\n");
+			}
+			else
+			{
+				UE_LOG(LogIoStore, Warning, TEXT("Failed to encryption key '%s' in key chain"), EncryptionKeyName);
+			}
+		}
+	}
+
+	FString ConfigFilePath = UploadParams.TocOutputDir / TEXT("IoStoreOnDemand.ini");
+	UE_LOG(LogIoStore, Display, TEXT("Saving on demand config file '%s'"), *ConfigFilePath);
+	if (FFileHelper::SaveStringToFile(Sb.ToString(), *ConfigFilePath) == false)
+	{
+		UE_LOG(LogIoStore, Error, TEXT("Failed to save on demand config file '%s'"), *ConfigFilePath);
+	}
+}
+
+////////////////////////////////////////////////////////////////////////////////
+class FUploadQueue
+{
+public:
+	FUploadQueue(FS3Client& Client, const FString& Bucket, int32 ThreadCount);
+	bool Enqueue(FStringView Key, FIoBuffer Payload);
+	bool Flush();
+
+private:
+	void ThreadEntry();
+
+	struct FQueueEntry
+	{
+		FString Key;
+		FIoBuffer Payload;
+	};
+
+	FS3Client& Client;
+	TArray<TFuture<void>> Threads;
+	FString Bucket;
+	FCriticalSection CriticalSection;
+	TQueue<FQueueEntry> Queue;
+	FEventRef WakeUpEvent;
+	FEventRef UploadCompleteEvent;
+	std::atomic_int32_t ConcurrentUploads{0};
+	std::atomic_int32_t ActiveThreadCount{0};
+	std::atomic_int32_t ErrorCount{0};
+	std::atomic_bool bCompleteAdding{false};
+};
+
+FUploadQueue::FUploadQueue(FS3Client& InClient, const FString& InBucket, int32 ThreadCount)
+	: Client(InClient)
+	, Bucket(InBucket)
+{
+	ActiveThreadCount = ThreadCount;
+	for (int32 Idx = 0; Idx < ThreadCount; ++Idx)
+	{
+		Threads.Add(AsyncThread([this]()
+		{
+			ThreadEntry();
+		}));
+	}
+}
+
+bool FUploadQueue::Enqueue(FStringView Key, FIoBuffer Payload)
+{
+	if (ActiveThreadCount == 0)
+	{
+		return false;
+	}
+
+	for(;;)
+	{
+		bool bEnqueued = false;
+		{
+			FScopeLock _(&CriticalSection);
+			if (ConcurrentUploads < Threads.Num())
+			{
+				bEnqueued = Queue.Enqueue(FQueueEntry {FString(Key), Payload});
+			}
+		}
+
+		if (bEnqueued)
+		{
+			WakeUpEvent->Trigger();
+			break;
+		}
+
+		UploadCompleteEvent->Wait();
+	}
+
+	return true;
+}
+
+void FUploadQueue::ThreadEntry()
+{
+	for(;;)
+	{
+		FQueueEntry Entry;
+		bool bDequeued = false;
+		{
+			FScopeLock _(&CriticalSection);
+			bDequeued = Queue.Dequeue(Entry);
+		}
+
+		if (!bDequeued)
+		{
+			if (bCompleteAdding)
+			{
+				break;
+			}
+			WakeUpEvent->Wait();
+			continue;
+		}
+
+		ConcurrentUploads++;
+		const FS3PutObjectResponse Response = Client.TryPutObject(FS3PutObjectRequest{Bucket, Entry.Key, Entry.Payload.GetView()});
+		ConcurrentUploads--;
+		UploadCompleteEvent->Trigger();
+
+		if (Response.IsOk())
+		{
+			UE_LOG(LogIas, Log, TEXT("Uploaded chunk '%s/%s/%s'"), *Client.GetConfig().ServiceUrl, *Bucket, *Entry.Key);
+		}
+		else
+		{
+			TStringBuilder<256> ErrorResponse;
+			Response.GetErrorResponse(ErrorResponse);
+
+			UE_LOG(LogIas, Warning, TEXT("Failed to upload chunk '%s/%s/%s' (%s)"), *Client.GetConfig().ServiceUrl, *Bucket, *Entry.Key, ErrorResponse.ToString());
+			ErrorCount++;
+			break;
+		}
+	}
+
+	ActiveThreadCount--;
+}
+
+bool FUploadQueue::Flush()
+{
+	bCompleteAdding = true;
+	for (int32 Idx = 0; Idx < Threads.Num(); ++Idx)
+	{
+		WakeUpEvent->Trigger();
+	}
+
+	for (TFuture<void>& Thread : Threads)
+	{
+		Thread.Wait();
+	}
+
+	return ErrorCount == 0;
+}
+
+////////////////////////////////////////////////////////////////////////////////
+static TIoStatusOr<FUploadResult> UploadContainerFiles(
+	const FUploadParams& UploadParams,
+	TConstArrayView<FString> ContainerFiles,
+	const FKeyChain& KeyChain)
+{
+	TMap<FGuid, FAES::FAESKey> EncryptionKeys;
+	for (const TPair<FGuid, FNamedAESKey>& KeyPair: KeyChain.GetEncryptionKeys())
+	{
+		EncryptionKeys.Add(KeyPair.Key, KeyPair.Value.Key);
+	}
+
+	struct FContainerStats
+	{
+		FString ContainerName;
+		uint64 ChunkCount = 0;
+		uint64 TotalBytes = 0;
+		uint64 UploadedChunkCount = 0;
+		uint64 UploadedBytes = 0;
+	};
+	TMap<FString, FContainerStats> ContainerSummary;
+
+	const double StartTime = FPlatformTime::Seconds();
+
+	FS3ClientConfig Config;
+	Config.ServiceUrl = UploadParams.ServiceUrl;
+	Config.Region = UploadParams.Region;
+	
+	FS3ClientCredentials Credentials;
+	if (UploadParams.CredentialsFile.IsEmpty() == false)
+	{
+		UE_LOG(LogIas, Display, TEXT("Loading credentials file '%s'"), *UploadParams.CredentialsFile);
+		FS3CredentialsProfileStore CredentialsStore = FS3CredentialsProfileStore::FromFile(UploadParams.CredentialsFile);
+		if (CredentialsStore.TryGetCredentials(UploadParams.CredentialsFileKeyName, Credentials) == false)
+		{
+			return FIoStatus(EIoErrorCode::InvalidParameter, TEXT("Failed to find valid credentials in credentials file"));
+		}
+		else
+		{
+			UE_LOG(LogIas, Display, TEXT("Found credentials for '%s'"), *UploadParams.CredentialsFileKeyName);
+		}
+	}
+	else
+	{
+		Credentials = FS3ClientCredentials(UploadParams.AccessKey, UploadParams.SecretKey, UploadParams.SessionToken);
+	}
+
+	FS3Client Client(Config, Credentials);
+	FUploadQueue UploadQueue(Client, UploadParams.Bucket, UploadParams.MaxConcurrentUploads);
+
+	if (ContainerFiles.Num() == 0)
+	{
+		return FIoStatus(EIoErrorCode::InvalidParameter, TEXT("No container file(s) specified"));
+	}
+
+	TSet<FIoHash> ExistingChunks;
+	uint64 TotalExistingTocs = 0;
+	uint64 TotalExistingBytes = 0;
+	{
+		TStringBuilder<256> TocsKey;
+		if (!UploadParams.BucketPrefix.IsEmpty())
+		{
+			TocsKey << UploadParams.BucketPrefix << "/";
+		}
+
+		UE_LOG(LogIas, Display, TEXT("	 '%s/%s/%s'"), *Client.GetConfig().ServiceUrl, *UploadParams.Bucket, TocsKey.ToString());
+		FS3ListObjectResponse Response = Client.ListObjects(FS3ListObjectsRequest
+		{
+			UploadParams.Bucket,
+			TocsKey.ToString(),
+			TEXT('/')
+		});
+
+		for (const FS3Object& TocInfo : Response.Objects)
+		{
+			if (TocInfo.Key.EndsWith(TEXT("iochunktoc")) == false)
+			{
+				continue;
+			}
+
+			UE_LOG(LogIas, Display, TEXT("Fetching TOC '%s/%s/%s'"), *Client.GetConfig().ServiceUrl, *UploadParams.Bucket, *TocInfo.Key);
+			FS3GetObjectResponse TocResponse = Client.GetObject(FS3GetObjectRequest
+			{
+				UploadParams.Bucket,
+				TocInfo.Key
+			});
+			
+			if (TocResponse.IsOk() == false)
+			{
+				UE_LOG(LogIas, Warning, TEXT("Failed to fetch TOC '%s/%s/%s'"), *Client.GetConfig().ServiceUrl, *UploadParams.Bucket, *TocInfo.Key);
+				continue;
+			}
+
+			FOnDemandToc Toc;
+			FMemoryReaderView Ar(TocResponse.GetBody().GetView());
+			Ar << Toc;
+
+			if (Ar.IsError()) 
+			{
+				Toc = FOnDemandToc{};
+				if (LoadFromCompactBinary(FCbFieldView(TocResponse.GetBody().GetData()), Toc) == false)
+				{
+					UE_LOG(LogIas, Display, TEXT("Failed to load TOC '%s/%s/%s'"), *Client.GetConfig().ServiceUrl, *UploadParams.Bucket, *TocInfo.Key);
+					continue;
+				}
+			}
+
+			for (const FOnDemandTocContainerEntry& ContainerEntry : Toc.Containers)
+			{
+				for (const FOnDemandTocEntry& TocEntry : ContainerEntry.Entries)
+				{
+					bool bIsAlreadyInSet;
+					ExistingChunks.Add(TocEntry.Hash, &bIsAlreadyInSet);
+					if (!bIsAlreadyInSet)
+					{
+						TotalExistingBytes += TocEntry.EncodedSize;
+					}
+				}
+			}
+
+			TotalExistingTocs++;
+		}
+	}
+
+	FString ChunksRelativePath = UploadParams.BucketPrefix.IsEmpty() ? TEXT("Chunks") : FString::Printf(TEXT("%s/Chunks"), *UploadParams.BucketPrefix);
+	ChunksRelativePath.ToLowerInline();
+
+	bool bPerContainerTocsConfigValue = false;
+	GConfig->GetBool(TEXT("Ias"), TEXT("CreatePerContainerTocs"), bPerContainerTocsConfigValue, GEngineIni);
+	const bool bWritePerContainerToc = bPerContainerTocsConfigValue || UploadParams.bPerContainerTocs;
+
+	uint64 TotalUploadedChunks = 0;
+	uint64 TotalUploadedBytes = 0;
+
+	FOnDemandToc OnDemandToc;
+	OnDemandToc.Containers.Reserve(ContainerFiles.Num());
+
+	// Map of the .utoc paths that we have created ondemand containers for, indexed by their hash so that the paths can
+	// be looked up later. We do not rely on the filename as we cannot be sure that there won't be duplicate file names
+	// stored in different directories.
+	TMap<FIoHash, FString> UTocPaths;
+
+	TArray<FString> FilesToDelete;
+	for (const FString& Path : ContainerFiles)
+	{
+		FIoStoreReader ContainerFileReader;
+		{
+			FIoStatus Status = ContainerFileReader.Initialize(*FPaths::ChangeExtension(Path, TEXT("")), EncryptionKeys);
+			if (!Status.IsOk())
+			{
+				UE_LOG(LogIas, Error, TEXT("Failed to open container '%s' for reading"), *Path);
+				continue;
+			}
+		}
+
+		if (EnumHasAnyFlags(ContainerFileReader.GetContainerFlags(), EIoContainerFlags::OnDemand) == false)
+		{
+			continue;
+		}
+		UE_LOG(LogIas, Display, TEXT("Uploading ondemand container '%s'"), *Path);
+
+		const uint32 BlockSize = ContainerFileReader.GetCompressionBlockSize();
+		if (OnDemandToc.Header.BlockSize == 0)
+		{
+			OnDemandToc.Header.BlockSize = ContainerFileReader.GetCompressionBlockSize();
+		}
+		check(OnDemandToc.Header.BlockSize == ContainerFileReader.GetCompressionBlockSize());
+
+		TArray<FIoStoreTocChunkInfo> ChunkInfos;
+		ContainerFileReader.EnumerateChunks([&ChunkInfos](FIoStoreTocChunkInfo&& Info)
+		{ 
+			ChunkInfos.Emplace(MoveTemp(Info));
+			return true;
+		});
+		
+		FOnDemandTocContainerEntry& ContainerEntry = OnDemandToc.Containers.AddDefaulted_GetRef();
+		ContainerEntry.ContainerName = FPaths::GetBaseFilename(Path);
+		ContainerEntry.ContainerId = ContainerFileReader.GetContainerId();
+
+		if (EnumHasAnyFlags(ContainerFileReader.GetContainerFlags(), EIoContainerFlags::Encrypted))
+		{
+			ContainerEntry.EncryptionKeyGuid = LexToString(ContainerFileReader.GetEncryptionKeyGuid());
+		}
+		FContainerStats& ContainerStats = ContainerSummary.FindOrAdd(ContainerEntry.ContainerName);
+		
+		ContainerEntry.Entries.Reserve(ChunkInfos.Num());
+
+		for (const FIoStoreTocChunkInfo& ChunkInfo : ChunkInfos)
+		{
+			const bool bDecrypt = false;
+			TIoStatusOr<FIoStoreCompressedReadResult> Status = ContainerFileReader.ReadCompressed(ChunkInfo.Id, FIoReadOptions(), bDecrypt);
+			if (!Status.IsOk())
+			{
+				return Status.Status();
+			}
+
+			FIoStoreCompressedReadResult ReadResult = Status.ConsumeValueOrDie();
+
+			const uint32 BlockOffset = ContainerEntry.BlockSizes.Num();
+			const uint32 BlockCount = ReadResult.Blocks.Num();
+			const FIoHash ChunkHash = FIoHash::HashBuffer(ReadResult.IoBuffer.GetView());
+
+			FMemoryView EncodedBlocks = ReadResult.IoBuffer.GetView();
+			uint64 RawChunkSize = 0;
+			uint64 EncodedChunkSize = 0;
+			for (const FIoStoreCompressedBlockInfo& BlockInfo : ReadResult.Blocks)
+			{
+				check(Align(BlockInfo.CompressedSize, FAES::AESBlockSize) == BlockInfo.AlignedSize);
+				const uint64 EncodedBlockSize = BlockInfo.AlignedSize;
+				ContainerEntry.BlockSizes.Add(uint32(BlockInfo.CompressedSize));
+
+				FMemoryView EncodedBlock = EncodedBlocks.Left(EncodedBlockSize);
+				EncodedBlocks += EncodedBlock.GetSize();
+				ContainerEntry.BlockHashes.Add(FIoChunkEncoding::HashBlock(EncodedBlock));
+
+				EncodedChunkSize += EncodedBlockSize;
+				RawChunkSize += BlockInfo.UncompressedSize;
+
+				if (OnDemandToc.Header.CompressionFormat.IsEmpty() && BlockInfo.CompressionMethod != NAME_None)
+				{
+					OnDemandToc.Header.CompressionFormat = BlockInfo.CompressionMethod.ToString();
+				}
+			}
+
+			if (EncodedChunkSize != ReadResult.IoBuffer.GetSize())
+			{
+				return FIoStatus(EIoErrorCode::ReadError, TEXT("Encoded chunk size does not match buffer"));
+			}
+
+			FOnDemandTocEntry& TocEntry = ContainerEntry.Entries.AddDefaulted_GetRef();
+			TocEntry.ChunkId = ChunkInfo.Id;
+			TocEntry.Hash = ChunkHash;
+			TocEntry.RawSize = RawChunkSize;
+			TocEntry.EncodedSize = EncodedChunkSize;
+			TocEntry.BlockOffset = BlockOffset;
+			TocEntry.BlockCount = BlockCount;
+
+			ContainerStats.ChunkCount++;
+			ContainerStats.TotalBytes += EncodedChunkSize;
+
+			if (ExistingChunks.Contains(TocEntry.Hash) == false)
+			{
+				const FString HashString = LexToString(ChunkHash);
+
+				TStringBuilder<256> Key;
+				Key << ChunksRelativePath
+					<< TEXT("/") << HashString.Left(2)
+					<< TEXT("/") << HashString
+					<< TEXT(".iochunk");
+
+				if (UploadQueue.Enqueue(Key, ReadResult.IoBuffer) == false)
+				{
+					return FIoStatus(EIoErrorCode::WriteError, TEXT("Failed to upload chunk"));
+				}
+
+				TotalUploadedChunks++;
+				TotalUploadedBytes += EncodedChunkSize;
+				ContainerStats.UploadedChunkCount++;
+				ContainerStats.UploadedBytes += EncodedChunkSize;
+			}
+		}
+
+		{
+			const FString UTocFilePath = FPaths::ChangeExtension(Path, TEXT(".utoc"));
+			TArray<uint8> Buffer;
+			if (!FFileHelper::LoadFileToArray(Buffer, *UTocFilePath))
+			{
+				return FIoStatus(EIoErrorCode::ReadError, TEXT("Failed to upload .utoc file"));
+			}
+
+			ContainerEntry.UTocHash = FIoHash::HashBuffer(Buffer.GetData(), Buffer.Num());
+			TStringBuilder<256> Key;
+			if (UploadParams.BucketPrefix.IsEmpty() == false)
+			{
+				Key << UploadParams.BucketPrefix.ToLower() << TEXT("/");
+			}
+			Key << LexToString(ContainerEntry.UTocHash) << TEXT(".utoc");
+
+			const FS3PutObjectResponse Response = Client.TryPutObject(
+				FS3PutObjectRequest{UploadParams.Bucket, Key.ToString(), MakeMemoryView(Buffer.GetData(), Buffer.Num())});
+
+			if (!Response.IsOk())
+			{
+				return FIoStatus(EIoErrorCode::WriteError, FString::Printf(TEXT("Failed to upload '%s' (%s)"), *UTocFilePath, *Response.GetErrorStatus()));
+			}
+			
+			UTocPaths.Add(ContainerEntry.UTocHash, UTocFilePath);
+
+			UE_LOG(LogIas, Display, TEXT("Uploaded '%s'"), *UTocFilePath);	
+		}
+
+		if (UploadParams.bDeleteContainerFiles)
+		{
+			FilesToDelete.Add(Path);
+			ContainerFileReader.GetContainerFilePaths(FilesToDelete);
+
+			// We need the pak files in order to mount OnDemand toc files!
+			if (UploadParams.bDeletePakFiles && !bWritePerContainerToc)
+			{
+				FilesToDelete.Add(FPaths::ChangeExtension(Path, TEXT(".pak")));
+				FilesToDelete.Add(FPaths::ChangeExtension(Path, TEXT(".sig")));
+			}
+		}
+	}
+
+	if (OnDemandToc.Containers.IsEmpty())
+	{
+		return FIoStatus(EIoErrorCode::InvalidParameter, TEXT("No container file(s) marked as on demand"));
+	}
+
+	if (!UploadQueue.Flush())
+	{
+		return FIoStatus(EIoErrorCode::WriteError, TEXT("Failed to upload chunk(s)"));
+	}
+
+	FString ServiceUrl;
+	GetTocPath(UploadParams.ServiceUrl, UploadParams.Bucket, UploadParams.BucketPrefix, ServiceUrl, OnDemandToc.Header.ChunksDirectory);
+
+	OnDemandToc.Meta.EpochTimestamp = FDateTime::Now().ToUnixTimestamp();
+	OnDemandToc.Meta.BuildVersion = UploadParams.BuildVersion;
+	OnDemandToc.Meta.TargetPlatform = UploadParams.TargetPlatform;
+
+	FUploadResult UploadResult;
+	{
+		FLargeMemoryWriter Ar;
+		Ar << OnDemandToc;
+
+		UploadResult.TocHash = FIoHash::HashBuffer(Ar.GetView());
+		const FString TocFilename = LexToString(UploadResult.TocHash) + TEXT(".iochunktoc");
+
+		TStringBuilder<256> Key;
+		if (!UploadParams.BucketPrefix.IsEmpty())
+		{
+			Key << UploadParams.BucketPrefix.ToLower() << TEXT("/");
+		}
+		Key << TocFilename;
+
+		UploadResult.ServiceUrl = ServiceUrl;
+		UploadResult.TocPath = OnDemandToc.Header.ChunksDirectory / TocFilename;
+		UploadResult.TocSize = Ar.TotalSize();
+
+		const FS3PutObjectResponse Response = Client.TryPutObject(FS3PutObjectRequest{UploadParams.Bucket, Key.ToString(), Ar.GetView()});
+		if (Response.IsOk())
+		{
+			UE_LOG(LogIas, Display, TEXT("Uploaded on demand TOC '%s/%s/%s'"), *Client.GetConfig().ServiceUrl, *UploadParams.Bucket, Key.ToString());
+		}
+		else
+		{
+			UE_LOG(LogIas, Warning, TEXT("Failed to upload TOC '%s/%s/%s' (%s)"), *Client.GetConfig().ServiceUrl, *UploadParams.Bucket, Key.ToString(), *Response.GetErrorStatus());
+			return FIoStatus(EIoErrorCode::WriteError, TEXT("Failed to upload TOC"));
+		}
+
+		if (UploadParams.bWriteTocToDisk)
+		{
+			// Write a single .iochunktoc containing all on demand data for the current build
+			TStringBuilder<512> OnDemandTocFilePath;
+			FPathViews::Append(OnDemandTocFilePath, UploadParams.TocOutputDir, UploadResult.TocHash);
+			OnDemandTocFilePath << TEXT(".iochunktoc");
+
+			if (SaveArrayToFile(TArrayView64<uint8>(Ar.GetData(), Ar.TotalSize()), OnDemandTocFilePath.ToString()) == false)
+			{
+				UE_LOG(LogIoStore, Error, TEXT("Failed to save on demand toc file '%s'"), OnDemandTocFilePath.ToString());
+			}
+		}
+	}
+
+	if (bWritePerContainerToc)
+	{
+		// Write out separate .uondemandtoc files, one per.utoc containing ondemand data.
+		FIoStatus Result = WriteContainerFiles(OnDemandToc, UTocPaths);
+		if (!Result.IsOk())
+		{
+			return Result;
+		}
+	}
+
+	for (const FString& Path : FilesToDelete)
+	{
+		if (IFileManager::Get().FileExists(*Path))
+		{
+			UE_LOG(LogIas, Display, TEXT("Deleting '%s'"), *Path); 
+			IFileManager::Get().Delete(*Path);
+		}
+	}
+
+	if (!UploadParams.TocOutputDir.IsEmpty())
+	{
+		WriteConfigFile(UploadParams, UploadResult, KeyChain);
+	}
+
+	{
+		const double Duration = FPlatformTime::Seconds() - StartTime;
+		
+		UE_LOG(LogIas, Display, TEXT(""));
+		UE_LOG(LogIas, Display, TEXT("------------------------------------------------- Upload Summary -------------------------------------------------"));
+		UE_LOG(LogIas, Display, TEXT("%-15s: %s"), TEXT("Service URL"), *UploadParams.ServiceUrl);
+		UE_LOG(LogIas, Display, TEXT("%-15s: %s"), TEXT("Bucket"), *UploadParams.Bucket);
+		UE_LOG(LogIas, Display, TEXT("%-15s: %s"), TEXT("TargetPlatform"), *UploadParams.TargetPlatform);
+		UE_LOG(LogIas, Display, TEXT("%-15s: %s"), TEXT("BuildVersion"), *UploadParams.BuildVersion);
+		UE_LOG(LogIas, Display, TEXT("%-15s: %s"), TEXT("TOC path"), *UploadResult.TocPath);
+		UE_LOG(LogIas, Display, TEXT("%-15s: %.2lf KiB"), TEXT("TOC size"), double(UploadResult.TocSize) / 1024.0);
+		UE_LOG(LogIas, Display, TEXT("%-15s: %.2lf second(s)"), TEXT("Duration"), Duration);
+		UE_LOG(LogIas, Display, TEXT(""));
+		
+		UE_LOG(LogIas, Display, TEXT("%-50s %15s %15s %15s %15s"),
+			TEXT("Container"), TEXT("Chunk(s)"), TEXT("Size (MiB)"), TEXT("Uploaded"), TEXT("Uploaded (MiB)"));
+		UE_LOG(LogIas, Display, TEXT("-------------------------------------------------------------------------------------------------------------------"));
+
+		FContainerStats TotalStats;
+		for (const TTuple<FString, FContainerStats>& Kv : ContainerSummary)
+		{
+			UE_LOG(LogIas, Display, TEXT("%-50s %15llu %15.2lf %15llu %15.2lf"),
+				*Kv.Key, Kv.Value.ChunkCount, double(Kv.Value.TotalBytes) / 1024.0 / 1024.0, Kv.Value.UploadedChunkCount, double(Kv.Value.UploadedBytes) / 1024.0 / 1024.0);
+			
+			TotalStats.ChunkCount += Kv.Value.ChunkCount;
+			TotalStats.TotalBytes += Kv.Value.TotalBytes;
+			TotalStats.UploadedChunkCount += Kv.Value.UploadedChunkCount;
+			TotalStats.UploadedBytes += Kv.Value.UploadedBytes;
+		}
+		UE_LOG(LogIas, Display, TEXT("-------------------------------------------------------------------------------------------------------------------"));
+		UE_LOG(LogIas, Display, TEXT("%-50s %15llu %15.2lf %15llu %15.2lf"),
+			TEXT("Total"),TotalStats.ChunkCount, double(TotalStats.TotalBytes) / 1024.0 / 1024.0, TotalStats.UploadedChunkCount, double(TotalStats.UploadedBytes) / 1024.0 / 1024.0);
+		UE_LOG(LogIas, Display, TEXT(""));
+		
+		UE_LOG(LogIas, Display, TEXT("%-50s %15s %15s %15s"), TEXT("Bucket"), TEXT("TOC(s)"), TEXT("Chunk(s)"), TEXT("MiB"));
+		UE_LOG(LogIas, Display, TEXT("-------------------------------------------------------------------------------------------------------------------"));
+		UE_LOG(LogIas, Display, TEXT("%-50s %15llu %15d %15.2lf"), TEXT("Existing"), TotalExistingTocs, ExistingChunks.Num(), double(TotalExistingBytes) / 1024.0 / 1024.0);
+		UE_LOG(LogIas, Display, TEXT("%-50s %15llu %15d %15.2lf"), TEXT("Uploaded"), 1, TotalStats.UploadedChunkCount, double(TotalStats.UploadedBytes) / 1024.0 / 1024.0);
+		UE_LOG(LogIas, Display, TEXT("-------------------------------------------------------------------------------------------------------------------"));
+		UE_LOG(LogIas, Display, TEXT("%-50s %15llu %15d %15.2lf"), TEXT("Total"),
+			TotalExistingTocs + 1, ExistingChunks.Num() + TotalStats.UploadedChunkCount, double(TotalExistingBytes + TotalStats.UploadedBytes) / 1024.0 / 1024.0);
+		UE_LOG(LogIas, Display, TEXT(""));
+	}
+	
+	return UploadResult;
+}
+
 ////////////////////////////////////////////////////////////////////////////////
 static TArray<FString> GlobContainers(const FContext& Context)
 {
@@ -116,9 +969,9 @@ static int32 UploadCommandEntry(const FContext& Context)
 {
 	TArray<FString> Containers = GlobContainers(Context);
 	FKeyChain KeyChain = LoadKeys(Context);
-	FIoStoreUploadParams Params = BuildUploadParams(Context);
+	FUploadParams Params = BuildUploadParams(Context);
 
-	TIoStatusOr<FIoStoreUploadResult> Result = UploadContainerFiles(Params, Containers, KeyChain);
+	TIoStatusOr<FUploadResult> Result = UploadContainerFiles(Params, Containers, KeyChain);
 	if (!Result.IsOk())
 	{
 		FString Reason = Result.Status().ToString();
diff --git a/Engine/Source/Runtime/Experimental/IoStore/OnDemand/Public/IO/IoStoreOnDemand.h b/Engine/Source/Runtime/Experimental/IoStore/OnDemand/Public/IO/IoStoreOnDemand.h
--- a/Engine/Source/Runtime/Experimental/IoStore/OnDemand/Public/IO/IoStoreOnDemand.h
+++ b/Engine/Source/Runtime/Experimental/IoStore/OnDemand/Public/IO/IoStoreOnDemand.h
@@ -192,57 +192,6 @@ UE_API bool LoadFromCompactBinary(FCbFieldView Field, FOnDemandToc& OutToc);
 
 #if (IS_PROGRAM || WITH_EDITOR)
 
-////////////////////////////////////////////////////////////////////////////////
-struct FIoStoreUploadParams
-{
-	FString ServiceUrl;
-	FString DistributionUrl;
-	FString FallbackUrl;
-	FString Bucket;
-	FString BucketPrefix;
-	FString Region; 
-	FString AccessKey;
-	FString SecretKey;
-	FString SessionToken;
-	FString CredentialsFile;
-	FString CredentialsFileKeyName;
-	FString BuildVersion;
-	FString TargetPlatform;
-	FString EncryptionKeyName;
-	int32 MaxConcurrentUploads = 16;
-	bool bDeleteContainerFiles = true;
-	bool bDeletePakFiles = true;
-	bool bPerContainerTocs = false;
-
-	/** If we should write out the .iochunktoc to disk as well as uploading it. */
-	bool bWriteTocToDisk = false;
-	/** Where the .iochunktoc file should be written out. */
-	FString TocOutputDir;
-
-	static TIoStatusOr<FIoStoreUploadParams> Parse(const TCHAR* CommandLine);
-	FIoStatus Validate() const;
-};
-
-/** Results from uploading a FOnDemandToc */
-struct FIoStoreUploadResult
-{
-	/** Hash of the toc when written as a binary blob */
-	FIoHash TocHash;
-	
-	/** Url of the service that the toc was uploaded too */
-	FString ServiceUrl;
-	/** Path of the toc on the service */
-	FString TocPath;
-
-	/** Size (in bytes) of the toc when written as a binary blob */
-	uint64 TocSize = 0;
-};
-
-TIoStatusOr<FIoStoreUploadResult> UploadContainerFiles(
-	const FIoStoreUploadParams& UploadParams,
-	TConstArrayView<FString> ContainerFiles,
-	const FKeyChain& KeyChain);
-
 ////////////////////////////////////////////////////////////////////////////////
 struct FIoStoreDownloadParams
 {
