Iris
- Fixed missing referencecount on destructioninfoprotocol (take2)

diff --git a/Engine/Source/Runtime/Experimental/Iris/Core/Private/Iris/ReplicationSystem/ReplicationBridge.cpp b/Engine/Source/Runtime/Experimental/Iris/Core/Private/Iris/ReplicationSystem/ReplicationBridge.cpp
--- a/Engine/Source/Runtime/Experimental/Iris/Core/Private/Iris/ReplicationSystem/ReplicationBridge.cpp
+++ b/Engine/Source/Runtime/Experimental/Iris/Core/Private/Iris/ReplicationSystem/ReplicationBridge.cpp
@@ -342,6 +342,8 @@ void UReplicationBridge::Initialize(UReplicationSystem* InReplicationSystem)
 		const FReplicationFragments RegisteredFragments;
 		FCreateReplicationProtocolParameters CreateProtocolParams {.bValidateProtocolId = false, .TypeStatsIndex =  GetReplicationSystem()->GetReplicationSystemInternal()->GetNetTypeStats().GetOrCreateTypeStats(FName("DestructionInfo"))};
 		DestructionInfoProtocol = ReplicationProtocolManager->CreateReplicationProtocol(FReplicationProtocolManager::CalculateProtocolIdentifier(RegisteredFragments), RegisteredFragments, TEXT("InternalDestructionInfo"), CreateProtocolParams);
+		// Explicit refcount
+		DestructionInfoProtocol->AddRef();
 	}
 }
 
@@ -376,7 +378,11 @@ void UReplicationBridge::Deinitialize()
 	});
 
 	// Release destructioninfo protocol.
-	DestructionInfoProtocol = nullptr;
+	if (DestructionInfoProtocol)
+	{
+		DestructionInfoProtocol->Release();
+		DestructionInfoProtocol = nullptr;
+	}
 
 	ReplicationSystem = nullptr;
 	ReplicationProtocolManager = nullptr;
diff --git a/Engine/Source/Runtime/Experimental/Iris/Core/Public/Iris/ReplicationSystem/ReplicationBridge.h b/Engine/Source/Runtime/Experimental/Iris/Core/Public/Iris/ReplicationSystem/ReplicationBridge.h
--- a/Engine/Source/Runtime/Experimental/Iris/Core/Public/Iris/ReplicationSystem/ReplicationBridge.h
+++ b/Engine/Source/Runtime/Experimental/Iris/Core/Public/Iris/ReplicationSystem/ReplicationBridge.h
@@ -10,7 +10,6 @@
 #include "Iris/ReplicationSystem/NetObjectGroupHandle.h"
 #include "Misc/EnumClassFlags.h"
 #include "Net/Core/NetHandle/NetHandle.h"
-#include "Templates/RefCounting.h"
 #include "UObject/ObjectKey.h"
 #include "UObject/ObjectMacros.h"
 
@@ -385,7 +384,7 @@ private:
 		UE::Net::Private::FInternalNetRefIndex InternalReplicationIndex;
 	};
 
-	TRefCountPtr<const UE::Net::FReplicationProtocol> DestructionInfoProtocol;
+	const UE::Net::FReplicationProtocol* DestructionInfoProtocol;
 	
 	// Need to track the objects with destruction infos so that we can clean them up properly
 	// We use this to be able ask remote to destroy static objects
