Load root authorities the same way on all systems and set them as the defaults for IAS.

To verify a TLS connection, the server provides a chain of signed certificates that we can compare against a set of "root certificate authorities". Long story short - getting "system" root CAs varies by platform and UE's SSL module is incomplete. So for IAS one has to explicitly load root CAs - a simple path that's same for all systems, and we get protection against MITM proxying. Can also pin later if need be.

#rnx

diff --git a/Engine/Source/Runtime/Experimental/IoStore/OnDemand/Private/IoStoreOnDemand.cpp b/Engine/Source/Runtime/Experimental/IoStore/OnDemand/Private/IoStoreOnDemand.cpp
--- a/Engine/Source/Runtime/Experimental/IoStore/OnDemand/Private/IoStoreOnDemand.cpp
+++ b/Engine/Source/Runtime/Experimental/IoStore/OnDemand/Private/IoStoreOnDemand.cpp
@@ -304,6 +304,30 @@ static FIasCacheConfig GetIasCacheConfig(const TCHAR* CommandLine)
 	return Ret;
 }
 
+////////////////////////////////////////////////////////////////////////////////
+static void LoadCaCerts()
+{
+	using namespace UE::IoStore::HTTP;
+
+	IFileManager& Ifm = IFileManager::Get();
+	FString PemPath = FPaths::EngineContentDir() / TEXT("Certificates/ThirdParty/cacert.pem");
+	FArchive* Reader = Ifm.CreateFileReader(*PemPath);
+	check(Reader != nullptr)
+
+	uint32 Size = uint32(Reader->TotalSize());
+	FIoBuffer PemData(Size);
+	FMutableMemoryView PemView = PemData.GetMutableView();
+	Reader->Serialize(PemView.GetData(), Size);
+
+	FCertRoots CaRoots(PemData.GetView());
+
+	uint32 NumCerts = CaRoots.Num();
+	UE_LOG(LogIas, Display, TEXT("CaRoots: %u (%u .pem bytes))"), NumCerts, Size);
+
+	FCertRoots::SetDefault(MoveTemp(CaRoots));
+	delete Reader;
+}
+
 ////////////////////////////////////////////////////////////////////////////////
 /**
  * Utility to create a FArchive capable of reading from disk using the exact same pathing
@@ -1811,6 +1835,8 @@ void FIoStoreOnDemandModule::InitializeInternal()
 	}
 #endif
 
+	LoadCaCerts();
+
 	// Make sure we haven't called initialize before
 	check(!HttpIoDispatcherBackend.IsValid());
 
diff --git a/Engine/Source/Runtime/Experimental/IoStore/OnDemand/Private/Tool/Purl.cpp b/Engine/Source/Runtime/Experimental/IoStore/OnDemand/Private/Tool/Purl.cpp
--- a/Engine/Source/Runtime/Experimental/IoStore/OnDemand/Private/Tool/Purl.cpp
+++ b/Engine/Source/Runtime/Experimental/IoStore/OnDemand/Private/Tool/Purl.cpp
@@ -5,9 +5,11 @@
 #include "Command.h"
 
 #include <Containers/UnrealString.h>
+#include <HAL/FileManager.h>
 #include <HAL/PlatformProcess.h>
 #include <IO/IoBuffer.h>
 #include <IO/Http/Client.h>
+#include <Misc/Paths.h>
 #include <Misc/ScopeExit.h>
 
 #if PLATFORM_WINDOWS
@@ -22,6 +24,28 @@
 namespace UE::IoStore::Tool
 {
 
+////////////////////////////////////////////////////////////////////////////////
+static void LoadCaCerts()
+{
+	using namespace UE::IoStore::HTTP;
+
+	IFileManager& Ifm = IFileManager::Get();
+	FString PemPath = FPaths::EngineContentDir() / TEXT("Certificates/ThirdParty/cacert.pem");
+	FArchive* Reader = Ifm.CreateFileReader(*PemPath);
+	check(Reader != nullptr)
+
+	uint32 Size = uint32(Reader->TotalSize());
+	FIoBuffer PemData(Size);
+	FMutableMemoryView PemView = PemData.GetMutableView();
+	Reader->Serialize(PemView.GetData(), Size);
+
+	FCertRoots CaRoots(PemData.GetView());
+	check(CaRoots.IsValid());
+
+	FCertRoots::SetDefault(MoveTemp(CaRoots));
+	delete Reader;
+}
+
 ////////////////////////////////////////////////////////////////////////////////
 static int32 PurlCommandEntry(const FContext& Context)
 {
@@ -34,6 +58,8 @@ static int32 PurlCommandEntry(const FContext& Context)
 
 	using namespace UE::IoStore::HTTP;
 
+	LoadCaCerts();
+
 	FStringView Url = Context.Get<FStringView>(TEXT("Url"));
 	auto AnsiUrl = StringCast<ANSICHAR>(Url.GetData(), Url.Len());
 
