Updated IAS test server to also support TLS traffic

#rnx

diff --git a/Engine/Source/Runtime/Experimental/IoStore/HttpClient/Private/TestServer.py b/Engine/Source/Runtime/Experimental/IoStore/HttpClient/Private/TestServer.py
--- a/Engine/Source/Runtime/Experimental/IoStore/HttpClient/Private/TestServer.py
+++ b/Engine/Source/Runtime/Experimental/IoStore/HttpClient/Private/TestServer.py
@@ -4,9 +4,11 @@ import time
 import base64
 import socket
 import random
+import tempfile
 import threading
 import http.server
 import http.client
+import subprocess as sp
 from pathlib import Path
 from urllib.parse import parse_qsl
 
@@ -14,6 +16,8 @@ from urllib.parse import parse_qsl
 
 def intercept(fd):
     line = fd.readline()
+    if b"HTTP/1.1" not in line:
+        return False
     try:
         headers = http.client.parse_headers(fd)
     except http.client.HTTPException:
@@ -30,6 +34,8 @@ def make_preamble(line, headers):
 def proxy_impl(client, httpd):
     # get request
     req = intercept(client)
+    if not req:
+        return False
     if isinstance(req, int):
         client.write(f"HTTP/1.1 {req} IasTestServerProxyError\r\n".encode())
         client.write(b"Content-Length: 0\r\n\r\n")
@@ -61,7 +67,7 @@ def proxy_impl(client, httpd):
     # get response
     line, headers = intercept(httpd)
     close = close or (headers.get("Connection", "").lower() == "close")
-    content_len = int(headers["Content-Length"])
+    content_len = int(headers.get("Content-Length", "0"))
 
     # get data to retransmit
     data = make_preamble(line, headers)
@@ -113,28 +119,68 @@ def proxy_client(client, httpd_port):
     client.close()
 
 def proxy_loop(httpd_port):
+    port = 9493
+    print(f"clear-text proxy: {port}")
     with socket.socket(socket.AF_INET, socket.SOCK_STREAM) as sock:
-        sock.bind(("", 9493))
+        sock.bind(("", port))
         sock.listen(16)
         sock.setblocking(True)
         while True:
             client, address = sock.accept()
+            threading.Thread(target=proxy_client, args=(client, httpd_port), daemon=True).start()
+
+
+
+# {{{1 tls .....................................................................
+
+def tls_proxy_loop(httpd_port, root_cert, *server_pems):
+    port = 4939
+    print(f"tls proxy: {port}")
+
+    temp_dir_obj = tempfile.TemporaryDirectory(prefix="iastestserver_")
+    temp_dir = Path(temp_dir_obj.name)
+
+    cert_path = temp_dir / "server_kc"
+    with cert_path.open("wb") as out:
+        for pem in server_pems:
+            out.write(pem)
+        # out.write(root_cert)
+
+    ca_path = temp_dir / "server_ca"
+    with ca_path.open("wb") as out:
+        out.write(root_cert)
+
+    import ssl
+    ssl_ctx = ssl.SSLContext(ssl.PROTOCOL_TLS_SERVER)
+    # ssl_ctx.load_verify_locations(ca_path)
+    ssl_ctx.load_cert_chain(cert_path)
+
+    with socket.socket(socket.AF_INET, socket.SOCK_STREAM) as sock:
+        sock.bind(("", port))
+        sock.listen(16)
+        sock.setblocking(True)
+        while True:
             try:
-                threading.Thread(target=proxy_client, args=(client, httpd_port), daemon=True).start()
-            except (ConnectionResetError, ConnectionAbortedError):
-                pass
+                client, address = sock.accept()
+                client = ssl_ctx.wrap_socket(client, server_side=True)
+            except (ssl.SSLError, Exception) as e:
+                print("ERR:", str(e))
+                continue
+            threading.Thread(target=proxy_client, args=(client, httpd_port), daemon=True).start()
+
+
 
 # {{{1 httpd ...................................................................
 
 payload_data = random.randbytes(2 << 20)
 
-def http_seed(handler, value):
+def http_seed(handler, value=0):
     handler.send_response(200)
     handler.send_header("Content-Length", 0)
     handler.end_headers()
     random.seed(value)
 
-def http_data(handler, payload_size):
+def http_data(handler, payload_size=0):
     handler.send_response(200)
 
     mega_size = int(random.random() * (4 << 10))
@@ -213,6 +259,14 @@ class Handler(http.server.BaseHTTPRequestHandler):
         if query := parts[-1].split("?"):
             parts[-1] = query[0]
 
+        if parts[0] == "ca":
+            ca_pem = self.server.ca_pem
+            self.send_response(200)
+            self.send_header("Content-Length", len(ca_pem))
+            self.end_headers()
+            self.wfile.write(ca_pem)
+            return
+
         if parts[0] == "hello":
             self.send_response(200)
             self.send_header("Content-Length", 5)
@@ -224,25 +278,84 @@ class Handler(http.server.BaseHTTPRequestHandler):
             return http_redirect(self, *parts[1:])
 
         if parts[0] == "data":
-            size = -1
-            if len(parts) > 1:
-                size = int(parts[1])
-            return http_data(self, size)
+            return http_data(self, *parts[1:])
 
         if parts[0] == "seed":
-            size = int(parts[1])
-            return http_seed(self, size)
+            return http_seed(self, *parts[1:])
 
         return self.send_error(404, f"not found '{self.path}'")
 
-def plain_httpd_loop(port):
+def plain_httpd_loop(port, root_pem):
+    print(f"httpd: {port}")
     server = http.server.ThreadingHTTPServer(("", port), Handler)
+    server.ca_pem = root_pem
     server.serve_forever()
 
 
 
 # {{{1 main ....................................................................
 
+def gen_test_certs(openssl_bin):
+    temp_dir_obj = tempfile.TemporaryDirectory(prefix="iastestserver_")
+    temp_dir = Path(temp_dir_obj.name)
+
+    empty_cnf = temp_dir / "empty.cnf"
+    with empty_cnf.open("wb") as out:
+        out.write(b"[req]\n")
+        out.write(b"distinguished_name=ridgers\n")
+        out.write(b"[ridgers]\n")
+
+    root_key_path = temp_dir / "root_k"
+    with root_key_path.open("wb") as out:
+        sp.run((openssl_bin, "genrsa", "2048"), stdout=out)
+
+    root_path = temp_dir / "root_c"
+    with root_path.open("wb") as out:
+        sp.run((openssl_bin,
+            "req", "-new", "-x509",
+            "-nodes",
+            "-sha256",
+            "-key", str(root_key_path),
+            "-subj", "/C=SE/ST=SE/L=Stockholm/O=IasRoot/CN=localhost",
+            "-days", "10",
+            "-config", str(empty_cnf)),
+            stdout=out
+        )
+
+    key_path = temp_dir / "server_k"
+    with key_path.open("wb") as out:
+        sp.run((openssl_bin, "genrsa", "2048"), stdout=out)
+
+    req_path = temp_dir / "server_r"
+    with req_path.open("wb") as out:
+        sp.run((openssl_bin,
+            "req", "-new",
+            "-nodes",
+            "-sha256",
+            "-key", str(key_path),
+            "-subj", "/C=SE/ST=SE/L=Stockholm/O=Ias/CN=localhost",
+            "-config", str(empty_cnf)),
+            stdout=out
+        )
+
+    cert_path = temp_dir / "server_c"
+    with cert_path.open("wb") as out:
+        sp.run((openssl_bin,
+            "x509", "-req",
+            "-sha256",
+            "-in", str(req_path),
+            "-days", "10",
+            "-set_serial", "493",
+            "-CA", str(root_path),
+            "-CAkey", str(root_key_path)),
+            stdout=out
+        )
+
+    with root_path.open("rb") as inp: r = inp.read()
+    with cert_path.open("rb") as inp: s = inp.read()
+    with key_path.open("rb") as inp:  k = inp.read()
+    return r, s, k
+
 def main():
     for item in Path(__file__).parents:
         if (item / "GenerateProjectFiles.bat").is_file():
@@ -251,14 +364,24 @@ def main():
     else:
         assert False
 
+    print("\n## generating certificates")
+    openssl_bin = os.getenv("OPENSSL_BIN", "openssl")
+    if (x := Path(".") / "Engine/Binaries/DotNET/IOS/openssl.exe").is_file():
+        openssl_bin = str(x)
+    root_cert, server_cert, server_key = gen_test_certs(openssl_bin)
+
     httpd_port = int(random.random() * 0x8000) + 0x4000
 
     def start_svc(target, *args):
         threading.Thread(target=target, args=args, daemon=True).start()
 
+    print("\n## starting servers")
     start_svc(proxy_loop, httpd_port)
-    start_svc(plain_httpd_loop, httpd_port)
+    start_svc(tls_proxy_loop, httpd_port, root_cert, server_key, server_cert)
+    start_svc(plain_httpd_loop, httpd_port, root_cert)
 
+    time.sleep(0.25)
+    print("\n## ready")
     while True:
         time.sleep(3600)
 
