Geometry Collection : Fix Nanite Index attribute being created with the wrong type (FBox) but being used  as an int32
Note : this attribute does not seem to be used anywhere , will likely remove it in future changes

diff --git a/Engine/Source/Runtime/Experimental/GeometryCollectionEngine/Private/GeometryCollection/GeometryCollectionObject.cpp b/Engine/Source/Runtime/Experimental/GeometryCollectionEngine/Private/GeometryCollection/GeometryCollectionObject.cpp
--- a/Engine/Source/Runtime/Experimental/GeometryCollectionEngine/Private/GeometryCollection/GeometryCollectionObject.cpp
+++ b/Engine/Source/Runtime/Experimental/GeometryCollectionEngine/Private/GeometryCollection/GeometryCollectionObject.cpp
@@ -1484,7 +1484,7 @@ TSharedPtr<FGeometryCollection, ESPMode::ThreadSafe> UGeometryCollection::Genera
 
 		if (!DuplicateGeometryCollection->HasAttribute("NaniteIndex", "Transform"))
 		{
-			DuplicateGeometryCollection->AddAttribute<FBox>("NaniteIndex", "Transform");
+			DuplicateGeometryCollection->AddAttribute<int32>("NaniteIndex", "Transform");
 		}
 
 		const int32 NumTransforms = GeometryCollection->NumElements(FGeometryCollection::TransformGroup);
