Log the top 20 biggest cull distances at the end of every Iris match.

#tests Verified functionality with a local client + server match
#rnx

diff --git a/Engine/Source/Runtime/Experimental/Iris/Core/Private/Iris/ReplicationSystem/ObjectReplicationBridgeDebugging.cpp b/Engine/Source/Runtime/Experimental/Iris/Core/Private/Iris/ReplicationSystem/ObjectReplicationBridgeDebugging.cpp
--- a/Engine/Source/Runtime/Experimental/Iris/Core/Private/Iris/ReplicationSystem/ObjectReplicationBridgeDebugging.cpp
+++ b/Engine/Source/Runtime/Experimental/Iris/Core/Private/Iris/ReplicationSystem/ObjectReplicationBridgeDebugging.cpp
@@ -568,7 +568,7 @@ void UObjectReplicationBridge::PrintRelevantObjectsForConnections(const TArray<F
 //-----------------------------------------------
 FAutoConsoleCommand ObjectBridgePrintNetCullDistances(
 TEXT("Net.Iris.PrintNetCullDistances"),
-TEXT("Prints the list of replicated objects and their current netculldistance."),
+TEXT("Prints the list of replicated objects and their current netculldistance. Add -NumClasses=X to limit the printing to the X classes with the largest net cull distances."),
 FConsoleCommandWithArgsDelegate::CreateLambda([](const TArray< FString >& Args)
 {
 	using namespace UE::Net;
@@ -580,8 +580,6 @@ FConsoleCommandWithArgsDelegate::CreateLambda([](const TArray< FString >& Args)
 	{
 		if (UObjectReplicationBridge* ObjectBridge = CastChecked<UObjectReplicationBridge>(RepSystem->GetReplicationBridge()))
 		{
-			FReplicationSystemInternal* ReplicationSystemInternal = RepSystem->GetReplicationSystemInternal();
-
 			ObjectBridge->PrintNetCullDistances(Args);
 		}
 	}
@@ -593,6 +591,13 @@ void UObjectReplicationBridge::PrintNetCullDistances(const TArray<FString>& Args
 	using namespace UE::Net::Private;
 	using namespace UE::Net::Private::ObjectBridgeDebugging;
 
+	// Number of classes to print. If 0, print all.
+	int32 NumClassesToPrint = 0;
+	if (const FString* ClassCountArg = Args.FindByPredicate([](const FString& Str) { return Str.Contains(TEXT("NumClasses=")); }))
+	{
+		FParse::Value(**ClassCountArg, TEXT("NumClasses="), NumClassesToPrint);
+	}
+
 	FReplicationSystemInternal* ReplicationSystemInternal = GetReplicationSystem()->GetReplicationSystemInternal();
 	const FWorldLocations& WorldLocations = ReplicationSystemInternal->GetWorldLocations();
 	const FNetCullDistanceOverrides& CullDistanceOverrides = ReplicationSystemInternal->GetNetCullDistanceOverrides();
@@ -688,6 +693,7 @@ void UObjectReplicationBridge::PrintNetCullDistances(const TArray<FString>& Args
 	UE_LOG(LogIrisBridge, Display, TEXT("################ Start Printing NetCullDistance Values ################"));
 	UE_LOG(LogIrisBridge, Display, TEXT(""));
 
+	int32 NumClassesPrinted = 0;
 	for (auto ClassIt = ClassCullDistanceMap.CreateIterator(); ClassIt; ++ClassIt)
 	{
 		FCullDistanceInfo& Info = ClassIt.Value();
@@ -700,6 +706,11 @@ void UObjectReplicationBridge::PrintNetCullDistances(const TArray<FString>& Args
 		{
 			UE_LOG(LogIrisBridge, Display, TEXT("\tNetCullDistance: %f | UseCount: %d/%d (%.2f%%)"), DivergentIt.Key(), DivergentIt.Value(), Info.NumTotal,((float)DivergentIt.Value()/(float)Info.NumTotal)*100.f);
 		}
+
+		if (++NumClassesPrinted == NumClassesToPrint)
+		{
+			break;
+		}
 	}
 	
 	UE_LOG(LogIrisBridge, Display, TEXT(""));
