Iris
* Support runtime NetBitArray reallocations when maximum replicated object count is hit.
* FReplicationSystemParams can now configure an initial size for netobjet lists that can grow up to the MaxReplicatedObjectCount when the initialize size is all assigned to registered netobjets.
* Use InitialNetObjectListCount to set the initial list size.
* Use NetObjectListGrowCount to set by how many indexes we grow and realloca all netobject lists.
* Hitting the maximum active replicated object limit will now cause a fatal error and kill the process.

NetDriver
* Split Iris config into two structs: ReplicationSystemConfigClient and ReplicationSystemConfigServer

diff --git a/Engine/Plugins/Runtime/ReplicationSystemTestPlugin/Source/Private/Tests/Core/TestNetBitArray.cpp b/Engine/Plugins/Runtime/ReplicationSystemTestPlugin/Source/Private/Tests/Core/TestNetBitArray.cpp
--- a/Engine/Plugins/Runtime/ReplicationSystemTestPlugin/Source/Private/Tests/Core/TestNetBitArray.cpp
+++ b/Engine/Plugins/Runtime/ReplicationSystemTestPlugin/Source/Private/Tests/Core/TestNetBitArray.cpp
@@ -1550,6 +1550,42 @@ UE_NET_TEST(FNetBitArrayView, PrintHelper)
 	}
 }
 
+UE_NET_TEST(FNetBitArrayView, TestRoundUpToMaxWord)
+{
+	{
+		const uint32 NumBits = 0;
+		uint32 RoundUpZero = FNetBitArray::RoundUpToMaxWordBitCount(NumBits);
+		UE_NET_ASSERT_TRUE(RoundUpZero > NumBits);
+		UE_NET_ASSERT_TRUE((RoundUpZero % FNetBitArray::WordBitCount) == 0);
+
+		FNetBitArray BitArray(RoundUpZero);
+		UE_NET_ASSERT_TRUE(BitArray.GetNumWords() == 1);
+		UE_NET_ASSERT_TRUE(BitArray.GetNumBits() == RoundUpZero);
+	}
+
+	{
+		const uint32 NumBits = 32;
+		uint32 RoundUpWord = FNetBitArray::RoundUpToMaxWordBitCount(NumBits);
+		UE_NET_ASSERT_TRUE(RoundUpWord == NumBits);
+		UE_NET_ASSERT_TRUE((RoundUpWord % FNetBitArray::WordBitCount) == 0);
+
+		FNetBitArray BitArray(RoundUpWord);
+		UE_NET_ASSERT_TRUE(BitArray.GetNumWords() == 1);
+		UE_NET_ASSERT_TRUE(BitArray.GetNumBits() == RoundUpWord);
+	}
+
+	{
+		const uint32 NumBits = 48;
+		uint32 RoundUpTwoWords = FNetBitArray::RoundUpToMaxWordBitCount(NumBits);
+		UE_NET_ASSERT_TRUE(RoundUpTwoWords > NumBits);
+		UE_NET_ASSERT_TRUE((RoundUpTwoWords % FNetBitArray::WordBitCount) == 0);
+
+		FNetBitArray BitArray(RoundUpTwoWords);
+		UE_NET_ASSERT_TRUE(BitArray.GetNumWords() == 2);
+		UE_NET_ASSERT_TRUE(BitArray.GetNumBits() == RoundUpTwoWords);
+	}
+}
+
 class FNetBitArrayViewFixture : public FNetworkAutomationTestSuiteFixture
 {
 public:
diff --git a/Engine/Plugins/Runtime/ReplicationSystemTestPlugin/Source/Private/Tests/ReplicationSystem/Filtering/MockNetObjectFilter.cpp b/Engine/Plugins/Runtime/ReplicationSystemTestPlugin/Source/Private/Tests/ReplicationSystem/Filtering/MockNetObjectFilter.cpp
--- a/Engine/Plugins/Runtime/ReplicationSystemTestPlugin/Source/Private/Tests/ReplicationSystem/Filtering/MockNetObjectFilter.cpp
+++ b/Engine/Plugins/Runtime/ReplicationSystemTestPlugin/Source/Private/Tests/ReplicationSystem/Filtering/MockNetObjectFilter.cpp
@@ -21,10 +21,19 @@ void UMockNetObjectFilter::OnInit(const FNetObjectFilterInitParams& Params)
 
 	CallStatus.SuccessfulCallCounts.Init += Cast<UMockNetObjectFilterConfig>(Params.Config) != nullptr;
 
-	AddedObjectIndices.Init(Params.MaxObjectCount);
+	AddedObjectIndices.Init(Params.CurrentMaxInternalIndex);
 	AddedConnectionIndices.Init(Params.MaxConnectionCount + 1);
 }
 
+void UMockNetObjectFilter::OnDeinit()
+{
+}
+
+void UMockNetObjectFilter::OnMaxInternalNetRefIndexIncreased(uint32 NewMaxInternalIndex)
+{
+	AddedObjectIndices.SetNumBits(NewMaxInternalIndex);
+}
+
 void UMockNetObjectFilter::AddConnection(uint32 ConnectionId)
 {
 	++CallStatus.CallCounts.AddConnection;
diff --git a/Engine/Plugins/Runtime/ReplicationSystemTestPlugin/Source/Private/Tests/ReplicationSystem/Filtering/MockNetObjectFilter.h b/Engine/Plugins/Runtime/ReplicationSystemTestPlugin/Source/Private/Tests/ReplicationSystem/Filtering/MockNetObjectFilter.h
--- a/Engine/Plugins/Runtime/ReplicationSystemTestPlugin/Source/Private/Tests/ReplicationSystem/Filtering/MockNetObjectFilter.h
+++ b/Engine/Plugins/Runtime/ReplicationSystemTestPlugin/Source/Private/Tests/ReplicationSystem/Filtering/MockNetObjectFilter.h
@@ -67,6 +67,8 @@ public:
 
 protected:
 	virtual void OnInit(const FNetObjectFilterInitParams&) override;
+	virtual void OnDeinit() override;
+	virtual void OnMaxInternalNetRefIndexIncreased(uint32 NewMaxInternalIndex) override;
 	virtual void AddConnection(uint32 ConnectionId) override;
 	virtual void RemoveConnection(uint32 ConnectionId) override;
 	virtual bool AddObject(uint32 ObjectIndex, FNetObjectFilterAddObjectParams&) override;
diff --git a/Engine/Plugins/Runtime/ReplicationSystemTestPlugin/Source/Private/Tests/ReplicationSystem/Prioritization/MockNetObjectPrioritizer.cpp b/Engine/Plugins/Runtime/ReplicationSystemTestPlugin/Source/Private/Tests/ReplicationSystem/Prioritization/MockNetObjectPrioritizer.cpp
--- a/Engine/Plugins/Runtime/ReplicationSystemTestPlugin/Source/Private/Tests/ReplicationSystem/Prioritization/MockNetObjectPrioritizer.cpp
+++ b/Engine/Plugins/Runtime/ReplicationSystemTestPlugin/Source/Private/Tests/ReplicationSystem/Prioritization/MockNetObjectPrioritizer.cpp
@@ -15,7 +15,12 @@ void UMockNetObjectPrioritizer::Init(FNetObjectPrioritizerInitParams& Params)
 
 	CallStatus.SuccessfulCallCounts.Init += Cast<UMockNetObjectPrioritizerConfig>(Params.Config) != nullptr;
 
-	AddedIndices.Init(false, Params.MaxObjectCount);
+	AddedIndices.Init(Params.CurrentMaxInternalIndex);
+}
+
+void UMockNetObjectPrioritizer::OnMaxInternalNetRefIndexIncreased(uint32 NewMaxInternalIndex)
+{
+	AddedIndices.SetNumBits(NewMaxInternalIndex);
 }
 
 bool UMockNetObjectPrioritizer::AddObject(uint32 ObjectIndex, FNetObjectPrioritizerAddObjectParams& Params)
@@ -39,7 +44,7 @@ bool UMockNetObjectPrioritizer::AddObject(uint32 ObjectIndex, FNetObjectPrioriti
 		Params.OutInfo.Data[2] = uint16(ObjectIndex);
 		Params.OutInfo.Data[3] = uint16(ObjectIndex);
 
-		AddedIndices[ObjectIndex] = true;
+		AddedIndices.SetBit(ObjectIndex);
 		++AddedCount;
 
 		// Check if object has the NetTest_Priority RepTag which the priority should be read from.
@@ -60,7 +65,7 @@ void UMockNetObjectPrioritizer::RemoveObject(uint32 ObjectIndex, const FNetObjec
 	++CallStatus.CallCounts.RemoveObject;
 
 	// If this object wasn't added to us it shouldn't be removed either.
-	if (AddedIndices[ObjectIndex])
+	if (AddedIndices.IsBitSet(ObjectIndex))
 	{
 		bool bIsProperCall = true;
 
@@ -72,7 +77,7 @@ void UMockNetObjectPrioritizer::RemoveObject(uint32 ObjectIndex, const FNetObjec
 
 		CallStatus.SuccessfulCallCounts.RemoveObject += bIsProperCall;
 
-		AddedIndices[ObjectIndex] = false;
+		AddedIndices.ClearBit(ObjectIndex);
 		--AddedCount;
 
 		ObjectToPriority.Remove(ObjectIndex);
@@ -113,7 +118,7 @@ void UMockNetObjectPrioritizer::Prioritize(FNetObjectPrioritizationParams& Param
 	uint32 LastObjectIndex = 0;
 	for (uint32 ObjectIndex : MakeArrayView(Params.ObjectIndices, Params.ObjectCount))
 	{
-		bIsProperCall &= (AddedIndices[ObjectIndex] == true);
+		bIsProperCall &= (AddedIndices.IsBitSet(ObjectIndex));
 		bIsProperCall &= (ObjectIndex > LastObjectIndex);
 
 		LastObjectIndex = ObjectIndex;
@@ -134,7 +139,7 @@ void UMockNetObjectPrioritizer::Prioritize(FNetObjectPrioritizationParams& Param
 
 float UMockNetObjectPrioritizer::GetPriority(UE::Net::Private::FInternalNetRefIndex ObjectIndex) const
 {
-	check(AddedIndices[ObjectIndex]);
+	check(AddedIndices.IsBitSet(ObjectIndex));
 	if (const float* Priority = ObjectToPriority.Find(ObjectIndex))
 	{
 		return *Priority;
diff --git a/Engine/Plugins/Runtime/ReplicationSystemTestPlugin/Source/Private/Tests/ReplicationSystem/Prioritization/MockNetObjectPrioritizer.h b/Engine/Plugins/Runtime/ReplicationSystemTestPlugin/Source/Private/Tests/ReplicationSystem/Prioritization/MockNetObjectPrioritizer.h
--- a/Engine/Plugins/Runtime/ReplicationSystemTestPlugin/Source/Private/Tests/ReplicationSystem/Prioritization/MockNetObjectPrioritizer.h
+++ b/Engine/Plugins/Runtime/ReplicationSystemTestPlugin/Source/Private/Tests/ReplicationSystem/Prioritization/MockNetObjectPrioritizer.h
@@ -3,8 +3,8 @@
 #pragma once
 
 #include "Iris/ReplicationSystem/Prioritization/NetObjectPrioritizer.h"
+#include "Net/Core/NetBitArray.h"
 #include "Iris/ReplicationSystem/RepTag.h"
-#include "Containers/BitArray.h"
 #include "MockNetObjectPrioritizer.generated.h"
 
 namespace UE::Net::Private
@@ -60,6 +60,8 @@ public:
 
 protected:
 	virtual void Init(FNetObjectPrioritizerInitParams& Params) override;
+	virtual void Deinit() override {}
+	virtual void OnMaxInternalNetRefIndexIncreased(uint32 NewMaxInternalIndex) override;
 	virtual bool AddObject(uint32 ObjectIndex, FNetObjectPrioritizerAddObjectParams& Params) override;
 	virtual void RemoveObject(uint32 ObjectIndex, const FNetObjectPrioritizationInfo& Info) override;
 	virtual void UpdateObjects(FNetObjectPrioritizerUpdateParams&) override;
@@ -70,7 +72,7 @@ private:
 
 	FFunctionCallStatus CallStatus;
 	FFunctionCallSetup CallSetup;
-	TBitArray<> AddedIndices;
+	UE::Net::FNetBitArray AddedIndices;
 	SIZE_T AddedCount;
 	TMap<uint32, UPTRINT> ObjectToPriorityOffset;
 	TMap<uint32, float> ObjectToPriority;
diff --git a/Engine/Source/Runtime/Engine/Classes/Engine/NetDriver.h b/Engine/Source/Runtime/Engine/Classes/Engine/NetDriver.h
--- a/Engine/Source/Runtime/Engine/Classes/Engine/NetDriver.h
+++ b/Engine/Source/Runtime/Engine/Classes/Engine/NetDriver.h
@@ -632,25 +632,25 @@ struct FNetDriverReplicationSystemConfig
 {
 	GENERATED_USTRUCT_BODY()
 
-	/** Override the max object count when running as a client. If 0 use the default system value. */
+	/** Override the max object count. If 0 use the default system value. */
 	UPROPERTY()
-	uint32 MaxReplicatedObjectClientCount = 0;
+	uint32 MaxReplicatedObjectCount = 0;
 
-	/** Override the max object count when running as a server. If 0 use the default system value. */
+	/** Override the preallocated size of net object lists. If 0 use the default value. */
 	UPROPERTY()
-	uint32 MaxReplicatedObjectServerCount = 0;
+	uint32 InitialNetObjectListCount = 0;
 
-	/** Override the number of pre-allocated objects when running as a client. */
+	/** Override the amount to grow every net object list by when they hit the preallocated count. If 0 use the default system value. */
 	UPROPERTY()
-	uint32 PreAllocatedReplicatedObjectClientCount = 0;
+	uint32 NetObjectListGrowCount = 0;
 
-	/** Override the number of pre-allocated objects when running as a server. */
+	/** Override the number of pre-allocated memory buffers that can hold up to the specified number of objects before they have to grow. */
 	UPROPERTY()
-	uint32 PreAllocatedReplicatedObjectServerCount = 0;
+	uint32 PreAllocatedMemoryBuffersObjectCount = 0;
 
-	/** Override the number of pre-allocated objects in FReplicationWriter on the client. */
+	/** Override the number of pre-allocated objects in FReplicationWriter. */
 	UPROPERTY()
-	uint32 MaxReplicatedWriterObjectClientCount = 0;
+	uint32 MaxReplicationWriterObjectCount = 0;
 	
 	/** Override the max compressed object count. If 0 use the default system value. */
 	UPROPERTY()
@@ -823,9 +823,13 @@ public:
 	UPROPERTY(Config)
 	FString ReplicationBridgeClassName;
 	
-	/** Can be used to configure settings for the ReplicationSystem */
+	/** Can be used to configure Server settings for the ReplicationSystem */
 	UPROPERTY(Config)
-	FNetDriverReplicationSystemConfig ReplicationSystemConfig;
+	FNetDriverReplicationSystemConfig ReplicationSystemConfigServer;
+
+	/** Can be used to configure Client settings for the ReplicationSystem */
+	UPROPERTY(Config)
+	FNetDriverReplicationSystemConfig ReplicationSystemConfigClient;
 
 	/** @todo document */
 	UPROPERTY(Config)
diff --git a/Engine/Source/Runtime/Engine/Private/NetDriver.cpp b/Engine/Source/Runtime/Engine/Private/NetDriver.cpp
--- a/Engine/Source/Runtime/Engine/Private/NetDriver.cpp
+++ b/Engine/Source/Runtime/Engine/Private/NetDriver.cpp
@@ -249,32 +249,31 @@ namespace UE::Net::Private
 	static bool bIgnoreStaticActorDestruction = false;
 
 
-	void ApplyReplicationSystemConfig(const FNetDriverReplicationSystemConfig& ReplicationSystemConfig, UReplicationSystem::FReplicationSystemParams& OutParams, bool bIsServer)
+	void ApplyReplicationSystemConfig(const FNetDriverReplicationSystemConfig& ReplicationSystemConfig, UReplicationSystem::FReplicationSystemParams& OutParams)
 	{
-		if (bIsServer)
+		if (ReplicationSystemConfig.MaxReplicatedObjectCount != 0)
 		{
-			if (ReplicationSystemConfig.MaxReplicatedObjectServerCount != 0)
-			{
-				OutParams.MaxReplicatedObjectCount = ReplicationSystemConfig.MaxReplicatedObjectServerCount;
-			}
-
-			OutParams.PreAllocatedReplicatedObjectCount = ReplicationSystemConfig.PreAllocatedReplicatedObjectServerCount;
-
-			OutParams.MaxReplicatedWriterObjectCount = OutParams.MaxReplicatedObjectCount;
+			OutParams.MaxReplicatedObjectCount = ReplicationSystemConfig.MaxReplicatedObjectCount;
 		}
-		else
+
+		if (ReplicationSystemConfig.InitialNetObjectListCount != 0 )
 		{
-			if (ReplicationSystemConfig.MaxReplicatedObjectClientCount != 0)
-			{
-				OutParams.MaxReplicatedObjectCount = ReplicationSystemConfig.MaxReplicatedObjectClientCount;
-			}
+			OutParams.InitialNetObjectListCount = ReplicationSystemConfig.InitialNetObjectListCount;
+		}
 
-			OutParams.PreAllocatedReplicatedObjectCount = ReplicationSystemConfig.PreAllocatedReplicatedObjectClientCount;
+		if (ReplicationSystemConfig.NetObjectListGrowCount != 0)
+		{
+			OutParams.NetObjectListGrowCount = ReplicationSystemConfig.NetObjectListGrowCount;
+		}
 
-			if (ReplicationSystemConfig.MaxReplicatedWriterObjectClientCount != 0)
-			{
-				OutParams.MaxReplicatedWriterObjectCount = ReplicationSystemConfig.MaxReplicatedWriterObjectClientCount;
-			}
+		if (ReplicationSystemConfig.PreAllocatedMemoryBuffersObjectCount != 0)
+		{
+			OutParams.PreAllocatedMemoryBuffersObjectCount = ReplicationSystemConfig.PreAllocatedMemoryBuffersObjectCount;
+		}
+		
+		if (ReplicationSystemConfig.MaxReplicationWriterObjectCount!= 0)
+		{
+			OutParams.MaxReplicationWriterObjectCount = OutParams.MaxReplicationWriterObjectCount;
 		}
 
 		if (ReplicationSystemConfig.MaxDeltaCompressedObjectCount != 0)
@@ -7006,7 +7005,14 @@ void UNetDriver::CreateReplicationSystem(bool bInitAsClient)
 		Params.bAllowObjectReplication = !bInitAsClient;
 		Params.ForwardNetRPCCallDelegate.BindUObject(this, &UNetDriver::ForwardRemoteFunction);
 
-		UE::Net::Private::ApplyReplicationSystemConfig(ReplicationSystemConfig, Params, !bInitAsClient);
+		if (bInitAsClient)
+		{
+			UE::Net::Private::ApplyReplicationSystemConfig(ReplicationSystemConfigClient, Params);
+		}
+		else
+		{
+			UE::Net::Private::ApplyReplicationSystemConfig(ReplicationSystemConfigServer, Params);
+		}
 
 		SetReplicationSystem(UE::Net::FReplicationSystemFactory::CreateReplicationSystem(Params));
 	}
diff --git a/Engine/Source/Runtime/Experimental/Iris/Core/Private/Iris/ReplicationState/ReplicationStateStorage.cpp b/Engine/Source/Runtime/Experimental/Iris/Core/Private/Iris/ReplicationState/ReplicationStateStorage.cpp
--- a/Engine/Source/Runtime/Experimental/Iris/Core/Private/Iris/ReplicationState/ReplicationStateStorage.cpp
+++ b/Engine/Source/Runtime/Experimental/Iris/Core/Private/Iris/ReplicationState/ReplicationStateStorage.cpp
@@ -20,22 +20,22 @@ FReplicationStateStorage::FReplicationStateStorage()
 
 FReplicationStateStorage::~FReplicationStateStorage()
 {
-	Deinit();
 }
 
 void FReplicationStateStorage::Init(FReplicationStateStorageInitParams& InitParams)
 {
-	const uint32 MaxObjectInfoCount = 1U + FPlatformMath::Min(uint32(std::numeric_limits<ObjectInfoIndexType>::max()), FPlatformMath::Min(InitParams.MaxObjectCount, InitParams.MaxDeltaCompressedObjectCount));
+	const uint32 MaxDeltaCompressedObjectInfoCount = 1U + FPlatformMath::Min(uint32(std::numeric_limits<ObjectInfoIndexType>::max()), FPlatformMath::Min(InitParams.MaxObjectCount, InitParams.MaxDeltaCompressedObjectCount));
 
 	// Make sure the MaxObjectInfoCount calculation can't overflow.
-	static_assert(std::numeric_limits<decltype(MaxObjectInfoCount)>::max() > std::numeric_limits<ObjectInfoIndexType>::max(), "");
+	static_assert(std::numeric_limits<decltype(MaxDeltaCompressedObjectInfoCount)>::max() > std::numeric_limits<ObjectInfoIndexType>::max(), "");
 
 	NetRefHandleManager = InitParams.NetRefHandleManager;
 
-	UsedPerObjectInfos.Init(MaxObjectInfoCount);
+	UsedPerObjectInfos.Init(MaxDeltaCompressedObjectInfoCount);
 	UsedPerObjectInfos.SetBit(InvalidObjectInfoIndex);
-	ObjectIndexToObjectInfoIndex.SetNumZeroed(InitParams.MaxObjectCount);
 	
+	ObjectIndexToObjectInfoIndex.SetNumZeroed(InitParams.MaxInternalNetRefIndex);
+
 	InternalSerializationContext = Private::FInternalNetSerializationContext(InitParams.ReplicationSystem);
 	SerializationContext.SetInternalContext(&InternalSerializationContext);
 }
@@ -47,6 +47,11 @@ void FReplicationStateStorage::Deinit()
 #endif
 }
 
+void FReplicationStateStorage::OnMaxInternalNetRefIndexIncreased(UE::Net::Private::FInternalNetRefIndex NewMaxInternalIndex)
+{
+	ObjectIndexToObjectInfoIndex.SetNumZeroed(NewMaxInternalIndex);
+}
+
 const uint8* FReplicationStateStorage::GetState(uint32 ObjectIndex, EReplicationStateType StateType) const
 {
 	if (StateType != EReplicationStateType::CurrentSendState && StateType != EReplicationStateType::CurrentRecvState)
@@ -164,7 +169,7 @@ uint8* FReplicationStateStorage::AllocBaseline(uint32 ObjectIndex, EReplicationS
 void FReplicationStateStorage::FreeBaseline(uint32 ObjectIndex, uint8* Storage)
 {
 	FPerObjectInfo* ObjectInfo = GetPerObjectInfoForObject(ObjectIndex);
-	if (!ensureAlwaysMsgf(ObjectInfo != nullptr, TEXT("Trying to free baseline storage pointer %p for object ( InternalIndex: %u ) with no PerObjectInfo."), Storage, ObjectIndex))
+	if (!ensureMsgf(ObjectInfo != nullptr, TEXT("Trying to free baseline storage pointer %p for object ( InternalIndex: %u ) with no PerObjectInfo."), Storage, ObjectIndex))
 	{
 		return;
 	}
@@ -243,7 +248,7 @@ FReplicationStateStorage::FBaselineReservation FReplicationStateStorage::Reserve
 void FReplicationStateStorage::CancelBaselineReservation(uint32 ObjectIndex, uint8* Storage)
 {
 	FPerObjectInfo* ObjectInfo = GetPerObjectInfoForObject(ObjectIndex);
-	if (!ensureAlwaysMsgf(ObjectInfo != nullptr, TEXT("Trying to cancel a baseline reservation with pointer %p for object ( InternalIndex: %u ) with no PerObjectInfo."), Storage, ObjectIndex))
+	if (!ensureMsgf(ObjectInfo != nullptr, TEXT("Trying to cancel a baseline reservation with pointer %p for object ( InternalIndex: %u ) with no PerObjectInfo."), Storage, ObjectIndex))
 	{
 		return;
 	}
@@ -272,7 +277,7 @@ void FReplicationStateStorage::CommitBaselineReservation(uint32 ObjectIndex, uin
 #endif
 
 	FPerObjectInfo* ObjectInfo = GetPerObjectInfoForObject(ObjectIndex);
-	if (!ensureAlwaysMsgf(ObjectInfo != nullptr, TEXT("Trying to commit baseline storage pointer %p for object ( InternalIndex: %u ) with no PerObjectInfo."), Storage, ObjectIndex))
+	if (!ensureMsgf(ObjectInfo != nullptr, TEXT("Trying to commit baseline storage pointer %p for object ( InternalIndex: %u ) with no PerObjectInfo."), Storage, ObjectIndex))
 	{
 		return;
 	}
diff --git a/Engine/Source/Runtime/Experimental/Iris/Core/Private/Iris/ReplicationState/ReplicationStateStorage.h b/Engine/Source/Runtime/Experimental/Iris/Core/Private/Iris/ReplicationState/ReplicationStateStorage.h
--- a/Engine/Source/Runtime/Experimental/Iris/Core/Private/Iris/ReplicationState/ReplicationStateStorage.h
+++ b/Engine/Source/Runtime/Experimental/Iris/Core/Private/Iris/ReplicationState/ReplicationStateStorage.h
@@ -17,6 +17,7 @@ namespace UE::Net
 	struct FReplicationProtocol;
 	namespace Private
 	{
+		typedef uint32 FInternalNetRefIndex;
 		class FNetRefHandleManager;
 	};
 }
@@ -29,6 +30,7 @@ struct FReplicationStateStorageInitParams
 	UReplicationSystem* ReplicationSystem = nullptr;
 	const Private::FNetRefHandleManager* NetRefHandleManager = nullptr;
 	uint32 MaxObjectCount = 0;
+	uint32 MaxInternalNetRefIndex = 0;
 	uint32 MaxConnectionCount = 0;
 	uint32 MaxDeltaCompressedObjectCount = 0;
 };
@@ -64,6 +66,7 @@ public:
 	~FReplicationStateStorage();
 
 	void Init(FReplicationStateStorageInitParams& InitParams);
+	void Deinit();
 
 	/*
 	 * Returns a buffer to a state of the requested type. Only supports CurrentSendState and CurrentRecvState. May return null
@@ -96,6 +99,9 @@ public:
 	 */
 	void CommitBaselineReservation(uint32 ObjectIndex, uint8* Storage, EReplicationStateType Base);
 
+	/** Called when the maximum InternalNetRefIndex increased and we need to realloc our lists */
+	void OnMaxInternalNetRefIndexIncreased(UE::Net::Private::FInternalNetRefIndex NewMaxInternalIndex);
+
 private:
 	using ObjectInfoIndexType = uint16;
 
@@ -121,8 +127,6 @@ private:
 		uint16 AllocationCount = 0;
 	};
 
-	void Deinit();
-
 	FPerObjectInfo* GetOrCreatePerObjectInfoForObject(uint32 ObjectIndex);
 	FPerObjectInfo* GetPerObjectInfoForObject(uint32 ObjectIndex);
 	const FPerObjectInfo* GetPerObjectInfoForObject(uint32 ObjectIndex) const;
diff --git a/Engine/Source/Runtime/Experimental/Iris/Core/Private/Iris/ReplicationSystem/Conditionals/ReplicationConditionals.cpp b/Engine/Source/Runtime/Experimental/Iris/Core/Private/Iris/ReplicationSystem/Conditionals/ReplicationConditionals.cpp
--- a/Engine/Source/Runtime/Experimental/Iris/Core/Private/Iris/ReplicationSystem/Conditionals/ReplicationConditionals.cpp
+++ b/Engine/Source/Runtime/Experimental/Iris/Core/Private/Iris/ReplicationSystem/Conditionals/ReplicationConditionals.cpp
@@ -42,13 +42,29 @@ void FReplicationConditionals::Init(FReplicationConditionalsInitParams& Params)
 	ReplicationConnections = Params.ReplicationConnections;
 	BaselineInvalidationTracker = Params.BaselineInvalidationTracker;
 	NetObjectGroups = Params.NetObjectGroups;
-	MaxObjectCount = Params.MaxObjectCount;
+	MaxInternalNetRefIndex = Params.MaxInternalNetRefIndex;
 	MaxConnectionCount = Params.MaxConnectionCount;
 
-	PerObjectInfos.SetNumZeroed(MaxObjectCount);
+	PerObjectInfos.SetNumZeroed(MaxInternalNetRefIndex);
 	ConnectionInfos.SetNum(MaxConnectionCount + 1U);
 }
 
+void FReplicationConditionals::OnMaxInternalNetRefIndexIncreased(FInternalNetRefIndex NewMaxInternalIndex)
+{
+	MaxInternalNetRefIndex = NewMaxInternalIndex;
+
+	PerObjectInfos.SetNumZeroed(NewMaxInternalIndex);
+
+	for (FPerConnectionInfo& ConnectionInfo : ConnectionInfos)
+	{
+		// Resize the netobject list for valid connections
+		if (!ConnectionInfo.ObjectConditionals.IsEmpty())
+		{
+			ConnectionInfo.ObjectConditionals.SetNumZeroed(NewMaxInternalIndex);
+		}
+	}
+}
+
 bool FReplicationConditionals::SetConditionConnectionFilter(FInternalNetRefIndex ObjectIndex, EReplicationCondition Condition, uint32 ConnectionId, bool bEnable)
 {
 	if (ConnectionId >= MaxConnectionCount)
@@ -81,7 +97,7 @@ void FReplicationConditionals::AddConnection(uint32 ConnectionId)
 {
 	// Init connection info
 	FPerConnectionInfo& ConnectionInfo = ConnectionInfos[ConnectionId];
-	ConnectionInfo.ObjectConditionals.SetNumZeroed(MaxObjectCount);
+	ConnectionInfo.ObjectConditionals.SetNumZeroed(MaxInternalNetRefIndex);
 }
 
 void FReplicationConditionals::RemoveConnection(uint32 ConnectionId)
@@ -688,7 +704,7 @@ void FReplicationConditionals::UpdateObjectsInScope()
 	const FNetBitArrayView ObjectsInScope = NetRefHandleManager->GetCurrentFrameScopableInternalIndices();
 	const FNetBitArrayView PrevObjectsInScope = NetRefHandleManager->GetPrevFrameScopableInternalIndices();
 
-	const uint32 WordCountForModifiedWords = Align(FPlatformMath::Max(MaxObjectCount, 1U), 32U)/32U;
+	const uint32 WordCountForModifiedWords = Align(FPlatformMath::Max(MaxInternalNetRefIndex, 1U), 32U)/32U;
 	TArray<uint32> ModifiedWords;
 	ModifiedWords.SetNumUninitialized(WordCountForModifiedWords);
 	uint32* ModifiedWordsStorage = ModifiedWords.GetData();
diff --git a/Engine/Source/Runtime/Experimental/Iris/Core/Private/Iris/ReplicationSystem/Conditionals/ReplicationConditionals.h b/Engine/Source/Runtime/Experimental/Iris/Core/Private/Iris/ReplicationSystem/Conditionals/ReplicationConditionals.h
--- a/Engine/Source/Runtime/Experimental/Iris/Core/Private/Iris/ReplicationSystem/Conditionals/ReplicationConditionals.h
+++ b/Engine/Source/Runtime/Experimental/Iris/Core/Private/Iris/ReplicationSystem/Conditionals/ReplicationConditionals.h
@@ -33,7 +33,7 @@ struct FReplicationConditionalsInitParams
 	const FReplicationConnections* ReplicationConnections = nullptr;
 	const FNetObjectGroups* NetObjectGroups = nullptr;
 	FDeltaCompressionBaselineInvalidationTracker* BaselineInvalidationTracker = nullptr;
-	uint32 MaxObjectCount = 0;
+	FInternalNetRefIndex MaxInternalNetRefIndex = 0;
 	uint32 MaxConnectionCount = 0;
 };
 
@@ -52,6 +52,9 @@ public:
 
 	void Init(FReplicationConditionalsInitParams& Params);
 
+	/** Called when the maximum InternalNetRefIndex increased and we need to realloc our lists */
+	void OnMaxInternalNetRefIndexIncreased(FInternalNetRefIndex NewMaxInternalIndex);
+
 	void AddConnection(uint32 ConnectionId);
 	void RemoveConnection(uint32 ConnectionId);
 
@@ -140,7 +143,7 @@ private:
 	TArray<FPerConnectionInfo> ConnectionInfos;
 	TMap<FInternalNetRefIndex, FObjectDynamicConditions> DynamicConditions;
 
-	uint32 MaxObjectCount = 0;
+	FInternalNetRefIndex MaxInternalNetRefIndex = 0;
 	uint32 MaxConnectionCount = 0;
 	uint16 CachedRemoteRoleRepIndex = InvalidRepIndex;
 };
diff --git a/Engine/Source/Runtime/Experimental/Iris/Core/Private/Iris/ReplicationSystem/DeltaCompression/DeltaCompressionBaselineInvalidationTracker.cpp b/Engine/Source/Runtime/Experimental/Iris/Core/Private/Iris/ReplicationSystem/DeltaCompression/DeltaCompressionBaselineInvalidationTracker.cpp
--- a/Engine/Source/Runtime/Experimental/Iris/Core/Private/Iris/ReplicationSystem/DeltaCompression/DeltaCompressionBaselineInvalidationTracker.cpp
+++ b/Engine/Source/Runtime/Experimental/Iris/Core/Private/Iris/ReplicationSystem/DeltaCompression/DeltaCompressionBaselineInvalidationTracker.cpp
@@ -18,7 +18,12 @@ FDeltaCompressionBaselineInvalidationTracker::~FDeltaCompressionBaselineInvalida
 void FDeltaCompressionBaselineInvalidationTracker::Init(FDeltaCompressionBaselineInvalidationTrackerInitParams& InitParams)
 {
 	BaselineManager = InitParams.BaselineManager;
-	InvalidatedObjects.Init(InitParams.MaxObjectCount);
+	InvalidatedObjects.Init(InitParams.MaxInternalNetRefIndex);
+}
+
+void FDeltaCompressionBaselineInvalidationTracker::OnMaxInternalNetRefIndexIncreased(FInternalNetRefIndex NewMaxInternalIndex)
+{
+	InvalidatedObjects.SetNumBits(NewMaxInternalIndex);
 }
 
 void FDeltaCompressionBaselineInvalidationTracker::InvalidateBaselines(FInternalNetRefIndex ObjectIndex, uint32 ConnId)
diff --git a/Engine/Source/Runtime/Experimental/Iris/Core/Private/Iris/ReplicationSystem/DeltaCompression/DeltaCompressionBaselineInvalidationTracker.h b/Engine/Source/Runtime/Experimental/Iris/Core/Private/Iris/ReplicationSystem/DeltaCompression/DeltaCompressionBaselineInvalidationTracker.h
--- a/Engine/Source/Runtime/Experimental/Iris/Core/Private/Iris/ReplicationSystem/DeltaCompression/DeltaCompressionBaselineInvalidationTracker.h
+++ b/Engine/Source/Runtime/Experimental/Iris/Core/Private/Iris/ReplicationSystem/DeltaCompression/DeltaCompressionBaselineInvalidationTracker.h
@@ -21,7 +21,7 @@ namespace UE::Net::Private
 struct FDeltaCompressionBaselineInvalidationTrackerInitParams
 {
 	const FDeltaCompressionBaselineManager* BaselineManager = nullptr;
-	uint32 MaxObjectCount = 0;
+	FInternalNetRefIndex MaxInternalNetRefIndex = 0;
 };
 
 class FDeltaCompressionBaselineInvalidationTracker
@@ -52,6 +52,9 @@ public:
 	void PreSendUpdate();
 	void PostSendUpdate();
 
+	/** Called when the maximum InternalNetRefIndex increased and we need to realloc our lists */
+	void OnMaxInternalNetRefIndexIncreased(FInternalNetRefIndex NewMaxInternalIndex);
+
 private:
 	enum : unsigned
 	{
diff --git a/Engine/Source/Runtime/Experimental/Iris/Core/Private/Iris/ReplicationSystem/DeltaCompression/DeltaCompressionBaselineManager.cpp b/Engine/Source/Runtime/Experimental/Iris/Core/Private/Iris/ReplicationSystem/DeltaCompression/DeltaCompressionBaselineManager.cpp
--- a/Engine/Source/Runtime/Experimental/Iris/Core/Private/Iris/ReplicationSystem/DeltaCompression/DeltaCompressionBaselineManager.cpp
+++ b/Engine/Source/Runtime/Experimental/Iris/Core/Private/Iris/ReplicationSystem/DeltaCompression/DeltaCompressionBaselineManager.cpp
@@ -69,16 +69,16 @@ void FDeltaCompressionBaselineManager::Init(FDeltaCompressionBaselineManagerInit
 	BaselineInvalidationTracker = InitParams.BaselineInvalidationTracker;
 	ReplicationSystem = InitParams.ReplicationSystem;
 	MaxConnectionCount = InitParams.Connections->GetMaxConnectionCount();
-	MaxDeltaCompressedObjectCount = FPlatformMath::Min(InitParams.MaxObjectCount, InitParams.MaxDeltaCompressedObjectCount);
+	MaxDeltaCompressedObjectCount = FPlatformMath::Min(InitParams.MaxNetObjectCount, InitParams.MaxDeltaCompressedObjectCount);
 
-	DeltaCompressionEnabledObjects.Init(InitParams.MaxObjectCount);
+	DeltaCompressionEnabledObjects.Init(InitParams.MaxInternalNetRefIndex);
 
 	// PerObjectInfo initialization
 	{
 		check(MaxDeltaCompressedObjectCount < std::numeric_limits<ObjectInfoIndexType>::max());
 		UsedPerObjectInfos.Init(MaxDeltaCompressedObjectCount + 1U);
 		UsedPerObjectInfos.SetBit(InvalidObjectInfoIndex);
-		ObjectIndexToObjectInfoIndex.SetNumZeroed(InitParams.MaxObjectCount);
+		ObjectIndexToObjectInfoIndex.SetNumZeroed(InitParams.MaxInternalNetRefIndex);
 
 		const uint32 BytesPerConnection = sizeof(FObjectBaselineInfo);
 		// The PerObjectInfo already contains the connection specific data for one connection,
@@ -104,6 +104,12 @@ void FDeltaCompressionBaselineManager::Deinit()
 	BaselineStorage.Deinit();
 }
 
+void FDeltaCompressionBaselineManager::OnMaxInternalNetRefIndexIncreased(FInternalNetRefIndex NewMaxInternalIndex)
+{
+	DeltaCompressionEnabledObjects.SetNumBits(NewMaxInternalIndex);
+	ObjectIndexToObjectInfoIndex.SetNumZeroed(NewMaxInternalIndex);
+}
+
 void FDeltaCompressionBaselineManager::PreSendUpdate(FDeltaCompressionBaselineManagerPreSendUpdateParams& UpdateParams)
 {
 	IRIS_PROFILER_SCOPE(DeltaCompressionBaselineManager_PreSendUpdate);
diff --git a/Engine/Source/Runtime/Experimental/Iris/Core/Private/Iris/ReplicationSystem/DeltaCompression/DeltaCompressionBaselineManager.h b/Engine/Source/Runtime/Experimental/Iris/Core/Private/Iris/ReplicationSystem/DeltaCompression/DeltaCompressionBaselineManager.h
--- a/Engine/Source/Runtime/Experimental/Iris/Core/Private/Iris/ReplicationSystem/DeltaCompression/DeltaCompressionBaselineManager.h
+++ b/Engine/Source/Runtime/Experimental/Iris/Core/Private/Iris/ReplicationSystem/DeltaCompression/DeltaCompressionBaselineManager.h
@@ -36,7 +36,8 @@ struct FDeltaCompressionBaselineManagerInitParams
 	const FNetRefHandleManager* NetRefHandleManager = nullptr;
 	FReplicationStateStorage* ReplicationStateStorage = nullptr;
 	UReplicationSystem* ReplicationSystem = nullptr;
-	uint32 MaxObjectCount = 0;
+	FInternalNetRefIndex MaxNetObjectCount = 0;
+	FInternalNetRefIndex MaxInternalNetRefIndex = 0;
 	uint32 MaxDeltaCompressedObjectCount = 0;
 };
 
@@ -62,6 +63,9 @@ public:
 	void Init(FDeltaCompressionBaselineManagerInitParams& InitParams);
 	void Deinit();
 
+	/** Called when the maximum InternalNetRefIndex increased and we need to realloc our lists */
+	void OnMaxInternalNetRefIndexIncreased(UE::Net::Private::FInternalNetRefIndex NewMaxInternalIndex);
+
 	void PreSendUpdate(FDeltaCompressionBaselineManagerPreSendUpdateParams& UpdateParams);
 	void PostSendUpdate(FDeltaCompressionBaselineManagerPostSendUpdateParams& UpdateParams);
 
@@ -165,6 +169,8 @@ private:
 		uint32 CreatedBaselineCount = 0;
 	};
 
+private:
+
 	void UpdateScope();
 	void UpdateDirtyStateMasks(const FChangeMaskCache* ChangeMaskCache);
 
diff --git a/Engine/Source/Runtime/Experimental/Iris/Core/Private/Iris/ReplicationSystem/DeltaCompression/DeltaCompressionBaselineStorage.cpp b/Engine/Source/Runtime/Experimental/Iris/Core/Private/Iris/ReplicationSystem/DeltaCompression/DeltaCompressionBaselineStorage.cpp
--- a/Engine/Source/Runtime/Experimental/Iris/Core/Private/Iris/ReplicationSystem/DeltaCompression/DeltaCompressionBaselineStorage.cpp
+++ b/Engine/Source/Runtime/Experimental/Iris/Core/Private/Iris/ReplicationSystem/DeltaCompression/DeltaCompressionBaselineStorage.cpp
@@ -22,6 +22,7 @@ FDeltaCompressionBaselineStorage::FDeltaCompressionBaselineStorage()
 
 FDeltaCompressionBaselineStorage::~FDeltaCompressionBaselineStorage()
 {
+	checkf(UsedBaselineStateInfos.IsAnyBitSet() == false, TEXT("Leak in FDeltaCompressionBaselineStorage"));
 }
 
 void FDeltaCompressionBaselineStorage::Init(FDeltaCompressionBaselineStorageInitParams& InitParams)
@@ -219,6 +220,7 @@ void FDeltaCompressionBaselineStorage::FreeAllBaselineStateInfos()
 
 	UsedBaselineStateInfos.ClearBit(InvalidDeltaCompressionBaselineStateInfoIndex);
 	UsedBaselineStateInfos.ForAllSetBits(DestructStateInfo);
+	UsedBaselineStateInfos.Reset();
 }
 
 }
diff --git a/Engine/Source/Runtime/Experimental/Iris/Core/Private/Iris/ReplicationSystem/DirtyNetObjectTracker.cpp b/Engine/Source/Runtime/Experimental/Iris/Core/Private/Iris/ReplicationSystem/DirtyNetObjectTracker.cpp
--- a/Engine/Source/Runtime/Experimental/Iris/Core/Private/Iris/ReplicationSystem/DirtyNetObjectTracker.cpp
+++ b/Engine/Source/Runtime/Experimental/Iris/Core/Private/Iris/ReplicationSystem/DirtyNetObjectTracker.cpp
@@ -29,7 +29,6 @@ FDirtyNetObjectTracker::FDirtyNetObjectTracker()
 
 FDirtyNetObjectTracker::~FDirtyNetObjectTracker()
 {
-	Deinit();
 }
 
 void FDirtyNetObjectTracker::Init(const FDirtyNetObjectTrackerInitParams& Params)
@@ -39,13 +38,13 @@ void FDirtyNetObjectTracker::Init(const FDirtyNetObjectTrackerInitParams& Params
 	NetRefHandleManager = Params.NetRefHandleManager;
 	ReplicationSystemId = Params.ReplicationSystemId;
 	
-	NetObjectIdCount = Params.MaxObjectCount;
+	NetObjectIdCount = Params.MaxInternalNetRefIndex;
 
 	GlobalDirtyTrackerPollHandle = FGlobalDirtyNetObjectTracker::CreatePoller();
 
-	AccumulatedDirtyNetObjects.Init(NetObjectIdCount);
-	ForceNetUpdateObjects.Init(NetObjectIdCount);
-	DirtyNetObjects.Init(NetObjectIdCount);
+	SetNetObjectListsSize(Params.MaxInternalNetRefIndex);
+
+	NetRefHandleManager->GetOnMaxInternalNetRefIndexIncreasedDelegate().AddRaw(this, &FDirtyNetObjectTracker::OnMaxInternalNetRefIndexIncreased);
 
 	AllowExternalAccess();
 
@@ -54,10 +53,24 @@ void FDirtyNetObjectTracker::Init(const FDirtyNetObjectTrackerInitParams& Params
 
 void FDirtyNetObjectTracker::Deinit()
 {
+	NetRefHandleManager->GetOnMaxInternalNetRefIndexIncreasedDelegate().RemoveAll(this);
 	GlobalDirtyTrackerPollHandle.Destroy();
 	bShouldResetPolledGlobalDirtyTracker = false;
 }
 
+void FDirtyNetObjectTracker::SetNetObjectListsSize(FInternalNetRefIndex NewMaxInternalIndex)
+{
+	AccumulatedDirtyNetObjects.SetNumBits(NewMaxInternalIndex);
+	ForceNetUpdateObjects.SetNumBits(NewMaxInternalIndex);
+	DirtyNetObjects.SetNumBits(NewMaxInternalIndex);
+}
+
+void FDirtyNetObjectTracker::OnMaxInternalNetRefIndexIncreased(FInternalNetRefIndex NewMaxInternalIndex)
+{
+	SetNetObjectListsSize(NewMaxInternalIndex);
+	NetObjectIdCount = NewMaxInternalIndex;
+}
+
 void FDirtyNetObjectTracker::GrabAndApplyGlobalDirtyObjectList()
 {
 	{
diff --git a/Engine/Source/Runtime/Experimental/Iris/Core/Private/Iris/ReplicationSystem/DirtyNetObjectTracker.h b/Engine/Source/Runtime/Experimental/Iris/Core/Private/Iris/ReplicationSystem/DirtyNetObjectTracker.h
--- a/Engine/Source/Runtime/Experimental/Iris/Core/Private/Iris/ReplicationSystem/DirtyNetObjectTracker.h
+++ b/Engine/Source/Runtime/Experimental/Iris/Core/Private/Iris/ReplicationSystem/DirtyNetObjectTracker.h
@@ -25,7 +25,7 @@ struct FDirtyNetObjectTrackerInitParams
 {
 	const FNetRefHandleManager* NetRefHandleManager = nullptr;
 	uint32 ReplicationSystemId = 0;
-	uint32 MaxObjectCount = 0;
+	uint32 MaxInternalNetRefIndex = 0;
 };
 
 class FDirtyNetObjectTracker
@@ -35,6 +35,7 @@ public:
 	~FDirtyNetObjectTracker();
 
 	void Init(const FDirtyNetObjectTrackerInitParams& Params);
+	void Deinit();
 
 	/** Returns true if this dirty tracker can be used by the replication system */
 	bool IsInit() const { return NetRefHandleManager != nullptr; }
@@ -73,7 +74,9 @@ private:
 	using StorageType = FNetBitArrayView::StorageWordType;
 	static constexpr uint32 StorageTypeBitCount = FNetBitArrayView::WordBitCount;
 
-	void Deinit();
+	void SetNetObjectListsSize(FInternalNetRefIndex NewMaxInternalIndex);
+	void OnMaxInternalNetRefIndexIncreased(FInternalNetRefIndex NewMaxInternalIndex);
+
 	void MarkNetObjectDirty(FInternalNetRefIndex NetObjectIndex);
 	void ForceNetUpdate(FInternalNetRefIndex NetObjectIndex);
 	void GrabAndApplyGlobalDirtyObjectList();
diff --git a/Engine/Source/Runtime/Experimental/Iris/Core/Private/Iris/ReplicationSystem/Filtering/NetObjectFilter.cpp b/Engine/Source/Runtime/Experimental/Iris/Core/Private/Iris/ReplicationSystem/Filtering/NetObjectFilter.cpp
--- a/Engine/Source/Runtime/Experimental/Iris/Core/Private/Iris/ReplicationSystem/Filtering/NetObjectFilter.cpp
+++ b/Engine/Source/Runtime/Experimental/Iris/Core/Private/Iris/ReplicationSystem/Filtering/NetObjectFilter.cpp
@@ -9,7 +9,7 @@ UNetObjectFilter::UNetObjectFilter()
 
 void UNetObjectFilter::Init(const FNetObjectFilterInitParams& Params)
 {
-	FilteredObjects.Init(Params.MaxObjectCount);
+	FilteredObjects.Init(Params.CurrentMaxInternalIndex);
 
 	{
 		UE::Net::Private::FNetObjectFilteringInfoAccessor FilteringInfoAccessor;
@@ -19,6 +19,24 @@ void UNetObjectFilter::Init(const FNetObjectFilterInitParams& Params)
 	OnInit(Params);
 }
 
+void UNetObjectFilter::Deinit()
+{
+	OnDeinit();
+
+	FilteringInfos = TArrayView<FNetObjectFilteringInfo>();
+}
+
+void UNetObjectFilter::MaxInternalNetRefIndexIncreased(UE::Net::Private::FInternalNetRefIndex MaxInternalIndex, TArrayView<FNetObjectFilteringInfo> NewFilterInfoView)
+{
+	FilteredObjects.SetNumBits(MaxInternalIndex);
+	
+	//$IRIS TODO: Move the FilteringInfo somewhere else or pass it via function param only ?
+	//      We shouldn't be holding views on arrays we don't own exactly for this reason...
+	FilteringInfos = NewFilterInfoView;
+
+	OnMaxInternalNetRefIndexIncreased(MaxInternalIndex);
+}
+
 void UNetObjectFilter::AddConnection(uint32 ConnectionId)
 {
 }
@@ -45,7 +63,7 @@ void UNetObjectFilter::PostFilter(FNetObjectPostFilteringParams&)
 
 FNetObjectFilteringInfo* UNetObjectFilter::GetFilteringInfo(uint32 ObjectIndex)
 {
-	// Only allow retreiving infos for objects handled by this instance.
+	// Only allow retrieving infos for objects handled by this instance.
 	if (!IsObjectFiltered(ObjectIndex))
 	{
 		return nullptr;
diff --git a/Engine/Source/Runtime/Experimental/Iris/Core/Private/Iris/ReplicationSystem/Filtering/NetObjectGridFilter.cpp b/Engine/Source/Runtime/Experimental/Iris/Core/Private/Iris/ReplicationSystem/Filtering/NetObjectGridFilter.cpp
--- a/Engine/Source/Runtime/Experimental/Iris/Core/Private/Iris/ReplicationSystem/Filtering/NetObjectGridFilter.cpp
+++ b/Engine/Source/Runtime/Experimental/Iris/Core/Private/Iris/ReplicationSystem/Filtering/NetObjectGridFilter.cpp
@@ -19,7 +19,7 @@ void UNetObjectGridFilter::OnInit(const FNetObjectFilterInitParams& Params)
 	Config = TStrongObjectPtr<UNetObjectGridFilterConfig>(CastChecked<UNetObjectGridFilterConfig>(Params.Config));
 	checkf(Config.IsValid(), TEXT("Need config to operate."));
 
-	AssignedObjectInfoIndices.Init(Params.MaxObjectCount);
+	AssignedObjectInfoIndices.Init(Params.CurrentMaxInternalIndex);
 
 	PerConnectionInfos.SetNum(Params.MaxConnectionCount + 1);
 
@@ -27,6 +27,24 @@ void UNetObjectGridFilter::OnInit(const FNetObjectFilterInitParams& Params)
 	NetCullDistanceOverrides = &Params.ReplicationSystem->GetNetCullDistanceOverrides();
 }
 
+void UNetObjectGridFilter::OnDeinit()
+{
+	Config = nullptr;
+	NetRefHandleManager = nullptr;
+	NetCullDistanceOverrides = nullptr;
+
+	PerConnectionInfos.Empty();
+	ObjectInfos.Empty();
+	AssignedObjectInfoIndices.Empty();
+	
+	Cells.Empty();
+}
+
+void UNetObjectGridFilter::OnMaxInternalNetRefIndexIncreased(uint32 NewMaxInternalIndex)
+{
+	AssignedObjectInfoIndices.SetNumBits(NewMaxInternalIndex);
+}
+
 void UNetObjectGridFilter::AddConnection(uint32 ConnectionId)
 {
 	PerConnectionInfos[ConnectionId] = {};
@@ -562,6 +580,11 @@ void UNetObjectGridWorldLocFilter::OnInit(const FNetObjectFilterInitParams& Para
 	WorldLocations = &Params.ReplicationSystem->GetWorldLocations();
 }
 
+void UNetObjectGridWorldLocFilter::OnDeinit()
+{
+	WorldLocations = nullptr;
+}
+
 void UNetObjectGridWorldLocFilter::PreFilter(FNetObjectPreFilteringParams& Params)
 {
 	Super::PreFilter(Params);
diff --git a/Engine/Source/Runtime/Experimental/Iris/Core/Private/Iris/ReplicationSystem/Filtering/NetObjectGroups.cpp b/Engine/Source/Runtime/Experimental/Iris/Core/Private/Iris/ReplicationSystem/Filtering/NetObjectGroups.cpp
--- a/Engine/Source/Runtime/Experimental/Iris/Core/Private/Iris/ReplicationSystem/Filtering/NetObjectGroups.cpp
+++ b/Engine/Source/Runtime/Experimental/Iris/Core/Private/Iris/ReplicationSystem/Filtering/NetObjectGroups.cpp
@@ -62,8 +62,14 @@ void FNetObjectGroups::Init(const FNetObjectGroupInitParams& Params)
 	// Reserve first as invalid group
 	Groups.Add(FNetObjectGroup());
 
-	GroupMemberships.SetNum(Params.MaxObjectCount);
-	GroupFilteredOutObjects.Init(Params.MaxObjectCount);
+	GroupMemberships.SetNum(Params.MaxInternalNetRefIndex);
+	GroupFilteredOutObjects.Init(Params.MaxInternalNetRefIndex);
+}
+
+void FNetObjectGroups::OnMaxInternalNetRefIndexIncreased(FInternalNetRefIndex NewMaxInternalIndex)
+{
+	GroupMemberships.SetNum(NewMaxInternalIndex);
+	GroupFilteredOutObjects.SetNumBits(NewMaxInternalIndex);
 }
 
 FNetObjectGroupHandle FNetObjectGroups::CreateGroup()
diff --git a/Engine/Source/Runtime/Experimental/Iris/Core/Private/Iris/ReplicationSystem/Filtering/NetObjectGroups.h b/Engine/Source/Runtime/Experimental/Iris/Core/Private/Iris/ReplicationSystem/Filtering/NetObjectGroups.h
--- a/Engine/Source/Runtime/Experimental/Iris/Core/Private/Iris/ReplicationSystem/Filtering/NetObjectGroups.h
+++ b/Engine/Source/Runtime/Experimental/Iris/Core/Private/Iris/ReplicationSystem/Filtering/NetObjectGroups.h
@@ -43,7 +43,7 @@ struct FNetObjectGroup
 struct FNetObjectGroupInitParams
 {
 	FNetRefHandleManager* NetRefHandleManager = nullptr;
-	uint32 MaxObjectCount = 0;
+	uint32 MaxInternalNetRefIndex = 0;
 	uint32 MaxGroupCount = 0;
 };
 
@@ -117,6 +117,9 @@ public:
 		return MakeNetBitArrayView(GroupFilteredOutObjects);
 	}
 
+	/** Called when the maximum InternalNetRefIndex increased and we need to realloc our lists */
+	void OnMaxInternalNetRefIndexIncreased(FInternalNetRefIndex NewMaxInternalIndex);
+
 private:
 	struct FNetObjectGroupMembership
 	{
diff --git a/Engine/Source/Runtime/Experimental/Iris/Core/Private/Iris/ReplicationSystem/Filtering/ReplicationFiltering.cpp b/Engine/Source/Runtime/Experimental/Iris/Core/Private/Iris/ReplicationSystem/Filtering/ReplicationFiltering.cpp
--- a/Engine/Source/Runtime/Experimental/Iris/Core/Private/Iris/ReplicationSystem/Filtering/ReplicationFiltering.cpp
+++ b/Engine/Source/Runtime/Experimental/Iris/Core/Private/Iris/ReplicationSystem/Filtering/ReplicationFiltering.cpp
@@ -197,22 +197,15 @@ void FReplicationFiltering::Init(FReplicationFilteringInitParams& Params)
 	BaselineInvalidationTracker = Params.BaselineInvalidationTracker;
 	Groups = Params.Groups;
 
-	MaxObjectCount = Params.MaxObjectCount;
-	WordCountForObjectBitArrays = Align(FPlatformMath::Max(MaxObjectCount, 1U), sizeof(FNetBitArrayBase::StorageWordType)*8U)/(sizeof(FNetBitArrayBase::StorageWordType)*8U);
+	MaxInternalNetRefIndex = Params.MaxInternalNetRefIndex;
 
 	// Connection specifics
 	ConnectionInfos.SetNum(Params.Connections->GetMaxConnectionCount() + 1U);
 	ValidConnections.Init(ConnectionInfos.Num());
 	NewConnections.Init(ConnectionInfos.Num());
 
-	// Filter specifics
-	ObjectsWithDirtyConnectionFilter.Init(MaxObjectCount);
-	ObjectsWithDirtyOwner.Init(MaxObjectCount);
-	ObjectsWithOwnerFilter.Init(MaxObjectCount);
-	ObjectsWithPerObjectInfo.Init(MaxObjectCount);
-
-	AllConnectionFilteredObjects.Init(MaxObjectCount);
-	DynamicFilterEnabledObjects.Init(MaxObjectCount);
+	// Initialize all InternalNetRefIndex lists
+	SetNetObjectListsSize(MaxInternalNetRefIndex);
 
 	// Group filtering
 	{
@@ -229,29 +222,95 @@ void FReplicationFiltering::Init(FReplicationFilteringInitParams& Params)
 		DirtySubObjectFilterGroups.Init(InMaxGroupCount);
 	}
 
-	ObjectIndexToPerObjectInfoIndex.SetNumZeroed(MaxObjectCount);
-
 	PerObjectInfoStorageCountForConnections = Align(FPlatformMath::Max(uint32(ConnectionInfos.Num()), 1U), 32U)/32U;
 	PerObjectInfoStorageCountPerItem = sizeof(FPerObjectInfo) + PerObjectInfoStorageCountForConnections - 1U;
 
-	// Owning connections
-	ObjectIndexToOwningConnection.SetNumZeroed(MaxObjectCount);
+	InitFilters();
+	InitObjectScopeHysteresis();
+}
+
+void FReplicationFiltering::Deinit()
+{
+	for (FFilterInfo& FilterInfo : DynamicFilterInfos)
+	{
+		FilterInfo.Filter->Deinit();
+	}
+
+	// Clear most buffers by setting size to 0
+	SetNetObjectListsSize(0);
+}
 
-	// Object scope hysteresis
-	ObjectScopeHysteresisFrameCounts.SetNumZeroed(MaxObjectCount);
+void FReplicationFiltering::SetNetObjectListsSize(FInternalNetRefIndex MaxInternalIndex)
+{
+	WordCountForObjectBitArrays = Align(MaxInternalIndex, sizeof(FNetBitArrayBase::StorageWordType) * 8U) / (sizeof(FNetBitArrayBase::StorageWordType) * 8U);
 
-	// Dynamic filters
+	// Increase NetBitArrays
 	{
-		NetObjectFilteringInfos.SetNumUninitialized(Params.MaxObjectCount);
+		ObjectsWithDirtyConnectionFilter.SetNumBits(MaxInternalIndex);
+		ObjectsWithDirtyOwner.SetNumBits(MaxInternalIndex);
+		ObjectsWithOwnerFilter.SetNumBits(MaxInternalIndex);
+		ObjectsWithPerObjectInfo.SetNumBits(MaxInternalIndex);
 
-		ObjectIndexToDynamicFilterIndex.SetNumUninitialized(Params.MaxObjectCount);
-		FMemory::Memset(ObjectIndexToDynamicFilterIndex.GetData(), InvalidDynamicFilterIndex, Params.MaxObjectCount*sizeof(decltype(ObjectIndexToDynamicFilterIndex)::ElementType));
+		AllConnectionFilteredObjects.SetNumBits(MaxInternalIndex);
+		DynamicFilterEnabledObjects.SetNumBits(MaxInternalIndex);
+		
+		ObjectsRequiringDynamicFilterUpdate.SetNumBits(MaxInternalIndex);
+	}
 
-		ObjectsRequiringDynamicFilterUpdate.Init(MaxObjectCount);
+	// Increase TArrays whose index maps to a NetBitArray
+	{
+		ObjectIndexToPerObjectInfoIndex.SetNumZeroed(MaxInternalIndex);
+		ObjectIndexToOwningConnection.SetNumZeroed(MaxInternalIndex);
+		ObjectScopeHysteresisFrameCounts.SetNumZeroed(MaxInternalIndex);
+		NetObjectFilteringInfos.SetNumZeroed(MaxInternalIndex);
 	}
 
-	InitFilters();
-	InitObjectScopeHysteresis();
+	// ObjectIndexToDynamicFilterIndex is initialized to a non-zero value.
+	{
+		const int32 PrevMaxSize = ObjectIndexToDynamicFilterIndex.Num();
+		ObjectIndexToDynamicFilterIndex.SetNumUninitialized(MaxInternalIndex);
+
+		// Initialize the newly allocated buffer portion
+		if (MaxInternalIndex > 0)
+		{
+			checkf(MaxInternalIndex > (uint32)PrevMaxSize, TEXT("Not expected for the array to get smaller."));
+			uint8* NewBufferToInit = ObjectIndexToDynamicFilterIndex.GetData();
+			NewBufferToInit += PrevMaxSize;
+			FMemory::Memset(NewBufferToInit, InvalidDynamicFilterIndex, (MaxInternalIndex - PrevMaxSize) * sizeof(decltype(ObjectIndexToDynamicFilterIndex)::ElementType));
+		}
+	}
+
+	// Always allocated and maintained regardless of whether the feature is enabled or not.
+	HysteresisState.ObjectsToClear.SetNumBits(MaxInternalNetRefIndex);
+}
+
+void FReplicationFiltering::OnMaxInternalNetRefIndexIncreased(FInternalNetRefIndex NewMaxInternalIndex)
+{
+	MaxInternalNetRefIndex = NewMaxInternalIndex;
+
+	SetNetObjectListsSize(NewMaxInternalIndex);
+
+	// Resize the per-connection data
+	{
+		IRIS_PROFILER_SCOPE(FReplicationFiltering_ResizeAllPerConnectionLists);
+
+		auto ResizePerConnectionInfo = [this, NewMaxInternalIndex](uint32 ConnectionId)
+		{
+			FPerConnectionInfo& ConnectionInfo = this->ConnectionInfos[ConnectionId];
+			this->SetPerConnectionListsSize(ConnectionInfo, NewMaxInternalIndex);
+		};
+
+		ValidConnections.ForAllSetBits(ResizePerConnectionInfo);
+	}
+
+	// Propagate the increase to the DynamicFilters
+	{
+		TArrayView<FNetObjectFilteringInfo> NewFilterInfoView = GetNetObjectFilteringInfos();
+		for (FFilterInfo& FilterInfo : DynamicFilterInfos)
+		{
+			FilterInfo.Filter->MaxInternalNetRefIndexIncreased(NewMaxInternalIndex, NewFilterInfoView);
+		}
+	}
 }
 
 void FReplicationFiltering::Filter()
@@ -653,29 +712,12 @@ void FReplicationFiltering::InitNewConnections()
 		const FNetBitArrayView ScopableInternalIndices = NetRefHandleManager->GetCurrentFrameScopableInternalIndices();
 		FPerConnectionInfo& ConnectionInfo = this->ConnectionInfos[ConnectionId];
 
-		ConnectionInfo.ConnectionFilteredObjects.Init(MaxObjectCount);
-		ConnectionInfo.ConnectionFilteredObjects.Copy(ScopableInternalIndices);
-		ConnectionInfo.ConnectionFilteredObjects.ClearBit(FNetRefHandleManager::InvalidInternalIndex);
-
-		// Do not filter out anything by default.
-		ConnectionInfo.GroupExcludedObjects.Init(MaxObjectCount);
-
-		// Do not override dynamic filtering by default.
-		ConnectionInfo.GroupIncludedObjects.Init(MaxObjectCount);
+		const FInternalNetRefIndex CurrentMaxInternalindex = NetRefHandleManager->GetCurrentMaxInternalNetRefIndex();
 
-		// The combined result of scoped objects, connection filtering and group exclusion filtering.
-		ConnectionInfo.ObjectsInScopeBeforeDynamicFiltering.Init(MaxObjectCount);
+		this->SetPerConnectionListsSize(ConnectionInfo, CurrentMaxInternalindex);
 
-		// The final result of all filtering.
-		ConnectionInfo.ObjectsInScope.Init(MaxObjectCount);
-
-		if (HasDynamicFilters())
-		{
-			ConnectionInfo.DynamicFilteredOutObjects.Init(MaxObjectCount);
-			ConnectionInfo.InProgressDynamicFilteredOutObjects.Init(MaxObjectCount);
-			ConnectionInfo.DynamicFilteredOutObjectsHysteresisAdjusted.Init(MaxObjectCount);
-			ConnectionInfo.HysteresisUpdater.Init(MaxObjectCount);
-		}
+		ConnectionInfo.ConnectionFilteredObjects.Copy(ScopableInternalIndices);
+		ConnectionInfo.ConnectionFilteredObjects.ClearBit(FNetRefHandleManager::InvalidInternalIndex);
 
 		// Update group exclusion filtering
 		{
@@ -809,6 +851,31 @@ void FReplicationFiltering::ResetRemovedConnections()
 	FNetBitArray::ForAllSetBits(ExclusionFilterGroups, InclusionFilterGroups, FNetBitArrayBase::OrOp, ResetGroupFilterStatus);
 }
 
+void FReplicationFiltering::SetPerConnectionListsSize(FPerConnectionInfo& ConnectionInfo, FInternalNetRefIndex NewMaxInternalIndex)
+{
+	ConnectionInfo.ConnectionFilteredObjects.SetNumBits(NewMaxInternalIndex);
+
+	// Do not filter out anything by default.
+	ConnectionInfo.GroupExcludedObjects.SetNumBits(NewMaxInternalIndex);
+
+	// Do not override dynamic filtering by default.
+	ConnectionInfo.GroupIncludedObjects.SetNumBits(NewMaxInternalIndex);
+
+	// The combined result of scoped objects, connection filtering and group exclusion filtering.
+	ConnectionInfo.ObjectsInScopeBeforeDynamicFiltering.SetNumBits(NewMaxInternalIndex);
+
+	// The final result of all filtering.
+	ConnectionInfo.ObjectsInScope.SetNumBits(NewMaxInternalIndex);
+
+	if (HasDynamicFilters())
+	{
+		ConnectionInfo.DynamicFilteredOutObjects.SetNumBits(NewMaxInternalIndex);
+		ConnectionInfo.InProgressDynamicFilteredOutObjects.SetNumBits(NewMaxInternalIndex);
+		ConnectionInfo.DynamicFilteredOutObjectsHysteresisAdjusted.Init(NewMaxInternalIndex);
+		ConnectionInfo.HysteresisUpdater.Init(NewMaxInternalIndex);
+	}
+}
+
 void FReplicationFiltering::UpdateObjectsInScope()
 {
 	IRIS_PROFILER_SCOPE(FReplicationFiltering_UpdateObjectsInScope);
@@ -1270,7 +1337,7 @@ void FReplicationFiltering::UpdateDynamicFiltering()
 	FilteredOutByHysteresisObjects.Reserve(256);
 
 	uint32* AllowedObjectsData = static_cast<uint32*>(FMemory_Alloca(WordCountForObjectBitArrays * sizeof(uint32)));
-	FNetBitArrayView AllowedObjects(AllowedObjectsData, MaxObjectCount, FNetBitArrayView::NoResetNoValidate);
+	FNetBitArrayView AllowedObjects(AllowedObjectsData, MaxInternalNetRefIndex, FNetBitArrayView::NoResetNoValidate);
 
 	// Subobjects will never be added to a dynamic filter, but objects can become dependent at any time.
 	// We need to make sure they are not filtered out.
@@ -2598,7 +2665,8 @@ void FReplicationFiltering::InitFilters()
 		FNetObjectFilterInitParams InitParams;
 		InitParams.ReplicationSystem = ReplicationSystem;
 		InitParams.Config = (NetObjectFilterConfigClass ? NewObject<UNetObjectFilterConfig>((UObject*)GetTransientPackage(), NetObjectFilterConfigClass) : nullptr);
-		InitParams.MaxObjectCount = MaxObjectCount;
+		InitParams.AbsoluteMaxNetObjectCount = NetRefHandleManager->GetMaxActiveObjectCount();
+		InitParams.CurrentMaxInternalIndex = MaxInternalNetRefIndex;
 		InitParams.MaxConnectionCount = Connections->GetMaxConnectionCount();
 
 		Info.Filter->Init(InitParams);
@@ -2611,8 +2679,6 @@ void FReplicationFiltering::InitFilters()
 void FReplicationFiltering::InitObjectScopeHysteresis()
 {
 	HysteresisState.Mode = Config->IsObjectScopeHysteresisEnabled() ? EHysteresisProcessingMode::Enabled : EHysteresisProcessingMode::Disabled;
-	// Always allocated and maintained regardless of whether the feature is enabled or not.
-	HysteresisState.ObjectsToClear.SetNumBits(MaxObjectCount);
 }
 
 void FReplicationFiltering::RemoveFromDynamicFilter(uint32 ObjectIndex, uint32 FilterIndex)
diff --git a/Engine/Source/Runtime/Experimental/Iris/Core/Private/Iris/ReplicationSystem/Filtering/ReplicationFiltering.h b/Engine/Source/Runtime/Experimental/Iris/Core/Private/Iris/ReplicationSystem/Filtering/ReplicationFiltering.h
--- a/Engine/Source/Runtime/Experimental/Iris/Core/Private/Iris/ReplicationSystem/Filtering/ReplicationFiltering.h
+++ b/Engine/Source/Runtime/Experimental/Iris/Core/Private/Iris/ReplicationSystem/Filtering/ReplicationFiltering.h
@@ -48,8 +48,8 @@ struct FReplicationFilteringInitParams
 	FNetObjectGroups* Groups = nullptr;
 	FDeltaCompressionBaselineInvalidationTracker* BaselineInvalidationTracker = nullptr;
 	FReplicationConnections* Connections = nullptr;
+	FInternalNetRefIndex MaxInternalNetRefIndex = 0;
 	uint32 MaxGroupCount = 0;
-	uint32 MaxObjectCount = 0;
 };
 
 class FReplicationFiltering
@@ -58,6 +58,10 @@ public:
 	FReplicationFiltering();
 
 	void Init(FReplicationFilteringInitParams& Params);
+	void Deinit();
+
+	/** Called when the maximum InternalNetRefIndex increased and we need to realloc our lists */
+	void OnMaxInternalNetRefIndexIncreased(FInternalNetRefIndex NewMaxInternalIndex);
 
 	/**
 	 * Executes group, owner and connection filtering then any dynamic filters.
@@ -201,6 +205,9 @@ private:
 	void InitFilters();
 	void InitObjectScopeHysteresis();
 
+	void SetNetObjectListsSize(FInternalNetRefIndex MaxInternalIndex);
+	void SetPerConnectionListsSize(FPerConnectionInfo& ConnectionInfo, FInternalNetRefIndex NewMaxInternalIndex);
+
 	void InitNewConnections();
 	void ResetRemovedConnections();
 	void UpdateObjectsInScope();
@@ -318,7 +325,7 @@ private:
 	FNetBitArray NewConnections;
 
 	// Object specifics
-	uint32 MaxObjectCount = 0;
+	FInternalNetRefIndex MaxInternalNetRefIndex = 0;
 	uint32 WordCountForObjectBitArrays = 0;
 
 	// Filter specifics
diff --git a/Engine/Source/Runtime/Experimental/Iris/Core/Private/Iris/ReplicationSystem/NetCullDistanceOverrides.cpp b/Engine/Source/Runtime/Experimental/Iris/Core/Private/Iris/ReplicationSystem/NetCullDistanceOverrides.cpp
--- a/Engine/Source/Runtime/Experimental/Iris/Core/Private/Iris/ReplicationSystem/NetCullDistanceOverrides.cpp
+++ b/Engine/Source/Runtime/Experimental/Iris/Core/Private/Iris/ReplicationSystem/NetCullDistanceOverrides.cpp
@@ -8,7 +8,12 @@ namespace UE::Net
 
 void FNetCullDistanceOverrides::Init(const FNetCullDistanceOverridesInitParams& InitParams)
 {
-	ValidCullDistanceSqr.Init(InitParams.MaxObjectCount);
+	ValidCullDistanceSqr.Init(InitParams.MaxInternalNetRefIndex);
+}
+
+void FNetCullDistanceOverrides::OnMaxInternalNetRefIndexIncreased(UE::Net::Private::FInternalNetRefIndex NewMaxInternalIndex)
+{
+	ValidCullDistanceSqr.SetNumBits(NewMaxInternalIndex);
 }
 
 bool FNetCullDistanceOverrides::ClearCullDistanceSqr(uint32 ObjectIndex)
diff --git a/Engine/Source/Runtime/Experimental/Iris/Core/Private/Iris/ReplicationSystem/NetRefHandleManager.cpp b/Engine/Source/Runtime/Experimental/Iris/Core/Private/Iris/ReplicationSystem/NetRefHandleManager.cpp
--- a/Engine/Source/Runtime/Experimental/Iris/Core/Private/Iris/ReplicationSystem/NetRefHandleManager.cpp
+++ b/Engine/Source/Runtime/Experimental/Iris/Core/Private/Iris/ReplicationSystem/NetRefHandleManager.cpp
@@ -25,36 +25,34 @@ namespace UE::Net::Private
 FNetRefHandleManager::FNetRefHandleManager(FReplicationProtocolManager& InReplicationProtocolManager)
 	: ReplicationProtocolManager(InReplicationProtocolManager)
 {
+	static_assert(InvalidInternalIndex == 0, "FNetRefHandleManager::InvalidInternalIndex has an unexpected value");
 }
 
 void FNetRefHandleManager::Init(const FInitParams& InitParams)
 {
-	MaxActiveObjectCount = InitParams.MaxActiveObjectCount;
+	MaxActiveObjectCount = FNetBitArray::RoundUpToMaxWordBitCount(InitParams.MaxActiveObjectCount);
+	InternalNetRefIndexGrowSize = InitParams.InternalNetRefIndexGrowSize > 0 ? FNetBitArray::RoundUpToMaxWordBitCount(InitParams.InternalNetRefIndexGrowSize) : MaxActiveObjectCount;
+
 	ReplicationSystemId = InitParams.ReplicationSystemId;
 
-	// PreAllocatedObjectCount must be a minimum of 1 to account for InvalidInternalIndex.
-	PreAllocatedObjectCount = InitParams.PreAllocatedObjectCount > 0 ? InitParams.PreAllocatedObjectCount : 1;
+	// Must be a minimum of 1 to account for InvalidInternalIndex.
+	uint32 PreAllocatedNetChunkedArrayCount = FMath::Clamp(InitParams.NetChunkedArrayCount, 1U, MaxActiveObjectCount);
 
-	// Ensure that the pre-allocated object count is not greater than the maximum. If it is, just set it to the maximum.
-	if (PreAllocatedObjectCount > MaxActiveObjectCount)
-	{
-		UE_LOG(LogIris, Error, TEXT("PreAllocatedObjectCount (%d) must be smaller or equal to MaxActiveObjectCount (%d). Setting PreAllocatedObjectCount to MaxActiveObjectCount."), PreAllocatedObjectCount, MaxActiveObjectCount);
-		PreAllocatedObjectCount = MaxActiveObjectCount;
-	}
+	// Calculate the highest internal index possible with the current preallocated buffers. Must be a minimum of 0 to support InvalidInternalIndex.
+	HighestNetChunkedArrayInternalIndex = PreAllocatedNetChunkedArrayCount - 1;
 
-	// Calculate the largest internal index and must be a minimum of 0 to support InvalidInternalIndex.
-	LargestInternalIndex = PreAllocatedObjectCount > 0 ? PreAllocatedObjectCount - 1 : 0;
+	// Initialize NetObjectList configs
+	CurrentMaxInternalNetRefIndex = InitParams.InternalNetRefIndexInitSize > 0 ? FMath::Min(InitParams.InternalNetRefIndexInitSize, MaxActiveObjectCount) : MaxActiveObjectCount;
+	CurrentMaxInternalNetRefIndex = FNetBitArray::RoundUpToMaxWordBitCount(CurrentMaxInternalNetRefIndex);
 
-	UE_LOG(LogIris, Log, TEXT("NetRefHandleManager: Configured with MaxActiveObjectCount=%d, PreAllocatedObjectCount=%d, LargestInternalIndex=%d"),
-		MaxActiveObjectCount, PreAllocatedObjectCount, LargestInternalIndex);
-
-	static_assert(InvalidInternalIndex == 0, "FNetRefHandleManager::InvalidInternalIndex has an unexpected value");
+	UE_LOG(LogIris, Log, TEXT("NetRefHandleManager: Configured with MaxActiveObjectCount=%d, MaxInternalNetRefIndex: %u, Grow=%u, NetChunkedArray: Init=%u|Highest=%u"),
+		MaxActiveObjectCount, CurrentMaxInternalNetRefIndex, InternalNetRefIndexGrowSize, PreAllocatedNetChunkedArrayCount, HighestNetChunkedArrayInternalIndex);
 
 	// Initialize TNetChunkedArrays with PreAllocatedObjectCount
-	ReplicatedObjectData = TNetChunkedArray<FReplicatedObjectData>(PreAllocatedObjectCount, EInitMemory::Constructor);
-	ReplicatedObjectRefCount = TNetChunkedArray<uint16>(PreAllocatedObjectCount, EInitMemory::Zero);
-	ReplicatedObjectStateBuffers = TNetChunkedArray<uint8*>(PreAllocatedObjectCount, EInitMemory::Zero);
-	ReplicatedInstances = TNetChunkedArray<TObjectPtr<UObject>>(PreAllocatedObjectCount, EInitMemory::Zero);
+	ReplicatedObjectData = TNetChunkedArray<FReplicatedObjectData>(PreAllocatedNetChunkedArrayCount, EInitMemory::Constructor);
+	ReplicatedObjectRefCount = TNetChunkedArray<uint16>(PreAllocatedNetChunkedArrayCount, EInitMemory::Zero);
+	ReplicatedObjectStateBuffers = TNetChunkedArray<uint8*>(PreAllocatedNetChunkedArrayCount, EInitMemory::Zero);
+	ReplicatedInstances = TNetChunkedArray<TObjectPtr<UObject>>(PreAllocatedNetChunkedArrayCount, EInitMemory::Zero);
 
 	// For convenience we initialize ReplicatedObjectData for InvalidInternalIndex so that GetReplicatedObjectDataNoCheck returns something useful.
 	ReplicatedObjectData[InvalidInternalIndex] = FReplicatedObjectData();
@@ -82,44 +80,89 @@ void FNetRefHandleManager::Init(const FInitParams& InitParams)
 
 void FNetRefHandleManager::Deinit()
 {
-	// We can clear out the reserved index for invalid handles now.
 	AssignedInternalIndices.ClearBit(0);
-	AssignedInternalIndices.ForAllSetBits([this](uint32 InternalIndex) { 
-		this->InternalDestroyNetObject(InternalIndex); 
+	AssignedInternalIndices.ForAllSetBits([this](uint32 InternalIndex) 
+	{
+		this->InternalDestroyNetObject(InternalIndex);
 	});
+
+	OwnedNetBitArrays.Empty();
+
+	ensureMsgf(OnMaxInternalNetRefIndexIncreased.IsBound() == false, TEXT("FNetRefHandleManager still has delegates registered to OnMaxInternalNetRefIndexIncreased while deinitializing."));
 }
 
 void FNetRefHandleManager::InitNetBitArray(FNetBitArray* NetBitArray)
 {
-	NetBitArray->Init(MaxActiveObjectCount);
+	NetBitArray->Init(CurrentMaxInternalNetRefIndex);
+	OwnedNetBitArrays.Add(NetBitArray);
 }
 
-void FNetRefHandleManager::GrowBuffersToLargestIndex(uint32 InternalIndex)
+FInternalNetRefIndex FNetRefHandleManager::GrowNetObjectLists()
 {
-	if (InternalIndex > LargestInternalIndex)
+	check(AssignedInternalIndices.GetNumBits() == CurrentMaxInternalNetRefIndex);
+
+	// The old max is the next available index
+	const FInternalNetRefIndex NextFreeIndex = CurrentMaxInternalNetRefIndex;
+
+	// We already are at the max, return InvalidIndex and abort.
+	if (CurrentMaxInternalNetRefIndex >= MaxActiveObjectCount)
 	{
-		QUICK_SCOPE_CYCLE_COUNTER(STAT_FNetRefHandleManager_GrowBuffersToLargestIndex);
-		CSV_CUSTOM_STAT(IrisCommon, DynamicBufferGrowCount, 1, ECsvCustomStatOp::Accumulate);
+		return InvalidInternalIndex;
+	}
 
-		// This call will add the necessary number of elements and chunks to ReplicatedObjectRefCount
-		// to accommodate InternalIndex. Once this is done, we determine how many more elements could be
-		// added to the array without adding a new chunk and ensure all of the other buffers have this
-		// many elements. 
-		// 
-		// This optimization will reduce the number of calls to AddToIndexUninitialized() and AddToIndexZeroed()
-		// but assumes each buffer is going to have the same number of elements all the time.
-		ReplicatedObjectRefCount.AddToIndexUninitialized(InternalIndex);
-		const int32 LargestIndexInCurrentChunk = ReplicatedObjectRefCount.Capacity() - 1;
+    CurrentMaxInternalNetRefIndex += InternalNetRefIndexGrowSize;
 
-		ReplicatedObjectRefCount.AddToIndexUninitialized(LargestIndexInCurrentChunk);
-		ReplicatedObjectStateBuffers.AddToIndexZeroed(LargestIndexInCurrentChunk);
-		ReplicatedInstances.AddToIndexZeroed(LargestIndexInCurrentChunk);
-		ReplicatedObjectData.AddToIndexZeroed(LargestIndexInCurrentChunk);
-		
-		LargestInternalIndex = static_cast<uint32>(LargestIndexInCurrentChunk);
+	if (CurrentMaxInternalNetRefIndex > MaxActiveObjectCount)
+	{
+		// Last grow possibility before we cause a critical failure
+		CurrentMaxInternalNetRefIndex = MaxActiveObjectCount;
+	}
+
+	UE_LOG(LogIris, Log, TEXT("FNetRefHandleManager::GrowNetObjectLists grew MaxInternalIndex from %u to %u (+%u)"), NextFreeIndex, CurrentMaxInternalNetRefIndex, CurrentMaxInternalNetRefIndex - NextFreeIndex);
+
+	MaxInternalNetRefIndexIncreased(CurrentMaxInternalNetRefIndex);
+
+	return NextFreeIndex;
+}
+
+void FNetRefHandleManager::MaxInternalNetRefIndexIncreased(FInternalNetRefIndex NewMaxInternalNetRefIndex)
+{
+	QUICK_SCOPE_CYCLE_COUNTER(STAT_FNetRefHandleManager_MaxInternalNetRefIndexIncreased);
+	CSV_CUSTOM_STAT(IrisCommon, MaxInternalIndexIncreasedCount, 1, ECsvCustomStatOp::Accumulate);
 
-		OnLargestIndexIncreaseDelegate.Broadcast(LargestInternalIndex);
+	// Start by reallocating all the NetBitArrays we own
+	for (FNetBitArray* NetBitArray : OwnedNetBitArrays)
+	{
+		NetBitArray->SetNumBits(NewMaxInternalNetRefIndex);
 	}
+
+	// Tell other systems to increase their lists too
+	OnMaxInternalNetRefIndexIncreased.Broadcast(NewMaxInternalNetRefIndex);
+}
+
+void FNetRefHandleManager::GrowNetChunkedArrayBuffers(FInternalNetRefIndex InternalIndex)
+{
+	QUICK_SCOPE_CYCLE_COUNTER(STAT_FNetRefHandleManager_GrowNetChunkedArrayBuffers);
+	CSV_CUSTOM_STAT(IrisCommon, DynamicNetChunkedArrayGrowCount, 1, ECsvCustomStatOp::Accumulate);
+
+	// This call will add the necessary number of elements and chunks to ReplicatedObjectRefCount
+	// to accommodate InternalIndex. Once this is done, we determine how many more elements could be
+	// added to the array without adding a new chunk and ensure all of the other buffers have this
+	// many elements. 
+	// 
+	// This optimization will reduce the number of calls to AddToIndexUninitialized() and AddToIndexZeroed()
+	// but assumes each buffer is going to have the same number of elements all the time.
+	ReplicatedObjectRefCount.AddToIndexUninitialized(InternalIndex);
+	const int32 LargestIndexInCurrentChunk = ReplicatedObjectRefCount.Capacity() - 1;
+
+	ReplicatedObjectRefCount.AddToIndexUninitialized(LargestIndexInCurrentChunk);
+	ReplicatedObjectStateBuffers.AddToIndexZeroed(LargestIndexInCurrentChunk);
+	ReplicatedInstances.AddToIndexZeroed(LargestIndexInCurrentChunk);
+	ReplicatedObjectData.AddToIndexZeroed(LargestIndexInCurrentChunk);
+		
+	HighestNetChunkedArrayInternalIndex = static_cast<uint32>(LargestIndexInCurrentChunk);
+
+	OnNetChunkedArrayIncrease.Broadcast(HighestNetChunkedArrayInternalIndex);
 }
 
 uint64 FNetRefHandleManager::GetNextNetRefHandleId(uint64 HandleId) const
@@ -157,63 +200,70 @@ FInternalNetRefIndex FNetRefHandleManager::InternalCreateNetObject(const FNetRef
 		return InvalidInternalIndex;
 	}
 
-	const uint32 InternalIndex = GetNextFreeInternalIndex();
-	if (InternalIndex != InvalidInternalIndex)
+	uint32 InternalIndex = GetNextFreeInternalIndex();
+
+	// Try to grow the NetObjectLists if no more indexes are available.
+	if (InternalIndex == InvalidInternalIndex)
 	{
-		UE_LOG(LogIris, Verbose, TEXT("FNetRefHandleManager::InternalCreateNetObject: (InternalIndex: %u) (%s)"), InternalIndex, *NetRefHandle.ToString());
+		InternalIndex = GrowNetObjectLists();
 
-		// Track the largest internal index and grow internal buffers if neccessary.
-		GrowBuffersToLargestIndex(InternalIndex);
+		// If we could not grow anymore, kill the process now. The system cannot replicate objects anymore and the game behavior is undefined.
+		if (InternalIndex == InvalidInternalIndex)
+		{
+			UE_LOG(LogIris, Fatal, TEXT("NetRefHandleManager: Hit the maximum limit of active replicated objects: %u. Aborting since we cannot replicate %s"), MaxActiveObjectCount, ReplicationProtocol->DebugName->Name);
+			return InvalidInternalIndex;
+		}
+	}
 
-		// Store data;
-		FReplicatedObjectData& Data = ReplicatedObjectData[InternalIndex];
+	UE_LOG(LogIris, Verbose, TEXT("FNetRefHandleManager::InternalCreateNetObject: (InternalIndex: %u) (%s)"), InternalIndex, *NetRefHandle.ToString());
 
-		Data = FReplicatedObjectData();
+	// Track the largest internal index and grow internal buffers if necessary.
+	if (InternalIndex > HighestNetChunkedArrayInternalIndex)
+	{
+		GrowNetChunkedArrayBuffers(InternalIndex);
+	}
 
-		Data.RefHandle = NetRefHandle;
-		Data.NetHandle = GlobalHandle;
-		Data.Protocol = ReplicationProtocol;
-		Data.InstanceProtocol = nullptr;
-		ObjectsWithPreUpdate.ClearBit(InternalIndex);
-		ReplicatedObjectStateBuffers[InternalIndex] = nullptr;
-		Data.ReceiveStateBuffer = nullptr;
-		Data.bShouldPropagateChangedStates = 1U;
+	// Store data;
+	FReplicatedObjectData& Data = ReplicatedObjectData[InternalIndex];
 
-		++ActiveObjectCount;
+	Data = FReplicatedObjectData();
+
+	Data.RefHandle = NetRefHandle;
+	Data.NetHandle = GlobalHandle;
+	Data.Protocol = ReplicationProtocol;
+	Data.InstanceProtocol = nullptr;
+	ObjectsWithPreUpdate.ClearBit(InternalIndex);
+	ReplicatedObjectStateBuffers[InternalIndex] = nullptr;
+	Data.ReceiveStateBuffer = nullptr;
+	Data.bShouldPropagateChangedStates = 1U;
+
+	++ActiveObjectCount;
 
-		// Add map entry so that we can map from NetRefHandle to InternalIndex
-		RefHandleToInternalIndex.Add(NetRefHandle, InternalIndex);
+	// Add map entry so that we can map from NetRefHandle to InternalIndex
+	RefHandleToInternalIndex.Add(NetRefHandle, InternalIndex);
 		
-		// Add mapping from global handle to InternalIndex to speed up lookups for ReplicationSystem public API
-		if (GlobalHandle.IsValid())
-		{
-			NetHandleToInternalIndex.Add(GlobalHandle, InternalIndex);
-		}
+	// Add mapping from global handle to InternalIndex to speed up lookups for ReplicationSystem public API
+	if (GlobalHandle.IsValid())
+	{
+		NetHandleToInternalIndex.Add(GlobalHandle, InternalIndex);
+	}
 
-		// Mark Handle index as assigned and scopable for now
-		AssignedInternalIndices.SetBit(InternalIndex);
-		GlobalScopableInternalIndices.SetBit(InternalIndex);
+	// Mark Handle index as assigned and scopable for now
+	AssignedInternalIndices.SetBit(InternalIndex);
+	GlobalScopableInternalIndices.SetBit(InternalIndex);
 
-		// When a handle is first created, it is not set to be a subobject
-		SubObjectInternalIndices.ClearBit(InternalIndex);
+	// When a handle is first created, it is not set to be a subobject
+	SubObjectInternalIndices.ClearBit(InternalIndex);
 
-		// Need a full copy if set, normally only needed for new objects.
-		Data.bNeedsFullCopyAndQuantize = 1U;
+	// Need a full copy if set, normally only needed for new objects.
+	Data.bNeedsFullCopyAndQuantize = 1U;
 
-		// Make sure we do a full poll of all properties the first time the object gets polled.
-		Data.bWantsFullPoll = 1U;
+	// Make sure we do a full poll of all properties the first time the object gets polled.
+	Data.bWantsFullPoll = 1U;
 
-		ReplicatedObjectRefCount[InternalIndex] = 0;
+	ReplicatedObjectRefCount[InternalIndex] = 0;
 
-		return InternalIndex;
-	}
-	else
-	{
-		UE_LOG(LogIris, Error, TEXT("NetRefHandleManager: Cannot find free active object (%d/%d)."), ActiveObjectCount, MaxActiveObjectCount);
-		ensureMsgf(false, TEXT("NetRefHandleManager: Cannot find free active object (%d/%d)."), ActiveObjectCount, MaxActiveObjectCount);
-	}
-		
-	return InvalidInternalIndex;
+	return InternalIndex;
 }
 
 void FNetRefHandleManager::AttachInstanceProtocol(FInternalNetRefIndex InternalIndex, const FReplicationInstanceProtocol* InstanceProtocol, UObject* Instance)
diff --git a/Engine/Source/Runtime/Experimental/Iris/Core/Private/Iris/ReplicationSystem/NetRefHandleManager.h b/Engine/Source/Runtime/Experimental/Iris/Core/Private/Iris/ReplicationSystem/NetRefHandleManager.h
--- a/Engine/Source/Runtime/Experimental/Iris/Core/Private/Iris/ReplicationSystem/NetRefHandleManager.h
+++ b/Engine/Source/Runtime/Experimental/Iris/Core/Private/Iris/ReplicationSystem/NetRefHandleManager.h
@@ -124,11 +124,14 @@ public:
 
 public:
 
+	/** External configuration variables used to initialize the NetRefHandleManager */
 	struct FInitParams
 	{
 		uint32 ReplicationSystemId = 0;
 		uint32 MaxActiveObjectCount = 0;
-		uint32 PreAllocatedObjectCount = 0;
+		uint32 InternalNetRefIndexInitSize = 0;
+		uint32 NetChunkedArrayCount = 0;
+		uint32 InternalNetRefIndexGrowSize = 0;
 	};
 
 	FNetRefHandleManager(FReplicationProtocolManager& InReplicationProtocolManager);
@@ -257,10 +260,12 @@ public:
 	void SetShouldPropagateChangedStates(FNetRefHandle Handle, bool bShouldPropagateChangedStates);
 	void SetShouldPropagateChangedStates(FInternalNetRefIndex ObjectInternalIndex, bool bShouldPropagateChangedStates);
 
+	/** The absolute maximum replicated objects that can be registered in the ReplicationSystem. Hitting this limit will cause a critical failure. */
 	uint32 GetMaxActiveObjectCount() const { return MaxActiveObjectCount; }
-	uint32 GetActiveObjectCount() const { return ActiveObjectCount; }
-	uint32 GetPreAllocatedObjectCount() const { return PreAllocatedObjectCount; }
 
+	/** The current amount of replicated objects registered for replication. */
+	uint32 GetActiveObjectCount() const { return ActiveObjectCount; }
+	
 	// We do refcount objects tracked by each connection in order to know when it is safe to reuse an InternalIndex
 	void AddNetObjectRef(FInternalNetRefIndex InternalIndex) { ++ReplicatedObjectRefCount[InternalIndex]; }
 	void ReleaseNetObjectRef(FInternalNetRefIndex InternalIndex) { check(ReplicatedObjectRefCount[InternalIndex] > 0); --ReplicatedObjectRefCount[InternalIndex]; }
@@ -302,16 +307,22 @@ public:
 	FNetBitArrayView GetWantToBeDormantInternalIndices() { return MakeNetBitArrayView(WantToBeDormantInternalIndices); }
 
 	/** Return a string to identify the object linked to an index in logs */
-	FString PrintObjectFromIndex(FInternalNetRefIndex ObjectIndex) const;
-	FString PrintObjectFromNetRefHandle(FNetRefHandle ObjectHandle) const;
+	[[nodiscard]] FString PrintObjectFromIndex(FInternalNetRefIndex ObjectIndex) const;
+	[[nodiscard]] FString PrintObjectFromNetRefHandle(FNetRefHandle ObjectHandle) const;
 
-	DECLARE_MULTICAST_DELEGATE_OneParam(FOnLargestIndexIncrease, uint32 LargestIndex);
+	/** Delegate that will notify when the highest NetChunkedArray internal index has increased (Highest = Max-1)*/
+	DECLARE_MULTICAST_DELEGATE_OneParam(FOnNetChunkedArrayIncrease, FInternalNetRefIndex HighestInternalIndex);
+	FOnNetChunkedArrayIncrease& GetOnNetChunkedArrayIncreaseDelegate() const { return OnNetChunkedArrayIncrease; };
 
-	/** Return a delegate that will notify when the largest internal index has increased. */
-	FOnLargestIndexIncrease& GetLargestIndexIncreaseDelegate() { return OnLargestIndexIncreaseDelegate; };
+	/** Return the highest internal index that NetChunkedArrays are currently allocated for. */
+	uint32 GetHighestNetChunkedArrayInternalIndex() const { return HighestNetChunkedArrayInternalIndex; };
 
-	/** Return the largest internal index that has been used. */
-	uint32 GetLargestInternalIndex() const { return LargestInternalIndex; };
+	/** Delegate that will notify when the NetObjectLists (eg. NetBitArray and TArray indexed via FInternalNetRefIndex) have a new maximum index to support */
+	DECLARE_MULTICAST_DELEGATE_OneParam(FOnMaxInternalNetRefIndexIncreased, FInternalNetRefIndex MaxInternalIndex);
+	FOnMaxInternalNetRefIndexIncreased& GetOnMaxInternalNetRefIndexIncreasedDelegate() const { return OnMaxInternalNetRefIndexIncreased; };
+
+	/** Return the maximum internal index that NetBitArrays and TArrays are currently allocated for. */
+	FInternalNetRefIndex GetCurrentMaxInternalNetRefIndex() const { return CurrentMaxInternalNetRefIndex; };
 
 	/** Get Objects that is flagged for PreUpdate (aka PreReplication) */
 	FNetBitArrayView GetObjectsWithPreUpdate() const { return MakeNetBitArrayView(ObjectsWithPreUpdate); }
@@ -354,23 +365,34 @@ private:
 	void InternalRemoveDependentObject(FInternalNetRefIndex ParentInternalIndex, FInternalNetRefIndex DependentInternalIndex, ERemoveDependentObjectFlags Flags = ERemoveDependentObjectFlags::All);
 	void InternalRemoveDependentObject(FInternalNetRefIndex DependentInternalIndex);
 
-	// Ensure that all buffers that depend on the largest internal index are sized correctly.
-	void GrowBuffersToLargestIndex(uint32 InternalIndex);
+	/**
+	 * Grow the NetObjectLists to fit more internal indexes.
+	 * Returns the newest available index if the increase worked or InvalidInternalIndex if we hit the maximum.
+	 */
+	FInternalNetRefIndex GrowNetObjectLists();
 
-	// The current replicated object count
-	uint32 ActiveObjectCount = 0;
+	/** Grow memory buffers when we use an index higher than the previous maximum. */
+	void GrowNetChunkedArrayBuffers(FInternalNetRefIndex InternalIndex);
+
+	/** Deal with the increase to the max index and tell other systems about it */
+	void MaxInternalNetRefIndexIncreased(uint32 NewBitCount);
+
+private:
 
 	// Max allowed replicated object count
 	uint32 MaxActiveObjectCount = 0;
 
-	// The number of pre-allocated objects used by internal buffers.
-	uint32 PreAllocatedObjectCount = 0;
+	// The current replicated object count
+	uint32 ActiveObjectCount = 0;
 
-	// The largest internal index value that has been used.
-	uint32 LargestInternalIndex = 0;
+	// By how many indexes to grow the NetObjectLists when we hit the highest index.
+	FInternalNetRefIndex InternalNetRefIndexGrowSize = 0;
 
-	// A delegate that is triggered when the largest encountered internal index increases.
-	FOnLargestIndexIncrease OnLargestIndexIncreaseDelegate;
+	// The maximum internal index that NetObjectLists need to be able to reference
+	FInternalNetRefIndex CurrentMaxInternalNetRefIndex = 0;
+
+	// The highest internal index currently allocated for NetChunkedArrays
+	uint32 HighestNetChunkedArrayInternalIndex = 0;
 
 	uint32 ReplicationSystemId = UE::Net::InvalidReplicationSystemId;
 
@@ -434,8 +456,6 @@ private:
 	// Map DestructionInfoIndex -> Orignal
 	TMap<uint32, uint32> DestroyedStartupObject;
 
-private:
-
 	// Array used in order to track objects pending destroy
 	TArray<FInternalNetRefIndex> PendingDestroyInternalIndices;
 
@@ -459,6 +479,14 @@ private:
 	FNetDependencyData SubObjects;
 
 	FReplicationProtocolManager& ReplicationProtocolManager;
+
+	// Delegates
+	mutable FOnMaxInternalNetRefIndexIncreased OnMaxInternalNetRefIndexIncreased;
+	mutable FOnNetChunkedArrayIncrease OnNetChunkedArrayIncrease;
+
+	// Track NetBitArrays owned by the NetRefHandleManager so they can automatically be resized.
+	TArray<FNetBitArray*> OwnedNetBitArrays;
+
 };
 
 const FNetRefHandleManager::FReplicatedObjectData& FNetRefHandleManager::GetReplicatedObjectData(FInternalNetRefIndex InternalIndex) const
diff --git a/Engine/Source/Runtime/Experimental/Iris/Core/Private/Iris/ReplicationSystem/Prioritization/LocationBasedNetObjectPrioritizer.cpp b/Engine/Source/Runtime/Experimental/Iris/Core/Private/Iris/ReplicationSystem/Prioritization/LocationBasedNetObjectPrioritizer.cpp
--- a/Engine/Source/Runtime/Experimental/Iris/Core/Private/Iris/ReplicationSystem/Prioritization/LocationBasedNetObjectPrioritizer.cpp
+++ b/Engine/Source/Runtime/Experimental/Iris/Core/Private/Iris/ReplicationSystem/Prioritization/LocationBasedNetObjectPrioritizer.cpp
@@ -15,10 +15,20 @@ ULocationBasedNetObjectPrioritizer::ULocationBasedNetObjectPrioritizer()
 
 void ULocationBasedNetObjectPrioritizer::Init(FNetObjectPrioritizerInitParams& Params)
 {
-	AssignedLocationIndices.Init(Params.MaxObjectCount);
+	AssignedLocationIndices.Init(Params.CurrentMaxInternalIndex);
 	WorldLocations = &Params.ReplicationSystem->GetWorldLocations();
 }
 
+void ULocationBasedNetObjectPrioritizer::Deinit()
+{
+	WorldLocations = nullptr;
+}
+
+void ULocationBasedNetObjectPrioritizer::OnMaxInternalNetRefIndexIncreased(uint32 NewMaxInternalIndex)
+{
+	AssignedLocationIndices.SetNumBits(NewMaxInternalIndex);
+}
+
 bool ULocationBasedNetObjectPrioritizer::AddObject(uint32 ObjectIndex, FNetObjectPrioritizerAddObjectParams& Params)
 {
 	// We support either a world location in the state, tagged with RepTag_WorldLocation, or via the WorldLocations instance.
diff --git a/Engine/Source/Runtime/Experimental/Iris/Core/Private/Iris/ReplicationSystem/Prioritization/NetObjectConnectionFilter.cpp b/Engine/Source/Runtime/Experimental/Iris/Core/Private/Iris/ReplicationSystem/Prioritization/NetObjectConnectionFilter.cpp
--- a/Engine/Source/Runtime/Experimental/Iris/Core/Private/Iris/ReplicationSystem/Prioritization/NetObjectConnectionFilter.cpp
+++ b/Engine/Source/Runtime/Experimental/Iris/Core/Private/Iris/ReplicationSystem/Prioritization/NetObjectConnectionFilter.cpp
@@ -5,11 +5,13 @@
 #include "Iris/ReplicationSystem/ReplicationSystem.h"
 #include "Iris/ReplicationSystem/ReplicationSystemInternal.h"
 #include "Iris/Core/IrisProfiler.h"
+#include "Iris/Core/IrisLog.h"
 
 void UNetObjectConnectionFilter::SetReplicateToConnection(UE::Net::FNetRefHandle RefHandle, uint32 ConnectionId, UE::Net::ENetFilterStatus FilterStatus)
 {
-	if (!ensure(ConnectionId < static_cast<uint32>(PerConnectionInfos.Num())))
+	if (ConnectionId >= static_cast<uint32>(PerConnectionInfos.Num()))
 	{
+		ensureMsgf(false, TEXT("UNetObjectConnectionFilter::SetReplicateToConnection received invalid ConnectionId:%u | Max:%u"), ConnectionId, PerConnectionInfos.Num());
 		return;
 	}
 
@@ -21,8 +23,10 @@ void UNetObjectConnectionFilter::SetReplicateToConnection(UE::Net::FNetRefHandle
 
 	if (const FFilteringInfo* FilteringInfo = static_cast<const FFilteringInfo*>(GetFilteringInfo(ObjectIndex)))
 	{
-		const uint32 LocalIndex = FilteringInfo->GetLocalObjectIndex();
+		const uint16 LocalIndex = FilteringInfo->GetLocalObjectIndex();
+
 		FPerConnectionInfo& PerConnectionInfo = PerConnectionInfos[ConnectionId];
+		ensureMsgf(LocalIndex < PerConnectionInfo.ReplicationEnabledObjects.GetNumBits(), TEXT("UNetObjectConnectionFilter::SetReplicateToConnection Object %u mapped to invalid local index %u"), ObjectIndex, LocalIndex);
 		PerConnectionInfo.ReplicationEnabledObjects.SetBitValue(LocalIndex, FilterStatus == UE::Net::ENetFilterStatus::Allow);
 	}
 }
@@ -31,15 +35,26 @@ void UNetObjectConnectionFilter::OnInit(const FNetObjectFilterInitParams& Params
 {
 	NetRefHandleManager = &Params.ReplicationSystem->GetReplicationSystemInternal()->GetNetRefHandleManager();
 
+	MaxInternalIndex = Params.CurrentMaxInternalIndex; 
+
 	Config = TStrongObjectPtr<UNetObjectConnectionFilterConfig>(CastChecked<UNetObjectConnectionFilterConfig>(Params.Config));
 
-	const uint16 MaxLocalObjectCount = static_cast<uint16>(FPlatformMath::Min<uint32>(Params.MaxObjectCount, Config->MaxObjectCount));
+	const uint16 MaxLocalObjectCount = static_cast<uint16>(FPlatformMath::Min<uint32>(Params.AbsoluteMaxNetObjectCount, Config->MaxObjectCount));
 	UsedLocalInfoIndices.Init(MaxLocalObjectCount);
 	LocalToNetRefIndex.SetNumZeroed(MaxLocalObjectCount);
 
 	PerConnectionInfos.SetNum(Params.MaxConnectionCount + 1);
 }
 
+void UNetObjectConnectionFilter::OnDeinit()
+{
+	NetRefHandleManager = nullptr;
+	Config = nullptr;
+
+	UsedLocalInfoIndices.Empty();
+	LocalToNetRefIndex.Empty();
+}
+
 void UNetObjectConnectionFilter::AddConnection(uint32 ConnectionId)
 {
 	FPerConnectionInfo& ConnInfo = PerConnectionInfos[ConnectionId];
@@ -56,16 +71,19 @@ void UNetObjectConnectionFilter::RemoveConnection(uint32 ConnectionId)
 bool UNetObjectConnectionFilter::AddObject(uint32 ObjectIndex, FNetObjectFilterAddObjectParams& Params)
 {
 	const uint32 LocalIndex = UsedLocalInfoIndices.FindFirstZero();
-	if (!ensureMsgf(LocalIndex != UE::Net::FNetBitArrayBase::InvalidIndex, TEXT("%hs. MaxObjectCount: %u. Config type %s."), "Too many objects added to NetObjectConnectionFilter. Object will not be handled by filter!", (Config.IsValid() ? Config->MaxObjectCount : uint16(0)), ToCStr(GetNameSafe(Config.Get()))))
+	if (LocalIndex == UE::Net::FNetBitArrayBase::InvalidIndex)
 	{
+		ensureMsgf(false, TEXT("%hs. MaxObjectCount: %u. Config type %s."), "Too many objects added to NetObjectConnectionFilter. Object will not be handled by filter!", (Config.IsValid() ? Config->MaxObjectCount : uint16(0)), ToCStr(GetNameSafe(Config.Get())));
 		return false;
 	}
 
+	UE_LOG(LogIris, Verbose, TEXT("UNetObjectConnectionFilter::AddObject added %u | Mapped to LocalIndex %u"), ObjectIndex, LocalIndex);
+
 	UsedLocalInfoIndices.SetBit(LocalIndex);
 	LocalToNetRefIndex[LocalIndex] = ObjectIndex;
 
 	FFilteringInfo& Info = static_cast<FFilteringInfo&>(Params.OutInfo);
-	Info.SetLocalObjectIndex(static_cast<uint16>(LocalIndex));
+	Info.SetLocalObjectIndex((uint16)(LocalIndex)); //LPS: Assert if overflow
 
 	return true;
 }
@@ -79,6 +97,8 @@ void UNetObjectConnectionFilter::RemoveObject(uint32 ObjectIndex, const FNetObje
 	const uint16 LocalIndex = Info.GetLocalObjectIndex();
 	UsedLocalInfoIndices.ClearBit(LocalIndex);
 
+	UE_LOG(LogIris, Verbose, TEXT("UNetObjectConnectionFilter::RemoveObject removed %u | Mapped to LocalIndex %u"), ObjectIndex, LocalIndex);
+
 	// For good measure. It's not strictly needed.
 	LocalToNetRefIndex[LocalIndex] = UE::Net::Private::FNetRefHandleManager::InvalidInternalIndex;
 }
diff --git a/Engine/Source/Runtime/Experimental/Iris/Core/Private/Iris/ReplicationSystem/Prioritization/NetObjectCountLimiter.cpp b/Engine/Source/Runtime/Experimental/Iris/Core/Private/Iris/ReplicationSystem/Prioritization/NetObjectCountLimiter.cpp
--- a/Engine/Source/Runtime/Experimental/Iris/Core/Private/Iris/ReplicationSystem/Prioritization/NetObjectCountLimiter.cpp
+++ b/Engine/Source/Runtime/Experimental/Iris/Core/Private/Iris/ReplicationSystem/Prioritization/NetObjectCountLimiter.cpp
@@ -22,7 +22,7 @@ void UNetObjectCountLimiter::Init(FNetObjectPrioritizerInitParams& Params)
 	checkf(Params.Config != nullptr, TEXT("Need config to operate."));
 	checkf(Params.MaxConnectionCount < 65536U, TEXT("Assumption being able to use uint16 for ConnectionIds is incorrect."));
 	Config = TStrongObjectPtr<UNetObjectCountLimiterConfig>(CastChecked<UNetObjectCountLimiterConfig>(Params.Config));
-	ensureAlwaysMsgf(Config->MaxObjectCount >= 1U, TEXT("Prioritizer will not consider any object for replication. They will be replicated once when constructed, but never again."));
+	ensureMsgf(Config->MaxObjectCount >= 1U, TEXT("Prioritizer will not consider any object for replication. They will be replicated once when constructed, but never again."));
 
 	InternalObjectIndices.Init(ObjectGrowCount);
 	PerConnectionInfos.SetNum(Params.MaxConnectionCount + 1);
@@ -30,6 +30,14 @@ void UNetObjectCountLimiter::Init(FNetObjectPrioritizerInitParams& Params)
 	ReplicationSystem = Params.ReplicationSystem;
 }
 
+void UNetObjectCountLimiter::Deinit()
+{
+	ReplicationSystem = nullptr;
+	Config = nullptr;
+
+	InternalObjectIndices.Empty();
+}
+
 void UNetObjectCountLimiter::AddConnection(uint32 ConnectionId)
 {
 	FPerConnectionInfo& Info = PerConnectionInfos[ConnectionId];
diff --git a/Engine/Source/Runtime/Experimental/Iris/Core/Private/Iris/ReplicationSystem/Prioritization/ReplicationPrioritization.cpp b/Engine/Source/Runtime/Experimental/Iris/Core/Private/Iris/ReplicationSystem/Prioritization/ReplicationPrioritization.cpp
--- a/Engine/Source/Runtime/Experimental/Iris/Core/Private/Iris/ReplicationSystem/Prioritization/ReplicationPrioritization.cpp
+++ b/Engine/Source/Runtime/Experimental/Iris/Core/Private/Iris/ReplicationSystem/Prioritization/ReplicationPrioritization.cpp
@@ -5,7 +5,6 @@
 #include "Iris/Core/IrisLog.h"
 #include "Iris/Core/IrisProfiler.h"
 #include "Net/Core/Misc/NetCVars.h"
-#include "Net/Core/NetBitArray.h"
 #include "Iris/ReplicationSystem/NetRefHandleManager.h"
 #include "Iris/ReplicationSystem/ReplicationConnections.h"
 #include "Iris/ReplicationSystem/ReplicationWriter.h"
@@ -298,31 +297,91 @@ void FReplicationPrioritization::Init(FReplicationPrioritizationInitParams& Para
 	Connections = Params.Connections;
 	NetRefHandleManager = Params.NetRefHandleManager;
 
-	MaxObjectCount = Params.MaxObjectCount;
+	MaxInternalNetRefIndex = Params.MaxInternalNetRefIndex;
 
-	constexpr EAllowShrinking AllowShrinking = EAllowShrinking::No;
+	SetNetObjectListsSize(MaxInternalNetRefIndex);
+
+	ConnectionInfos.Reserve(Params.Connections->GetMaxConnectionCount());
+
+	InitPrioritizers();
+}
+
+void FReplicationPrioritization::Deinit()
+{
+	for (FPrioritizerInfo& Info : PrioritizerInfos)
+	{
+		Info.Prioritizer->Deinit();
+		Info.Prioritizer = nullptr;
+	}
+
+	PrioritizerDefinitions = nullptr;
+
+	ReplicationSystem = nullptr;
+	Connections = nullptr;
+	NetRefHandleManager = nullptr;
+
+	NetObjectPrioritizationInfos.Empty();
+	ObjectIndexToPrioritizer.Empty();
+	DefaultPriorities.Empty();
+	ObjectsWithNewStaticPriority.Empty();
+}
+
+void FReplicationPrioritization::SetNetObjectListsSize(FInternalNetRefIndex MaxInternalIndex)
+{
+	constexpr EAllowShrinking NoShrinking = EAllowShrinking::No;
 
 	// $IRIS TODO: This can be quite wasteful in terms of memory assuming many objects will use a static priority. Need object pool!
-	NetObjectPrioritizationInfos.SetNumUninitialized(Params.MaxObjectCount, AllowShrinking);
+	NetObjectPrioritizationInfos.SetNumUninitialized(MaxInternalIndex, NoShrinking);
 
+	// Properly Initialize ObjectIndexToPrioritizer
 	{
-		ObjectIndexToPrioritizer.SetNumUninitialized(Params.MaxObjectCount, AllowShrinking);
-		FMemory::Memset(ObjectIndexToPrioritizer.GetData(), FReplicationPrioritization_InvalidNetObjectPrioritizerIndex, Params.MaxObjectCount*sizeof(decltype(ObjectIndexToPrioritizer)::ElementType));
+		const int32 PreviousSize = ObjectIndexToPrioritizer.Num();
+		ObjectIndexToPrioritizer.SetNumUninitialized(MaxInternalIndex, NoShrinking);
+
+		uint8* BufferData = ObjectIndexToPrioritizer.GetData();
+		BufferData += PreviousSize;
+
+		check((int32)MaxInternalIndex >= PreviousSize);
+		const int32 UninitNum = (int32)MaxInternalIndex - PreviousSize;
+
+		FMemory::Memset(BufferData, FReplicationPrioritization_InvalidNetObjectPrioritizerIndex, UninitNum * sizeof(decltype(ObjectIndexToPrioritizer)::ElementType));
 	}
-	
+
+	// Property initialize DefaultPriorities
+	ResizePrioritiesList(DefaultPriorities, MaxInternalIndex);
+
+	ObjectsWithNewStaticPriority.SetNumBits(MaxInternalIndex);
+}
+
+void FReplicationPrioritization::ResizePrioritiesList(TArray<float>& OutPriorities, FInternalNetRefIndex MaxInternalIndex)
+{
+	const int32 PreviousSize = OutPriorities.Num();
+	OutPriorities.SetNumUninitialized(MaxInternalIndex, EAllowShrinking::No);
+
+	for (int32 PrioIndex = PreviousSize; PrioIndex < (int32)MaxInternalIndex; ++PrioIndex)
 	{
-		DefaultPriorities.SetNumUninitialized(Params.MaxObjectCount, AllowShrinking);
-		float* Priorities = DefaultPriorities.GetData();
-		for (SIZE_T PrioIt = 0, PrioEndIt = Params.MaxObjectCount; PrioIt != PrioEndIt; ++PrioIt)
-		{
-			Priorities[PrioIt] = DefaultPriority;
-		}
+		OutPriorities[PrioIndex] = DefaultPriority;
 	}
+}
 
-	ObjectsWithNewStaticPriority.Init(false, Params.MaxObjectCount);
-	ConnectionInfos.Reserve(Params.Connections->GetMaxConnectionCount());
+void FReplicationPrioritization::OnMaxInternalNetRefIndexIncreased(FInternalNetRefIndex NewMaxInternalIndex)
+{
+	MaxInternalNetRefIndex = NewMaxInternalIndex;
 
-	InitPrioritizers();
+	SetNetObjectListsSize(NewMaxInternalIndex);
+
+	for (FPrioritizerInfo& PrioritizerInfo : PrioritizerInfos)
+	{
+		PrioritizerInfo.Prioritizer->OnMaxInternalNetRefIndexIncreased(NewMaxInternalIndex);
+	}
+
+	for (FPerConnectionInfo& ConnectionInfo : ConnectionInfos)
+	{
+		if (ConnectionInfo.IsValid)
+		{
+			ResizePrioritiesList(ConnectionInfo.Priorities, NewMaxInternalIndex);
+		}
+	}
 }
 
 void FReplicationPrioritization::SetStaticPriority(uint32 ObjectIndex, float NewPrio)
@@ -348,7 +407,7 @@ void FReplicationPrioritization::SetStaticPriority(uint32 ObjectIndex, float New
 			Prioritizer = FReplicationPrioritization_InvalidNetObjectPrioritizerIndex;
 		}
 
-		ObjectsWithNewStaticPriority[ObjectIndex] = true;
+		ObjectsWithNewStaticPriority.SetBit(ObjectIndex);
 	}
 }
 
@@ -405,7 +464,7 @@ bool FReplicationPrioritization::SetPrioritizer(uint32 ObjectIndex, FNetObjectPr
 		{
 			DefaultPriorities[ObjectIndex] = DefaultPriority;
 			Prioritizer = FReplicationPrioritization_InvalidNetObjectPrioritizerIndex;
-			ObjectsWithNewStaticPriority[ObjectIndex] = true;
+			ObjectsWithNewStaticPriority.SetBit(ObjectIndex);
 		}
 	}
 
@@ -621,13 +680,13 @@ void FReplicationPrioritization::UpdatePrioritiesForNewAndDeletedObjects()
 	{
 		NewIndices.Add(ObjectIndex);
 		// Prevent the same index from being added twice
-		this->ObjectsWithNewStaticPriority[ObjectIndex] = false;
+		this->ObjectsWithNewStaticPriority.ClearBit(ObjectIndex);
 	};
 
 	TFunction<void(uint32)> ForEachNewObject = (ConnectionCount > 0 ? AddIndexAndClearFromNewPriority : DoNothing);
 	if (ConnectionCount > 0)
 	{
-		NewIndices.Reserve(FMath::Min(1024U, MaxObjectCount));
+		NewIndices.Reserve(FMath::Min(1024U, MaxInternalNetRefIndex));
 	}
 
 	FNetBitArrayView::ForAllExclusiveBits(ScopedIndices, PrevScopedIndices, ForEachNewObject, ForEachRemovedObject);
@@ -638,7 +697,10 @@ void FReplicationPrioritization::UpdatePrioritiesForNewAndDeletedObjects()
 
 		if (ConnectionCount > 0)
 		{
-			MakeNetBitArrayView(ObjectsWithNewStaticPriority.GetData(), ObjectsWithNewStaticPriority.Num()).ForAllSetBits([this, &NewIndices](uint32 ObjectIndex) { NewIndices.Add(ObjectIndex); });
+			ObjectsWithNewStaticPriority.ForAllSetBits([this, &NewIndices](uint32 ObjectIndex)
+			{ 
+				NewIndices.Add(ObjectIndex); 
+			});
 		}
 
 		// $IRIS TODO: Want ForAllSetBits with clear.
@@ -847,7 +909,8 @@ void FReplicationPrioritization::InitPrioritizers()
 		FNetObjectPrioritizerInitParams InitParams;
 		InitParams.ReplicationSystem = ReplicationSystem;
 		InitParams.Config = (Definition.ConfigClass != nullptr ? NewObject<UNetObjectPrioritizerConfig>((UObject*)GetTransientPackage(), Definition.ConfigClass) : nullptr);
-		InitParams.MaxObjectCount = MaxObjectCount;
+		InitParams.AbsoluteMaxNetObjectCount = NetRefHandleManager->GetMaxActiveObjectCount();
+		InitParams.CurrentMaxInternalIndex = MaxInternalNetRefIndex;
 		InitParams.MaxConnectionCount = Connections->GetMaxConnectionCount();
 
 		Prioritizer->Init(InitParams);
diff --git a/Engine/Source/Runtime/Experimental/Iris/Core/Private/Iris/ReplicationSystem/Prioritization/ReplicationPrioritization.h b/Engine/Source/Runtime/Experimental/Iris/Core/Private/Iris/ReplicationSystem/Prioritization/ReplicationPrioritization.h
--- a/Engine/Source/Runtime/Experimental/Iris/Core/Private/Iris/ReplicationSystem/Prioritization/ReplicationPrioritization.h
+++ b/Engine/Source/Runtime/Experimental/Iris/Core/Private/Iris/ReplicationSystem/Prioritization/ReplicationPrioritization.h
@@ -4,6 +4,7 @@
 
 #include "CoreTypes.h"
 #include "Containers/Array.h"
+#include "Net/Core/NetBitArray.h"
 #include "Iris/ReplicationSystem/Prioritization/NetObjectPrioritizer.h"
 #include "UObject/StrongObjectPtr.h"
 
@@ -16,6 +17,8 @@ namespace UE::Net
 	{
 		class FNetRefHandleManager;
 		class FReplicationConnections;
+
+		typedef uint32 FInternalNetRefIndex;
 	}
 }
 
@@ -27,7 +30,7 @@ struct FReplicationPrioritizationInitParams
 	TObjectPtr<const UReplicationSystem> ReplicationSystem;
 	const FNetRefHandleManager* NetRefHandleManager = nullptr;
 	FReplicationConnections* Connections = nullptr;
-	uint32 MaxObjectCount = 0;
+	FInternalNetRefIndex MaxInternalNetRefIndex = 0;
 };
 
 class FReplicationPrioritization
@@ -36,6 +39,10 @@ public:
 	FReplicationPrioritization();
 
 	void Init(FReplicationPrioritizationInitParams& Params);
+	void Deinit();
+
+	/** Called when the maximum InternalNetRefIndex increased and we need to realloc our lists */
+	void OnMaxInternalNetRefIndexIncreased(FInternalNetRefIndex NewMaxInternalIndex);
 
 	void Prioritize(const FNetBitArrayView& ConnectionsToSend, const FNetBitArrayView& DirtyObjectsThisFrame);
 
@@ -51,6 +58,9 @@ private:
 	class FPrioritizerBatchHelper;
 	class FUpdateDirtyObjectsBatchHelper;
 
+	void SetNetObjectListsSize(FInternalNetRefIndex MaxInternalIndex);
+	void ResizePrioritiesList(TArray<float>& OutPriorities, FInternalNetRefIndex MaxInternalIndex);
+
 	void UpdatePrioritiesForNewAndDeletedObjects();
 	void PrioritizeForConnection(uint32 ConnId, FPrioritizerBatchHelper& BatchHelper, FNetBitArrayView Objects);
 	void SetHighPriorityOnViewTargets(const TArrayView<float>& Priorities, const FReplicationView& View);
@@ -78,20 +88,22 @@ private:
 	static constexpr float DefaultPriority = 1.0f;
 	static constexpr float ViewTargetHighPriority = 1.0E7f;
 
-	TObjectPtr<const UReplicationSystem> ReplicationSystem;
+	TObjectPtr<const UReplicationSystem> ReplicationSystem = nullptr;
 	FReplicationConnections* Connections = nullptr;
 	const FNetRefHandleManager* NetRefHandleManager = nullptr;
+
 	TStrongObjectPtr<UNetObjectPrioritizerDefinitions> PrioritizerDefinitions;
+
 	TArray<FNetObjectPrioritizationInfo> NetObjectPrioritizationInfos;
 	TArray<uint8> ObjectIndexToPrioritizer;
 	TArray<FPrioritizerInfo> PrioritizerInfos;
 	TArray<FPerConnectionInfo> ConnectionInfos;
 	TArray<float> DefaultPriorities;
-	TBitArray<> ObjectsWithNewStaticPriority;
-	uint32 ConnectionCount = 0;
+	FNetBitArray ObjectsWithNewStaticPriority;
 
-	//
-	uint32 MaxObjectCount = 0;
+	FInternalNetRefIndex MaxInternalNetRefIndex = 0;
+
+	uint32 ConnectionCount = 0;
 	uint32 HasNewObjectsWithStaticPriority : 1;
 };
 
diff --git a/Engine/Source/Runtime/Experimental/Iris/Core/Private/Iris/ReplicationSystem/Prioritization/SphereNetObjectPrioritizer.cpp b/Engine/Source/Runtime/Experimental/Iris/Core/Private/Iris/ReplicationSystem/Prioritization/SphereNetObjectPrioritizer.cpp
--- a/Engine/Source/Runtime/Experimental/Iris/Core/Private/Iris/ReplicationSystem/Prioritization/SphereNetObjectPrioritizer.cpp
+++ b/Engine/Source/Runtime/Experimental/Iris/Core/Private/Iris/ReplicationSystem/Prioritization/SphereNetObjectPrioritizer.cpp
@@ -15,6 +15,13 @@ void USphereNetObjectPrioritizer::Init(FNetObjectPrioritizerInitParams& Params)
 	Super::Init(Params);
 }
 
+void USphereNetObjectPrioritizer::Deinit()
+{
+	Super::Deinit();
+
+	Config = nullptr;
+}
+
 void USphereNetObjectPrioritizer::Prioritize(FNetObjectPrioritizationParams& PrioritizationParams)
 {
 	IRIS_PROFILER_SCOPE(USphereNetObjectPrioritizer_Prioritize);
diff --git a/Engine/Source/Runtime/Experimental/Iris/Core/Private/Iris/ReplicationSystem/Prioritization/SphereWithOwnerBoostNetObjectPrioritizer.cpp b/Engine/Source/Runtime/Experimental/Iris/Core/Private/Iris/ReplicationSystem/Prioritization/SphereWithOwnerBoostNetObjectPrioritizer.cpp
--- a/Engine/Source/Runtime/Experimental/Iris/Core/Private/Iris/ReplicationSystem/Prioritization/SphereWithOwnerBoostNetObjectPrioritizer.cpp
+++ b/Engine/Source/Runtime/Experimental/Iris/Core/Private/Iris/ReplicationSystem/Prioritization/SphereWithOwnerBoostNetObjectPrioritizer.cpp
@@ -18,7 +18,18 @@ void USphereWithOwnerBoostNetObjectPrioritizer::Init(FNetObjectPrioritizerInitPa
 	Super::Init(Params);
 
 	ReplicationSystem = Params.ReplicationSystem;
-	AssignedOwningConnectionIndices.Init(Params.MaxObjectCount);
+
+	AssignedOwningConnectionIndices.Init(Params.AbsoluteMaxNetObjectCount);
+}
+
+void USphereWithOwnerBoostNetObjectPrioritizer::Deinit()
+{
+	Super::Deinit();
+
+	ReplicationSystem = nullptr;
+
+	AssignedOwningConnectionIndices.Empty();
+	OwningConnections.Empty();
 }
 
 bool USphereWithOwnerBoostNetObjectPrioritizer::AddObject(uint32 ObjectIndex, FNetObjectPrioritizerAddObjectParams& Params)
diff --git a/Engine/Source/Runtime/Experimental/Iris/Core/Private/Iris/ReplicationSystem/ReplicationConnections.cpp b/Engine/Source/Runtime/Experimental/Iris/Core/Private/Iris/ReplicationSystem/ReplicationConnections.cpp
--- a/Engine/Source/Runtime/Experimental/Iris/Core/Private/Iris/ReplicationSystem/ReplicationConnections.cpp
+++ b/Engine/Source/Runtime/Experimental/Iris/Core/Private/Iris/ReplicationSystem/ReplicationConnections.cpp
@@ -74,6 +74,9 @@ void FReplicationConnections::DestroyReplicationReaderAndWriter(uint32 Connectio
 		ReplicationDataStream->SetReaderAndWriter(nullptr, nullptr);
 	}
 
+	Connection->ReplicationReader->Deinit();
+	Connection->ReplicationWriter->Deinit();
+
 	delete Connection->ReplicationReader;
 	Connection->ReplicationReader = nullptr;
 	delete Connection->ReplicationWriter;
diff --git a/Engine/Source/Runtime/Experimental/Iris/Core/Private/Iris/ReplicationSystem/ReplicationReader.cpp b/Engine/Source/Runtime/Experimental/Iris/Core/Private/Iris/ReplicationSystem/ReplicationReader.cpp
--- a/Engine/Source/Runtime/Experimental/Iris/Core/Private/Iris/ReplicationSystem/ReplicationReader.cpp
+++ b/Engine/Source/Runtime/Experimental/Iris/Core/Private/Iris/ReplicationSystem/ReplicationReader.cpp
@@ -270,7 +270,8 @@ FReplicationReader::FReplicationReader()
 
 FReplicationReader::~FReplicationReader()
 {
-	Deinit();
+	checkf(ReplicatedObjects.IsEmpty(), TEXT("Possible leak detected in FReplicationReader. Nothing should be registered after Deinit()"));
+	checkf(PendingBatches.PendingBatches.IsEmpty(), TEXT("Possible leak detected in FReplicationReader. Nothing should be registered after Deinit()"));
 }
 
 void FReplicationReader::Init(const FReplicationParameters& InParameters)
@@ -320,6 +321,7 @@ void FReplicationReader::Deinit()
 			}
 		}		
 	}
+	PendingBatches.PendingBatches.Empty();
 
 	// Cleanup any allocation stored in the per object info
 	for (auto& ObjectIt : ReplicatedObjects)
diff --git a/Engine/Source/Runtime/Experimental/Iris/Core/Private/Iris/ReplicationSystem/ReplicationSystem.cpp b/Engine/Source/Runtime/Experimental/Iris/Core/Private/Iris/ReplicationSystem/ReplicationSystem.cpp
--- a/Engine/Source/Runtime/Experimental/Iris/Core/Private/Iris/ReplicationSystem/ReplicationSystem.cpp
+++ b/Engine/Source/Runtime/Experimental/Iris/Core/Private/Iris/ReplicationSystem/ReplicationSystem.cpp
@@ -61,7 +61,14 @@ public:
 
 	explicit FReplicationSystemImpl(UReplicationSystem* InReplicationSystem, const UReplicationSystem::FReplicationSystemParams& Params)
 	: ReplicationSystem(InReplicationSystem)
-	, ReplicationSystemInternal(FReplicationSystemInternalInitParams({ InReplicationSystem->GetId(), Params.MaxReplicatedObjectCount, Params.PreAllocatedReplicatedObjectCount, Params.MaxReplicatedWriterObjectCount }))
+	, ReplicationSystemInternal(
+		FReplicationSystemInternalInitParams(
+		{ 
+			.ReplicationSystemId = InReplicationSystem->GetId(),
+			.MaxReplicatedObjectCount = Params.MaxReplicatedObjectCount,
+			.NetChunkedArrayCount = Params.PreAllocatedMemoryBuffersObjectCount,
+			.MaxReplicationWriterObjectCount = Params.MaxReplicationWriterObjectCount,
+		}))
 	{
 	}
 
@@ -90,23 +97,25 @@ public:
 #if !UE_BUILD_SHIPPING
 		IrisDebugHelperDummy = UE::Net::IrisDebugHelper::Init();
 #endif
-
 		const uint32 ReplicationSystemId = ReplicationSystem->GetId();
 
 		FNetRefHandleManager& NetRefHandleManager = ReplicationSystemInternal.GetNetRefHandleManager();
-
 		{
-			FNetRefHandleManager::FInitParams HandleInitParams
-			{
-				.ReplicationSystemId = ReplicationSystemId,
-				.MaxActiveObjectCount = Params.MaxReplicatedObjectCount,
-				.PreAllocatedObjectCount = Params.PreAllocatedReplicatedObjectCount,
-			};
-
-			NetRefHandleManager.Init(HandleInitParams);
+			FNetRefHandleManager::FInitParams NetRefHandleManagerInitParams;
+			NetRefHandleManagerInitParams.ReplicationSystemId = ReplicationSystemId;
+			NetRefHandleManagerInitParams.MaxActiveObjectCount = Params.MaxReplicatedObjectCount;
+			NetRefHandleManagerInitParams.InternalNetRefIndexInitSize = Params.InitialNetObjectListCount;
+			NetRefHandleManagerInitParams.InternalNetRefIndexGrowSize = Params.NetObjectListGrowCount;
+			NetRefHandleManagerInitParams.NetChunkedArrayCount = Params.PreAllocatedMemoryBuffersObjectCount;
+			NetRefHandleManager.Init(NetRefHandleManagerInitParams);
+
+			NetRefHandleManager.GetOnMaxInternalNetRefIndexIncreasedDelegate().AddRaw(this, &FReplicationSystemImpl::OnMaxInternalNetRefIndexIncreased);
 		}
 
-		const uint32 MaxObjectCount =  NetRefHandleManager.GetMaxActiveObjectCount();
+		// Note that Params.MaxReplicatedObjectCount was just a suggestion for the NetRefHandleManager.
+		// From here systems must rely on the NetRefHandleManager configuration.
+		const uint32 AbsoluteMaxObjectCount =  NetRefHandleManager.GetMaxActiveObjectCount();
+		const uint32 CurrentMaxInternalNetRefIndex = NetRefHandleManager.GetCurrentMaxInternalNetRefIndex();
 
 		// DirtyNetObjectTracking is only needed when object replication is allowed
 		if (Params.bAllowObjectReplication)
@@ -116,109 +125,112 @@ public:
 
 			DirtyNetObjectTrackerInitParams.NetRefHandleManager = &NetRefHandleManager;
 			DirtyNetObjectTrackerInitParams.ReplicationSystemId = ReplicationSystemId;
-			DirtyNetObjectTrackerInitParams.MaxObjectCount = MaxObjectCount;
+			DirtyNetObjectTrackerInitParams.MaxInternalNetRefIndex = CurrentMaxInternalNetRefIndex;
 
 			ReplicationSystemInternal.InitDirtyNetObjectTracker(DirtyNetObjectTrackerInitParams);
 		}
 
-		FReplicationStateStorage& StateStorage = ReplicationSystemInternal.GetReplicationStateStorage();
 		{
+			FReplicationStateStorage& StateStorage = ReplicationSystemInternal.GetReplicationStateStorage();
 			FReplicationStateStorageInitParams InitParams;
 			InitParams.ReplicationSystem = ReplicationSystem;
 			InitParams.NetRefHandleManager = &NetRefHandleManager;
-			InitParams.MaxObjectCount = MaxObjectCount;
+			InitParams.MaxObjectCount = AbsoluteMaxObjectCount;
+			InitParams.MaxInternalNetRefIndex = CurrentMaxInternalNetRefIndex;
 			InitParams.MaxConnectionCount = ReplicationSystemInternal.GetConnections().GetMaxConnectionCount();
 			InitParams.MaxDeltaCompressedObjectCount = Params.MaxDeltaCompressedObjectCount;
 			StateStorage.Init(InitParams);
 		}
 
-		FNetObjectGroups& Groups = ReplicationSystemInternal.GetGroups();
 		{
+			FNetObjectGroups& Groups = ReplicationSystemInternal.GetGroups();
 			FNetObjectGroupInitParams InitParams = {};
 			InitParams.NetRefHandleManager = &NetRefHandleManager;
-			InitParams.MaxObjectCount = MaxObjectCount;
+			InitParams.MaxInternalNetRefIndex = CurrentMaxInternalNetRefIndex;
 			InitParams.MaxGroupCount = Params.MaxNetObjectGroupCount;
 
 			Groups.Init(InitParams);
 		}
 
-		FReplicationStateDescriptorRegistry& Registry = ReplicationSystemInternal.GetReplicationStateDescriptorRegistry();
 		{
+			FReplicationStateDescriptorRegistry& Registry = ReplicationSystemInternal.GetReplicationStateDescriptorRegistry();
 			FReplicationStateDescriptorRegistryInitParams InitParams = {};
 			InitParams.ProtocolManager = &ReplicationSystemInternal.GetReplicationProtocolManager();
 
 			Registry.Init(InitParams);
 		}		
 
-		FNetCullDistanceOverrides& NetCullDistanceOverrides = ReplicationSystemInternal.GetNetCullDistanceOverrides();
 		{
+			FNetCullDistanceOverrides& NetCullDistanceOverrides = ReplicationSystemInternal.GetNetCullDistanceOverrides();
 			FNetCullDistanceOverridesInitParams InitParams;
-			InitParams.MaxObjectCount = MaxObjectCount;
+			InitParams.MaxInternalNetRefIndex = CurrentMaxInternalNetRefIndex;
 			NetCullDistanceOverrides.Init(InitParams);
 		}
 
-		FWorldLocations& WorldLocations = ReplicationSystemInternal.GetWorldLocations();
 		{
+			FWorldLocations& WorldLocations = ReplicationSystemInternal.GetWorldLocations();
 			FWorldLocationsInitParams InitParams;
-			InitParams.MaxObjectCount = MaxObjectCount;
+			InitParams.MaxInternalNetRefIndex = CurrentMaxInternalNetRefIndex;
 			WorldLocations.Init(InitParams);
 		}
 	
-		FDeltaCompressionBaselineInvalidationTracker& DeltaCompressionBaselineInvalidationTracker = ReplicationSystemInternal.GetDeltaCompressionBaselineInvalidationTracker();
-		FDeltaCompressionBaselineManager& DeltaCompressionBaselineManager = ReplicationSystemInternal.GetDeltaCompressionBaselineManager();
 		{
+			FDeltaCompressionBaselineInvalidationTracker& DeltaCompressionBaselineInvalidationTracker = ReplicationSystemInternal.GetDeltaCompressionBaselineInvalidationTracker();
 			FDeltaCompressionBaselineInvalidationTrackerInitParams InitParams;
-			InitParams.BaselineManager = &DeltaCompressionBaselineManager;
-			InitParams.MaxObjectCount = MaxObjectCount;
+			InitParams.BaselineManager = &ReplicationSystemInternal.GetDeltaCompressionBaselineManager();
+			InitParams.MaxInternalNetRefIndex = CurrentMaxInternalNetRefIndex;
 			DeltaCompressionBaselineInvalidationTracker.Init(InitParams);
 		}
+
 		{
+			FDeltaCompressionBaselineManager& DeltaCompressionBaselineManager = ReplicationSystemInternal.GetDeltaCompressionBaselineManager();
 			FDeltaCompressionBaselineManagerInitParams InitParams;
-			InitParams.BaselineInvalidationTracker = &DeltaCompressionBaselineInvalidationTracker;
+			InitParams.BaselineInvalidationTracker = &ReplicationSystemInternal.GetDeltaCompressionBaselineInvalidationTracker();
 			InitParams.Connections = &ReplicationSystemInternal.GetConnections();
 			InitParams.NetRefHandleManager = &NetRefHandleManager;
-			InitParams.ReplicationStateStorage = &StateStorage;
-			InitParams.MaxObjectCount = MaxObjectCount;
+			InitParams.ReplicationStateStorage = &ReplicationSystemInternal.GetReplicationStateStorage();
+			InitParams.MaxNetObjectCount = AbsoluteMaxObjectCount;
+			InitParams.MaxInternalNetRefIndex = CurrentMaxInternalNetRefIndex;
 			InitParams.MaxDeltaCompressedObjectCount = Params.MaxDeltaCompressedObjectCount;
 			InitParams.ReplicationSystem = ReplicationSystem;
 			DeltaCompressionBaselineManager.Init(InitParams);
 		}
 
-		FReplicationFiltering& ReplicationFiltering = ReplicationSystemInternal.GetFiltering();
 		{
+			FReplicationFiltering& ReplicationFiltering = ReplicationSystemInternal.GetFiltering();
 			FReplicationFilteringInitParams InitParams;
 			InitParams.ReplicationSystem = ReplicationSystem;
 			InitParams.Connections = &ReplicationSystemInternal.GetConnections();
 			InitParams.NetRefHandleManager = &NetRefHandleManager;
-			InitParams.Groups = &Groups;
+			InitParams.Groups = &ReplicationSystemInternal.GetGroups();
 			InitParams.BaselineInvalidationTracker = &ReplicationSystemInternal.GetDeltaCompressionBaselineInvalidationTracker();
-			InitParams.MaxObjectCount = MaxObjectCount;
+			InitParams.MaxInternalNetRefIndex = CurrentMaxInternalNetRefIndex;
 			InitParams.MaxGroupCount = Params.MaxNetObjectGroupCount;
 			ReplicationFiltering.Init(InitParams);
 		}
 
 		InitDefaultFilteringGroups();
 
-		FReplicationConditionals& ReplicationConditionals = ReplicationSystemInternal.GetConditionals();
 		{
+			FReplicationConditionals& ReplicationConditionals = ReplicationSystemInternal.GetConditionals();
 			FReplicationConditionalsInitParams InitParams = {};
 			InitParams.NetRefHandleManager = &NetRefHandleManager;
 			InitParams.ReplicationConnections = &ReplicationSystemInternal.GetConnections();
-			InitParams.ReplicationFiltering = &ReplicationFiltering;
-			InitParams.NetObjectGroups = &Groups;
+			InitParams.ReplicationFiltering = &ReplicationSystemInternal.GetFiltering();
+			InitParams.NetObjectGroups = &ReplicationSystemInternal.GetGroups();
 			InitParams.BaselineInvalidationTracker = &ReplicationSystemInternal.GetDeltaCompressionBaselineInvalidationTracker();
-			InitParams.MaxObjectCount = MaxObjectCount;
+			InitParams.MaxInternalNetRefIndex = CurrentMaxInternalNetRefIndex;
 			InitParams.MaxConnectionCount = ReplicationSystemInternal.GetConnections().GetMaxConnectionCount();
 			ReplicationConditionals.Init(InitParams);
 		}
 
-		FReplicationPrioritization& ReplicationPrioritization = ReplicationSystemInternal.GetPrioritization();
 		{
+			FReplicationPrioritization& ReplicationPrioritization = ReplicationSystemInternal.GetPrioritization();
 			FReplicationPrioritizationInitParams InitParams;
 			InitParams.ReplicationSystem = ReplicationSystem;
 			InitParams.Connections = &ReplicationSystemInternal.GetConnections();
 			InitParams.NetRefHandleManager = &NetRefHandleManager;
-			InitParams.MaxObjectCount = MaxObjectCount;
+			InitParams.MaxInternalNetRefIndex = CurrentMaxInternalNetRefIndex;
 			ReplicationPrioritization.Init(InitParams);
 		}
 
@@ -236,8 +248,8 @@ public:
 			ReplicationSystemInternal.SetIrisObjectReferencePackageMap(ObjectReferencePackageMap);
 		}
 
-		FNetBlobManager& BlobManager = ReplicationSystemInternal.GetNetBlobManager();
 		{
+			FNetBlobManager& BlobManager = ReplicationSystemInternal.GetNetBlobManager();
 			FNetBlobManagerInitParams InitParams = {};
 			InitParams.ReplicationSystem = ReplicationSystem;
 			InitParams.bSendAttachmentsWithObject = ReplicationSystem->IsServer();
@@ -251,8 +263,8 @@ public:
 
 		ConnectionsPendingPostTickDispatchSend.Init(ReplicationSystemInternal.GetConnections().GetMaxConnectionCount());
 
-		FNetTypeStats& NetStats = ReplicationSystemInternal.GetNetTypeStats();
 		{
+			FNetTypeStats& NetStats = ReplicationSystemInternal.GetNetTypeStats();
 			FNetTypeStats::FInitParams InitParams;
 			InitParams.NetRefHandleManager = &NetRefHandleManager;
 			NetStats.Init(InitParams);
@@ -261,7 +273,16 @@ public:
 
 	void Deinit()
 	{
+		ReplicationSystemInternal.GetPrioritization().Deinit();
+		ReplicationSystemInternal.GetFiltering().Deinit();
 		ReplicationSystemInternal.GetConnections().Deinit();
+		ReplicationSystemInternal.GetDeltaCompressionBaselineManager().Deinit();
+		ReplicationSystemInternal.GetReplicationStateStorage().Deinit();
+
+		if (ReplicationSystemInternal.IsDirtyNetObjectTrackerInitialized())
+		{
+			ReplicationSystemInternal.GetDirtyNetObjectTracker().Deinit();
+		}
 
 		// Reset replication bridge
 		ReplicationSystemInternal.GetReplicationBridge()->Deinitialize();
@@ -273,10 +294,23 @@ public:
 			ReplicationSystemInternal.SetIrisObjectReferencePackageMap(static_cast<UIrisObjectReferencePackageMap*>(nullptr));
 		}
 
-		ReplicationSystemInternal.GetDeltaCompressionBaselineManager().Deinit();
+		ReplicationSystemInternal.GetNetRefHandleManager().GetOnMaxInternalNetRefIndexIncreasedDelegate().RemoveAll(this);
 		ReplicationSystemInternal.GetNetRefHandleManager().Deinit();
 	}
 
+	void OnMaxInternalNetRefIndexIncreased(FInternalNetRefIndex NewMaxInternalIndex)
+	{
+		ReplicationSystemInternal.GetReplicationStateStorage().OnMaxInternalNetRefIndexIncreased(NewMaxInternalIndex);
+		ReplicationSystemInternal.GetGroups().OnMaxInternalNetRefIndexIncreased(NewMaxInternalIndex);
+		ReplicationSystemInternal.GetNetCullDistanceOverrides().OnMaxInternalNetRefIndexIncreased(NewMaxInternalIndex);
+		ReplicationSystemInternal.GetWorldLocations().OnMaxInternalNetRefIndexIncreased(NewMaxInternalIndex);
+		ReplicationSystemInternal.GetDeltaCompressionBaselineInvalidationTracker().OnMaxInternalNetRefIndexIncreased(NewMaxInternalIndex);
+		ReplicationSystemInternal.GetDeltaCompressionBaselineManager().OnMaxInternalNetRefIndexIncreased(NewMaxInternalIndex);
+		ReplicationSystemInternal.GetFiltering().OnMaxInternalNetRefIndexIncreased(NewMaxInternalIndex);
+		ReplicationSystemInternal.GetConditionals().OnMaxInternalNetRefIndexIncreased(NewMaxInternalIndex);
+		ReplicationSystemInternal.GetPrioritization().OnMaxInternalNetRefIndexIncreased(NewMaxInternalIndex);
+	}
+
 	void StartPreSendUpdate()
 	{
 		// Block unsupported operations until SendUpdate is finished
@@ -507,9 +541,8 @@ public:
 			Params.ReplicationSystem = ReplicationSystem;
 			Params.PacketSendWindowSize = 256;
 			Params.ConnectionId = ConnectionId;
-			Params.MaxActiveReplicatedObjectCount = ReplicationSystemInternal.GetNetRefHandleManager().GetMaxActiveObjectCount();
-			Params.PreAllocatedReplicatedObjectCount = ReplicationSystemInternal.GetNetRefHandleManager().GetPreAllocatedObjectCount();
-			Params.MaxReplicatedWriterObjectCount = ReplicationSystemInternal.GetInitParams().MaxReplicatedWriterObjectCount;
+			Params.MaxInternalNetRefIndex = ReplicationSystemInternal.GetNetRefHandleManager().GetCurrentMaxInternalNetRefIndex();
+			Params.MaxReplicationWriterObjectCount = ReplicationSystemInternal.GetInitParams().MaxReplicationWriterObjectCount;
 
 			/** 
 			  * Currently we expect all objects to be replicated from server to client.
@@ -592,9 +625,9 @@ public:
 		const FNetBitArray& ValidConnections = Connections.GetValidConnections();
 		ValidConnections.ForAllSetBits(UpdateUnresolvableReferenceTracking);
 	}
-};
+}; // end class FReplicationSystemImpl
 
-}
+} // end namespace UE::Net::Private
 
 UReplicationSystem::UReplicationSystem()
 : Super()
@@ -1729,7 +1762,7 @@ UReplicationSystem* FReplicationSystemFactory::CreateReplicationSystem(const URe
 			MaxReplicationSystemId = ReplicationSystemId;
 		}
 
-		UE_LOG(LogIris, Display, TEXT("Iris ReplicationSystem[%i] is created"), ReplicationSystemId);
+		UE_LOG(LogIris, Display, TEXT("Iris ReplicationSystem[%i]: %s (0x%p) is created"), ReplicationSystemId, *ReplicationSystem->GetName(), ReplicationSystem);
 
 		ReplicationSystem->Init(ReplicationSystemId, Params);
 
@@ -1754,7 +1787,7 @@ void FReplicationSystemFactory::DestroyReplicationSystem(UReplicationSystem* Sys
 
 	const uint32 Id = System->GetId();
 
-	UE_LOG(LogIris, Display, TEXT("Iris ReplicationSystem[%i] is about to be destroyed"), Id);
+	UE_LOG(LogIris, Display, TEXT("Iris ReplicationSystem[%i]: %s (0x%p) is about to be destroyed"), Id, *System->GetName(), System);
 
 	if (Id < MaxReplicationSystemCount)
 	{
diff --git a/Engine/Source/Runtime/Experimental/Iris/Core/Private/Iris/ReplicationSystem/ReplicationSystemInternal.h b/Engine/Source/Runtime/Experimental/Iris/Core/Private/Iris/ReplicationSystem/ReplicationSystemInternal.h
--- a/Engine/Source/Runtime/Experimental/Iris/Core/Private/Iris/ReplicationSystem/ReplicationSystemInternal.h
+++ b/Engine/Source/Runtime/Experimental/Iris/Core/Private/Iris/ReplicationSystem/ReplicationSystemInternal.h
@@ -35,8 +35,8 @@ struct FReplicationSystemInternalInitParams
 {
 	uint32 ReplicationSystemId;
 	uint32 MaxReplicatedObjectCount;
-	uint32 PreAllocatedReplicatedObjectCount;
-	uint32 MaxReplicatedWriterObjectCount;
+	uint32 NetChunkedArrayCount;
+	uint32 MaxReplicationWriterObjectCount;
 };
 
 class FReplicationSystemInternal
@@ -58,6 +58,7 @@ public:
 	const FNetRefHandleManager& GetNetRefHandleManager() const { return NetRefHandleManager; }
 
 	void InitDirtyNetObjectTracker(const struct FDirtyNetObjectTrackerInitParams& Params) { DirtyNetObjectTracker.Init(Params); }
+	bool IsDirtyNetObjectTrackerInitialized() const { return DirtyNetObjectTracker.IsInit(); }
 	FDirtyNetObjectTracker& GetDirtyNetObjectTracker() { checkf(DirtyNetObjectTracker.IsInit(), TEXT("Not allowed to access the DirtyNetObjectTracker unless object replication is enabled.")); return DirtyNetObjectTracker; }
 
 	FReplicationStateDescriptorRegistry& GetReplicationStateDescriptorRegistry() { return ReplicationStateDescriptorRegistry; }
diff --git a/Engine/Source/Runtime/Experimental/Iris/Core/Private/Iris/ReplicationSystem/ReplicationTypes.h b/Engine/Source/Runtime/Experimental/Iris/Core/Private/Iris/ReplicationSystem/ReplicationTypes.h
--- a/Engine/Source/Runtime/Experimental/Iris/Core/Private/Iris/ReplicationSystem/ReplicationTypes.h
+++ b/Engine/Source/Runtime/Experimental/Iris/Core/Private/Iris/ReplicationSystem/ReplicationTypes.h
@@ -5,14 +5,18 @@
 
 class UReplicationSystem;
 
+namespace UE::Net::Private
+{
+	typedef uint32 FInternalNetRefIndex;
+}
+
 namespace UE::Net::Private
 {
 
 struct FReplicationParameters
 {
-	uint32 MaxActiveReplicatedObjectCount = 0;
-	uint32 PreAllocatedReplicatedObjectCount = 0;
-	uint32 MaxReplicatedWriterObjectCount = 0;
+	FInternalNetRefIndex MaxInternalNetRefIndex = 0;
+	uint32 MaxReplicationWriterObjectCount = 0;
 	uint32 PacketSendWindowSize = 0;
 	uint32 ConnectionId = 0;
 	UReplicationSystem* ReplicationSystem = nullptr;
diff --git a/Engine/Source/Runtime/Experimental/Iris/Core/Private/Iris/ReplicationSystem/ReplicationWriter.cpp b/Engine/Source/Runtime/Experimental/Iris/Core/Private/Iris/ReplicationSystem/ReplicationWriter.cpp
--- a/Engine/Source/Runtime/Experimental/Iris/Core/Private/Iris/ReplicationSystem/ReplicationWriter.cpp
+++ b/Engine/Source/Runtime/Experimental/Iris/Core/Private/Iris/ReplicationSystem/ReplicationWriter.cpp
@@ -361,7 +361,8 @@ void FReplicationWriter::Init(const FReplicationParameters& InParameters)
 	// Store copy of parameters
 	Parameters = InParameters;
 
-	UE_LOG(LogIris, Log, TEXT("ReplicationWriter: Configured with MaxActiveReplicatedObjectCount=%d, PreallocatedObjectCount=%d and MaxReplicatedWriterObjectCount=%d."), Parameters.MaxActiveReplicatedObjectCount, Parameters.PreAllocatedReplicatedObjectCount, Parameters.MaxReplicatedWriterObjectCount);
+	UE_LOG(LogIris, Log, TEXT("ReplicationWriter: Configured with MaxInternalNetRefIndex=%d and MaxReplicationWriterObjectCount=%d."), 
+		Parameters.MaxInternalNetRefIndex, Parameters.MaxReplicationWriterObjectCount);
 
 	// Cache internal systems
 	ReplicationSystemInternal = Parameters.ReplicationSystem->GetReplicationSystemInternal();
@@ -376,19 +377,18 @@ void FReplicationWriter::Init(const FReplicationParameters& InParameters)
 	NetObjectBlobHandler = NetBlobManager->GetNetObjectBlobHandler();
 	NetTypeStats = &ReplicationSystemInternal->GetNetTypeStats();
 
-	// Init book keeping
-	const int32 PreAllocatedBufferSize = Parameters.MaxReplicatedWriterObjectCount;
-	ReplicatedObjects.SetNumZeroed(PreAllocatedBufferSize);
-	ReplicatedObjectsRecordInfoLists.SetNumZeroed(PreAllocatedBufferSize);
-	SchedulingPriorities.SetNumZeroed(PreAllocatedBufferSize);
-
+	// See if we want to limit the amount of objects able to do property replication, otherwise follow the system max and grow as needed
+	const uint32 MaxSupportedObjects = Parameters.MaxReplicationWriterObjectCount > 0 ? Parameters.MaxReplicationWriterObjectCount : Parameters.MaxInternalNetRefIndex;
+	ReplicatedObjects.SetNumZeroed(MaxSupportedObjects);
+	ReplicatedObjectsRecordInfoLists.SetNumZeroed(MaxSupportedObjects);
+	SchedulingPriorities.SetNumZeroed(MaxSupportedObjects);
+	
 	// NOTE: Currently disabled because FReplicationWriter until the performance impact of TNetChunkedArray can be measured on the server.
 	//OnLargestIndexIncreaseHandle = NetRefHandleManager->GetLargestIndexIncreaseDelegate().AddRaw(this, &FReplicationWriter::OnLargestIndexIncrease);
 
-	ObjectsPendingDestroy.Init(Parameters.MaxActiveReplicatedObjectCount);
-	ObjectsWithDirtyChanges.Init(Parameters.MaxActiveReplicatedObjectCount);
-	ObjectsInScope.Init(Parameters.MaxActiveReplicatedObjectCount);	
-	WriteContext.ObjectsWrittenThisPacket.Init(Parameters.MaxActiveReplicatedObjectCount);
+	SetNetObjectListsSize(Parameters.MaxInternalNetRefIndex);
+
+	NetRefHandleManager->GetOnMaxInternalNetRefIndexIncreasedDelegate().AddRaw(this, &FReplicationWriter::OnMaxInternalNetRefIndexIncreased);
 
 	// Attachments init
 	SetupReplicationInfoForAttachmentsToObjectsNotInScope();
@@ -396,6 +396,32 @@ void FReplicationWriter::Init(const FReplicationParameters& InParameters)
 	bReplicationEnabled = false;
 }
 
+void FReplicationWriter::Deinit()
+{
+	NetRefHandleManager->GetOnMaxInternalNetRefIndexIncreasedDelegate().RemoveAll(this);
+}
+
+void FReplicationWriter::SetNetObjectListsSize(FInternalNetRefIndex NewMaxInternalIndex)
+{
+	ObjectsPendingDestroy.SetNumBits(NewMaxInternalIndex);
+	ObjectsWithDirtyChanges.SetNumBits(NewMaxInternalIndex);
+	ObjectsInScope.SetNumBits(NewMaxInternalIndex);
+	WriteContext.ObjectsWrittenThisPacket.SetNumBits(NewMaxInternalIndex);
+}
+
+void FReplicationWriter::OnMaxInternalNetRefIndexIncreased(FInternalNetRefIndex NewMaxInternalIndex)
+{
+	// Only grow the objects if no limits were set
+	if (Parameters.MaxReplicationWriterObjectCount == 0)
+	{
+		ReplicatedObjects.SetNumZeroed(NewMaxInternalIndex);
+		ReplicatedObjectsRecordInfoLists.SetNumZeroed(NewMaxInternalIndex);
+		SchedulingPriorities.SetNumZeroed(NewMaxInternalIndex);
+	}
+
+	SetNetObjectListsSize(NewMaxInternalIndex);
+}
+
 void FReplicationWriter::GetInitialChangeMask(ChangeMaskStorageType* ChangeMaskData, const FReplicationProtocol* Protocol)
 {
 	FNetBitArrayView ChangeMask(ChangeMaskData, Protocol->ChangeMaskBitCount, FNetBitArrayView::NoResetNoValidate);
@@ -1572,7 +1598,7 @@ void FReplicationWriter::HandleDroppedRecord(const FReplicationRecord::FRecordIn
 void FReplicationWriter::ProcessDeliveryNotification(EPacketDeliveryStatus PacketDeliveryStatus)
 {
 #if UE_NET_VALIDATE_REPLICATION_RECORD
-	check(s_ValidateReplicationRecord(&ReplicationRecord, Parameters.MaxActiveReplicatedObjectCount + 1U, true));
+	check(s_ValidateReplicationRecord(&ReplicationRecord, NetRefHandleManager->GetMaxActiveObjectCount() + 1U, true));
 #endif
 
 	const uint32 RecordCount = ReplicationRecord.PopRecord();
@@ -3188,7 +3214,7 @@ UDataStream::EWriteResult FReplicationWriter::BeginWrite(const UDataStream::FBeg
 		// Allocate space for indices to send
 		// This should be allocated from frame temp allocator and be cleaned up end of frame, we might want this data to persist over multiple write calls but not over multiple frames 
 		// https://jira.it.epicgames.com/browse/UE-127374	
-		WriteContext.ScheduledObjectInfos = reinterpret_cast<FScheduleObjectInfo*>(FMemory::Malloc(sizeof(FScheduleObjectInfo) * Parameters.MaxActiveReplicatedObjectCount));
+		WriteContext.ScheduledObjectInfos = reinterpret_cast<FScheduleObjectInfo*>(FMemory::Malloc(sizeof(FScheduleObjectInfo) * NetRefHandleManager->GetCurrentMaxInternalNetRefIndex()));
 		WriteContext.ScheduledObjectCount = ScheduleObjects(WriteContext.ScheduledObjectInfos);
 	}
 
@@ -3352,7 +3378,7 @@ UDataStream::EWriteResult FReplicationWriter::Write(FNetSerializationContext& Co
 		ReplicationRecord.PushRecord(ReplicationInfoCount);
 
 #if UE_NET_VALIDATE_REPLICATION_RECORD
-		check(s_ValidateReplicationRecord(&ReplicationRecord, Parameters.MaxActiveReplicatedObjectCount + 1U, false));
+		check(s_ValidateReplicationRecord(&ReplicationRecord, NetRefHandleManager->GetMaxActiveObjectCount() + 1U, false));
 #endif
 
 #if UE_NET_TRACE_ENABLED
diff --git a/Engine/Source/Runtime/Experimental/Iris/Core/Private/Iris/ReplicationSystem/ReplicationWriter.h b/Engine/Source/Runtime/Experimental/Iris/Core/Private/Iris/ReplicationSystem/ReplicationWriter.h
--- a/Engine/Source/Runtime/Experimental/Iris/Core/Private/Iris/ReplicationSystem/ReplicationWriter.h
+++ b/Engine/Source/Runtime/Experimental/Iris/Core/Private/Iris/ReplicationSystem/ReplicationWriter.h
@@ -150,6 +150,8 @@ public:
 	// Init
 	void Init(const FReplicationParameters& InParameters);
 
+	void Deinit();
+
 	// Update new or existing/destroyed 
 	void UpdateScope(const FNetBitArrayView& ScopedObjects);
 
@@ -402,6 +404,9 @@ private:
 
 private:
 
+	void SetNetObjectListsSize(FInternalNetRefIndex NewMaxInternalIndex);
+	void OnMaxInternalNetRefIndexIncreased(FInternalNetRefIndex NewMaxInternalIndex);
+
 	uint32 GetDefaultFlushFlags() const;
 	uint32 GetFlushStatus(uint32 InternalIndex, const FReplicationInfo& Info, uint32 FlushFlagsToTest = EFlushFlags::FlushFlags_Default) const;
 	void SetPendingDestroyOrSubObjectPendingDestroyState(uint32 Index, FReplicationInfo& Info);
diff --git a/Engine/Source/Runtime/Experimental/Iris/Core/Private/Iris/ReplicationSystem/WorldLocations.cpp b/Engine/Source/Runtime/Experimental/Iris/Core/Private/Iris/ReplicationSystem/WorldLocations.cpp
--- a/Engine/Source/Runtime/Experimental/Iris/Core/Private/Iris/ReplicationSystem/WorldLocations.cpp
+++ b/Engine/Source/Runtime/Experimental/Iris/Core/Private/Iris/ReplicationSystem/WorldLocations.cpp
@@ -2,6 +2,7 @@
 
 #include "Iris/ReplicationSystem/WorldLocations.h"
 #include "Iris/Core/IrisMemoryTracker.h"
+#include "Iris/ReplicationSystem/NetRefHandleManager.h"
 
 #include UE_INLINE_GENERATED_CPP_BY_NAME(WorldLocations)
 
@@ -10,14 +11,21 @@ namespace UE::Net
 
 void FWorldLocations::Init(const FWorldLocationsInitParams& InitParams)
 {
-	ValidInfoIndexes.Init(InitParams.MaxObjectCount);
-	ObjectsWithDirtyInfo.Init(InitParams.MaxObjectCount);
-	ObjectsRequiringFrequentWorldLocationUpdate.Init(InitParams.MaxObjectCount);
+	ValidInfoIndexes.Init(InitParams.MaxInternalNetRefIndex);
+	ObjectsWithDirtyInfo.Init(InitParams.MaxInternalNetRefIndex);
+	ObjectsRequiringFrequentWorldLocationUpdate.Init(InitParams.MaxInternalNetRefIndex);
 
 	MinWorldPos = GetDefault<UWorldLocationsConfig>()->MinPos;
 	MaxWorldPos = GetDefault<UWorldLocationsConfig>()->MaxPos;
 }
 
+void FWorldLocations::OnMaxInternalNetRefIndexIncreased(UE::Net::Private::FInternalNetRefIndex NewMaxInternalIndex)
+{
+	ValidInfoIndexes.SetNumBits(NewMaxInternalIndex);
+	ObjectsWithDirtyInfo.SetNumBits(NewMaxInternalIndex);
+	ObjectsRequiringFrequentWorldLocationUpdate.SetNumBits(NewMaxInternalIndex);
+}
+
 void FWorldLocations::InitObjectInfoCache(uint32 ObjectIndex)
 {
 	if (ValidInfoIndexes.IsBitSet(ObjectIndex))
diff --git a/Engine/Source/Runtime/Experimental/Iris/Core/Public/Iris/ReplicationSystem/Filtering/FilterOutNetObjectFilter.h b/Engine/Source/Runtime/Experimental/Iris/Core/Public/Iris/ReplicationSystem/Filtering/FilterOutNetObjectFilter.h
--- a/Engine/Source/Runtime/Experimental/Iris/Core/Public/Iris/ReplicationSystem/Filtering/FilterOutNetObjectFilter.h
+++ b/Engine/Source/Runtime/Experimental/Iris/Core/Public/Iris/ReplicationSystem/Filtering/FilterOutNetObjectFilter.h
@@ -19,6 +19,8 @@ class UFilterOutNetObjectFilter final : public UNetObjectFilter
 protected:
 	// UNetObjectFilter interface
 	IRISCORE_API virtual void OnInit(const FNetObjectFilterInitParams&) override;
+	virtual void OnDeinit() override {}
+	virtual void OnMaxInternalNetRefIndexIncreased(uint32 NewMaxInternalIndex) override {}
 	IRISCORE_API virtual bool AddObject(uint32 ObjectIndex, FNetObjectFilterAddObjectParams&) override;
 	IRISCORE_API virtual void RemoveObject(uint32 ObjectIndex, const FNetObjectFilteringInfo&) override;
 	IRISCORE_API virtual void Filter(FNetObjectFilteringParams&) override;
diff --git a/Engine/Source/Runtime/Experimental/Iris/Core/Public/Iris/ReplicationSystem/Filtering/NetObjectConnectionFilter.h b/Engine/Source/Runtime/Experimental/Iris/Core/Public/Iris/ReplicationSystem/Filtering/NetObjectConnectionFilter.h
--- a/Engine/Source/Runtime/Experimental/Iris/Core/Public/Iris/ReplicationSystem/Filtering/NetObjectConnectionFilter.h
+++ b/Engine/Source/Runtime/Experimental/Iris/Core/Public/Iris/ReplicationSystem/Filtering/NetObjectConnectionFilter.h
@@ -51,6 +51,8 @@ protected:
 
 	// UNetObjectFilter interface
 	IRISCORE_API virtual void OnInit(const FNetObjectFilterInitParams&) override;
+	IRISCORE_API virtual void OnDeinit() override;
+	virtual void OnMaxInternalNetRefIndexIncreased(uint32 NewMaxInternalIndex) override {}
 	IRISCORE_API virtual void AddConnection(uint32 ConnectionId) override;
 	IRISCORE_API virtual void RemoveConnection(uint32 ConnectionId) override;
 	IRISCORE_API virtual bool AddObject(uint32 ObjectIndex, FNetObjectFilterAddObjectParams&) override;
@@ -63,5 +65,7 @@ protected:
 	TArray<uint32> LocalToNetRefIndex;
 	TArray<FPerConnectionInfo> PerConnectionInfos;
 	UE::Net::FNetBitArray UsedLocalInfoIndices;
+
+	uint32 MaxInternalIndex = 0;
 	bool bObjectRemoved = false;
 };
diff --git a/Engine/Source/Runtime/Experimental/Iris/Core/Public/Iris/ReplicationSystem/Filtering/NetObjectFilter.h b/Engine/Source/Runtime/Experimental/Iris/Core/Public/Iris/ReplicationSystem/Filtering/NetObjectFilter.h
--- a/Engine/Source/Runtime/Experimental/Iris/Core/Public/Iris/ReplicationSystem/Filtering/NetObjectFilter.h
+++ b/Engine/Source/Runtime/Experimental/Iris/Core/Public/Iris/ReplicationSystem/Filtering/NetObjectFilter.h
@@ -15,6 +15,11 @@ namespace UE::Net
 	typedef uint32 FNetObjectFilterHandle;
 	struct FReplicationInstanceProtocol;
 	struct FReplicationProtocol;
+
+	namespace Private
+	{
+		typedef uint32 FInternalNetRefIndex;
+	}
 }
 
 namespace UE::Net
@@ -119,8 +124,10 @@ public:
 	TObjectPtr<UReplicationSystem> ReplicationSystem = nullptr;
 	/** Optional config as set in the FNetObjectFilterDefinition. */
 	UNetObjectFilterConfig* Config = nullptr;
-	/** The maximum number of objects in the system. */
-	uint32 MaxObjectCount = 0;
+	/** The maximum number of replicated objects in the system. */
+	uint32 AbsoluteMaxNetObjectCount = 0;
+	/** The current maximum replicated objects referenced by an index (may grow at runtime). */
+	uint32 CurrentMaxInternalIndex = 0;
 	/** The maximum number of connections in the system. */
 	uint32 MaxConnectionCount = 0;
 };
@@ -166,6 +173,9 @@ class UNetObjectFilter : public UObject
 
 public:
 	IRISCORE_API void Init(const FNetObjectFilterInitParams& Params);
+	IRISCORE_API void Deinit();
+
+	IRISCORE_API void MaxInternalNetRefIndexIncreased(UE::Net::Private::FInternalNetRefIndex MaxInternalIndex, TArrayView<FNetObjectFilteringInfo> NewFilterInfoView);
 	
 	/** A new connection has been added. An opportunity for the filter to allocate per connection info. */
 	IRISCORE_API virtual void AddConnection(uint32 ConnectionId);
@@ -226,6 +236,12 @@ protected:
 	/** Called right after constructor for enabled filters. Must be overriden. */
 	IRISCORE_API virtual void OnInit(const FNetObjectFilterInitParams&) PURE_VIRTUAL(OnInit, );
 
+	/** Called when the replication system is shutting down. Use this to remove references to other systems */
+	IRISCORE_API virtual void OnDeinit() PURE_VIRTUAL(OnDeinit);
+
+	/** Called when the maximum InternalNetRefIndex increased and we need to realloc our lists */
+	IRISCORE_API virtual void OnMaxInternalNetRefIndexIncreased(uint32 NewMaxInternalIndex) PURE_VIRTUAL(OnMaxInternalNetRefIndexIncreased);
+
 	/* Returns the filtering info for this object if it's handled by this filter, nullptr otherwise. */
 	IRISCORE_API FNetObjectFilteringInfo* GetFilteringInfo(uint32 ObjectIndex);
 
diff --git a/Engine/Source/Runtime/Experimental/Iris/Core/Public/Iris/ReplicationSystem/Filtering/NetObjectGridFilter.h b/Engine/Source/Runtime/Experimental/Iris/Core/Public/Iris/ReplicationSystem/Filtering/NetObjectGridFilter.h
--- a/Engine/Source/Runtime/Experimental/Iris/Core/Public/Iris/ReplicationSystem/Filtering/NetObjectGridFilter.h
+++ b/Engine/Source/Runtime/Experimental/Iris/Core/Public/Iris/ReplicationSystem/Filtering/NetObjectGridFilter.h
@@ -95,6 +95,8 @@ class UNetObjectGridFilter : public UNetObjectFilter
 protected:
 	// UNetObjectFilter interface
 	IRISCORE_API virtual void OnInit(const FNetObjectFilterInitParams&) override;
+	IRISCORE_API virtual void OnDeinit() override;
+	IRISCORE_API virtual void OnMaxInternalNetRefIndexIncreased(uint32 NewMaxInternalIndex) override;
 	IRISCORE_API virtual void AddConnection(uint32 ConnectionId) override;
 	IRISCORE_API virtual void RemoveConnection(uint32 ConnectionId) override;
 	IRISCORE_API virtual bool AddObject(uint32 ObjectIndex, FNetObjectFilterAddObjectParams&) override;
@@ -308,6 +310,7 @@ class UNetObjectGridWorldLocFilter : public UNetObjectGridFilter
 protected:
 
 	virtual void OnInit(const FNetObjectFilterInitParams&) override;
+	virtual void OnDeinit() override;
 	virtual void PreFilter(FNetObjectPreFilteringParams&) override;
 	virtual void UpdateObjectInfo(FPerObjectInfo& PerObjectInfo, const UNetObjectGridFilter::FObjectLocationInfo& ObjectLocationInfo, const UE::Net::FReplicationInstanceProtocol* InstanceProtocol) override;
 	virtual bool BuildObjectInfo(uint32 ObjectIndex, FNetObjectFilterAddObjectParams& Params) override;
diff --git a/Engine/Source/Runtime/Experimental/Iris/Core/Public/Iris/ReplicationSystem/Filtering/NopNetObjectFilter.h b/Engine/Source/Runtime/Experimental/Iris/Core/Public/Iris/ReplicationSystem/Filtering/NopNetObjectFilter.h
--- a/Engine/Source/Runtime/Experimental/Iris/Core/Public/Iris/ReplicationSystem/Filtering/NopNetObjectFilter.h
+++ b/Engine/Source/Runtime/Experimental/Iris/Core/Public/Iris/ReplicationSystem/Filtering/NopNetObjectFilter.h
@@ -19,6 +19,8 @@ class UNopNetObjectFilter final : public UNetObjectFilter
 protected:
 	// UNetObjectFilter interface
 	IRISCORE_API virtual void OnInit(const FNetObjectFilterInitParams&) override;
+	virtual void OnDeinit() override {}
+	virtual void OnMaxInternalNetRefIndexIncreased(uint32 NewMaxInternalIndex) override {}
 	IRISCORE_API virtual bool AddObject(uint32 ObjectIndex, FNetObjectFilterAddObjectParams&) override;
 	IRISCORE_API virtual void RemoveObject(uint32 ObjectIndex, const FNetObjectFilteringInfo&) override;
 	IRISCORE_API virtual void Filter(FNetObjectFilteringParams&) override;
diff --git a/Engine/Source/Runtime/Experimental/Iris/Core/Public/Iris/ReplicationSystem/NetCullDistanceOverrides.h b/Engine/Source/Runtime/Experimental/Iris/Core/Public/Iris/ReplicationSystem/NetCullDistanceOverrides.h
--- a/Engine/Source/Runtime/Experimental/Iris/Core/Public/Iris/ReplicationSystem/NetCullDistanceOverrides.h
+++ b/Engine/Source/Runtime/Experimental/Iris/Core/Public/Iris/ReplicationSystem/NetCullDistanceOverrides.h
@@ -6,12 +6,17 @@
 #include "Containers/ChunkedArray.h"
 #include "Net/Core/NetBitArray.h"
 
+namespace UE::Net::Private
+{
+	typedef uint32 FInternalNetRefIndex;
+}
+
 namespace UE::Net
 {
 
 struct FNetCullDistanceOverridesInitParams
 {
-	uint32 MaxObjectCount = 0;
+	uint32 MaxInternalNetRefIndex = 0;
 };
 
 class FNetCullDistanceOverrides
@@ -28,13 +33,15 @@ public:
 	/** Returns the object's squared cull distance override if it's valid or the provided DefaultValue if it's not. */
 	float GetCullDistanceSqr(uint32 ObjectIndex, float DefaultValue) const;
 
-public:
 	/** Remove cull distance override for object. */
 	bool ClearCullDistanceSqr(uint32 ObjectIndex);
 
 	/** Set cull distance override for object. */
 	void SetCullDistanceSqr(uint32 ObjectIndex, float CullDistSqr);
 
+	/** Called when the maximum InternalNetRefIndex increased and we need to realloc our lists */
+	void OnMaxInternalNetRefIndexIncreased(UE::Net::Private::FInternalNetRefIndex NewMaxInternalIndex);
+
 private:
 	enum : unsigned
 	{
diff --git a/Engine/Source/Runtime/Experimental/Iris/Core/Public/Iris/ReplicationSystem/Prioritization/LocationBasedNetObjectPrioritizer.h b/Engine/Source/Runtime/Experimental/Iris/Core/Public/Iris/ReplicationSystem/Prioritization/LocationBasedNetObjectPrioritizer.h
--- a/Engine/Source/Runtime/Experimental/Iris/Core/Public/Iris/ReplicationSystem/Prioritization/LocationBasedNetObjectPrioritizer.h
+++ b/Engine/Source/Runtime/Experimental/Iris/Core/Public/Iris/ReplicationSystem/Prioritization/LocationBasedNetObjectPrioritizer.h
@@ -21,6 +21,8 @@ class ULocationBasedNetObjectPrioritizer : public UNetObjectPrioritizer
 protected:
 	// UNetObjectPrioritizer interface
 	IRISCORE_API virtual void Init(FNetObjectPrioritizerInitParams& Params) override;
+	IRISCORE_API virtual void Deinit() override;
+	IRISCORE_API virtual void OnMaxInternalNetRefIndexIncreased(uint32 NewMaxInternalIndex) override;
 	IRISCORE_API virtual bool AddObject(uint32 ObjectIndex, FNetObjectPrioritizerAddObjectParams& Params) override;
 	IRISCORE_API virtual void RemoveObject(uint32 ObjectIndex, const FNetObjectPrioritizationInfo& Info) override;
 	IRISCORE_API virtual void UpdateObjects(FNetObjectPrioritizerUpdateParams&) override;
diff --git a/Engine/Source/Runtime/Experimental/Iris/Core/Public/Iris/ReplicationSystem/Prioritization/NetObjectCountLimiter.h b/Engine/Source/Runtime/Experimental/Iris/Core/Public/Iris/ReplicationSystem/Prioritization/NetObjectCountLimiter.h
--- a/Engine/Source/Runtime/Experimental/Iris/Core/Public/Iris/ReplicationSystem/Prioritization/NetObjectCountLimiter.h
+++ b/Engine/Source/Runtime/Experimental/Iris/Core/Public/Iris/ReplicationSystem/Prioritization/NetObjectCountLimiter.h
@@ -72,6 +72,8 @@ class UNetObjectCountLimiter : public UNetObjectPrioritizer
 protected:
 	// UNetObjectPrioritizer interface
 	IRISCORE_API virtual void Init(FNetObjectPrioritizerInitParams&) override;
+	IRISCORE_API virtual void Deinit() override;
+	IRISCORE_API virtual void OnMaxInternalNetRefIndexIncreased(uint32 NewMaxInternalIndex) override {}
 	IRISCORE_API virtual void AddConnection(uint32 ConnectionId) override;
 	IRISCORE_API virtual void RemoveConnection(uint32 ConnectionId) override;
 	IRISCORE_API virtual bool AddObject(uint32 ObjectIndex, FNetObjectPrioritizerAddObjectParams&) override;
diff --git a/Engine/Source/Runtime/Experimental/Iris/Core/Public/Iris/ReplicationSystem/Prioritization/NetObjectPrioritizer.h b/Engine/Source/Runtime/Experimental/Iris/Core/Public/Iris/ReplicationSystem/Prioritization/NetObjectPrioritizer.h
--- a/Engine/Source/Runtime/Experimental/Iris/Core/Public/Iris/ReplicationSystem/Prioritization/NetObjectPrioritizer.h
+++ b/Engine/Source/Runtime/Experimental/Iris/Core/Public/Iris/ReplicationSystem/Prioritization/NetObjectPrioritizer.h
@@ -85,8 +85,10 @@ struct FNetObjectPrioritizerInitParams
 	TObjectPtr<const UReplicationSystem> ReplicationSystem;
 	/** Optional config as set in the FNetObjectPrioritizerDefinition. */
 	UNetObjectPrioritizerConfig* Config = nullptr;
-	/** The maximum number of objects in the system. */
-	uint32 MaxObjectCount = 0;
+	/** The maximum number of replicated objects in the system. */
+	uint32 AbsoluteMaxNetObjectCount = 0;
+	/** The current maximum replicated objects referenced by an index (may grow at runtime). */
+	uint32 CurrentMaxInternalIndex = 0;
 	/** The maximum number of connections in the system. */
 	uint32 MaxConnectionCount = 0;
 };
@@ -145,6 +147,12 @@ public:
 	/** Called once at init time before any other calls to the prioritizer. */
 	IRISCORE_API virtual void Init(FNetObjectPrioritizerInitParams& Params) PURE_VIRTUAL(Init,)
 
+	/** Called when the replication system is shutting down. Use this to remove references to other systems */
+	IRISCORE_API virtual void Deinit() PURE_VIRTUAL(Deinit)
+
+	/** Called when the maximum InternalNetRefIndex increased and we need to realloc our lists */
+	IRISCORE_API virtual void OnMaxInternalNetRefIndexIncreased(uint32 NewMaxInternalIndex) PURE_VIRTUAL(OnMaxInternalNetRefIndexIncreased);
+
 	/** A new connection has been added. An opportunity for the prioritizer to allocate per connection info. */
 	IRISCORE_API virtual void AddConnection(uint32 ConnectionId);
 
diff --git a/Engine/Source/Runtime/Experimental/Iris/Core/Public/Iris/ReplicationSystem/Prioritization/SphereNetObjectPrioritizer.h b/Engine/Source/Runtime/Experimental/Iris/Core/Public/Iris/ReplicationSystem/Prioritization/SphereNetObjectPrioritizer.h
--- a/Engine/Source/Runtime/Experimental/Iris/Core/Public/Iris/ReplicationSystem/Prioritization/SphereNetObjectPrioritizer.h
+++ b/Engine/Source/Runtime/Experimental/Iris/Core/Public/Iris/ReplicationSystem/Prioritization/SphereNetObjectPrioritizer.h
@@ -43,6 +43,7 @@ class USphereNetObjectPrioritizer : public ULocationBasedNetObjectPrioritizer
 protected:
 	// UNetObjectPrioritizer interface
 	IRISCORE_API virtual void Init(FNetObjectPrioritizerInitParams& Params) override;
+	IRISCORE_API virtual void Deinit() override;
 	IRISCORE_API virtual void Prioritize(FNetObjectPrioritizationParams&) override;
 
 protected:
diff --git a/Engine/Source/Runtime/Experimental/Iris/Core/Public/Iris/ReplicationSystem/Prioritization/SphereWithOwnerBoostNetObjectPrioritizer.h b/Engine/Source/Runtime/Experimental/Iris/Core/Public/Iris/ReplicationSystem/Prioritization/SphereWithOwnerBoostNetObjectPrioritizer.h
--- a/Engine/Source/Runtime/Experimental/Iris/Core/Public/Iris/ReplicationSystem/Prioritization/SphereWithOwnerBoostNetObjectPrioritizer.h
+++ b/Engine/Source/Runtime/Experimental/Iris/Core/Public/Iris/ReplicationSystem/Prioritization/SphereWithOwnerBoostNetObjectPrioritizer.h
@@ -25,6 +25,7 @@ class USphereWithOwnerBoostNetObjectPrioritizer : public USphereNetObjectPriorit
 protected:
 	// UNetObjectPrioritizer interface
 	IRISCORE_API virtual void Init(FNetObjectPrioritizerInitParams& Params) override;
+	IRISCORE_API virtual void Deinit() override;
 	IRISCORE_API virtual bool AddObject(uint32 ObjectIndex, FNetObjectPrioritizerAddObjectParams& Params) override;
 	IRISCORE_API virtual void RemoveObject(uint32 ObjectIndex, const FNetObjectPrioritizationInfo& Info) override;
 	IRISCORE_API virtual void UpdateObjects(FNetObjectPrioritizerUpdateParams&) override;
diff --git a/Engine/Source/Runtime/Experimental/Iris/Core/Public/Iris/ReplicationSystem/ReplicationSystem.h b/Engine/Source/Runtime/Experimental/Iris/Core/Public/Iris/ReplicationSystem/ReplicationSystem.h
--- a/Engine/Source/Runtime/Experimental/Iris/Core/Public/Iris/ReplicationSystem/ReplicationSystem.h
+++ b/Engine/Source/Runtime/Experimental/Iris/Core/Public/Iris/ReplicationSystem/ReplicationSystem.h
@@ -64,20 +64,42 @@ public:
 		/** The replication bridge that allows communication between the replication system and the game engine  */
 		UReplicationBridge* ReplicationBridge = nullptr;
 
-		/** The maximum amount of netobjects that can be registered to the replication system */
-		uint32 MaxReplicatedObjectCount = 65535U;
+		/** 
+		 * The maximum amount of netobjects that can be registered to the replication system 
+		 * Note that this variable is automatically rounded up to a multiple of 32 so that all available bits in the NetBitArray storage type are used.
+		 */
+		uint32 MaxReplicatedObjectCount = 65536U;
+
+		/**
+		 * The default allocated size for lists referencing NetObjects by their internal index (NetBitArray or TArray).
+		 * Use 0 to preallocate for all possible replicated objects and never reallocate the lists.
+		 * Setting a value smaller than Max minimizes the memory footprint of the replication system when few replicated objects are registered.
+		 * The downside is you have to pay a CPU hit when the initial list size is met.
+		 */
+		uint32 InitialNetObjectListCount = 65536U;
+
+		/**
+		 * The amount by which we increase the size of every NetObjectList (NetBitArray and TArray) when we hit the initial amount.
+		 * Use a small value if you want to keep the memory footprint of the system to a minimum.
+		 * But be aware that increasing the NetObjectList's is costly and may increase memory fragmentation so you'll want to do pay the reallocation cost as little as possible.
+		 * Note that this variable is automatically rounded up to a multiple of 32 so that all available bits in the NetBitArray storage type are used.
+		 */
+		uint32 NetObjectListGrowCount = 16384U;
 
 		/**
-		 * The amount of netobjects to preallocate internal memory buffers for. Preallocated memory provides faster cache-friendly operations but
-		 * has the downside of using much more memory than might actually be needed. Reduce this value if you are operating on a memory constrained platform.
+		 * The amount of netobjects to preallocate internal memory buffers for (NetChunkedArray types).
+		 * These arrays hold the biggest memory blocks in the replication system and can grow independently of the NetObjectLists.
+		 * Using a large amount of preallocated memory provides faster cache-friendly CPU operations but has the downside of holding into much more memory than might actually be needed. 
+		 * Reduce this value if you are operating on a memory constrained platform.
 		 */
-		uint32 PreAllocatedReplicatedObjectCount = 65535U;
+		uint32 PreAllocatedMemoryBuffersObjectCount = 65536U;
 
 		/**
-		 * The maximum amount of netobjects that can replicate properties to remote connection. Can be much lower on clients where very few
-		 * netobjects have authority and support property replication
+		 * The maximum amount of netobjects that can replicate properties to remote connection. 
+		 * Can be much lower on clients where very few netobjects have authority and support property replication (often just 1 player controller)
+		 * When set to 0 it will follow the MaxReplicatedObjectCount and InitialNetObjectListCount limits
 		 */
-		uint32 MaxReplicatedWriterObjectCount = 65535U;
+		uint32 MaxReplicationWriterObjectCount = 0;
 
 		/** The maximum amount of netobjects that can be added to the delta compression manager */
 		uint32 MaxDeltaCompressedObjectCount = 2048U;
diff --git a/Engine/Source/Runtime/Experimental/Iris/Core/Public/Iris/ReplicationSystem/WorldLocations.h b/Engine/Source/Runtime/Experimental/Iris/Core/Public/Iris/ReplicationSystem/WorldLocations.h
--- a/Engine/Source/Runtime/Experimental/Iris/Core/Public/Iris/ReplicationSystem/WorldLocations.h
+++ b/Engine/Source/Runtime/Experimental/Iris/Core/Public/Iris/ReplicationSystem/WorldLocations.h
@@ -8,6 +8,11 @@
 #include "Math/Vector.h"
 #include "WorldLocations.generated.h"
 
+namespace UE::Net::Private
+{
+	typedef uint32 FInternalNetRefIndex;
+}
+
 /**
 * Common settings used to configure how the GridFilter behaves
 */
@@ -31,13 +36,14 @@ namespace UE::Net
 
 struct FWorldLocationsInitParams
 {
-	uint32 MaxObjectCount = 0;
+	UE::Net::Private::FInternalNetRefIndex MaxInternalNetRefIndex = 0;
 };
 
 class FWorldLocations
 {
 public:
 	void Init(const FWorldLocationsInitParams& InitParams);
+	void Deinit();
 
 	/** Returns whether the object has a valid cached data or not. */
 	bool HasInfoForObject(uint32 ObjectIndex) const;
@@ -93,6 +99,8 @@ public:
 	/** Return a position clamped to the configured world boundary. */
 	FVector ClampPositionToBoundary(const FVector& Position);
 
+	void OnMaxInternalNetRefIndexIncreased(UE::Net::Private::FInternalNetRefIndex NewMaxInternalIndex);
+
 private:
 	enum : uint32
 	{
diff --git a/Engine/Source/Runtime/Net/Core/Public/Net/Core/NetBitArray.h b/Engine/Source/Runtime/Net/Core/Public/Net/Core/NetBitArray.h
--- a/Engine/Source/Runtime/Net/Core/Public/Net/Core/NetBitArray.h
+++ b/Engine/Source/Runtime/Net/Core/Public/Net/Core/NetBitArray.h
@@ -58,6 +58,12 @@ public:
 	static constexpr StorageWordType AndNotOp(StorageWordType A, StorageWordType B) { return A & ~B; }
 	static constexpr StorageWordType OrOp(StorageWordType A, StorageWordType B) { return A | B; }
 	static constexpr StorageWordType XorOp(StorageWordType A, StorageWordType B) { return A ^ B; }
+
+	// Round up to a value that uses all available bits in a single NetBitArray Word
+	static uint32 RoundUpToMaxWordBitCount(uint32 Value)
+	{
+		return Value > 0  ? ((Value + WordBitCount - 1) & ~(WordBitCount - 1)) : WordBitCount;
+	}
 };
 
 /**
