Fix crash when loading some geometry collection asset missing mass/inertia attributes
- When attributes are missing we regenerate them   ( editor time only , so assets will be fixed next cook )
- in the future we should revise the way we cook those assets

diff --git a/Engine/Source/Runtime/Experimental/Chaos/Private/PhysicsProxy/GeometryCollectionPhysicsProxy.cpp b/Engine/Source/Runtime/Experimental/Chaos/Private/PhysicsProxy/GeometryCollectionPhysicsProxy.cpp
--- a/Engine/Source/Runtime/Experimental/Chaos/Private/PhysicsProxy/GeometryCollectionPhysicsProxy.cpp
+++ b/Engine/Source/Runtime/Experimental/Chaos/Private/PhysicsProxy/GeometryCollectionPhysicsProxy.cpp
@@ -4940,6 +4940,17 @@ static Chaos::FImplicitObjectPtr MakeTransformImplicitObject(const Chaos::FImpli
 	return ResultObject;
 }
 
+bool FGeometryCollectionPhysicsProxy::NeedToInitializeSharedCollisionStructures(const FGeometryCollection& RestCollection)
+{
+	return !RestCollection.HasAttribute(InertiaTensorAttributeName, FTransformCollection::TransformGroup)
+		|| !RestCollection.HasAttribute(MassAttributeName, FTransformCollection::TransformGroup)
+		|| !RestCollection.HasAttribute(FGeometryDynamicCollection::SimplicialsAttribute, FTransformCollection::TransformGroup)
+		|| !RestCollection.HasAttribute(LevelAttributeName, FTransformCollection::TransformGroup)
+		|| !RestCollection.HasAttribute(FGeometryDynamicCollection::ImplicitsAttribute, FTransformCollection::TransformGroup)
+		|| !RestCollection.HasAttribute(MassToLocalAttributeName, FTransformCollection::TransformGroup)
+		;
+}
+
 /** 
 	NOTE - Making any changes to data stored on the rest collection below MUST be accompanied
 	by a rotation of the DDC key in FDerivedDataGeometryCollectionCooker::GetVersionString
@@ -4971,7 +4982,7 @@ void FGeometryCollectionPhysicsProxy::InitializeSharedCollisionStructures(
 	TManagedArray<FRealSingle>& CollectionMass = RestCollection.AddAttribute<FRealSingle>(MassAttributeName, FTransformCollection::TransformGroup);
 	TManagedArray<TUniquePtr<FSimplicial>>& CollectionSimplicials =	RestCollection.AddAttribute<TUniquePtr<FSimplicial>>(FGeometryDynamicCollection::SimplicialsAttribute, FTransformCollection::TransformGroup);
 
-	TManagedArray<int32>& Levels = RestCollection.AddAttribute<int32>(TEXT("Level"), FTransformCollection::TransformGroup);
+	TManagedArray<int32>& Levels = RestCollection.AddAttribute<int32>(LevelAttributeName, FTransformCollection::TransformGroup);
 
 	RestCollection.RemoveAttribute(FGeometryDynamicCollection::ImplicitsAttribute, FTransformCollection::TransformGroup);
 	TManagedArray<Chaos::FImplicitObjectPtr>& CollectionImplicits = RestCollection.AddAttribute<Chaos::FImplicitObjectPtr>(FGeometryDynamicCollection::ImplicitsAttribute, FTransformCollection::TransformGroup);
diff --git a/Engine/Source/Runtime/Experimental/Chaos/Public/PhysicsProxy/GeometryCollectionPhysicsProxy.h b/Engine/Source/Runtime/Experimental/Chaos/Public/PhysicsProxy/GeometryCollectionPhysicsProxy.h
--- a/Engine/Source/Runtime/Experimental/Chaos/Public/PhysicsProxy/GeometryCollectionPhysicsProxy.h
+++ b/Engine/Source/Runtime/Experimental/Chaos/Public/PhysicsProxy/GeometryCollectionPhysicsProxy.h
@@ -267,6 +267,7 @@ public:
 		Commands.Add(Command); 
 	}
 
+	static CHAOS_API bool NeedToInitializeSharedCollisionStructures(const FGeometryCollection& RestCollection);
 	static CHAOS_API void InitializeSharedCollisionStructures(Chaos::FErrorReporter& ErrorReporter, FGeometryCollection& RestCollection, const FSharedSimulationParameters& SharedParams);
 
 	CHAOS_API void FieldForcesUpdateCallback(Chaos::FPBDRigidsSolver* RigidSolver);
diff --git a/Engine/Source/Runtime/Experimental/GeometryCollectionEngine/Private/GeometryCollection/GeometryCollectionObject.cpp b/Engine/Source/Runtime/Experimental/GeometryCollectionEngine/Private/GeometryCollection/GeometryCollectionObject.cpp
--- a/Engine/Source/Runtime/Experimental/GeometryCollectionEngine/Private/GeometryCollection/GeometryCollectionObject.cpp
+++ b/Engine/Source/Runtime/Experimental/GeometryCollectionEngine/Private/GeometryCollection/GeometryCollectionObject.cpp
@@ -25,6 +25,8 @@
 #include "Rendering/NaniteResources.h"
 #include "Engine/AssetUserData.h"
 #include "PhysicalMaterials/PhysicalMaterial.h"
+#include "PhysicsProxy/GeometryCollectionPhysicsProxy.h"
+#include "Chaos/ErrorReporter.h"
 
 #if WITH_EDITOR
 #include "GeometryCollection/DerivedDataGeometryCollectionCooker.h"
@@ -1862,9 +1864,27 @@ void UGeometryCollection::EnsureDataIsCooked(bool bInitResources, bool bIsTransa
 	if (StateGuid != LastBuiltSimulationDataGuid)
 	{
 		CreateSimulationDataImp(/*bCopyFromDDC=*/ bAllowCopyFromDDC && !bIsTransacting);
+
 		LastBuiltSimulationDataGuid = StateGuid;
 	}
 
+	// todo(chaos) - this is temporary solution to make sure the data is computed accordingly if the attribute are missing
+	//				 in the future we should probably get rid of this all cooker logic and have a proper dependent attribute system 
+	if (GeometryCollection)
+	{
+		if (FGeometryCollectionPhysicsProxy::NeedToInitializeSharedCollisionStructures(*GeometryCollection))
+		{
+			FSharedSimulationParameters SharedParams;
+			GetSharedSimulationParams(SharedParams);
+
+			Chaos::FErrorReporter ErrorReporter(GetName());
+			BuildSimulationData(ErrorReporter, *GeometryCollection, SharedParams);
+			// important : this is necessary to make sure we compute mass scale on the instances properly
+			// sadly we cannot call this in BuildSimulationData because we have no access to the asset
+			CacheMaterialDensity();
+		}
+	}
+
 	// Render data only goes through DDC when loading and saving (bIsPersistant).
 	// Using DDC during edits isn't worth it especially as we use a continually mutating guid instead of a state hash.
 	// That ensures that all edits are cache misses (slow) and unnecessarily fill up DDC disk space.
