Introduce an API for expressing how IAS should validate its peers (if at all)

#rnx

diff --git a/Engine/Source/Runtime/Experimental/IoStore/HttpClient/Internal/IO/Http/Client.h b/Engine/Source/Runtime/Experimental/IoStore/HttpClient/Internal/IO/Http/Client.h
--- a/Engine/Source/Runtime/Experimental/IoStore/HttpClient/Internal/IO/Http/Client.h
+++ b/Engine/Source/Runtime/Experimental/IoStore/HttpClient/Internal/IO/Http/Client.h
@@ -45,6 +45,7 @@ enum class EStatusCodeClass
 };
 
 ////////////////////////////////////////////////////////////////////////////////
+using	FCertRootsRef	= UPTRINT;
 using	FTicket		= uint64;
 using	FPemCert	= FMemoryView;
 struct	FActivity;
@@ -61,6 +62,9 @@ public:
 	bool					IsValid() const					{ return Handle != 0; }
 	int32					Num() const;
 	static void				SetDefault(FCertRoots&& CertRoots);
+	static FCertRootsRef	NoTls();
+	static FCertRootsRef	Default();
+	static FCertRootsRef	Explicit(const FCertRoots& CertRoots);
 
 private:
 	UPTRINT					Handle = 0;
diff --git a/Engine/Source/Runtime/Experimental/IoStore/HttpClient/Private/Peer.inl b/Engine/Source/Runtime/Experimental/IoStore/HttpClient/Private/Peer.inl
--- a/Engine/Source/Runtime/Experimental/IoStore/HttpClient/Private/Peer.inl
+++ b/Engine/Source/Runtime/Experimental/IoStore/HttpClient/Private/Peer.inl
@@ -76,6 +76,12 @@ bool FSslContext::AddCert(const FPemCert& Cert)
 ////////////////////////////////////////////////////////////////////////////////
 static FCertRoots GDefaultCertRoots;
 
+struct ECertRootsRefType
+{
+	static const FCertRootsRef	None	= 0;
+	static const FCertRootsRef	Default = ~0ull;
+};
+
 ////////////////////////////////////////////////////////////////////////////////
 FCertRoots::~FCertRoots()
 {
@@ -167,6 +173,25 @@ void FCertRoots::SetDefault(FCertRoots&& CertRoots)
 	GDefaultCertRoots = MoveTemp(CertRoots);
 }
 
+////////////////////////////////////////////////////////////////////////////////
+FCertRootsRef FCertRoots::NoTls()
+{
+	return ECertRootsRefType::None;
+}
+
+////////////////////////////////////////////////////////////////////////////////
+FCertRootsRef FCertRoots::Default()
+{
+	return ECertRootsRefType::Default;
+}
+
+////////////////////////////////////////////////////////////////////////////////
+FCertRootsRef FCertRoots::Explicit(const FCertRoots& CertRoots)
+{
+	check(CertRoots.IsValid());
+	return CertRoots.Handle;
+}
+
 
 
 ////////////////////////////////////////////////////////////////////////////////
