Move listing of on-demand TOC's to JSON to I/O store tool

diff --git a/Engine/Source/Runtime/Experimental/IoStore/OnDemand/Private/IoStoreOnDemand.cpp b/Engine/Source/Runtime/Experimental/IoStore/OnDemand/Private/IoStoreOnDemand.cpp
--- a/Engine/Source/Runtime/Experimental/IoStore/OnDemand/Private/IoStoreOnDemand.cpp
+++ b/Engine/Source/Runtime/Experimental/IoStore/OnDemand/Private/IoStoreOnDemand.cpp
@@ -31,21 +31,6 @@
 #include "Statistics.h"
 #include "String/LexFromString.h"
 
-#if (PLATFORM_DESKTOP && (IS_PROGRAM || WITH_EDITOR))
-#include "Algo/Sort.h"
-#include "Async/Async.h"
-#include "Async/Mutex.h"
-#include "Async/ParallelFor.h"
-#include "Async/UniqueLock.h"
-#include "HAL/PlatformFileManager.h"
-#include "IO/IoStore.h"
-#include "Misc/EnumClassFlags.h"
-#include "Misc/KeyChainUtilities.h"
-#include "Serialization/JsonSerializer.h"
-#include "Serialization/JsonWriter.h"
-#include "S3/S3Client.h"
-#endif // (PLATFORM_DESKTOP && (IS_PROGRAM || WITH_EDITOR))
-
 DEFINE_LOG_CATEGORY(LogIoStoreOnDemand);
 DEFINE_LOG_CATEGORY(LogIas);
 
@@ -855,565 +840,6 @@ TIoStatusOr<FOnDemandToc> FOnDemandToc::LoadFromUrl(FStringView Url, uint32 Retr
 	return LoadFromUrl(AnsiUrl, RetryCount, bFollowRedirects);
 }
 
-////////////////////////////////////////////////////////////////////////////////
-#if (PLATFORM_DESKTOP && (IS_PROGRAM || WITH_EDITOR))
-
-using FJsonWriter = TSharedPtr<TJsonWriter<TCHAR, TPrettyJsonPrintPolicy<TCHAR>>>; 
-using FJsonWriterFactory = TJsonWriterFactory<TCHAR, TPrettyJsonPrintPolicy<TCHAR>>;
-
-////////////////////////////////////////////////////////////////////////////////
-static void GetChunkObjectKeys(const FOnDemandToc& Toc, FString Prefix, TArray<FString>& OutKeys)
-{
-	if (Prefix.EndsWith(TEXT("/")))
-	{
-		Prefix = Prefix.Left(Prefix.Len() - 1);
-	}
-
-	TStringBuilder<256> Sb;
-	for (const FOnDemandTocContainerEntry& Container : Toc.Containers)
-	{
-		for (const FOnDemandTocEntry& Entry : Container.Entries)
-		{
-			Sb.Reset();
-			const FString Hash = LexToString(Entry.Hash);
-			Sb << Prefix << TEXT("/") << Toc.Header.ChunksDirectory << TEXT("/") << Hash.Left(2) << TEXT("/") << Hash << TEXT(".iochunk");
-			OutKeys.Add(FString(Sb.ToString()).ToLower());
-		}
-	}
-}
-
-////////////////////////////////////////////////////////////////////////////////
-enum class EOnDemandTocJsonOptions
-{
-	Header			= (1 << 0),
-	TocEntries		= (1 << 1),
-	BlockSizes		= (1 << 2),
-	BlockHashes		= (1 << 3),
-	All				= Header | TocEntries | BlockSizes | BlockHashes
-};
-ENUM_CLASS_FLAGS(EOnDemandTocJsonOptions);
-
-static void ToJson(FJsonWriter JsonWriter, const FOnDemandToc& Toc, EOnDemandTocJsonOptions Options)
-{
-	JsonWriter->WriteObjectStart(TEXT("Header"));
-	{
-		JsonWriter->WriteValue(TEXT("Magic"), Toc.Header.Magic);
-		JsonWriter->WriteValue(TEXT("Version"), uint64(Toc.Header.Version));
-		JsonWriter->WriteValue(TEXT("BlockSize"), uint64(Toc.Header.BlockSize));
-		JsonWriter->WriteValue(TEXT("CompressonFormat"), Toc.Header.CompressionFormat);
-		JsonWriter->WriteValue(TEXT("ChunksDirectory"), Toc.Header.ChunksDirectory);
-	}
-	JsonWriter->WriteObjectEnd();
-	JsonWriter->WriteObjectStart(TEXT("Meta"));
-	{
-		const FDateTime Dt = FDateTime::FromUnixTimestamp(Toc.Meta.EpochTimestamp);
-		JsonWriter->WriteValue(TEXT("DateTime"), Dt.ToString()), 
-		JsonWriter->WriteValue(TEXT("BuildVersion"), Toc.Meta.BuildVersion);
-		JsonWriter->WriteValue(TEXT("TargetPlatform"), Toc.Meta.TargetPlatform);
-	}
-	JsonWriter->WriteObjectEnd();
-	JsonWriter->WriteArrayStart(TEXT("Containers"));
-	{
-		for (const FOnDemandTocContainerEntry& Container : Toc.Containers)
-		{
-			JsonWriter->WriteObjectStart();
-			JsonWriter->WriteValue(TEXT("Name"), Container.ContainerName);
-			JsonWriter->WriteValue(TEXT("EncryptionKeyGuid"), Container.EncryptionKeyGuid);
-
-			if (EnumHasAnyFlags(Options, EOnDemandTocJsonOptions::TocEntries))
-			{
-				JsonWriter->WriteArrayStart("Entries");
-				for (const FOnDemandTocEntry& Entry : Container.Entries)
-				{
-					JsonWriter->WriteObjectStart();
-					JsonWriter->WriteValue(TEXT("Hash"), LexToString(Entry.Hash));
-					JsonWriter->WriteValue(TEXT("ChunkId"), LexToString(Entry.ChunkId));
-					JsonWriter->WriteValue(TEXT("RawSize"), Entry.RawSize);
-					JsonWriter->WriteValue(TEXT("EncodedSize"), Entry.EncodedSize);
-					JsonWriter->WriteValue(TEXT("BlockOffset"), int32(Entry.BlockOffset));
-					JsonWriter->WriteValue(TEXT("BlockCount"), int32(Entry.BlockCount));
-					JsonWriter->WriteObjectEnd();
-				}
-				JsonWriter->WriteArrayEnd();
-			}
-
-			if (EnumHasAnyFlags(Options, EOnDemandTocJsonOptions::BlockSizes))
-			{
-				JsonWriter->WriteArrayStart("Blocks");
-				{
-					for (uint32 BlockSize : Container.BlockSizes)
-					{
-						JsonWriter->WriteValue(int32(BlockSize));
-					}
-				}
-				JsonWriter->WriteArrayEnd();
-			}
-
-			if (EnumHasAnyFlags(Options, EOnDemandTocJsonOptions::BlockHashes))
-			{
-				JsonWriter->WriteArrayStart("BlockHashes");
-				{
-					for (const FIoBlockHash& BlockHash : Container.BlockHashes)
-					{
-						JsonWriter->WriteValue(int32(BlockHash));
-					}
-				}
-				JsonWriter->WriteArrayEnd();
-			}
-			JsonWriter->WriteObjectEnd();
-		}
-	}
-	JsonWriter->WriteArrayEnd();
-}
-
-////////////////////////////////////////////////////////////////////////////////
-TIoStatusOr<FIoStoreListTocsParams> FIoStoreListTocsParams::Parse(const TCHAR* CommandLine)
-{
-	FIoStoreListTocsParams Params;
-
-	// Convenience argument to specify both bucket and bucket prefix
-	// -BucketPath="mybucket/some/data/path" is equal to -Bucket="mybucket" -BucketPrefix="some/data/path
-	FString BucketPath;
-	if (FParse::Value(CommandLine, TEXT("-BucketPath="), BucketPath))
-	{
-		FStringView PathView = BucketPath;
-		int32 Idx = INDEX_NONE;
-		if (PathView.FindChar(TCHAR('/'), Idx))
-		{
-			Params.Bucket = PathView.Left(Idx);
-			FStringView BucketPrefix = PathView.RightChop(Idx + 1);
-			if (BucketPrefix.EndsWith(TEXT("/")))
-			{
-				BucketPrefix.LeftChopInline(1);
-			}
-			Params.BucketPrefix = BucketPrefix;
-		}
-		else
-		{
-			Params.Bucket = MoveTemp(BucketPath);
-		}
-	}
-
-	if (Params.Bucket.IsEmpty())
-	{
-		if (!FParse::Value(CommandLine, TEXT("Bucket="), Params.Bucket))
-		{
-			return FIoStatus(EIoErrorCode::InvalidParameter, TEXT("Invalid bucket name"));
-		}
-	}
-
-	FParse::Value(CommandLine, TEXT("BucketPrefix="), Params.BucketPrefix);
-	FParse::Value(CommandLine, TEXT("ServiceUrl="), Params.ServiceUrl);
-	FParse::Value(CommandLine, TEXT("Region="), Params.Region);
-	FParse::Value(CommandLine, TEXT("AccessKey="), Params.AccessKey);
-	FParse::Value(CommandLine, TEXT("SecretKey="), Params.SecretKey);
-	FParse::Value(CommandLine, TEXT("SessionToken="), Params.SessionToken);
-	FParse::Value(CommandLine, TEXT("CredentialsFile="), Params.CredentialsFile);
-	FParse::Value(CommandLine, TEXT("TocKey="), Params.TocKey);
-	FParse::Value(CommandLine, TEXT("BuildVersion="), Params.BuildVersion);
-	FParse::Value(CommandLine, TEXT("TargetPlatform="), Params.TargetPlatform);
-	FParse::Value(CommandLine, TEXT("ChunkKeys="), Params.ChunkKeys);
-	// JSON serialization options
-	Params.bTocEntries = FParse::Param(CommandLine, TEXT("TocEntries"));
-	Params.bBlockSizes = FParse::Param(CommandLine, TEXT("BlockSizes"));
-	Params.bBlockHashes = FParse::Param(CommandLine, TEXT("BlockHashes"));
-
-	if (!FParse::Value(CommandLine, TEXT("CredentialsFileKeyName="), Params.CredentialsFileKeyName))
-	{
-		Params.CredentialsFileKeyName = TEXT("default");
-	}
-
-	FParse::Value(CommandLine, TEXT("Json="), Params.OutFile);
-
-	if (FParse::Value(CommandLine, TEXT("TocUrl="), Params.TocUrl))
-	{
-		FStringView UrlView(Params.TocUrl);
-		if (UrlView.StartsWith(TEXTVIEW("http://")) && UrlView.EndsWith(TEXTVIEW(".iochunktoc")))
-		{
-			int32 Delim = INDEX_NONE;
-			if (UrlView.RightChop(7).FindChar(TEXT('/'), Delim))
-			{
-				Params.ServiceUrl = FString(UrlView.Left(7 +  Delim));
-				Params.TocKey = UrlView.RightChop(Params.ServiceUrl.Len() + 1);
-			}
-		}
-	}
-
-	if (FIoStatus Status = Params.Validate(); !Status.IsOk())
-	{
-		return Status;
-	}
-
-	return Params;
-}
-
-FIoStatus FIoStoreListTocsParams::Validate() const
-{
-	if (Bucket.IsEmpty())
-	{
-		return FIoStatus(EIoErrorCode::InvalidParameter, TEXT("Invalid bucket name"));
-	}
-	if (!AccessKey.IsEmpty() && SecretKey.IsEmpty())
-	{
-		return FIoStatus(EIoErrorCode::InvalidParameter, TEXT("Invalid secret key"));
-	}
-	else if (AccessKey.IsEmpty() && !SecretKey.IsEmpty())
-	{
-		return FIoStatus(EIoErrorCode::InvalidParameter, TEXT("Invalid access key"));
-	}
-
-	if (!CredentialsFile.IsEmpty() && CredentialsFileKeyName.IsEmpty())
-	{
-		return FIoStatus(EIoErrorCode::InvalidParameter, TEXT("Invalid credential file key name"));
-	}
-
-	if (TocUrl.IsEmpty() && ServiceUrl.IsEmpty() && Region.IsEmpty())
-	{
-		return FIoStatus(EIoErrorCode::InvalidParameter, TEXT("Service URL or AWS region needs to be specified"));
-	}
-
-	return FIoStatus::Ok;
-}
-
-FIoStatus ListTocs(const FIoStoreListTocsParams& Params)
-{
-	struct FChunkStats
-	{
-		TSet<FIoHash> Chunks;
-		uint64 TotalChunkSize = 0;
-
-		void Add(const FIoHash& Hash, uint64 ChunkSize)
-		{
-			bool bExist = false;
-			Chunks.Add(Hash, &bExist);
-			if (!bExist)
-			{
-				TotalChunkSize += ChunkSize;
-			}
-		}
-	};
-	TMap<FString, FChunkStats> Stats;
-
-	struct FTocDescription
-	{
-		FOnDemandToc Toc;
-		FDateTime DateTime;
-		FString Key;
-		uint64 Size = 0;
-		uint64 ChunkCount = 0;
-		uint64 TotalChunkSize = 0;
-	};
-
-	const bool bFilteredQuery = !Params.TocKey.IsEmpty() || !Params.BuildVersion.IsEmpty() || !Params.TargetPlatform.IsEmpty();
-
-	FS3ClientConfig Config;
-	Config.ServiceUrl = Params.ServiceUrl;
-	Config.Region = Params.Region;
-	
-	FS3ClientCredentials Credentials;
-	if (Params.CredentialsFile.IsEmpty() == false)
-	{
-		UE_LOG(LogIas, Display, TEXT("Loading credentials file '%s'"), *Params.CredentialsFile);
-		FS3CredentialsProfileStore CredentialsStore = FS3CredentialsProfileStore::FromFile(Params.CredentialsFile);
-		if (CredentialsStore.TryGetCredentials(Params.CredentialsFileKeyName, Credentials) == false)
-		{
-			return FIoStatus(EIoErrorCode::InvalidParameter, TEXT("Failed to find valid credentials in credentials file"));
-		}
-		else
-		{
-			UE_LOG(LogIas, Display, TEXT("Found credentials for '%s'"), *Params.CredentialsFileKeyName);
-		}
-	}
-	else
-	{
-		Credentials = FS3ClientCredentials(Params.AccessKey, Params.SecretKey, Params.SessionToken);
-	}
-
-	FS3Client Client(Config, Credentials);
-	FMutex TocMutex;
-	TArray<FTocDescription> Tocs;
-
-	if (Params.TocUrl.IsEmpty())
-	{
-		FStringView PrefixView = Params.BucketPrefix;
-		if (PrefixView.StartsWith(TEXT("/")))
-		{
-			PrefixView.RightChopInline(1);
-		}
-
-		TStringBuilder<256> Path;
-		if (!PrefixView.IsEmpty())
-		{
-			Path << PrefixView << TEXT("/");
-		}
-
-		UE_LOG(LogIas, Display, TEXT("Fetching TOC's from '%s/%s/%s'"), *Client.GetConfig().ServiceUrl, *Params.Bucket, Path.ToString());
-		FS3ListObjectResponse Response = Client.ListObjects(FS3ListObjectsRequest
-		{
-			Params.Bucket,
-			Path.ToString(),
-			TEXT('/')
-		});
-
-		if (Response.Objects.IsEmpty())
-		{
-			UE_LOG(LogIas, Display, TEXT("Not TOC's found at '%s/%s/%s' (%s)"), *Client.GetConfig().ServiceUrl, *Params.Bucket, Path.ToString(), *Response.GetErrorStatus());
-			return FIoStatus(EIoErrorCode::NotFound);
-		}
-
-		ParallelFor(Response.Objects.Num(), [&Params, &Client, &Response, &Tocs, &TocMutex](int32 Index)
-		{
-			const FS3Object& Obj = Response.Objects[Index];
-			if (Obj.Key.EndsWith(TEXT("iochunktoc")) == false)
-			{
-				return;
-			}
-
-			if (!Params.TocKey.IsEmpty() && !Params.TocKey.Equals(FPaths::GetBaseFilename(Obj.Key), ESearchCase::IgnoreCase))
-			{
-				return;
-			}
-
-			UE_LOG(LogIas, Display, TEXT("Fetching TOC '%s/%s/%s'"), *Client.GetConfig().ServiceUrl, *Params.Bucket, *Obj.Key);
-			FS3GetObjectResponse TocResponse = Client.GetObject(FS3GetObjectRequest
-			{
-				Params.Bucket,
-				Obj.Key
-			});
-
-			if (TocResponse.IsOk() == false)
-			{
-				UE_LOG(LogIas, Warning, TEXT("Failed to fetch TOC '%s/%s/%s'"), *Client.GetConfig().ServiceUrl, *Params.Bucket, *Obj.Key);
-				return;
-			}
-
-			FOnDemandToc Toc;
-			FMemoryReaderView Ar(TocResponse.GetBody().GetView());
-			Ar << Toc;
-
-			if (Ar.IsError() || Toc.Header.Magic != FOnDemandTocHeader::ExpectedMagic)
-			{
-				UE_LOG(LogIas, Warning, TEXT("Failed to serialize TOC '%s/%s/%s'. Header version/magic mimsatch"), *Client.GetConfig().ServiceUrl, *Params.Bucket, *Obj.Key);
-				return;
-			}
-
-			if (!Params.BuildVersion.IsEmpty() && !Params.BuildVersion.Equals(Toc.Meta.BuildVersion, ESearchCase::IgnoreCase))
-			{
-				return;
-			}
-
-			if (!Params.TargetPlatform.IsEmpty() && !Params.TargetPlatform.Equals(Toc.Meta.TargetPlatform, ESearchCase::IgnoreCase))
-			{
-				return;
-			}
-
-			FDateTime DateTime = FDateTime::FromUnixTimestamp(Toc.Meta.EpochTimestamp);
-			FTocDescription Desc = FTocDescription 
-			{
-				.Toc = MoveTemp(Toc),
-				.DateTime = DateTime,
-				.Key = Obj.Key,
-				.Size = Obj.Size
-			};
-
-			{
-				TUniqueLock Lock(TocMutex);
-				Tocs.Add(MoveTemp(Desc));
-			}
-		});
-	}
-	else
-	{
-		UE_LOG(LogIas, Display, TEXT("Fetching TOC '%s/%s'"), *Client.GetConfig().ServiceUrl, *Params.TocKey);
-		FS3GetObjectResponse TocResponse = Client.GetObject(FS3GetObjectRequest
-		{
-			TEXT("/"),
-			Params.TocKey
-		});
-
-		if (TocResponse.IsOk())
-		{
-			FOnDemandToc Toc;
-			FMemoryReaderView Ar(TocResponse.GetBody().GetView());
-			Ar << Toc;
-
-			if (Toc.Header.Magic == FOnDemandTocHeader::ExpectedMagic)
-			{
-				FDateTime DateTime = FDateTime::FromUnixTimestamp(Toc.Meta.EpochTimestamp);
-				Tocs.Add(FTocDescription 
-				{
-					.Toc = MoveTemp(Toc),
-					.DateTime = DateTime,
-					.Key = Params.TocKey,
-					.Size = TocResponse.GetBody().GetSize()
-				});
-			}
-			else
-			{
-				UE_LOG(LogIas, Warning, TEXT("Failed to serialize TOC '%s/%s'. Header magic mimsatch"), *Client.GetConfig().ServiceUrl, *Params.TocKey);
-			}
-		}
-	}
-
-	if (Tocs.IsEmpty())
-	{
-		return FIoStatus(EIoErrorCode::NotFound);
-	}
-
-	for (FTocDescription& Desc : Tocs)
-	{
-		FChunkStats& BuildStats = Stats.FindOrAdd(Desc.Toc.Meta.BuildVersion);
-		FChunkStats& BucketStats = Stats.FindOrAdd(Params.Bucket);
-		for (const FOnDemandTocContainerEntry& Container : Desc.Toc.Containers)
-		{
-			for (const FOnDemandTocEntry& Entry : Container.Entries)
-			{
-				Desc.ChunkCount++;
-				Desc.TotalChunkSize += Entry.EncodedSize;
-				BucketStats.Add(Entry.Hash, Entry.EncodedSize);
-				BuildStats.Add(Entry.Hash, Entry.EncodedSize);
-			}
-		}
-	}
-
-	Tocs.Sort([](const auto& LHS, const auto& RHS) { return LHS.DateTime > RHS.DateTime; });
-
-	if (Params.OutFile.IsEmpty())
-	{
-		int32 Counter = 1;
-		TStringBuilder<256> Url;
-
-		for (const FTocDescription& Desc : Tocs)
-		{
-			Url.Reset();
-			Url << Client.GetConfig().ServiceUrl << TEXT("/") << Params.Bucket << TEXT("/") << Desc.Key;
-
-			UE_LOG(LogIas, Display, TEXT(""));
-			UE_LOG(LogIas, Display, TEXT("%d) %s"), Counter++, *Desc.Key);
-			UE_LOG(LogIas, Display, TEXT("%-20s: %s"),TEXT("Date"), *Desc.DateTime.ToString());
-			UE_LOG(LogIas, Display, TEXT("%-20s: %s"), TEXT("BuildVersion"), *Desc.Toc.Meta.BuildVersion);
-			UE_LOG(LogIas, Display, TEXT("%-20s: %s"), TEXT("TargetPlatform"), *Desc.Toc.Meta.TargetPlatform);
-			UE_LOG(LogIas, Display, TEXT("%-20s: %.2lf KiB"), TEXT("TocSize"), double(Desc.Size) / 1024.0);
-			UE_LOG(LogIas, Display, TEXT("%-20s: %llu"), TEXT("ChunkCount"), Desc.ChunkCount);
-			UE_LOG(LogIas, Display, TEXT("%-20s: %.2lf MiB"), TEXT("TotalChunkSize"), double(Desc.TotalChunkSize) / 1024.0 / 1024.0);
-			UE_LOG(LogIas, Display, TEXT("%-20s: %s"), TEXT("Url"), Url.ToString());
-		}
-		UE_LOG(LogIas, Display, TEXT(""));
-		
-		if (bFilteredQuery == false)
-		{
-			TArray<FString> Keys;
-			Stats.GetKeys(Keys);
-			Keys.Sort();
-
-			UE_LOG(LogIas, Display, TEXT("%-80s %15s %15s"), TEXT("BuildVersion"), TEXT("Chunk(s)"), TEXT("MiB"));
-			UE_LOG(LogIas, Display, TEXT("-------------------------------------------------------------------------------------------------------------------"));
-			for (const FString& Key : Keys)
-			{
-				if (Key != Params.Bucket)
-				{
-					const FChunkStats& ChunkStats = Stats.FindChecked(Key);
-					UE_LOG(LogIas, Display, TEXT("%-80s %15llu %15.2f"), *Key, ChunkStats.Chunks.Num(), double(ChunkStats.TotalChunkSize) / 1024.0 / 1024.0);
-				}
-			}
-			UE_LOG(LogIas, Display, TEXT("-------------------------------------------------------------------------------------------------------------------"));
-			FChunkStats& BucketStats = Stats.FindOrAdd(Params.Bucket);
-			UE_LOG(LogIas, Display, TEXT("%-80s %15llu %15.2f"), *Params.Bucket, BucketStats.Chunks.Num(), double(BucketStats.TotalChunkSize) / 1024.0 / 1024.0);
-			UE_LOG(LogIas, Display, TEXT(""));
-		}
-	}
-	else
-	{
-		FString Json;
-		FJsonWriter JsonWriter = FJsonWriterFactory::Create(&Json);
-		JsonWriter->WriteObjectStart();
-		JsonWriter->WriteValue(TEXT("ServiceUrl"), Client.GetConfig().ServiceUrl);
-		JsonWriter->WriteValue(TEXT("Bucket"), Params.Bucket);
-		JsonWriter->WriteValue(TEXT("BucketPrefix"), Params.BucketPrefix);
-
-		JsonWriter->WriteArrayStart(TEXT("Tocs"));
-		for (const FTocDescription& Desc : Tocs)
-		{
-			JsonWriter->WriteObjectStart();
-			JsonWriter->WriteValue(TEXT("Key"), Desc.Key);
-			JsonWriter->WriteValue(TEXT("ChunkCount"), Desc.ChunkCount);
-			JsonWriter->WriteValue(TEXT("TotalChunkSize"), Desc.TotalChunkSize);
-
-			EOnDemandTocJsonOptions JsonOptions = EOnDemandTocJsonOptions::Header;
-			if (Params.bTocEntries)
-			{
-				JsonOptions |= EOnDemandTocJsonOptions::TocEntries;
-			}
-			if (Params.bBlockSizes)
-			{
-				JsonOptions |= EOnDemandTocJsonOptions::BlockSizes;
-			}
-			if (Params.bBlockHashes)
-			{
-				JsonOptions |= EOnDemandTocJsonOptions::BlockHashes;
-			}
-			ToJson(JsonWriter, Desc.Toc, JsonOptions);
-			JsonWriter->WriteObjectEnd();
-		}
-		JsonWriter->WriteArrayEnd();
-
-		JsonWriter->WriteObjectEnd();
-		JsonWriter->Close();
-
-		UE_LOG(LogIas, Display, TEXT("Saving file '%s'"), *Params.OutFile);
-		if (!FFileHelper::SaveStringToFile(Json, *Params.OutFile))
-		{
-			return FIoStatus(EIoErrorCode::WriteError, TEXTVIEW("Failed writing JSON file")); 
-		}
-	}
-
-	if (Params.ChunkKeys.IsEmpty() == false)
-	{
-		FString Json;
-		FJsonWriter JsonWriter = FJsonWriterFactory::Create(&Json);
-		
-		JsonWriter->WriteObjectStart();
-		JsonWriter->WriteValue(TEXT("ServiceUrl"), Client.GetConfig().ServiceUrl);
-		JsonWriter->WriteValue(TEXT("Bucket"), Params.Bucket);
-		JsonWriter->WriteValue(TEXT("BucketPrefix"), Params.BucketPrefix);
-
-		JsonWriter->WriteArrayStart(TEXT("Tocs"));
-		for (const FTocDescription& Desc : Tocs)
-		{
-			TArray<FString> ObjKeys;
-			GetChunkObjectKeys(Desc.Toc, Params.BucketPrefix, ObjKeys);
-
-			JsonWriter->WriteObjectStart();
-			JsonWriter->WriteValue(TEXT("Key"), Desc.Key);
-			JsonWriter->WriteValue(TEXT("BuildVersion"), Desc.Toc.Meta.BuildVersion);
-			JsonWriter->WriteValue(TEXT("TargetPlatform"), Desc.Toc.Meta.TargetPlatform);
-			JsonWriter->WriteArrayStart(TEXT("ChunkKeys"));
-			for (const FString& Key : ObjKeys)
-			{
-				JsonWriter->WriteValue(Key);
-			}
-			JsonWriter->WriteArrayEnd();
-			JsonWriter->WriteObjectEnd();
-		}
-		JsonWriter->WriteArrayEnd();
-
-		JsonWriter->WriteObjectEnd();
-		JsonWriter->Close();
-
-		UE_LOG(LogIas, Display, TEXT("Saving chunk key(s) '%s'"), *Params.ChunkKeys);
-		if (!FFileHelper::SaveStringToFile(Json, *Params.ChunkKeys))
-		{
-			return FIoStatus(EIoErrorCode::WriteError, TEXTVIEW("Failed writing JSON file")); 
-		}
-	}
-
-	return FIoStatus::Ok;
-}
-
-#endif // (PLATFORM_DESKTOP && (IS_PROGRAM || WITH_EDITOR))
-
-
-
 ////////////////////////////////////////////////////////////////////////////////
 void FIoStoreOnDemandModule::SetBulkOptionalEnabled(bool bInEnabled)
 {
diff --git a/Engine/Source/Runtime/Experimental/IoStore/OnDemand/Private/Tool/ListTocs.cpp b/Engine/Source/Runtime/Experimental/IoStore/OnDemand/Private/Tool/ListTocs.cpp
--- a/Engine/Source/Runtime/Experimental/IoStore/OnDemand/Private/Tool/ListTocs.cpp
+++ b/Engine/Source/Runtime/Experimental/IoStore/OnDemand/Private/Tool/ListTocs.cpp
@@ -4,14 +4,113 @@
 
 #include "Command.h"
 
-#include <IO/IoStoreOnDemand.h>
+#include "Async/Mutex.h"
+#include "Async/ParallelFor.h"
+#include "Async/UniqueLock.h"
+#include "IO/IoStoreOnDemand.h"
+#include "Misc/EnumClassFlags.h"
+#include "Misc/FileHelper.h"
+#include "Misc/Paths.h"
+#include "Serialization/JsonSerializer.h"
+#include "Serialization/JsonWriter.h"
+#include "Serialization/MemoryReader.h"
+#include "S3/S3Client.h"
 
-namespace UE::IoStore::Tool {
+namespace UE::IoStore::Tool
+{
+
+using FJsonWriter = TSharedPtr<TJsonWriter<TCHAR, TPrettyJsonPrintPolicy<TCHAR>>>; 
+using FJsonWriterFactory = TJsonWriterFactory<TCHAR, TPrettyJsonPrintPolicy<TCHAR>>;
+
+////////////////////////////////////////////////////////////////////////////////
+
+/**
+ * Parameters for listing uploaded TOC file(s) from an S3 compatible endpoint.
+ *
+ * Example usage:
+ *
+ * 1) Print available TOC's from a local server to standard out.
+ * IasTool.exe ListTocs -ServiceUrl="http://10.24.101.92:9000" -Bucket=<bucketname> -BucketPrefix=<some/path/to/data> -AccessKey=<accesskey> -SecretKey=<secretkey>
+ *
+ * 2) Print available TOC's from AWS S3.
+ * IasTool.exe ListTocs -Region="us-east-1" -BucketPath=<mybucket/some/data/path/> -CredentialsFile=<path/to/credentials.txt> -BuildVersion=<version> -Json=<path/to/file.json>
+ * 
+ * 3) Serialize all TOC's matching a specific build version to JSON:
+ * IasTool.exe ListTocs -Region="us-east-1" -BucketPath=<mybucket/some/data/path/> -CredentialsFile=<path/to/credentials.txt> -BuildVersion=<version> -Json=<path/to/file.json>
+ *
+ * 4) Serialize all chunk object key(s) to JSON.
+ * IasTool.exe ListTocs -Region="us-east-1" -BucketPath=<mybucket/some/data/path/> -CredentialsFile=<path/to/credentials.txt> -BuildVersion=<version> -ChunkKeys=<path/to/file.json>
+ *
+ * 5) Fetch a TOC from a public CDN.
+ * IasTool.exe ListTocs -TocUrl=<http://some.public.endpoint.net/path/to/1a32076ca12bfc6feb982ffb064d18f28156606c.iochunktoc>
+ *
+ * Parameters: -TocEntries, -BlockSizes and -BlockHashes controls what to include when serializing TOC's to JSON.
+ *
+ * Credentials file example:
+ *
+ * [default]
+ * aws_access_key_id="<key>"
+ * aws_secret_access_key="<key>
+ * aws_session_token="<token>"
+ *
+ * Note: All values must be surounded with "".
+ */
+struct FListTocsParams
+{
+	FString OutFile;
+	FString ServiceUrl;
+	FString Bucket;
+	FString BucketPrefix;
+	FString Region; 
+	FString AccessKey;
+	FString SecretKey;
+	FString SessionToken;
+	FString CredentialsFile;
+	FString CredentialsFileKeyName;
+	FString TocUrl;
+	FString TocKey;
+	FString BuildVersion;
+	FString TargetPlatform;
+	FString ChunkKeys;
+	bool bTocEntries = false;
+	bool bBlockSizes = false;
+	bool bBlockHashes = false;
+
+	FIoStatus Validate() const;
+};
+
+FIoStatus FListTocsParams::Validate() const
+{
+	if (Bucket.IsEmpty())
+	{
+		return FIoStatus(EIoErrorCode::InvalidParameter, TEXT("Invalid bucket name"));
+	}
+	if (!AccessKey.IsEmpty() && SecretKey.IsEmpty())
+	{
+		return FIoStatus(EIoErrorCode::InvalidParameter, TEXT("Invalid secret key"));
+	}
+	else if (AccessKey.IsEmpty() && !SecretKey.IsEmpty())
+	{
+		return FIoStatus(EIoErrorCode::InvalidParameter, TEXT("Invalid access key"));
+	}
+
+	if (!CredentialsFile.IsEmpty() && CredentialsFileKeyName.IsEmpty())
+	{
+		return FIoStatus(EIoErrorCode::InvalidParameter, TEXT("Invalid credential file key name"));
+	}
+
+	if (TocUrl.IsEmpty() && ServiceUrl.IsEmpty() && Region.IsEmpty())
+	{
+		return FIoStatus(EIoErrorCode::InvalidParameter, TEXT("Service URL or AWS region needs to be specified"));
+	}
+
+	return FIoStatus::Ok;
+}
 
 ////////////////////////////////////////////////////////////////////////////////
-static FIoStoreListTocsParams BuildListTocsParams(const FContext& Context)
+static FListTocsParams BuildListTocsParams(const FContext& Context)
 {
-	FIoStoreListTocsParams Ret;
+	FListTocsParams Ret;
 
 	if (auto Value = Context.Get<FStringView>(TEXT("-BucketPath")); !Value.IsEmpty())
 	{
@@ -74,10 +173,452 @@ static FIoStoreListTocsParams BuildListTocsParams(const FContext& Context)
 	return Ret;
 }
 
+////////////////////////////////////////////////////////////////////////////////
+static void GetChunkObjectKeys(const FOnDemandToc& Toc, FString Prefix, TArray<FString>& OutKeys)
+{
+	if (Prefix.EndsWith(TEXT("/")))
+	{
+		Prefix = Prefix.Left(Prefix.Len() - 1);
+	}
+
+	TStringBuilder<256> Sb;
+	for (const FOnDemandTocContainerEntry& Container : Toc.Containers)
+	{
+		for (const FOnDemandTocEntry& Entry : Container.Entries)
+		{
+			Sb.Reset();
+			const FString Hash = LexToString(Entry.Hash);
+			Sb << Prefix << TEXT("/") << Toc.Header.ChunksDirectory << TEXT("/") << Hash.Left(2) << TEXT("/") << Hash << TEXT(".iochunk");
+			OutKeys.Add(FString(Sb.ToString()).ToLower());
+		}
+	}
+}
+
+////////////////////////////////////////////////////////////////////////////////
+enum class EOnDemandTocJsonOptions
+{
+	Header			= (1 << 0),
+	TocEntries		= (1 << 1),
+	BlockSizes		= (1 << 2),
+	BlockHashes		= (1 << 3),
+	All				= Header | TocEntries | BlockSizes | BlockHashes
+};
+ENUM_CLASS_FLAGS(EOnDemandTocJsonOptions);
+
+////////////////////////////////////////////////////////////////////////////////
+static void ToJson(FJsonWriter JsonWriter, const FOnDemandToc& Toc, EOnDemandTocJsonOptions Options)
+{
+	JsonWriter->WriteObjectStart(TEXT("Header"));
+	{
+		JsonWriter->WriteValue(TEXT("Magic"), Toc.Header.Magic);
+		JsonWriter->WriteValue(TEXT("Version"), uint64(Toc.Header.Version));
+		JsonWriter->WriteValue(TEXT("BlockSize"), uint64(Toc.Header.BlockSize));
+		JsonWriter->WriteValue(TEXT("CompressonFormat"), Toc.Header.CompressionFormat);
+		JsonWriter->WriteValue(TEXT("ChunksDirectory"), Toc.Header.ChunksDirectory);
+	}
+	JsonWriter->WriteObjectEnd();
+	JsonWriter->WriteObjectStart(TEXT("Meta"));
+	{
+		const FDateTime Dt = FDateTime::FromUnixTimestamp(Toc.Meta.EpochTimestamp);
+		JsonWriter->WriteValue(TEXT("DateTime"), Dt.ToString()), 
+		JsonWriter->WriteValue(TEXT("BuildVersion"), Toc.Meta.BuildVersion);
+		JsonWriter->WriteValue(TEXT("TargetPlatform"), Toc.Meta.TargetPlatform);
+	}
+	JsonWriter->WriteObjectEnd();
+	JsonWriter->WriteArrayStart(TEXT("Containers"));
+	{
+		for (const FOnDemandTocContainerEntry& Container : Toc.Containers)
+		{
+			JsonWriter->WriteObjectStart();
+			JsonWriter->WriteValue(TEXT("Name"), Container.ContainerName);
+			JsonWriter->WriteValue(TEXT("EncryptionKeyGuid"), Container.EncryptionKeyGuid);
+
+			if (EnumHasAnyFlags(Options, EOnDemandTocJsonOptions::TocEntries))
+			{
+				JsonWriter->WriteArrayStart("Entries");
+				for (const FOnDemandTocEntry& Entry : Container.Entries)
+				{
+					JsonWriter->WriteObjectStart();
+					JsonWriter->WriteValue(TEXT("Hash"), LexToString(Entry.Hash));
+					JsonWriter->WriteValue(TEXT("ChunkId"), LexToString(Entry.ChunkId));
+					JsonWriter->WriteValue(TEXT("RawSize"), Entry.RawSize);
+					JsonWriter->WriteValue(TEXT("EncodedSize"), Entry.EncodedSize);
+					JsonWriter->WriteValue(TEXT("BlockOffset"), int32(Entry.BlockOffset));
+					JsonWriter->WriteValue(TEXT("BlockCount"), int32(Entry.BlockCount));
+					JsonWriter->WriteObjectEnd();
+				}
+				JsonWriter->WriteArrayEnd();
+			}
+
+			if (EnumHasAnyFlags(Options, EOnDemandTocJsonOptions::BlockSizes))
+			{
+				JsonWriter->WriteArrayStart("Blocks");
+				{
+					for (uint32 BlockSize : Container.BlockSizes)
+					{
+						JsonWriter->WriteValue(int32(BlockSize));
+					}
+				}
+				JsonWriter->WriteArrayEnd();
+			}
+
+			if (EnumHasAnyFlags(Options, EOnDemandTocJsonOptions::BlockHashes))
+			{
+				JsonWriter->WriteArrayStart("BlockHashes");
+				{
+					for (const FIoBlockHash& BlockHash : Container.BlockHashes)
+					{
+						JsonWriter->WriteValue(int32(BlockHash));
+					}
+				}
+				JsonWriter->WriteArrayEnd();
+			}
+			JsonWriter->WriteObjectEnd();
+		}
+	}
+	JsonWriter->WriteArrayEnd();
+}
+
+////////////////////////////////////////////////////////////////////////////////
+FIoStatus ListTocs(const FListTocsParams& Params)
+{
+	struct FChunkStats
+	{
+		TSet<FIoHash> Chunks;
+		uint64 TotalChunkSize = 0;
+
+		void Add(const FIoHash& Hash, uint64 ChunkSize)
+		{
+			bool bExist = false;
+			Chunks.Add(Hash, &bExist);
+			if (!bExist)
+			{
+				TotalChunkSize += ChunkSize;
+			}
+		}
+	};
+	TMap<FString, FChunkStats> Stats;
+
+	struct FTocDescription
+	{
+		FOnDemandToc Toc;
+		FDateTime DateTime;
+		FString Key;
+		uint64 Size = 0;
+		uint64 ChunkCount = 0;
+		uint64 TotalChunkSize = 0;
+	};
+
+	const bool bFilteredQuery = !Params.TocKey.IsEmpty() || !Params.BuildVersion.IsEmpty() || !Params.TargetPlatform.IsEmpty();
+
+	FS3ClientConfig Config;
+	Config.ServiceUrl = Params.ServiceUrl;
+	Config.Region = Params.Region;
+	
+	FS3ClientCredentials Credentials;
+	if (Params.CredentialsFile.IsEmpty() == false)
+	{
+		UE_LOG(LogIas, Display, TEXT("Loading credentials file '%s'"), *Params.CredentialsFile);
+		FS3CredentialsProfileStore CredentialsStore = FS3CredentialsProfileStore::FromFile(Params.CredentialsFile);
+		if (CredentialsStore.TryGetCredentials(Params.CredentialsFileKeyName, Credentials) == false)
+		{
+			return FIoStatus(EIoErrorCode::InvalidParameter, TEXT("Failed to find valid credentials in credentials file"));
+		}
+		else
+		{
+			UE_LOG(LogIas, Display, TEXT("Found credentials for '%s'"), *Params.CredentialsFileKeyName);
+		}
+	}
+	else
+	{
+		Credentials = FS3ClientCredentials(Params.AccessKey, Params.SecretKey, Params.SessionToken);
+	}
+
+	FS3Client Client(Config, Credentials);
+	FMutex TocMutex;
+	TArray<FTocDescription> Tocs;
+
+	if (Params.TocUrl.IsEmpty())
+	{
+		FStringView PrefixView = Params.BucketPrefix;
+		if (PrefixView.StartsWith(TEXT("/")))
+		{
+			PrefixView.RightChopInline(1);
+		}
+
+		TStringBuilder<256> Path;
+		if (!PrefixView.IsEmpty())
+		{
+			Path << PrefixView << TEXT("/");
+		}
+
+		UE_LOG(LogIas, Display, TEXT("Fetching TOC's from '%s/%s/%s'"), *Client.GetConfig().ServiceUrl, *Params.Bucket, Path.ToString());
+		FS3ListObjectResponse Response = Client.ListObjects(FS3ListObjectsRequest
+		{
+			Params.Bucket,
+			Path.ToString(),
+			TEXT('/')
+		});
+
+		if (Response.Objects.IsEmpty())
+		{
+			UE_LOG(LogIas, Display, TEXT("Not TOC's found at '%s/%s/%s' (%s)"), *Client.GetConfig().ServiceUrl, *Params.Bucket, Path.ToString(), *Response.GetErrorStatus());
+			return FIoStatus(EIoErrorCode::NotFound);
+		}
+
+		ParallelFor(Response.Objects.Num(), [&Params, &Client, &Response, &Tocs, &TocMutex](int32 Index)
+		{
+			const FS3Object& Obj = Response.Objects[Index];
+			if (Obj.Key.EndsWith(TEXT("iochunktoc")) == false)
+			{
+				return;
+			}
+
+			if (!Params.TocKey.IsEmpty() && !Params.TocKey.Equals(FPaths::GetBaseFilename(Obj.Key), ESearchCase::IgnoreCase))
+			{
+				return;
+			}
+
+			UE_LOG(LogIas, Display, TEXT("Fetching TOC '%s/%s/%s'"), *Client.GetConfig().ServiceUrl, *Params.Bucket, *Obj.Key);
+			FS3GetObjectResponse TocResponse = Client.GetObject(FS3GetObjectRequest
+			{
+				Params.Bucket,
+				Obj.Key
+			});
+
+			if (TocResponse.IsOk() == false)
+			{
+				UE_LOG(LogIas, Warning, TEXT("Failed to fetch TOC '%s/%s/%s'"), *Client.GetConfig().ServiceUrl, *Params.Bucket, *Obj.Key);
+				return;
+			}
+
+			FOnDemandToc Toc;
+			FMemoryReaderView Ar(TocResponse.GetBody().GetView());
+			Ar << Toc;
+
+			if (Ar.IsError() || Toc.Header.Magic != FOnDemandTocHeader::ExpectedMagic)
+			{
+				UE_LOG(LogIas, Warning, TEXT("Failed to serialize TOC '%s/%s/%s'. Header version/magic mimsatch"), *Client.GetConfig().ServiceUrl, *Params.Bucket, *Obj.Key);
+				return;
+			}
+
+			if (!Params.BuildVersion.IsEmpty() && !Params.BuildVersion.Equals(Toc.Meta.BuildVersion, ESearchCase::IgnoreCase))
+			{
+				return;
+			}
+
+			if (!Params.TargetPlatform.IsEmpty() && !Params.TargetPlatform.Equals(Toc.Meta.TargetPlatform, ESearchCase::IgnoreCase))
+			{
+				return;
+			}
+
+			FDateTime DateTime = FDateTime::FromUnixTimestamp(Toc.Meta.EpochTimestamp);
+			FTocDescription Desc = FTocDescription 
+			{
+				.Toc = MoveTemp(Toc),
+				.DateTime = DateTime,
+				.Key = Obj.Key,
+				.Size = Obj.Size
+			};
+
+			{
+				TUniqueLock Lock(TocMutex);
+				Tocs.Add(MoveTemp(Desc));
+			}
+		});
+	}
+	else
+	{
+		UE_LOG(LogIas, Display, TEXT("Fetching TOC '%s/%s'"), *Client.GetConfig().ServiceUrl, *Params.TocKey);
+		FS3GetObjectResponse TocResponse = Client.GetObject(FS3GetObjectRequest
+		{
+			TEXT("/"),
+			Params.TocKey
+		});
+
+		if (TocResponse.IsOk())
+		{
+			FOnDemandToc Toc;
+			FMemoryReaderView Ar(TocResponse.GetBody().GetView());
+			Ar << Toc;
+
+			if (Toc.Header.Magic == FOnDemandTocHeader::ExpectedMagic)
+			{
+				FDateTime DateTime = FDateTime::FromUnixTimestamp(Toc.Meta.EpochTimestamp);
+				Tocs.Add(FTocDescription 
+				{
+					.Toc = MoveTemp(Toc),
+					.DateTime = DateTime,
+					.Key = Params.TocKey,
+					.Size = TocResponse.GetBody().GetSize()
+				});
+			}
+			else
+			{
+				UE_LOG(LogIas, Warning, TEXT("Failed to serialize TOC '%s/%s'. Header magic mimsatch"), *Client.GetConfig().ServiceUrl, *Params.TocKey);
+			}
+		}
+	}
+
+	if (Tocs.IsEmpty())
+	{
+		return FIoStatus(EIoErrorCode::NotFound);
+	}
+
+	for (FTocDescription& Desc : Tocs)
+	{
+		FChunkStats& BuildStats = Stats.FindOrAdd(Desc.Toc.Meta.BuildVersion);
+		FChunkStats& BucketStats = Stats.FindOrAdd(Params.Bucket);
+		for (const FOnDemandTocContainerEntry& Container : Desc.Toc.Containers)
+		{
+			for (const FOnDemandTocEntry& Entry : Container.Entries)
+			{
+				Desc.ChunkCount++;
+				Desc.TotalChunkSize += Entry.EncodedSize;
+				BucketStats.Add(Entry.Hash, Entry.EncodedSize);
+				BuildStats.Add(Entry.Hash, Entry.EncodedSize);
+			}
+		}
+	}
+
+	Tocs.Sort([](const auto& LHS, const auto& RHS) { return LHS.DateTime > RHS.DateTime; });
+
+	if (Params.OutFile.IsEmpty())
+	{
+		int32 Counter = 1;
+		TStringBuilder<256> Url;
+
+		for (const FTocDescription& Desc : Tocs)
+		{
+			Url.Reset();
+			Url << Client.GetConfig().ServiceUrl << TEXT("/") << Params.Bucket << TEXT("/") << Desc.Key;
+
+			UE_LOG(LogIas, Display, TEXT(""));
+			UE_LOG(LogIas, Display, TEXT("%d) %s"), Counter++, *Desc.Key);
+			UE_LOG(LogIas, Display, TEXT("%-20s: %s"),TEXT("Date"), *Desc.DateTime.ToString());
+			UE_LOG(LogIas, Display, TEXT("%-20s: %s"), TEXT("BuildVersion"), *Desc.Toc.Meta.BuildVersion);
+			UE_LOG(LogIas, Display, TEXT("%-20s: %s"), TEXT("TargetPlatform"), *Desc.Toc.Meta.TargetPlatform);
+			UE_LOG(LogIas, Display, TEXT("%-20s: %.2lf KiB"), TEXT("TocSize"), double(Desc.Size) / 1024.0);
+			UE_LOG(LogIas, Display, TEXT("%-20s: %llu"), TEXT("ChunkCount"), Desc.ChunkCount);
+			UE_LOG(LogIas, Display, TEXT("%-20s: %.2lf MiB"), TEXT("TotalChunkSize"), double(Desc.TotalChunkSize) / 1024.0 / 1024.0);
+			UE_LOG(LogIas, Display, TEXT("%-20s: %s"), TEXT("Url"), Url.ToString());
+		}
+		UE_LOG(LogIas, Display, TEXT(""));
+		
+		if (bFilteredQuery == false)
+		{
+			TArray<FString> Keys;
+			Stats.GetKeys(Keys);
+			Keys.Sort();
+
+			UE_LOG(LogIas, Display, TEXT("%-80s %15s %15s"), TEXT("BuildVersion"), TEXT("Chunk(s)"), TEXT("MiB"));
+			UE_LOG(LogIas, Display, TEXT("-------------------------------------------------------------------------------------------------------------------"));
+			for (const FString& Key : Keys)
+			{
+				if (Key != Params.Bucket)
+				{
+					const FChunkStats& ChunkStats = Stats.FindChecked(Key);
+					UE_LOG(LogIas, Display, TEXT("%-80s %15llu %15.2f"), *Key, ChunkStats.Chunks.Num(), double(ChunkStats.TotalChunkSize) / 1024.0 / 1024.0);
+				}
+			}
+			UE_LOG(LogIas, Display, TEXT("-------------------------------------------------------------------------------------------------------------------"));
+			FChunkStats& BucketStats = Stats.FindOrAdd(Params.Bucket);
+			UE_LOG(LogIas, Display, TEXT("%-80s %15llu %15.2f"), *Params.Bucket, BucketStats.Chunks.Num(), double(BucketStats.TotalChunkSize) / 1024.0 / 1024.0);
+			UE_LOG(LogIas, Display, TEXT(""));
+		}
+	}
+	else
+	{
+		FString Json;
+		FJsonWriter JsonWriter = FJsonWriterFactory::Create(&Json);
+		JsonWriter->WriteObjectStart();
+		JsonWriter->WriteValue(TEXT("ServiceUrl"), Client.GetConfig().ServiceUrl);
+		JsonWriter->WriteValue(TEXT("Bucket"), Params.Bucket);
+		JsonWriter->WriteValue(TEXT("BucketPrefix"), Params.BucketPrefix);
+
+		JsonWriter->WriteArrayStart(TEXT("Tocs"));
+		for (const FTocDescription& Desc : Tocs)
+		{
+			JsonWriter->WriteObjectStart();
+			JsonWriter->WriteValue(TEXT("Key"), Desc.Key);
+			JsonWriter->WriteValue(TEXT("ChunkCount"), Desc.ChunkCount);
+			JsonWriter->WriteValue(TEXT("TotalChunkSize"), Desc.TotalChunkSize);
+
+			EOnDemandTocJsonOptions JsonOptions = EOnDemandTocJsonOptions::Header;
+			if (Params.bTocEntries)
+			{
+				JsonOptions |= EOnDemandTocJsonOptions::TocEntries;
+			}
+			if (Params.bBlockSizes)
+			{
+				JsonOptions |= EOnDemandTocJsonOptions::BlockSizes;
+			}
+			if (Params.bBlockHashes)
+			{
+				JsonOptions |= EOnDemandTocJsonOptions::BlockHashes;
+			}
+			ToJson(JsonWriter, Desc.Toc, JsonOptions);
+			JsonWriter->WriteObjectEnd();
+		}
+		JsonWriter->WriteArrayEnd();
+
+		JsonWriter->WriteObjectEnd();
+		JsonWriter->Close();
+
+		UE_LOG(LogIas, Display, TEXT("Saving file '%s'"), *Params.OutFile);
+		if (!FFileHelper::SaveStringToFile(Json, *Params.OutFile))
+		{
+			return FIoStatus(EIoErrorCode::WriteError, TEXTVIEW("Failed writing JSON file")); 
+		}
+	}
+
+	if (Params.ChunkKeys.IsEmpty() == false)
+	{
+		FString Json;
+		FJsonWriter JsonWriter = FJsonWriterFactory::Create(&Json);
+		
+		JsonWriter->WriteObjectStart();
+		JsonWriter->WriteValue(TEXT("ServiceUrl"), Client.GetConfig().ServiceUrl);
+		JsonWriter->WriteValue(TEXT("Bucket"), Params.Bucket);
+		JsonWriter->WriteValue(TEXT("BucketPrefix"), Params.BucketPrefix);
+
+		JsonWriter->WriteArrayStart(TEXT("Tocs"));
+		for (const FTocDescription& Desc : Tocs)
+		{
+			TArray<FString> ObjKeys;
+			GetChunkObjectKeys(Desc.Toc, Params.BucketPrefix, ObjKeys);
+
+			JsonWriter->WriteObjectStart();
+			JsonWriter->WriteValue(TEXT("Key"), Desc.Key);
+			JsonWriter->WriteValue(TEXT("BuildVersion"), Desc.Toc.Meta.BuildVersion);
+			JsonWriter->WriteValue(TEXT("TargetPlatform"), Desc.Toc.Meta.TargetPlatform);
+			JsonWriter->WriteArrayStart(TEXT("ChunkKeys"));
+			for (const FString& Key : ObjKeys)
+			{
+				JsonWriter->WriteValue(Key);
+			}
+			JsonWriter->WriteArrayEnd();
+			JsonWriter->WriteObjectEnd();
+		}
+		JsonWriter->WriteArrayEnd();
+
+		JsonWriter->WriteObjectEnd();
+		JsonWriter->Close();
+
+		UE_LOG(LogIas, Display, TEXT("Saving chunk key(s) '%s'"), *Params.ChunkKeys);
+		if (!FFileHelper::SaveStringToFile(Json, *Params.ChunkKeys))
+		{
+			return FIoStatus(EIoErrorCode::WriteError, TEXTVIEW("Failed writing JSON file")); 
+		}
+	}
+
+	return FIoStatus::Ok;
+}
+
 ////////////////////////////////////////////////////////////////////////////////
 static int32 ListTocsCommandEntry(const Tool::FContext& Context)
 {
-	FIoStoreListTocsParams Params = BuildListTocsParams(Context);
+	FListTocsParams Params = BuildListTocsParams(Context);
 
 	FIoStatus Status = ListTocs(Params);
 	if (!Status.IsOk())
diff --git a/Engine/Source/Runtime/Experimental/IoStore/OnDemand/Public/IO/IoStoreOnDemand.h b/Engine/Source/Runtime/Experimental/IoStore/OnDemand/Public/IO/IoStoreOnDemand.h
--- a/Engine/Source/Runtime/Experimental/IoStore/OnDemand/Public/IO/IoStoreOnDemand.h
+++ b/Engine/Source/Runtime/Experimental/IoStore/OnDemand/Public/IO/IoStoreOnDemand.h
@@ -29,6 +29,7 @@ struct FIoContainerSettings;
 struct FIoStoreWriterSettings;
 namespace UE::IoStore { struct FOnDemandEndpoint; }
 namespace UE::IoStore { class FOnDemandIoStore; }
+namespace UE::IoStore { class IOnDemandIoDispatcherBackend; }
 using FIoBlockHash = uint32;
 
 // Custom initialization allows users to control when
@@ -190,71 +191,6 @@ struct FOnDemandToc
 
 UE_API bool LoadFromCompactBinary(FCbFieldView Field, FOnDemandToc& OutToc);
 
-#if (IS_PROGRAM || WITH_EDITOR)
-
-////////////////////////////////////////////////////////////////////////////////
-/**
- * Parameters for listing uploaded TOC file(s) from an S3 compatible endpoint.
- *
- * Example usage:
- *
- * 1) Print available TOC's from a local server to standard out.
- * UnrealPak.exe -ListTocs -ServiceUrl="http://10.24.101.92:9000" -Bucket=<bucketname> -BucketPrefix=<some/path/to/data> -AccessKey=<accesskey> -SecretKey=<secretkey>
- *
- * 2) Print available TOC's from AWS S3.
- * UnrealPak.exe -Region="us-east-1" -BucketPath=<mybucket/some/data/path/> -CredentialsFile=<path/to/credentials.txt> -BuildVersion=<version> -Json=<path/to/file.json>
- * 
- * 3) Serialize all TOC's matching a specific build version to JSON:
- * UnrealPak.exe -Region="us-east-1" -BucketPath=<mybucket/some/data/path/> -CredentialsFile=<path/to/credentials.txt> -BuildVersion=<version> -Json=<path/to/file.json>
- *
- * 4) Serialize all chunk object key(s) to JSON.
- * UnrealPak.exe -Region="us-east-1" -BucketPath=<mybucket/some/data/path/> -CredentialsFile=<path/to/credentials.txt> -BuildVersion=<version> -ChunkKeys=<path/to/file.json>
- *
- * 5) Fetch a TOC from a public CDN.
- * UnrealPak.exe -ListTocs -TocUrl=<http://some.public.endpoint.net/path/to/1a32076ca12bfc6feb982ffb064d18f28156606c.iochunktoc>
- *
- * Parameters: -TocEntries, -BlockSizes and -BlockHashes controls what to include when serializing TOC's to JSON.
- *
- * Credentials file example:
- *
- * [default]
- * aws_access_key_id="<key>"
- * aws_secret_access_key="<key>
- * aws_session_token="<token>"
- *
- * Note: All values must be surounded with "".
- */
-struct FIoStoreListTocsParams
-{
-	FString OutFile;
-	FString ServiceUrl;
-	FString Bucket;
-	FString BucketPrefix;
-	FString Region; 
-	FString AccessKey;
-	FString SecretKey;
-	FString SessionToken;
-	FString CredentialsFile;
-	FString CredentialsFileKeyName;
-	FString TocUrl;
-	FString TocKey;
-	FString BuildVersion;
-	FString TargetPlatform;
-	FString ChunkKeys;
-	bool bTocEntries = false;
-	bool bBlockSizes = false;
-	bool bBlockHashes = false;
-
-	static TIoStatusOr<FIoStoreListTocsParams> Parse(const TCHAR* CommandLine);
-	FIoStatus Validate() const;
-};
-
-FIoStatus ListTocs(const FIoStoreListTocsParams& Params);
-
-#endif // (IS_PROGRAM || WITH_EDITOR)
-
-class IOnDemandIoDispatcherBackend;
-
 #if UE_IAS_CUSTOM_INITIALIZATION
 
 /** Result of calling FIoStoreOnDemandModule::Initialize */
