Remove use of context-wrapping object to expressing everything in terms of certs.

Prior to this change, a per-connection context was used either with system default authorities or ones given by the requester. This was both wasteful, and problematic as there is no consistent way to get authorities across all systems. Now the context is abtracted away and presented as a set of authorities, and intent on how a connection should be verified (or not) is explicit.

#rnx

diff --git a/Engine/Source/Runtime/Experimental/IoStore/HttpClient/Internal/IO/Http/Client.h b/Engine/Source/Runtime/Experimental/IoStore/HttpClient/Internal/IO/Http/Client.h
--- a/Engine/Source/Runtime/Experimental/IoStore/HttpClient/Internal/IO/Http/Client.h
+++ b/Engine/Source/Runtime/Experimental/IoStore/HttpClient/Internal/IO/Http/Client.h
@@ -46,8 +46,7 @@ enum class EStatusCodeClass
 
 ////////////////////////////////////////////////////////////////////////////////
 using	FCertRootsRef	= UPTRINT;
-using	FTicket		= uint64;
-using	FPemCert	= FMemoryView;
+using	FTicket			= uint64;
 struct	FActivity;
 
 ////////////////////////////////////////////////////////////////////////////////
@@ -82,12 +81,11 @@ public:
 	{
 		int32				SetHostFromUrl(FAnsiStringView Url);
 		FAnsiStringView		HostName;
-		FPemCert			VerifyCert;
+		FCertRootsRef		VerifyCert = FCertRoots::NoTls();
 		int32				SendBufSize = -1;
 		int32				RecvBufSize = -1;
 		uint32				Port = 0;
 		uint16				ConnectionCount = 1;
-		bool				bUseTls = false;
 		/*
 		enum class ProxyType { Http, Socks4 };
 		Proxy = { ip, port, type }
@@ -214,7 +212,7 @@ public:
 
 	struct FRequestParams
 	{
-		FPemCert			VerifyCert;
+		FCertRootsRef		VerifyCert		= {};
 		uint32				BufferSize		= 256;
 		bool				bAutoRedirect	= false;
 	};
diff --git a/Engine/Source/Runtime/Experimental/IoStore/HttpClient/Private/ConnectionPool.inl b/Engine/Source/Runtime/Experimental/IoStore/HttpClient/Private/ConnectionPool.inl
--- a/Engine/Source/Runtime/Experimental/IoStore/HttpClient/Private/ConnectionPool.inl
+++ b/Engine/Source/Runtime/Experimental/IoStore/HttpClient/Private/ConnectionPool.inl
@@ -19,7 +19,7 @@ public:
 		const ANSICHAR* HostName;
 		uint32			Port = 0;
 		uint32			MaxConnections = 1;
-		FPemCert		VerifyCert;
+		FCertRootsRef	VerifyCert;
 	};
 
 					FHost(const FParams& Params);
@@ -28,8 +28,7 @@ public:
 	FOutcome		Connect(FTlsPeer& Peer);
 	int32			IsResolved() const;
 	FOutcome		ResolveHostName();
-	bool			WithTls() const				{ return GetVerifyCert().GetData() != nullptr; }
-	FPemCert		GetVerifyCert() const		{ return VerifyCert; }
+	FCertRootsRef	GetVerifyCert() const		{ return VerifyCert; }
 	uint32			GetMaxConnections() const	{ return MaxConnections; }
 	uint32			GetIpAddress() const		{ return IpAddresses[0]; }
 	FAnsiStringView	GetHostName() const			{ return HostName; }
@@ -37,7 +36,7 @@ public:
 
 private:
 	FOutcome		Connect(FSocket& Socket);
-	FPemCert		VerifyCert;
+	FCertRootsRef	VerifyCert;
 	const ANSICHAR*	HostName;
 	uint32			IpAddresses[4] = {};
 	int16			SendBufKb = -1;
@@ -57,7 +56,7 @@ FHost::FHost(const FParams& Params)
 
 	if (Port == 0)
 	{
-		Port = (VerifyCert.GetData() == nullptr) ? 80 : 443;
+		Port = (VerifyCert == ECertRootsRefType::None) ? 80 : 443;
 	}
 }
 
@@ -154,14 +153,13 @@ FOutcome FHost::Connect(FTlsPeer& Peer)
 		return Outcome;
 	}
 
-	if (VerifyCert.GetData() == nullptr)
+	if (VerifyCert == ECertRootsRefType::None)
 	{
 		Peer = FTlsPeer(MoveTemp(Socket));
 		return Outcome;
 	}
 
-	FSslContext SslContext(VerifyCert);
-	Peer = FTlsPeer(MoveTemp(Socket), &SslContext, HostName);
+	Peer = FTlsPeer(MoveTemp(Socket), VerifyCert, HostName);
 	return Outcome;
 }
 
@@ -243,7 +241,12 @@ int32 FConnectionPool::FParams::SetHostFromUrl(FAnsiStringView Url)
 	}
 
 	HostName = Offsets.HostName.Get(Url);
-	bUseTls = (Offsets.SchemeLength == 5);
+	
+	VerifyCert = FCertRoots::NoTls();	
+	if (Offsets.SchemeLength == 5)
+	{
+		VerifyCert = FCertRoots::Default();
+	}
 
 	if (Offsets.Port)
 	{
@@ -270,20 +273,12 @@ FConnectionPool::FConnectionPool(const FParams& Params)
 	memcpy(HostDest, Params.HostName.GetData(), HostNameLen);
 	HostDest[HostNameLen] = '\0';
 
-	// Verify
-	FPemCert VerifyCert;
-	if (Params.bUseTls)
-	{
-		VerifyCert = Params.VerifyCert;
-		VerifyCert = (VerifyCert.GetData() != nullptr) ? VerifyCert : FPemCert("", 0);
-	}
-
 	// Init internal object
 	new (Internal) FHost({
 		.HostName		= HostDest,
 		.Port			= Params.Port,
 		.MaxConnections	= Params.ConnectionCount,
-		.VerifyCert		= VerifyCert,
+		.VerifyCert		= Params.VerifyCert,
 	});
 	Internal->SetBufferSize(FHost::EDirection::Send, Params.SendBufSize);
 	Internal->SetBufferSize(FHost::EDirection::Recv, Params.RecvBufSize);
diff --git a/Engine/Source/Runtime/Experimental/IoStore/HttpClient/Private/Loop.inl b/Engine/Source/Runtime/Experimental/IoStore/HttpClient/Private/Loop.inl
--- a/Engine/Source/Runtime/Experimental/IoStore/HttpClient/Private/Loop.inl
+++ b/Engine/Source/Runtime/Experimental/IoStore/HttpClient/Private/Loop.inl
@@ -1294,13 +1294,14 @@ FRequest FEventLoop::Request(
 	// Create an activity and an emphemeral host
 	Params = (Params != nullptr) ? Params : &GDefaultParams;
 
-	FPemCert VerifyCert;
+	FCertRootsRef VerifyCert = FCertRoots::NoTls();
 	if (UrlOffsets.SchemeLength == 5)
 	{
-		if (VerifyCert = Params->VerifyCert; VerifyCert.GetSize() == 0)
+		if (VerifyCert = Params->VerifyCert; VerifyCert == ECertRootsRefType::None)
 		{
-			VerifyCert = FPemCert("", 0);
+			VerifyCert = FCertRoots::Default();
 		}
+		check(VerifyCert != ECertRootsRefType::None);
 	}
 
 	uint32 BufferSize = Params->BufferSize;
@@ -1338,7 +1339,7 @@ FRequest FEventLoop::Request(
 	const FRequestParams* Params)
 {
 	check(Pool.Ptr != nullptr);
-	check(Params == nullptr || Params->VerifyCert.GetData() == nullptr); // add cert to FConPool instead
+	check(Params == nullptr || Params->VerifyCert == ECertRootsRefType::None); // add cert to FConPool instead
 
 	Params = (Params != nullptr) ? Params : &GDefaultParams;
 
@@ -1402,7 +1403,7 @@ bool FEventLoop::Redirect(const FTicketStatus& Status, FTicketSink& OuterSink)
 		FHost& Host = *(Activity.Host);
 
 		TAnsiStringBuilder<256> Url;
-		Url << (Host.WithTls() ? "https" : "http");
+		Url << ((Host.GetVerifyCert() != ECertRootsRefType::None) ? "https" : "http");
 		Url << "://";
 		Url << Host.GetHostName();
 		Url << ":" << Host.GetPort();
diff --git a/Engine/Source/Runtime/Experimental/IoStore/HttpClient/Private/Peer.inl b/Engine/Source/Runtime/Experimental/IoStore/HttpClient/Private/Peer.inl
--- a/Engine/Source/Runtime/Experimental/IoStore/HttpClient/Private/Peer.inl
+++ b/Engine/Source/Runtime/Experimental/IoStore/HttpClient/Private/Peer.inl
@@ -19,60 +19,6 @@ static_assert(OPENSSL_VERSION_NUMBER >= 0x10100000L, "version supporting autoini
 
 
 
-////////////////////////////////////////////////////////////////////////////////
-class FSslContext
-{
-public:
-				FSslContext(const FPemCert& CertPin={});
-				~FSslContext()					{ SSL_CTX_free(Context); }
-				operator SSL_CTX* () const		{ return Context; }
-	bool		AddCert(const FPemCert& Cert);
-
-private:
-	SSL_CTX*	Context;
-};
-
-////////////////////////////////////////////////////////////////////////////////
-FSslContext::FSslContext(const FPemCert& CertPin)
-{
-	Context = SSL_CTX_new(TLS_client_method());
-
-	uint32 ProtoFlags = SSL_OP_NO_SSLv2 | SSL_OP_NO_SSLv3;
-	SSL_CTX_set_options(Context, ProtoFlags);
-
-	if (CertPin.GetSize() == 0)
-	{
-#if WITH_SSL
-		const ISslCertificateManager& CertManager = FSslModule::Get().GetCertificateManager();
-		CertManager.AddCertificatesToSslContext(Context);
-#endif
-		return;
-	}
-
-	AddCert(CertPin);
-}
-
-////////////////////////////////////////////////////////////////////////////////
-bool FSslContext::AddCert(const FPemCert& Cert)
-{
-	const void* Data = Cert.GetData();
-	uint32 Size = uint32(Cert.GetSize());
-	if (Data == nullptr || Size == 0)
-	{
-		return false;
-	}
-
-	BIO* Bio = BIO_new_mem_buf(Data, Size);
-	X509* FiveOhNine = PEM_read_bio_X509(Bio, nullptr, 0, nullptr);
-	check(FiveOhNine != nullptr);
-
-	X509_STORE* Store = SSL_CTX_get_cert_store(Context);
-	int32 Result = X509_STORE_add_cert(Store, FiveOhNine);
-	return Result == 1;
-}
-
-
-
 ////////////////////////////////////////////////////////////////////////////////
 static FCertRoots GDefaultCertRoots;
 
@@ -232,7 +178,7 @@ public:
 				~FTlsPeer();
 				FTlsPeer(FTlsPeer&& Rhs)				{ Move(MoveTemp(Rhs)); }
 				FTlsPeer& operator = (FTlsPeer&& Rhs)	{ return Move(MoveTemp(Rhs)); }
-				FTlsPeer(FSocket InSocket, const FSslContext* Context=nullptr, const char* HostName=nullptr);
+				FTlsPeer(FSocket InSocket, FCertRootsRef Certs=ECertRootsRefType::None, const char* HostName=nullptr);
 	FTlsPeer&	Move(FTlsPeer&& Rhs);
 	FOutcome	Handshake();
 	FOutcome	Send(const char* Data, int32 Size);
@@ -249,14 +195,21 @@ private:
 };
 
 ////////////////////////////////////////////////////////////////////////////////
-FTlsPeer::FTlsPeer(FSocket InSocket, const FSslContext* Context, const char* HostName)
+FTlsPeer::FTlsPeer(FSocket InSocket, FCertRootsRef Certs, const char* HostName)
 : FPeer(MoveTemp(InSocket))
 {
-	if (Context == nullptr)
+	if (Certs == ECertRootsRefType::None)
 	{
 		return;
 	}
 
+	if (Certs == ECertRootsRefType::Default)
+	{
+		Certs = FCertRoots::Explicit(GDefaultCertRoots);
+		check(Certs != 0);
+	}
+	auto* Context = (SSL_CTX*)Certs;
+
 	static BIO_METHOD* BioMethod = nullptr;
 	if (BioMethod == nullptr)
 	{
@@ -278,7 +231,7 @@ FTlsPeer::FTlsPeer(FSocket InSocket, const FSslContext* Context, const char* Hos
 
 	// SSL_MODE_ENABLE_PARTIAL_WRITE ??!!!
 
-	Ssl = SSL_new(*Context);
+	Ssl = SSL_new(Context);
 	SSL_set_connect_state(Ssl);
 	SSL_set0_rbio(Ssl, Bio);
 	SSL_set0_wbio(Ssl, Bio);
@@ -434,7 +387,7 @@ class FHttpPeer
 {
 public:
 				FHttpPeer() = default;
-				FHttpPeer(FSocket InSocket, FSslContext* Context=nullptr, const char* HostName=nullptr);
+				FHttpPeer(FSocket InSocket, FCertRootsRef Certs=ECertRootsRefType::None, const char* HostName=nullptr);
 	FOutcome	Handshake();
 
 private:
@@ -443,8 +396,8 @@ private:
 };
 
 ////////////////////////////////////////////////////////////////////////////////
-FHttpPeer::FHttpPeer(FSocket InSocket, FSslContext* Context, const char* HostName)
-: FTlsPeer(MoveTemp(InSocket), Context, HostName)
+FHttpPeer::FHttpPeer(FSocket InSocket, FCertRootsRef Certs, const char* HostName)
+: FTlsPeer(MoveTemp(InSocket), Certs, HostName)
 {
 	if (Ssl == nullptr)
 	{
diff --git a/Engine/Source/Runtime/Experimental/IoStore/HttpClient/Private/Test.inl b/Engine/Source/Runtime/Experimental/IoStore/HttpClient/Private/Test.inl
--- a/Engine/Source/Runtime/Experimental/IoStore/HttpClient/Private/Test.inl
+++ b/Engine/Source/Runtime/Experimental/IoStore/HttpClient/Private/Test.inl
@@ -2,6 +2,9 @@
 
 #pragma once
 
+#include <HAL/FileManager.h>
+#include <Misc/Paths.h>
+
 namespace UE::IoStore::HTTP
 {
 
@@ -253,7 +256,70 @@ static void ThrottleTest(FAnsiStringView TestUrl)
 }
 
 ////////////////////////////////////////////////////////////////////////////////
-static void RedirectTest(const ANSICHAR* TestHost)
+static void TlsLoadRootCerts()
+{
+	IFileManager& Ifm = IFileManager::Get();
+	FString PemPath = FPaths::EngineDir() / TEXT("Content/Certificates/ThirdParty/cacert.pem");
+	FArchive* Reader = Ifm.CreateFileReader(*PemPath);
+
+	uint32 Size = uint32(Reader->TotalSize());
+	FIoBuffer PemData(Size);
+	FMutableMemoryView PemView = PemData.GetMutableView();
+	Reader->Serialize(PemView.GetData(), Size);
+
+	FCertRoots CaRoots(PemData.GetView());
+	FCertRoots::SetDefault(MoveTemp(CaRoots));
+
+	delete Reader;
+}
+
+////////////////////////////////////////////////////////////////////////////////
+static void TlsTest()
+{
+	FEventLoop Loop;
+
+	auto WaitForLoopIdle = [&] {
+		for (; Loop.Tick(-1); FPlatformProcess::SleepNoStats(0.02f));
+	};
+
+	auto OkSink = [Dest=FIoBuffer()] (const FTicketStatus& Status) mutable {
+		check(Status.GetId() != FTicketStatus::EId::Error);
+		if (Status.GetId() == FTicketStatus::EId::Response)
+		{
+			FResponse& Response = Status.GetResponse();
+			check(Response.GetStatusCode() == 200);
+			Response.SetDestination(&Dest);
+			return;
+		}
+		check(Status.GetId() == FTicketStatus::EId::Content);
+	};
+
+	auto NotOkSink = [Dest=FIoBuffer()] (const FTicketStatus& Status) mutable {
+		check(Status.GetId() == FTicketStatus::EId::Error);
+	};
+
+	static const ANSICHAR* Url = "https://httpbin.org/get";
+
+	{
+		FRequest Request = Loop.Get(Url);
+		Loop.Send(MoveTemp(Request), OkSink);
+		WaitForLoopIdle();
+	}
+
+	if (false) {
+		FRequest Request = Loop.Request("HEAD", "[REDACTED_PROFILE_URL]/clink/releases/download/1.0.0a4/clink-1.0.0a4.zip");
+		Loop.Send(MoveTemp(Request), OkSink);
+		WaitForLoopIdle();
+	}
+
+	{
+		FCertRoots NotACert(FMemoryView("493", 3));
+		check(NotACert.IsValid() == false);
+	}
+}
+
+////////////////////////////////////////////////////////////////////////////////
+static void RedirectTest(const ANSICHAR* TestHost, FCertRootsRef VerifyCert)
 {
 	FEventLoop Loop;
 
@@ -265,18 +331,19 @@ static void RedirectTest(const ANSICHAR* TestHost)
 		.bAutoRedirect = true,
 	};
 
-	enum ReTyp { ReAbs, ReRel };
+	enum ReTyp { ReAbs, ReAbsTls, ReRel, ReRelTls };
 	enum { RecvDataSize = 48 };
 
 	TAnsiStringBuilder<64> Builder;
 	auto BuildUrl = [&] (ReTyp Typ, uint32 Code) -> const auto&
 	{
+		bool bTls = (Typ & 1);
 		Builder.Reset();
-		Builder << "https://";
+		Builder << ((bTls) ? "https://" : "http://");
 		Builder << TestHost;
-		Builder << ":9493";
+		Builder << ":" << (bTls ? 4939 : 9493);
 		Builder << "/redirect";
-		Builder << ((Typ < ReRel) ? "/abs/" : "/rel/");
+		Builder << ((Typ <= ReAbsTls) ? "/abs/" : "/rel/");
 		Builder << Code;
 		Builder << "/data/" << uint32(RecvDataSize);
 		return Builder;
@@ -301,8 +368,9 @@ static void RedirectTest(const ANSICHAR* TestHost)
 
 	uint32 TestCodes[] = { 301, 302, 307, 308 };
 
-	for (auto ReTest : { ReAbs, ReRel })
+	for (auto ReTest : { ReAbs, ReAbsTls, ReRel, ReRelTls })
 	{
+		RequestParams.VerifyCert = (ReTest & 1) ? VerifyCert : 0;
 		RecvCount = 0;
 		for (uint32 Code : TestCodes)
 		{
@@ -320,25 +388,40 @@ static void RedirectTest(const ANSICHAR* TestHost)
 	RequestParams = FEventLoop::FRequestParams();
 	RequestParams.bAutoRedirect = true;
 
-	FConnectionPool::FParams Params;
-	Params.SetHostFromUrl(BuildUrl(ReAbs, 0));
-	Params.ConnectionCount = 4;
-	FConnectionPool Pool(Params);
+	for (auto ReTest : { ReAbs, ReAbsTls, ReRel, ReRelTls })
+	{
+		FConnectionPool::FParams Params;
+		Params.SetHostFromUrl(BuildUrl(ReTest, 0));
+		Params.VerifyCert = (ReTest & 1) ? VerifyCert : 0;
+		Params.ConnectionCount = 4;
+		FConnectionPool Pool(Params);
 
-	RecvCount = 0;
-	Loop.Send(Loop.Get("/redirect/abs/307/data/55", Pool, &RequestParams), OkSink, SinkParam);
-	WaitForLoopIdle();
-	check(RecvCount == 55);
+		RecvCount = 0;
+		uint32 ExpectCount = 0;
+		for (uint32 TestCount : { 4, 267, 55, 17, 1024, 13, 26, 39, 52, 493 })
+		{
+			ExpectCount += TestCount;
+			TAnsiStringBuilder<64> Path;
+			Path << "/redirect/abs/307/data/";
+			Path << TestCount;
+			Loop.Send(Loop.Get(Path.ToString(), Pool, &RequestParams), OkSink, SinkParam);
+		}
+		WaitForLoopIdle();
+		check(RecvCount == ExpectCount);
+	}
 }
 
 ////////////////////////////////////////////////////////////////////////////////
-static void HttpTest(const ANSICHAR* TestHost)
+static void HttpTest(const ANSICHAR* TestHost, FCertRootsRef VerifyCert)
 {
+	const uint32 DefaultPort = (VerifyCert != 0) ? 4939 : 9493;
+
 	TAnsiStringBuilder<64> Ret;
-	auto BuildUrl = [&] (const ANSICHAR* Suffix=nullptr, uint32 Port=9493) -> const auto&
+	auto BuildUrl = [&] (const ANSICHAR* Suffix=nullptr, uint32 Port=0) -> const auto&
 	{
+		Port = Port ? Port : DefaultPort;
 		Ret.Reset();
-		Ret << "http://";
+		Ret << ((Port == 4939) ? "https://" : "http://");
 		Ret << TestHost;
 		Ret << ":" << Port;
 		return (Suffix != nullptr) ? (Ret << Suffix) : Ret;
@@ -435,6 +518,13 @@ static void HttpTest(const ANSICHAR* TestHost)
 		}
 	};
 
+	FEventLoop::FRequestParams ReqParamObj = { .VerifyCert = VerifyCert };
+	const FEventLoop::FRequestParams* ReqParams = nullptr;
+	if (VerifyCert != 0)
+	{
+		ReqParams = &ReqParamObj;
+	}
+
 	// unused request
 	{
 		FRequest Request = Loop.Request("GET", BuildUrl("/data"));
@@ -442,7 +532,7 @@ static void HttpTest(const ANSICHAR* TestHost)
 
 	// foundational
 	{
-		FRequest Request = Loop.Request("GET", BuildUrl("/seed/493"));
+		FRequest Request = Loop.Request("GET", BuildUrl("/seed/493"), ReqParams);
 		Request.Accept(EMimeType::Json);
 
 		FTicket Ticket = Loop.Send(MoveTemp(Request), NullSink);
@@ -452,11 +542,11 @@ static void HttpTest(const ANSICHAR* TestHost)
 
 	// convenience
 	{
-		FRequest Request = Loop.Get(BuildUrl("/data")).Accept(EMimeType::Json);
+		FRequest Request = Loop.Get(BuildUrl("/data"), ReqParams).Accept(EMimeType::Json);
 
 		FTicket Tickets[] = {
+			Loop.Send(Loop.Get(BuildUrl("/data"), ReqParams).Accept(EMimeType::Json), HashSink),
 			Loop.Send(MoveTemp(Request), HashSink),
-			Loop.Send(Loop.Get(BuildUrl("/data")).Accept(EMimeType::Json), HashSink),
 			Loop.Send(Loop.Get("http://httpbin.org/get"), NoErrorSink),
 		};
 		WaitForLoopIdle();
@@ -464,7 +554,7 @@ static void HttpTest(const ANSICHAR* TestHost)
 
 	// convenience
 	{
-		FRequest Request = Loop.Get(BuildUrl("/data")).Accept(EMimeType::Json);
+		FRequest Request = Loop.Get(BuildUrl("/data"), ReqParams).Accept(EMimeType::Json);
 		FTicket Ticket = Loop.Send(MoveTemp(Request), HashSink);
 		WaitForLoopIdle();
 	}
@@ -474,11 +564,14 @@ static void HttpTest(const ANSICHAR* TestHost)
 	{
 		FConnectionPool::FParams Params;
 		Params.SetHostFromUrl(BuildUrl());
+		Params.VerifyCert = VerifyCert;
 		Params.ConnectionCount = (i % 2) + 1;
 		FConnectionPool Pool(Params);
 		for (int32 j = 0; j < i; ++j)
 		{
-			FRequest Request = Loop.Get("/data", Pool);
+			TAnsiStringBuilder<16> Path;
+			Path << "/data?pool=" << i << "x" << j;
+			FRequest Request = Loop.Get(Path, Pool);
 			Loop.Send(MoveTemp(Request), HashSink);
 		}
 		WaitForLoopIdle();
@@ -511,12 +604,13 @@ static void HttpTest(const ANSICHAR* TestHost)
 
 		FConnectionPool::FParams Params;
 		Params.SetHostFromUrl(BuildUrl());
+		Params.VerifyCert = VerifyCert;
 		FConnectionPool Pool(Params);
 
 		FEventLoop Loop2;
 		Loop2.Send(Loop2.Get("/data?stall=1", Pool), Sink);
 
-		// Requests are pipelined. The second one will get went during the stall so
+		// Requests are pipelined. The second one will get sent during the stall so
 		// we expect it to fail. The subsequent ones are expected to succeed.
 		Loop2.Send(Loop2.Get("/data", Pool), ErrorSink);
 		Loop2.Send(Loop2.Get("/data", Pool), HashSink);
@@ -563,7 +657,7 @@ static void HttpTest(const ANSICHAR* TestHost)
 
 		for (int32 i = 0; (i += 69493) < 2 << 20;)
 		{
-			FRequest Request = Loop.Request("HEAD", BuildUrl("/data"));
+			FRequest Request = Loop.Request("HEAD", BuildUrl("/data"), ReqParams);
 			for (int32 j = i; j > 0;)
 			{
 				FAnsiStringView Name(AsciiData, MixTh() + 1);
@@ -630,7 +724,7 @@ static void HttpTest(const ANSICHAR* TestHost)
 				FTicket Tickets[StressLoad];
 				for (FTicket& Ticket : Tickets)
 				{
-					Ticket = Loop.Send(Loop.Get(StressUrl).Header("Accept", "*/*"), Sink);
+					Ticket = Loop.Send(Loop.Get(StressUrl, ReqParams).Header("Accept", "*/*"), Sink);
 				}
 
 				LoopTickDelay = AddDelay;
@@ -653,7 +747,7 @@ static void HttpTest(const ANSICHAR* TestHost)
 		auto StressTaskEntry = [&] {
 			for (uint32 i = 0; i < StressLoad; ++i)
 			{
-				FTicket Ticket = Loop.Send(Loop.Get(Url), HashSink);
+				FTicket Ticket = Loop.Send(Loop.Get(Url, ReqParams), HashSink);
 				if (!Ticket)
 				{
 					FPlatformProcess::SleepNoStats(0.01f);
@@ -682,9 +776,9 @@ static void HttpTest(const ANSICHAR* TestHost)
 		TamperUrl << "/data?tamper=" << i;
 		FAnsiStringView Url = BuildUrl(TamperUrl.ToString());
 
-		for (int j = 0; j < 48; ++j)
+		for (int j = 0; j < 13; ++j)
 		{
-			FRequest Request = Loop.Request("GET", Url);
+			FRequest Request = Loop.Request("GET", Url, ReqParams);
 			Loop.Send(MoveTemp(Request), NullSink);
 		}
 
@@ -697,7 +791,10 @@ static void HttpTest(const ANSICHAR* TestHost)
 	check(Loop.IsIdle());
 
 #if IS_PROGRAM
-	ThrottleTest(BuildUrl("/data/"));
+	if (VerifyCert == 0)
+	{
+		ThrottleTest(BuildUrl("/data/"));
+	}
 #endif
 
 	// pre-generated headers
@@ -705,9 +802,7 @@ static void HttpTest(const ANSICHAR* TestHost)
 	// proxy
 	// chunked transfer encoding
 	// gzip / deflate
-	// redirects
 	// loop multi-req.
-	// tls
 	// url auth credentials
 	// transfer-file / splice / sendfile
 	// (header field parser)
@@ -730,12 +825,43 @@ IOSTOREHTTPCLIENT_API void IasHttpTest(const ANSICHAR* TestHost="localhost")
 #endif
 
 	MiscTest();
-	HttpTest(TestHost);
-	RedirectTest(TestHost);
+
+	FCertRoots TestServerCaChain;
+	{
+		TAnsiStringBuilder<64> CaUrl;
+		CaUrl << "http://" << TestHost << ":9493/ca";
+
+		FIoBuffer CertBuffer;
+
+		FEventLoop Loop;
+		FRequest Request = Loop.Get(CaUrl);
+		Loop.Send(MoveTemp(Request), [&] (const FTicketStatus& Status) {
+			check(Status.GetId() != FTicketStatus::EId::Error);
+
+			if (Status.GetId() == FTicketStatus::EId::Response)
+			{
+				FResponse& Response = Status.GetResponse();
+				Response.SetDestination(&CertBuffer);
+			}
+		});
+		for (; Loop.Tick(-1); FPlatformProcess::SleepNoStats(0.02f));
+
+		TestServerCaChain = FCertRoots(CertBuffer.GetView());
+	}
+
+	FCertRootsRef TestServerCertRef = FCertRoots::Explicit(TestServerCaChain);
+
+	HttpTest(TestHost, FCertRoots::NoTls());
+	HttpTest(TestHost, TestServerCertRef);
+
+	RedirectTest(TestHost, TestServerCertRef);
+
+	TlsLoadRootCerts();
+	TlsTest();
 }
 
 #endif // !SHIP|TEST
+
 // }}}
 
 } // namespace UE::IoStore::HTTP
-
