Removed dependency on unused module

#rnx

diff --git a/Engine/Source/Runtime/Experimental/IoStore/HttpClient/IoStoreHttpClient.Build.cs b/Engine/Source/Runtime/Experimental/IoStore/HttpClient/IoStoreHttpClient.Build.cs
--- a/Engine/Source/Runtime/Experimental/IoStore/HttpClient/IoStoreHttpClient.Build.cs
+++ b/Engine/Source/Runtime/Experimental/IoStore/HttpClient/IoStoreHttpClient.Build.cs
@@ -11,7 +11,6 @@ public class IoStoreHttpClient : ModuleRules
 		PrivateDependencyModuleNames.AddRange(
 			new string[] {
 				"OpenSSL",
-				"SSL",
 			}
 		);
 
diff --git a/Engine/Source/Runtime/Experimental/IoStore/HttpClient/Private/Http.cpp b/Engine/Source/Runtime/Experimental/IoStore/HttpClient/Private/Http.cpp
--- a/Engine/Source/Runtime/Experimental/IoStore/HttpClient/Private/Http.cpp
+++ b/Engine/Source/Runtime/Experimental/IoStore/HttpClient/Private/Http.cpp
@@ -13,8 +13,6 @@
 #include <Misc/ScopeLock.h>
 #include <Misc/StringBuilder.h>
 #include <ProfilingDebugging/CpuProfilerTrace.h>
-#include <Ssl.h>
-#include <Interfaces/ISslCertificateManager.h>
 #include <Tasks/Task.h>
 #include <Trace/Trace.h>
 #endif
