Iris
- Compile with logReferenceCache in test (all) and shipping builds (log)

diff --git a/Engine/Source/Runtime/Experimental/Iris/Core/Private/Iris/ReplicationSystem/ObjectReferenceCache.cpp b/Engine/Source/Runtime/Experimental/Iris/Core/Private/Iris/ReplicationSystem/ObjectReferenceCache.cpp
--- a/Engine/Source/Runtime/Experimental/Iris/Core/Private/Iris/ReplicationSystem/ObjectReferenceCache.cpp
+++ b/Engine/Source/Runtime/Experimental/Iris/Core/Private/Iris/ReplicationSystem/ObjectReferenceCache.cpp
@@ -25,20 +25,18 @@
 #include "UObject/Package.h"
 #include "HAL/IConsoleManager.h"
 
-#ifndef UE_NET_ENABLE_REFERENCECACHE_LOG
-#	define UE_NET_ENABLE_REFERENCECACHE_LOG !(UE_BUILD_SHIPPING || UE_BUILD_TEST)
-#endif 
-
 #ifndef UE_NET_VALIDATE_REFERENCECACHE
 #	define UE_NET_VALIDATE_REFERENCECACHE 0
 #endif
 
-// Only compile warnings or errors when UE_NET_ENABLE_REFERENCECACHE_LOG is false
-#if UE_NET_ENABLE_REFERENCECACHE_LOG
-	DEFINE_LOG_CATEGORY_STATIC(LogIrisReferences, Log, All);
+// Don't compile verbose logs in Shipping builds
+#if (UE_BUILD_SHIPPING)
+#	define UE_NET_REFERENCECACHE_LOG_COMPILE_VERBOSITY Log
 #else
-	DEFINE_LOG_CATEGORY_STATIC(LogIrisReferences, Warning, Warning);
-#endif
+#	define UE_NET_REFERENCECACHE_LOG_COMPILE_VERBOSITY All
+#endif 
+
+DEFINE_LOG_CATEGORY_STATIC(LogIrisReferences, Log, UE_NET_REFERENCECACHE_LOG_COMPILE_VERBOSITY);
 
 #define UE_LOG_REFERENCECACHE(Verbosity, Format, ...)  UE_LOG(LogIrisReferences, Verbosity, Format, ##__VA_ARGS__)
 #define UE_CLOG_REFERENCECACHE(Condition, Verbosity, Format, ...)  UE_CLOG(Condition, LogIrisReferences, Verbosity, Format, ##__VA_ARGS__)
@@ -894,7 +892,7 @@ UObject* FObjectReferenceCache::ResolveObjectReferenceHandleInternal(FNetRefHand
 			{
 				// Something else is already async loading this package, calling load again will add our callback to the existing load request
 				StartAsyncLoadingPackage(*CacheObjectPtr, ObjectPathName, RefHandle, true);
-				UE_LOG(LogNetPackageMap, Log, TEXT("GetObjectFromRefHandle: Listening to existing async load. Path: %s, NetRefHandle: %s"), *ObjectPath, *RefHandle.ToString());
+				UE_LOG(LogIrisReferences, Verbose, TEXT("GetObjectFromRefHandle: Listening to existing async load. Path: %s, NetRefHandle: %s"), *ObjectPath, *RefHandle.ToString());
 			}
 			else if (/*!GbNetCheckNoLoadPackages ||*/ !CacheObjectPtr->bNoLoad)
 			{
@@ -1247,7 +1245,7 @@ void FObjectReferenceCache::WriteFullReferenceInternal(FNetSerializationContext&
 	// Write invalid handle
 	if (!CachedObject)
 	{
-		UE_CLOG(RefHandle.IsValid(), LogIris, Log, TEXT("ObjectReferenceCache::WriteFullReference Trying to write Stale handle, %s"), *RefHandle.ToString());
+		UE_CLOG(RefHandle.IsValid(), LogIrisReferences, Log, TEXT("ObjectReferenceCache::WriteFullReference Trying to write Stale handle, %s"), *RefHandle.ToString());
 
 		WriteNetRefHandle(Context, FNetRefHandle::GetInvalid());
 		return;			
@@ -1507,7 +1505,7 @@ void FObjectReferenceCache::WriteReference(FNetSerializationContext& Context, FN
 	// Write invalid handle
 	if (!CachedObject)
 	{
-		UE_CLOG(RefHandle.IsValid(), LogIris, Verbose, TEXT("ObjectReferenceCache::WriteReference Trying to write Stale handle, %s"), *RefHandle.ToString());
+		UE_CLOG(RefHandle.IsValid(), LogIrisReferences, Verbose, TEXT("ObjectReferenceCache::WriteReference Trying to write Stale handle, %s"), *RefHandle.ToString());
 
 		WriteNetRefHandle(Context, FNetRefHandle::GetInvalid());
 
