Replaced some TEMPLATE_REQUIRES usage.

diff --git a/Engine/Source/Runtime/Engine/Classes/Engine/NetSerialization.h b/Engine/Source/Runtime/Engine/Classes/Engine/NetSerialization.h
--- a/Engine/Source/Runtime/Engine/Classes/Engine/NetSerialization.h
+++ b/Engine/Source/Runtime/Engine/Classes/Engine/NetSerialization.h
@@ -15,6 +15,8 @@
 #include "EngineLogs.h"
 #include "Net/Core/Serialization/QuantizedVectorSerialization.h"
 
+#include <type_traits>
+
 #include "NetSerialization.generated.h"
 
 /**
@@ -290,7 +292,7 @@ struct TFixedCompressedFloatDetails
 	static constexpr int32 MaxDelta = (1 << (NumBits - 0)) - 1;     //   1111 1111 - Max delta is
 };
 
-template<int32 MaxValue, uint32 NumBits, typename T, TEMPLATE_REQUIRES(TIsFloatingPoint<T>::Value), TEMPLATE_REQUIRES(NumBits < 32)>
+template<int32 MaxValue, uint32 NumBits, typename T UE_REQUIRES(std::is_floating_point_v<T> && NumBits < 32)>
 bool WriteFixedCompressedFloat(const T Value, FArchive& Ar)
 {
 	using Details = TFixedCompressedFloatDetails<MaxValue, NumBits>;
diff --git a/Engine/Source/Runtime/Engine/Private/KismetMathLibrary.cpp b/Engine/Source/Runtime/Engine/Private/KismetMathLibrary.cpp
--- a/Engine/Source/Runtime/Engine/Private/KismetMathLibrary.cpp
+++ b/Engine/Source/Runtime/Engine/Private/KismetMathLibrary.cpp
@@ -15,12 +15,14 @@
 #include "Misc/QualifiedFrameTime.h"
 #include "Templates/SubclassOf.h"
 
+#include <type_traits>
+
 #include UE_INLINE_GENERATED_CPP_BY_NAME(KismetMathLibrary)
 
 #define LOCTEXT_NAMESPACE "UKismetMathLibrary"
 
 /** Interpolate a linear alpha value using an ease mode and function. */
-template<typename FloatType, TEMPLATE_REQUIRES(TIsFloatingPoint<FloatType>::Value)>
+template<typename FloatType UE_REQUIRES(std::is_floating_point_v<FloatType>)>
 FloatType EaseAlpha(FloatType InAlpha, uint8 EasingFunc, FloatType BlendExp, int32 Steps)
 {
 	switch (EasingFunc)
diff --git a/Engine/Source/Runtime/Engine/Public/MeshCardRepresentation.h b/Engine/Source/Runtime/Engine/Public/MeshCardRepresentation.h
--- a/Engine/Source/Runtime/Engine/Public/MeshCardRepresentation.h
+++ b/Engine/Source/Runtime/Engine/Public/MeshCardRepresentation.h
@@ -46,7 +46,7 @@ public:
 	}
 
 	// Conversion from other type.
-	template<typename FArg, TEMPLATE_REQUIRES(!std::is_same_v<T, FArg>)>
+	template<typename FArg UE_REQUIRES(!std::is_same_v<T, FArg>)>
 	explicit TLumenCardOBB(const TLumenCardOBB<FArg>& From)
 	{
 		Origin = UE::Math::TVector<T>(From.Origin);
diff --git a/Engine/Source/Runtime/Engine/Public/WorldPartition/WorldPartitionActorDescType.h b/Engine/Source/Runtime/Engine/Public/WorldPartition/WorldPartitionActorDescType.h
--- a/Engine/Source/Runtime/Engine/Public/WorldPartition/WorldPartitionActorDescType.h
+++ b/Engine/Source/Runtime/Engine/Public/WorldPartition/WorldPartitionActorDescType.h
@@ -4,10 +4,12 @@
 #include "Templates/Requires.h"
 #include "Templates/UnrealTypeTraits.h"
 
+#include <type_traits>
+
 class AActor;
 
 #if WITH_EDITOR
-template <typename ActorType, TEMPLATE_REQUIRES(TIsDerivedFrom<ActorType, AActor>::IsDerived)>
+template <typename ActorType UE_REQUIRES(std::is_base_of_v<AActor, ActorType>)>
 struct FWorldPartitionActorDescType
 {};
 
diff --git a/Engine/Source/Runtime/Experimental/Chaos/Public/Chaos/CollectionPropertyFacade.h b/Engine/Source/Runtime/Experimental/Chaos/Public/Chaos/CollectionPropertyFacade.h
--- a/Engine/Source/Runtime/Experimental/Chaos/Public/Chaos/CollectionPropertyFacade.h
+++ b/Engine/Source/Runtime/Experimental/Chaos/Public/Chaos/CollectionPropertyFacade.h
@@ -74,18 +74,18 @@ namespace Chaos::Softs
 		//~ Values access per index, fast, no check, index must be valid (0 <= KeyIndex < Num())
 		const FString& GetKey(int32 KeyIndex) const { return GetValue<const FString&>(KeyIndex, KeyArray); }
 
-		template<typename T, TEMPLATE_REQUIRES(TIsWeightedType<T>::Value)>
+		template<typename T UE_REQUIRES(TIsWeightedType<T>::Value)>
 		T GetLowValue(int32 KeyIndex) const { return GetValue<T>(KeyIndex, LowValueArray); }
 
-		template<typename T, TEMPLATE_REQUIRES(TIsWeightedType<T>::Value)>
+		template<typename T UE_REQUIRES(TIsWeightedType<T>::Value)>
 		T GetHighValue(int32 KeyIndex) const { return GetValue<T>(KeyIndex, HighValueArray); }
 
-		template<typename T, TEMPLATE_REQUIRES(TIsWeightedType<T>::Value)>
+		template<typename T UE_REQUIRES(TIsWeightedType<T>::Value)>
 		TPair<T, T> GetWeightedValue(int32 KeyIndex) const { return MakeTuple(GetLowValue<T>(KeyIndex), GetHighValue<T>(KeyIndex)); }
 
 		FVector2f GetWeightedFloatValue(int32 KeyIndex) const { return FVector2f(GetLowValue<float>(KeyIndex), GetHighValue<float>(KeyIndex)); }
 
-		template<typename T, TEMPLATE_REQUIRES(TIsWeightedType<T>::Value)>
+		template<typename T UE_REQUIRES(TIsWeightedType<T>::Value)>
 		T GetValue(int32 KeyIndex) const { return GetLowValue<T>(KeyIndex); }
 
 		const FString& GetStringValue(int32 KeyIndex) const { return GetValue<const FString&>(KeyIndex, StringValueArray); }
@@ -104,19 +104,19 @@ namespace Chaos::Softs
 		bool IsInterpolable(int32 KeyIndex) const { return HasAnyFlags(KeyIndex, ECollectionPropertyFlags::Interpolable); }
 
 		//~ Values access per key
-		template<typename T, TEMPLATE_REQUIRES(TIsWeightedType<T>::Value)>
+		template<typename T UE_REQUIRES(TIsWeightedType<T>::Value)>
 		T GetLowValue(const FString& Key, const T& Default = T(0), int32* OutKeyIndex = nullptr) const
 		{
 			return SafeGet(Key, [this](int32 KeyIndex)->T { return GetLowValue<T>(KeyIndex); }, Default, OutKeyIndex);
 		}
 
-		template<typename T, TEMPLATE_REQUIRES(TIsWeightedType<T>::Value)>
+		template<typename T UE_REQUIRES(TIsWeightedType<T>::Value)>
 		T GetHighValue(const FString& Key, const T& Default = T(0), int32* OutKeyIndex = nullptr) const
 		{
 			return SafeGet(Key, [this](int32 KeyIndex)->T { return GetHighValue<T>(KeyIndex); }, Default, OutKeyIndex);
 		}
 
-		template<typename T, TEMPLATE_REQUIRES(TIsWeightedType<T>::Value)>
+		template<typename T UE_REQUIRES(TIsWeightedType<T>::Value)>
 		TPair<T, T> GetWeightedValue(const FString& Key, const T& Default = T(0), int32* OutKeyIndex = nullptr) const
 		{
 			return SafeGet(Key, [this](int32 KeyIndex)->TPair<T, T> { return GetWeightedValue<T>(KeyIndex); }, MakeTuple(Default, Default), OutKeyIndex);
@@ -132,7 +132,7 @@ namespace Chaos::Softs
 			return SafeGet(Key, [this](int32 KeyIndex)->FVector2f { return GetWeightedFloatValue(KeyIndex); }, Default, OutKeyIndex);
 		}
 
-		template<typename T, TEMPLATE_REQUIRES(TIsWeightedType<T>::Value)>
+		template<typename T UE_REQUIRES(TIsWeightedType<T>::Value)>
 		T GetValue(const FString& Key, const T& Default = T(0), int32* OutKeyIndex = nullptr) const
 		{
 			return SafeGet(Key, [this](int32 KeyIndex)->T { return GetValue<T>(KeyIndex); }, Default, OutKeyIndex);
@@ -272,18 +272,18 @@ PRAGMA_ENABLE_DEPRECATION_WARNINGS
 		FCollectionPropertyFacade& operator=(FCollectionPropertyFacade&&) = default;
 
 		//~ Values set per index
-		template<typename T, TEMPLATE_REQUIRES(TIsWeightedType<T>::Value)>
+		template<typename T UE_REQUIRES(TIsWeightedType<T>::Value)>
 		void SetLowValue(int32 KeyIndex, const T& Value) { SetValue(KeyIndex, GetLowValueArray(), FVector3f(Value)); }
 
-		template<typename T, TEMPLATE_REQUIRES(TIsWeightedType<T>::Value)>
+		template<typename T UE_REQUIRES(TIsWeightedType<T>::Value)>
 		void SetHighValue(int32 KeyIndex, const T& Value) { SetValue(KeyIndex, GetHighValueArray(), FVector3f(Value)); }
 
-		template<typename T, TEMPLATE_REQUIRES(TIsWeightedType<T>::Value)>
+		template<typename T UE_REQUIRES(TIsWeightedType<T>::Value)>
 		void SetWeightedValue(int32 KeyIndex, const T& LowValue, const T& HighValue) { SetLowValue(KeyIndex, LowValue); SetHighValue(KeyIndex, HighValue); }
 
 		void SetWeightedFloatValue(int32 KeyIndex, const FVector2f& Value) { SetLowValue<float>(KeyIndex, Value.X); SetHighValue<float>(KeyIndex, Value.Y); }
 
-		template<typename T, TEMPLATE_REQUIRES(TIsWeightedType<T>::Value)>
+		template<typename T UE_REQUIRES(TIsWeightedType<T>::Value)>
 		void SetValue(int32 KeyIndex, const T& Value) { SetWeightedValue(KeyIndex, Value, Value); }
 
 		void SetStringValue(int32 KeyIndex, const FString& Value) { if (GetStringValueArray()[KeyIndex] != Value) { GetStringValueArray()[KeyIndex] = Value; SetStringDirty(KeyIndex); } }
@@ -306,19 +306,19 @@ PRAGMA_ENABLE_DEPRECATION_WARNINGS
 		void SetInterpolable(int32 KeyIndex) { EnableFlags(KeyIndex, ECollectionPropertyFlags::Interpolable, true); }
 
 		//~ Values set per key
-		template<typename T, TEMPLATE_REQUIRES(TIsWeightedType<T>::Value)>
+		template<typename T UE_REQUIRES(TIsWeightedType<T>::Value)>
 		int32 SetLowValue(const FString& Key, const T& Value)
 		{
 			return SafeSet(Key, [this, &Value](int32 KeyIndex) { SetLowValue(KeyIndex, Value); });
 		}
 
-		template<typename T, TEMPLATE_REQUIRES(TIsWeightedType<T>::Value)>
+		template<typename T UE_REQUIRES(TIsWeightedType<T>::Value)>
 		int32 SetHighValue(const FString& Key, const T& Value)
 		{
 			return SafeSet(Key, [this, &Value](int32 KeyIndex) { SetHighValue(KeyIndex, Value); });
 		}
 
-		template<typename T, TEMPLATE_REQUIRES(TIsWeightedType<T>::Value)>
+		template<typename T UE_REQUIRES(TIsWeightedType<T>::Value)>
 		int32 SetWeightedValue(const FString& Key, const T& LowValue, const T& HighValue)
 		{
 			return SafeSet(Key, [this, &LowValue, &HighValue](int32 KeyIndex) { SetWeightedValue(KeyIndex, LowValue, HighValue); });
@@ -329,7 +329,7 @@ PRAGMA_ENABLE_DEPRECATION_WARNINGS
 			return SafeSet(Key, [this, &Value](int32 KeyIndex) { SetWeightedFloatValue(KeyIndex, Value); });
 		}
 
-		template<typename T, TEMPLATE_REQUIRES(TIsWeightedType<T>::Value)>
+		template<typename T UE_REQUIRES(TIsWeightedType<T>::Value)>
 		int32 SetValue(const FString& Key, const T& Value)
 		{
 			return SafeSet(Key, [this, &Value](int32 KeyIndex) { SetValue(KeyIndex, Value); });
@@ -498,17 +498,17 @@ PRAGMA_ENABLE_DEPRECATION_WARNINGS
 		CHAOS_API void Update(const TSharedPtr<const FManagedArrayCollection>& InManagedArrayCollection, ECollectionPropertyUpdateFlags UpdateFlags);
 
 		//~ Add values
-		template<typename T, TEMPLATE_REQUIRES(TIsWeightedType<T>::Value)>
+		template<typename T UE_REQUIRES(TIsWeightedType<T>::Value)>
 		inline int32 AddWeightedValue(const FString& Key, const T& LowValue, const T& HighValue, ECollectionPropertyFlags Flags = ECollectionPropertyFlags::Enabled);
-		template<typename T, TEMPLATE_REQUIRES(TIsWeightedType<T>::Value)>
+		template<typename T UE_REQUIRES(TIsWeightedType<T>::Value)>
 		inline int32 AddWeightedValue(const FString& Key, const T& LowValue, const T& HighValue, bool bEnabled, bool bAnimatable = false, bool bIntrinsic = false);
 
 		CHAOS_API int32 AddWeightedFloatValue(const FString& Key, const FVector2f& Value, ECollectionPropertyFlags Flags = ECollectionPropertyFlags::Enabled);
 		CHAOS_API int32 AddWeightedFloatValue(const FString& Key, const FVector2f& Value, bool bEnabled, bool bAnimatable, bool bIntrinsic = false);
 
-		template<typename T, TEMPLATE_REQUIRES(TIsWeightedType<T>::Value)>
+		template<typename T UE_REQUIRES(TIsWeightedType<T>::Value)>
 		int32 AddValue(const FString& Key, const T& Value, ECollectionPropertyFlags Flags = ECollectionPropertyFlags::Enabled) { return AddWeightedValue(Key, Value, Value, Flags); }
-		template<typename T, TEMPLATE_REQUIRES(TIsWeightedType<T>::Value)>
+		template<typename T UE_REQUIRES(TIsWeightedType<T>::Value)>
 		int32 AddValue(const FString& Key, const T& Value, bool bEnabled, bool bAnimatable = false) { return AddWeightedValue(Key, Value, Value, bEnabled, bAnimatable); }
 
 		CHAOS_API int32 AddStringValue(const FString& Key, const FString& Value, ECollectionPropertyFlags Flags = ECollectionPropertyFlags::Enabled);
@@ -525,7 +525,7 @@ PRAGMA_ENABLE_DEPRECATION_WARNINGS
 		}
 	}
 
-	template<typename T, typename TEnableIf<TIsWeightedType<T>::Value, int>::type>
+	template<typename T UE_REQUIRES_DEFINITION(TIsWeightedType<T>::Value)>
 	inline int32 FCollectionPropertyMutableFacade::AddWeightedValue(const FString& Key, const T& LowValue, const T& HighValue, ECollectionPropertyFlags Flags)
 	{
 		const int32 KeyIndex = AddProperty(Key, Flags);
@@ -533,7 +533,7 @@ PRAGMA_ENABLE_DEPRECATION_WARNINGS
 		return KeyIndex;
 	}
 
-	template<typename T, typename TEnableIf<TIsWeightedType<T>::Value, int>::type>
+	template<typename T UE_REQUIRES_DEFINITION(TIsWeightedType<T>::Value)>
 	inline int32 FCollectionPropertyMutableFacade::AddWeightedValue(const FString& Key, const T& LowValue, const T& HighValue, bool bEnabled, bool bAnimatable, bool bIntrinsic)
 	{
 		const int32 KeyIndex = AddProperty(Key, bEnabled, bAnimatable, bIntrinsic);
diff --git a/Engine/Source/Runtime/Experimental/Chaos/Public/Chaos/PBDSpringConstraintsBase.h b/Engine/Source/Runtime/Experimental/Chaos/Public/Chaos/PBDSpringConstraintsBase.h
--- a/Engine/Source/Runtime/Experimental/Chaos/Public/Chaos/PBDSpringConstraintsBase.h
+++ b/Engine/Source/Runtime/Experimental/Chaos/Public/Chaos/PBDSpringConstraintsBase.h
@@ -16,7 +16,7 @@ namespace Chaos::Softs
 class FPBDSpringConstraintsBase
 {
 public:
-	template<int32 Valence, TEMPLATE_REQUIRES(Valence >= 2 && Valence <= 4)>
+	template<int32 Valence UE_REQUIRES(Valence >= 2 && Valence <= 4)>
 	FPBDSpringConstraintsBase(
 		const FSolverParticlesRange& Particles,
 		const TArray<TVector<int32, Valence>>& InConstraints,
@@ -51,7 +51,7 @@ public:
 		}
 	}
 
-	template<int32 Valence, TEMPLATE_REQUIRES(Valence >= 2 && Valence <= 4)>
+	template<int32 Valence UE_REQUIRES(Valence >= 2 && Valence <= 4)>
 	FPBDSpringConstraintsBase(
 		const FSolverParticles& Particles,
 		int32 InParticleOffset,
@@ -123,9 +123,8 @@ protected:
 	}
 
 private:
-	template<int32 Valence, typename Predicate>
-	typename TEnableIf<Valence >= 2 && Valence <= 4, TArray<TVector<int32, 2>>>::Type
-	TrimConstraints(const TArray<TVector<int32, Valence>>& InConstraints, Predicate TrimPredicate)
+	template<int32 Valence, typename Predicate UE_REQUIRES(Valence >= 2 && Valence <= 4)>
+	TArray<TVector<int32, 2>> TrimConstraints(const TArray<TVector<int32, Valence>>& InConstraints, Predicate TrimPredicate)
 	{
 		TSet<TVec2<int32>> TrimmedConstraints;
 		TrimmedConstraints.Reserve(Valence == 2 ? InConstraints.Num() : InConstraints.Num() * Chaos::Utilities::NChooseR(Valence, 2));
diff --git a/Engine/Source/Runtime/Experimental/Chaos/Public/Chaos/PBDStiffness.h b/Engine/Source/Runtime/Experimental/Chaos/Public/Chaos/PBDStiffness.h
--- a/Engine/Source/Runtime/Experimental/Chaos/Public/Chaos/PBDStiffness.h
+++ b/Engine/Source/Runtime/Experimental/Chaos/Public/Chaos/PBDStiffness.h
@@ -39,7 +39,7 @@ public:
 	/**
 	 * Weightmap constraint constructor. 
 	 */
-	template<int32 Valence, TEMPLATE_REQUIRES(Valence >= 2 && Valence <= 4)>
+	template<int32 Valence UE_REQUIRES(Valence >= 2 && Valence <= 4)>
 	inline FPBDStiffness(
 		const FSolverVec2& InWeightedValue,
 		const TConstArrayView<FRealSingle>& Multipliers = TConstArrayView<FRealSingle>(),
@@ -99,7 +99,7 @@ FPBDStiffness::FPBDStiffness(
 	SetWeightedValue(InWeightedValue, MaxStiffness);
 }
 
-template<int32 Valence, typename TEnableIf<Valence >= 2 && Valence <= 4, int>::type>
+template<int32 Valence UE_REQUIRES_DEFINITION(Valence >= 2 && Valence <= 4)>
 FPBDStiffness::FPBDStiffness(
 	const FSolverVec2& InWeightedValue,
 	const TConstArrayView<FRealSingle>& Multipliers,
diff --git a/Engine/Source/Runtime/Experimental/Chaos/Public/Chaos/XPBDSpringConstraints.h b/Engine/Source/Runtime/Experimental/Chaos/Public/Chaos/XPBDSpringConstraints.h
--- a/Engine/Source/Runtime/Experimental/Chaos/Public/Chaos/XPBDSpringConstraints.h
+++ b/Engine/Source/Runtime/Experimental/Chaos/Public/Chaos/XPBDSpringConstraints.h
@@ -26,32 +26,7 @@ public:
 	static constexpr FSolverReal MinDampingRatio = (FSolverReal)0.;
 	static constexpr FSolverReal MaxDampingRatio = (FSolverReal)1000.;
 
-	template<int32 Valence, TEMPLATE_REQUIRES(Valence >= 2 && Valence <= 4)>
-	UE_DEPRECATED(5.2, "Use the other constructor instead.")
-	FXPBDSpringConstraints(
-		const FSolverParticles& Particles,
-		int32 ParticleOffset,
-		int32 ParticleCount,
-		const TArray<TVector<int32, Valence>>& InConstraints,
-		const TConstArrayView<FRealSingle>& StiffnessMultipliers,
-		const FSolverVec2& InStiffness,
-		bool bTrimKinematicConstraints = false)
-		: Base(
-			Particles,
-			ParticleOffset,
-			ParticleCount,
-			InConstraints,
-			StiffnessMultipliers,
-			InStiffness,
-			true /*bTrimKinematicConstraints*/,
-			MaxStiffness)
-		, DampingRatio(FSolverVec2::ZeroVector)
-	{
-		Lambdas.Init((FSolverReal)0., Constraints.Num());
-		InitColor(Particles);
-	}	
-	
-	template<int32 Valence, TEMPLATE_REQUIRES(Valence >= 2 && Valence <= 4)>
+	template<int32 Valence UE_REQUIRES(Valence >= 2 && Valence <= 4)>
 	FXPBDSpringConstraints(
 		const FSolverParticlesRange& Particles,
 		const TArray<TVector<int32, Valence>>& InConstraints,
@@ -77,7 +52,7 @@ public:
 		InitColor(Particles);
 	}
 
-	template<int32 Valence, TEMPLATE_REQUIRES(Valence >= 2 && Valence <= 4)>
+	template<int32 Valence UE_REQUIRES(Valence >= 2 && Valence <= 4)>
 	UE_DEPRECATED(5.4, "XPBD Constraints must always trim kinematic constraints")
 	FXPBDSpringConstraints(
 		const FSolverParticlesRange& Particles,
@@ -105,7 +80,7 @@ public:
 		InitColor(Particles);
 	}
 
-	template<int32 Valence, TEMPLATE_REQUIRES(Valence >= 2 && Valence <= 4)>
+	template<int32 Valence UE_REQUIRES(Valence >= 2 && Valence <= 4)>
 	FXPBDSpringConstraints(
 		const FSolverParticles& Particles,
 		int32 ParticleOffset,
@@ -138,7 +113,7 @@ public:
 		InitColor(Particles);
 	}
 
-	template<int32 Valence, TEMPLATE_REQUIRES(Valence >= 2 && Valence <= 4)>
+	template<int32 Valence UE_REQUIRES(Valence >= 2 && Valence <= 4)>
 	UE_DEPRECATED(5.4, "XPBD Constraints must always trim kinematic constraints")
 	FXPBDSpringConstraints(
 		const FSolverParticles& Particles,
