Fix assert due to different netbitarray sizes  in Net.Iris.PrintNetCullDistances

#rnx

diff --git a/Engine/Source/Runtime/Experimental/Iris/Core/Private/Iris/ReplicationSystem/ObjectReplicationBridgeDebugging.cpp b/Engine/Source/Runtime/Experimental/Iris/Core/Private/Iris/ReplicationSystem/ObjectReplicationBridgeDebugging.cpp
--- a/Engine/Source/Runtime/Experimental/Iris/Core/Private/Iris/ReplicationSystem/ObjectReplicationBridgeDebugging.cpp
+++ b/Engine/Source/Runtime/Experimental/Iris/Core/Private/Iris/ReplicationSystem/ObjectReplicationBridgeDebugging.cpp
@@ -336,7 +336,7 @@ void UObjectReplicationBridge::PrintReplicatedObjects(uint32 ArgTraits) const
 	uint32 TotalSubObjects = 0;
 
 	FNetBitArray RootObjects;
-	RootObjects.Init(NetRefHandleManager->GetMaxActiveObjectCount());
+	RootObjects.Init(NetRefHandleManager->GetCurrentMaxInternalNetRefIndex());
 	FNetBitArrayView RootObjectsView = MakeNetBitArrayView(RootObjects);
 	RootObjectsView.Set(NetRefHandleManager->GetGlobalScopableInternalIndices(), FNetBitArrayView::AndNotOp, NetRefHandleManager->GetSubObjectInternalIndicesView());
 
@@ -390,7 +390,7 @@ void UObjectReplicationBridge::PrintRelevantObjects(uint32 ArgTraits) const
 	UE_LOG(LogIrisBridge, Display, TEXT(""));
 
 	FNetBitArray RootObjects;
-	RootObjects.Init(NetRefHandleManager->GetMaxActiveObjectCount());
+	RootObjects.Init(NetRefHandleManager->GetCurrentMaxInternalNetRefIndex());
 	FNetBitArrayView RootObjectsView = MakeNetBitArrayView(RootObjects);
 	RootObjectsView.Set(NetRefHandleManager->GetRelevantObjectsInternalIndices(), FNetBitArrayView::AndNotOp, NetRefHandleManager->GetSubObjectInternalIndicesView());
 
@@ -438,7 +438,7 @@ void UObjectReplicationBridge::PrintAlwaysRelevantObjects(uint32 ArgTraits) cons
 	UE_LOG(LogIrisBridge, Display, TEXT(""));
 
 	FNetBitArray AlwaysRelevantList;
-	AlwaysRelevantList.Init(NetRefHandleManager->GetMaxActiveObjectCount());
+	AlwaysRelevantList.Init(NetRefHandleManager->GetCurrentMaxInternalNetRefIndex());
 	
 	ReplicationSystemInternal->GetFiltering().BuildAlwaysRelevantList(MakeNetBitArrayView(AlwaysRelevantList), ReplicationSystemInternal->GetNetRefHandleManager().GetGlobalScopableInternalIndices());
 
@@ -538,7 +538,7 @@ void UObjectReplicationBridge::PrintRelevantObjectsForConnections(const TArray<F
 		UE_LOG(LogIrisBridge, Display, TEXT(""));
 
 		FNetBitArray RootObjects;
-		RootObjects.Init(NetRefHandleManager->GetMaxActiveObjectCount());
+		RootObjects.Init(NetRefHandleManager->GetCurrentMaxInternalNetRefIndex());
 		MakeNetBitArrayView(RootObjects).Set(GetReplicationSystem()->GetReplicationSystemInternal()->GetFiltering().GetRelevantObjectsInScope(ConnectionId), FNetBitArrayView::AndNotOp, NetRefHandleManager->GetSubObjectInternalIndicesView());
 
 		TArray<FRootObjectData> RelevantObjects;
@@ -641,7 +641,7 @@ void UObjectReplicationBridge::PrintNetCullDistances(const TArray<FString>& Args
 	// Filter down to objects in the GridFilter. Other filters do not use net culling
 	{
 		FNetBitArray GridFilterList;
-		GridFilterList.Init(NetRefHandleManager->GetMaxActiveObjectCount());
+		GridFilterList.Init(NetRefHandleManager->GetCurrentMaxInternalNetRefIndex());
 		ReplicationSystemInternal->GetFiltering().BuildObjectsInFilterList(MakeNetBitArrayView(GridFilterList), TEXT("Spatial"));
 		RootObjects.Combine(GridFilterList, FNetBitArray::AndOp);
 	}
