Iris
- Added ReadJournal to make it easier to track down bitstream errors

diff --git a/Engine/Source/Runtime/Experimental/Iris/Core/Private/Iris/ReplicationSystem/NetBlob/NetRPC.cpp b/Engine/Source/Runtime/Experimental/Iris/Core/Private/Iris/ReplicationSystem/NetBlob/NetRPC.cpp
--- a/Engine/Source/Runtime/Experimental/Iris/Core/Private/Iris/ReplicationSystem/NetBlob/NetRPC.cpp
+++ b/Engine/Source/Runtime/Experimental/Iris/Core/Private/Iris/ReplicationSystem/NetBlob/NetRPC.cpp
@@ -243,7 +243,11 @@ void FNetRPC::Deserialize(FNetSerializationContext& Context)
 	
 	if (!Context.HasErrorOrOverflow())
 	{
-		check(PostNetRPCPos == Context.GetBitStreamReader()->GetPosBits());
+		if (PostNetRPCPos != Context.GetBitStreamReader()->GetPosBits())
+		{
+			UE_LOG(LogIrisRpc, Error, TEXT("DeserializeWithObject::RPC %s did not read expected number of bits ErrorContext: %s"), ToCStr(BlobDescriptor->DebugName), *Context.PrintReadJournal())
+			ensure(PostNetRPCPos == Context.GetBitStreamReader()->GetPosBits());
+		}
 	}
 }
 
diff --git a/Engine/Source/Runtime/Experimental/Iris/Core/Private/Iris/ReplicationSystem/ReplicationOperations.cpp b/Engine/Source/Runtime/Experimental/Iris/Core/Private/Iris/ReplicationSystem/ReplicationOperations.cpp
--- a/Engine/Source/Runtime/Experimental/Iris/Core/Private/Iris/ReplicationSystem/ReplicationOperations.cpp
+++ b/Engine/Source/Runtime/Experimental/Iris/Core/Private/Iris/ReplicationSystem/ReplicationOperations.cpp
@@ -301,6 +301,7 @@ void FReplicationStateOperations::Deserialize(FNetSerializationContext& Context,
 	for (uint32 MemberIt = 0; MemberIt < MemberCount; ++MemberIt)
 	{
 		UE_NET_TRACE_DYNAMIC_NAME_SCOPE(MemberDebugDescriptors[MemberIt].DebugName, *Context.GetBitStreamReader(), Context.GetTraceCollector(), ENetTraceVerbosity::Trace);
+		UE_ADD_READ_JOURNAL_ENTRY(Context, MemberDebugDescriptors[MemberIt].DebugName);
 
 		const FReplicationStateMemberDescriptor& MemberDescriptor = MemberDescriptors[MemberIt];
 		const FReplicationStateMemberSerializerDescriptor& MemberSerializerDescriptor = MemberSerializerDescriptors[MemberIt];
@@ -351,6 +352,7 @@ void FReplicationStateOperations::DeserializeDelta(FNetSerializationContext& Con
 	for (uint32 MemberIt = 0; MemberIt < MemberCount; ++MemberIt)
 	{
 		UE_NET_TRACE_DYNAMIC_NAME_SCOPE(MemberDebugDescriptors[MemberIt].DebugName, *Context.GetBitStreamReader(), Context.GetTraceCollector(), ENetTraceVerbosity::Trace);
+		UE_ADD_READ_JOURNAL_ENTRY(Context, MemberDebugDescriptors[MemberIt].DebugName);
 
 		const FReplicationStateMemberDescriptor& MemberDescriptor = MemberDescriptors[MemberIt];
 		const FReplicationStateMemberSerializerDescriptor& MemberSerializerDescriptor = MemberSerializerDescriptors[MemberIt];
@@ -448,6 +450,8 @@ void FReplicationStateOperations::DeserializeWithMask(FNetSerializationContext&
 			UE_NET_TRACE_DYNAMIC_NAME_SCOPE(MemberSerializerDescriptor.Serializer->Name, *Context.GetBitStreamReader(), Context.GetTraceCollector(), ENetTraceVerbosity::VeryVerbose);
 #endif
 
+			UE_ADD_READ_JOURNAL_ENTRY(Context, MemberDebugDescriptors[MemberIt].DebugName);
+
 			FNetDeserializeArgs Args;
 			Args.NetSerializerConfig = MemberSerializerDescriptor.SerializerConfig;
 			Args.Target = reinterpret_cast<NetSerializerValuePointer>(DstInternalBuffer + MemberDescriptor.InternalMemberOffset);
@@ -541,6 +545,8 @@ void FReplicationStateOperations::DeserializeDeltaWithMask(FNetSerializationCont
 			UE_NET_TRACE_DYNAMIC_NAME_SCOPE(MemberDebugDescriptors[MemberIt].DebugName, *Context.GetBitStreamReader(), Context.GetTraceCollector(), ENetTraceVerbosity::Trace);
 #endif
 
+			UE_ADD_READ_JOURNAL_ENTRY(Context, MemberDebugDescriptors[MemberIt].DebugName);
+
 			FNetDeserializeDeltaArgs Args;
 			Args.NetSerializerConfig = MemberSerializerDescriptor.SerializerConfig;
 			Args.Target = reinterpret_cast<NetSerializerValuePointer>(DstInternalBuffer + MemberDescriptor.InternalMemberOffset);
@@ -856,6 +862,8 @@ void FReplicationProtocolOperations::Deserialize(FNetSerializationContext& Conte
 	{
 		const FReplicationStateDescriptor* CurrentDescriptor = ReplicationStateDescriptors[StateIt];
 
+		UE_ADD_READ_JOURNAL_ENTRY(Context, CurrentDescriptor->DebugName);
+
 		CurrentInternalStateBuffer = Align(CurrentInternalStateBuffer, CurrentDescriptor->InternalAlignment);
 		
 		FReplicationStateOperations::Deserialize(Context, CurrentInternalStateBuffer, CurrentDescriptor);
@@ -935,6 +943,7 @@ void FReplicationProtocolOperations::DeserializeWithMask(FNetSerializationContex
 	// Read the ChangeMask
 	{
 		UE_NET_TRACE_SCOPE(ChangeMasks, *Context.GetBitStreamReader(), Context.GetTraceCollector(), ENetTraceVerbosity::Trace);
+		UE_ADD_READ_JOURNAL_ENTRY(Context, TEXT("ChangeMask"));
 		ReadSparseBitArray(Context.GetBitStreamReader(), DstChangeMaskData, Protocol->ChangeMaskBitCount);
 	}
 
@@ -945,6 +954,8 @@ void FReplicationProtocolOperations::DeserializeWithMask(FNetSerializationContex
 		const FReplicationStateDescriptor* CurrentDescriptor = ReplicationStateDescriptors[StateIt];
 
 		CurrentInternalStateBuffer = Align(CurrentInternalStateBuffer, CurrentDescriptor->InternalAlignment);
+
+		UE_ADD_READ_JOURNAL_ENTRY(Context, CurrentDescriptor->DebugName);
 		
 		if (CurrentDescriptor->IsInitState())
 		{
@@ -1174,6 +1185,7 @@ void FReplicationProtocolOperations::DeserializeWithMaskDelta(FNetSerializationC
 	// Read the ChangeMask
 	{
 		UE_NET_TRACE_SCOPE(ChangeMasks, *Context.GetBitStreamReader(), Context.GetTraceCollector(), ENetTraceVerbosity::Trace);
+		UE_ADD_READ_JOURNAL_ENTRY(Context, TEXT("ChangeMask"));
 		ReadSparseBitArray(Context.GetBitStreamReader(), DstChangeMaskData, Protocol->ChangeMaskBitCount);
 	}
 
@@ -1187,6 +1199,7 @@ void FReplicationProtocolOperations::DeserializeWithMaskDelta(FNetSerializationC
 		PrevInternalStateBuffer = Align(PrevInternalStateBuffer, CurrentDescriptor->InternalAlignment);
 
 		UE_NET_TRACE_DYNAMIC_NAME_SCOPE(CurrentDescriptor->DebugName, *Context.GetBitStreamReader(), Context.GetTraceCollector(), ENetTraceVerbosity::Trace);
+		UE_ADD_READ_JOURNAL_ENTRY(Context, CurrentDescriptor->DebugName);
 		if (CurrentDescriptor->IsInitState())
 		{
 			if (bIsInitState)
diff --git a/Engine/Source/Runtime/Experimental/Iris/Core/Private/Iris/ReplicationSystem/ReplicationReader.cpp b/Engine/Source/Runtime/Experimental/Iris/Core/Private/Iris/ReplicationSystem/ReplicationReader.cpp
--- a/Engine/Source/Runtime/Experimental/Iris/Core/Private/Iris/ReplicationSystem/ReplicationReader.cpp
+++ b/Engine/Source/Runtime/Experimental/Iris/Core/Private/Iris/ReplicationSystem/ReplicationReader.cpp
@@ -667,6 +667,8 @@ uint32 FReplicationReader::ReadObjectBatch(FNetSerializationContext& Context)
 
 	UE_NET_TRACE_SCOPE(Batch, Reader, Context.GetTraceCollector(), ENetTraceVerbosity::Trace);
 
+	UE_ADD_READ_JOURNAL_ENTRY(Context, TEXT("ReadObjectBatch"));
+
 	// Special handling for destruction infos
 	if (const bool bIsDestructionInfo = Reader.ReadBool())
 	{
@@ -679,7 +681,14 @@ uint32 FReplicationReader::ReadObjectBatch(FNetSerializationContext& Context)
 #if UE_NET_USE_READER_WRITER_SENTINEL
 		ReadAndVerifySentinelBits(&Reader, TEXT("DestructionInfo"), 8);
 #endif
-		
+
+		if (Context.HasErrorOrOverflow())
+		{
+			UE_LOG(LogIris, Error, TEXT("FReplicationReader::ReadObject Failed to read destruction info. \n%s"), *Context.PrintReadJournal());
+			Context.SetError(GNetError_BitStreamError);
+			return 0U;
+		}
+
 		return 1U;
 	}
 
@@ -702,6 +711,9 @@ uint32 FReplicationReader::ReadObjectBatch(FNetSerializationContext& Context)
 		BatchSize = Reader.ReadBits(NumBitsUsedForBatchSize);
 	}
 
+	// Store the current handle if we encounter errors
+	Context.SetErrorHandleContext(IncompleteHandle);
+
 	if (Context.HasErrorOrOverflow() || BatchSize > Reader.GetBitsLeft())
 	{
 		Context.SetError(GNetError_InvalidValue);
@@ -734,6 +746,7 @@ uint32 FReplicationReader::ReadObjectBatch(FNetSerializationContext& Context)
 		ObjectReferenceCache->ReadExports(Context, &TempMustBeMappedReferences);
 		if (Context.HasErrorOrOverflow())
 		{
+			UE_LOG(LogIris, Error, TEXT("FReplicationReader::ReadObject Failed to read exports for handle: %s.\n%s"), *IncompleteHandle.ToString(), *Context.PrintReadJournal());
 			return 0U;
 		}
 
@@ -761,7 +774,7 @@ uint32 FReplicationReader::ReadObjectBatch(FNetSerializationContext& Context)
 	FPendingBatchData* PendingBatchData = ObjectReferenceCache->ShouldAsyncLoad() ? UpdateUnresolvedMustBeMappedReferences(IncompleteHandle, TempMustBeMappedReferences) : nullptr;
 	if (PendingBatchData)
 	{
-		UE_LOG(LogIris, Verbose, TEXT("FReplicationReader::ReadObjectBatch Handle %s will be defered as it has unresolved must be mapped references"), *IncompleteHandle.ToString());				
+		UE_LOG(LogIris, Verbose, TEXT("FReplicationReader::ReadObjectBatch Handle %s will be defered as it has unresolved must be mapped references"), *IncompleteHandle.ToString());
 
 		UE_NET_TRACE_OBJECT_SCOPE(IncompleteHandle, Reader, Context.GetTraceCollector(), ENetTraceVerbosity::Trace);
 		UE_NET_TRACE_SCOPE(QueuedBatch, Reader, Context.GetTraceCollector(), ENetTraceVerbosity::Trace);
@@ -785,6 +798,9 @@ uint32 FReplicationReader::ReadObjectBatch(FNetSerializationContext& Context)
 
 		if (Context.HasErrorOrOverflow())
 		{
+			// Log error, this is something we cannot recover from.
+			UE_LOG(LogIris, Error, TEXT("FReplicationReader::ReadObject Failed to read object batch data chunk for handle:%s \n%s"), ToCStr(IncompleteHandle.ToString()), *Context.PrintReadJournal());
+
 			return 0U;
 		}
 
@@ -818,6 +834,13 @@ uint32 FReplicationReader::ReadObjectBatch(FNetSerializationContext& Context)
 		}
 	}
 
+	if (!ensure(Reader.GetPosBits() == BatchEndOrStartOfExportsPos))
+	{
+		UE_LOG(LogIris, Error, TEXT("FReplicationReader::ReadObjectsInBatch Did not read the expected number of bits when reading batch: %s. \n%s"), *IncompleteHandle.ToString(), *Context.PrintReadJournal());
+		Context.SetError(GNetError_BitStreamOverflow, true);
+		return 0U;
+	}
+
 	// Skip to the end as we already have read any exports
 	Reader.Seek(BatchEndPos);
 
@@ -850,6 +873,12 @@ void FReplicationReader::ReadObjectInBatch(FNetSerializationContext& Context, FN
 
 	UE_NET_TRACE_OBJECT_SCOPE(IncompleteHandle, Reader, Context.GetTraceCollector(), ENetTraceVerbosity::Trace);
 
+	// Store the current handle in case we encounter errors
+	if (IncompleteHandle != BatchHandle)
+	{
+		Context.SetErrorHandleContext(IncompleteHandle);
+	}
+
 	//UE_LOG_REPLICATIONREADER(TEXT("FReplicationReader::Read Object with %s InitialState: %u"), *IncompleteHandle.ToString(), bIsInitialState ? 1u : 0u);
 
 	bool bHasErrors = false;
@@ -861,6 +890,8 @@ void FReplicationReader::ReadObjectInBatch(FNetSerializationContext& Context, FN
 	{
 		UE_NET_TRACE_SCOPE(CreationInfo, *Context.GetBitStreamReader(), Context.GetTraceCollector(), ENetTraceVerbosity::Trace);
 
+		UE_ADD_READ_JOURNAL_ENTRY(Context, TEXT("ReadCreationInfo"));
+
 		// SubObject data for initial state
 		FNetRefHandle RootObjectOfSubObject;
 		if (bIsSubObject)
@@ -874,7 +905,9 @@ void FReplicationReader::ReadObjectInBatch(FNetSerializationContext& Context, FN
 				UE_LOG(LogIris, Error, TEXT("FReplicationReader::ReadObject Invalid subobjectowner handle. %s"), ToCStr(IncompleteOwnerHandle.ToString()));
 				const FName& NetError = (Reader.IsOverflown() ? GNetError_BitStreamOverflow : GNetError_InvalidNetHandle);
 				Context.SetError(NetError);
-				return;			
+
+				bHasErrors = true;
+				goto ErrorHandling;
 			}
 
 			RootObjectOfSubObject = NetRefHandleManager->GetReplicatedObjectDataNoCheck(RootObjectInternalIndex).RefHandle;
@@ -893,7 +926,9 @@ void FReplicationReader::ReadObjectInBatch(FNetSerializationContext& Context, FN
 			UE_LOG(LogIris, Error, TEXT("FReplicationReader::ReadObject Bitstream corrupted."));
 			const FName& NetError = (Reader.IsOverflown() ? GNetError_BitStreamOverflow : GNetError_BitStreamError);
 			Context.SetError(NetError);
-			return;
+
+			bHasErrors = true;
+			goto ErrorHandling;
 		}
 	
 		// Get Bridge
@@ -933,7 +968,7 @@ void FReplicationReader::ReadObjectInBatch(FNetSerializationContext& Context, FN
 	else
 	{
 		bHasErrors = Context.HasErrorOrOverflow();
-		UE_CLOG(bHasErrors, LogIris, Error, TEXT("FReplicationReader::ReadObject ErrorOrOverFlow after reading object header"))
+		UE_CLOG(bHasErrors, LogIris, Error, TEXT("FReplicationReader::ReadObject ErrorOrOverFlow after reading object header %s"), *Context.PrintReadJournal())
 
 		if (bHasErrors || !IncompleteHandle.IsValid())
 		{
@@ -970,6 +1005,8 @@ void FReplicationReader::ReadObjectInBatch(FNetSerializationContext& Context, FN
 	{
 		const FNetRefHandleManager::FReplicatedObjectData& ObjectData = NetRefHandleManager->GetReplicatedObjectDataNoCheck(InternalIndex);
 
+		UE_ADD_READ_JOURNAL_ENTRY(Context, ObjectData.Protocol ? ObjectData.Protocol->DebugName->Name : TEXT("OOB"));
+
 		// Add entry in our received data as we postpone state application until we have received all data in order to be able to properly resolve references
 		FDispatchObjectInfo& Info = ObjectsToDispatchArray->AddPendingDispatchObjectInfo(TempLinearAllocator);
 
@@ -981,6 +1018,8 @@ void FReplicationReader::ReadObjectInBatch(FNetSerializationContext& Context, FN
 
 		if (bHasState)
 		{
+			UE_ADD_READ_JOURNAL_ENTRY(Context, TEXT("State"));
+
 			if (IsObjectIndexForOOBAttachment(InternalIndex) || bIsReplicatedDestroyForInvalidObject)
 			{
 				bHasErrors = true;
@@ -1054,6 +1093,8 @@ void FReplicationReader::ReadObjectInBatch(FNetSerializationContext& Context, FN
 		ENetObjectAttachmentType AttachmentType = ENetObjectAttachmentType::Normal;
 		if (bHasAttachments)
 		{
+			UE_ADD_READ_JOURNAL_ENTRY(Context, TEXT("Attachments"));
+
 			if (IsObjectIndexForOOBAttachment(InternalIndex))
 			{
 				bHasErrors = bIsReplicatedDestroyForInvalidObject;
@@ -1101,7 +1142,7 @@ ErrorHandling:
 	if (bHasErrors)
 	{
 		Context.SetErrorHandleContext(IncompleteHandle);
-		UE_LOG(LogIris, Error, TEXT("FReplicationReader::ReadObject Failed to read replicated object with %s. Error '%s'."), *IncompleteHandle.ToString(), (Context.HasError() ? ToCStr(Context.GetError().ToString()) : TEXT("BitStream Overflow")));
+		UE_LOG(LogIris, Error, TEXT("FReplicationReader::ReadObject Failed to read replicated object with Handle: %s. Error '%s'. %s"), *IncompleteHandle.ToString(), (Context.HasError() ? ToCStr(Context.GetError().ToString()) : TEXT("BitStream Overflow")), *Context.PrintReadJournal());
 	}
 }
 
@@ -2471,11 +2512,6 @@ void FReplicationReader::Read(FNetSerializationContext& Context)
 		return;
 	}
 
-	if (Context.HasErrorOrOverflow())
-	{
-		return;
-	}
-
 	// Stats
 	UE_NET_TRACE_PACKET_STATSCOUNTER(Parameters.ReplicationSystem->GetId(), Parameters.ConnectionId, ReplicationReader.ReadObjectBatchCount, ReceivedObjectBatchCountToRead, ENetTraceVerbosity::Trace);
 	UE_NET_TRACE_PACKET_STATSCOUNTER(Parameters.ReplicationSystem->GetId(), Parameters.ConnectionId, ReplicationReader.ReadObjectsAndSubObjectsToDispatchCount, ObjectsToDispatchArray->Num(), ENetTraceVerbosity::Trace);
diff --git a/Engine/Source/Runtime/Experimental/Iris/Core/Private/Iris/Serialization/NetJournal.cpp b/Engine/Source/Runtime/Experimental/Iris/Core/Private/Iris/Serialization/NetJournal.cpp
new file mode 100644
--- /dev/null
+++ b/Engine/Source/Runtime/Experimental/Iris/Core/Private/Iris/Serialization/NetJournal.cpp
@@ -0,0 +1,36 @@
+// Copyright Epic Games, Inc. All Rights Reserved.
+
+#include "Iris/Serialization/NetJournal.h"
+#include "Iris/ReplicationSystem/ReplicationSystem.h"
+#include "Iris/ReplicationSystem/ReplicationSystemInternal.h"
+#include "Iris/ReplicationSystem/NetRefHandleManager.h"
+
+namespace UE::Net
+{
+
+FString FNetJournal::Print(const UReplicationSystem* ReplicationSystem) const
+{
+	using namespace UE::Net::Private;
+
+	const FNetRefHandleManager* NetRefHandleManager = ReplicationSystem ? &ReplicationSystem->GetReplicationSystemInternal()->GetNetRefHandleManager() : nullptr;
+
+	FString Result;
+	const uint32 StoredCount = FMath::Min(NumEntries, JournalSize);
+	const uint32 StartIndex = (NumEntries - StoredCount);
+	Result.Appendf(TEXT("ErrorContext:\n"));
+
+	FNetRefHandle LastNetRefHandle;
+	for (uint32 EntryIt = 0U; EntryIt < StoredCount; ++EntryIt)
+	{
+		const FJournalEntry& Entry = Entries[(StartIndex + EntryIt) & JournalMask];
+		if (LastNetRefHandle != Entry.NetRefHandle)
+		{
+			Result.Appendf(TEXT("%s\n"), NetRefHandleManager ? *NetRefHandleManager->PrintObjectFromNetRefHandle(Entry.NetRefHandle) : *Entry.NetRefHandle.ToString());
+			LastNetRefHandle = Entry.NetRefHandle;
+		}
+		Result.Appendf(TEXT("%u: - BitOffset: %u:%s\n"), EntryIt, Entry.BitOffset, Entry.Name);		
+	}
+	return Result;
+}
+
+}
diff --git a/Engine/Source/Runtime/Experimental/Iris/Core/Private/Iris/Serialization/NetSerializationContext.cpp b/Engine/Source/Runtime/Experimental/Iris/Core/Private/Iris/Serialization/NetSerializationContext.cpp
--- a/Engine/Source/Runtime/Experimental/Iris/Core/Private/Iris/Serialization/NetSerializationContext.cpp
+++ b/Engine/Source/Runtime/Experimental/Iris/Core/Private/Iris/Serialization/NetSerializationContext.cpp
@@ -24,6 +24,21 @@ bool FNetSerializationContext::IsBitStreamOverflown() const
 	return false;
 }
 
+void FNetSerializationContext::AddReadJournalEntry(const TCHAR* Name)
+{ 
+	ReadJournal.AddEntry(Name, BitStreamReader->GetPosBits(), ErrorContext.GetObjectHandle());
+}
+
+void FNetSerializationContext::AddReadJournalEntry(const FNetDebugName* DebugName)
+{
+	ReadJournal.AddEntry(DebugName->Name, BitStreamReader->GetPosBits(), ErrorContext.GetObjectHandle());
+}
+
+FString FNetSerializationContext::PrintReadJournal()
+{
+	return ReadJournal.Print(InternalContext ? InternalContext->ReplicationSystem : nullptr);
+}
+
 void FNetSerializationContext::SetBitStreamOverflow()
 {
 	if (BitStreamReader != nullptr && !BitStreamReader->IsOverflown())
diff --git a/Engine/Source/Runtime/Experimental/Iris/Core/Public/Iris/Serialization/NetJournal.h b/Engine/Source/Runtime/Experimental/Iris/Core/Public/Iris/Serialization/NetJournal.h
new file mode 100644
--- /dev/null
+++ b/Engine/Source/Runtime/Experimental/Iris/Core/Public/Iris/Serialization/NetJournal.h
@@ -0,0 +1,55 @@
+// Copyright Epic Games, Inc. All Rights Reserved.
+
+#pragma once
+
+#include "Containers/StringFwd.h"
+#include "Iris/ReplicationSystem/NetRefHandle.h"
+#include "Net/Core/Trace/NetDebugName.h"
+
+class UReplicationSystem;
+
+#ifndef UE_NET_ENABLE_READ_JOURNAL
+#if (UE_BUILD_SHIPPING)
+#	define UE_NET_ENABLE_READ_JOURNAL 1
+#else
+#	define UE_NET_ENABLE_READ_JOURNAL 1
+#endif 
+#endif
+
+#if UE_NET_ENABLE_READ_JOURNAL
+	#define UE_ADD_READ_JOURNAL_ENTRY(SerializationContext, X) SerializationContext.AddReadJournalEntry(X);
+#else
+	#define UE_ADD_READ_JOURNAL_ENTRY(...)
+#endif
+
+namespace UE::Net
+{
+// Simple journal to track last few entries of read data
+class FNetJournal
+{
+	static constexpr uint32 JournalSize = 8U;
+	static constexpr uint32 JournalMask = JournalSize - 1U;
+
+public:
+	void Reset() { NumEntries = 0U; }
+	void AddEntry(const TCHAR* Name, uint32 BitOffset, FNetRefHandle NetRefHandle);
+	FString Print(const UReplicationSystem* ReplicationSystem) const;
+
+private:
+	struct FJournalEntry
+	{
+		const TCHAR* Name;
+		FNetRefHandle NetRefHandle;
+		uint32 BitOffset;
+	};
+	FJournalEntry Entries[JournalSize];
+	uint32 NumEntries = 0U;
+};
+
+inline void FNetJournal::AddEntry(const TCHAR* Name, uint32 BitOffset, FNetRefHandle NetRefHandle)
+{ 
+	Entries[NumEntries & JournalMask] = FJournalEntry({Name, NetRefHandle, BitOffset});
+	++NumEntries;
+}
+
+}
diff --git a/Engine/Source/Runtime/Experimental/Iris/Core/Public/Iris/Serialization/NetSerializationContext.h b/Engine/Source/Runtime/Experimental/Iris/Core/Public/Iris/Serialization/NetSerializationContext.h
--- a/Engine/Source/Runtime/Experimental/Iris/Core/Public/Iris/Serialization/NetSerializationContext.h
+++ b/Engine/Source/Runtime/Experimental/Iris/Core/Public/Iris/Serialization/NetSerializationContext.h
@@ -3,6 +3,7 @@
 #pragma once
 
 #include "Iris/Serialization/NetErrorContext.h"
+#include "Iris/Serialization/NetJournal.h"
 
 class FName;
 class INetBlobReceiver;
@@ -55,6 +56,17 @@ public:
 	/** There are cases where an error is handled and reported where we want to stay calm, reset the error context and carry on */
 	void ResetErrorContext() { ErrorContext = FNetErrorContext(); }
 
+	/** Add entry into read journal, Name must be a static string as the pointer will be stored */
+	void AddReadJournalEntry(const TCHAR* Name);
+
+	/** Add entry in to error context, Name must be a static string as the pointer will be stored */
+	void AddReadJournalEntry(const FNetDebugName* DebugName);
+
+	/** Print the ReadJournal */
+	FString PrintReadJournal();
+
+	void ResetReadJournal() { ReadJournal.Reset(); }
+
 	void SetIsInitState(bool bInIsInitState) { bIsInitState = bInIsInitState; }
 	bool IsInitState() const { return bIsInitState; }
 
@@ -96,6 +108,7 @@ private:
 	IRISCORE_API void SetBitStreamOverflow();
 
 	FNetErrorContext ErrorContext;
+	FNetJournal ReadJournal;
 
 	FNetBitStreamReader* BitStreamReader = nullptr;
 	FNetBitStreamWriter* BitStreamWriter = nullptr;
