NetBitArray
* Added safer ways to access internal buffers held by a FNetBitArray and FNetBitArrayView
* GetWord() can be used to safely access the full 32bit value to modify 32bit entries at once.
* GetDataChecked(NumWordSize) can be used to access the raw array buffer but will also validate that the maximum word iteration count is within the bounds of the NetBitArray
* Reset() has been deprecated in favor of ClearAllBits()

Iris
* Replaced usage of deprecated NetBitArray::Reset
* Replaced a few usages of unsafe GetData() with safer GetWord and GetDataChecked
* Replaced Reset with ClearAllBits

diff --git a/Engine/Plugins/Runtime/ReplicationSystemTestPlugin/Source/Private/Tests/Core/TestNetBitArray.cpp b/Engine/Plugins/Runtime/ReplicationSystemTestPlugin/Source/Private/Tests/Core/TestNetBitArray.cpp
--- a/Engine/Plugins/Runtime/ReplicationSystemTestPlugin/Source/Private/Tests/Core/TestNetBitArray.cpp
+++ b/Engine/Plugins/Runtime/ReplicationSystemTestPlugin/Source/Private/Tests/Core/TestNetBitArray.cpp
@@ -28,13 +28,13 @@ UE_NET_TEST(FNetBitArrayView, Construct)
 	}
 }
 
-UE_NET_TEST(FNetBitArrayView, Reset)
+UE_NET_TEST(FNetBitArrayView, ClearAllBits)
 {	
 	{
 		const uint32 ExpectedSingleWordBuffer = 0u;
 		uint32 SingleWordBuffer = 0xfe;
 		FNetBitArrayView BitArray(&SingleWordBuffer, 8);
-		BitArray.Reset();
+		BitArray.ClearAllBits();
 
 		UE_NET_ASSERT_EQ(ExpectedSingleWordBuffer, SingleWordBuffer);
 	}
@@ -43,7 +43,7 @@ UE_NET_TEST(FNetBitArrayView, Reset)
 		const uint32 ExpectedSingleWordBuffer = 0u;
 		uint32 SingleWordBuffer = 0xfefefefe;
 		FNetBitArrayView BitArray(&SingleWordBuffer, 32);
-		BitArray.Reset();
+		BitArray.ClearAllBits();
 
 		UE_NET_ASSERT_EQ(ExpectedSingleWordBuffer, SingleWordBuffer);
 	}
@@ -53,7 +53,7 @@ UE_NET_TEST(FNetBitArrayView, Reset)
 		uint32 WordBuffer[] = { 1u, 2u, 3u, 4 };
 
 		FNetBitArrayView BitArray(&WordBuffer[0], 8);
-		BitArray.Reset();
+		BitArray.ClearAllBits();
 
 		for (uint32 it = 0; it < UE_ARRAY_COUNT(ExpectedWordBuffer); ++it)
 		{
@@ -66,7 +66,7 @@ UE_NET_TEST(FNetBitArrayView, Reset)
 		uint32 WordBuffer[] = { 1, 2, 3, 4 };
 
 		FNetBitArrayView BitArray(&WordBuffer[0], 128);
-		BitArray.Reset();
+		BitArray.ClearAllBits();
 
 		for (uint32 it = 0; it < UE_ARRAY_COUNT(ExpectedWordBuffer); ++it)
 		{
@@ -75,12 +75,12 @@ UE_NET_TEST(FNetBitArrayView, Reset)
 	}
 }
 
-UE_NET_TEST(FNetBitArrayView, IsAnyBitSetIsFalseAfterReset)
+UE_NET_TEST(FNetBitArrayView, IsAnyBitSetIsFalseAfterClearAllBits)
 {	
 	{
 		uint32 SingleWordBuffer = 0xfe;
 		FNetBitArrayView BitArray(&SingleWordBuffer, 8);
-		BitArray.Reset();
+		BitArray.ClearAllBits();
 
 		UE_NET_ASSERT_FALSE(BitArray.IsAnyBitSet());
 		UE_NET_ASSERT_FALSE(BitArray.IsAnyBitSet(0, ~0U));
@@ -90,7 +90,7 @@ UE_NET_TEST(FNetBitArrayView, IsAnyBitSetIsFalseAfterReset)
 		uint32 WordBuffer[] = { 1, 2, 3, 4 };
 
 		FNetBitArrayView BitArray(&WordBuffer[0], 128);
-		BitArray.Reset();
+		BitArray.ClearAllBits();
 
 		UE_NET_ASSERT_FALSE(BitArray.IsAnyBitSet());
 		UE_NET_ASSERT_FALSE(BitArray.IsAnyBitSet(0, ~0U));
@@ -433,7 +433,7 @@ UE_NET_TEST(FNetBitArrayView, GetSetBit)
 
 		const FNetBitArrayView SrcBitArray(&ExpectedWordBuffer, 32);
 		FNetBitArrayView DstBitArray(&DstWordBuffer, 32);
-		DstBitArray.Reset();
+		DstBitArray.ClearAllBits();
 
 		UE_NET_ASSERT_FALSE(DstBitArray.IsAnyBitSet());
 
@@ -454,7 +454,7 @@ UE_NET_TEST(FNetBitArrayView, GetSetBit)
 
 		const FNetBitArrayView SrcBitArray(ExpectedWordBuffer, 128);
 		FNetBitArrayView DstBitArray(DstWordBuffer, 128);
-		DstBitArray.Reset();
+		DstBitArray.ClearAllBits();
 
 		UE_NET_ASSERT_FALSE(DstBitArray.IsAnyBitSet());
 
@@ -1983,6 +1983,56 @@ UE_NET_TEST_FIXTURE(FNetBitArrayFixture, MakeNetBitArrayView)
 	}
 }
 
+UE_NET_TEST_FIXTURE(FNetBitArrayFixture, Test_NetBitArray_GetWord)
+{
+	constexpr uint32 WordCount = 16;
+	FNetBitArray BitArray;
+	BitArray.Init(WordCount * FNetBitArray::WordBitCount);
+
+	for (uint32 WordIt=0; WordIt < WordCount; ++WordIt)
+	{
+		UE_NET_ASSERT_TRUE(BitArray.GetWord(WordIt) == 0x00);
+		BitArray.GetWord(WordIt) = 0xAA;
+		UE_NET_ASSERT_TRUE(BitArray.GetWord(WordIt) == 0xAA);
+	}
+
+	FNetBitArrayView BitArrayView = MakeNetBitArrayView(BitArray);
+	for (uint32 WordIt = 0; WordIt < WordCount; ++WordIt)
+	{
+		UE_NET_ASSERT_TRUE(BitArrayView.GetWord(WordIt) == 0xAA);
+		BitArrayView.GetWord(WordIt) = 0xBB;
+		UE_NET_ASSERT_TRUE(BitArrayView.GetWord(WordIt) == 0xBB);
+	}
+}
+
+UE_NET_TEST_FIXTURE(FNetBitArrayFixture, Test_NetBitArray_GetDataChecked)
+{
+	constexpr uint32 WordCount = 16;
+	FNetBitArray BitArray;
+	BitArray.Init(WordCount * FNetBitArray::WordBitCount);
+
+	{
+		FNetBitArray::StorageWordType* RawData = BitArray.GetDataChecked(WordCount);
+		for (uint32 WordIt = 0; WordIt < WordCount; ++WordIt)
+		{
+			UE_NET_ASSERT_TRUE(RawData[WordIt] == 0x00);
+			RawData[WordIt] = 0xAA;
+			UE_NET_ASSERT_TRUE(BitArray.GetWord(WordIt) == 0xAA);
+		}
+	}
+
+	FNetBitArrayView BitArrayView = MakeNetBitArrayView(BitArray);
+	{
+		FNetBitArray::StorageWordType* RawData = BitArrayView.GetDataChecked(WordCount);
+		for (uint32 WordIt = 0; WordIt < WordCount; ++WordIt)
+		{
+			UE_NET_ASSERT_TRUE(RawData[WordIt] == 0xAA);
+			RawData[WordIt] = 0xBB;
+			UE_NET_ASSERT_TRUE(BitArrayView.GetWord(WordIt) == 0xBB);
+		}
+	}
+}
+
 UE_NET_TEST(FNetBitArrayConstRangeIterator, CanIterateOverEmptyBitArray)
 {
 	bool bIteratedOverEmptyArray = false;
@@ -2034,4 +2084,4 @@ UE_NET_TEST(FNetBitArrayConstRangeIterator, CanIterateOverArbitrailySizedBitArra
 	UE_NET_ASSERT_EQ(It, (uint32)UE_ARRAY_COUNT(IndicesToSet));
 }
 
-}
+} // end namespace UE::Net::Private
diff --git a/Engine/Plugins/Runtime/ReplicationSystemTestPlugin/Source/Private/Tests/ReplicationSystem/Filtering/MockNetObjectFilter.cpp b/Engine/Plugins/Runtime/ReplicationSystemTestPlugin/Source/Private/Tests/ReplicationSystem/Filtering/MockNetObjectFilter.cpp
--- a/Engine/Plugins/Runtime/ReplicationSystemTestPlugin/Source/Private/Tests/ReplicationSystem/Filtering/MockNetObjectFilter.cpp
+++ b/Engine/Plugins/Runtime/ReplicationSystemTestPlugin/Source/Private/Tests/ReplicationSystem/Filtering/MockNetObjectFilter.cpp
@@ -132,7 +132,7 @@ void UMockNetObjectFilter::Filter(FNetObjectFilteringParams& Params)
 
 	if (CallSetup.Filter.bFilterOutByDefault)
 	{
-		Params.OutAllowedObjects.Reset();
+		Params.OutAllowedObjects.ClearAllBits();
 	}
 	else
 	{
diff --git a/Engine/Plugins/Runtime/ReplicationSystemTestPlugin/Source/Private/Tests/ReplicationSystem/TestObjectPollFrequencyLimiter.cpp b/Engine/Plugins/Runtime/ReplicationSystemTestPlugin/Source/Private/Tests/ReplicationSystem/TestObjectPollFrequencyLimiter.cpp
--- a/Engine/Plugins/Runtime/ReplicationSystemTestPlugin/Source/Private/Tests/ReplicationSystem/TestObjectPollFrequencyLimiter.cpp
+++ b/Engine/Plugins/Runtime/ReplicationSystemTestPlugin/Source/Private/Tests/ReplicationSystem/TestObjectPollFrequencyLimiter.cpp
@@ -40,7 +40,7 @@ UE_NET_TEST(ObjectPollFrequencyLimiter, ObjectsToPollAreSpreadOutWhenCreatingObj
 	// All objects are created. Verify we get one object polled per frame
 	for (uint32 It = 0, EndIt = MaxObjectCount; It < EndIt; ++It)
 	{
-		OutputBitArrayView.Reset();
+		OutputBitArrayView.ClearAllBits();
 		FrequencyLimiter.Update(MakeNetBitArrayView(ScopeBitArray), MakeNetBitArrayView(EmptyBitArray), OutputBitArrayView);
 
 		const uint32 PollCount = OutputArray.CountSetBits();
@@ -99,7 +99,7 @@ UE_NET_TEST(ObjectPollFrequencyLimiter, ObjectsToPollAreSpreadOutRegardlessOfUpd
 	// Verify the objects get polled evenly.
 	for (uint32 PollIt = 0; PollIt <= PollFrequency; ++PollIt)
 	{
-		OutputArray.Reset();
+		OutputArray.ClearAllBits();
 		FNetBitArrayView OutputBitArrayView = MakeNetBitArrayView(OutputArray);
 		FrequencyLimiter.Update(MakeNetBitArrayView(ScopeBitArray), MakeNetBitArrayView(EmptyBitArray), OutputBitArrayView);
 
diff --git a/Engine/Plugins/Runtime/ReplicationSystemTestPlugin/Source/Private/Tests/ReplicationSystem/TestObjectReferences.cpp b/Engine/Plugins/Runtime/ReplicationSystemTestPlugin/Source/Private/Tests/ReplicationSystem/TestObjectReferences.cpp
--- a/Engine/Plugins/Runtime/ReplicationSystemTestPlugin/Source/Private/Tests/ReplicationSystem/TestObjectReferences.cpp
+++ b/Engine/Plugins/Runtime/ReplicationSystemTestPlugin/Source/Private/Tests/ReplicationSystem/TestObjectReferences.cpp
@@ -83,7 +83,7 @@ public:
 
 		void Reset()
 		{
-			BitArray.Reset();
+			BitArray.ClearAllBits();
 			References.Reset();
 		}
 		const FNetBitArrayView* GetChangeMask() const { return &ChangeMask; }
@@ -719,7 +719,7 @@ UE_NET_TEST_FIXTURE(FTestObjectReferencesFixture, TestVisitAllDirty)
 	// Partial
 	{
 		FCollectValidReferenceCollector Collector(false, ChangeMask);
-		ChangeMask.Reset();
+		ChangeMask.ClearAllBits();
 
 		// Mark member 11 as dirty. This corresponds to TestStructWithNestedRefTArray_TArray.
 		SetMemberDirty(ChangeMask, 11);
diff --git a/Engine/Plugins/Runtime/ReplicationSystemTestPlugin/Source/Private/Tests/Serialization/TestArrayPropertyNetSerializer.cpp b/Engine/Plugins/Runtime/ReplicationSystemTestPlugin/Source/Private/Tests/Serialization/TestArrayPropertyNetSerializer.cpp
--- a/Engine/Plugins/Runtime/ReplicationSystemTestPlugin/Source/Private/Tests/Serialization/TestArrayPropertyNetSerializer.cpp
+++ b/Engine/Plugins/Runtime/ReplicationSystemTestPlugin/Source/Private/Tests/Serialization/TestArrayPropertyNetSerializer.cpp
@@ -228,7 +228,7 @@ UE_NET_TEST_FIXTURE(FTestElementChangeMaskForArrayPropertyNetSerializer, TestCha
 	TObjectPtr<UClassWithDynamicArrayOfPrimitiveTypeForArrayPropertyNetSerializerTest> Source = NewObject<UClassWithDynamicArrayOfPrimitiveTypeForArrayPropertyNetSerializerTest>();
 
 	FNetBitArrayView ArrayChangeMask = UE::Net::Private::GetMemberChangeMask(ReplicationState.GetStateBuffer(), ReplicationStateDescriptor);
-	ArrayChangeMask.Reset();
+	ArrayChangeMask.ClearAllBits();
 
 	// Add one element at a time until the last bit in the changemask is set and make sure the array is always marked as dirty as well as the corresponding element
 	for (unsigned ElementIt = 0, ElementEndIt = FPropertyReplicationState::TArrayElementChangeMaskBits; ElementIt != ElementEndIt; ++ElementIt)
@@ -240,7 +240,7 @@ UE_NET_TEST_FIXTURE(FTestElementChangeMaskForArrayPropertyNetSerializer, TestCha
 		UE_NET_ASSERT_TRUE(ArrayChangeMask.GetBit(FPropertyReplicationState::TArrayElementChangeMaskBitOffset + ElementIt));
 		UE_NET_ASSERT_EQ(ArrayChangeMask.CountSetBits(), 2U);
 
-		ArrayChangeMask.Reset();
+		ArrayChangeMask.ClearAllBits();
 	}
 }
 
@@ -250,7 +250,7 @@ UE_NET_TEST_FIXTURE(FTestElementChangeMaskForArrayPropertyNetSerializer, TestCha
 	TObjectPtr<UClassWithDynamicArrayOfPrimitiveTypeForArrayPropertyNetSerializerTest> Source = NewObject<UClassWithDynamicArrayOfPrimitiveTypeForArrayPropertyNetSerializerTest>();
 
 	FNetBitArrayView ArrayChangeMask = UE::Net::Private::GetMemberChangeMask(ReplicationState.GetStateBuffer(), ReplicationStateDescriptor);
-	ArrayChangeMask.Reset();
+	ArrayChangeMask.ClearAllBits();
 
 	// First add elements up to the max element changemask capacity.
 	{
@@ -258,7 +258,7 @@ UE_NET_TEST_FIXTURE(FTestElementChangeMaskForArrayPropertyNetSerializer, TestCha
 		ReplicationState.PollPropertyReplicationState(Source.Get());
 
 		UE_NET_ASSERT_EQ(ArrayChangeMask.CountSetBits(), ArrayChangeMask.GetNumBits());
-		ArrayChangeMask.Reset();
+		ArrayChangeMask.ClearAllBits();
 	}
 
 	// Double the array size by adding one element at a time until the last bit in the changemask is set and make sure the array is always marked as dirty as well as the corresponding element.
@@ -271,7 +271,7 @@ UE_NET_TEST_FIXTURE(FTestElementChangeMaskForArrayPropertyNetSerializer, TestCha
 		UE_NET_ASSERT_TRUE(ArrayChangeMask.GetBit(FPropertyReplicationState::TArrayElementChangeMaskBitOffset + ElementIt));
 		UE_NET_ASSERT_EQ(ArrayChangeMask.CountSetBits(), 2U);
 
-		ArrayChangeMask.Reset();
+		ArrayChangeMask.ClearAllBits();
 	}
 }
 
@@ -287,7 +287,7 @@ UE_NET_TEST_FIXTURE(FTestElementChangeMaskForArrayPropertyNetSerializer, TestCha
 		Source->ArrayOfUint.SetNum(FPropertyReplicationState::TArrayElementChangeMaskBits);
 		ReplicationState.PollPropertyReplicationState(Source.Get());
 
-		ArrayChangeMask.Reset();
+		ArrayChangeMask.ClearAllBits();
 	}
 
 	// Modify one element at a time until the last bit in the changemask is set and make sure the array is always marked as dirty as well as the corresponding element
@@ -300,7 +300,7 @@ UE_NET_TEST_FIXTURE(FTestElementChangeMaskForArrayPropertyNetSerializer, TestCha
 		UE_NET_ASSERT_TRUE(ArrayChangeMask.GetBit(FPropertyReplicationState::TArrayElementChangeMaskBitOffset + ElementIt));
 		UE_NET_ASSERT_EQ(ArrayChangeMask.CountSetBits(), 2U);
 
-		ArrayChangeMask.Reset();
+		ArrayChangeMask.ClearAllBits();
 	}
 }
 
@@ -310,14 +310,14 @@ UE_NET_TEST_FIXTURE(FTestElementChangeMaskForArrayPropertyNetSerializer, TestCha
 	TObjectPtr<UClassWithDynamicArrayOfPrimitiveTypeForArrayPropertyNetSerializerTest> Source = NewObject<UClassWithDynamicArrayOfPrimitiveTypeForArrayPropertyNetSerializerTest>();
 
 	FNetBitArrayView ArrayChangeMask = UE::Net::Private::GetMemberChangeMask(ReplicationState.GetStateBuffer(), ReplicationStateDescriptor);
-	ArrayChangeMask.Reset();
+	ArrayChangeMask.ClearAllBits();
 
 	// First add elements up to double the max element changemask capacity.
 	{
 		Source->ArrayOfUint.SetNum(2*FPropertyReplicationState::TArrayElementChangeMaskBits);
 		ReplicationState.PollPropertyReplicationState(Source.Get());
 
-		ArrayChangeMask.Reset();
+		ArrayChangeMask.ClearAllBits();
 	}
 
 	// Double the array size by adding one element at a time until the last bit in the changemask is set and make sure the array is always marked as dirty as well as the corresponding element.
@@ -330,7 +330,7 @@ UE_NET_TEST_FIXTURE(FTestElementChangeMaskForArrayPropertyNetSerializer, TestCha
 		UE_NET_ASSERT_TRUE(ArrayChangeMask.GetBit(FPropertyReplicationState::TArrayElementChangeMaskBitOffset + ElementIt));
 		UE_NET_ASSERT_EQ(ArrayChangeMask.CountSetBits(), 2U);
 
-		ArrayChangeMask.Reset();
+		ArrayChangeMask.ClearAllBits();
 	}
 }
 
@@ -340,7 +340,7 @@ UE_NET_TEST_FIXTURE(FTestElementChangeMaskForArrayPropertyNetSerializer, TestCha
 	TObjectPtr<UClassWithDynamicArrayOfPrimitiveTypeForArrayPropertyNetSerializerTest> Source = NewObject<UClassWithDynamicArrayOfPrimitiveTypeForArrayPropertyNetSerializerTest>();
 
 	FNetBitArrayView ArrayChangeMask = UE::Net::Private::GetMemberChangeMask(ReplicationState.GetStateBuffer(), ReplicationStateDescriptor);
-	ArrayChangeMask.Reset();
+	ArrayChangeMask.ClearAllBits();
 
 	constexpr int32 ElementCount = 17;
 	// First add some elements
@@ -348,7 +348,7 @@ UE_NET_TEST_FIXTURE(FTestElementChangeMaskForArrayPropertyNetSerializer, TestCha
 		Source->ArrayOfUint.SetNum(ElementCount);
 		ReplicationState.PollPropertyReplicationState(Source.Get());
 
-		ArrayChangeMask.Reset();
+		ArrayChangeMask.ClearAllBits();
 	}
 
 	// Modify some arbitrary elements
@@ -376,7 +376,7 @@ UE_NET_TEST_FIXTURE(FTestElementChangeMaskForArrayPropertyNetSerializer, TestSer
 	TObjectPtr<UClassWithDynamicArrayOfPrimitiveTypeForArrayPropertyNetSerializerTest> Source = NewObject<UClassWithDynamicArrayOfPrimitiveTypeForArrayPropertyNetSerializerTest>();
 
 	FNetBitArrayView ArrayChangeMask = UE::Net::Private::GetMemberChangeMask(ReplicationState.GetStateBuffer(), ReplicationStateDescriptor);
-	ArrayChangeMask.Reset();
+	ArrayChangeMask.ClearAllBits();
 
 	constexpr int32 ElementCount = 17;
 	// First add some elements
@@ -418,7 +418,7 @@ UE_NET_TEST_FIXTURE(FTestElementChangeMaskForArrayPropertyNetSerializer, TestSer
 
 	// Modify a few elements
 	{
-		ArrayChangeMask.Reset();
+		ArrayChangeMask.ClearAllBits();
 
 		const int32 ElementIndicesToModify[] = {3, 7, 11};
 		for (int32 ElementIndex : ElementIndicesToModify)
diff --git a/Engine/Source/Runtime/Experimental/Iris/Core/Private/Iris/ReplicationState/InternalPropertyReplicationState.cpp b/Engine/Source/Runtime/Experimental/Iris/Core/Private/Iris/ReplicationState/InternalPropertyReplicationState.cpp
--- a/Engine/Source/Runtime/Experimental/Iris/Core/Private/Iris/ReplicationState/InternalPropertyReplicationState.cpp
+++ b/Engine/Source/Runtime/Experimental/Iris/Core/Private/Iris/ReplicationState/InternalPropertyReplicationState.cpp
@@ -18,7 +18,7 @@ void InitReplicationStateInternals(uint8* StateBuffer, const FReplicationStateDe
 
 	// init dirty state tracking
 	FNetBitArrayView DirtyStates = GetMemberChangeMask(StateBuffer, Descriptor);
-	DirtyStates.Reset();
+	DirtyStates.ClearAllBits();
 
 	// Init optional conditionals
 	if (EnumHasAnyFlags(Descriptor->Traits, EReplicationStateTraits::HasLifetimeConditionals))
diff --git a/Engine/Source/Runtime/Experimental/Iris/Core/Private/Iris/ReplicationSystem/DeltaCompression/DeltaCompressionBaselineInvalidationTracker.cpp b/Engine/Source/Runtime/Experimental/Iris/Core/Private/Iris/ReplicationSystem/DeltaCompression/DeltaCompressionBaselineInvalidationTracker.cpp
--- a/Engine/Source/Runtime/Experimental/Iris/Core/Private/Iris/ReplicationSystem/DeltaCompression/DeltaCompressionBaselineInvalidationTracker.cpp
+++ b/Engine/Source/Runtime/Experimental/Iris/Core/Private/Iris/ReplicationSystem/DeltaCompression/DeltaCompressionBaselineInvalidationTracker.cpp
@@ -61,7 +61,7 @@ void FDeltaCompressionBaselineInvalidationTracker::PostSendUpdate()
 	if (InvalidationInfos.Num() > 0)
 	{
 		InvalidationInfos.Empty();
-		InvalidatedObjects.Reset();
+		InvalidatedObjects.ClearAllBits();
 	}
 }
 
diff --git a/Engine/Source/Runtime/Experimental/Iris/Core/Private/Iris/ReplicationSystem/DeltaCompression/DeltaCompressionBaselineStorage.cpp b/Engine/Source/Runtime/Experimental/Iris/Core/Private/Iris/ReplicationSystem/DeltaCompression/DeltaCompressionBaselineStorage.cpp
--- a/Engine/Source/Runtime/Experimental/Iris/Core/Private/Iris/ReplicationSystem/DeltaCompression/DeltaCompressionBaselineStorage.cpp
+++ b/Engine/Source/Runtime/Experimental/Iris/Core/Private/Iris/ReplicationSystem/DeltaCompression/DeltaCompressionBaselineStorage.cpp
@@ -220,7 +220,7 @@ void FDeltaCompressionBaselineStorage::FreeAllBaselineStateInfos()
 
 	UsedBaselineStateInfos.ClearBit(InvalidDeltaCompressionBaselineStateInfoIndex);
 	UsedBaselineStateInfos.ForAllSetBits(DestructStateInfo);
-	UsedBaselineStateInfos.Reset();
+	UsedBaselineStateInfos.ClearAllBits();
 }
 
 }
diff --git a/Engine/Source/Runtime/Experimental/Iris/Core/Private/Iris/ReplicationSystem/DirtyNetObjectTracker.cpp b/Engine/Source/Runtime/Experimental/Iris/Core/Private/Iris/ReplicationSystem/DirtyNetObjectTracker.cpp
--- a/Engine/Source/Runtime/Experimental/Iris/Core/Private/Iris/ReplicationSystem/DirtyNetObjectTracker.cpp
+++ b/Engine/Source/Runtime/Experimental/Iris/Core/Private/Iris/ReplicationSystem/DirtyNetObjectTracker.cpp
@@ -105,16 +105,14 @@ void FDirtyNetObjectTracker::UpdateDirtyNetObjects()
 
 	//$IRIS TODO:  We could look if any objects where actually in the global list and skip the array iteration if not needed.
 
-	const FNetBitArrayView GlobalScopeList = NetRefHandleManager->GetCurrentFrameScopableInternalIndices();
-	const uint32* GlobalScopeListData = GlobalScopeList.GetData();
-	
-	uint32* AccumulatedDirtyNetObjectsData = AccumulatedDirtyNetObjects.GetData();
-	uint32* DirtyNetObjectsData = DirtyNetObjects.GetData();
-
 	const uint32 NumWords = AccumulatedDirtyNetObjects.GetNumWords();
-	check(NumWords == DirtyNetObjects.GetNumWords());
-	check(NumWords == GlobalScopeList.GetNumWords());
 
+	const FNetBitArrayView GlobalScopeList = NetRefHandleManager->GetCurrentFrameScopableInternalIndices();
+	const uint32* GlobalScopeListData = GlobalScopeList.GetDataChecked(NumWords);
+	
+	uint32* AccumulatedDirtyNetObjectsData = AccumulatedDirtyNetObjects.GetDataChecked(NumWords);
+	uint32* DirtyNetObjectsData = DirtyNetObjects.GetDataChecked(NumWords);
+		
 	for (uint32 WordIndex = 0; WordIndex < NumWords; ++WordIndex)
 	{
 		// Due to objects having been marked as dirty and later removed we must make sure that all dirty objects are still in scope.
@@ -218,7 +216,7 @@ void FDirtyNetObjectTracker::ReconcilePolledList(const FNetBitArrayView& Objects
 	MakeNetBitArrayView(AccumulatedDirtyNetObjects).Combine(ObjectsPolled, FNetBitArrayView::AndNotOp);
 
 	// Clear the current frame dirty objects
-	DirtyNetObjects.Reset();
+	DirtyNetObjects.ClearAllBits();
 
 	AllowExternalAccess();
 
diff --git a/Engine/Source/Runtime/Experimental/Iris/Core/Private/Iris/ReplicationSystem/Filtering/FilterOutNetObjectFilter.cpp b/Engine/Source/Runtime/Experimental/Iris/Core/Private/Iris/ReplicationSystem/Filtering/FilterOutNetObjectFilter.cpp
--- a/Engine/Source/Runtime/Experimental/Iris/Core/Private/Iris/ReplicationSystem/Filtering/FilterOutNetObjectFilter.cpp
+++ b/Engine/Source/Runtime/Experimental/Iris/Core/Private/Iris/ReplicationSystem/Filtering/FilterOutNetObjectFilter.cpp
@@ -18,5 +18,5 @@ void UFilterOutNetObjectFilter::RemoveObject(uint32 ObjectIndex, const FNetObjec
 void UFilterOutNetObjectFilter::Filter(FNetObjectFilteringParams& Params)
 {
 	// Filter out everything
-	Params.OutAllowedObjects.Reset();
+	Params.OutAllowedObjects.ClearAllBits();
 }
diff --git a/Engine/Source/Runtime/Experimental/Iris/Core/Private/Iris/ReplicationSystem/Filtering/NetObjectGridFilter.cpp b/Engine/Source/Runtime/Experimental/Iris/Core/Private/Iris/ReplicationSystem/Filtering/NetObjectGridFilter.cpp
--- a/Engine/Source/Runtime/Experimental/Iris/Core/Private/Iris/ReplicationSystem/Filtering/NetObjectGridFilter.cpp
+++ b/Engine/Source/Runtime/Experimental/Iris/Core/Private/Iris/ReplicationSystem/Filtering/NetObjectGridFilter.cpp
@@ -183,7 +183,7 @@ void UNetObjectGridFilter::Filter(FNetObjectFilteringParams& Params)
 
 	// Only allow objects in any of the relevant cells to replicate.
 	FNetBitArrayView AllowedObjects = Params.OutAllowedObjects;
-	AllowedObjects.Reset();
+	AllowedObjects.ClearAllBits();
 
 	/**
 	 * The algorithm will simply iterate over all relevant cells and set the bits
diff --git a/Engine/Source/Runtime/Experimental/Iris/Core/Private/Iris/ReplicationSystem/Filtering/ReplicationFiltering.cpp b/Engine/Source/Runtime/Experimental/Iris/Core/Private/Iris/ReplicationSystem/Filtering/ReplicationFiltering.cpp
--- a/Engine/Source/Runtime/Experimental/Iris/Core/Private/Iris/ReplicationSystem/Filtering/ReplicationFiltering.cpp
+++ b/Engine/Source/Runtime/Experimental/Iris/Core/Private/Iris/ReplicationSystem/Filtering/ReplicationFiltering.cpp
@@ -386,18 +386,19 @@ void FReplicationFiltering::FilterNonRelevantObjects()
 
 void FReplicationFiltering::BuildAlwaysRelevantList(FNetBitArrayView OutAlwaysRelevantList, const FNetBitArrayView ScopeList) const
 {
+	const uint32 MaxWords = OutAlwaysRelevantList.GetNumWords();
+
 	// The list of all replicated objects
-	const uint32* const ScopeListData = ScopeList.GetData();
+	const uint32* const ScopeListData = ScopeList.GetDataChecked(MaxWords);
 
 	// The different list of filtered objects
-	const uint32* const WithOwnerData = ObjectsWithOwnerFilter.GetData();
-	const uint32* const ConnectionFiltersData = AllConnectionFilteredObjects.GetData();
-	const uint32* const DynamicFilteredData = DynamicFilterEnabledObjects.GetData();
-	const uint32* const GroupFilteredOutData = Groups->GetGroupFilteredOutObjects().GetData();
+	const uint32* const WithOwnerData = ObjectsWithOwnerFilter.GetDataChecked(MaxWords);
+	const uint32* const ConnectionFiltersData = AllConnectionFilteredObjects.GetDataChecked(MaxWords);
+	const uint32* const DynamicFilteredData = DynamicFilterEnabledObjects.GetDataChecked(MaxWords);
+	const uint32* const GroupFilteredOutData = Groups->GetGroupFilteredOutObjects().GetDataChecked(MaxWords);
 
-	uint32* OutAlwaysRelevantListData = OutAlwaysRelevantList.GetData();
+	uint32* OutAlwaysRelevantListData = OutAlwaysRelevantList.GetDataChecked(MaxWords);
 
-	const uint32 MaxWords = OutAlwaysRelevantList.GetNumWords();
 	for (uint32 WordIndex = 0; WordIndex < MaxWords; ++WordIndex)
 	{
 		// Build the list of always relevant objects, e.g. objects that have no filters
@@ -580,7 +581,7 @@ bool FReplicationFiltering::SetConnectionFilter(FInternalNetRefIndex ObjectIndex
 	bHasDirtyConnectionFilter = 1;
 	ObjectsWithDirtyConnectionFilter.SetBit(ObjectIndex);
 
-	PerObjectInfoIndexType ObjectInfoIndex = ObjectIndexToPerObjectInfoIndex.GetData()[ObjectIndex];
+	PerObjectInfoIndexType ObjectInfoIndex = ObjectIndexToPerObjectInfoIndex[ObjectIndex];
 	FPerObjectInfo* ObjectInfo = nullptr;
 	if (ObjectInfoIndex == 0)
 	{
@@ -595,11 +596,10 @@ bool FReplicationFiltering::SetConnectionFilter(FInternalNetRefIndex ObjectIndex
 		// The ConnectionIndices bit array does not necessarily have to have storage for the max amount of connections
 		const uint32 InWordCount = ConnectionIndices.GetNumWords();
 		const uint32 MaxWordCount = PerObjectInfoStorageCountForConnections;
-		const uint32* InConnectionIndicesData = ConnectionIndices.GetData();
 		const uint32 WordMask = (ReplicationStatus == ENetFilterStatus::Allow ? 0U : ~0U);
 		for (uint32 WordIt = 0, WordEndIt = FPlatformMath::Min(InWordCount, MaxWordCount); WordIt < WordEndIt; ++WordIt)
 		{
-			ObjectInfo->ConnectionIds[WordIt] = InConnectionIndicesData[WordIt] ^ WordMask;
+			ObjectInfo->ConnectionIds[WordIt] = ConnectionIndices.GetWord(WordIt) ^ WordMask;
 		}
 
 		// Fill in connections that weren't passed in ConnectionIndices
@@ -790,7 +790,7 @@ void FReplicationFiltering::InitNewConnections()
 		{
 			auto MaskObjectToOwner = [this, ConnectionId, &ConnectionInfo](uint32 ObjectIndex)
 			{
-				const bool bIsOwner = (ConnectionId == ObjectIndexToOwningConnection.GetData()[ObjectIndex]);
+				const bool bIsOwner = (ConnectionId == ObjectIndexToOwningConnection[ObjectIndex]);
 				ConnectionInfo.ConnectionFilteredObjects.SetBitValue(ObjectIndex, bIsOwner);
 			};
 
@@ -801,7 +801,7 @@ void FReplicationFiltering::InitNewConnections()
 		{
 			auto MaskObjectToConnection = [this, ConnectionId, &ConnectionInfo](uint32 ObjectIndex)
 			{
-				const PerObjectInfoIndexType ObjectInfoIndex = ObjectIndexToPerObjectInfoIndex.GetData()[ObjectIndex];
+				const PerObjectInfoIndexType ObjectInfoIndex = ObjectIndexToPerObjectInfoIndex[ObjectIndex];
 				const FPerObjectInfo* ObjectInfo = GetPerObjectInfo(ObjectInfoIndex);
 				const ENetFilterStatus ReplicationStatus = GetConnectionFilterStatus(*ObjectInfo, ConnectionId);
 				ConnectionInfo.ConnectionFilteredObjects.SetBitValue(ObjectIndex, ReplicationStatus == ENetFilterStatus::Allow);
@@ -816,7 +816,7 @@ void FReplicationFiltering::InitNewConnections()
 	};
 
 	NewConnections.ForAllSetBits(InitNewConnection);
-	NewConnections.Reset();
+	NewConnections.ClearAllBits();
 }
 
 void FReplicationFiltering::ResetRemovedConnections()
@@ -838,13 +838,12 @@ void FReplicationFiltering::ResetRemovedConnections()
 		FPerObjectInfo* ConnectionStateInfo = GetPerObjectInfo(this->GroupInfos[GroupIndex].ConnectionStateIndex);
 
 		// Special case we want to mask with the valid connections
-		const uint32* ValidConnectionsData = this->ValidConnections.GetData();
 		const uint32 NumWords = this->ValidConnections.GetNumWords();
 
 		// We do not need to do anything more than to restore the state to the default as we reset the effects of the filter
 		for (uint32 WordIt = 0; WordIt < NumWords; ++WordIt)
 		{
-			ConnectionStateInfo->ConnectionIds[WordIt] &= ValidConnectionsData[WordIt];
+			ConnectionStateInfo->ConnectionIds[WordIt] &= ValidConnections.GetWord(WordIt);
 		}
 	};
 
@@ -891,14 +890,15 @@ void FReplicationFiltering::UpdateObjectsInScope()
 
 	TArray<FNetBitArrayBase::StorageWordType> ModifiedWords;
 	ModifiedWords.SetNumUninitialized(WordCountForObjectBitArrays);
-	uint32* ModifiedWordsStorage = ModifiedWords.GetData();
+
 	uint32 ModifiedWordIndex = 0;
-	uint32* ObjectsWithDirtyConnectionFilterStorage = ObjectsWithDirtyConnectionFilter.GetData();
-	uint32* ObjectsWithDirtyOwnerStorage = ObjectsWithDirtyOwner.GetData();
-	const uint32* ObjectsInScopeStorage = ObjectsInScope.GetData();
-	uint32* FakePrevObjectsInScopeStorage = FakePrevObjectsInScope.GetData();
+	uint32* ObjectsWithDirtyConnectionFilterStorage = ObjectsWithDirtyConnectionFilter.GetDataChecked(WordCountForObjectBitArrays);
+	uint32* ObjectsWithDirtyOwnerStorage = ObjectsWithDirtyOwner.GetDataChecked(WordCountForObjectBitArrays);
+	const uint32* ObjectsInScopeStorage = ObjectsInScope.GetDataChecked(WordCountForObjectBitArrays);
+	uint32* FakePrevObjectsInScopeStorage = FakePrevObjectsInScope.GetDataChecked(WordCountForObjectBitArrays);
 	{
-		const uint32* PrevObjectsInScopeStorage = PrevObjectsInScope.GetData();
+		uint32* ModifiedWordsStorage = ModifiedWords.GetData();
+		const uint32* PrevObjectsInScopeStorage = PrevObjectsInScope.GetDataChecked(WordCountForObjectBitArrays);
 
 		for (uint32 WordIt = 0, WordEndIt = WordCountForObjectBitArrays; WordIt != WordEndIt; ++WordIt)
 		{
@@ -926,19 +926,17 @@ void FReplicationFiltering::UpdateObjectsInScope()
 
 	// Clear info for deleted objects
 	{
-		uint32* ObjectsWithOwnerFilterStorage = ObjectsWithOwnerFilter.GetData();
-		uint16* ObjectIndexToOwningConnectionStorage = ObjectIndexToOwningConnection.GetData();
-		const uint32* SubObjectInternalIndicesStorage = NetRefHandleManager->GetSubObjectInternalIndices().GetData();
+		const FNetBitArrayView SubObjectInternalIndices = NetRefHandleManager->GetSubObjectInternalIndicesView();
 		uint32 PrevParentIndex = FNetRefHandleManager::InvalidInternalIndex;
 		for (uint32 WordIt = 0, WordEndIt = ModifiedWordIndex; WordIt != WordEndIt; ++WordIt)
 		{
-			const uint32 WordIndex = ModifiedWordsStorage[WordIt];
+			const uint32 WordIndex = ModifiedWords[WordIt];
 			const uint32 PrevExistingObjects = FakePrevObjectsInScopeStorage[WordIndex];
 			const uint32 ExistingObjects = ObjectsInScopeStorage[WordIndex];
 
 			// Deleted objects can't be dirty and can't have filtering.
 			ObjectsWithDirtyConnectionFilterStorage[WordIndex] &= ExistingObjects;
-			ObjectsWithOwnerFilterStorage[WordIndex] &= ExistingObjects;
+			ObjectsWithOwnerFilter.GetWord(WordIndex) &= ExistingObjects;
 			ObjectsWithDirtyOwnerStorage[WordIndex] &= ExistingObjects;
 
 			// Clear dynamic filters and owner info from deleted objects.
@@ -966,7 +964,7 @@ void FReplicationFiltering::UpdateObjectsInScope()
 			// Make sure subobjects that are added after the parent gets properly updated.
 			// Dirtying the parent will cause the subobjects to be updated too.
 			uint32 AddedObjects = ExistingObjects & ~PrevExistingObjects;
-			if (uint32 AddedSubObjects = AddedObjects & SubObjectInternalIndicesStorage[WordIndex])
+			if (uint32 AddedSubObjects = AddedObjects & SubObjectInternalIndices.GetWord(WordIndex))
 			{
 				for ( ; AddedSubObjects; )
 				{
@@ -1011,7 +1009,7 @@ void FReplicationFiltering::UpdateObjectsInScope()
 						continue;
 					}
 
-					if (ObjectIndexToOwningConnectionStorage[ParentIndex])
+					if (ObjectIndexToOwningConnection[ParentIndex])
 					{
 						bHasDirtyOwner = 1;
 						ObjectsWithDirtyOwner.SetBit(ParentIndex);
@@ -1034,16 +1032,16 @@ void FReplicationFiltering::UpdateObjectsInScope()
 		FPerConnectionInfo& ConnectionInfo = ConnectionInfos[ConnectionId];
 
 		// Or in brand new objects and mask off deleted objects by anding with now existing objects.
-		uint32* FilteredObjectsStorage = ConnectionInfo.ConnectionFilteredObjects.GetData();
-		uint32* GroupExcludedObjectsStorage = ConnectionInfo.GroupExcludedObjects.GetData();
-		uint32* GroupIncludedObjectsStorage = ConnectionInfo.GroupIncludedObjects.GetData();
-		uint32* ObjectsInScopeBeforeDynamicFiltering = ConnectionInfo.ObjectsInScopeBeforeDynamicFiltering.GetData();
-		uint32* DynamicFilteredOutObjects = ConnectionInfo.DynamicFilteredOutObjects.GetData();
-		uint32* DynamicFilteredOutObjectsHysteresisAdjusted = ConnectionInfo.DynamicFilteredOutObjectsHysteresisAdjusted.GetData();
+		uint32* FilteredObjectsStorage = ConnectionInfo.ConnectionFilteredObjects.GetDataChecked(WordCountForObjectBitArrays);
+		uint32* GroupExcludedObjectsStorage = ConnectionInfo.GroupExcludedObjects.GetDataChecked(WordCountForObjectBitArrays);
+		uint32* GroupIncludedObjectsStorage = ConnectionInfo.GroupIncludedObjects.GetDataChecked(WordCountForObjectBitArrays);
+		uint32* ObjectsInScopeBeforeDynamicFiltering = ConnectionInfo.ObjectsInScopeBeforeDynamicFiltering.GetDataChecked(WordCountForObjectBitArrays);
+		uint32* DynamicFilteredOutObjects = HasDynamicFilters() ? ConnectionInfo.DynamicFilteredOutObjects.GetDataChecked(WordCountForObjectBitArrays) : nullptr;
+		uint32* DynamicFilteredOutObjectsHysteresisAdjusted = HasDynamicFilters() ? ConnectionInfo.DynamicFilteredOutObjectsHysteresisAdjusted.GetDataChecked(WordCountForObjectBitArrays) : nullptr;
 
 		for (uint32 WordIt = 0, WordEndIt = ModifiedWordIndex; WordIt != WordEndIt; ++WordIt)
 		{
-			const uint32 WordIndex = ModifiedWordsStorage[WordIt];
+			const uint32 WordIndex = ModifiedWords[WordIt];
 
 			const uint32 PrevExistingObjects = FakePrevObjectsInScopeStorage[WordIndex];
 			const uint32 ExistingObjects = ObjectsInScopeStorage[WordIndex];
@@ -1099,17 +1097,16 @@ void FReplicationFiltering::UpdateOwnerAndConnectionFiltering()
 	{
 		IRIS_PROFILER_SCOPE(FReplicationFiltering_UpdateDirtyOwnerValue);
 
-		uint16* const ObjectIndexToOwningConnectionStorage = ObjectIndexToOwningConnection.GetData();
-		auto UpdateOwners = [this, ObjectIndexToOwningConnectionStorage](uint32 ObjectIndex)
+		auto UpdateOwners = [this](uint32 ObjectIndex)
 		{
-			const uint32 OwningConnectionId = ObjectIndexToOwningConnectionStorage[ObjectIndex];
+			const uint32 OwningConnectionId = ObjectIndexToOwningConnection[ObjectIndex];
 
 			for (const FInternalNetRefIndex SubObjectIndex : NetRefHandleManager->GetSubObjects(ObjectIndex))
 			{
 				const FNetRefHandleManager::FReplicatedObjectData& ObjectData = NetRefHandleManager->GetReplicatedObjectDataNoCheck(SubObjectIndex);
 				if (ObjectData.IsSubObject())
 				{
-					ObjectIndexToOwningConnectionStorage[SubObjectIndex] = static_cast<uint16>(OwningConnectionId);
+					ObjectIndexToOwningConnection[SubObjectIndex] = static_cast<uint16>(OwningConnectionId);
 				}
 			}
 		};
@@ -1143,13 +1140,13 @@ void FReplicationFiltering::UpdateOwnerAndConnectionFiltering()
 					bool bObjectIsInScope = true;
 					if (HasOwnerFilter(ObjectIndex))
 					{
-						const uint32 OwningConnection = ObjectIndexToOwningConnection.GetData()[ObjectIndex];
+						const uint32 OwningConnection = ObjectIndexToOwningConnection[ObjectIndex];
 						const bool bIsOwner = (ConnectionId == OwningConnection);
 						bObjectIsInScope = bIsOwner;
 					}
 					else if (HasConnectionFilter(ObjectIndex))
 					{
-						const PerObjectInfoIndexType ObjectInfoIndex = ObjectIndexToPerObjectInfoIndex.GetData()[ObjectIndex];
+						const PerObjectInfoIndexType ObjectInfoIndex = ObjectIndexToPerObjectInfoIndex[ObjectIndex];
 						const FPerObjectInfo* ObjectInfo = GetPerObjectInfo(ObjectInfoIndex);
 						const ENetFilterStatus ReplicationStatus = GetConnectionFilterStatus(*ObjectInfo, ConnectionId);
 						bObjectIsInScope = ReplicationStatus == ENetFilterStatus::Allow;
@@ -1184,8 +1181,8 @@ void FReplicationFiltering::UpdateOwnerAndConnectionFiltering()
 	// Clear out dirtiness
 	bHasDirtyConnectionFilter = 0;
 	bHasDirtyOwner = 0;
-	ObjectsWithDirtyConnectionFilter.Reset();
-	ObjectsWithDirtyOwner.Reset();
+	ObjectsWithDirtyConnectionFilter.ClearAllBits();
+	ObjectsWithDirtyOwner.ClearAllBits();
 }
 
 void FReplicationFiltering::UpdateGroupExclusionFiltering()
@@ -1202,16 +1199,15 @@ void FReplicationFiltering::UpdateGroupExclusionFiltering()
 
 		const FPerObjectInfo* ConnectionStateInfo = GetPerObjectInfo(this->GroupInfos[GroupIndex].ConnectionStateIndex);
 		const FNetObjectGroup* Group = this->Groups->GetGroupByIndex(FNetObjectGroupHandle::FGroupIndexType(GroupIndex));
-		FPerConnectionInfo* LocalConnectionInfos = this->ConnectionInfos.GetData();
 
 		const FNetBitArrayView CurrentFrameScopableObjects = NetRefHandleManager->GetCurrentFrameScopableInternalIndices();
 
-		auto UpdateGroupFilterForConnection = [this, ConnectionStateInfo, Group, LocalConnectionInfos, CurrentFrameScopableObjects](uint32 ConnectionId)
+		auto UpdateGroupFilterForConnection = [this, ConnectionStateInfo, Group, CurrentFrameScopableObjects](uint32 ConnectionId)
 		{
 			if (this->GetConnectionFilterStatus(*ConnectionStateInfo, ConnectionId) == ENetFilterStatus::Disallow)
 			{
-				FNetBitArray& GroupExcludedObjects = LocalConnectionInfos[ConnectionId].GroupExcludedObjects;
-				FNetBitArray& ObjectsInScopeBeforeDynamicFiltering = LocalConnectionInfos[ConnectionId].ObjectsInScopeBeforeDynamicFiltering;
+				FNetBitArray& GroupExcludedObjects = ConnectionInfos[ConnectionId].GroupExcludedObjects;
+				FNetBitArray& ObjectsInScopeBeforeDynamicFiltering = ConnectionInfos[ConnectionId].ObjectsInScopeBeforeDynamicFiltering;
 
 				for (const FInternalNetRefIndex ObjectIndex : Group->Members)
 				{
@@ -1236,7 +1232,7 @@ void FReplicationFiltering::UpdateGroupExclusionFiltering()
 
 	// Clear out dirtiness
 	bHasDirtyExclusionFilterGroup = 0;
-	DirtyExclusionFilterGroups.Reset();
+	DirtyExclusionFilterGroups.ClearAllBits();
 }
 
 void FReplicationFiltering::UpdateGroupInclusionFiltering()
@@ -1252,16 +1248,15 @@ void FReplicationFiltering::UpdateGroupInclusionFiltering()
 
 		const FPerObjectInfo* ConnectionStateInfo = GetPerObjectInfo(this->GroupInfos[GroupIndex].ConnectionStateIndex);
 		const FNetObjectGroup* Group = this->Groups->GetGroupByIndex(FNetObjectGroupHandle::FGroupIndexType(GroupIndex));
-		FPerConnectionInfo* LocalConnectionInfos = this->ConnectionInfos.GetData();
 
 		const FNetBitArrayView CurrentFrameScopableObjects = NetRefHandleManager->GetCurrentFrameScopableInternalIndices();
 		const FNetBitArrayView SubObjectInternalIndices = NetRefHandleManager->GetSubObjectInternalIndicesView();
 
-		auto UpdateGroupFilterForConnection = [this, ConnectionStateInfo, Group, LocalConnectionInfos, &CurrentFrameScopableObjects, &SubObjectInternalIndices](uint32 ConnectionId)
+		auto UpdateGroupFilterForConnection = [this, ConnectionStateInfo, Group, &CurrentFrameScopableObjects, &SubObjectInternalIndices](uint32 ConnectionId)
 		{
 			if (this->GetConnectionFilterStatus(*ConnectionStateInfo, ConnectionId) == ENetFilterStatus::Allow)
 			{
-				FNetBitArray& GroupIncludedObjects = LocalConnectionInfos[ConnectionId].GroupIncludedObjects;
+				FNetBitArray& GroupIncludedObjects = ConnectionInfos[ConnectionId].GroupIncludedObjects;
 
 				for (const FInternalNetRefIndex ObjectIndex : Group->Members)
 				{
@@ -1289,7 +1284,7 @@ void FReplicationFiltering::UpdateGroupInclusionFiltering()
 
 	// Clear out dirtiness
 	bHasDirtyInclusionFilterGroup = 0;
-	DirtyInclusionFilterGroups.Reset();
+	DirtyInclusionFilterGroups.ClearAllBits();
 }
 
 void FReplicationFiltering::PreUpdateDynamicFiltering()
@@ -1341,9 +1336,9 @@ void FReplicationFiltering::UpdateDynamicFiltering()
 
 	// Subobjects will never be added to a dynamic filter, but objects can become dependent at any time.
 	// We need to make sure they are not filtered out.
-	const uint32* SubObjectsData = NetRefHandleManager->GetSubObjectInternalIndices().GetData();
-	const uint32* DependentObjectsData = NetRefHandleManager->GetDependentObjectInternalIndices().GetData();
-	const uint32* ObjectsRequiringDynamicFilterUpdateData = ObjectsRequiringDynamicFilterUpdate.GetData();
+	const uint32* SubObjectsData = NetRefHandleManager->GetSubObjectInternalIndices().GetDataChecked(WordCountForObjectBitArrays);
+	const uint32* DependentObjectsData = NetRefHandleManager->GetDependentObjectInternalIndices().GetDataChecked(WordCountForObjectBitArrays);
+	const uint32* ObjectsRequiringDynamicFilterUpdateData = ObjectsRequiringDynamicFilterUpdate.GetDataChecked(WordCountForObjectBitArrays);
 
 	uint32* ConnectionIds = static_cast<uint32*>(FMemory_Alloca(ValidConnections.GetNumBits() * sizeof(uint32)));
 	uint32 ConnectionCount = 0;
@@ -1353,8 +1348,8 @@ void FReplicationFiltering::UpdateDynamicFiltering()
 	{
 		FPerConnectionInfo& ConnectionInfo = this->ConnectionInfos[ConnId];
 
-		ConnectionInfo.InProgressDynamicFilteredOutObjects.Reset();
-		uint32* InProgressDynamicFilteredOutObjectsData = ConnectionInfo.InProgressDynamicFilteredOutObjects.GetData();
+		ConnectionInfo.InProgressDynamicFilteredOutObjects.ClearAllBits();
+		uint32* InProgressDynamicFilteredOutObjectsData = HasDynamicFilters() ? ConnectionInfo.InProgressDynamicFilteredOutObjects.GetDataChecked(WordCountForObjectBitArrays) : nullptr;
 
 		/*
 		 * Apply dynamic filters.
@@ -1560,9 +1555,9 @@ void FReplicationFiltering::UpdateDynamicFiltering()
 		{
 			IRIS_PROFILER_SCOPE(FReplicationFiltering_UpdateConnectionScope);
 
-			uint32* ObjectsInScopeData = ConnectionInfo.ObjectsInScope.GetData();
-			const uint32* ObjectsInScopeBeforeDynamicFilteringData = ConnectionInfo.ObjectsInScopeBeforeDynamicFiltering.GetData();
-			const uint32* DynamicFilteredOutObjectsHysteresisAdjustedData = DynamicFilteredOutObjectsHysteresisAdjusted.GetData();
+			uint32* ObjectsInScopeData = ConnectionInfo.ObjectsInScope.GetDataChecked(WordCountForObjectBitArrays);
+			const uint32* ObjectsInScopeBeforeDynamicFilteringData = ConnectionInfo.ObjectsInScopeBeforeDynamicFiltering.GetDataChecked(WordCountForObjectBitArrays);
+			const uint32* DynamicFilteredOutObjectsHysteresisAdjustedData = DynamicFilteredOutObjectsHysteresisAdjusted.GetDataChecked(WordCountForObjectBitArrays);
 			for (SIZE_T WordIt = 0, WordEndIt = WordCountForObjectBitArrays; WordIt != WordEndIt; ++WordIt)
 			{
 				const uint32 ObjectsInScopeBeforeWord = ObjectsInScopeBeforeDynamicFilteringData[WordIt];
@@ -1670,7 +1665,7 @@ void FReplicationFiltering::PostUpdateDynamicFiltering()
 {
 	IRIS_PROFILER_SCOPE(FReplicationFiltering_PostUpdateDynamicFiltering);
 
-	ObjectsRequiringDynamicFilterUpdate.Reset();
+	ObjectsRequiringDynamicFilterUpdate.ClearAllBits();
 
 	// Tell filters to clean up after filtering. It's only called if any object has the filter set.
 	{
@@ -1797,7 +1792,7 @@ FReplicationFiltering::FPerObjectInfo* FReplicationFiltering::AllocPerObjectInfo
 {
 	ObjectsWithPerObjectInfo.SetBit(ObjectIndex);
 	const PerObjectInfoIndexType ObjectInfoIndex = AllocPerObjectInfo();
-	ObjectIndexToPerObjectInfoIndex.GetData()[ObjectIndex] = ObjectInfoIndex;
+	ObjectIndexToPerObjectInfoIndex[ObjectIndex] = ObjectInfoIndex;
 
 	FPerObjectInfo* ObjectInfo = GetPerObjectInfo(ObjectInfoIndex);
 	SetPerObjectInfoFilterStatus(*ObjectInfo, ENetFilterStatus::Allow);
@@ -1807,7 +1802,7 @@ FReplicationFiltering::FPerObjectInfo* FReplicationFiltering::AllocPerObjectInfo
 
 void FReplicationFiltering::FreePerObjectInfoForObject(uint32 ObjectIndex)
 {
-	PerObjectInfoIndexType& ObjectInfoIndex = ObjectIndexToPerObjectInfoIndex.GetData()[ObjectIndex];
+	PerObjectInfoIndexType& ObjectInfoIndex = ObjectIndexToPerObjectInfoIndex[ObjectIndex];
 	if (ObjectInfoIndex == 0)
 	{
 		return;
@@ -1897,7 +1892,7 @@ void FReplicationFiltering::UpdateSubObjectFilters()
 	};
 
 	FNetBitArray::ForAllSetBits(DirtySubObjectFilterGroups, SubObjectFilterGroups, FNetBitArray::AndOp, UpdateSubObjectFilterGroup);
-	DirtySubObjectFilterGroups.Reset();
+	DirtySubObjectFilterGroups.ClearAllBits();
 }
 
 void FReplicationFiltering::SetSubObjectFilterStatus(FNetObjectGroupHandle GroupHandle, ENetFilterStatus ReplicationStatus)
@@ -2781,7 +2776,6 @@ void FReplicationFiltering::BuildObjectsInFilterList(FNetBitArrayView OutObjects
 	}
 }
 
-
 void FReplicationFiltering::PreUpdateObjectScopeHysteresis()
 {
 	if (HysteresisState.Mode == EHysteresisProcessingMode::Enabled)
@@ -2796,7 +2790,7 @@ void FReplicationFiltering::PreUpdateObjectScopeHysteresis()
 void FReplicationFiltering::PostUpdateObjectScopeHysteresis()
 {
 	HysteresisState.ObjectsToClearCount = 0;
-	HysteresisState.ObjectsToClear.Reset();
+	HysteresisState.ObjectsToClear.ClearAllBits();
 }
 
 void FReplicationFiltering::ClearObjectsFromHysteresis()
diff --git a/Engine/Source/Runtime/Experimental/Iris/Core/Private/Iris/ReplicationSystem/NetBlob/NetBlobManager.cpp b/Engine/Source/Runtime/Experimental/Iris/Core/Private/Iris/ReplicationSystem/NetBlob/NetBlobManager.cpp
--- a/Engine/Source/Runtime/Experimental/Iris/Core/Private/Iris/ReplicationSystem/NetBlob/NetBlobManager.cpp
+++ b/Engine/Source/Runtime/Experimental/Iris/Core/Private/Iris/ReplicationSystem/NetBlob/NetBlobManager.cpp
@@ -449,7 +449,7 @@ void FNetBlobManager::FNetObjectAttachmentSendQueue::PrepareAndProcessOOBAttachm
 
 	if (ScheduleAsOOBAttachmentQueue.Num() <= 0)
 	{
-		OutConnectionsPendingImmediateSend.Reset();
+		OutConnectionsPendingImmediateSend.ClearAllBits();
 		return;
 	}
 
diff --git a/Engine/Source/Runtime/Experimental/Iris/Core/Private/Iris/ReplicationSystem/NetBlob/NetBlobManager.h b/Engine/Source/Runtime/Experimental/Iris/Core/Private/Iris/ReplicationSystem/NetBlob/NetBlobManager.h
--- a/Engine/Source/Runtime/Experimental/Iris/Core/Private/Iris/ReplicationSystem/NetBlob/NetBlobManager.h
+++ b/Engine/Source/Runtime/Experimental/Iris/Core/Private/Iris/ReplicationSystem/NetBlob/NetBlobManager.h
@@ -161,7 +161,7 @@ private:
 				Connections = nullptr;
 				NetRefHandleManager = nullptr;
 				QueueToProcess = nullptr;
-				ConnectionsPendingSendInPostDispatch.Reset();
+				ConnectionsPendingSendInPostDispatch.ClearAllBits();
 			}
 
 			bool IsValid() const { return NetRefHandleManager != nullptr; }
diff --git a/Engine/Source/Runtime/Experimental/Iris/Core/Private/Iris/ReplicationSystem/NetRefHandleManager.cpp b/Engine/Source/Runtime/Experimental/Iris/Core/Private/Iris/ReplicationSystem/NetRefHandleManager.cpp
--- a/Engine/Source/Runtime/Experimental/Iris/Core/Private/Iris/ReplicationSystem/NetRefHandleManager.cpp
+++ b/Engine/Source/Runtime/Experimental/Iris/Core/Private/Iris/ReplicationSystem/NetRefHandleManager.cpp
@@ -705,7 +705,7 @@ bool FNetRefHandleManager::SetSubObjectNetCondition(FInternalNetRefIndex SubObje
 
 					check(SubObjectConditionals->Num() == SubObjectsArray->Num());
 
-					SubObjectConditionals->GetData()[SubObjectArrayIndex] = SubObjectCondition;
+					(*SubObjectConditionals)[SubObjectArrayIndex] = SubObjectCondition;
 					return true;
 				}
 			}
diff --git a/Engine/Source/Runtime/Experimental/Iris/Core/Private/Iris/ReplicationSystem/ObjectReplicationBridge.cpp b/Engine/Source/Runtime/Experimental/Iris/Core/Private/Iris/ReplicationSystem/ObjectReplicationBridge.cpp
--- a/Engine/Source/Runtime/Experimental/Iris/Core/Private/Iris/ReplicationSystem/ObjectReplicationBridge.cpp
+++ b/Engine/Source/Runtime/Experimental/Iris/Core/Private/Iris/ReplicationSystem/ObjectReplicationBridge.cpp
@@ -844,7 +844,7 @@ void UObjectReplicationBridge::PreSendUpdate()
 	using namespace UE::Net;
 
 	FNetBitArrayView ObjectsConsideredForPolling = GetReplicationSystem()->GetReplicationSystemInternal()->GetNetRefHandleManager().GetPolledObjectsInternalIndices();
-	ObjectsConsideredForPolling.Reset();
+	ObjectsConsideredForPolling.ClearAllBits();
 
 	BuildPollList(ObjectsConsideredForPolling);
 
diff --git a/Engine/Source/Runtime/Experimental/Iris/Core/Private/Iris/ReplicationSystem/ObjectReplicationBridgeDebugging.cpp b/Engine/Source/Runtime/Experimental/Iris/Core/Private/Iris/ReplicationSystem/ObjectReplicationBridgeDebugging.cpp
--- a/Engine/Source/Runtime/Experimental/Iris/Core/Private/Iris/ReplicationSystem/ObjectReplicationBridgeDebugging.cpp
+++ b/Engine/Source/Runtime/Experimental/Iris/Core/Private/Iris/ReplicationSystem/ObjectReplicationBridgeDebugging.cpp
@@ -503,7 +503,7 @@ void UObjectReplicationBridge::PrintRelevantObjectsForConnections(const TArray<F
 	TArray<uint32> RequestedConnectionList = FindConnectionsFromArgs(Args);
 	if (RequestedConnectionList.Num())
 	{
-		ConnectionsToPrint.Reset();
+		ConnectionsToPrint.ClearAllBits();
 		for (uint32 ConnectionId : RequestedConnectionList)
 		{
 			if (ValidConnections.IsBitSet(ConnectionId))
diff --git a/Engine/Source/Runtime/Experimental/Iris/Core/Private/Iris/ReplicationSystem/Prioritization/NetObjectConnectionFilter.cpp b/Engine/Source/Runtime/Experimental/Iris/Core/Private/Iris/ReplicationSystem/Prioritization/NetObjectConnectionFilter.cpp
--- a/Engine/Source/Runtime/Experimental/Iris/Core/Private/Iris/ReplicationSystem/Prioritization/NetObjectConnectionFilter.cpp
+++ b/Engine/Source/Runtime/Experimental/Iris/Core/Private/Iris/ReplicationSystem/Prioritization/NetObjectConnectionFilter.cpp
@@ -127,7 +127,7 @@ void UNetObjectConnectionFilter::Filter(FNetObjectFilteringParams& Params)
 	IRIS_PROFILER_SCOPE(UNetObjectConnectionFilter_Filter);
 
 	UE::Net::FNetBitArrayView& AllowedObjects = Params.OutAllowedObjects;
-	AllowedObjects.Reset();
+	AllowedObjects.ClearAllBits();
 
 	FPerConnectionInfo& ConnectionInfo = PerConnectionInfos[Params.ConnectionId];
 	ConnectionInfo.ReplicationEnabledObjects.ForAllSetBits(
diff --git a/Engine/Source/Runtime/Experimental/Iris/Core/Private/Iris/ReplicationSystem/Prioritization/ReplicationPrioritization.cpp b/Engine/Source/Runtime/Experimental/Iris/Core/Private/Iris/ReplicationSystem/Prioritization/ReplicationPrioritization.cpp
--- a/Engine/Source/Runtime/Experimental/Iris/Core/Private/Iris/ReplicationSystem/Prioritization/ReplicationPrioritization.cpp
+++ b/Engine/Source/Runtime/Experimental/Iris/Core/Private/Iris/ReplicationSystem/Prioritization/ReplicationPrioritization.cpp
@@ -704,7 +704,7 @@ void FReplicationPrioritization::UpdatePrioritiesForNewAndDeletedObjects()
 		}
 
 		// $IRIS TODO: Want ForAllSetBits with clear.
-		ObjectsWithNewStaticPriority.Reset();
+		ObjectsWithNewStaticPriority.ClearAllBits();
 	}
 
 	// Copy the priorities for the new objects to each connection
diff --git a/Engine/Source/Runtime/Experimental/Iris/Core/Private/Iris/ReplicationSystem/ReplicationReader.cpp b/Engine/Source/Runtime/Experimental/Iris/Core/Private/Iris/ReplicationSystem/ReplicationReader.cpp
--- a/Engine/Source/Runtime/Experimental/Iris/Core/Private/Iris/ReplicationSystem/ReplicationReader.cpp
+++ b/Engine/Source/Runtime/Experimental/Iris/Core/Private/Iris/ReplicationSystem/ReplicationReader.cpp
@@ -1567,7 +1567,7 @@ void FReplicationReader::CleanupReferenceTracking(FReplicatedObjectInfo* ObjectI
 
 void FReplicationReader::BuildUnresolvedChangeMaskAndUpdateObjectReferenceTracking(const FResolveAndCollectUnresolvedAndResolvedReferenceCollector& Collector, FNetBitArrayView CollectorChangeMask, FReplicatedObjectInfo* ReplicationInfo, FNetBitArrayView& OutUnresolvedChangeMask, FResolvedNetRefHandlesArray& OutNewResolvedRefHandles)
 {
-	OutUnresolvedChangeMask.Reset();
+	OutUnresolvedChangeMask.ClearAllBits();
 	bool bHasUnresolvedInitReferences = false;
 
 	UnresolvedReferencesCache.Reset();
@@ -1656,7 +1656,7 @@ void FReplicationReader::ResolveAndDispatchUnresolvedReferencesForObject(FNetSer
 		BuildUnresolvedChangeMaskAndUpdateObjectReferenceTracking(Collector, TempUnresolvedChangeMask, ReplicationInfo, UnresolvedChangeMask, NewResolvedRefHandles);
 
 		// Re-purpose temp changemask for members that has resolved references.
-		TempChangeMask.Reset();
+		TempChangeMask.ClearAllBits();
 		FNetBitArrayView ResolvedChangeMask = TempChangeMask;
 
 		// Merge in partially resolved changes
diff --git a/Engine/Source/Runtime/Experimental/Iris/Core/Private/Iris/ReplicationSystem/ReplicationSystem.cpp b/Engine/Source/Runtime/Experimental/Iris/Core/Private/Iris/ReplicationSystem/ReplicationSystem.cpp
--- a/Engine/Source/Runtime/Experimental/Iris/Core/Private/Iris/ReplicationSystem/ReplicationSystem.cpp
+++ b/Engine/Source/Runtime/Experimental/Iris/Core/Private/Iris/ReplicationSystem/ReplicationSystem.cpp
@@ -481,7 +481,7 @@ public:
 		};
 
 		DirtyObjectsToQuantize.ForAllSetBits(QuantizeFunction);
-		DirtyObjectsToQuantize.Reset();
+		DirtyObjectsToQuantize.ClearAllBits();
 
 		const uint32 ReplicationSystemId = ReplicationSystem->GetId();
 		UE_NET_TRACE_FRAME_STATSCOUNTER(ReplicationSystemId, ReplicationSystem.QuantizedObjectCount, QuantizedObjectCount, ENetTraceVerbosity::Trace);
@@ -817,7 +817,7 @@ IRISCORE_API void UReplicationSystem::SendUpdate(TFunctionRef<void(TArrayView<ui
 		// We only need to send data to connections that has data to send in PostTickDispatch
 
 		FNetBitArray::ForAllSetBits(Impl->ConnectionsPendingPostTickDispatchSend, ReplicatingConnections, FNetBitArray::AndOp, [&ConnectionToUpdate](uint32 ConnId) { ConnectionToUpdate.Add(ConnId);});
-		Impl->ConnectionsPendingPostTickDispatchSend.Reset();
+		Impl->ConnectionsPendingPostTickDispatchSend.ClearAllBits();
 	}
 
 	SendFunction(MakeArrayView(ConnectionToUpdate));
diff --git a/Engine/Source/Runtime/Experimental/Iris/Core/Private/Iris/ReplicationSystem/ReplicationWriter.cpp b/Engine/Source/Runtime/Experimental/Iris/Core/Private/Iris/ReplicationSystem/ReplicationWriter.cpp
--- a/Engine/Source/Runtime/Experimental/Iris/Core/Private/Iris/ReplicationSystem/ReplicationWriter.cpp
+++ b/Engine/Source/Runtime/Experimental/Iris/Core/Private/Iris/ReplicationSystem/ReplicationWriter.cpp
@@ -271,9 +271,6 @@ static bool s_ValidateReplicationRecord(const FReplicationRecord* ReplicationRec
 
 FReplicationWriter::~FReplicationWriter()
 {
-	// NOTE: Currently disabled because FReplicationWriter until the performance impact of TNetChunkedArray can be measured on the server.
-	//NetRefHandleManager->GetLargestIndexIncreaseDelegate().Remove(OnLargestIndexIncreaseHandle);
-
 	DiscardAllRecords();
 
 	// Freeing the huge object queue needs to be done before calling StopAllReplication() in order to be able to free any changemask allocations.
@@ -383,9 +380,6 @@ void FReplicationWriter::Init(const FReplicationParameters& InParameters)
 	ReplicatedObjectsRecordInfoLists.SetNumZeroed(MaxSupportedObjects);
 	SchedulingPriorities.SetNumZeroed(MaxSupportedObjects);
 	
-	// NOTE: Currently disabled because FReplicationWriter until the performance impact of TNetChunkedArray can be measured on the server.
-	//OnLargestIndexIncreaseHandle = NetRefHandleManager->GetLargestIndexIncreaseDelegate().AddRaw(this, &FReplicationWriter::OnLargestIndexIncrease);
-
 	SetNetObjectListsSize(Parameters.MaxInternalNetRefIndex);
 
 	NetRefHandleManager->GetOnMaxInternalNetRefIndexIncreasedDelegate().AddRaw(this, &FReplicationWriter::OnMaxInternalNetRefIndexIncreased);
@@ -2964,7 +2958,7 @@ uint32 FReplicationWriter::WriteObjects(FNetSerializationContext& Context)
 	WriteContext.SortedObjectCount = SortedCount;
 
 	// Reset objects written this packet
-	WriteContext.ObjectsWrittenThisPacket.Reset();
+	WriteContext.ObjectsWrittenThisPacket.ClearAllBits();
 
 	return WrittenObjectCount;
 }
@@ -3023,7 +3017,7 @@ int FReplicationWriter::HandleObjectBatchSuccess(const FBatchInfo& BatchInfo, FR
 			CreateObjectRecord(&ChangeMask, Info, BatchObjectInfo, ObjectRecord);
 
 			// The object no longer has any dirty state, but may still have attachments that didn't fit
-			ChangeMask.Reset();
+			ChangeMask.ClearAllBits();
 
 			++ObjectCount;
 			if (Info.LastAckedBaselineIndex != FDeltaCompressionBaselineManager::InvalidBaselineIndex)
@@ -3625,14 +3619,6 @@ void FReplicationWriter::MarkObjectDirty(FInternalNetRefIndex InternalIndex, con
 	ObjectsWithDirtyChanges.SetBit(InternalIndex);
 }
 
-void FReplicationWriter::OnLargestIndexIncrease(uint32 InternalIndex)
-{
-	ReplicatedObjects.SetNumZeroed(InternalIndex);
-	ReplicatedObjectsRecordInfoLists.SetNumZeroed(InternalIndex);
-	SchedulingPriorities.SetNumZeroed(InternalIndex);
-}
-
-
 FReplicationWriter::FHugeObjectContext::FHugeObjectContext() = default;
 
 FReplicationWriter::FHugeObjectContext::~FHugeObjectContext() = default;
diff --git a/Engine/Source/Runtime/Experimental/Iris/Core/Private/Iris/ReplicationSystem/ReplicationWriter.h b/Engine/Source/Runtime/Experimental/Iris/Core/Private/Iris/ReplicationSystem/ReplicationWriter.h
--- a/Engine/Source/Runtime/Experimental/Iris/Core/Private/Iris/ReplicationSystem/ReplicationWriter.h
+++ b/Engine/Source/Runtime/Experimental/Iris/Core/Private/Iris/ReplicationSystem/ReplicationWriter.h
@@ -532,9 +532,6 @@ private:
 
 	void MarkObjectDirty(FInternalNetRefIndex InternalIndex, const char* Caller);
 
-	/* Grow internal buffers to support a given index. */
-	void OnLargestIndexIncrease(uint32 InternalIndex);
-
 private:
 	// Replication parameters
 	FReplicationParameters Parameters;
@@ -563,8 +560,6 @@ private:
 	// Handles logic for all attachments to objects.
 	FNetObjectAttachmentsWriter Attachments;
 
-	FDelegateHandle OnLargestIndexIncreaseHandle;
-
 	// Cached internal systems
 	FReplicationSystemInternal* ReplicationSystemInternal = nullptr;
 	FNetRefHandleManager* NetRefHandleManager = nullptr;
diff --git a/Engine/Source/Runtime/Experimental/Iris/Core/Private/Iris/ReplicationSystem/WorldLocations.cpp b/Engine/Source/Runtime/Experimental/Iris/Core/Private/Iris/ReplicationSystem/WorldLocations.cpp
--- a/Engine/Source/Runtime/Experimental/Iris/Core/Private/Iris/ReplicationSystem/WorldLocations.cpp
+++ b/Engine/Source/Runtime/Experimental/Iris/Core/Private/Iris/ReplicationSystem/WorldLocations.cpp
@@ -77,7 +77,7 @@ void FWorldLocations::UpdateWorldLocation(uint32 ObjectIndex, const FVector& Wor
 
 void FWorldLocations::ResetObjectsWithDirtyInfo()
 {
-	ObjectsWithDirtyInfo.Reset();
+	ObjectsWithDirtyInfo.ClearAllBits();
 }
 
 FVector FWorldLocations::ClampPositionToBoundary(const FVector& Position)
diff --git a/Engine/Source/Runtime/Net/Core/Private/Net/Core/DirtyNetObjectTracker/GlobalDirtyNetObjectTracker.cpp b/Engine/Source/Runtime/Net/Core/Private/Net/Core/DirtyNetObjectTracker/GlobalDirtyNetObjectTracker.cpp
--- a/Engine/Source/Runtime/Net/Core/Private/Net/Core/DirtyNetObjectTracker/GlobalDirtyNetObjectTracker.cpp
+++ b/Engine/Source/Runtime/Net/Core/Private/Net/Core/DirtyNetObjectTracker/GlobalDirtyNetObjectTracker.cpp
@@ -226,7 +226,7 @@ void FGlobalDirtyNetObjectTracker::FPimpl::ResetDirtyNetObjects()
 {
 	if (!ensureMsgf(Pollers.IsNoBitSet(), TEXT("FGlobalDirtyNetObjectTracker poller %u forgot to call ResetDirtyNetObjects."), Pollers.FindFirstOne()))
 	{
-		Pollers.Reset();
+		Pollers.ClearAllBits();
 
 		// DirtyObjects should already be reset if there are no pollers.
 		DirtyObjects.Reset();
diff --git a/Engine/Source/Runtime/Net/Core/Public/Net/Core/NetBitArray.h b/Engine/Source/Runtime/Net/Core/Public/Net/Core/NetBitArray.h
--- a/Engine/Source/Runtime/Net/Core/Public/Net/Core/NetBitArray.h
+++ b/Engine/Source/Runtime/Net/Core/Public/Net/Core/NetBitArray.h
@@ -123,14 +123,52 @@ public:
 	/** Returns the number of words in our storage */
 	uint32 GetNumWords() const;
 
-	/** Returns a pointer to the internal storage. */
-	const StorageWordType* GetData() const { return Storage.GetData(); }
+	/**
+	 * Returns a pointer to the internal storage but validates that the array buffer is the expected size.
+	 * Use this before directly iterating over one or multiple buffers and ensure they all support the max iteration index.
+	 */
+	FORCEINLINE const StorageWordType* GetDataChecked(uint32 MaxWordIterationCount) const UE_LIFETIMEBOUND
+	{
+		checkf((int32)MaxWordIterationCount <= Storage.Num(), TEXT("MaxWordIterationCount (%u) is higher than this NetBitArray word count (%d). Out of Bounds memory access will occur."), MaxWordIterationCount, Storage.Num());
+		return Storage.GetData();
+	}
+
+	/**
+	 * Returns a pointer to the internal storage but validates that the array buffer is the expected size.
+	 * Use this before directly iterating over one or multiple buffers and ensure they all support the max iteration index.
+	 */
+	FORCEINLINE StorageWordType* GetDataChecked(uint32 MaxWordIterationCount) UE_LIFETIMEBOUND
+	{
+		checkf((int32)MaxWordIterationCount <= Storage.Num(), TEXT("MaxWordIterationCount (%u) is higher than this NetBitArray word count (%d). Out of Bounds memory access will occur."), MaxWordIterationCount, Storage.Num());
+		return Storage.GetData();
+	}
+
+	/** Returns a pointer to the internal storage. It's safer to use GetWord() or GetDataChecked() instead. */
+	FORCEINLINE const StorageWordType* GetData() const UE_LIFETIMEBOUND { return Storage.GetData(); }
+
+	/** Returns a pointer to the internal storage. It's safer to use GetWord() or GetDataChecked() instead. */
+	FORCEINLINE StorageWordType* GetData() UE_LIFETIMEBOUND { return Storage.GetData(); }
+
+	/** Read an entire word. Useful when you want to read 'WordBitCount(32)' bits at once */
+	FORCEINLINE StorageWordType GetWord(uint32 WordIndex) const UE_LIFETIMEBOUND 
+	{ 
+		checkf(WordIndex < (uint32)Storage.Num(), TEXT("NetBitArray index out of bounds: %u / max %u"), WordIndex, Storage.Num());
+		return Storage.GetData()[WordIndex]; 
+	}
+
+	/** Access an entire word. Useful when you want to write 'WordBitCount(32)' bits at once */
+	FORCEINLINE StorageWordType& GetWord(uint32 WordIndex) UE_LIFETIMEBOUND 
+	{  
+		checkf(WordIndex < (uint32)Storage.Num(), TEXT("NetBitArray index out of bounds: %u / max %u"), WordIndex, Storage.Num());
+		return Storage.GetData()[WordIndex]; 
+	}
 
-	/** Returns a pointer to the internal storage. */
-	StorageWordType* GetData() { return Storage.GetData(); }
+	/** Clear all bits in the array. */
+	UE_DEPRECATED(5.5, "Use ClearAllBits instead.")
+	void Reset() { ClearAllBits(); }
 
 	/** Clear all bits in the array. */
-	void Reset();
+	void ClearAllBits();
 
 	/** Sets all bits in the array, but clears padding bits. */
 	void SetAllBits();
@@ -303,7 +341,11 @@ public:
 	inline bool IsNoBitSet() const;
 
 	/** Reset the storage of the BitArray including any padding bits */
-	inline void Reset();
+	UE_DEPRECATED(5.5, "Use ClearAllBits instead.")
+	inline void Reset() { ClearAllBits(); }
+
+	/** Clear all bits to zero*/
+	inline void ClearAllBits();
 
 	/** All padding bits will be set to zero */
 	inline void ClearPaddingBits();
@@ -353,11 +395,45 @@ public:
 	/** Returns the number of words in our storage */
 	inline uint32 GetNumWords() const;
 
-	/** Returns a pointer to the internal storage. */
-	const StorageWordType* GetData() const { return Storage; }
+	/**
+	 * Returns a pointer to the internal storage but validates that the array buffer is the expected size.
+	 * Use this before directly iterating over one or multiple buffers and ensure they all support the max iteration index.
+	 */
+	FORCEINLINE const StorageWordType* GetDataChecked(uint32 MaxWordIterationCount) const UE_LIFETIMEBOUND
+	{
+		checkf(MaxWordIterationCount <= WordCount, TEXT("MaxWordIterationCount (%u) is higher than this NetBitArrayView word count (%u). Out of Bounds memory access will occur."), MaxWordIterationCount, WordCount);
+		return Storage;
+	}
+
+	/**
+	 * Returns a pointer to the internal storage but validates that the array buffer is the expected size.
+	 * Use this before directly iterating over one or multiple buffers and ensure they all support the max iteration index.
+	 */
+	FORCEINLINE StorageWordType* GetDataChecked(uint32 MaxWordIterationCount) UE_LIFETIMEBOUND
+	{
+		checkf(MaxWordIterationCount <= WordCount, TEXT("MaxWordIterationCount (%u) is higher than this NetBitArrayView word count (%u). Out of Bounds memory access will occur."), MaxWordIterationCount, WordCount);
+		return Storage;
+	}
+
+	/** Returns a pointer to the internal storage. It's safer to use GetWord() or GetDataChecked() instead. */
+	FORCEINLINE const StorageWordType* GetData() const UE_LIFETIMEBOUND { return Storage; }
+
+	/** Returns a pointer to the internal storage. It's safer to use GetWord() or GetDataChecked() instead. */
+	FORCEINLINE StorageWordType* GetData() UE_LIFETIMEBOUND { return Storage; }
 
-	/** Returns a pointer to the internal storage. */
-	StorageWordType* GetData() { return Storage; }
+	/** Read an entire word. Useful when you want to read 'WordBitCount(32)' bits at once */
+	FORCEINLINE StorageWordType GetWord(uint32 WordIndex) const UE_LIFETIMEBOUND 
+	{ 
+		checkf(WordIndex < WordCount, TEXT("NetBitArrayView index out of bounds: %u / max %u"), WordIndex, WordCount);
+		return Storage[WordIndex]; 
+	}
+
+	/** Access an entire word. Useful when you want to write 'WordBitCount(32)' bits at once */
+	FORCEINLINE StorageWordType& GetWord(uint32 WordIndex) UE_LIFETIMEBOUND 
+	{
+		checkf(WordIndex < WordCount, TEXT("NetBitArrayView index out of bounds: %u / max %u"), WordIndex, WordCount);
+		return Storage[WordIndex]; 
+	}
 
 	/** Only prints the amount of set bits in the array. @see FNetBitArrayPrinter for more print options */
 	NETCORE_API FString ToString() const;
@@ -959,7 +1035,7 @@ inline uint32 FNetBitArray::GetNumWords() const
 	return static_cast<uint32>(Storage.Num());
 }
 
-inline void FNetBitArray::Reset()
+inline void FNetBitArray::ClearAllBits()
 {
 	FPlatformMemory::Memset(Storage.GetData(), 0, Storage.Num()*sizeof(StorageWordType));
 }
@@ -1164,7 +1240,7 @@ FNetBitArrayView::FNetBitArrayView(StorageWordType* StorageIn, uint32 BitCountIn
 FNetBitArrayView::FNetBitArrayView(StorageWordType* StorageIn, uint32 BitCountIn, const EResetOnInitType)
 : FNetBitArrayView(StorageIn, BitCountIn, NoResetNoValidate)
 {
-	Reset();
+	ClearAllBits();
 }
 
 FNetBitArrayView::FNetBitArrayView(StorageWordType* StorageIn, uint32 BitCountIn)
@@ -1199,7 +1275,7 @@ bool FNetBitArrayView::IsNoBitSet() const
 	return !IsAnyBitSet();
 }
 
-void FNetBitArrayView::Reset()
+void FNetBitArrayView::ClearAllBits()
 {
 	FPlatformMemory::Memset(&Storage[0], 0, WordCount * sizeof(StorageWordType));
 }
