Dataflow: cache offset on connections. Use offsets to lookup connections on the nodes.

diff --git a/Engine/Source/Runtime/Experimental/Dataflow/Core/Private/Dataflow/DataflowConnection.cpp b/Engine/Source/Runtime/Experimental/Dataflow/Core/Private/Dataflow/DataflowConnection.cpp
--- a/Engine/Source/Runtime/Experimental/Dataflow/Core/Private/Dataflow/DataflowConnection.cpp
+++ b/Engine/Source/Runtime/Experimental/Dataflow/Core/Private/Dataflow/DataflowConnection.cpp
@@ -13,6 +13,7 @@ FDataflowConnection::FDataflowConnection(Dataflow::FPin::EDirection InDirection,
 	, OwningNode(InOwningNode)
 	, Property(InProperty)
 	, Guid(InGuid)
+	, Offset(OwningNode ? OwningNode->GetPropertyOffset(Name) : INDEX_NONE)
 	, Direction(InDirection)
 {
 	InitFromType();
@@ -24,6 +25,7 @@ FDataflowConnection::FDataflowConnection(Dataflow::FPin::EDirection InDirection,
 	, OwningNode(Params.Owner)
 	, Property(Params.Property)
 	, Guid(Params.Guid)
+	, Offset(Params.Offset)
 	, Direction(InDirection)
 {
 	InitFromType();
@@ -47,16 +49,6 @@ void FDataflowConnection::InitFromType()
 	}
 }
 
-uint32 FDataflowConnection::GetOffset() const
-{
-	if (ensure(OwningNode))
-	{
-		return OwningNode->GetPropertyOffset(Name);
-	}
-	return INDEX_NONE;
-}
-
-
 bool FDataflowConnection::IsOwningNodeEnabled() const
 {
 	return (OwningNode && OwningNode->bActive);
diff --git a/Engine/Source/Runtime/Experimental/Dataflow/Core/Private/Dataflow/DataflowNode.cpp b/Engine/Source/Runtime/Experimental/Dataflow/Core/Private/Dataflow/DataflowNode.cpp
--- a/Engine/Source/Runtime/Experimental/Dataflow/Core/Private/Dataflow/DataflowNode.cpp
+++ b/Engine/Source/Runtime/Experimental/Dataflow/Core/Private/Dataflow/DataflowNode.cpp
@@ -54,15 +54,13 @@ void FDataflowNode::AddInput(FDataflowInput* InPtr)
 			ensureMsgf(!In->GetName().IsEqual(InPtr->GetName()), TEXT("Add Input Failed: Existing Node input already defined with name (%s)"), *InPtr->GetName().ToString());
 		}
 
+		check(InPtr->GetOwningNode() == this);
+
 		const uint32 PropertyOffset = InPtr->GetOffset();
-		if (!Inputs.Contains(PropertyOffset))
+		if (ensure(!Inputs.Contains(PropertyOffset)))
 		{
 			Inputs.Add(PropertyOffset, InPtr);
 		}
-		else
-		{
-			Inputs[PropertyOffset] = InPtr;
-		}
 	}
 }
 
@@ -95,26 +93,22 @@ const FDataflowInput* FDataflowNode::FindInput(FName InName) const
 
 const FDataflowInput* FDataflowNode::FindInput(const void* Reference) const
 {
-	for (const TPair<int32, FDataflowInput*>& Elem : Inputs)
+	const int32 Key = GetConnectionOffsetFromReference(Reference);
+	if (const FDataflowInput* const* Con = Inputs.Find(Key))
 	{
-		const FDataflowInput* const Con = Elem.Value;
-		if (Con->RealAddress() == Reference)
-		{
-			return Con;
-		}
+		check(*Con && (*Con)->RealAddress() == Reference);
+		return *Con;
 	}
 	return nullptr;
 }
 
 FDataflowInput* FDataflowNode::FindInput(void* Reference)
 {
-	for (TPair<int32, FDataflowInput*>& Elem : Inputs)
+	const int32 Key = GetConnectionOffsetFromReference(Reference);
+	if (FDataflowInput* const* Con = Inputs.Find(Key))
 	{
-		FDataflowInput* const Con = Elem.Value;
-		if (Con->RealAddress() == Reference)
-		{
-			return Con;
-		}
+		check(*Con && (*Con)->RealAddress() == Reference);
+		return *Con;
 	}
 	return nullptr;
 }
@@ -190,15 +184,13 @@ void FDataflowNode::AddOutput(FDataflowOutput* InPtr)
 			ensureMsgf(!Out->GetName().IsEqual(InPtr->GetName()), TEXT("Add Output Failed: Existing Node output already defined with name (%s)"), *InPtr->GetName().ToString());
 		}
 
+		check(InPtr->GetOwningNode() == this);
+
 		const uint32 PropertyOffset = InPtr->GetOffset();
-		if (!Outputs.Contains(PropertyOffset))
+		if (ensure(!Outputs.Contains(PropertyOffset)))
 		{
 			Outputs.Add(PropertyOffset, InPtr);
 		}
-		else
-		{
-			Outputs[PropertyOffset] = InPtr;
-		}
 	}
 }
 
@@ -257,26 +249,22 @@ const FDataflowOutput* FDataflowNode::FindOutput(uint32 InGuidHash) const
 
 const FDataflowOutput* FDataflowNode::FindOutput(const void* Reference) const
 {
-	for (const TPair<int32, FDataflowOutput*>& Elem : Outputs)
+	const int32 Key = GetConnectionOffsetFromReference(Reference);
+	if (const FDataflowOutput* const* Con = Outputs.Find(Key))
 	{
-		const FDataflowOutput* const Con = Elem.Value;
-		if (Con->RealAddress() == Reference)
-		{
-			return Con;
-		}
+		check(*Con && (*Con)->RealAddress() == Reference);
+		return *Con;
 	}
 	return nullptr;
 }
 
 FDataflowOutput* FDataflowNode::FindOutput(void* Reference)
 {
-	for (TPair<int32, FDataflowOutput*>& Elem : Outputs)
+	const int32 Key = GetConnectionOffsetFromReference(Reference);
+	if (FDataflowOutput* const* Con = Outputs.Find(Key))
 	{
-		FDataflowOutput* const Con = Elem.Value;
-		if (Con->RealAddress() == Reference)
-		{
-			return Con;
-		}
+		check(*Con && (*Con)->RealAddress() == Reference);
+		return *Con;
 	}
 	return nullptr;
 }
@@ -478,6 +466,11 @@ uint32 FDataflowNode::GetPropertyOffset(const FName& PropertyFullName) const
 	return Offset;
 }
 
+uint32 FDataflowNode::GetConnectionOffsetFromReference(const void* Reference) const
+{
+	return (uint32)((size_t)Reference - (size_t)this);
+}
+
 FString FDataflowNode::GetPropertyFullNameString(const TConstArrayView<const FProperty*>& PropertyChain)
 {
 	FString PropertyFullName;
@@ -528,6 +521,8 @@ bool FDataflowNode::InitConnectionParametersFromPropertyReference(const FStructO
 				OutParams.Name = GetPropertyFullName(PropertyChain);
 				OutParams.Property = Property;
 				OutParams.Owner = this;
+				OutParams.Offset = GetConnectionOffsetFromReference(PropertyRef);
+				check(OutParams.Offset == GetPropertyOffset(PropertyChain));
 				return true;
 			}
 		}
@@ -542,7 +537,9 @@ FDataflowInput* FDataflowNode::RegisterInputConnectionInternal(const void* InPro
 		if (InitConnectionParametersFromPropertyReference(*ScriptOnStruct, InProperty, PropertyName, InputParams))
 		{
 			FDataflowInput* const Input = new FDataflowInput(InputParams);
+			check(Input->RealAddress() == InProperty);
 			AddInput(Input);
+			check(FindInput(InProperty) == Input);
 			return Input;
 		}
 	}
@@ -582,15 +579,19 @@ FDataflowOutput* FDataflowNode::RegisterOutputConnectionInternal(const void* InP
 		if (InitConnectionParametersFromPropertyReference(*ScriptOnStruct, InProperty, PropertyName, OutputParams))
 		{
 			FDataflowOutput* OutputConnection = new FDataflowOutput(OutputParams);
+			check(OutputConnection->RealAddress() == InProperty);
 
 			TArray<const FProperty*> PassthroughPropertyChain;
 			if (FindProperty(ScriptOnStruct->GetStruct(), Passthrough, PassthroughName, &PassthroughPropertyChain))
 			{
-				const uint32 PassthroughOffset = GetPropertyOffset(PassthroughPropertyChain);
+				const uint32 PassthroughOffset = GetConnectionOffsetFromReference(Passthrough);
+				check(PassthroughOffset == GetPropertyOffset(PassthroughPropertyChain));
 				OutputConnection->SetPassthroughOffset(PassthroughOffset);
 			}
 
 			AddOutput(OutputConnection);
+			check(FindOutput(InProperty) == OutputConnection);
+			return OutputConnection;
 		}
 	}
 	return nullptr;
diff --git a/Engine/Source/Runtime/Experimental/Dataflow/Core/Public/Dataflow/DataflowConnection.h b/Engine/Source/Runtime/Experimental/Dataflow/Core/Public/Dataflow/DataflowConnection.h
--- a/Engine/Source/Runtime/Experimental/Dataflow/Core/Public/Dataflow/DataflowConnection.h
+++ b/Engine/Source/Runtime/Experimental/Dataflow/Core/Public/Dataflow/DataflowConnection.h
@@ -37,11 +37,12 @@ namespace Dataflow
 
 	struct FConnectionParameters
 	{
-		FConnectionParameters(FName InType = FName(""), FName InName = FName(""), FDataflowNode* InOwner = nullptr, const FProperty* InProperty = nullptr, FGuid InGuid = FGuid::NewGuid())
+		FConnectionParameters(FName InType = FName(""), FName InName = FName(""), FDataflowNode* InOwner = nullptr, const FProperty* InProperty = nullptr, uint32 InOffset = INDEX_NONE, FGuid InGuid = FGuid::NewGuid())
 			: Type(InType)
 			, Name(InName)
 			, Owner(InOwner)
 			, Property(InProperty)
+			, Offset(InOffset)
 			, Guid(InGuid)
 		{}
 
@@ -49,6 +50,7 @@ namespace Dataflow
 		FName Name;
 		FDataflowNode* Owner = nullptr;
 		const FProperty* Property = nullptr;
+		uint32 Offset = INDEX_NONE;
 		FGuid Guid;
 	};
 
@@ -70,6 +72,7 @@ protected:
 	const FProperty* Property = nullptr;
 	FGuid  Guid;
 	IDataflowTypePolicy* TypePolicy = nullptr;
+	uint32 Offset;
 	Dataflow::FPin::EDirection Direction;
 	bool bIsAnyType:1 = false;
 	bool bHasConcreteType : 1 = false;
@@ -112,7 +115,7 @@ public:
 	const FProperty* GetProperty() const { return Property; }
 
 	Dataflow::FPin::EDirection GetDirection() const { return Direction; }
-	DATAFLOWCORE_API uint32 GetOffset() const;
+	uint32 GetOffset() const { return Offset; }
 
 	FName GetType() const { return Type; }
 
diff --git a/Engine/Source/Runtime/Experimental/Dataflow/Core/Public/Dataflow/DataflowInputOutput.h b/Engine/Source/Runtime/Experimental/Dataflow/Core/Public/Dataflow/DataflowInputOutput.h
--- a/Engine/Source/Runtime/Experimental/Dataflow/Core/Public/Dataflow/DataflowInputOutput.h
+++ b/Engine/Source/Runtime/Experimental/Dataflow/Core/Public/Dataflow/DataflowInputOutput.h
@@ -21,8 +21,8 @@ namespace Dataflow
 {
 	struct FInputParameters : public FConnectionParameters
 	{
-		FInputParameters(FName InType = NAME_None, FName InName = NAME_None, FDataflowNode* InOwner = nullptr, const FProperty* InProperty = nullptr, FGuid InGuid = FGuid::NewGuid())
-			: FConnectionParameters(InType, InName, InOwner, InProperty, InGuid)
+		FInputParameters(FName InType = NAME_None, FName InName = NAME_None, FDataflowNode* InOwner = nullptr, const FProperty* InProperty = nullptr, uint32 InOffset = INDEX_NONE, FGuid InGuid = FGuid::NewGuid())
+			: FConnectionParameters(InType, InName, InOwner, InProperty, InOffset, InGuid)
 		{}
 	};
 }
@@ -87,8 +87,8 @@ namespace Dataflow
 {
 	struct FOutputParameters: public FConnectionParameters
 	{
-		FOutputParameters(FName InType = NAME_None, FName InName = NAME_None, FDataflowNode* InOwner = nullptr, const FProperty* InProperty = nullptr, FGuid InGuid = FGuid::NewGuid())
-			: FConnectionParameters(InType, InName, InOwner, InProperty, InGuid)
+		FOutputParameters(FName InType = NAME_None, FName InName = NAME_None, FDataflowNode* InOwner = nullptr, const FProperty* InProperty = nullptr, uint32 InOffset = INDEX_NONE, FGuid InGuid = FGuid::NewGuid())
+			: FConnectionParameters(InType, InName, InOwner, InProperty, InOffset, InGuid)
 		{}
 	};
 }
diff --git a/Engine/Source/Runtime/Experimental/Dataflow/Core/Public/Dataflow/DataflowNode.h b/Engine/Source/Runtime/Experimental/Dataflow/Core/Public/Dataflow/DataflowNode.h
--- a/Engine/Source/Runtime/Experimental/Dataflow/Core/Public/Dataflow/DataflowNode.h
+++ b/Engine/Source/Runtime/Experimental/Dataflow/Core/Public/Dataflow/DataflowNode.h
@@ -393,6 +393,7 @@ private:
 	static FName GetPropertyFullName(const TArray<const FProperty*>& PropertyChain);
 	static FText GetPropertyDisplayNameText(const TArray<const FProperty*>& PropertyChain);
 	static uint32 GetPropertyOffset(const TArray<const FProperty*>& PropertyChain);
+	uint32 GetConnectionOffsetFromReference(const void* Reference) const;
 
 	/**
 	* Find a property using the property address and name (not including its parent struct property names).
