Fix NetBitArray::GetData lifetimebound errors on clang

#rnx

diff --git a/Engine/Source/Runtime/Experimental/Iris/Core/Private/Iris/ReplicationSystem/Filtering/ReplicationFiltering.cpp b/Engine/Source/Runtime/Experimental/Iris/Core/Private/Iris/ReplicationSystem/Filtering/ReplicationFiltering.cpp
--- a/Engine/Source/Runtime/Experimental/Iris/Core/Private/Iris/ReplicationSystem/Filtering/ReplicationFiltering.cpp
+++ b/Engine/Source/Runtime/Experimental/Iris/Core/Private/Iris/ReplicationSystem/Filtering/ReplicationFiltering.cpp
@@ -395,7 +395,9 @@ void FReplicationFiltering::BuildAlwaysRelevantList(FNetBitArrayView OutAlwaysRe
 	const uint32* const WithOwnerData = ObjectsWithOwnerFilter.GetDataChecked(MaxWords);
 	const uint32* const ConnectionFiltersData = AllConnectionFilteredObjects.GetDataChecked(MaxWords);
 	const uint32* const DynamicFilteredData = DynamicFilterEnabledObjects.GetDataChecked(MaxWords);
-	const uint32* const GroupFilteredOutData = Groups->GetGroupFilteredOutObjects().GetDataChecked(MaxWords);
+
+	const FNetBitArrayView GroupFilteredOutView = Groups->GetGroupFilteredOutObjects();
+	const uint32* const GroupFilteredOutData = GroupFilteredOutView.GetDataChecked(MaxWords);
 
 	uint32* OutAlwaysRelevantListData = OutAlwaysRelevantList.GetDataChecked(MaxWords);
 
@@ -1337,9 +1339,12 @@ void FReplicationFiltering::UpdateDynamicFiltering()
 	// Subobjects will never be added to a dynamic filter, but objects can become dependent at any time.
 	// We need to make sure they are not filtered out.
 	const uint32* SubObjectsData = NetRefHandleManager->GetSubObjectInternalIndices().GetDataChecked(WordCountForObjectBitArrays);
-	const uint32* DependentObjectsData = NetRefHandleManager->GetDependentObjectInternalIndices().GetDataChecked(WordCountForObjectBitArrays);
 	const uint32* ObjectsRequiringDynamicFilterUpdateData = ObjectsRequiringDynamicFilterUpdate.GetDataChecked(WordCountForObjectBitArrays);
 
+	const FNetBitArrayView DependentObjectsView = NetRefHandleManager->GetDependentObjectInternalIndices();
+	const uint32* DependentObjectsData = DependentObjectsView.GetDataChecked(WordCountForObjectBitArrays);
+	
+
 	uint32* ConnectionIds = static_cast<uint32*>(FMemory_Alloca(ValidConnections.GetNumBits() * sizeof(uint32)));
 	uint32 ConnectionCount = 0;
 	ValidConnections.ForAllSetBits([ConnectionIds, &ConnectionCount](uint32 Bit) { ConnectionIds[ConnectionCount++] = Bit; });
@@ -1377,7 +1382,8 @@ void FReplicationFiltering::UpdateDynamicFiltering()
 				// Execute the filter here
 				Info.Filter->Filter(FilteringParams);
 
-				const uint32* FilteredObjectsData = Info.Filter->GetFilteredObjects().GetData();
+				const FNetBitArrayView FilteredObjectsView = Info.Filter->GetFilteredObjects();
+				const uint32* FilteredObjectsData = FilteredObjectsView.GetData();
 				for (SIZE_T WordIt = 0, WordEndIt = WordCountForObjectBitArrays; WordIt != WordEndIt; ++WordIt)
 				{
 					const uint32 FilteredObjects = FilteredObjectsData[WordIt];
