IOStore on demand package tag set support

diff --git a/Engine/Source/Runtime/Experimental/IoStore/OnDemand/Private/IoStoreOnDemand.cpp b/Engine/Source/Runtime/Experimental/IoStore/OnDemand/Private/IoStoreOnDemand.cpp
--- a/Engine/Source/Runtime/Experimental/IoStore/OnDemand/Private/IoStoreOnDemand.cpp
+++ b/Engine/Source/Runtime/Experimental/IoStore/OnDemand/Private/IoStoreOnDemand.cpp
@@ -8,6 +8,7 @@
 #include "HAL/PlatformMisc.h"
 #include "HAL/PlatformTime.h"
 #include "IO/IoChunkEncoding.h"
+#include "IO/PackageId.h"
 #include "IasCache.h"
 #include "Misc/Base64.h"
 #include "Misc/CommandLine.h"
@@ -655,6 +656,103 @@ bool LoadFromCompactBinary(FCbFieldView Field, FOnDemandTocAdditionalFile& Addit
 	return false;
 }
 
+FArchive& operator<<(FArchive& Ar, FOnDemandTocTagSetPackageList& TagSet)
+{
+	Ar << TagSet.ContainerIndex;
+	Ar << TagSet.PackageIndicies;
+	return Ar;
+}
+
+FCbWriter& operator<<(FCbWriter& Writer, const FOnDemandTocTagSetPackageList& TagSet)
+{
+	Writer.BeginObject();
+	Writer.AddInteger(UTF8TEXTVIEW("ContainerIndex"), TagSet.ContainerIndex);
+	Writer.BeginArray(UTF8TEXTVIEW("PackageIndicies"));
+	for (const uint32 Index : TagSet.PackageIndicies)
+	{
+		Writer << Index;
+	}
+	Writer.EndArray();
+	Writer.EndObject();
+	return Writer;
+}
+
+bool LoadFromCompactBinary(FCbFieldView Field, FOnDemandTocTagSetPackageList& TagSet)
+{
+	if (FCbObjectView Obj = Field.AsObjectView())
+	{
+		FCbFieldView ContainerIndex = Obj["ContainerIndex"];
+		TagSet.ContainerIndex = ContainerIndex.AsUInt32();
+		if (ContainerIndex.HasError())
+		{
+			return false;
+		}
+
+		FCbFieldView PackageIndicies = Obj["PackageIndicies"];
+		FCbArrayView PackageIndiciesArray = PackageIndicies.AsArrayView();
+		if(PackageIndicies.HasError())
+		{
+			return false;
+		}
+
+		TagSet.PackageIndicies.Reserve(int32(PackageIndiciesArray.Num()));
+		for (FCbFieldView ArrayField : PackageIndiciesArray)
+		{
+			uint32 Index = ArrayField.AsUInt32();
+			if (ArrayField.HasError())
+			{
+				return false;
+			}
+			TagSet.PackageIndicies.Emplace(Index);
+		}
+
+		return true;
+	}
+
+	return false;
+}
+
+FArchive& operator<<(FArchive& Ar, FOnDemandTocTagSet& TagSet)
+{
+	Ar << TagSet.Tag;
+	Ar << TagSet.Packages;
+	return Ar;
+}
+
+FCbWriter& operator<<(FCbWriter& Writer, const FOnDemandTocTagSet& TagSet)
+{
+	Writer.BeginObject();
+	Writer.AddString(UTF8TEXTVIEW("Tag"), TagSet.Tag);
+	Writer.BeginArray(UTF8TEXTVIEW("Packages"));
+	for (const FOnDemandTocTagSetPackageList& PackageList : TagSet.Packages)
+	{
+		Writer << PackageList;
+	}
+	Writer.EndArray();
+	Writer.EndObject();
+	return Writer;
+}
+
+bool LoadFromCompactBinary(FCbFieldView Field, FOnDemandTocTagSet& TagSet)
+{
+	if (FCbObjectView Obj = Field.AsObjectView())
+	{
+		TagSet.Tag = FString(Obj["Tag"].AsString());
+		FCbArrayView Packages = Obj["Packages"].AsArrayView();
+		TagSet.Packages.Reserve(int32(Packages.Num()));
+		for (FCbFieldView ArrayField : Packages)
+		{
+			if (!LoadFromCompactBinary(ArrayField, TagSet.Packages.Emplace_GetRef()))
+			{
+				return false;
+			}
+		}
+		return true;
+	}
+
+	return false;
+}
+
 FArchive& operator<<(FArchive& Ar, FOnDemandToc& Toc)
 {
 	Ar << Toc.Header;
@@ -676,6 +774,11 @@ FArchive& operator<<(FArchive& Ar, FOnDemandToc& Toc)
 		Ar << Toc.AdditionalFiles;
 	}
 
+	if (uint32(Toc.Header.Version) >= uint32(EOnDemandTocVersion::TagSets))
+	{
+		Ar << Toc.TagSets;
+	}
+
 	return Ar;
 }
 
@@ -701,6 +804,16 @@ FCbWriter& operator<<(FCbWriter& Writer, const FOnDemandToc& Toc)
 		Writer.EndArray();
 	}
 
+	if (Toc.TagSets.Num() > 0)
+	{
+		Writer.BeginArray(UTF8TEXTVIEW("TagSets"));
+		for (const FOnDemandTocTagSet& TagSet : Toc.TagSets)
+		{
+			Writer << TagSet;
+		}
+		Writer.EndArray();
+	}
+
 	Writer.EndObject();
 	
 	return Writer;
@@ -748,6 +861,19 @@ bool LoadFromCompactBinary(FCbFieldView Field, FOnDemandToc& OutToc)
 			}
 		}
 
+		if (uint32(OutToc.Header.Version) >= uint32(EOnDemandTocVersion::TagSets))
+		{
+			FCbArrayView TagSets = Obj["TagSets"].AsArrayView();
+			OutToc.TagSets.Reserve(int32(TagSets.Num()));
+			for (FCbFieldView ArrayField : TagSets)
+			{
+				if (!LoadFromCompactBinary(ArrayField, OutToc.TagSets.AddDefaulted_GetRef()))
+				{
+					return false;
+				}
+			}
+		}
+
 		return true;
 	}
 
diff --git a/Engine/Source/Runtime/Experimental/IoStore/OnDemand/Private/OnDemandIoStore.cpp b/Engine/Source/Runtime/Experimental/IoStore/OnDemand/Private/OnDemandIoStore.cpp
--- a/Engine/Source/Runtime/Experimental/IoStore/OnDemand/Private/OnDemandIoStore.cpp
+++ b/Engine/Source/Runtime/Experimental/IoStore/OnDemand/Private/OnDemandIoStore.cpp
@@ -110,9 +110,14 @@ FString LexToString(EOnDemandContainerFlags Flags)
 const FOnDemandChunkEntry FOnDemandChunkEntry::Null = {};
 
 ///////////////////////////////////////////////////////////////////////////////
+static FString OnDemandContainerUniqueName(FStringView MountId, FStringView Name)
+{
+	return FString::Printf(TEXT("%.*s-%.*s"), MountId.Len(), MountId.GetData(), Name.Len(), Name.GetData());
+}
+
 FString FOnDemandContainer::UniqueName() const
 {
-	return FString::Printf(TEXT("%s-%s"), *MountId, *Name);
+	return OnDemandContainerUniqueName(MountId, Name);
 }
 
 ///////////////////////////////////////////////////////////////////////////////
@@ -504,7 +509,9 @@ FIoStatus FOnDemandIoStore::TickMountRequest(FMountRequest& MountRequest)
 		}
 
 		FOnDemandToc Toc = TocStatus.ConsumeValueOrDie();
-		CreateContainersFromToc(Args.MountId, TocPath, Toc, MountRequest.Containers);
+		Args.Toc.Emplace(MoveTemp(Toc));
+
+		CreateContainersFromToc(Args.MountId, TocPath, *Args.Toc, MountRequest.Containers);
 	}
 	else if (Args.Url.IsEmpty() == false)
 	{
@@ -520,7 +527,9 @@ FIoStatus FOnDemandIoStore::TickMountRequest(FMountRequest& MountRequest)
 		}
 
 		FOnDemandToc Toc = TocStatus.ConsumeValueOrDie();
-		CreateContainersFromToc(Args.MountId, TocPath, Toc, MountRequest.Containers);
+		Args.Toc.Emplace(MoveTemp(Toc));
+
+		CreateContainersFromToc(Args.MountId, TocPath, *Args.Toc, MountRequest.Containers);
 	}
 
 	// Remove already mounted containers
@@ -567,7 +576,36 @@ FIoStatus FOnDemandIoStore::TickMountRequest(FMountRequest& MountRequest)
 			return FIoStatusBuilder(EIoErrorCode::InvalidCode) << TEXT("Install cache not configured");
 		}
 
-		if (FIoStatus Status = InstallContainers(Args.Url, MountRequest.Containers); !Status.IsOk())
+		const bool bUsePackageFilter = Args.TagSets.Num() > 0;
+
+		// List of unique container name, package index list pairs
+		FPackageFilter PackageFilter;
+
+		if (bUsePackageFilter && ensure(Args.Toc))
+		{
+			const FOnDemandToc& Toc = Args.Toc.GetValue();
+			for (const FString& TagSet : Args.TagSets)
+			{
+				const FOnDemandTocTagSet* MaybeTagSet = Algo::FindBy(Toc.TagSets, TagSet, &FOnDemandTocTagSet::Tag);
+				if (MaybeTagSet)
+				{
+					for (const FOnDemandTocTagSetPackageList& PkgList : MaybeTagSet->Packages)
+					{
+						if (!Toc.Containers.IsValidIndex(PkgList.ContainerIndex))
+						{
+							return FIoStatusBuilder(EIoErrorCode::CorruptToc) << TEXT("Bad container index in tag set");
+						}
+
+						FString ContainerName = OnDemandContainerUniqueName(Args.MountId, Toc.Containers[PkgList.ContainerIndex].ContainerName);
+						PackageFilter.Emplace(MoveTemp(ContainerName), MakeArrayView(PkgList.PackageIndicies));
+					}
+				}
+			}
+		}
+
+		FIoStatus Status = InstallContainers(
+			Args.Url, MountRequest.Containers, bUsePackageFilter ? &PackageFilter : nullptr);
+		if (!Status.IsOk())
 		{
 			return Status;
 		}
@@ -667,7 +705,8 @@ void FOnDemandIoStore::CreateContainersFromToc(
 
 FIoStatus FOnDemandIoStore::InstallContainers(
 	const FString& Url,
-	const TConstArrayView<FSharedOnDemandContainer>& ContainersToInstall)
+	const TConstArrayView<FSharedOnDemandContainer>& ContainersToInstall,
+	const FOnDemandIoStore::FPackageFilter* PackageFilter /*= nullptr*/)
 {
 	FStringView Host, TocRelUrl;
 	Private::SplitHostUrl(Url, Host, TocRelUrl);
@@ -739,7 +778,36 @@ FIoStatus FOnDemandIoStore::InstallContainers(
 				reinterpret_cast<const FFilePackageStoreEntry*>(ContainerHeader.StoreEntries.GetData()),
 				ContainerHeader.PackageIds.Num());
 
-			//TODO: Find all referenced chunks
+			//TODO: Find all referenced chunks, do not apply filter recursively, it should only be for exports
+			
+			if (!PackageFilter)
+			{
+				for (int32 Index = 0; Index < ContainerHeader.PackageIds.Num(); ++Index)
+				{
+					const FPackageId& PackageId = ContainerHeader.PackageIds[Index];
+					const FFilePackageStoreEntry& FilePackageStoreEntry = PackageStoreEntries[Index];
+				}
+			}
+			else
+			{
+				const FString ContainerName = Container->UniqueName();
+				for (const TPair<FString, TConstArrayView<uint32>>& Pair : *PackageFilter)
+				{
+					if (Pair.Key == ContainerName)
+					{
+						for (const uint32 Index : Pair.Value)
+						{
+							if (!ContainerHeader.PackageIds.IsValidIndex(Index))
+							{
+								return FIoStatusBuilder(EIoErrorCode::CorruptToc) << TEXT("Bad pkg index in tagset");
+							}
+
+							const FPackageId& PackageId = ContainerHeader.PackageIds[Index];
+							const FFilePackageStoreEntry& FilePackageStoreEntry = PackageStoreEntries[Index];
+						}
+					}
+				}
+			}
 		}
 
 		// Download all chunks
diff --git a/Engine/Source/Runtime/Experimental/IoStore/OnDemand/Private/OnDemandIoStore.h b/Engine/Source/Runtime/Experimental/IoStore/OnDemand/Private/OnDemandIoStore.h
--- a/Engine/Source/Runtime/Experimental/IoStore/OnDemand/Private/OnDemandIoStore.h
+++ b/Engine/Source/Runtime/Experimental/IoStore/OnDemand/Private/OnDemandIoStore.h
@@ -129,6 +129,8 @@ class FOnDemandIoStore
 
 	using FSharedMountRequest	= TSharedPtr<FMountRequest>;
 	using FMountRequestMap		= TMap<FString, FSharedMountRequest>;
+	// List of (unique container name, package index list) pairs
+	using FPackageFilter		= TArray<TPair<FString, TConstArrayView<uint32>>>;
 
 public:
 	FOnDemandIoStore();
@@ -159,7 +161,8 @@ private:
 								TArray<FSharedOnDemandContainer>& Out);
 	FIoStatus				InstallContainers(
 								const FString& Url,
-								const TConstArrayView<FSharedOnDemandContainer>& ContainersToInstall);
+								const TConstArrayView<FSharedOnDemandContainer>& ContainersToInstall,
+								const FPackageFilter* PackageFilter = nullptr);
 
 	TArray<FSharedOnDemandContainer> GetMountedContainers();
 
diff --git a/Engine/Source/Runtime/Experimental/IoStore/OnDemand/Private/Tool/ChunkPlugin.cpp b/Engine/Source/Runtime/Experimental/IoStore/OnDemand/Private/Tool/ChunkPlugin.cpp
--- a/Engine/Source/Runtime/Experimental/IoStore/OnDemand/Private/Tool/ChunkPlugin.cpp
+++ b/Engine/Source/Runtime/Experimental/IoStore/OnDemand/Private/Tool/ChunkPlugin.cpp
@@ -11,18 +11,28 @@
 #include "IO/IoStatus.h"
 #include "IO/IoStore.h"
 #include "IO/IoStoreOnDemand.h"
+#include "IO/IoContainerHeader.h"
 #include "Misc/DateTime.h"
 #include "Misc/FileHelper.h"
 #include "Misc/Paths.h"
 #include "Misc/PathViews.h"
 #include "Misc/StringBuilder.h"
+#include "Serialization/JsonReader.h"
+#include "Serialization/JsonSerializer.h"
 #include "Serialization/LargeMemoryWriter.h"
+#include "Serialization/MemoryReader.h"
 
 namespace UE::IoStore::Tool
 {
 
 ////////////////////////////////////////////////////////////////////////////////
-FIoStatus WriteChunk(const FString& Directory, FMemoryView Chunk, const FIoHash& Hash)
+struct FChunkPluginSettings
+{
+	TMap<FString, TArray<FString>> PackageSets;
+};
+
+////////////////////////////////////////////////////////////////////////////////
+static FIoStatus WriteChunk(const FString& Directory, FMemoryView Chunk, const FIoHash& Hash)
 {
 	IFileManager& FileMgr = IFileManager::Get();
 	const FString HashString = LexToString(Hash);
@@ -57,7 +67,7 @@ FIoStatus WriteChunk(const FString& Directory, FMemoryView Chunk, const FIoHash&
 }
 
 ////////////////////////////////////////////////////////////////////////////////
-FIoStatus WriteChunk(const FString& Directory, FMemoryView Chunk)
+static FIoStatus WriteChunk(const FString& Directory, FMemoryView Chunk)
 {
 	return WriteChunk(Directory, Chunk, FIoHash::HashBuffer(Chunk));
 }
@@ -65,12 +75,16 @@ FIoStatus WriteChunk(const FString& Directory, FMemoryView Chunk)
 ////////////////////////////////////////////////////////////////////////////////
 static int32 ChunkPluginCommandEntry(const FContext& Context)
 {
+	//while (!FPlatformMisc::IsDebuggerPresent());
+	//UE_DEBUG_BREAK();
+
 	const FString Platform				= FString(Context.Get<FStringView>(TEXT("-Platform"), FString()));
 	const FString BuildVersion			= FString(Context.Get<FStringView>(TEXT("-BuildVersion"), FString()));
 	const FString OnDemandTocName		= FString(Context.Get<FStringView>(TEXT("-OnDemandTocName"), FString()));
 	const FString InputFolder			= FString(Context.Get<FStringView>(TEXT("-InputFolder"), FString()));
 	const FString OutputFolder			= FString(Context.Get<FStringView>(TEXT("-OutputFolder"), FString()));
 	const FString IntermediateFolder	= FString(Context.Get<FStringView>(TEXT("-IntermediateFolder"), FString()));
+	FString SettingsFile				= FString(Context.Get<FStringView>(TEXT("-SettingsFile"), FString()));
 	const bool bIncludeSigPak			= Context.Get<bool>(TEXT("-IncludeSigPak"), false);
 	const bool bDeleteContainerFiles	= !Context.Get<bool>(TEXT("-KeepContainerFiles"), false);
 	FString ContainerFolder				= InputFolder;
@@ -80,6 +94,7 @@ static int32 ChunkPluginCommandEntry(const FContext& Context)
 	FPaths::NormalizeDirectoryName(ContainerFolder);
 	FPaths::NormalizeDirectoryName(IoStoreOutputFolder);
 	FPaths::NormalizeDirectoryName(ChunksOutputFolder);
+	FPaths::NormalizeFilename(SettingsFile);
 
 	UE_LOG(LogIoStore, Display, TEXT("I/O store chunk plugin:"));
 	UE_LOG(LogIoStore, Display, TEXT("----------------------------------------"));
@@ -89,10 +104,68 @@ static int32 ChunkPluginCommandEntry(const FContext& Context)
 	UE_LOG(LogIoStore, Display, TEXT("\tInputFolder: %s"), *InputFolder);
 	UE_LOG(LogIoStore, Display, TEXT("\tOutputFolder: %s"), *OutputFolder);
 	UE_LOG(LogIoStore, Display, TEXT("\tIntermediateFolder: %s"), *IntermediateFolder);
+	UE_LOG(LogIoStore, Display, TEXT("\tSettingsFile: %s"), *SettingsFile);
 	UE_LOG(LogIoStore, Display, TEXT("\tIncludeSigPak: %s"), bIncludeSigPak ? TEXT("true") : TEXT("false"));
 	UE_LOG(LogIoStore, Display, TEXT("\tDeleteContainerFiles: %s"), bDeleteContainerFiles ? TEXT("true") : TEXT("false"));
 
 	IFileManager& FileMgr = IFileManager::Get();
+
+	FChunkPluginSettings Settings;
+	if (!SettingsFile.IsEmpty())
+	{
+		TUniquePtr<FArchive> Ar(FileMgr.CreateFileReader(*SettingsFile));
+		if (!Ar || Ar->IsError())
+		{
+			UE_LOG(LogIoStore, Error, TEXT("Failed to open settings file '%s'"), *SettingsFile);
+			return -1;
+		}
+
+		TSharedRef<TJsonReader<UTF8CHAR>> JsonReader = TJsonReaderFactory<UTF8CHAR>::Create(Ar.Get());
+
+		TSharedPtr<FJsonValue> JsonSettings;
+		if (!FJsonSerializer::Deserialize(*JsonReader, JsonSettings))
+		{
+			UE_LOG(LogIoStore, Error, TEXT("Failed to read settings file '%s'"), *SettingsFile);
+			return -1;
+		}
+
+		TSharedPtr<FJsonObject> JsonSettingsObject = JsonSettings->AsObject();
+		if (!JsonSettingsObject)
+		{
+			UE_LOG(LogIoStore, Error, TEXT("Bad settings file '%s'"), *SettingsFile);
+			return -1;
+		}
+
+		if (TSharedPtr<FJsonValue> PackageSets = JsonSettingsObject->TryGetField(TEXT("PackageSets")))
+		{
+			TSharedPtr<FJsonObject> PackageSetsObject = PackageSets->AsObject();
+			if (!PackageSetsObject)
+			{
+				UE_LOG(LogIoStore, Error, TEXT("Bad settings file '%s'"), *SettingsFile);
+				return -1;
+			}
+
+			for (const TPair<FString, TSharedPtr<FJsonValue>>& Pair : PackageSetsObject->Values)
+			{
+				TArray<FString> Packages;
+				if (!PackageSetsObject->TryGetStringArrayField(Pair.Key, Packages))
+				{
+					UE_LOG(LogIoStore, Error, TEXT("Bad settings file '%s'"), *SettingsFile);
+					return -1;
+				}
+
+				UE_LOG(LogIoStore, Display, TEXT("Found Package Set '%s'"), *Pair.Key);
+				for (const FString& Package : Packages)
+				{
+					UE_LOG(LogIoStore, Display, TEXT("-> \t'%s'"), *Package);
+				}
+
+				Settings.PackageSets.Add(Pair.Key, MoveTemp(Packages));
+			}
+		}
+	}
+
+	
 	if (FileMgr.MakeDirectory(*ChunksOutputFolder, true) == false)
 	{
 		UE_LOG(LogIoStore, Error, TEXT("Failed to create directory '%s'"), *ChunksOutputFolder);
@@ -114,6 +187,9 @@ static int32 ChunkPluginCommandEntry(const FContext& Context)
 	//OnDemandToc.Header.ChunksDirectory = TODO 
 	OnDemandToc.Containers.Reserve(ContainerFilenames.Num());
 
+	TMap<FString, TArray<FOnDemandTocTagSetPackageList>> FoundTagSets;
+	FoundTagSets.Reserve(Settings.PackageSets.Num());
+
 	TArray<FString> FilesToDelete;
 
 	for (const FString& Filename : ContainerFilenames)
@@ -131,6 +207,13 @@ static int32 ChunkPluginCommandEntry(const FContext& Context)
 
 		UE_LOG(LogIoStore, Display, TEXT("Processing container '%s'"), *FullPath);
 
+		/*TArray<FString> ContainerFileList;
+		ContainerFileReader.GetFilenames(ContainerFileList);
+		for (const FString& File : ContainerFileList)
+		{
+			UE_LOG(LogIoStore, Display, TEXT("Found file in container '%s'"), *File);
+		}*/
+
 		const uint32 BlockSize = ContainerFileReader.GetCompressionBlockSize();
 		if (OnDemandToc.Header.BlockSize == 0)
 		{
@@ -138,6 +221,71 @@ static int32 ChunkPluginCommandEntry(const FContext& Context)
 		}
 		check(OnDemandToc.Header.BlockSize == ContainerFileReader.GetCompressionBlockSize());
 
+		const int32 ContainerIndex = OnDemandToc.Containers.AddDefaulted();
+		FOnDemandTocContainerEntry& ContainerEntry = OnDemandToc.Containers[ContainerIndex];
+		ContainerEntry.ContainerId = ContainerFileReader.GetContainerId();
+		ContainerEntry.ContainerName = FPaths::GetBaseFilename(FullPath);
+
+		if (EnumHasAnyFlags(ContainerFileReader.GetContainerFlags(), EIoContainerFlags::Encrypted))
+		{
+			ContainerEntry.EncryptionKeyGuid = LexToString(ContainerFileReader.GetEncryptionKeyGuid());
+		}
+
+		TArray<FOnDemandTocTagSet> TagSets;
+		if (Settings.PackageSets.Num() > 0)
+		{
+			const FIoChunkId HeaderChunkId = CreateContainerHeaderChunkId(ContainerFileReader.GetContainerId());
+			TIoStatusOr<FIoBuffer> Status = ContainerFileReader.Read(HeaderChunkId, FIoReadOptions());
+			if (!Status.IsOk())
+			{
+				if (Status.Status().GetErrorCode() == EIoErrorCode::NotFound)
+				{
+					UE_LOG(LogIoStore, Display, TEXT("No header chunk for container '%s'"), *FullPath);
+				}
+				else
+				{
+					UE_LOG(LogIoStore, Error, TEXT("Failed to read header chunk for container '%s'"), *FullPath);
+					return -1;
+				}
+			}
+			else
+			{
+				FIoContainerHeader ContainerHeader;
+				FMemoryReaderView Ar(Status.ValueOrDie().GetView());
+				Ar << ContainerHeader;
+				Ar.Close();
+				if (Ar.IsError() || Ar.IsCriticalError())
+				{
+					UE_LOG(LogIoStore, Error, TEXT("Failed to deserialize header chunk for container '%s'"), *FullPath);
+					return -1;
+				}
+
+				for (const TPair<FString, TArray<FString>>& Pair : Settings.PackageSets)
+				{
+					const FString& Tag = Pair.Key;
+
+					FOnDemandTocTagSetPackageList TagSetPackageList;
+					TagSetPackageList.ContainerIndex = ContainerIndex;
+
+					for (const FString& Package : Pair.Value)
+					{
+						const FPackageId PackageId = FPackageId::FromName(FName(Package));
+						int32 PackageIndex = INDEX_NONE;
+						if (ContainerHeader.PackageIds.Find(PackageId, PackageIndex))
+						{
+							UE_LOG(LogIoStore, Display, TEXT("Found package %s for tag %s in container %s"), *Package, *Tag, *Filename)
+							TagSetPackageList.PackageIndicies.Add(PackageIndex);
+						}
+					}
+
+					if (TagSetPackageList.PackageIndicies.Num() > 0)
+					{
+						FoundTagSets.FindOrAdd(Tag).Add(MoveTemp(TagSetPackageList));
+					}
+				}
+			}
+		}
+
 		TArray<FIoStoreTocChunkInfo> ChunkInfos;
 		ContainerFileReader.EnumerateChunks([&ChunkInfos](FIoStoreTocChunkInfo&& Info)
 		{ 
@@ -146,14 +294,6 @@ static int32 ChunkPluginCommandEntry(const FContext& Context)
 		});
 
 		UE_LOG(LogIoStore, Display, TEXT("Serializing %d chunks"), ChunkInfos.Num());
-		FOnDemandTocContainerEntry& ContainerEntry = OnDemandToc.Containers.AddDefaulted_GetRef();
-		ContainerEntry.ContainerId		= ContainerFileReader.GetContainerId();
-		ContainerEntry.ContainerName	= FPaths::GetBaseFilename(FullPath);
-
-		if (EnumHasAnyFlags(ContainerFileReader.GetContainerFlags(), EIoContainerFlags::Encrypted))
-		{
-			ContainerEntry.EncryptionKeyGuid = LexToString(ContainerFileReader.GetEncryptionKeyGuid());
-		}
 
 		for (const FIoStoreTocChunkInfo& ChunkInfo : ChunkInfos)
 		{
@@ -225,6 +365,14 @@ static int32 ChunkPluginCommandEntry(const FContext& Context)
 		}
 	}
 
+	OnDemandToc.TagSets.Reserve(FoundTagSets.Num());
+	for (TPair<FString, TArray<FOnDemandTocTagSetPackageList>>& Pair : FoundTagSets)
+	{
+		FOnDemandTocTagSet& TocTagSet = OnDemandToc.TagSets.Emplace_GetRef();
+		TocTagSet.Tag = Pair.Key;
+		TocTagSet.Packages = MoveTemp(Pair.Value);
+	}
+
 	IFileManager& FileMan = IFileManager::Get();
 	for (const FString& Path : FilesToDelete)
 	{
@@ -383,6 +531,7 @@ static FCommand ChunkPluginCommand(
 		TArgument<FStringView>(TEXT("-InputFolder"),		TEXT("Input folder to plugin information.")),
 		TArgument<FStringView>(TEXT("-OutputFolder"),		TEXT("Ouptut folder.")),
 		TArgument<FStringView>(TEXT("-IntermediateFolder"),	TEXT("Intermediate folder.")),
+		TArgument<FStringView>(TEXT("-SettingsFile"),		TEXT("Optional settings file.")),
 		TArgument<FStringView>(TEXT("-ErrorOutput"),		TEXT("Error output.")),
 		TArgument<bool>(TEXT("-IncludeSigPak"),				TEXT("Include .sig and .pak file in the uondemandtoc")),
 		TArgument<bool>(TEXT("-KeepContainerFiles"),		TEXT("Should we keep the container files after processing them.")),
diff --git a/Engine/Source/Runtime/Experimental/IoStore/OnDemand/Public/IO/IoStoreOnDemand.h b/Engine/Source/Runtime/Experimental/IoStore/OnDemand/Public/IO/IoStoreOnDemand.h
--- a/Engine/Source/Runtime/Experimental/IoStore/OnDemand/Public/IO/IoStoreOnDemand.h
+++ b/Engine/Source/Runtime/Experimental/IoStore/OnDemand/Public/IO/IoStoreOnDemand.h
@@ -59,6 +59,7 @@ enum class EOnDemandTocVersion : uint32
 	Meta			= 5,
 	ContainerId		= 6,
 	AdditionalFiles	= 7,
+	TagSets			= 8,
 
 	LatestPlusOne,
 	Latest			= (LatestPlusOne - 1)
@@ -161,6 +162,28 @@ struct FOnDemandTocAdditionalFile
 
 UE_API bool LoadFromCompactBinary(FCbFieldView Field, FOnDemandTocAdditionalFile& AdditionalFile);
 
+struct FOnDemandTocTagSetPackageList
+{
+	uint32 ContainerIndex = 0;
+	TArray<uint32> PackageIndicies;
+
+	UE_API friend FArchive& operator<<(FArchive& Ar, FOnDemandTocTagSetPackageList& TagSet);
+	UE_API friend FCbWriter& operator<<(FCbWriter& Writer, const FOnDemandTocTagSetPackageList& TagSet);
+};
+
+UE_API bool LoadFromCompactBinary(FCbFieldView Field, FOnDemandTocTagSetPackageList& TagSet);
+
+struct FOnDemandTocTagSet
+{
+	FString Tag;
+	TArray<FOnDemandTocTagSetPackageList> Packages;
+
+	UE_API friend FArchive& operator<<(FArchive& Ar, FOnDemandTocTagSet& TagSet);
+	UE_API friend FCbWriter& operator<<(FCbWriter& Writer, const FOnDemandTocTagSet& TagSet);
+};
+
+UE_API bool LoadFromCompactBinary(FCbFieldView Field, FOnDemandTocTagSet& TagSet);
+
 struct FOnDemandToc
 {
 	FOnDemandToc() = default;
@@ -178,6 +201,7 @@ struct FOnDemandToc
 	FTocMeta Meta;
 	TArray<FOnDemandTocContainerEntry> Containers;
 	TArray<FOnDemandTocAdditionalFile> AdditionalFiles;
+	TArray<FOnDemandTocTagSet> TagSets;
 
 	UE_API friend FArchive& operator<<(FArchive& Ar, FOnDemandToc& Toc);
 	UE_API friend FCbWriter& operator<<(FCbWriter& Writer, const FOnDemandToc& Toc);
@@ -234,6 +258,8 @@ struct FOnDemandMountArgs
 	FString Url;
 	/** Serialize the TOC from the specified file path. */
 	FString FilePath;
+	/** Used with EOnDemandMountOptions::Install to fliter requested content */
+	TArray<FString> TagSets;
 	/** Mount options. */
 	EOnDemandMountOptions Options;
 };
