Add a cvar ias.CompleteMaterializeTaskPriority so we can change the priority of CompleteMaterialize Ias tasks. Default this to BackgroundNormal
#tests Tested on console

diff --git a/Engine/Source/Runtime/Experimental/IoStore/OnDemand/Private/OnDemandIoDispatcherBackend.cpp b/Engine/Source/Runtime/Experimental/IoStore/OnDemand/Private/OnDemandIoDispatcherBackend.cpp
--- a/Engine/Source/Runtime/Experimental/IoStore/OnDemand/Private/OnDemandIoDispatcherBackend.cpp
+++ b/Engine/Source/Runtime/Experimental/IoStore/OnDemand/Private/OnDemandIoDispatcherBackend.cpp
@@ -191,6 +191,25 @@ static FAutoConsoleVariableRef CVar_DistributedEndpointAttemptCount(
 	TEXT("Number of times we should try to resolve a distributed endpoint befor eusing the fallback url (if there is one)")
 );
 
+
+// These priorities are indexed using the cvar below
+static UE::Tasks::ETaskPriority GCompleteMaterializeTaskPriorities[] =
+{
+	UE::Tasks::ETaskPriority::High,
+	UE::Tasks::ETaskPriority::Normal,
+	UE::Tasks::ETaskPriority::BackgroundHigh,
+	UE::Tasks::ETaskPriority::BackgroundNormal,
+	UE::Tasks::ETaskPriority::BackgroundLow
+};
+
+static int32 GCompleteMaterializeTaskPriority = 3;
+FAutoConsoleVariableRef CVarCompleteMaterializeTaskPriority(
+	TEXT("ias.CompleteMaterializeTaskPriority"),
+	GCompleteMaterializeTaskPriority,
+	TEXT("Task priority for the CompleteMaterialize task (0 = foreground/high, 1 = foreground/normal, 2 = background/high, 3 = background/normal, 4 = background/low)."),
+	ECVF_Default
+);
+
 #if !UE_BUILD_SHIPPING
 static FAutoConsoleCommand CVar_IasAbandonCache(
 	TEXT("Ias.AbandonCache"),
@@ -1314,11 +1333,12 @@ bool FOnDemandIoBackend::Resolve(FIoRequestImpl* Request)
 			InflightCacheRequestCount.fetch_add(1, std::memory_order_relaxed);
 
 			FTaskEvent OnReadyEvent(TEXT("IasCacheMaterializeDone"));
+			UE::Tasks::ETaskPriority TaskPriority = GCompleteMaterializeTaskPriorities[FMath::Clamp(GCompleteMaterializeTaskPriority, 0, UE_ARRAY_COUNT(GCompleteMaterializeTaskPriorities)-1)];
 
 			Launch(UE_SOURCE_LOCATION, [this, ChunkRequest] {
 				InflightCacheRequestCount.fetch_sub(1, std::memory_order_relaxed);
 				CompleteMaterialize(ChunkRequest);
-			}, OnReadyEvent);
+			}, OnReadyEvent, TaskPriority);
 
 			EIoErrorCode& OutStatus = ChunkRequest->CacheGetStatus;
 			Cache->Materialize(Key, Buffer, OutStatus, MoveTemp(OnReadyEvent));
