Iris
* Limit ObjectBridge allocs to CurrentMaxInternalIndex and support regrow when needed

#rnx

diff --git a/Engine/Source/Runtime/Experimental/Iris/Core/Private/Iris/ReplicationSystem/ObjectPollFrequencyLimiter.cpp b/Engine/Source/Runtime/Experimental/Iris/Core/Private/Iris/ReplicationSystem/ObjectPollFrequencyLimiter.cpp
--- a/Engine/Source/Runtime/Experimental/Iris/Core/Private/Iris/ReplicationSystem/ObjectPollFrequencyLimiter.cpp
+++ b/Engine/Source/Runtime/Experimental/Iris/Core/Private/Iris/ReplicationSystem/ObjectPollFrequencyLimiter.cpp
@@ -29,6 +29,14 @@ void FObjectPollFrequencyLimiter::Deinit()
 	FrameCounters.Empty();
 }
 
+void FObjectPollFrequencyLimiter::OnMaxInternalNetRefIndexIncreased(FInternalNetRefIndex NewMaxInternalIndex)
+{
+	const uint32 StorageObjectCount = Align(NewMaxInternalIndex, 32U);
+
+	FramesBetweenUpdates.SetNumZeroed(StorageObjectCount);
+	FrameCounters.SetNumZeroed(StorageObjectCount);
+}
+
 void FObjectPollFrequencyLimiter::Update(const FNetBitArrayView& RelevantObjects, const FNetBitArrayView& DirtyObjects, FNetBitArrayView& OutObjectsToPoll)
 {
 	IRIS_PROFILER_SCOPE(ObjectPollFrequencyLimiter_Update);
diff --git a/Engine/Source/Runtime/Experimental/Iris/Core/Private/Iris/ReplicationSystem/ObjectPollFrequencyLimiter.h b/Engine/Source/Runtime/Experimental/Iris/Core/Private/Iris/ReplicationSystem/ObjectPollFrequencyLimiter.h
--- a/Engine/Source/Runtime/Experimental/Iris/Core/Private/Iris/ReplicationSystem/ObjectPollFrequencyLimiter.h
+++ b/Engine/Source/Runtime/Experimental/Iris/Core/Private/Iris/ReplicationSystem/ObjectPollFrequencyLimiter.h
@@ -22,6 +22,8 @@ public:
 	void Init(uint32 MaxActiveObjectCount);
 	void Deinit();
 
+	void OnMaxInternalNetRefIndexIncreased(FInternalNetRefIndex NewMaxInternalIndex);
+
 	void SetPollFramePeriod(FInternalNetRefIndex InternalIndex, uint8 PollFramePeriod);
 
 	void SetPollWithObject(FInternalNetRefIndex ObjectToPollWithInternalIndex, FInternalNetRefIndex InternalIndex);
diff --git a/Engine/Source/Runtime/Experimental/Iris/Core/Private/Iris/ReplicationSystem/ObjectReplicationBridge.cpp b/Engine/Source/Runtime/Experimental/Iris/Core/Private/Iris/ReplicationSystem/ObjectReplicationBridge.cpp
--- a/Engine/Source/Runtime/Experimental/Iris/Core/Private/Iris/ReplicationSystem/ObjectReplicationBridge.cpp
+++ b/Engine/Source/Runtime/Experimental/Iris/Core/Private/Iris/ReplicationSystem/ObjectReplicationBridge.cpp
@@ -169,10 +169,12 @@ void UObjectReplicationBridge::Initialize(UReplicationSystem* InReplicationSyste
 
 	Super::Initialize(InReplicationSystem);
 
-	const uint32 MaxActiveObjectCount = NetRefHandleManager->GetMaxActiveObjectCount();
-	PollFrequencyLimiter->Init(MaxActiveObjectCount);
-	ObjectsWithObjectReferences.Init(MaxActiveObjectCount);
-	GarbageCollectionAffectedObjects.Init(MaxActiveObjectCount);
+	const uint32 CurrentMaxInternalIndex = NetRefHandleManager->GetCurrentMaxInternalNetRefIndex();
+	PollFrequencyLimiter->Init(CurrentMaxInternalIndex);
+	ObjectsWithObjectReferences.Init(CurrentMaxInternalIndex);
+	GarbageCollectionAffectedObjects.Init(CurrentMaxInternalIndex);
+
+	NetRefHandleManager->GetOnMaxInternalNetRefIndexIncreasedDelegate().AddUObject(this, &UObjectReplicationBridge::OnMaxInternalNetRefIndexIncreased);
 
 	LoadConfig();
 
@@ -181,6 +183,8 @@ void UObjectReplicationBridge::Initialize(UReplicationSystem* InReplicationSyste
 
 void UObjectReplicationBridge::Deinitialize()
 {
+	NetRefHandleManager->GetOnMaxInternalNetRefIndexIncreasedDelegate().RemoveAll(this);
+
 	UE::Net::Private::FPropertyConditionDelegates::GetOnPropertyCustomConditionChangedDelegate().Remove(OnCustomConditionChangedHandle);
 	UE::Net::Private::FPropertyConditionDelegates::GetOnPropertyDynamicConditionChangedDelegate().Remove(OnDynamicConditionChangedHandle);
 	OnCustomConditionChangedHandle.Reset();
@@ -189,6 +193,13 @@ void UObjectReplicationBridge::Deinitialize()
 	Super::Deinitialize();
 }
 
+void UObjectReplicationBridge::OnMaxInternalNetRefIndexIncreased(UE::Net::Private::FInternalNetRefIndex NewMaxInternalIndex)
+{
+	PollFrequencyLimiter->OnMaxInternalNetRefIndexIncreased(NewMaxInternalIndex);
+	ObjectsWithObjectReferences.SetNumBits(NewMaxInternalIndex);
+	GarbageCollectionAffectedObjects.SetNumBits(NewMaxInternalIndex);
+}
+
 UObject* UObjectReplicationBridge::GetObjectFromReferenceHandle(FNetRefHandle RefHandle) const
 {
 	return GetObjectReferenceCache()->GetObjectFromReferenceHandle(RefHandle);
diff --git a/Engine/Source/Runtime/Experimental/Iris/Core/Public/Iris/ReplicationSystem/ObjectReplicationBridge.h b/Engine/Source/Runtime/Experimental/Iris/Core/Public/Iris/ReplicationSystem/ObjectReplicationBridge.h
--- a/Engine/Source/Runtime/Experimental/Iris/Core/Public/Iris/ReplicationSystem/ObjectReplicationBridge.h
+++ b/Engine/Source/Runtime/Experimental/Iris/Core/Public/Iris/ReplicationSystem/ObjectReplicationBridge.h
@@ -392,6 +392,8 @@ private:
 
 	void InitConditionalPropertyDelegates();
 
+	void OnMaxInternalNetRefIndexIncreased(UE::Net::Private::FInternalNetRefIndex NewMaxInternalIndex);
+
 private:
 
 	friend UE::Net::Private::FObjectPoller;
