1- Fix deformable and GC crashes

diff --git a/Engine/Source/Runtime/Experimental/Chaos/Private/Chaos/Deformable/ChaosDeformableSolver.cpp b/Engine/Source/Runtime/Experimental/Chaos/Private/Chaos/Deformable/ChaosDeformableSolver.cpp
--- a/Engine/Source/Runtime/Experimental/Chaos/Private/Chaos/Deformable/ChaosDeformableSolver.cpp
+++ b/Engine/Source/Runtime/Experimental/Chaos/Private/Chaos/Deformable/ChaosDeformableSolver.cpp
@@ -1724,6 +1724,19 @@ namespace Chaos::Softs
 			Evolution->ResetConstraintRules();
 			Evolution->DeactivateParticleRanges();
 
+			// Sorting of the proxies to be removed to avoid particles indices issues in the next loop
+			RemovedProxies.Sort([](FThreadingProxy& BaseProxyA, FThreadingProxy& BaseProxyB) ->bool
+			{
+				const FFleshThreadingProxy* FleshProxyA = BaseProxyA.As<FFleshThreadingProxy>();
+				const FFleshThreadingProxy* FleshProxyB = BaseProxyB.As<FFleshThreadingProxy>();
+
+				if(FleshProxyA && FleshProxyB)
+				{
+					return FleshProxyA->GetSolverParticleRange()[0] < FleshProxyB->GetSolverParticleRange()[0];
+				}
+				return true;
+			});
+
 			// delete the simulated particles in block moves
 			for (FThreadingProxy* BaseProxy : RemovedProxies)
 			{
diff --git a/Engine/Source/Runtime/Experimental/Chaos/Private/GeometryCollection/Facades/CollectionTransformFacade.cpp b/Engine/Source/Runtime/Experimental/Chaos/Private/GeometryCollection/Facades/CollectionTransformFacade.cpp
--- a/Engine/Source/Runtime/Experimental/Chaos/Private/GeometryCollection/Facades/CollectionTransformFacade.cpp
+++ b/Engine/Source/Runtime/Experimental/Chaos/Private/GeometryCollection/Facades/CollectionTransformFacade.cpp
@@ -59,10 +59,13 @@ namespace GeometryCollection::Facades
 	{
 		TArray<FTransform> OutTransforms;
 
-		const TManagedArray<FTransform3f>& BoneTransforms = TransformAttribute.Get();
-		const TManagedArray<int32>& Parents = ParentAttribute.Get();
+		if(IsValid())
+		{
+			const TManagedArray<FTransform3f>& BoneTransforms = TransformAttribute.Get();
+			const TManagedArray<int32>& Parents = ParentAttribute.Get();
 
-		GeometryCollectionAlgo::GlobalMatrices(BoneTransforms, Parents, OutTransforms);
+			GeometryCollectionAlgo::GlobalMatrices(BoneTransforms, Parents, OutTransforms);
+		}
 
 		return OutTransforms;
 	}
