Added a GroupHash member to FISMComponentDescription for identical SMs to be separated into different groups for finer grained culling.

diff --git a/Engine/Source/Runtime/Experimental/GeometryCollectionEngine/Public/GeometryCollection/GeometryCollectionISMPoolComponent.h b/Engine/Source/Runtime/Experimental/GeometryCollectionEngine/Public/GeometryCollection/GeometryCollectionISMPoolComponent.h
--- a/Engine/Source/Runtime/Experimental/GeometryCollectionEngine/Public/GeometryCollection/GeometryCollectionISMPoolComponent.h
+++ b/Engine/Source/Runtime/Experimental/GeometryCollectionEngine/Public/GeometryCollection/GeometryCollectionISMPoolComponent.h
@@ -121,6 +121,7 @@ struct FISMComponentDescription
 	int32 StartCullDistance = 0;
 	int32 EndCullDistance = 0;
 	int32 MinLod = 0;
+	uint32 GroupHash = 0;	// Optional, allows identical SMs to be separated into different groups for finer grained culling
 	float LodScale = 1.f;
 	TArray<FName> Tags;
 	FName StatsCategory;
@@ -135,6 +136,7 @@ struct FISMComponentDescription
 			MinLod == Other.MinLod &&
 			LodScale == Other.LodScale &&
 			Tags == Other.Tags &&
+			GroupHash == Other.GroupHash &&
 			StatsCategory == Other.StatsCategory;
 	}
 };
@@ -147,6 +149,7 @@ FORCEINLINE uint32 GetTypeHash(const FISMComponentDescription& Desc)
 	Hash = HashCombineFast(Hash, GetTypeHash(Desc.EndCullDistance));
 	Hash = HashCombineFast(Hash, GetTypeHash(Desc.MinLod));
 	Hash = HashCombineFast(Hash, GetTypeHash(Desc.LodScale));
+	Hash = HashCombineFast(Hash, GetTypeHash(Desc.GroupHash));
 	Hash = HashCombineFast(Hash, GetArrayHash(Desc.Tags.GetData(), Desc.Tags.Num()));
 	return HashCombineFast(Hash, GetTypeHash(Desc.StatsCategory));
 }
