Added IAS processing to output sentry.

diff --git a/Engine/Source/Runtime/Experimental/IoStore/OnDemand/Private/OnDemandIoStore.cpp b/Engine/Source/Runtime/Experimental/IoStore/OnDemand/Private/OnDemandIoStore.cpp
--- a/Engine/Source/Runtime/Experimental/IoStore/OnDemand/Private/OnDemandIoStore.cpp
+++ b/Engine/Source/Runtime/Experimental/IoStore/OnDemand/Private/OnDemandIoStore.cpp
@@ -651,11 +651,13 @@ void FOnDemandIoStore::CreateContainersFromToc(
 	{
 		if (TocPath.IsEmpty() == false)
 		{
-			Algo::Transform(TocPath, AppendChars(Sb), FChar::ToLower);
+			//Algo::Transform(TocPath, AppendChars(Sb), FChar::ToLower);
+			FPathViews::Append(Sb, TocPath);
 		}
 		else
 		{
-			Algo::Transform(Toc.Header.ChunksDirectory, AppendChars(Sb), FChar::ToLower);
+			// Algo::Transform(Toc.Header.ChunksDirectory, AppendChars(Sb), FChar::ToLower);
+			FPathViews::Append(Sb, Toc.Header.ChunksDirectory);
 		}
 		FPathViews::Append(Sb, TEXT("chunks"));
 
