[IAS] Follow up to CL 33708323. Use CompleteMaterializeTaskPriority CVar for all CompleteMaterialize tasks, including decode from an HTTP request or when the requested chunk resides in memory.

#tests Tested ReplayRuns on Consolewith -Ias.DropCache and different values of CVar with some perf impact (and ensured logs set CVar as needed)

diff --git a/Engine/Source/Runtime/Experimental/IoStore/OnDemand/Private/OnDemandIoDispatcherBackend.cpp b/Engine/Source/Runtime/Experimental/IoStore/OnDemand/Private/OnDemandIoDispatcherBackend.cpp
--- a/Engine/Source/Runtime/Experimental/IoStore/OnDemand/Private/OnDemandIoDispatcherBackend.cpp
+++ b/Engine/Source/Runtime/Experimental/IoStore/OnDemand/Private/OnDemandIoDispatcherBackend.cpp
@@ -1322,9 +1322,11 @@ bool FOnDemandIoBackend::Resolve(FIoRequestImpl* Request)
 		{
 			check(Buffer.GetData() != nullptr);
 			ChunkRequest->bCached = true;
+
+			UE::Tasks::ETaskPriority TaskPriority = GCompleteMaterializeTaskPriorities[FMath::Clamp(GCompleteMaterializeTaskPriority, 0, UE_ARRAY_COUNT(GCompleteMaterializeTaskPriorities) - 1)];
 			Launch(UE_SOURCE_LOCATION, [this, ChunkRequest] {
 				CompleteRequest(ChunkRequest);
-			});
+			}, TaskPriority);
 			return true;
 		}
 
@@ -1644,10 +1646,11 @@ void FOnDemandIoBackend::ProcessHttpRequests(FHttpClient& HttpClient, FBitWindow
 							}
 							}
 
+							UE::Tasks::ETaskPriority TaskPriority = GCompleteMaterializeTaskPriorities[FMath::Clamp(GCompleteMaterializeTaskPriority, 0, UE_ARRAY_COUNT(GCompleteMaterializeTaskPriorities) - 1)];
 							UE::Tasks::Launch(UE_SOURCE_LOCATION, [this, ChunkRequest]()
 							{
 								CompleteRequest(ChunkRequest);
-							});
+							}, TaskPriority);
 						});
 				}
 			}
