Make the Ias OnDemand backend thread priority configurable via the cvar  ias.onDemandBackendThreadPriority. This is configurable at runtime while the thread is running if needed.
#tests Tested on console

diff --git a/Engine/Source/Runtime/Experimental/IoStore/OnDemand/Private/OnDemandIoDispatcherBackend.cpp b/Engine/Source/Runtime/Experimental/IoStore/OnDemand/Private/OnDemandIoDispatcherBackend.cpp
--- a/Engine/Source/Runtime/Experimental/IoStore/OnDemand/Private/OnDemandIoDispatcherBackend.cpp
+++ b/Engine/Source/Runtime/Experimental/IoStore/OnDemand/Private/OnDemandIoDispatcherBackend.cpp
@@ -210,6 +210,34 @@ FAutoConsoleVariableRef CVarCompleteMaterializeTaskPriority(
 	ECVF_Default
 );
 
+// Thread priority cvar (settable at runtime)
+// We declare these explicitly rather than just casting the cvar in case the enum changes in future
+const int32 GOnDemandBackendThreadPriorities[] =
+{
+	EThreadPriority::TPri_Lowest,
+	EThreadPriority::TPri_BelowNormal,
+	EThreadPriority::TPri_SlightlyBelowNormal,
+	EThreadPriority::TPri_Normal,
+	EThreadPriority::TPri_AboveNormal
+};
+
+const TCHAR* GOnDemandBackendThreadPriorityNames[] =
+{
+	TEXT("TPri_Lowest"),
+	TEXT("TPri_BelowNormal"),
+	TEXT("TPri_SlightlyBelowNormal"),
+	TEXT("TPri_Normal"),
+	TEXT("TPri_AboveNormal")
+};
+
+static int32 GOnDemandBackendThreadPriorityIndex = 4; // EThreadPriority::TPri_AboveNormal
+FAutoConsoleVariableRef CVarOnDemandBackendThreadPriority(
+	TEXT("ias.onDemandBackendThreadPriority"),
+	GOnDemandBackendThreadPriorityIndex,
+	TEXT("Thread priority of the on demand backend thread: 0=Lowest, 1=BelowNormal, 2=SlightlyBelowNormal, 3=Normal, 4=AboveNormal\n")
+	TEXT("Note that this is switchable at runtime"),
+	ECVF_Default);
+
 #if !UE_BUILD_SHIPPING
 static FAutoConsoleCommand CVar_IasAbandonCache(
 	TEXT("Ias.AbandonCache"),
@@ -1023,6 +1051,8 @@ private:
 	void ProcessHttpRequests(FHttpClient& HttpClient, FBitWindow& HttpErrors, int32 MaxConcurrentRequests);
 	int32 WaitForCompleteRequestTasks(float WaitTimeSeconds, float PollTimeSeconds);
 	void DrainHttpRequests();
+	void UpdateThreadPriorityIfNeeded();
+
 
 	FOnDemandIoStore& IoStore;
 	TUniquePtr<IIasCache> Cache;
@@ -1040,6 +1070,7 @@ private:
 	FEventRef DistributedEndpointEvent;
 
 	FString EndpointTestPath;
+	EThreadPriority CurrentThreadPriority;
 
 	mutable FRWLock Lock;
 	std::atomic_uint32_t InflightCacheRequestCount{0};
@@ -1054,6 +1085,7 @@ private:
 FOnDemandIoBackend::FOnDemandIoBackend(const FOnDemandEndpointConfig& Config, FOnDemandIoStore& InIoStore, TUniquePtr<IIasCache>&& InCache)
 	: IoStore(InIoStore)
 	, Cache(MoveTemp(InCache))
+	, CurrentThreadPriority(EThreadPriority::TPri_Num)
 {
 	EndpointTestPath = Config.TocPath;
 	if (Config.DistributionUrl.IsEmpty() == false)
@@ -1144,12 +1176,31 @@ FString FOnDemandIoBackend::GetEndpointTestPath() const
 	return EndpointTestPath;
 }
 
+void FOnDemandIoBackend::UpdateThreadPriorityIfNeeded()
+{
+	// Read the thread priority from the cvar
+	int32 ThreadPriorityIndex = FMath::Clamp(GOnDemandBackendThreadPriorityIndex, 0, (int32)UE_ARRAY_COUNT(GOnDemandBackendThreadPriorities) - 1);
+	EThreadPriority DesiredThreadPriority = (EThreadPriority)GOnDemandBackendThreadPriorities[ThreadPriorityIndex];
+	if (DesiredThreadPriority != CurrentThreadPriority)
+	{
+		UE_LOG(LogIas, Log, TEXT("Setting backend http thread priority to %s"), GOnDemandBackendThreadPriorityNames[ThreadPriorityIndex]);
+		FPlatformProcess::SetThreadPriority(DesiredThreadPriority);
+		CurrentThreadPriority = DesiredThreadPriority;
+	}
+}
+
+
 void FOnDemandIoBackend::ConditionallyStartBackendThread()
 {
 	FWriteScopeLock _(Lock);
 	if (BackendThread.IsValid() == false)
 	{
-		BackendThread.Reset(FRunnableThread::Create(this, TEXT("Ias.Http"), 0, TPri_AboveNormal));
+		// Read the desired thread priority from the cvar
+		int32 ThreadPriorityIndex = FMath::Clamp(GOnDemandBackendThreadPriorityIndex, 0, (int32)UE_ARRAY_COUNT(GOnDemandBackendThreadPriorities) - 1);
+		EThreadPriority DesiredThreadPriority = (EThreadPriority)GOnDemandBackendThreadPriorities[ThreadPriorityIndex];
+		CurrentThreadPriority = DesiredThreadPriority;
+
+		BackendThread.Reset(FRunnableThread::Create(this, TEXT("Ias.Http"), 0, DesiredThreadPriority));
 	}
 }
 
@@ -1750,6 +1801,8 @@ uint32 FOnDemandIoBackend::Run()
 
 	while (!bStopRequested)
 	{
+		UpdateThreadPriorityIfNeeded();
+
 		// Process HTTP request(s) even if the client is invalid to ensure enqueued request(s) gets completed.
 		ProcessHttpRequests(*HttpClient, HttpErrors, FMath::Min(GIasHttpConcurrentRequests, 32));
 		AvailableEps.Current = HttpClient->GetEndpoint();
@@ -1840,4 +1893,4 @@ TSharedPtr<IOnDemandIoDispatcherBackend> MakeOnDemandIoDispatcherBackend(
 
 } // namespace UE::IoStore
 
-#undef UE_IAS_DEBUG_CONSOLE_CMDS
+#undef UE_IAS_DEBUG_CONSOLE_CMDS
\ No newline at end of file
