Iris
- Optimized ResetObjectStateDirtiness to not touch all objects considered for polling
- Can be disabled by net.Iris.OnlyResetDirtinessForQuantizedObjects 0

diff --git a/Engine/Source/Runtime/Experimental/Iris/Core/Private/Iris/ReplicationSystem/ReplicationSystem.cpp b/Engine/Source/Runtime/Experimental/Iris/Core/Private/Iris/ReplicationSystem/ReplicationSystem.cpp
--- a/Engine/Source/Runtime/Experimental/Iris/Core/Private/Iris/ReplicationSystem/ReplicationSystem.cpp
+++ b/Engine/Source/Runtime/Experimental/Iris/Core/Private/Iris/ReplicationSystem/ReplicationSystem.cpp
@@ -41,6 +41,9 @@ static FAutoConsoleVariableRef CVarForcePruneBeforeUpdate(TEXT("net.Iris.ForcePr
 static bool bAllowAttachmentSendPolicyFlags = true;
 static FAutoConsoleVariableRef CVarAllowAttachmentSendPolicyFlags(TEXT("net.Iris.Attachments.AllowSendPolicyFlags"), bAllowAttachmentSendPolicyFlags, TEXT("Allow use of ENetObjectAttachmentSendPolicyFlags to specify behavior of RPCs."));
 
+static bool bOnlyResetDirtinessForQuantizedObjects = true;
+static FAutoConsoleVariableRef CVarOnlyResetDirtinessForQuantizedObjects(TEXT("net.Iris.OnlyResetDirtinessForQuantizedObjects"), bOnlyResetDirtinessForQuantizedObjects, TEXT("Only Reset Dirtiness For QuantizedObjects, optimization that only resets dirtiness for objects actually considered dirty."));
+
 }
 
 namespace UE::Net::Private
@@ -481,7 +484,7 @@ public:
 		};
 
 		DirtyObjectsToQuantize.ForAllSetBits(QuantizeFunction);
-		DirtyObjectsToQuantize.ClearAllBits();
+		// DirtyObjectsToQuantize is cleared in ResetObjectStateDirtiness
 
 		const uint32 ReplicationSystemId = ReplicationSystem->GetId();
 		UE_NET_TRACE_FRAME_STATSCOUNTER(ReplicationSystemId, ReplicationSystem.QuantizedObjectCount, QuantizedObjectCount, ENetTraceVerbosity::Trace);
@@ -495,12 +498,25 @@ public:
 
 		// Clear the objects that got polled this frame
 		const FNetBitArrayView PolledObjects = NetRefHandleManager.GetPolledObjectsInternalIndices();
+		FNetBitArrayView DirtyObjectsToQuantize = NetRefHandleManager.GetDirtyObjectsToQuantize();
 
-		// Reset object dirtyness
-		PolledObjects.ForAllSetBits([&NetRefHandleManager](uint32 DirtyIndex)
+		// This is clearing the internal changemask
+		if (ReplicationSystemCVars::bOnlyResetDirtinessForQuantizedObjects)
+		{
+			DirtyObjectsToQuantize.ForAllSetBits([&NetRefHandleManager](uint32 DirtyIndex)
+			{
+				FReplicationInstanceOperationsInternal::ResetObjectStateDirtiness(NetRefHandleManager, DirtyIndex);
+			});
+		}
+		else
 		{
-			FReplicationInstanceOperationsInternal::ResetObjectStateDirtiness(NetRefHandleManager, DirtyIndex);
-		});
+			PolledObjects.ForAllSetBits([&NetRefHandleManager](uint32 DirtyIndex)
+			{
+				FReplicationInstanceOperationsInternal::ResetObjectStateDirtiness(NetRefHandleManager, DirtyIndex);
+			});
+		}
+
+		DirtyObjectsToQuantize.ClearAllBits();
 
 		ReplicationSystemInternal.GetDirtyNetObjectTracker().ReconcilePolledList(PolledObjects);
 	}
