ChaosFlesh : Dynamic Constraints
Allow separate object to dynamically bind constraints between each other.

diff --git a/Engine/Plugins/Experimental/ChaosFlesh/Source/ChaosFleshEditor/Private/ChaosFlesh/ChaosFleshEditorPlugin.cpp b/Engine/Plugins/Experimental/ChaosFlesh/Source/ChaosFleshEditor/Private/ChaosFlesh/ChaosFleshEditorPlugin.cpp
--- a/Engine/Plugins/Experimental/ChaosFlesh/Source/ChaosFleshEditor/Private/ChaosFlesh/ChaosFleshEditorPlugin.cpp
+++ b/Engine/Plugins/Experimental/ChaosFlesh/Source/ChaosFleshEditor/Private/ChaosFlesh/ChaosFleshEditorPlugin.cpp
@@ -7,6 +7,7 @@
 #include "ChaosFlesh/Asset/FleshDeformableInterfaceDetails.h"
 #include "ChaosFlesh/Asset/FleshAssetThumbnailRenderer.h"
 #include "ChaosFlesh/ChaosDeformableCollisionsActor.h"
+#include "ChaosFlesh/ChaosDeformableConstraintsActor.h"
 #include "ChaosFlesh/ChaosDeformableSolverComponent.h"
 #include "ChaosFlesh/Cmd/ChaosFleshCommands.h"
 #include "ChaosFlesh/FleshActor.h"
@@ -63,6 +64,10 @@ void IChaosFleshEditorPlugin::StartupModule()
 		ADeformableCollisionsActor::StaticClass()->GetFName(),
 		FOnGetDetailCustomizationInstance::CreateStatic(&FDeformableInterfaceDetails::MakeInstance));
 
+	PropertyModule.RegisterCustomClassLayout(
+		ADeformableConstraintsActor::StaticClass()->GetFName(),
+		FOnGetDetailCustomizationInstance::CreateStatic(&FDeformableInterfaceDetails::MakeInstance));
+
 	PropertyModule.RegisterCustomClassLayout(
 		ADeformableSolverActor::StaticClass()->GetFName(),
 		FOnGetDetailCustomizationInstance::CreateStatic(&FDeformableInterfaceDetails::MakeInstance));
diff --git a/Engine/Plugins/Experimental/ChaosFlesh/Source/ChaosFleshEngine/Private/ChaosFlesh/ChaosDeformableConstraintsActor.cpp b/Engine/Plugins/Experimental/ChaosFlesh/Source/ChaosFleshEngine/Private/ChaosFlesh/ChaosDeformableConstraintsActor.cpp
new file mode 100644
--- /dev/null
+++ b/Engine/Plugins/Experimental/ChaosFlesh/Source/ChaosFleshEngine/Private/ChaosFlesh/ChaosDeformableConstraintsActor.cpp
@@ -0,0 +1,149 @@
+// Copyright Epic Games, Inc. All Rights Reserved.
+
+#include "ChaosFlesh/ChaosDeformableConstraintsActor.h"
+
+#include "ChaosFlesh/ChaosDeformableSolverComponent.h"
+#include "ChaosFlesh/ChaosDeformableSolverActor.h"
+#include "ChaosFlesh/FleshActor.h"
+
+#include UE_INLINE_GENERATED_CPP_BY_NAME(ChaosDeformableConstraintsActor)
+
+ADeformableConstraintsActor::ADeformableConstraintsActor(const FObjectInitializer& ObjectInitializer)
+	: Super(ObjectInitializer)
+{
+	DeformableConstraintsComponent = CreateDefaultSubobject<UDeformableConstraintsComponent>(TEXT("DeformableConstraintsComponent"));
+	RootComponent = DeformableConstraintsComponent;
+	PrimaryActorTick.bCanEverTick = false;
+}
+
+void ADeformableConstraintsActor::EnableSimulation(ADeformableSolverActor* InActor)
+{
+	if (InActor)
+	{
+		if (DeformableConstraintsComponent)
+		{
+			DeformableConstraintsComponent->EnableSimulation(InActor->GetDeformableSolverComponent());
+		}
+	}
+}
+
+#if WITH_EDITOR
+void ADeformableConstraintsActor::CustomizeDetails(IDetailLayoutBuilder& DetailBuilder)
+{
+	TSharedPtr<IPropertyHandle> bReplicatePhysicsProperty = DetailBuilder.GetProperty(GET_MEMBER_NAME_CHECKED(ADeformableConstraintsActor, bAsyncPhysicsTickEnabled), AActor::StaticClass());
+	bReplicatePhysicsProperty->MarkHiddenByCustomization();
+}
+
+void ADeformableConstraintsActor::PreEditChange(FProperty* PropertyThatWillChange)
+{
+	Super::PreEditChange(PropertyThatWillChange);
+	if (PropertyThatWillChange && PropertyThatWillChange->GetFName() == GET_MEMBER_NAME_CHECKED(ADeformableConstraintsActor, PrimarySolver))
+	{
+		PreEditChangePrimarySolver = PrimarySolver;
+	}
+	else if (PropertyThatWillChange && PropertyThatWillChange->GetFName() == GET_MEMBER_NAME_CHECKED(ADeformableConstraintsActor, SourceBodies))
+	{
+		PreEditChangeSourceBodies = SourceBodies;
+	}
+	else if (PropertyThatWillChange && PropertyThatWillChange->GetFName() == GET_MEMBER_NAME_CHECKED(ADeformableConstraintsActor, TargetBodies))
+	{
+		PreEditChangeTargetBodies = TargetBodies;
+	}
+}
+
+
+void ADeformableConstraintsActor::PostEditChangeProperty(struct FPropertyChangedEvent& PropertyChangedEvent)
+{
+	Super::PostEditChangeProperty(PropertyChangedEvent);
+
+	if (PropertyChangedEvent.Property && PropertyChangedEvent.Property->GetFName() == GET_MEMBER_NAME_CHECKED(ADeformableConstraintsActor, TargetBodies))
+	{
+		typedef TObjectPtr<AFleshActor> FActorType;
+		TSet< FActorType>  AT(PreEditChangeTargetBodies), BT(TargetBodies);
+		TSet< FActorType>  AS(PreEditChangeSourceBodies), BS(SourceBodies);
+		for (auto AddedTarget : BT.Difference(AT).Array())
+		{
+			for (auto AddedSource : BS.Difference(AS).Array())
+			{
+				if (AddedTarget && AddedTarget->GetFleshComponent())
+				{
+					if (AddedSource && AddedSource->GetFleshComponent())
+					{
+						GetConstraintsComponent()->AddConstrainedBodies(AddedSource->GetFleshComponent(), 
+							AddedTarget->GetFleshComponent(), FDeformableConstraintParameters());
+					}
+				}
+			}
+		}
+
+		TArray<FConstraintObject> ObjectsToRemove;
+		for (auto RemovedSource : AS.Difference(BS).Array())
+		{
+			for (FConstraintObject& Object : GetConstraintsComponent()->Constraints)
+			{
+				if (RemovedSource && RemovedSource->GetFleshComponent())
+				{
+					if (Object.Source == RemovedSource->GetFleshComponent())
+					{
+						ObjectsToRemove.Add(Object);
+					}
+				}
+			}
+		}
+		for (auto RemovedTarget : AT.Difference(BT).Array())
+		{
+			for (FConstraintObject& Object : GetConstraintsComponent()->Constraints)
+			{
+				if (RemovedTarget && RemovedTarget->GetFleshComponent())
+				{
+					if (Object.Target == RemovedTarget->GetFleshComponent())
+					{
+						ObjectsToRemove.Add(Object);
+					}
+				}
+			}
+		}
+
+		GetConstraintsComponent()->Constraints.RemoveAll([&ObjectsToRemove](const FConstraintObject& Request) {return ObjectsToRemove.Contains(Request); });
+		PreEditChangeSourceBodies.Empty();
+		PreEditChangeTargetBodies.Empty();
+	}
+	if (PropertyChangedEvent.Property && PropertyChangedEvent.Property->GetFName() == GET_MEMBER_NAME_CHECKED(ADeformableConstraintsActor, PrimarySolver))
+	{
+		if (PrimarySolver)
+		{
+			if (UDeformableSolverComponent* SolverComponent = PrimarySolver->GetDeformableSolverComponent())
+			{
+				if (DeformableConstraintsComponent)
+				{
+					DeformableConstraintsComponent->PrimarySolverComponent = SolverComponent;
+					if (!SolverComponent->ConnectedObjects.DeformableComponents.Contains(DeformableConstraintsComponent))
+					{
+						SolverComponent->ConnectedObjects.DeformableComponents.Add(TObjectPtr<UDeformablePhysicsComponent>(DeformableConstraintsComponent));
+					}
+				}
+			}
+		}
+		else if (PreEditChangePrimarySolver)
+		{
+			if (UDeformableSolverComponent* SolverComponent = PreEditChangePrimarySolver->GetDeformableSolverComponent())
+			{
+				if (DeformableConstraintsComponent)
+				{
+					DeformableConstraintsComponent->PrimarySolverComponent = nullptr;
+					if (SolverComponent->ConnectedObjects.DeformableComponents.Contains(DeformableConstraintsComponent))
+					{
+						SolverComponent->ConnectedObjects.DeformableComponents.Remove(DeformableConstraintsComponent);
+					}
+				}
+			}
+		}
+	}
+}
+#endif
+
+
+
+
+
+
diff --git a/Engine/Plugins/Experimental/ChaosFlesh/Source/ChaosFleshEngine/Private/ChaosFlesh/ChaosDeformableConstraintsComponent.cpp b/Engine/Plugins/Experimental/ChaosFlesh/Source/ChaosFleshEngine/Private/ChaosFlesh/ChaosDeformableConstraintsComponent.cpp
new file mode 100644
--- /dev/null
+++ b/Engine/Plugins/Experimental/ChaosFlesh/Source/ChaosFleshEngine/Private/ChaosFlesh/ChaosDeformableConstraintsComponent.cpp
@@ -0,0 +1,110 @@
+// Copyright Epic Games, Inc. All Rights Reserved.
+
+#include "ChaosFlesh/ChaosDeformableConstraintsComponent.h"
+
+#include "ChaosFlesh/ChaosDeformableSolverComponent.h"
+#include "Chaos/Deformable/ChaosDeformableConstraintsProxy.h"
+
+
+DEFINE_LOG_CATEGORY_STATIC(LogUDeformableConstraintsComponentInternal, Log, All);
+
+
+UDeformableConstraintsComponent::UDeformableConstraintsComponent(const FObjectInitializer& ObjectInitializer)
+	: Super(ObjectInitializer)
+{
+	PrimaryComponentTick.bCanEverTick = false;
+	bTickInEditor = false;
+}
+
+void UDeformableConstraintsComponent::AddConstrainedBodies(
+	UFleshComponent* SourceComponent,
+	UFleshComponent* TargetComponent,
+	FDeformableConstraintParameters InParameters)
+{
+	//PERF_SCOPE(STAT_ChaosDeformable_UDeformableConstraintsComponent_AddStaticMeshComponent);
+
+	FConstraintObject Constraint(SourceComponent, TargetComponent, InParameters);
+	if (IsValid(Constraint))
+	{
+		if (!Constraints.Contains(Constraint))
+		{
+			Constraints.Add(Constraint);
+			AddedConstraints.Add(Constraint);
+		}
+	}
+}
+
+
+void UDeformableConstraintsComponent::RemoveConstrainedBodies(
+	UFleshComponent* SourceComponent,
+	UFleshComponent* TargetComponent,
+	FDeformableConstraintParameters InParameters)
+{
+	//PERF_SCOPE(STAT_ChaosDeformable_UDeformableConstraintsComponent_RemoveStaticMeshComponent);
+
+	FConstraintObject Constraint(SourceComponent, TargetComponent, InParameters);
+	if (IsValid(Constraint))
+	{
+		if (Constraints.Contains(Constraint))
+		{
+			Constraints.Remove(Constraint);
+			RemovedConstraints.Add(Constraint);
+		}
+	}
+}
+
+
+UDeformablePhysicsComponent::FThreadingProxy* UDeformableConstraintsComponent::NewProxy()
+{
+	//PERF_SCOPE(STAT_ChaosDeformable_UDeformableConstraintsComponent_NewProxy);
+
+	for (auto& Constraint : Constraints)
+	{
+		if (IsValid(Constraint))
+		{
+			if (!AddedConstraints.Contains(Constraint))
+			{
+				AddedConstraints.Add(Constraint);
+			}
+		}
+	}
+	return new FConstraintThreadingProxy(this);
+}
+
+bool UDeformableConstraintsComponent::IsValid(const FConstraintObject& Key) const
+{
+	return Key.Source && Key.Target;
+}
+
+
+UDeformableConstraintsComponent::FDataMapValue 
+UDeformableConstraintsComponent::NewDeformableData()
+{
+	//PERF_SCOPE(STAT_ChaosDeformable_UDeformableConstraintsComponent_NewDeformableData);
+
+	TArray<Chaos::Softs::FConstraintObjectAdded> AddedConstraintsData;
+	TArray<Chaos::Softs::FConstraintObjectRemoved> RemovedConstraintsData;
+	TArray<Chaos::Softs::FConstraintObjectUpdated> UpdateConstraintData;
+
+	for (auto& Constraint : AddedConstraints)
+	{
+		if (IsValid(Constraint))
+		{
+			AddedConstraintsData.Add(Constraint.ToChaos());
+		}
+	}
+
+	for (auto& Constraint : RemovedConstraints)
+	{
+		if (IsValid(Constraint))
+		{
+			RemovedConstraintsData.Add(Constraint.ToChaos());
+		}
+	}
+
+	AddedConstraints.Empty();
+	RemovedConstraints.Empty();
+
+	return FDataMapValue(new Chaos::Softs::FConstraintManagerProxy::FConstraintsInputBuffer(
+		AddedConstraintsData, RemovedConstraintsData, UpdateConstraintData, this));
+}
diff --git a/Engine/Plugins/Experimental/ChaosFlesh/Source/ChaosFleshEngine/Public/ChaosFlesh/ChaosDeformableConstraintsActor.h b/Engine/Plugins/Experimental/ChaosFlesh/Source/ChaosFleshEngine/Public/ChaosFlesh/ChaosDeformableConstraintsActor.h
new file mode 100644
--- /dev/null
+++ b/Engine/Plugins/Experimental/ChaosFlesh/Source/ChaosFleshEngine/Public/ChaosFlesh/ChaosDeformableConstraintsActor.h
@@ -0,0 +1,56 @@
+// Copyright Epic Games, Inc. All Rights Reserved.
+
+#pragma once
+
+#include "ChaosFlesh/ChaosDeformableConstraintsComponent.h"
+#include "ChaosFlesh/ChaosDeformableSolverActor.h"
+#include "CoreMinimal.h"
+#include "DeformableInterface.h"
+#include "GameFramework/Actor.h"
+#include "UObject/ObjectMacros.h"
+#include "UObject/ObjectPtr.h"
+
+#include "ChaosDeformableConstraintsActor.generated.h"
+
+class AFleshActor;
+
+UCLASS()
+class CHAOSFLESHENGINE_API ADeformableConstraintsActor : public AActor, public IDeformableInterface
+{
+	GENERATED_UCLASS_BODY()
+
+public:
+
+	UFUNCTION(BlueprintCallable, Category = "Physics")
+	void EnableSimulation(ADeformableSolverActor* Actor);
+
+	UPROPERTY(VisibleAnywhere, BlueprintReadOnly, Category = "Physics")
+	TObjectPtr<UDeformableConstraintsComponent> DeformableConstraintsComponent;
+	UDeformableConstraintsComponent* GetConstraintsComponent() const { return DeformableConstraintsComponent; }
+
+	UPROPERTY(EditAnywhere, BlueprintReadWrite, Category = "Physics", meta = (DisplayPriority = 1))
+	TObjectPtr<ADeformableSolverActor> PrimarySolver;
+
+	UPROPERTY(EditAnywhere, BlueprintReadWrite, Category = "Physics", meta = (DisplayPriority = 2))
+	TArray<TObjectPtr<AFleshActor>> SourceBodies;
+
+	UPROPERTY(EditAnywhere, BlueprintReadWrite, Category = "Physics", meta = (DisplayPriority = 2))
+	TArray<TObjectPtr<AFleshActor>> TargetBodies;
+
+
+#if WITH_EDITOR
+	TArray < TObjectPtr<AFleshActor> > AddedSourceBodies;
+	TArray < TObjectPtr<AFleshActor> > RemovedSourceBodies;
+
+	TArray < TObjectPtr<AFleshActor> > AddedTargetBodies;
+	TArray < TObjectPtr<AFleshActor> > RemovedTargetBodies;
+
+	TArray < TObjectPtr<AFleshActor> > PreEditChangeSourceBodies;
+	TArray < TObjectPtr<AFleshActor> > PreEditChangeTargetBodies;
+
+	ADeformableSolverActor* PreEditChangePrimarySolver = nullptr;
+	virtual void PreEditChange(FProperty* PropertyThatWillChange) override;
+	virtual void PostEditChangeProperty(struct FPropertyChangedEvent& PropertyChangedEvent) override;
+	virtual void CustomizeDetails(IDetailLayoutBuilder& DetailBuilder) override;
+#endif
+};
diff --git a/Engine/Plugins/Experimental/ChaosFlesh/Source/ChaosFleshEngine/Public/ChaosFlesh/ChaosDeformableConstraintsComponent.h b/Engine/Plugins/Experimental/ChaosFlesh/Source/ChaosFleshEngine/Public/ChaosFlesh/ChaosDeformableConstraintsComponent.h
new file mode 100644
--- /dev/null
+++ b/Engine/Plugins/Experimental/ChaosFlesh/Source/ChaosFleshEngine/Public/ChaosFlesh/ChaosDeformableConstraintsComponent.h
@@ -0,0 +1,132 @@
+// Copyright Epic Games, Inc. All Rights Reserved.
+#pragma once
+
+#include "CoreMinimal.h"
+#include "Chaos/Deformable/ChaosDeformableSolverProxy.h"
+#include "Chaos/Deformable/ChaosDeformableConstraintsProxy.h"
+#include "ChaosFlesh/ChaosDeformablePhysicsComponent.h"
+#include "UObject/ObjectMacros.h"
+
+#include "ChaosDeformableConstraintsComponent.generated.h"
+
+class UFleshComponent;
+
+
+//ENUM()
+enum  EDeformableConstraintType : uint8
+{
+	Kinematic,
+	LinearSpring,
+	MAX
+};
+
+
+USTRUCT(BlueprintType)
+struct FDeformableConstraintParameters
+{
+	GENERATED_BODY()
+
+	FDeformableConstraintParameters(float InStiffness = 100000.f, float InDamping = 1.f,
+		EDeformableConstraintType InType = EDeformableConstraintType::Kinematic)
+		: Type(InType)
+		, Stiffness(InStiffness)
+		, Damping(InDamping)
+	{}
+	
+	//UPROPERTY(EditAnywhere)
+	EDeformableConstraintType Type;
+
+	UPROPERTY(EditAnywhere, Category = "Constraint")
+	float Stiffness = 100000.f;
+	
+	UPROPERTY(EditAnywhere, Category = "Constraint")
+	float Damping = 1.f;
+
+	Chaos::Softs::FDeformableConstraintParameters ToChaos() 
+	{ 
+		return Chaos::Softs::FDeformableConstraintParameters(Stiffness,Damping,
+			(Chaos::Softs::EDeformableConstraintType)Type);
+	}
+
+};
+
+USTRUCT(BlueprintType)
+struct FConstraintObject
+{
+	GENERATED_BODY()
+
+	FConstraintObject(
+		TObjectPtr<UFleshComponent> InSource = nullptr,
+		TObjectPtr<UFleshComponent> InTarget = nullptr,
+		FDeformableConstraintParameters InParameters = FDeformableConstraintParameters())
+		: Source(InSource)
+		, Target(InTarget)
+		, Parameters(InParameters)
+	{}
+
+	UPROPERTY(EditAnywhere, Category = "Constraint")
+	TObjectPtr<UFleshComponent> Source;
+
+	UPROPERTY(EditAnywhere, Category = "Constraint")
+	TObjectPtr<UFleshComponent> Target;
+
+	UPROPERTY(EditAnywhere, Category = "Constraint")
+	FDeformableConstraintParameters Parameters;
+
+	Chaos::Softs::FConstraintObjectKey ToChaos()
+	{
+		return Chaos::Softs::FConstraintObjectKey(
+			TObjectPtr<UObject>((UObject*)Source.Get()),
+			TObjectPtr<UObject>((UObject*)Target.Get()),
+			Chaos::Softs::EDeformableConstraintType(Parameters.Type));
+	}
+
+	bool operator==(const FConstraintObject& Other) const
+	{
+		return Other.Source==Source && Other.Target == Target;
+	}
+};
+
+/**
+*	UDeformableConstraintsComponent
+*/
+UCLASS(meta = (BlueprintSpawnableComponent))
+class CHAOSFLESHENGINE_API UDeformableConstraintsComponent : public UDeformablePhysicsComponent
+{
+	GENERATED_UCLASS_BODY()
+
+public:
+	typedef Chaos::Softs::FDataMapValue FDataMapValue;
+	typedef Chaos::Softs::FThreadingProxy FThreadingProxy;
+	typedef Chaos::Softs::FConstraintManagerProxy FConstraintThreadingProxy;
+	typedef Chaos::Softs::FConstraintObjectKey FConstraintObjectKey;
+
+	~UDeformableConstraintsComponent() {}
+
+	/** Simulation Interface*/
+	virtual FThreadingProxy* NewProxy() override;
+	virtual FDataMapValue NewDeformableData() override;
+
+	UFUNCTION(BlueprintCallable, Category = "Physics")
+	void AddConstrainedBodies(
+		UFleshComponent* SourceComponent,
+		UFleshComponent* TargetComponent,
+		FDeformableConstraintParameters InParameters);
+	
+	UFUNCTION(BlueprintCallable, Category = "Physics")
+	void RemoveConstrainedBodies(
+		UFleshComponent* SourceComponent,
+		UFleshComponent* TargetComponent,
+		FDeformableConstraintParameters InParameters);
+
+	UPROPERTY(BlueprintReadOnly, Category = "Physics")
+	TArray<FConstraintObject>  Constraints;
+
+protected:
+	bool IsValid(const FConstraintObject&) const;
+
+	TArray<FConstraintObject> RemovedConstraints;
+	TArray<FConstraintObject> AddedConstraints;
+
+	TMap<FConstraintObjectKey, FThreadingProxy*> ConstraintsMap;
+};
diff --git a/Engine/Source/Runtime/Experimental/Chaos/Private/Chaos/Deformable/ChaosDeformableSolver.cpp b/Engine/Source/Runtime/Experimental/Chaos/Private/Chaos/Deformable/ChaosDeformableSolver.cpp
--- a/Engine/Source/Runtime/Experimental/Chaos/Private/Chaos/Deformable/ChaosDeformableSolver.cpp
+++ b/Engine/Source/Runtime/Experimental/Chaos/Private/Chaos/Deformable/ChaosDeformableSolver.cpp
@@ -3,9 +3,13 @@
 #include "Chaos/Deformable/ChaosDeformableSolver.h"
 #include "Chaos/Deformable/ChaosDeformableSolverTypes.h"
 #include "Chaos/Deformable/ChaosDeformableSolverProxy.h"
+#include "Chaos/Deformable/ChaosDeformableConstraintsProxy.h"
 #include "Chaos/Deformable/ChaosDeformableCollisionsProxy.h"
 
+#include "ChaosLog.h"
+#include "Chaos/BoundingVolumeHierarchy.h"
 #include "Chaos/DebugDrawQueue.h"
+#include "Chaos/Tetrahedron.h"
 #include "Chaos/TriangleMesh.h"
 #include "Chaos/PBDAltitudeSpringConstraints.h"
 #include "Chaos/PBDBendingConstraints.h"
@@ -24,6 +28,7 @@
 #include "Chaos/Plane.h"
 #include "Chaos/Utilities.h"
 #include "Chaos/PBDEvolution.h"
+#include "Containers/ContainersFwd.h"
 #include "Containers/StringConv.h"
 #include "Containers/Set.h"
 #include "CoreMinimal.h"
@@ -33,6 +38,7 @@
 #include "GeometryCollection/Facades/CollectionConstraintOverrideFacade.h"
 #include "GeometryCollection/Facades/CollectionMeshFacade.h"
 #include "GeometryCollection/Facades/CollectionMuscleActivationFacade.h"
+#include "GeometryCollection/Facades/CollectionTetrahedralFacade.h"
 #include "Misc/FileHelper.h"
 #include "HAL/PlatformFileManager.h"
 #include "HAL/IConsoleManager.h"
@@ -60,6 +66,9 @@ DECLARE_CYCLE_STAT(TEXT("Chaos.Deformable.Solver.InitializeGidBasedConstraints")
 DECLARE_CYCLE_STAT(TEXT("Chaos.Deformable.Solver.InitializeWeakConstraints"), STAT_ChaosDeformableSolver_InitializeWeakConstraints, STATGROUP_Chaos);
 DECLARE_CYCLE_STAT(TEXT("Chaos.Deformable.Solver.InitializeKinematicConstraint"), STAT_ChaosDeformableSolver_InitializeKinematicConstraint, STATGROUP_Chaos);
 DECLARE_CYCLE_STAT(TEXT("Chaos.Deformable.Solver.InitializeCollisionBodies"), STAT_ChaosDeformableSolver_InitializeCollisionBodies, STATGROUP_Chaos);
+DECLARE_CYCLE_STAT(TEXT("Chaos.Deformable.Solver.UpdateCollisionBodies"), STAT_ChaosDeformableSolver_UpdateCollisionBodies, STATGROUP_Chaos);
+DECLARE_CYCLE_STAT(TEXT("Chaos.Deformable.Solver.InitializeConstraintBodies"), STAT_ChaosDeformableSolver_InitializeConstraintBodies, STATGROUP_Chaos);
+DECLARE_CYCLE_STAT(TEXT("Chaos.Deformable.Solver.UpdateConstraintBodies"), STAT_ChaosDeformableSolver_UpdateConstraintBodies, STATGROUP_Chaos);
 DECLARE_CYCLE_STAT(TEXT("Chaos.Deformable.Solver.InitializeSelfCollisionVariables"), STAT_ChaosDeformableSolver_InitializeSelfCollisionVariables, STATGROUP_Chaos);
 DECLARE_CYCLE_STAT(TEXT("Chaos.Deformable.Solver.InitializeGridBasedConstraintVariables"), STAT_ChaosDeformableSolver_InitializeGridBasedConstraintVariables, STATGROUP_Chaos);
 DECLARE_CYCLE_STAT(TEXT("Chaos.Deformable.Solver.InitializeGaussSeidelConstraintVariables"), STAT_ChaosDeformableSolver_InitializeGaussSeidelConstraintVariables, STATGROUP_Chaos);
@@ -373,6 +382,10 @@ namespace Chaos::Softs
 				{
 					UpdateCollisionBodies(*CollisionManagerProxy, Entry.Key, DeltaTime);
 				}
+				else if (FConstraintManagerProxy* ConstraintManagerProxy = Proxy.As< FConstraintManagerProxy>())
+				{
+					UpdateConstraintBodies(*ConstraintManagerProxy, Entry.Key, DeltaTime);
+				}
 			}
 		}
 
@@ -404,6 +417,11 @@ namespace Chaos::Softs
 		{
 			InitializeCollisionBodies(*CollisionManagerProxy);
 		}
+
+		if (FConstraintManagerProxy* ConstraintManagerProxy = InProxy.As< FConstraintManagerProxy>())
+		{
+			InitializeConstraintBodies(*ConstraintManagerProxy);
+		}
 	}
 
 	void FDeformableSolver::InitializeDeformableParticles(FFleshThreadingProxy& Proxy)
@@ -641,7 +659,7 @@ namespace Chaos::Softs
 
 	void FDeformableSolver::UpdateCollisionBodies(FCollisionManagerProxy& Proxy, FThreadingProxy::FKey Owner, FSolverReal DeltaTime)
 	{
-		PERF_SCOPE(STAT_ChaosDeformableSolver_InitializeCollisionBodies);
+		PERF_SCOPE(STAT_ChaosDeformableSolver_UpdateCollisionBodies);
 
 		FCollisionManagerProxy::FCollisionsInputBuffer* CollisionsInputBuffer = nullptr;
 		if (this->CurrentInputPackage && this->CurrentInputPackage->ObjectMap.Contains(Owner))
@@ -733,7 +751,56 @@ namespace Chaos::Softs
 
 	}
 
+	void FDeformableSolver::InitializeConstraintBodies(FConstraintManagerProxy& Proxy)
+	{
+		PERF_SCOPE(STAT_ChaosDeformableSolver_InitializeConstraintBodies);
+	}
+
+	void FDeformableSolver::UpdateConstraintBodies(FConstraintManagerProxy& Proxy, FThreadingProxy::FKey Owner, FSolverReal DeltaTime)
+	{
+		PERF_SCOPE(STAT_ChaosDeformableSolver_UpdateConstraintBodies);
 
+		FConstraintManagerProxy::FConstraintsInputBuffer* ConstraintsInputBuffer = nullptr;
+		if (this->CurrentInputPackage && this->CurrentInputPackage->ObjectMap.Contains(Owner))
+		{
+			if (this->CurrentInputPackage->ObjectMap[Owner] != nullptr)
+			{
+				ConstraintsInputBuffer = this->CurrentInputPackage->ObjectMap[Owner]->As<FConstraintManagerProxy::FConstraintsInputBuffer>();
+				if (ConstraintsInputBuffer)
+				{
+					for (auto& AddConstraints : ConstraintsInputBuffer->Added)
+					{
+						//@todo (defered initilziation) : Save the added constraint into a buffer for processing later. 
+						UE_LOG(LogChaosDeformableSolver, Log, TEXT("Process Constraint : %s"), *Proxy.GetOwner()->GetName());
+
+						if (Proxies.Contains(AddConstraints.Get<0>()) && Proxies.Contains(AddConstraints.Get<1>()))
+						{
+							FFleshThreadingProxy* SourceProxy = Proxies[AddConstraints.Get<0>()]->As<FFleshThreadingProxy>();
+							FFleshThreadingProxy* TargetProxy = Proxies[AddConstraints.Get<1>()]->As<FFleshThreadingProxy>();
+							if (SourceProxy && TargetProxy)
+							{
+								FIntVector2 SampleRange = SourceProxy->GetSolverParticleRange();
+								TConstArrayView<Softs::FSolverVec3> Samples(&Evolution->Particles().GetX(SampleRange[0]), SampleRange[1]);
+
+								FIntVector2 TargetRange = TargetProxy->GetSolverParticleRange();
+								TConstArrayView<Softs::FSolverVec3> TetVertices(&Evolution->Particles().GetX(TargetRange[0]), TargetRange[1]);
+
+								const GeometryCollection::Facades::FTetrahedralFacade Geom(TargetProxy->GetRestCollection());
+								TArray<GeometryCollection::Facades::TetrahedralParticleEmbedding> Intersections;
+								if (Geom.Intersection(Samples, TetVertices, Intersections))
+								{
+									UE_LOG(LogChaosDeformableSolver, Log, TEXT("... Intersections : %d"), Intersections.Num());
+									//GSWeakConstraints->AddExtraConstraints(PositionTargetIndices, PositionTargetWeights, PositionTargetStiffness, PositionTargetSecondIndices, PositionTargetSecondWeights);
+
+								}
+							}
+						}
+					}
+					ConstraintsInputBuffer->Added.Empty();
+				}
+			}
+		}
+	}
 
 	void FDeformableSolver::DebugDrawTetrahedralParticles(FFleshThreadingProxy& Proxy)
 	{
diff --git a/Engine/Source/Runtime/Experimental/Chaos/Private/Chaos/Deformable/ChaosDeformableSolverThreadingWrapper.cpp b/Engine/Source/Runtime/Experimental/Chaos/Private/Chaos/Deformable/ChaosDeformableSolverThreadingWrapper.cpp
--- a/Engine/Source/Runtime/Experimental/Chaos/Private/Chaos/Deformable/ChaosDeformableSolverThreadingWrapper.cpp
+++ b/Engine/Source/Runtime/Experimental/Chaos/Private/Chaos/Deformable/ChaosDeformableSolverThreadingWrapper.cpp
@@ -63,12 +63,6 @@ namespace Chaos::Softs
 			Solver->InitializeSimulationObject(InProxy);
 	}
 
-	void FDeformableSolver::FPhysicsThreadAccess::InitializeCollisionBodies(FCollisionManagerProxy& Proxy)
-	{
-		if (Solver)
-			Solver->InitializeCollisionBodies(Proxy);
-	}
-
 	void FDeformableSolver::FPhysicsThreadAccess::InitializeKinematicConstraint()
 	{
 		if(Solver) Solver->InitializeKinematicConstraint();
diff --git a/Engine/Source/Runtime/Experimental/Chaos/Private/GeometryCollection/Facades/CollectionTetrahedralFacade.cpp b/Engine/Source/Runtime/Experimental/Chaos/Private/GeometryCollection/Facades/CollectionTetrahedralFacade.cpp
new file mode 100644
--- /dev/null
+++ b/Engine/Source/Runtime/Experimental/Chaos/Private/GeometryCollection/Facades/CollectionTetrahedralFacade.cpp
@@ -0,0 +1,105 @@
+// Copyright Epic Games, Inc. All Rights Reserved.
+
+/*=============================================================================
+	GeometryCollection.cpp: FGeometryCollection methods.
+=============================================================================*/
+
+#include "GeometryCollection/Facades/CollectionTetrahedralFacade.h"
+#include "GeometryCollection/GeometryCollection.h"
+#include "GeometryCollection/ManagedArrayAccessor.h"
+#include "Chaos/Tetrahedron.h"
+#include "Chaos/BoundingVolumeHierarchy.h"
+
+namespace GeometryCollection::Facades
+{
+	FTetrahedralFacade::FTetrahedralFacade(FManagedArrayCollection& InCollection)
+		: ConstCollection(InCollection)
+		, Collection(&InCollection)
+		, Tetrahedron(InCollection, "Tetrahedron", "Tetrahedral")
+		, TetrahedronStart(InCollection, "TetrahedronStart", FGeometryCollection::GeometryGroup)
+		, TetrahedronCount(InCollection, "TetrahedronCount", FGeometryCollection::GeometryGroup)
+		, VertexStart(InCollection, "VertexStart", FGeometryCollection::GeometryGroup)
+		, VertexCount(InCollection, "VertexCount", FGeometryCollection::GeometryGroup)
+		, Vertex(InCollection, "Vertex", FGeometryCollection::VerticesGroup)
+	{}
+
+	FTetrahedralFacade::FTetrahedralFacade(const FManagedArrayCollection& InCollection)
+		: ConstCollection(InCollection)
+		, Collection(nullptr)
+		, Tetrahedron(InCollection, "Tetrahedron", "Tetrahedral")
+		, TetrahedronStart(InCollection, "TetrahedronStart", FGeometryCollection::GeometryGroup)
+		, TetrahedronCount(InCollection, "TetrahedronCount", FGeometryCollection::GeometryGroup)
+		, VertexStart(InCollection, "VertexStart", FGeometryCollection::GeometryGroup)
+		, VertexCount(InCollection, "VertexCount", FGeometryCollection::GeometryGroup)
+		, Vertex(InCollection, "Vertex", FGeometryCollection::VerticesGroup)
+	{}
+
+	FTetrahedralFacade::~FTetrahedralFacade()
+	{}
+
+	void FTetrahedralFacade::DefineSchema()
+	{
+	}
+
+	bool FTetrahedralFacade::IsValid() const
+	{
+		return Tetrahedron.IsValid() && TetrahedronStart.IsValid() && TetrahedronCount.IsValid()
+			&& VertexStart.IsValid() && VertexCount.IsValid() && Vertex.IsValid();
+	}
+
+	bool FTetrahedralFacade::Intersection(
+		const TConstArrayView<Chaos::Softs::FSolverVec3>& SamplePositions,
+		const TConstArrayView<Chaos::Softs::FSolverVec3>& TetarhedronPositions,
+		TArray<TetrahedralParticleEmbedding>& OutIntersections) const
+	{
+		using namespace Chaos;
+
+		auto TetrahedralMeshIndices = [&]() {
+			TArray<int32> RetList;
+			for (int32 i = 0; i < TetrahedronCount.Num(); i++) if (TetrahedronCount[i] > 0) RetList.Add(i);
+			return RetList;
+		};
+
+		for (int32 MeshIndex : TetrahedralMeshIndices())
+		{
+			int32 TetStart = TetrahedronStart[MeshIndex];
+			int32 TetCount = TetrahedronCount[MeshIndex];
+			int32 VrtStart = VertexStart[MeshIndex];
+			// Build Tetrahedra
+			typedef Chaos::TTetrahedron<Chaos::FReal> FTet;
+			TArray<FTet> Tets;
+			TArray<FTet*> TetPtrs;
+
+			Tets.SetNumUninitialized(TetCount);
+			TetPtrs.SetNumUninitialized(TetCount);
+
+			for (int32 Tdx = 0; Tdx < TetCount; Tdx++)
+			{
+				const int32 Idx = TetStart + Tdx;
+				FIntVector4 Tet = Tetrahedron[Idx] - FIntVector4(VrtStart);
+				Tets[Tdx] = FTet(TetarhedronPositions[Tet[0]], TetarhedronPositions[Tet[1]], TetarhedronPositions[Tet[2]], TetarhedronPositions[Tet[3]]);
+				TetPtrs[Tdx] = &Tets[Tdx];
+			}
+
+			Chaos::TBoundingVolumeHierarchy< TArray<FTet*>, TArray<int32>, Chaos::FReal, 3> BVH(TetPtrs);
+			for (int Pdx = 0; Pdx < SamplePositions.Num(); Pdx++)
+			{
+				const auto& Pos = SamplePositions[Pdx];
+				TArray<int32> Intersections = BVH.FindAllIntersections(Pos);
+				for (int32 j = 0; j < Intersections.Num(); j++)
+				{
+					const int32 TetIdx = Intersections[j];
+					if (!Tets[TetIdx].Outside(Pos, 0))
+					{
+						TVector<FReal, 4> W = Tets[TetIdx].GetBarycentricCoordinates(Pos);
+						OutIntersections.Add({ Pdx, MeshIndex, TetIdx, { (float)W[0],(float)W[1],(float)W[2],(float)W[3] } });
+						break;
+					}
+				}
+			}
+		}
+		return OutIntersections.Num()>0;
+	}
+};
+
+
diff --git a/Engine/Source/Runtime/Experimental/Chaos/Public/Chaos/Deformable/ChaosDeformableConstraintsProxy.h b/Engine/Source/Runtime/Experimental/Chaos/Public/Chaos/Deformable/ChaosDeformableConstraintsProxy.h
new file mode 100644
--- /dev/null
+++ b/Engine/Source/Runtime/Experimental/Chaos/Public/Chaos/Deformable/ChaosDeformableConstraintsProxy.h
@@ -0,0 +1,112 @@
+// Copyright Epic Games, Inc. All Rights Reserved.
+
+#pragma once
+#include "Chaos/Deformable/ChaosDeformableSolverProxy.h"
+#include "CoreMinimal.h"
+#include "UObject/ObjectMacros.h"
+
+namespace Chaos::Softs
+{
+	enum EDeformableConstraintType
+	{
+		Kinematic,
+		LinearSpring
+	};
+
+	struct CHAOS_API FDeformableConstraintParameters
+	{
+		FDeformableConstraintParameters(float InStiffness = 100000.f, float InDamping = 1.f,
+			EDeformableConstraintType InType = EDeformableConstraintType::Kinematic)
+		: Type(InType)
+		, Stiffness(InStiffness)
+		, Damping(InDamping)
+		{}
+		EDeformableConstraintType Type;
+		float Stiffness = 100000.f;
+		float Damping = 1.f;
+	};
+
+	typedef TTuple< const TObjectPtr<UObject>, const TObjectPtr<UObject>, EDeformableConstraintType> FConstraintObjectKey;
+
+	struct FConstraintObjectAdded : public FConstraintObjectKey
+	{
+		FConstraintObjectAdded(FConstraintObjectKey InKey = FConstraintObjectKey(),
+			FDeformableConstraintParameters InParamaters = FDeformableConstraintParameters())
+			: FConstraintObjectKey(InKey)
+			, Parameters(InParamaters) {}
+
+		FDeformableConstraintParameters Parameters = FDeformableConstraintParameters();
+	};
+
+	struct FConstraintObjectRemoved : public FConstraintObjectKey
+	{
+		FConstraintObjectRemoved(FConstraintObjectKey InKey = FConstraintObjectKey())
+			: FConstraintObjectKey(InKey) {}
+	};
+
+	struct FConstraintObjectUpdated : public FConstraintObjectKey
+	{
+		FConstraintObjectUpdated(FConstraintObjectKey InKey = FConstraintObjectKey(),
+			FDeformableConstraintParameters InParamaters = FDeformableConstraintParameters())
+			: FConstraintObjectKey(InKey)
+			, Parameters(InParamaters) {}
+
+		FDeformableConstraintParameters Parameters = FDeformableConstraintParameters();
+	};
+
+	struct FConstraintObjectParticleHandel
+	{
+		FConstraintObjectParticleHandel(
+			int32 InSourceParticleIndex = INDEX_NONE,
+			int32 InTargetParticleIndex = INDEX_NONE)
+			: SourceParticleIndex(InSourceParticleIndex)
+			, TargetParticleIndex(InTargetParticleIndex)
+		{}
+
+		int32 SourceParticleIndex = INDEX_NONE;
+		int32 TargetParticleIndex = INDEX_NONE;
+	};
+
+
+	class FConstraintManagerProxy : public FThreadingProxy
+	{
+	public:
+		typedef FThreadingProxy Super;
+
+		FConstraintManagerProxy(UObject* InOwner)
+			: Super(InOwner, TypeName())
+		{}
+
+		static FName TypeName() { return FName("ConstraintManager"); }
+
+		class FConstraintsInputBuffer : public FThreadingProxy::FBuffer
+		{
+			typedef FThreadingProxy::FBuffer Super;
+
+		public:
+			typedef FConstraintManagerProxy Source;
+
+			FConstraintsInputBuffer(
+				const TArray<FConstraintObjectAdded>& InAdded
+				, const TArray<FConstraintObjectRemoved>& InRemoved
+				, const TArray<FConstraintObjectUpdated>& InUpdate
+				, const UObject* InOwner)
+				: Super(InOwner, FConstraintManagerProxy::TypeName())
+				, Added(InAdded)
+				, Removed(InRemoved)
+				, Updated(InUpdate)
+			{}
+
+			TArray<FConstraintObjectAdded> Added;
+			TArray<FConstraintObjectRemoved> Removed;
+			TArray<FConstraintObjectUpdated> Updated;
+		};
+
+		TArray<FConstraintObjectAdded> ConstraintObjectsToAdd;
+		TArray< FConstraintObjectRemoved> ConstraintObjectsToRemove;
+		TMap< FConstraintObjectKey, FConstraintObjectParticleHandel > Constraints;
+	};
+
+}// namespace Chaos::Softs
+
+
diff --git a/Engine/Source/Runtime/Experimental/Chaos/Public/Chaos/Deformable/ChaosDeformableSolver.h b/Engine/Source/Runtime/Experimental/Chaos/Public/Chaos/Deformable/ChaosDeformableSolver.h
--- a/Engine/Source/Runtime/Experimental/Chaos/Public/Chaos/Deformable/ChaosDeformableSolver.h
+++ b/Engine/Source/Runtime/Experimental/Chaos/Public/Chaos/Deformable/ChaosDeformableSolver.h
@@ -24,6 +24,7 @@
 #include "Chaos/Deformable/ChaosDeformableSolverTypes.h"
 #include "Chaos/Deformable/ChaosDeformableSolverProxy.h"
 #include "Chaos/Deformable/ChaosDeformableCollisionsProxy.h"
+#include "Chaos/Deformable/ChaosDeformableConstraintsProxy.h"
 #include "CoreMinimal.h"
 #include "GeometryCollection/ManagedArrayCollection.h"
 #include "Templates/UniquePtr.h"
@@ -60,7 +61,6 @@ namespace Chaos::Softs
 			/* Iteration Advance */
 			CHAOS_API void InitializeSimulationObjects();
 			CHAOS_API void InitializeSimulationObject(FThreadingProxy&);
-			CHAOS_API void InitializeCollisionBodies(FCollisionManagerProxy&);
 			CHAOS_API void InitializeKinematicConstraint();
 			CHAOS_API void InitializeSelfCollisionVariables();
 			CHAOS_API void RemoveSimulationObjects();
@@ -132,15 +132,18 @@ namespace Chaos::Softs
 		CHAOS_API void InitializeWeakConstraint(FFleshThreadingProxy&);
 		CHAOS_API void InitializeKinematicConstraint();
 		CHAOS_API void InitializeCollisionBodies(FCollisionManagerProxy&);
+		CHAOS_API void InitializeConstraintBodies(FConstraintManagerProxy& Proxy);
 		CHAOS_API void InitializeSelfCollisionVariables();
 		CHAOS_API void InitializeGridBasedConstraintVariables();
 		CHAOS_API void InitializeGaussSeidelConstraintVariables();
 		CHAOS_API void InitializeMuscleActivationVariables();
 		CHAOS_API void InitializeMuscleActivation(FFleshThreadingProxy& Proxy);
 		CHAOS_API void UpdateCollisionBodies(FCollisionManagerProxy&, FThreadingProxy::FKey, FSolverReal DeltaTime);
+		CHAOS_API void UpdateConstraintBodies(FConstraintManagerProxy& Proxy, FThreadingProxy::FKey Owner, FSolverReal DeltaTime);
 		CHAOS_API void RemoveSimulationObjects();
 		CHAOS_API TArray<Chaos::TVec3<FSolverReal>> ComputeParticleTargets(const TArray<TArray<int32>>& ParticleIndices);
 
+
 		/*Debug Output*/
 		CHAOS_API void DebugDrawSimulationData();
 		CHAOS_API void DebugDrawTetrahedralParticles(FFleshThreadingProxy& Proxy);
diff --git a/Engine/Source/Runtime/Experimental/Chaos/Public/GeometryCollection/Facades/CollectionTetrahedralFacade.h b/Engine/Source/Runtime/Experimental/Chaos/Public/GeometryCollection/Facades/CollectionTetrahedralFacade.h
new file mode 100644
--- /dev/null
+++ b/Engine/Source/Runtime/Experimental/Chaos/Public/GeometryCollection/Facades/CollectionTetrahedralFacade.h
@@ -0,0 +1,85 @@
+// Copyright Epic Games, Inc. All Rights Reserved.
+
+#pragma once
+#include "Chaos/PBDSoftsEvolutionFwd.h"
+#include "Containers/UnrealString.h"
+#include "Containers/Array.h"
+#include "GeometryCollection/ManagedArrayCollection.h"
+#include "GeometryCollection/ManagedArrayAccessor.h"
+#include "GeometryCollection/Facades/CollectionPositionTargetFacade.h"
+
+namespace GeometryCollection::Facades
+{
+	struct TetrahedralParticleEmbedding
+	{
+		TetrahedralParticleEmbedding(
+			int32 InParticleIndex = INDEX_NONE, 
+			int32 InGeometryIndex = INDEX_NONE,
+			int32 InTetrahedronIndex = INDEX_NONE,
+			TArray<float> InBarycentricWeights = TArray<float>() )
+			: ParticleIndex(InParticleIndex)
+			, GeometryIndex(InGeometryIndex)
+			, TetrahedronIndex(InTetrahedronIndex)
+			, BarycentricWeights(InBarycentricWeights) {}
+
+		int32 ParticleIndex;
+		int32 GeometryIndex;
+		int32 TetrahedronIndex;
+		TArray<float> BarycentricWeights;
+	};
+
+	/**
+	* FTetrahedralFacade
+	*
+	*/
+	class FTetrahedralFacade
+	{
+		const FManagedArrayCollection& ConstCollection;
+		FManagedArrayCollection* Collection = nullptr;
+
+	public:
+
+		/**
+		* FSelectionFacade Constuctor
+		* @param VertixDependencyGroup : GroupName the index attribute is dependent on.
+		*/
+		CHAOS_API FTetrahedralFacade(FManagedArrayCollection& InSelf);
+		CHAOS_API FTetrahedralFacade(const FManagedArrayCollection& InSelf);
+		CHAOS_API virtual ~FTetrahedralFacade();
+
+		/**
+		* Create the facade schema.
+		*/
+		CHAOS_API void DefineSchema();
+
+		/** Returns \c true if the facade is operating on a read-only geometry collection. */
+		bool IsConst() const { return Collection==nullptr; }
+
+		/**
+		* Returns \c true if the Facade defined on the collection, and is initialized to
+		* a valid bindings group.
+		*/
+		CHAOS_API bool IsValid() const;
+
+		/**
+		*  Barycentric intersections with Tetarhedron
+		*/
+		bool Intersection(
+			const TConstArrayView<Chaos::Softs::FSolverVec3>& SamplePositions, 
+			const TConstArrayView<Chaos::Softs::FSolverVec3>& TetarhedronPositions,
+			TArray<TetrahedralParticleEmbedding>& OutIntersections) const;
+
+		//
+		// Attributes
+		//
+
+		const TManagedArrayAccessor<FIntVector4> Tetrahedron;
+		const TManagedArrayAccessor<int32> TetrahedronStart;
+		const TManagedArrayAccessor<int32> TetrahedronCount;
+		const TManagedArrayAccessor<int32> VertexStart;
+		const TManagedArrayAccessor<int32> VertexCount;
+		const TManagedArrayAccessor<FVector3f> Vertex;
+	};
+
+
+}
