Iris fixed crash when destroying rpc parameters
- ReplicationDescriptorBuilder now verifies traits accumulated from struct members against traits on struct/property itself when generating descriptor.

diff --git a/Engine/Plugins/Runtime/ReplicationSystemTestPlugin/Source/Private/Tests/ReplicationSystem/RPC/ReplicatedTestObjectWithRPC.cpp b/Engine/Plugins/Runtime/ReplicationSystemTestPlugin/Source/Private/Tests/ReplicationSystem/RPC/ReplicatedTestObjectWithRPC.cpp
--- a/Engine/Plugins/Runtime/ReplicationSystemTestPlugin/Source/Private/Tests/ReplicationSystem/RPC/ReplicatedTestObjectWithRPC.cpp
+++ b/Engine/Plugins/Runtime/ReplicationSystemTestPlugin/Source/Private/Tests/ReplicationSystem/RPC/ReplicatedTestObjectWithRPC.cpp
@@ -133,6 +133,11 @@ void UTestReplicatedObjectWithRPC::ServerRPCWithParam_Implementation(int32 IntPa
 	ServerRPCWithParamCallOrder = ++CallOrder;
 }
 
+void UTestReplicatedObjectWithRPC::ServerRPCWithParamWithHiddenVirtualBase_Implementation(const FReplicatedStructWithHiddenVirtualBase& StructParam)
+{
+	LocalStructWithHiddenVirtualBaseParam = StructParam;
+}
+
 void UTestReplicatedObjectWithRPC::ServerUnreliableRPC_Implementation()
 {
 	bServerUnreliableRPCCalled = true;
diff --git a/Engine/Plugins/Runtime/ReplicationSystemTestPlugin/Source/Private/Tests/ReplicationSystem/RPC/ReplicatedTestObjectWithRPC.h b/Engine/Plugins/Runtime/ReplicationSystemTestPlugin/Source/Private/Tests/ReplicationSystem/RPC/ReplicatedTestObjectWithRPC.h
--- a/Engine/Plugins/Runtime/ReplicationSystemTestPlugin/Source/Private/Tests/ReplicationSystem/RPC/ReplicatedTestObjectWithRPC.h
+++ b/Engine/Plugins/Runtime/ReplicationSystemTestPlugin/Source/Private/Tests/ReplicationSystem/RPC/ReplicatedTestObjectWithRPC.h
@@ -6,6 +6,32 @@
 
 #include "ReplicatedTestObjectWithRPC.generated.h"
 
+USTRUCT()
+struct FReplicatedStructVirtualBase
+{
+	GENERATED_BODY()
+	virtual ~FReplicatedStructVirtualBase() = default;
+
+protected:
+	int32 NotReplicatedVar;
+};
+
+// Exercise special case found in game code.
+USTRUCT()
+struct FReplicatedStructWithHiddenVirtualBase
+#if CPP
+: public FReplicatedStructVirtualBase
+#endif
+{
+	GENERATED_BODY()
+
+	UPROPERTY()
+	FString TestString;
+
+	UPROPERTY()
+	bool bTestBool;
+};
+
 /**
  *  A test class for testing RPC replication
  */
@@ -62,6 +88,10 @@ public:
 	int32 ServerRPCWithParamCalled = 0;
 	int32 ServerRPCWithParamCallOrder = 0;
 
+	UFUNCTION(Reliable, Server)
+	void ServerRPCWithParamWithHiddenVirtualBase(const FReplicatedStructWithHiddenVirtualBase& StructParam);
+	FReplicatedStructWithHiddenVirtualBase LocalStructWithHiddenVirtualBaseParam;
+	
 	UFUNCTION(Unreliable, Server)
 	void ServerUnreliableRPC();
 	bool bServerUnreliableRPCCalled = false;
diff --git a/Engine/Plugins/Runtime/ReplicationSystemTestPlugin/Source/Private/Tests/ReplicationSystem/RPC/TestObjectRPC.cpp b/Engine/Plugins/Runtime/ReplicationSystemTestPlugin/Source/Private/Tests/ReplicationSystem/RPC/TestObjectRPC.cpp
--- a/Engine/Plugins/Runtime/ReplicationSystemTestPlugin/Source/Private/Tests/ReplicationSystem/RPC/TestObjectRPC.cpp
+++ b/Engine/Plugins/Runtime/ReplicationSystemTestPlugin/Source/Private/Tests/ReplicationSystem/RPC/TestObjectRPC.cpp
@@ -70,6 +70,11 @@ namespace UE::Net::Private
 		// Call an RPC client->server
 		ClientObject->ServerRPCWithParam(IntParam);
 
+		// Call an RPC client->server with hidden virtual base
+		FReplicatedStructWithHiddenVirtualBase TestRpc;
+		TestRpc.TestString = TEXT("TestString");
+		ClientObject->ServerRPCWithParamWithHiddenVirtualBase(TestRpc);
+
 		// Send and deliver client packet
 		Client->UpdateAndSend(Server);
 
diff --git a/Engine/Source/Runtime/Experimental/Iris/Core/Private/Iris/ReplicationState/ReplicationStateDescriptorBuilder.cpp b/Engine/Source/Runtime/Experimental/Iris/Core/Private/Iris/ReplicationState/ReplicationStateDescriptorBuilder.cpp
--- a/Engine/Source/Runtime/Experimental/Iris/Core/Private/Iris/ReplicationState/ReplicationStateDescriptorBuilder.cpp
+++ b/Engine/Source/Runtime/Experimental/Iris/Core/Private/Iris/ReplicationState/ReplicationStateDescriptorBuilder.cpp
@@ -554,16 +554,6 @@ void FPropertyReplicationStateDescriptorBuilder::BuildMemberCache(FBuilderContex
 				Member.Traits |= EMemberPropertyTraits::HasDynamicState;
 			}
 
-			if (EnumHasAnyFlags(Descriptor->Traits, EReplicationStateTraits::IsSourceTriviallyConstructible))
-			{
-				Member.Traits |= EMemberPropertyTraits::IsSourceTriviallyConstructible;
-			}
-
-			if (EnumHasAnyFlags(Descriptor->Traits, EReplicationStateTraits::IsSourceTriviallyDestructible))
-			{
-				Member.Traits |= EMemberPropertyTraits::IsSourceTriviallyDestructible;
-			}
-
 			if (EnumHasAnyFlags(Descriptor->Traits, EReplicationStateTraits::HasConnectionSpecificSerialization))
 			{
 				Member.Traits |= EMemberPropertyTraits::HasConnectionSpecificSerialization;
@@ -1490,15 +1480,25 @@ EReplicationStateTraits FPropertyReplicationStateDescriptorBuilder::BuildReplica
 	// Special traits when all members are replicated
 	if (Context.BuildParams.bAllMembersAreReplicated)
 	{
+		const UScriptStruct* ScriptStruct = Cast<UScriptStruct>(StructInfo.Struct);
+		
 		// We cannot determine whether all properties are trivially constructible or destructible unless they're all replicated.
 		if (EnumHasAnyFlags(SharedPropertyTraits, EMemberPropertyTraits::IsSourceTriviallyConstructible))
 		{
-			Traits |= EReplicationStateTraits::IsSourceTriviallyConstructible;
+			const bool bPropagateSourceIsTrivallyConstructible = !ScriptStruct || ((ScriptStruct->StructFlags & (STRUCT_IsPlainOldData | STRUCT_ZeroConstructor)) != STRUCT_NoFlags);
+			if (bPropagateSourceIsTrivallyConstructible)
+			{
+				Traits |= EReplicationStateTraits::IsSourceTriviallyConstructible;
+			}
 		}
 
 		if (EnumHasAnyFlags(SharedPropertyTraits, EMemberPropertyTraits::IsSourceTriviallyDestructible))
 		{
-			Traits |= EReplicationStateTraits::IsSourceTriviallyDestructible;
+			const bool bPropagateSourceIsSourceTriviallyDestructible = !ScriptStruct || ((ScriptStruct->StructFlags & (STRUCT_IsPlainOldData | STRUCT_NoDestructor)) != STRUCT_NoFlags);
+			if (bPropagateSourceIsSourceTriviallyDestructible)
+			{
+				Traits |= EReplicationStateTraits::IsSourceTriviallyDestructible;
+			}
 		}
 
 		Traits |= EReplicationStateTraits::AllMembersAreReplicated;
