Remove unused NaniteIndex attribute when creating minimal collection for nanite

diff --git a/Engine/Source/Runtime/Experimental/GeometryCollectionEngine/Private/GeometryCollection/GeometryCollectionObject.cpp b/Engine/Source/Runtime/Experimental/GeometryCollectionEngine/Private/GeometryCollection/GeometryCollectionObject.cpp
--- a/Engine/Source/Runtime/Experimental/GeometryCollectionEngine/Private/GeometryCollection/GeometryCollectionObject.cpp
+++ b/Engine/Source/Runtime/Experimental/GeometryCollectionEngine/Private/GeometryCollection/GeometryCollectionObject.cpp
@@ -1500,20 +1500,12 @@ TSharedPtr<FGeometryCollection, ESPMode::ThreadSafe> UGeometryCollection::Genera
 			DuplicateGeometryCollection->AddAttribute<FBox>("BoundingBox", "Transform");
 		}
 
-		if (!DuplicateGeometryCollection->HasAttribute("NaniteIndex", "Transform"))
-		{
-			DuplicateGeometryCollection->AddAttribute<int32>("NaniteIndex", "Transform");
-		}
-
 		const int32 NumTransforms = GeometryCollection->NumElements(FGeometryCollection::TransformGroup);
-		TManagedArray<int32>& NaniteIndex = DuplicateGeometryCollection->ModifyAttribute<int32>("NaniteIndex", "Transform");
 		TManagedArray<FBox>& TransformBounds = DuplicateGeometryCollection->ModifyAttribute<FBox>("BoundingBox", "Transform");
 		const TManagedArray<FBox>& GeometryBounds = GeometryCollection->GetAttribute<FBox>("BoundingBox", "Geometry");
 
-		NaniteIndex.Fill(INDEX_NONE);
 		for (int TransformIndex = 0; TransformIndex < NumTransforms; TransformIndex++)
 		{
-			NaniteIndex[TransformIndex] = TransformToGeometryIndex[TransformIndex];
 			const int32 GeometryIndex = TransformToGeometryIndex[TransformIndex];
 			if (GeometryIndex != INDEX_NONE)
 			{
