Remove bone selected material reference from UGeometryCollection, and clean up some references / accesses to it and legacy methods referring to its (now nonexistent) material ID

diff --git a/Engine/Plugins/Experimental/ChaosEditor/Source/FractureEditor/Public/FractureTool.cpp b/Engine/Plugins/Experimental/ChaosEditor/Source/FractureEditor/Public/FractureTool.cpp
--- a/Engine/Plugins/Experimental/ChaosEditor/Source/FractureEditor/Public/FractureTool.cpp
+++ b/Engine/Plugins/Experimental/ChaosEditor/Source/FractureEditor/Public/FractureTool.cpp
@@ -228,10 +228,6 @@ TArray<FString> UFractureActionTool::GetSelectedComponentMaterialNames(bool bInc
 			int32 NumMaterials = Component->GetNumMaterials();
 			for (int32 MatIdx = 0; MatIdx < NumMaterials; ++MatIdx)
 			{
-				if (MatIdx == Component->GetRestCollection()->GetBoneSelectedMaterialIndex())
-				{
-					continue;
-				}
 				UMaterialInterface* Material = Component->GetMaterial(MatIdx);
 				FString MaterialName = Material ? Material->GetName() : LOCTEXT("NoMaterialName", "None").ToString();
 				MaterialNames.Add(FString::Printf(TEXT("[%d] %s"), MatIdx, *MaterialName));
diff --git a/Engine/Source/Runtime/Experimental/GeometryCollectionEngine/Private/GeometryCollection/GeometryCollectionComponent.cpp b/Engine/Source/Runtime/Experimental/GeometryCollectionEngine/Private/GeometryCollection/GeometryCollectionComponent.cpp
--- a/Engine/Source/Runtime/Experimental/GeometryCollectionEngine/Private/GeometryCollection/GeometryCollectionComponent.cpp
+++ b/Engine/Source/Runtime/Experimental/GeometryCollectionEngine/Private/GeometryCollection/GeometryCollectionComponent.cpp
@@ -5845,11 +5845,6 @@ bool UGeometryCollectionComponent::GetSuppressSelectionMaterial() const
 	return RestCollection->GetGeometryCollection()->HasAttribute("Hide", FGeometryCollection::TransformGroup);
 }
 
-const int UGeometryCollectionComponent::GetBoneSelectedMaterialID() const
-{
-	return RestCollection->GetBoneSelectedMaterialIndex();
-}
-
 FPhysScene_Chaos* UGeometryCollectionComponent::GetInnerChaosScene() const
 {
 	if (ChaosSolverActor)
@@ -6051,18 +6046,24 @@ void UGeometryCollectionComponent::GetUsedMaterials(TArray<UMaterialInterface*>&
 {
 	Super::GetUsedMaterials(OutMaterials, bGetDebugMaterials);
 
-	if (GetRestCollection() && GetRestCollection()->GetBoneSelectedMaterial())
+	if (GetRestCollection())
 	{
-		OutMaterials.Add(GetRestCollection()->GetBoneSelectedMaterial());
+		if (UMaterialInterface* BoneSelectedMaterial = GetRestCollection()->GetBoneSelectedMaterial())
+		{
+			OutMaterials.Add(BoneSelectedMaterial);
+		}
 	}
 }
 
 FMaterialRelevance UGeometryCollectionComponent::GetMaterialRelevance(ERHIFeatureLevel::Type InFeatureLevel) const
 {
 	FMaterialRelevance Result = Super::GetMaterialRelevance(InFeatureLevel);
-	if (RestCollection && RestCollection->GetBoneSelectedMaterial())
+	if (RestCollection)
 	{
-		Result |= RestCollection->GetBoneSelectedMaterial()->GetRelevance_Concurrent(InFeatureLevel);
+		if (UMaterialInterface* BoneSelectedMaterial = GetRestCollection()->GetBoneSelectedMaterial())
+		{
+			Result |= BoneSelectedMaterial->GetRelevance_Concurrent(InFeatureLevel);
+		}
 	}
 	return Result;
 }
diff --git a/Engine/Source/Runtime/Experimental/GeometryCollectionEngine/Private/GeometryCollection/GeometryCollectionObject.cpp b/Engine/Source/Runtime/Experimental/GeometryCollectionEngine/Private/GeometryCollection/GeometryCollectionObject.cpp
--- a/Engine/Source/Runtime/Experimental/GeometryCollectionEngine/Private/GeometryCollection/GeometryCollectionObject.cpp
+++ b/Engine/Source/Runtime/Experimental/GeometryCollectionEngine/Private/GeometryCollection/GeometryCollectionObject.cpp
@@ -656,12 +656,21 @@ void UGeometryCollection::ReindexMaterialSections()
 	InvalidateCollection();
 }
 
+UMaterialInterface* UGeometryCollection::GetBoneSelectedMaterial()
+{
+#if WITH_EDITORONLY_DATA
+	return LoadObject<UMaterialInterface>(nullptr, GetSelectedMaterialPath(), nullptr, LOAD_None, nullptr);
+#else
+	return nullptr;
+#endif
+}
+
 void UGeometryCollection::InitializeMaterials(bool bHasLegacyInternalMaterialsPairs)
 {
 	Modify();
 
 	// Initialize the BoneSelectedMaterial separate from the materials on the collection
-	BoneSelectedMaterial = LoadObject<UMaterialInterface>(nullptr, GetSelectedMaterialPath(), nullptr, LOAD_None, nullptr);
+	UMaterialInterface* BoneSelectedMaterial = GetBoneSelectedMaterial();
 
 	TManagedArray<int32>& MaterialIDs = GeometryCollection->MaterialID;
 
@@ -912,19 +921,6 @@ void UGeometryCollection::Serialize(FArchive& Ar)
 	if ((bIsCookedOrCooking && Ar.IsSaving()) || (Ar.IsCountingMemory() && Ar.IsFilterEditorOnly()))
 	{
 #if WITH_EDITOR
-		if (bIsCookedOrCooking && Ar.IsSaving())
-		{
-			// if we have a valid selection material, let's make sure we replace it with one that will be cooked
-			// this avoid getting warning about the selected material being reference but not cooked
-			const int32 SelectedMaterialIndex = GetBoneSelectedMaterialIndex();
-			if (!Materials.IsEmpty() && Materials.IsValidIndex(SelectedMaterialIndex))
-			{
-				Materials[SelectedMaterialIndex] = Materials[0];
-			}
-			// Likewise remove the direct reference to the BoneSelectedMaterial on cook
-			BoneSelectedMaterial = nullptr;
-		}
-
 		if (bStripOnCook)
 		{
 			// TODO: Since non-nanite path now stores mesh data in cooked build we may be able to unify 
@@ -1220,13 +1216,8 @@ void UGeometryCollection::Serialize(FArchive& Ar)
 
 	if (Ar.IsLoading() && !bIsCookedOrCooking && BoneSelectedMaterialIndex != INDEX_NONE)
 	{
-		BoneSelectedMaterial = LoadObject<UMaterialInterface>(nullptr, GetSelectedMaterialPath(), nullptr, LOAD_None, nullptr);
 		if (Materials.IsValidIndex(BoneSelectedMaterialIndex))
 		{
-			if (!BoneSelectedMaterial)
-			{
-				BoneSelectedMaterial = Materials[BoneSelectedMaterialIndex];
-			}
 			// Remove the material assuming it's the last in the list (otherwise, leave it, as it's not clear why it would be in that state)
 			if (BoneSelectedMaterialIndex == Materials.Num() - 1)
 			{
diff --git a/Engine/Source/Runtime/Experimental/GeometryCollectionEngine/Private/GeometryCollection/GeometryCollectionSceneProxy.cpp b/Engine/Source/Runtime/Experimental/GeometryCollectionEngine/Private/GeometryCollection/GeometryCollectionSceneProxy.cpp
--- a/Engine/Source/Runtime/Experimental/GeometryCollectionEngine/Private/GeometryCollection/GeometryCollectionSceneProxy.cpp
+++ b/Engine/Source/Runtime/Experimental/GeometryCollectionEngine/Private/GeometryCollection/GeometryCollectionSceneProxy.cpp
@@ -170,14 +170,6 @@ FGeometryCollectionSceneProxy::FGeometryCollectionSceneProxy(UGeometryCollection
 		{
 			BoneSelectedMaterial = Component->GetRestCollection()->GetBoneSelectedMaterial();
 		}
-		if (!BoneSelectedMaterial)
-		{
-			int32 LegacyBoneMaterialID = Component->GetBoneSelectedMaterialID();
-			if (Materials.IsValidIndex(LegacyBoneMaterialID))
-			{
-				BoneSelectedMaterial = Materials[LegacyBoneMaterialID];
-			}
-		}
 		if (BoneSelectedMaterial && !BoneSelectedMaterial->CheckMaterialUsage_Concurrent(MATUSAGE_GeometryCollections))
 		{
 			// If we have an invalid BoneSelectedMaterial, switch it back to null to skip its usage in GetDynamicMeshElements below
diff --git a/Engine/Source/Runtime/Experimental/GeometryCollectionEngine/Public/GeometryCollection/GeometryCollectionComponent.h b/Engine/Source/Runtime/Experimental/GeometryCollectionEngine/Public/GeometryCollection/GeometryCollectionComponent.h
--- a/Engine/Source/Runtime/Experimental/GeometryCollectionEngine/Public/GeometryCollection/GeometryCollectionComponent.h
+++ b/Engine/Source/Runtime/Experimental/GeometryCollectionEngine/Public/GeometryCollection/GeometryCollectionComponent.h
@@ -1151,7 +1151,11 @@ public:
 	
 	GEOMETRYCOLLECTIONENGINE_API bool GetSuppressSelectionMaterial() const;
 	
-	GEOMETRYCOLLECTIONENGINE_API const int GetBoneSelectedMaterialID() const;
+	UE_DEPRECATED(5.5, "Bone Selected Material does not have a material ID anymore.")
+	const int GetBoneSelectedMaterialID() const
+	{
+		return INDEX_NONE;
+	}
 	
 #if WITH_EDITORONLY_DATA
 	FORCEINLINE const TArray<int32>& GetSelectedBones() const { return SelectedBones; }
diff --git a/Engine/Source/Runtime/Experimental/GeometryCollectionEngine/Public/GeometryCollection/GeometryCollectionObject.h b/Engine/Source/Runtime/Experimental/GeometryCollectionEngine/Public/GeometryCollection/GeometryCollectionObject.h
--- a/Engine/Source/Runtime/Experimental/GeometryCollectionEngine/Public/GeometryCollection/GeometryCollectionObject.h
+++ b/Engine/Source/Runtime/Experimental/GeometryCollectionEngine/Public/GeometryCollection/GeometryCollectionObject.h
@@ -778,12 +778,11 @@ public:
 	UPROPERTY(meta = (DeprecatedProperty, DeprecationMessage = "Use remove on break feature instead ( Fracture editor tools )."))
 	TArray<TObjectPtr<UMaterialInterface>> RemoveOnFractureMaterials_DEPRECATED;
 
+	UE_DEPRECATED(5.5, "Bone Selected Material does not have a material index anymore.")
 	FORCEINLINE const int32 GetBoneSelectedMaterialIndex() const { return BoneSelectedMaterialIndex; }
 
-	UMaterialInterface* GetBoneSelectedMaterial() const
-	{
-		return BoneSelectedMaterial;
-	}
+	// Get the material to use for rendering bone selections in the editor, or nullptr
+	static UMaterialInterface* GetBoneSelectedMaterial();
 
 	/** Returns the asset path for the automatically populated selected material. */
 	static GEOMETRYCOLLECTIONENGINE_API const TCHAR* GetSelectedMaterialPath();
@@ -899,16 +898,12 @@ private:
 	UPROPERTY(VisibleAnywhere, Transient, Category = "Clustering")
 	TArray<int32> AutoInstanceTransformRemapIndices;
 
-	// #todo(dmp): rename to be consistent BoneSelectedMaterialID?
 	// Legacy index of the bone selected material in the object's Materials array, or INDEX_NONE if it is not stored there.
 	// Note for new objects the bone selected material should not be stored in the Materials array, so this should be INDEX_NONE
+	// The property is kept solely to support deletion of the bone selected material from the materials list of legacy assets
 	UPROPERTY()
 	int32 BoneSelectedMaterialIndex = INDEX_NONE;
 
-	// The material to use for rendering bone selections in the editor, or nullptr
-	UPROPERTY()
-	TObjectPtr<UMaterialInterface> BoneSelectedMaterial = nullptr;
-
 	TSharedPtr<FGeometryCollection, ESPMode::ThreadSafe> GeometryCollection;
 
 	/** Array of user data stored with the asset */
