[NetPhysics] Optimize memory, only create the derived render interpolation API that handles error corrections for proxies that actively use it.
- Also fix bug where if two physics ticks had been processed between game ticks, the error correction interpolation would be wrong due to not decaying the error in discrete steps.
- Also increased the directional decay, making the error decay faster if the object is moving against the error.

diff --git a/Engine/Source/Programs/HeadlessChaos/Private/HeadlessChaosTestRewind.cpp b/Engine/Source/Programs/HeadlessChaos/Private/HeadlessChaosTestRewind.cpp
--- a/Engine/Source/Programs/HeadlessChaos/Private/HeadlessChaosTestRewind.cpp
+++ b/Engine/Source/Programs/HeadlessChaos/Private/HeadlessChaosTestRewind.cpp
@@ -2158,18 +2158,20 @@ namespace ChaosTest {
 						const FReal ExpectedValue = ZStart + ZVel * InterpolatedTime;
 					//	const FReal TargetValue = ZStart + ZVel * NextSimStepTime;
 
-						if (Proxy->GetInterpolationData().IsErrorSmoothing())
+						FProxyInterpolationBase* InterpolationData = Proxy->GetInterpolationData();
+						if (InterpolationData && InterpolationData->IsErrorSmoothing())
 						{
-#if !RENDERINTERP_ERRORVELOCITYSMOOTHING
-							const FReal CorrectionStep = Particle.X()[2] - PrevZ;
-							if (LastCorrectionStep != 0)
+							if (!InterpolationData->IsErrorVelocitySmoothing())
 							{
-								// Make sure we have a linear correction
-								EXPECT_NEAR(LastCorrectionStep, CorrectionStep, 1e-2);
-							}
+								const FReal CorrectionStep = Particle.X()[2] - PrevZ;
+								if (LastCorrectionStep != 0)
+								{
+									// Make sure we have a linear correction
+									EXPECT_NEAR(LastCorrectionStep, CorrectionStep, 1e-2);
+								}
 
-							LastCorrectionStep = CorrectionStep;
-#endif
+								LastCorrectionStep = CorrectionStep;
+							}
 						}
 						else
 						{
@@ -4346,7 +4348,10 @@ namespace ChaosTest {
 			Particle2.SetGeometry(Sphere2);
 			Particle2.SetV(FVec3(0, 0, 1));
 			Particle2.SetGravityEnabled(false);
-			Proxy2->GetInterpolationData().SetInterpChannel_External(1);
+			if (FProxyInterpolationBase* InterpolationData = Proxy2->GetInterpolationData())
+			{
+				InterpolationData->SetInterpChannel_External(1);
+			}
 			Solver->RegisterObject(Proxy2);
 
 			FReal Time = 0;
diff --git a/Engine/Source/Runtime/Experimental/Chaos/Private/Chaos/Framework/ChaosResultsManager.cpp b/Engine/Source/Runtime/Experimental/Chaos/Private/Chaos/Framework/ChaosResultsManager.cpp
--- a/Engine/Source/Runtime/Experimental/Chaos/Private/Chaos/Framework/ChaosResultsManager.cpp
+++ b/Engine/Source/Runtime/Experimental/Chaos/Private/Chaos/Framework/ChaosResultsManager.cpp
@@ -62,7 +62,11 @@ namespace Chaos
 		{
 			if (TProxyType* Proxy = Data.Prev.GetProxy())
 			{
-				Proxy->GetInterpolationData().SetPullDataInterpIdx_External(INDEX_NONE);
+				FProxyInterpolationBase* InterpolationData = Proxy->GetInterpolationData();
+				if (InterpolationData)
+				{
+					InterpolationData->SetPullDataInterpIdx_External(INDEX_NONE);
+				}
 			}
 		}
 		Interpolations.Reset();
@@ -85,19 +89,19 @@ namespace Chaos
 		{
 			if (auto Proxy = Data.GetProxy())
 			{
-				FProxyInterpolationBase& InterpolationData = Proxy->GetInterpolationData();
+				FProxyInterpolationBase* InterpolationData = Proxy->GetInterpolationData();
 				
 				//If proxy is not associated with this channel, do nothing
-				if (InterpolationData.GetInterpChannel_External() != ChannelIdx)
+				if (!InterpolationData || InterpolationData->GetInterpChannel_External() != ChannelIdx)
 				{
 					continue;
 				}
 
-				int32 DataIdx = InterpolationData.GetPullDataInterpIdx_External();
+				int32 DataIdx = InterpolationData->GetPullDataInterpIdx_External();
 				if(DataIdx == INDEX_NONE)
 				{
 					DataIdx = Interpolations.AddDefaulted(1);
-					InterpolationData.SetPullDataInterpIdx_External(DataIdx);
+					InterpolationData->SetPullDataInterpIdx_External(DataIdx);
 
 					if(Mode == ESetPrevNextDataMode::Next)
 					{
@@ -150,8 +154,9 @@ namespace Chaos
 			{
 				if(FSingleParticlePhysicsProxy* Proxy = Data.GetProxy())
 				{
+					FProxyInterpolationBase* InterpolationData = Proxy->GetInterpolationData();
 					// only if the proxy is associated with this channel 
-					if(Proxy->GetInterpolationData().GetInterpChannel_External() == ChannelIdx)
+					if(InterpolationData && InterpolationData->GetInterpChannel_External() == ChannelIdx)
 					{
 						//update leash target
 						if(FDirtyRigidParticleData* ResimTarget = ParticleToResimTarget.Find(Proxy))
@@ -293,14 +298,14 @@ namespace Chaos
 		{
 			FSingleParticlePhysicsProxy* Proxy = Itr.Key;
 
-			FProxyInterpolationBase& InterpolationData = Proxy->GetInterpolationData();
-			if (InterpolationData.IsErrorSmoothing())
+			FProxyInterpolationBase* InterpolationData = Proxy->GetInterpolationData();
+			if (InterpolationData && InterpolationData->IsErrorSmoothing())
 			{
-				if (InterpolationData.GetPullDataInterpIdx_External() == INDEX_NONE)	//not in results array
+				if (InterpolationData->GetPullDataInterpIdx_External() == INDEX_NONE)	//not in results array
 				{
 					//still need to interpolate, so add to results array
 					const int32 DataIdx = Results.RigidInterpolations.AddDefaulted(1);
-					InterpolationData.SetPullDataInterpIdx_External(DataIdx);
+					InterpolationData->SetPullDataInterpIdx_External(DataIdx);
 					FChaosRigidInterpolationData& RigidData = Results.RigidInterpolations[DataIdx];
 
 					RigidData.Next = Itr.Value;			//not dirty from sim, so just use whatever last next was
@@ -399,9 +404,9 @@ namespace Chaos
 		{
 			if (FSingleParticlePhysicsProxy* ResimProxy = ResimDirty.GetProxy())
 			{
-				FProxyInterpolationBase& InterpolationData = ResimProxy->GetInterpolationData(); 
+				FProxyInterpolationBase* InterpolationData = ResimProxy->GetInterpolationData(); 
 				//Mark as resim only if proxy is owned by this channel
-				if(InterpolationData.GetInterpChannel_External() == ChannelIdx)
+				if(InterpolationData && InterpolationData->GetInterpChannel_External() == ChannelIdx)
 				{
 					ParticleToResimTarget.FindOrAdd(ResimProxy) = ResimDirty;
 				}
diff --git a/Engine/Source/Runtime/Experimental/Chaos/Private/Framework/PhysicsProxyBase.cpp b/Engine/Source/Runtime/Experimental/Chaos/Private/Framework/PhysicsProxyBase.cpp
--- a/Engine/Source/Runtime/Experimental/Chaos/Private/Framework/PhysicsProxyBase.cpp
+++ b/Engine/Source/Runtime/Experimental/Chaos/Private/Framework/PhysicsProxyBase.cpp
@@ -19,6 +19,9 @@ namespace RenderInterpolationCVars
 
 	bool bRenderInterpDebugDraw = false;
 	static FAutoConsoleVariableRef CVarRenderInterpDebugDraw(TEXT("p.RenderInterp.DebugDraw"), bRenderInterpDebugDraw, TEXT("Draw debug lines for physics render interpolation, also needs p.Chaos.DebugDraw.Enabled set"));
+
+	bool bRenderInterpErrorVelocityCorrection = false;
+	static FAutoConsoleVariableRef CVarRenderInterpErrorVelocityCorrection(TEXT("p.RenderInterp.ErrorVelocityCorrection"), bRenderInterpErrorVelocityCorrection, TEXT("EXPERIMENTAL - Take incoming velocity into consideration when performing render interpolation, the correction will be more organic but might result in clipping and it's heavier for memory and CPU."));
 }
 
 IPhysicsProxyBase::~IPhysicsProxyBase()
diff --git a/Engine/Source/Runtime/Experimental/Chaos/Private/PhysicsProxy/ClusterUnionPhysicsProxy.cpp b/Engine/Source/Runtime/Experimental/Chaos/Private/PhysicsProxy/ClusterUnionPhysicsProxy.cpp
--- a/Engine/Source/Runtime/Experimental/Chaos/Private/PhysicsProxy/ClusterUnionPhysicsProxy.cpp
+++ b/Engine/Source/Runtime/Experimental/Chaos/Private/PhysicsProxy/ClusterUnionPhysicsProxy.cpp
@@ -75,6 +75,7 @@ namespace Chaos
 		, ClusterParameters(InParameters)
 		, InitData(InInitData)
 	{
+		InterpolationData = MakeUnique<FProxyInterpolationBase>();
 	}
 
 	void FClusterUnionPhysicsProxy::Initialize_External()
@@ -478,17 +479,22 @@ namespace Chaos
 			return false;
 		}
 
+		FProxyInterpolationBase* InterpData = GetInterpolationData();
 		if (Error)
 		{
-			const FReal ErrorMagSq = Error->ErrorX.SizeSquared();
+			InterpData = GetOrCreateErrorInterpolationData<FProxyInterpolationError>();
+			check(InterpData);
+
+			// If error is within interpolation limit, set the number of physics frames to interpolate over, else leave at 0 frames to instantly correct the error.
 			const FReal MaxErrorCorrection = RenderInterpolationCVars::RenderInterpMaximumErrorCorrectionBeforeSnapping;
 			int32 RenderInterpErrorCorrectionDurationTicks = 0;
-			if (ErrorMagSq < MaxErrorCorrection * MaxErrorCorrection)
+			if (Error->ErrorX.SizeSquared() < (MaxErrorCorrection * MaxErrorCorrection))
 			{
 				RenderInterpErrorCorrectionDurationTicks = FMath::FloorToInt32(RenderInterpolationCVars::RenderInterpErrorCorrectionDuration / AsyncFixedTimeStep); // Convert duration from seconds to simulation ticks
 			}
-			InterpolationData.AccumlateErrorXR(Error->ErrorX, Error->ErrorR, SolverSyncTimestamp, RenderInterpErrorCorrectionDurationTicks);
+			InterpData->AccumlateErrorXR(Error->ErrorX, Error->ErrorR, SolverSyncTimestamp, RenderInterpErrorCorrectionDurationTicks);
 		}
+		const bool bHasInterpolationData = InterpData != nullptr;
 
 		SyncedData_External.bIsAnchored = CurrentPullData.bIsAnchored;
 		SyncedData_External.bDidSyncGeometry = false;
@@ -559,19 +565,22 @@ namespace Chaos
 				return OverwriteProperty.Timestamp <= SolverSyncTimestamp ? (OverwriteProperty.Timestamp < SolverSyncTimestamp ? &Prev : &OverwriteProperty.Value) : nullptr;
 			};
 
-			const bool bIsReplicationErrorSmoothing = InterpolationData.IsErrorSmoothing();
+			if (bHasInterpolationData)
+			{
+				InterpData->UpdateError(SolverSyncTimestamp, AsyncFixedTimeStep);
+			}
+			const bool bIsReplicationErrorSmoothing = bHasInterpolationData ? InterpData->IsErrorSmoothing() : false;
 			bool DirectionalDecayPerformed = false;
-			InterpolationData.UpdateError(SolverSyncTimestamp, AsyncFixedTimeStep);
 
 			if (const FVec3* Prev = LerpHelper(PullData.X, ProxyTimestamp->OverWriteX))
 			{
-				if (RenderInterpolationCVars::RenderInterpErrorDirectionalDecayMultiplier > 0.0f)
+				if (bHasInterpolationData && RenderInterpolationCVars::RenderInterpErrorDirectionalDecayMultiplier > 0.0f)
 				{
-					DirectionalDecayPerformed = InterpolationData.DirectionalDecay(NextPullData->X - *Prev);
+					DirectionalDecayPerformed = InterpData->DirectionalDecay(NextPullData->X - *Prev);
 				}
 
 				const FVec3 NewX = bIsReplicationErrorSmoothing ?
-					FMath::Lerp(*Prev, NextPullData->X, *Alpha) + InterpolationData.GetErrorX(*Alpha) :
+					FMath::Lerp(*Prev, NextPullData->X, *Alpha) + InterpData->GetErrorX(*Alpha) :
 					FMath::Lerp(*Prev, NextPullData->X, *Alpha);
 
 				Particle_External->SetX(NewX, false);
@@ -580,7 +589,7 @@ namespace Chaos
 			if (const FQuat* Prev = LerpHelper(PullData.R, ProxyTimestamp->OverWriteR))
 			{
 				const FQuat NewR = bIsReplicationErrorSmoothing ? 
-					InterpolationData.GetErrorR(*Alpha) * FMath::Lerp(*Prev, NextPullData->R, *Alpha) : 
+					InterpData->GetErrorR(*Alpha) * FMath::Lerp(*Prev, NextPullData->R, *Alpha) : 
 					FMath::Lerp(*Prev, NextPullData->R, *Alpha);
 		
 				Particle_External->SetR(NewR, false);
@@ -610,10 +619,10 @@ namespace Chaos
 					if (Error)
 					{
 						Chaos::FDebugDrawQueue::GetInstance().DrawDebugBox(PullData.X, FVector(4, 2, 2), PullData.R, FColor::Red, false, 5.f, 0, 0.5f);
-						Chaos::FDebugDrawQueue::GetInstance().DrawDebugDirectionalArrow(PullData.X, (PullData.X + InterpolationData.GetErrorX(0)), 1, FColor::Red, false, 5.0f, 0, 0.5f);
+						Chaos::FDebugDrawQueue::GetInstance().DrawDebugDirectionalArrow(PullData.X, (PullData.X + InterpData->GetErrorX(0)), 1, FColor::Red, false, 5.0f, 0, 0.5f);
 					}
 
-					Chaos::FDebugDrawQueue::GetInstance().DrawDebugDirectionalArrow((Particle_External->GetX() - InterpolationData.GetErrorX(*Alpha)), Particle_External->GetX(), 1, FColor::Blue, false, 5.0f, 0, 0.5f);
+					Chaos::FDebugDrawQueue::GetInstance().DrawDebugDirectionalArrow((Particle_External->GetX() - InterpData->GetErrorX(*Alpha)), Particle_External->GetX(), 1, FColor::Blue, false, 5.0f, 0, 0.5f);
 				}
 			}
 #endif // CHAOS_DEBUG_DRAW
diff --git a/Engine/Source/Runtime/Experimental/Chaos/Private/PhysicsProxy/GeometryCollectionPhysicsProxy.cpp b/Engine/Source/Runtime/Experimental/Chaos/Private/PhysicsProxy/GeometryCollectionPhysicsProxy.cpp
--- a/Engine/Source/Runtime/Experimental/Chaos/Private/PhysicsProxy/GeometryCollectionPhysicsProxy.cpp
+++ b/Engine/Source/Runtime/Experimental/Chaos/Private/PhysicsProxy/GeometryCollectionPhysicsProxy.cpp
@@ -534,6 +534,7 @@ FGeometryCollectionPhysicsProxy::FGeometryCollectionPhysicsProxy(
 {
 	// We rely on a guarded buffer.
 	check(BufferMode == Chaos::EMultiBufferMode::TripleGuarded);
+	InterpolationData = MakeUnique<FProxyInterpolationBase>();
 }
 
 
@@ -4295,17 +4296,22 @@ bool FGeometryCollectionPhysicsProxy::PullFromPhysicsState(const Chaos::FDirtyGe
 		return false;
 	}
 
+	FProxyInterpolationBase* InterpData = GetInterpolationData();
 	if (Error)
 	{
-		const Chaos::FReal ErrorMagSq = Error->ErrorX.SizeSquared();
+		InterpData = GetOrCreateErrorInterpolationData<FProxyInterpolationError>();
+		check(InterpData);
+
+		// If error is within interpolation limit, set the number of physics frames to interpolate over, else leave at 0 frames to instantly correct the error.
 		const Chaos::FReal MaxErrorCorrection = RenderInterpolationCVars::RenderInterpMaximumErrorCorrectionBeforeSnapping;
 		int32 RenderInterpErrorCorrectionDurationTicks = 0;
-		if (ErrorMagSq < MaxErrorCorrection * MaxErrorCorrection)
+		if (Error->ErrorX.SizeSquared() < (MaxErrorCorrection * MaxErrorCorrection))
 		{
 			RenderInterpErrorCorrectionDurationTicks = FMath::FloorToInt32(RenderInterpolationCVars::RenderInterpErrorCorrectionDuration / AsyncFixedTimeStep); // Convert duration from seconds to simulation ticks
 		}
-		InterpolationData.AccumlateErrorXR(Error->ErrorX, Error->ErrorR, SolverSyncTimestamp, RenderInterpErrorCorrectionDurationTicks);
+		InterpData->AccumlateErrorXR(Error->ErrorX, Error->ErrorR, SolverSyncTimestamp, RenderInterpErrorCorrectionDurationTicks);
 	}
+	const bool bHasInterpolationData = InterpData != nullptr;
 
 	check(NumTransforms == GameThreadCollection.GetNumTransforms());
 	const bool bNeedInterpolation = (NextPullData != nullptr);
@@ -4333,10 +4339,14 @@ bool FGeometryCollectionPhysicsProxy::PullFromPhysicsState(const Chaos::FDirtyGe
 			const FGeometryCollectionResults& PrevResults = PullData.Results();
 			const FGeometryCollectionResults& NextResults = NextPullData->Results();
 
-			InterpolationData.UpdateError(SolverSyncTimestamp, AsyncFixedTimeStep);
-			if (RenderInterpolationCVars::RenderInterpErrorDirectionalDecayMultiplier > 0.0f && PrevResults.GetNumEntries() && NextResults.GetNumEntries())
+			if (bHasInterpolationData)
 			{
-				InterpolationData.DirectionalDecay(NextResults.GetPositions(0).ParticleX - PrevResults.GetPositions(0).ParticleX);
+				InterpData->UpdateError(SolverSyncTimestamp, AsyncFixedTimeStep);
+
+				if (RenderInterpolationCVars::RenderInterpErrorDirectionalDecayMultiplier > 0.0f && PrevResults.GetNumEntries() && NextResults.GetNumEntries())
+				{
+					InterpData->DirectionalDecay(NextResults.GetPositions(0).ParticleX - PrevResults.GetPositions(0).ParticleX);
+				}
 			}
 
 			TManagedArray<FVector3f>* LinearVelocities = GameThreadCollection.GetLinearVelocitiesAttribute();
@@ -4377,10 +4387,10 @@ bool FGeometryCollectionPhysicsProxy::PullFromPhysicsState(const Chaos::FDirtyGe
 					Chaos::FRotation3 NewR;
 					ResultInterpolator.GetPositions(NewX, NewR);
 
-					if (InterpolationData.IsErrorSmoothing())
+					if (bHasInterpolationData && InterpData->IsErrorSmoothing())
 					{
-						NewX = NewX + InterpolationData.GetErrorX(*Alpha);
-						NewR = InterpolationData.GetErrorR(*Alpha) * NewR;
+						NewX = NewX + InterpData->GetErrorX(*Alpha);
+						NewR = InterpData->GetErrorR(*Alpha) * NewR;
 					}
 
 					const bool XRModified = UpdateGTParticleXR(GTParticle, NewX, NewR);
diff --git a/Engine/Source/Runtime/Experimental/Chaos/Private/PhysicsProxy/SingleParticlePhysicsProxy.cpp b/Engine/Source/Runtime/Experimental/Chaos/Private/PhysicsProxy/SingleParticlePhysicsProxy.cpp
--- a/Engine/Source/Runtime/Experimental/Chaos/Private/PhysicsProxy/SingleParticlePhysicsProxy.cpp
+++ b/Engine/Source/Runtime/Experimental/Chaos/Private/PhysicsProxy/SingleParticlePhysicsProxy.cpp
@@ -40,6 +40,7 @@ FSingleParticlePhysicsProxy::FSingleParticlePhysicsProxy(TUniquePtr<PARTICLE_TYP
 {
 	Particle->SetProxy(this);
 	Reference = FPhysicsObjectFactory::CreatePhysicsObject(this);
+	InterpolationData = MakeUnique<FProxyInterpolationBase>();
 }
 
 
@@ -284,25 +285,33 @@ bool FSingleParticlePhysicsProxy::PullFromPhysicsState(const Chaos::FDirtyRigidP
 		// logic in one place is changed, it should be checked in the other place too.
 		bool bUpdatePositionFromSimulation = ShouldUpdateTransformFromSimulation(*Rigid);
 		const FSingleParticleProxyTimestamp* ProxyTimestamp = PullData.GetTimestamp();
-		
-#if RENDERINTERP_ERRORVELOCITYSMOOTHING
-		const int32 RenderInterpErrorVelocitySmoothingDurationTicks = FMath::FloorToInt32(GetRenderInterpErrorVelocitySmoothingDuration() / AsyncFixedTimeStep); // Convert duration from seconds to simulation ticks
-#endif
 
+		FProxyInterpolationBase* InterpData = GetInterpolationData();
 		if (Error)
 		{
-			const FReal ErrorMagSq = Error->ErrorX.SizeSquared();
-			const FReal MaxErrorCorrection = RenderInterpolationCVars::RenderInterpMaximumErrorCorrectionBeforeSnapping;
+			if (RenderInterpolationCVars::bRenderInterpErrorVelocityCorrection)
+			{
+				InterpData = GetOrCreateErrorInterpolationData<FProxyInterpolationErrorVelocity>();
+			}
+			else
+			{
+				InterpData = GetOrCreateErrorInterpolationData<FProxyInterpolationError>();
+			}
+			check(InterpData);
+
+			// If error is within interpolation limit, set the number of physics frames to interpolate over, else leave at 0 frames to instantly correct the error.
 			int32 RenderInterpErrorCorrectionDurationTicks = 0;
-			if (ErrorMagSq < MaxErrorCorrection * MaxErrorCorrection)
+			const FReal MaxErrorCorrection = RenderInterpolationCVars::RenderInterpMaximumErrorCorrectionBeforeSnapping;
+			if (Error->ErrorX.SizeSquared() < (MaxErrorCorrection * MaxErrorCorrection))
 			{
 				RenderInterpErrorCorrectionDurationTicks = FMath::FloorToInt32(RenderInterpolationCVars::RenderInterpErrorCorrectionDuration / AsyncFixedTimeStep); // Convert duration from seconds to simulation ticks
 			}
-			InterpolationData.AccumlateErrorXR(Error->ErrorX, Error->ErrorR, SolverSyncTimestamp, RenderInterpErrorCorrectionDurationTicks);
-#if RENDERINTERP_ERRORVELOCITYSMOOTHING
-			InterpolationData.SetVelocitySmoothing(Rigid->V(), Rigid->X(), RenderInterpErrorVelocitySmoothingDurationTicks);
-#endif
+			InterpData->AccumlateErrorXR(Error->ErrorX, Error->ErrorR, SolverSyncTimestamp, RenderInterpErrorCorrectionDurationTicks);
+			
+			const int32 RenderInterpErrorVelocitySmoothingDurationTicks = FMath::FloorToInt32(RenderInterpolationCVars::RenderInterpErrorVelocitySmoothingDuration / AsyncFixedTimeStep); // Convert duration from seconds to simulation ticks
+			InterpData->SetVelocitySmoothing(Rigid->V(), Rigid->X(), RenderInterpErrorVelocitySmoothingDurationTicks);
 		}
+		const bool bHasInterpolationData = InterpData != nullptr;
 
 		if (NextPullData)
 		{
@@ -318,12 +327,13 @@ bool FSingleParticlePhysicsProxy::PullFromPhysicsState(const Chaos::FDirtyRigidP
 
 			if (bUpdatePositionFromSimulation)
 			{
-				const bool bIsReplicationErrorSmoothing = InterpolationData.IsErrorSmoothing();
+				if (bHasInterpolationData)
+				{
+					InterpData->UpdateError(SolverSyncTimestamp, AsyncFixedTimeStep);
+				}
+
+				const bool bIsReplicationErrorSmoothing = bHasInterpolationData ? InterpData->IsErrorSmoothing() : false;
 				bool DirectionalDecayPerformed = false;
-#if RENDERINTERP_ERRORVELOCITYSMOOTHING
-				const bool bIsErrorVelocitySmoothing = InterpolationData.IsErrorVelocitySmoothing();
-#endif
-				InterpolationData.UpdateError(SolverSyncTimestamp, AsyncFixedTimeStep);
 
 				if (const FVec3* Prev = LerpHelper(PullData.X, ProxyTimestamp->OverWriteX))
 				{
@@ -332,25 +342,25 @@ bool FSingleParticlePhysicsProxy::PullFromPhysicsState(const Chaos::FDirtyRigidP
 					{
 						if (RenderInterpolationCVars::RenderInterpErrorDirectionalDecayMultiplier > 0.0f)
 						{
-							DirectionalDecayPerformed = InterpolationData.DirectionalDecay(NextPullData->X - *Prev);
+							DirectionalDecayPerformed = InterpData->DirectionalDecay(NextPullData->X - *Prev);
 						}
 
-						Target += InterpolationData.GetErrorX(*Alpha);
+						Target += InterpData->GetErrorX(*Alpha);
 
-#if RENDERINTERP_ERRORVELOCITYSMOOTHING
-						if (bIsErrorVelocitySmoothing)
+						if (InterpData->IsErrorVelocitySmoothing())
 						{
 #if CHAOS_DEBUG_DRAW
-							if (!!RenderInterpDebugDraw)
+							if (RenderInterpolationCVars::bRenderInterpDebugDraw)
 							{
-								Chaos::FDebugDrawQueue::GetInstance().DrawDebugDirectionalArrow((Target - InterpolationData.GetErrorX(*Alpha)), Target, 1, FColor::Blue, false, 5.0f, 0, 0.5f);
+								Chaos::FDebugDrawQueue::GetInstance().DrawDebugDirectionalArrow((Target - InterpData->GetErrorX(*Alpha)), Target, 1, FColor::Blue, false, 5.0f, 0, 0.5f);
+								Chaos::FDebugDrawQueue::GetInstance().DrawDebugDirectionalArrow(InterpData->GetErrorVelocitySmoothingX(*Alpha), Target, 1, FColor::Blue, false, 5.0f, 0, 0.5f);
 								Chaos::FDebugDrawQueue::GetInstance().DrawDebugBox(Target, FVector(2, 1, 1), Rigid->R(), FColor::Cyan, false, 5.f, 0, 0.25f);
 							}
 #endif // CHAOS_DEBUG_DRAW
 
-							Target = FMath::Lerp(Target, InterpolationData.GetErrorVelocitySmoothingX(*Alpha), InterpolationData.GetErrorVelocitySmoothingAlpha(RenderInterpErrorVelocitySmoothingDurationTicks));
+							const int32 RenderInterpErrorVelocitySmoothingDurationTicks = FMath::FloorToInt32(RenderInterpolationCVars::RenderInterpErrorVelocitySmoothingDuration / AsyncFixedTimeStep); // Convert duration from seconds to simulation ticks
+							Target = FMath::Lerp(Target, InterpData->GetErrorVelocitySmoothingX(*Alpha), InterpData->GetErrorVelocitySmoothingAlpha(RenderInterpErrorVelocitySmoothingDurationTicks));
 						}
-#endif // RENDERINTERP_ERRORVELOCITYSMOOTHING
 					}
 
 					Rigid->SetX(Target, false);					
@@ -361,7 +371,7 @@ bool FSingleParticlePhysicsProxy::PullFromPhysicsState(const Chaos::FDirtyRigidP
 					FQuat Target = FMath::Lerp(*Prev, NextPullData->R, *Alpha);
 					if (bIsReplicationErrorSmoothing)
 					{
-						Target = InterpolationData.GetErrorR(*Alpha) * Target;
+						Target = InterpData->GetErrorR(*Alpha) * Target;
 					}
 					Rigid->SetR(Target, false);
 				}
@@ -378,18 +388,16 @@ bool FSingleParticlePhysicsProxy::PullFromPhysicsState(const Chaos::FDirtyRigidP
 						if (Error)
 						{
 							Chaos::FDebugDrawQueue::GetInstance().DrawDebugBox(PullData.X, FVector(4, 2, 2), PullData.R, FColor::Red, false, 5.f, 0, 0.5f);
-							Chaos::FDebugDrawQueue::GetInstance().DrawDebugDirectionalArrow(PullData.X, (PullData.X + InterpolationData.GetErrorX(0)), 1, FColor::Red, false, 5.0f, 0, 0.5f);
+							Chaos::FDebugDrawQueue::GetInstance().DrawDebugDirectionalArrow(PullData.X, (PullData.X + InterpData->GetErrorX(0)), 1, FColor::Red, false, 5.0f, 0, 0.5f);
 						}
 
-#if RENDERINTERP_ERRORVELOCITYSMOOTHING
-						if (bIsErrorVelocitySmoothing)
+						if (InterpData->IsErrorVelocitySmoothing())
 						{
-							Chaos::FDebugDrawQueue::GetInstance().DrawDebugBox(InterpolationData.GetErrorVelocitySmoothingX(*Alpha), FVector(2, 2, 2), Rigid->R(), FColor::Purple, false, 5.f, 0, 0.5f);
+							Chaos::FDebugDrawQueue::GetInstance().DrawDebugBox(InterpData->GetErrorVelocitySmoothingX(*Alpha), FVector(2, 2, 2), Rigid->R(), FColor::Purple, false, 5.f, 0, 0.5f);
 						}
 						else
-#endif // RENDERINTERP_ERRORVELOCITYSMOOTHING
 						{
-							Chaos::FDebugDrawQueue::GetInstance().DrawDebugDirectionalArrow((Rigid->X() - InterpolationData.GetErrorX(*Alpha)), Rigid->X(), 1, FColor::Blue, false, 5.0f, 0, 0.5f);
+							Chaos::FDebugDrawQueue::GetInstance().DrawDebugDirectionalArrow((Rigid->X() - InterpData->GetErrorX(*Alpha)), Rigid->X(), 1, FColor::Blue, false, 5.0f, 0, 0.5f);
 						}
 					}
 				}
diff --git a/Engine/Source/Runtime/Experimental/Chaos/Public/Chaos/Framework/PhysicsProxyBase.h b/Engine/Source/Runtime/Experimental/Chaos/Public/Chaos/Framework/PhysicsProxyBase.h
--- a/Engine/Source/Runtime/Experimental/Chaos/Public/Chaos/Framework/PhysicsProxyBase.h
+++ b/Engine/Source/Runtime/Experimental/Chaos/Public/Chaos/Framework/PhysicsProxyBase.h
@@ -34,8 +34,16 @@ namespace RenderInterpolationCVars
 	extern float RenderInterpErrorVelocitySmoothingDuration;
 	extern float RenderInterpErrorDirectionalDecayMultiplier;
 	extern bool bRenderInterpDebugDraw;
+	extern bool bRenderInterpErrorVelocityCorrection;
 }
 
+enum class EProxyInterpolationType : uint32
+{
+	Base = 0,
+	ErrorLinear = 1,
+	ErrorVelocity = 2
+};
+
 struct FProxyTimestampBase
 {
 	bool bDeleted = false;
@@ -195,10 +203,11 @@ struct PhysicsProxyWrapper
 
 struct FProxyInterpolationBase
 {
-	FProxyInterpolationBase()
-		: PullDataInterpIdx_External(INDEX_NONE)
-		, InterpChannel_External(0)
+	FProxyInterpolationBase(const int32 PullDataInterpIdx = INDEX_NONE, const int32 InterpChannel = 0)
+		: PullDataInterpIdx_External(PullDataInterpIdx)
+		, InterpChannel_External(InterpChannel)
 	{}
+	virtual ~FProxyInterpolationBase() { }
 
 	int32 GetPullDataInterpIdx_External() const { return PullDataInterpIdx_External; }
 	void SetPullDataInterpIdx_External(const int32 Idx) { PullDataInterpIdx_External = Idx; }
@@ -206,11 +215,48 @@ struct FProxyInterpolationBase
 	int32 GetInterpChannel_External() const { return InterpChannel_External; }
 	void SetInterpChannel_External(const int32 Channel) { InterpChannel_External = Channel; }
 
-	virtual bool IsErrorSmoothing() { return false; }
-
 protected:
 	int32 PullDataInterpIdx_External;
 	int32 InterpChannel_External;
+
+public:
+	// --- Error correction interpolation API ---
+
+	/** This interpolation structs type */
+	static const EProxyInterpolationType InterpolationType = EProxyInterpolationType::Base;
+
+	/** Get this interpolation structs type */
+	virtual const EProxyInterpolationType GetInterpolationType() const { return FProxyInterpolationBase::InterpolationType; }
+
+	/** If currently correcting an error through interpolation */
+	virtual const bool IsErrorSmoothing() const { return false; }
+
+	/** Get the position of the current error correction, taking current Alpha between GT and PT into account */
+	virtual const Chaos::FVec3 GetErrorX(const Chaos::FRealSingle Alpha) const { return Chaos::FVec3::ZeroVector; }
+
+	/** Get the rotation of the current error correction, taking current Alpha between GT and PT into account */
+	virtual const FQuat GetErrorR(const Chaos::FRealSingle Alpha) const { return FQuat::Identity; }
+
+	/** Add X and R error onto current error to correct through interpolation */
+	virtual void AccumlateErrorXR(const Chaos::FVec3 X, const FQuat R, const int32 CurrentSimTick, const int32 ErrorSmoothDuration) { };
+
+	/** Tick current error data and decay error */
+	virtual const bool UpdateError(const int32 CurrentSimTick, const Chaos::FReal AsyncFixedTimeStep) { return false; }
+
+	/** EXPERIMENTAL - Decay error based on moved direction and distance */
+	virtual const bool DirectionalDecay(Chaos::FVec3 Direction) { return false; }
+
+	/** EXPERIMENTAL - If currently correcting error while taking velocity into account */
+	virtual const bool IsErrorVelocitySmoothing() const { return false; }
+
+	/** EXPERIMENTAL - Returns the Alpha of how much to take previous velocity into account, used to lerp from linear extrapolation to the predicted position based on previous velocity */
+	virtual const Chaos::FRealSingle GetErrorVelocitySmoothingAlpha(const int32 ErrorVelocitySmoothDuration) const { return 0; }
+
+	/** EXPERIMENTAL - Get the position of the velocity-based correction, taking current Alpha between GT and PT into account*/
+	virtual const Chaos::FVec3 GetErrorVelocitySmoothingX(const Chaos::FRealSingle Alpha) const { return Chaos::FVec3::ZeroVector; }
+
+	/** EXPERIMENTAL - Register the current velocity and position for use in velocity correction calculations */
+	virtual void SetVelocitySmoothing(const Chaos::FVec3 CurrV, const Chaos::FVec3 CurrX, const int32 ErrorVelocitySmoothDuration) { };
 };
 
 // Render interpolation that can correct errors from resimulation / repositions through a linear decay over N simulation tick.
@@ -218,15 +264,20 @@ struct FProxyInterpolationError : FProxyInterpolationBase
 {
 	using Super = FProxyInterpolationBase;
 
-	FProxyInterpolationError() : Super()
+	FProxyInterpolationError(const int32 PullDataInterpIdx = INDEX_NONE, const int32 InterpChannel = 0) : Super(PullDataInterpIdx, InterpChannel)
 	{}
+	virtual ~FProxyInterpolationError() { }
+
+	static const EProxyInterpolationType InterpolationType = EProxyInterpolationType::ErrorLinear;
+	virtual const EProxyInterpolationType GetInterpolationType() const override { return FProxyInterpolationError::InterpolationType; }
 
-	virtual bool IsErrorSmoothing() override { return ErrorSmoothingCount > 0; }
+	virtual const bool IsErrorSmoothing() const override { return ErrorSmoothingCount > 0; }
 
-	Chaos::FVec3 GetErrorX(const Chaos::FRealSingle Alpha) { return FMath::Lerp(ErrorXPrev, ErrorX, Alpha); } // Get the ErrorX based on current Alpha between GT and PT
-	FQuat GetErrorR(const Chaos::FRealSingle Alpha) { return FMath::Lerp(ErrorRPrev, ErrorR, Alpha); } // Get the ErrorR based on current Alpha between GT and PT
+	virtual const Chaos::FVec3 GetErrorX(const Chaos::FRealSingle Alpha) const override { return FMath::Lerp(ErrorXPrev, ErrorX, Alpha); }
 
-	void AccumlateErrorXR(const Chaos::FVec3 X, const FQuat R, const int32 CurrentSimTick, const int32 ErrorSmoothDuration)
+	virtual const FQuat GetErrorR(const Chaos::FRealSingle Alpha) const override { return FMath::Lerp(ErrorRPrev, ErrorR, Alpha); }
+
+	virtual void AccumlateErrorXR(const Chaos::FVec3 X, const FQuat R, const int32 CurrentSimTick, const int32 ErrorSmoothDuration) override
 	{
 		ErrorSmoothingCount = ErrorSmoothDuration; // How many simulation ticks to correct error over
 		LastSimTick = CurrentSimTick - 1; // Error is from the previous simulation tick, not the current
@@ -245,7 +296,7 @@ struct FProxyInterpolationError : FProxyInterpolationBase
 		}
 	}
 
-	virtual bool UpdateError(const int32 CurrentSimTick, const Chaos::FReal AsyncFixedTimeStep)
+	virtual const bool UpdateError(const int32 CurrentSimTick, const Chaos::FReal AsyncFixedTimeStep) override
 	{
 		// Cache how many simulation ticks have passed since last call
 		SimTicks = CurrentSimTick - LastSimTick;
@@ -258,7 +309,7 @@ struct FProxyInterpolationError : FProxyInterpolationBase
 		return false;
 	}
 
-	virtual bool DirectionalDecay(Chaos::FVec3 Direction)
+	virtual const bool DirectionalDecay(Chaos::FVec3 Direction) override
 	{
 		if (IsErrorSmoothing() && SimTicks > 0)
 		{
@@ -275,7 +326,7 @@ struct FProxyInterpolationError : FProxyInterpolationBase
 	}
 
 protected:
-	bool DecayError()
+	virtual const bool DecayError()
 	{
 		if (!IsErrorSmoothing())
 		{
@@ -289,20 +340,23 @@ protected:
 		// Second step: 8/9  = 0.888 |  90 * 0.888 = 80 error
 		// Third step: 7/8  = 0.875 |  80 * 0.875 = 70 error
 		// etc.
-		Chaos::FRealSingle Alpha = Chaos::FRealSingle(ErrorSmoothingCount - SimTicks) / Chaos::FRealSingle(ErrorSmoothingCount);
-		Alpha = FMath::Clamp(Alpha, 0.f, 1.f);
+		for (int32 Ticks = FMath::Clamp(SimTicks, 0, ErrorSmoothingCount); Ticks > 0; Ticks--)
+		{
+			Chaos::FRealSingle Alpha = Chaos::FRealSingle(ErrorSmoothingCount - 1) / Chaos::FRealSingle(ErrorSmoothingCount);
+			Alpha = FMath::Clamp(Alpha, 0.f, 1.f);
 
-		ErrorXPrev = ErrorX;
-		ErrorX *= Alpha;
+			ErrorXPrev = ErrorX;
+			ErrorX *= Alpha;
+			ErrorRPrev = ErrorR;
+			ErrorR = FMath::Lerp(FQuat::Identity, ErrorR, Alpha);
 
-		ErrorRPrev = ErrorR;
-		ErrorR = FMath::Lerp(FQuat::Identity, ErrorR, Alpha);
+			ErrorSmoothingCount = FMath::Max(ErrorSmoothingCount - 1, 0);
+		}
 
-		ErrorSmoothingCount = FMath::Max(ErrorSmoothingCount - SimTicks, 0);
-		return true;
+		return SimTicks > 0;
 	}
 
-	void Reset()
+	virtual void Reset()
 	{
 		ErrorX = Chaos::FVec3::ZeroVector;
 		ErrorXPrev = Chaos::FVec3::ZeroVector;
@@ -322,29 +376,26 @@ protected:
 	FQuat ErrorR = FQuat::Identity;
 	FQuat ErrorRPrev = FQuat::Identity;
 	int32 ErrorSmoothingCount = 0;
-
 };
 
 
-// Take incoming velocity into consideration when performing render interpolation, the correction will be more organic but might result in clipping and it's heavier for memory and CPU.
-#ifndef RENDERINTERP_ERRORVELOCITYSMOOTHING
-#define RENDERINTERP_ERRORVELOCITYSMOOTHING 0
-#endif
 // Render Interpolation that both perform the linear error correction from FProxyInterpolationError and takes incoming velocity into account to make a smoother and more organic correction of the error.
 struct FProxyInterpolationErrorVelocity : FProxyInterpolationError
 {
 	using Super = FProxyInterpolationError;
 
-	FProxyInterpolationErrorVelocity() : Super()
+	FProxyInterpolationErrorVelocity(const int32 PullDataInterpIdx = INDEX_NONE, const int32 InterpChannel = 0) : Super(PullDataInterpIdx, InterpChannel)
 	{}
+	virtual ~FProxyInterpolationErrorVelocity() { }
 
-	bool IsErrorVelocitySmoothing() { return ErrorVelocitySmoothingCount > 0; }
+	static const EProxyInterpolationType InterpolationType = EProxyInterpolationType::ErrorVelocity;
+	virtual const EProxyInterpolationType GetInterpolationType() const override { return FProxyInterpolationErrorVelocity::InterpolationType; }
 
-	// Returns the Alpha of how much to take previous velocity into account, used to lerp from linear extrapolation to the predicted position based on previous velocity.
-	Chaos::FRealSingle GetErrorVelocitySmoothingAlpha(const int32 ErrorVelocitySmoothDuration) { return Chaos::FRealSingle(ErrorVelocitySmoothingCount) / Chaos::FRealSingle(ErrorVelocitySmoothDuration); }
-	Chaos::FVec3 GetErrorVelocitySmoothingX(const Chaos::FRealSingle Alpha) { return FMath::Lerp(ErrorVelocitySmoothingXPrev, ErrorVelocitySmoothingX, Alpha); } // Get the VelocityErrorX based on current Alpha between GT and PT
+	virtual const bool IsErrorVelocitySmoothing() const override { return ErrorVelocitySmoothingCount > 0; }
+	virtual const Chaos::FRealSingle GetErrorVelocitySmoothingAlpha(const int32 ErrorVelocitySmoothDuration) const override { return Chaos::FRealSingle(ErrorVelocitySmoothingCount) / Chaos::FRealSingle(ErrorVelocitySmoothDuration); }
+	virtual const Chaos::FVec3 GetErrorVelocitySmoothingX(const Chaos::FRealSingle Alpha) const override { return FMath::Lerp(ErrorVelocitySmoothingXPrev, ErrorVelocitySmoothingX, Alpha); }
 
-	bool UpdateError(const int32 CurrentSimTick, const Chaos::FReal AsyncFixedTimeStep) override
+	virtual const bool UpdateError(const int32 CurrentSimTick, const Chaos::FReal AsyncFixedTimeStep) override
 	{		
 		if (Super::UpdateError(CurrentSimTick, AsyncFixedTimeStep))
 		{
@@ -354,7 +405,7 @@ struct FProxyInterpolationErrorVelocity : FProxyInterpolationError
 		return false;
 	}
 
-	void SetVelocitySmoothing(const Chaos::FVec3 CurrV, const Chaos::FVec3 CurrX, const int32 ErrorVelocitySmoothDuration)
+	virtual void SetVelocitySmoothing(const Chaos::FVec3 CurrV, const Chaos::FVec3 CurrX, const int32 ErrorVelocitySmoothDuration) override
 	{
 		// Cache pre error velocity and position to be used when smoothing out error correction
 		ErrorVelocitySmoothingV = CurrV;
@@ -364,7 +415,7 @@ struct FProxyInterpolationErrorVelocity : FProxyInterpolationError
 	}
 
 protected:
-	void StepErrorVelocitySmoothingData(const Chaos::FReal AsyncFixedTimeStep)
+	virtual void StepErrorVelocitySmoothingData(const Chaos::FReal AsyncFixedTimeStep)
 	{
 		// Step the error velocity smoothing position forward along the previous velocity to have a new position to base smoothing on each tick
 		if (IsErrorVelocitySmoothing())
diff --git a/Engine/Source/Runtime/Experimental/Chaos/Public/PhysicsProxy/ClusterUnionPhysicsProxy.h b/Engine/Source/Runtime/Experimental/Chaos/Public/PhysicsProxy/ClusterUnionPhysicsProxy.h
--- a/Engine/Source/Runtime/Experimental/Chaos/Public/PhysicsProxy/ClusterUnionPhysicsProxy.h
+++ b/Engine/Source/Runtime/Experimental/Chaos/Public/PhysicsProxy/ClusterUnionPhysicsProxy.h
@@ -143,8 +143,8 @@ namespace Chaos
 		CHAOS_API void SyncRemoteData(FDirtyPropertiesManager& Manager, int32 DataIdx, FDirtyChaosProperties& RemoteData) const;
 		CHAOS_API void ClearAccumulatedData();
 
-		FProxyInterpolationBase& GetInterpolationData() { return InterpolationData; }
-		const FProxyInterpolationBase& GetInterpolationData() const { return InterpolationData; }
+		FProxyInterpolationBase* GetInterpolationData() { return InterpolationData.Get(); }
+		const FProxyInterpolationBase* GetInterpolationData() const { return InterpolationData.Get(); }
 
 		FClusterUnionIndex GetClusterUnionIndex() const { return ClusterUnionIndex; }
 
@@ -167,7 +167,23 @@ namespace Chaos
 		FInternalParticle* Particle_Internal = nullptr;
 		FClusterUnionIndex ClusterUnionIndex = INDEX_NONE;
 
-		FProxyInterpolationError InterpolationData;
+		TUniquePtr<FProxyInterpolationBase> InterpolationData;
+
+		/** Get or create a derived FProxyInterpolationBase that handles error corrections */
+		template<typename ErrorDataType>
+		ErrorDataType* GetOrCreateErrorInterpolationData()
+		{
+			if (!InterpolationData.IsValid())
+			{
+				InterpolationData = MakeUnique<ErrorDataType>();
+			}
+			else if (InterpolationData.Get()->GetInterpolationType() != ErrorDataType::InterpolationType)
+			{
+				InterpolationData = MakeUnique<ErrorDataType>(InterpolationData.Get()->GetPullDataInterpIdx_External(), InterpolationData.Get()->GetInterpChannel_External());
+			}
+
+			return static_cast<ErrorDataType*>(InterpolationData.Get());
+		}
 
 		// An array of a particles that exist in the external implicit object union.
 		// Note that this array should only be used for book-keeping. It is generally
diff --git a/Engine/Source/Runtime/Experimental/Chaos/Public/PhysicsProxy/GeometryCollectionPhysicsProxy.h b/Engine/Source/Runtime/Experimental/Chaos/Public/PhysicsProxy/GeometryCollectionPhysicsProxy.h
--- a/Engine/Source/Runtime/Experimental/Chaos/Public/PhysicsProxy/GeometryCollectionPhysicsProxy.h
+++ b/Engine/Source/Runtime/Experimental/Chaos/Public/PhysicsProxy/GeometryCollectionPhysicsProxy.h
@@ -479,8 +479,8 @@ public:
 
 	CHAOS_API float ComputeMaterialBasedDamageThreshold_Internal(Chaos::FPBDRigidClusteredParticleHandle& ClusteredParticle) const;
 
-	FProxyInterpolationBase& GetInterpolationData() { return InterpolationData; }
-	const FProxyInterpolationBase& GetInterpolationData() const { return InterpolationData; }
+	FProxyInterpolationBase* GetInterpolationData() { return InterpolationData.Get(); }
+	const FProxyInterpolationBase* GetInterpolationData() const { return InterpolationData.Get(); }
 
 	enum class EReplicationMode: uint8
 	{
@@ -750,7 +750,23 @@ private:
 	// paradigm, at least for this component of the handshake.
 	Chaos::FGuardedTripleBuffer<FGeometryCollectionResults> PhysToGameInterchange;
 
-	FProxyInterpolationError InterpolationData;
+	TUniquePtr<FProxyInterpolationBase> InterpolationData;
+
+	/** Get or create a derived FProxyInterpolationBase that handles error corrections */
+	template<typename ErrorDataType>
+	ErrorDataType* GetOrCreateErrorInterpolationData()
+	{
+		if (!InterpolationData.IsValid())
+		{
+			InterpolationData = MakeUnique<ErrorDataType>();
+		}
+		else if (InterpolationData.Get()->GetInterpolationType() != ErrorDataType::InterpolationType)
+		{
+			InterpolationData = MakeUnique<ErrorDataType>(InterpolationData.Get()->GetPullDataInterpIdx_External(), InterpolationData.Get()->GetInterpChannel_External());
+		}
+
+		return static_cast<ErrorDataType*>(InterpolationData.Get());
+	}
 
 	// this is used as a unique ID when collecting data from runtime
 	FGuid CollectorGuid;
diff --git a/Engine/Source/Runtime/Experimental/Chaos/Public/PhysicsProxy/SingleParticlePhysicsProxy.h b/Engine/Source/Runtime/Experimental/Chaos/Public/PhysicsProxy/SingleParticlePhysicsProxy.h
--- a/Engine/Source/Runtime/Experimental/Chaos/Public/PhysicsProxy/SingleParticlePhysicsProxy.h
+++ b/Engine/Source/Runtime/Experimental/Chaos/Public/PhysicsProxy/SingleParticlePhysicsProxy.h
@@ -67,8 +67,8 @@ public:
 	FSingleParticlePhysicsProxy(FSingleParticlePhysicsProxy&&) = delete;
 	CHAOS_API virtual ~FSingleParticlePhysicsProxy();
 
-	const FProxyInterpolationBase& GetInterpolationData() const { return InterpolationData; }
-	FProxyInterpolationBase& GetInterpolationData() { return InterpolationData; }
+	FProxyInterpolationBase* GetInterpolationData() { return InterpolationData.Get(); }
+	const FProxyInterpolationBase* GetInterpolationData() const { return InterpolationData.Get(); }
 
 	FORCEINLINE FRigidBodyHandle_External& GetGameThreadAPI()
 	{
@@ -178,11 +178,23 @@ protected:
 	FPhysicsObjectUniquePtr Reference;
 
 private:
-#if RENDERINTERP_ERRORVELOCITYSMOOTHING
-	FProxyInterpolationErrorVelocity InterpolationData;
-#else
-	FProxyInterpolationError InterpolationData;
-#endif
+	TUniquePtr<FProxyInterpolationBase> InterpolationData;
+
+	/** Get or create a derived FProxyInterpolationBase that handles error corrections */
+	template<typename ErrorDataType>
+	ErrorDataType* GetOrCreateErrorInterpolationData()
+	{
+		if (!InterpolationData.IsValid())
+		{
+			InterpolationData = MakeUnique<ErrorDataType>();
+		}
+		else if (InterpolationData.Get()->GetInterpolationType() != ErrorDataType::InterpolationType)
+		{
+			InterpolationData = MakeUnique<ErrorDataType>(InterpolationData.Get()->GetPullDataInterpIdx_External(), InterpolationData.Get()->GetInterpChannel_External());
+		}
+
+		return static_cast<ErrorDataType*>(InterpolationData.Get());
+	}
 
 	//use static Create
 	CHAOS_API FSingleParticlePhysicsProxy(TUniquePtr<PARTICLE_TYPE>&& InParticle, FParticleHandle* InHandle, UObject* InOwner = nullptr);
