Geometry collection - memory optimization
* packed and re-ordered data in structures to minimize padding as much as possible
* reduce simulation parameters data by 80 bytes ( 112 if stripping the deprecated members )
* remove dead code

diff --git a/Engine/Source/Runtime/Experimental/Chaos/Private/PhysicsProxy/GeometryCollectionPhysicsProxy.cpp b/Engine/Source/Runtime/Experimental/Chaos/Private/PhysicsProxy/GeometryCollectionPhysicsProxy.cpp
--- a/Engine/Source/Runtime/Experimental/Chaos/Private/PhysicsProxy/GeometryCollectionPhysicsProxy.cpp
+++ b/Engine/Source/Runtime/Experimental/Chaos/Private/PhysicsProxy/GeometryCollectionPhysicsProxy.cpp
@@ -38,18 +38,6 @@
 #include "GeometryCollection/Facades/CollectionAnchoringFacade.h"
 #include "GeometryCollection/Facades/CollectionConnectionGraphFacade.h"
 
-#ifndef TODO_REIMPLEMENT_INIT_COMMANDS
-#define TODO_REIMPLEMENT_INIT_COMMANDS 0
-#endif
-
-#ifndef TODO_REIMPLEMENT_FRACTURE
-#define TODO_REIMPLEMENT_FRACTURE 0
-#endif
-
-#ifndef TODO_REIMPLEMENT_RIGID_CACHING
-#define TODO_REIMPLEMENT_RIGID_CACHING 0
-#endif
-
 #define GC_PHYSICSPROXY_CHECK_FOR_NAN_ENABLED 0
 
 #if GC_PHYSICSPROXY_CHECK_FOR_NAN_ENABLED
@@ -519,10 +507,6 @@ FGeometryCollectionPhysicsProxy::FGeometryCollectionPhysicsProxy(
 	, IsObjectDeleting(false)
 	, SimFilter(InSimFilter)
 	, QueryFilter(InQueryFilter)
-#if TODO_REIMPLEMENT_RIGID_CACHING
-	, ProxySimDuration(0.0f)
-	, LastSyncCountGT(MAX_uint32)
-#endif
 	, CollisionParticlesPerObjectFraction(CollisionParticlesPerObjectFractionDefault)
 	, PhysicsThreadCollection(Parameters.RestCollectionShared)
 	, GameThreadCollection(GameThreadCollectionIn)
@@ -1630,18 +1614,6 @@ void FGeometryCollectionPhysicsProxy::InitializeBodiesPT(Chaos::FPBDRigidsSolver
 			}
 		} // end if EnableClustering
  
-
-#if TODO_REIMPLEMENT_RIGID_CACHING
-		// If we're recording and want to start immediately caching then we should cache the rest state
-		if (Parameters.IsCacheRecording() && Parameters.CacheBeginTime == 0.0f)
-		{
-			if (UpdateRecordedStateCallback)
-			{
-				UpdateRecordedStateCallback(0.0f, RigidBodyID, Particles, RigidSolver->GetCollisionConstraints());
-			}
-		}
-#endif // TODO_REIMPLEMENT_RIGID_CACHING
-
 		const bool bStartSleeping =
 				(Parameters.ObjectType == EObjectStateTypeEnum::Chaos_Object_Sleeping
 			 || (Parameters.ObjectType == EObjectStateTypeEnum::Chaos_Object_Dynamic && !Parameters.StartAwake));
diff --git a/Engine/Source/Runtime/Experimental/Chaos/Public/Chaos/PBDRigidClusteringTypes.h b/Engine/Source/Runtime/Experimental/Chaos/Public/Chaos/PBDRigidClusteringTypes.h
--- a/Engine/Source/Runtime/Experimental/Chaos/Public/Chaos/PBDRigidClusteringTypes.h
+++ b/Engine/Source/Runtime/Experimental/Chaos/Public/Chaos/PBDRigidClusteringTypes.h
@@ -1,9 +1,11 @@
 // Copyright Epic Games, Inc. All Rights Reserved.
 #pragma once
 
+#include "HAL/Platform.h"
+
 namespace Chaos
 {
-	enum class EDamageEvaluationModel
+	enum class EDamageEvaluationModel: uint8
 	{
 		/** 
 		* particles internal strains are set from user defined damage thresholds 
diff --git a/Engine/Source/Runtime/Experimental/Chaos/Public/GeometryCollection/GeometryCollectionSimulationCoreTypes.h b/Engine/Source/Runtime/Experimental/Chaos/Public/GeometryCollection/GeometryCollectionSimulationCoreTypes.h
--- a/Engine/Source/Runtime/Experimental/Chaos/Public/GeometryCollection/GeometryCollectionSimulationCoreTypes.h
+++ b/Engine/Source/Runtime/Experimental/Chaos/Public/GeometryCollection/GeometryCollectionSimulationCoreTypes.h
@@ -67,14 +67,14 @@
 {
 	FSharedSimulationSizeSpecificData()
 		: MaxSize(0.f)
-		, CollisionShapesData({ FCollectionCollisionTypeData() })
 		, DamageThreshold(5000.f)
+		, CollisionShapesData({ FCollectionCollisionTypeData() })
 	{
 	}
 
 	float MaxSize;
-	TArray<FCollectionCollisionTypeData> CollisionShapesData;
 	float DamageThreshold;
+	TArray<FCollectionCollisionTypeData> CollisionShapesData;
 
 	bool operator<(const FSharedSimulationSizeSpecificData& Rhs) const { return MaxSize < Rhs.MaxSize; }
 };
@@ -91,15 +91,15 @@
 struct FSharedSimulationParameters
 {
 	FSharedSimulationParameters()
-	: bMassAsDensity(true)
-	, Mass(1.0f)
-	, MinimumMassClamp(0.1f)								// todo : Expose to users with better initial values
+	: MinimumMassClamp(0.1f)								// todo : Expose to users with better initial values
 	, MaximumMassClamp(1e5f)								// todo : Expose to users with better initial values
 	, MinimumBoundingExtentClamp(0.1f)						// todo : Expose to users with better initial values
 	, MaximumBoundingExtentClamp(1e6f)						// todo : Expose to users with better initial values
 	, MinimumInertiaTensorDiagonalClamp(UE_SMALL_NUMBER)	// todo : Expose to users with better initial values
 	, MaximumInertiaTensorDiagonalClamp(1e20f)				// todo : Expose to users with better initial values
 	, MaximumCollisionParticleCount(60)
+	, Mass(1.0f)
+	, bMassAsDensity(true)
 	, bUseImportedCollisionImplicits(false)
 	{
 		SizeSpecificData.AddDefaulted();
@@ -123,15 +123,15 @@
 		, int32 InMaximumCollisionParticleCount
 		, float InCollisionMarginFraction
 		, bool InUseImportedCollisionImplicits )
-	: bMassAsDensity(InMassAsDensity)
-	, Mass(InMass)
-	, MinimumMassClamp(InMinimumMassClamp)
+	: MinimumMassClamp(InMinimumMassClamp)
 	, MaximumMassClamp(InMinimumMassClamp)
 	, MinimumBoundingExtentClamp(InMinimumBoundingExtentClamp)
 	, MaximumBoundingExtentClamp(InMinimumBoundingExtentClamp)
 	, MinimumInertiaTensorDiagonalClamp(InMinimumInertiaTensorDiagonalClamp)
 	, MaximumInertiaTensorDiagonalClamp(InMaximumInertiaTensorDiagonalClamp)
 	, MaximumCollisionParticleCount(InMaximumCollisionParticleCount)
+	, Mass(InMass)
+	, bMassAsDensity(InMassAsDensity)
 	, bUseImportedCollisionImplicits(InUseImportedCollisionImplicits)
 	{
 		SizeSpecificData.AddDefaulted();
@@ -149,23 +149,24 @@
 		}
 	}
 
-	bool bMassAsDensity;
-	float Mass;
+	TArray<FSharedSimulationSizeSpecificData> SizeSpecificData;
 	float MinimumMassClamp;
 	float MaximumMassClamp;
 	float MinimumBoundingExtentClamp;
 	float MaximumBoundingExtentClamp;
 	float MinimumInertiaTensorDiagonalClamp;
 	float MaximumInertiaTensorDiagonalClamp;
+	int32 MaximumCollisionParticleCount;
+	float Mass;
+	bool bMassAsDensity : 1;
+	bool bUseImportedCollisionImplicits : 1;
 
 	float MinimumVolumeClamp() const { return MinimumBoundingExtentClamp * MinimumBoundingExtentClamp * MinimumBoundingExtentClamp; }
 	float MaximumVolumeClamp() const { return MaximumBoundingExtentClamp * MaximumBoundingExtentClamp * MaximumBoundingExtentClamp; }
-
-	TArray<FSharedSimulationSizeSpecificData> SizeSpecificData;
-	int32 MaximumCollisionParticleCount;
-	bool bUseImportedCollisionImplicits;
 };
 
+#define SIMULATIONPARAMETERS_CACHE_PARAMETERS 1
+
 struct FSimulationParameters
 {
 	FSimulationParameters()
@@ -174,35 +175,29 @@
 		, RestCollection(nullptr)
 PRAGMA_ENABLE_DEPRECATION_WARNINGS
 		, RestCollectionShared(nullptr)
-		, InitialRootIndex(INDEX_NONE)
-		, RecordedTrack(nullptr)
-		, bOwnsTrack(false)
-		, Simulating(false)
 		, WorldTransform(FTransform::Identity)
-		, EnableClustering(true)
+		, DamageThreshold({ 500000.f, 50000.f, 5000.f })
+		, InitialRootIndex(INDEX_NONE)
 		, ClusterGroupIndex(0)
 		, MaxClusterLevel(100)
 		, MaxSimulatedLevel(100)
-		, bUseSizeSpecificDamageThresholds(false)
-		, bUseMaterialDamageModifiers(false)
+		, ObjectType(EObjectStateTypeEnum::Chaos_NONE)
+		, InitialVelocityType(EInitialVelocityTypeEnum::Chaos_Initial_Velocity_None)
 		, DamageModel(EDamageModelTypeEnum::Chaos_Damage_Model_UserDefined_Damage_Threshold)
 		, DamageEvaluationModel(Chaos::EDamageEvaluationModel::StrainFromDamageThreshold)
-		, DamageThreshold({500000.f, 50000.f, 5000.f})
-		, bUsePerClusterOnlyDamageThreshold(false)
 		, ClusterConnectionMethod(Chaos::FClusterCreationParameters::EConnectionMethod::PointImplicit)
 		, ConnectionGraphBoundsFilteringMargin(0)
 		, CollisionGroup(0)
 		, CollisionSampleFraction(1.0)
-		, InitialVelocityType(EInitialVelocityTypeEnum::Chaos_Initial_Velocity_None)
 		, InitialLinearVelocity(FVector(0))
 		, InitialAngularVelocity(FVector(0))
-		, CacheType(EGeometryCollectionCacheType::None)
-		, CacheBeginTime(0.0f)
-		, ReverseCacheBeginTime(0.0f)
-		, bClearCache(false)
-		, ObjectType(EObjectStateTypeEnum::Chaos_NONE)
-		, StartAwake(true)
 		, MaterialOverrideMassScaleMultiplier(1.0f)
+		, Simulating(false)
+		, EnableClustering(true)
+		, bUseSizeSpecificDamageThresholds(false)
+		, bUseMaterialDamageModifiers(false)
+		, bUsePerClusterOnlyDamageThreshold(false)
+		, StartAwake(true)
 		, bGenerateBreakingData(false)
 		, bGenerateCollisionData(false)
 		, bGenerateTrailingData(false)
@@ -213,24 +208,34 @@
 		, bGenerateGlobalCrumblingData(false)
 		, bGenerateGlobalCrumblingChildrenData(false)
 		, EnableGravity(true)
-		, GravityGroupIndex(0)
-		, OneWayInteractionLevel(INDEX_NONE)
 		, UseInertiaConditioning(true)
 		, UseCCD(false)
 		, UseMACD(false)
+		, bEnableStrainOnCollision(true)
+		, bUseStaticMeshCollisionForTraces(false)
+		, bOptimizeConvexes(true)
+		, bUseDamagePropagation(false)
+		, BreakDamagePropagationFactor(1.0f)
+		, ShockDamagePropagationFactor(0.0f)
 		, LinearDamping(0.01f)
 		, AngularDamping(0)
 		, InitialOverlapDepenetrationVelocity(-1.0f)
 		, SleepThresholdMultiplier(1.0f)
-		, bUseDamagePropagation(false)
-		, BreakDamagePropagationFactor(1.0f)
-		, ShockDamagePropagationFactor(0.0f)
+		, GravityGroupIndex(0)
+		, OneWayInteractionLevel(INDEX_NONE)
 		, SimulationFilterData()
 		, QueryFilterData()
 		, UserData(nullptr)
-		, bEnableStrainOnCollision(true)
-		, bUseStaticMeshCollisionForTraces(false)
-		, bOptimizeConvexes(true)
+#if SIMULATIONPARAMETERS_CACHE_PARAMETERS
+PRAGMA_DISABLE_DEPRECATION_WARNINGS
+		, RecordedTrack(nullptr)
+		, CacheBeginTime(0.0f)
+		, ReverseCacheBeginTime(0.0f)
+		, CacheType(EGeometryCollectionCacheType::None)
+		, bClearCache(false)
+		, bOwnsTrack(false)
+PRAGMA_ENABLE_DEPRECATION_WARNINGS
+#endif
 	{}
 
 	FSimulationParameters(const FSimulationParameters& Other)
@@ -239,37 +244,32 @@
 		, RestCollection(Other.RestCollection)
 PRAGMA_ENABLE_DEPRECATION_WARNINGS
 		, RestCollectionShared(Other.RestCollectionShared)
-		, InitialRootIndex(Other.InitialRootIndex)
-		, InitializationCommands(Other.InitializationCommands)
-		, RecordedTrack(Other.RecordedTrack)
-		, bOwnsTrack(false)
-		, Simulating(Other.Simulating)
+		, Shared(Other.Shared)
 		, WorldTransform(Other.WorldTransform)
-		, EnableClustering(Other.EnableClustering)
+		, InitializationCommands(Other.InitializationCommands)
+		, DamageThreshold(Other.DamageThreshold)
+		, InitialRootIndex(Other.InitialRootIndex)
 		, ClusterGroupIndex(Other.ClusterGroupIndex)
 		, MaxClusterLevel(Other.MaxClusterLevel)
 		, MaxSimulatedLevel(Other.MaxSimulatedLevel)
-		, bUseSizeSpecificDamageThresholds(Other.bUseSizeSpecificDamageThresholds)
-		, bUseMaterialDamageModifiers(Other.bUseMaterialDamageModifiers)
+		, ObjectType(Other.ObjectType)
+		, InitialVelocityType(Other.InitialVelocityType)
 		, DamageModel(Other.DamageModel)
 		, DamageEvaluationModel(Other.DamageEvaluationModel)
-		, DamageThreshold(Other.DamageThreshold)
-		, bUsePerClusterOnlyDamageThreshold(Other.bUsePerClusterOnlyDamageThreshold)
 		, ClusterConnectionMethod(Other.ClusterConnectionMethod)
 		, ConnectionGraphBoundsFilteringMargin(Other.ConnectionGraphBoundsFilteringMargin)
 		, CollisionGroup(Other.CollisionGroup)
 		, CollisionSampleFraction(Other.CollisionSampleFraction)
-		, InitialVelocityType(Other.InitialVelocityType)
 		, InitialLinearVelocity(Other.InitialLinearVelocity)
 		, InitialAngularVelocity(Other.InitialAngularVelocity)
-		, CacheType(Other.CacheType)
-		, CacheBeginTime(Other.CacheBeginTime)
-		, ReverseCacheBeginTime(Other.ReverseCacheBeginTime)
-		, bClearCache(Other.bClearCache)
-		, ObjectType(Other.ObjectType)
-		, StartAwake(Other.StartAwake)
 		, PhysicalMaterialHandle(Other.PhysicalMaterialHandle)
 		, MaterialOverrideMassScaleMultiplier(Other.MaterialOverrideMassScaleMultiplier)
+		, Simulating(Other.Simulating)
+		, EnableClustering(Other.EnableClustering)
+		, bUseSizeSpecificDamageThresholds(Other.bUseSizeSpecificDamageThresholds)
+		, bUseMaterialDamageModifiers(Other.bUseMaterialDamageModifiers)
+		, bUsePerClusterOnlyDamageThreshold(Other.bUsePerClusterOnlyDamageThreshold)
+		, StartAwake(Other.StartAwake)
 		, bGenerateBreakingData(Other.bGenerateBreakingData)
 		, bGenerateCollisionData(Other.bGenerateCollisionData)
 		, bGenerateTrailingData(Other.bGenerateTrailingData)
@@ -279,63 +279,74 @@
 		, bGenerateGlobalCollisionData(Other.bGenerateGlobalCollisionData)
 		, bGenerateGlobalCrumblingData(Other.bGenerateGlobalCrumblingData)
 		, bGenerateGlobalCrumblingChildrenData(Other.bGenerateGlobalCrumblingChildrenData)
-		, Shared(Other.Shared)
 		, EnableGravity(Other.EnableGravity)
-		, GravityGroupIndex(Other.GravityGroupIndex)
-		, OneWayInteractionLevel(Other.OneWayInteractionLevel)
 		, UseInertiaConditioning(Other.UseInertiaConditioning)
 		, UseCCD(Other.UseCCD)
 		, UseMACD(Other.UseMACD)
+		, bEnableStrainOnCollision(Other.bEnableStrainOnCollision)
+		, bUseStaticMeshCollisionForTraces(Other.bUseStaticMeshCollisionForTraces)
+		, bOptimizeConvexes(Other.bOptimizeConvexes)
+		, bUseDamagePropagation(Other.bUseDamagePropagation)
+		, BreakDamagePropagationFactor(Other.BreakDamagePropagationFactor)
+		, ShockDamagePropagationFactor(Other.ShockDamagePropagationFactor)
 		, LinearDamping(Other.LinearDamping)
 		, AngularDamping(Other.AngularDamping)
 		, InitialOverlapDepenetrationVelocity(Other.InitialOverlapDepenetrationVelocity)
 		, SleepThresholdMultiplier(Other.SleepThresholdMultiplier)
-		, bUseDamagePropagation(Other.bUseDamagePropagation)
-		, BreakDamagePropagationFactor(Other.BreakDamagePropagationFactor)
-		, ShockDamagePropagationFactor(Other.ShockDamagePropagationFactor)
+		, GravityGroupIndex(Other.GravityGroupIndex)
+		, OneWayInteractionLevel(Other.OneWayInteractionLevel)
 		, SimulationFilterData(Other.SimulationFilterData)
 		, QueryFilterData(Other.QueryFilterData)
 		, UserData(Other.UserData)
-		, bEnableStrainOnCollision(Other.bEnableStrainOnCollision)
-		, bUseStaticMeshCollisionForTraces(Other.bUseStaticMeshCollisionForTraces)
-		, bOptimizeConvexes(Other.bOptimizeConvexes)
+#if SIMULATIONPARAMETERS_CACHE_PARAMETERS
+PRAGMA_DISABLE_DEPRECATION_WARNINGS
+		, RecordedTrack(Other.RecordedTrack)
+		, CacheBeginTime(Other.CacheBeginTime)
+		, ReverseCacheBeginTime(Other.ReverseCacheBeginTime)
+		, CacheType(Other.CacheType)
+		, bClearCache(Other.bClearCache)
+		, bOwnsTrack(false)
+PRAGMA_ENABLE_DEPRECATION_WARNINGS
+#endif
+
 	{
 	}
 
 	~FSimulationParameters()
 	{
+#if SIMULATIONPARAMETERS_CACHE_PARAMETERS
+PRAGMA_DISABLE_DEPRECATION_WARNINGS
 		if (bOwnsTrack)
 		{
 			delete const_cast<FRecordedTransformTrack*>(RecordedTrack);
 		}
+PRAGMA_ENABLE_DEPRECATION_WARNINGS
+#endif
 	}
 
-	bool IsCacheRecording() { return CacheType == EGeometryCollectionCacheType::Record || CacheType == EGeometryCollectionCacheType::RecordAndPlay; }
-	bool IsCachePlaying() { return CacheType == EGeometryCollectionCacheType::Play || CacheType == EGeometryCollectionCacheType::RecordAndPlay; }
-
 	FString Name;
 	
 	UE_DEPRECATED(5.4, "Raw pointer no longer in use, instead prefer RestCollectionShared")
 	const FGeometryCollection* RestCollection;
 
 	TSharedPtr<const FGeometryCollection> RestCollectionShared;
+	FSharedSimulationParameters Shared;
 
-	int32 InitialRootIndex;
-	TArray<FFieldSystemCommand> InitializationCommands;
-	const FRecordedTransformTrack* RecordedTrack;
-	bool bOwnsTrack;
-
-	bool Simulating;
-
 	FTransform WorldTransform;
 	FTransform PrevWorldTransform;
 
-	bool EnableClustering;
+	TArray<FFieldSystemCommand> InitializationCommands;
+
+	TArray<float> DamageThreshold;
+
+	int32 InitialRootIndex;
 	int32 ClusterGroupIndex;
-	int32 MaxClusterLevel;
-	int32 MaxSimulatedLevel;
-	bool bUseSizeSpecificDamageThresholds;
-	bool bUseMaterialDamageModifiers;
+	int16 MaxClusterLevel;
+	int16 MaxSimulatedLevel;
 
+	EObjectStateTypeEnum ObjectType;
+	
+	EInitialVelocityTypeEnum InitialVelocityType;
+
 	/** this is the user expose damage model, used for creation of the particles */
 	EDamageModelTypeEnum DamageModel; 
@@ -343,64 +354,96 @@
 	/** this is the lower level damage model for clustering, used at runm time */
 	Chaos::EDamageEvaluationModel DamageEvaluationModel;
 
-	TArray<float> DamageThreshold;
-	bool bUsePerClusterOnlyDamageThreshold;
 	Chaos::FClusterCreationParameters::EConnectionMethod ClusterConnectionMethod;
 	float ConnectionGraphBoundsFilteringMargin;
 
 	int32 CollisionGroup;
 	float CollisionSampleFraction;
 
-	EInitialVelocityTypeEnum InitialVelocityType;
-	FVector InitialLinearVelocity;
-	FVector InitialAngularVelocity;
+	FVector3f InitialLinearVelocity;
+	FVector3f InitialAngularVelocity;
 
-	EGeometryCollectionCacheType CacheType;
-	float CacheBeginTime;
-	float ReverseCacheBeginTime;
-	bool bClearCache;
-
-	EObjectStateTypeEnum ObjectType;
-	bool StartAwake;
-
 	Chaos::FMaterialHandle PhysicalMaterialHandle;
 
 	float MaterialOverrideMassScaleMultiplier;
 
-	bool bGenerateBreakingData;
-	bool bGenerateCollisionData;
-	bool bGenerateTrailingData;
-	bool bGenerateCrumblingData;
-	bool bGenerateCrumblingChildrenData;
+	bool Simulating : 1;
+	bool EnableClustering : 1;
+	bool bUseSizeSpecificDamageThresholds : 1;
+	bool bUseMaterialDamageModifiers : 1;
+	bool bUsePerClusterOnlyDamageThreshold : 1;
+	bool StartAwake : 1;
 
-	bool bGenerateGlobalBreakingData;
-	bool bGenerateGlobalCollisionData;
-	bool bGenerateGlobalCrumblingData;
-	bool bGenerateGlobalCrumblingChildrenData;
+	bool bGenerateBreakingData : 1;
+	bool bGenerateCollisionData : 1;
+	bool bGenerateTrailingData : 1;
+	bool bGenerateCrumblingData : 1;
+	bool bGenerateCrumblingChildrenData : 1; 
 
-	FSharedSimulationParameters Shared;
+	bool bGenerateGlobalBreakingData : 1;
+	bool bGenerateGlobalCollisionData : 1;
+	bool bGenerateGlobalCrumblingData : 1;
+	bool bGenerateGlobalCrumblingChildrenData : 1;
 
-	bool EnableGravity;
-	int32 GravityGroupIndex;
-	int32 OneWayInteractionLevel;
-	bool UseInertiaConditioning;
-	bool UseCCD;
-	bool UseMACD;
+	bool EnableGravity : 1;
+	bool UseInertiaConditioning : 1;
+	bool UseCCD : 1;
+	bool UseMACD : 1;
+
+	bool bEnableStrainOnCollision : 1;
+	bool bUseStaticMeshCollisionForTraces : 1;
+	bool bOptimizeConvexes : 1 = true;
+
+	bool bUseDamagePropagation : 1;
+	float BreakDamagePropagationFactor;
+	float ShockDamagePropagationFactor;
+
 	float LinearDamping;
 	float AngularDamping;
 	float InitialOverlapDepenetrationVelocity;
 	float SleepThresholdMultiplier;
 
-	bool bUseDamagePropagation;
-	float BreakDamagePropagationFactor;
-	float ShockDamagePropagationFactor;
+	int32 GravityGroupIndex;
+	int32 OneWayInteractionLevel;
 
 	FCollisionFilterData SimulationFilterData;
 	FCollisionFilterData QueryFilterData;
+
 	void* UserData;
-	bool bEnableStrainOnCollision;
 
-	bool bUseStaticMeshCollisionForTraces;
+#if SIMULATIONPARAMETERS_CACHE_PARAMETERS
+	UE_DEPRECATED(5.5, "No longer used")
+	const FRecordedTransformTrack* RecordedTrack;
 
-	bool bOptimizeConvexes = true;
+	UE_DEPRECATED(5.5, "No longer used")
+	float CacheBeginTime;
+
+	UE_DEPRECATED(5.5, "No longer used")
+	float ReverseCacheBeginTime;
+
+	UE_DEPRECATED(5.5, "No longer used")
+	EGeometryCollectionCacheType CacheType;
+
+	UE_DEPRECATED(5.5, "No longer used")
+	bool bClearCache : 1;
+
+	UE_DEPRECATED(5.5, "No longer used")
+	bool bOwnsTrack : 1;
+
+	UE_DEPRECATED(5.5, "No longer used and underlying variable is deprecated")
+	bool IsCacheRecording()
+	{ 
+		PRAGMA_DISABLE_DEPRECATION_WARNINGS
+		return CacheType == EGeometryCollectionCacheType::Record || CacheType == EGeometryCollectionCacheType::RecordAndPlay; 
+		PRAGMA_ENABLE_DEPRECATION_WARNINGS
+	}
+
+	UE_DEPRECATED(5.5, "No longer used and underlying variable is deprecated")
+	bool IsCachePlaying()
+	{ 
+		PRAGMA_DISABLE_DEPRECATION_WARNINGS
+		return CacheType == EGeometryCollectionCacheType::Play || CacheType == EGeometryCollectionCacheType::RecordAndPlay;
+		PRAGMA_ENABLE_DEPRECATION_WARNINGS
+	}
+#endif
 };
diff --git a/Engine/Source/Runtime/Experimental/Chaos/Public/PhysicsProxy/GeometryCollectionPhysicsProxy.h b/Engine/Source/Runtime/Experimental/Chaos/Public/PhysicsProxy/GeometryCollectionPhysicsProxy.h
--- a/Engine/Source/Runtime/Experimental/Chaos/Public/PhysicsProxy/GeometryCollectionPhysicsProxy.h
+++ b/Engine/Source/Runtime/Experimental/Chaos/Public/PhysicsProxy/GeometryCollectionPhysicsProxy.h
@@ -698,25 +698,6 @@ private:
 	// todo(chaos): Remove this and move to a cook time approach of the SM data based on the GC property
 	FCreateTraceCollisionGeometryCallback CreateTraceCollisionGeometryCallback;
 	
-#ifdef TODO_REIMPLEMENT_RIGID_CACHING
-	TFunction<void(void)> ResetAnimationCacheCallback;
-	TFunction<void(const TArrayView<FTransform> &)> UpdateTransformsCallback;
-	TFunction<void(const int32 & CurrentFrame, const TManagedArray<int32> & RigidBodyID, const TManagedArray<int32>& Level, const TManagedArray<int32>& Parent, const TManagedArray<TSet<int32>>& Children, const TManagedArray<uint32>& SimulationType, const TManagedArray<uint32>& StatusFlags, const FParticlesType& Particles)> UpdateRestStateCallback;
-	TFunction<void(float SolverTime, const TManagedArray<int32> & RigidBodyID, const FParticlesType& Particles, const Chaos::FPBDCollisionConstraints& CollisionRule)> UpdateRecordedStateCallback;
-	TFunction<void(FRecordedTransformTrack& InTrack)> CommitRecordedStateCallback;
-
-	// Index of the first particles for this collection in the larger particle array
-	// Time since this object started simulating
-	float ProxySimDuration;
-
-	// Sync frame numbers so we don't do many syncs when physics is running behind
-	uint32 LastSyncCountGT;
-
-	// Storage for the recorded frame information when we're caching the geometry component results.
-	// Synced back to the component with SyncBeforeDestroy
-	FRecordedTransformTrack RecordedTracks;
-#endif
-
 	// called after we sync the physics thread data ( called on the game thread )
 	TFunction<void()> PostPhysicsSyncCallback;
 	TFunction<void()> PostParticlesCreatedCallback;
diff --git a/Engine/Source/Runtime/Experimental/GeometryCollectionEngine/Private/GeometryCollection/GeometryCollectionComponent.cpp b/Engine/Source/Runtime/Experimental/GeometryCollectionEngine/Private/GeometryCollection/GeometryCollectionComponent.cpp
--- a/Engine/Source/Runtime/Experimental/GeometryCollectionEngine/Private/GeometryCollection/GeometryCollectionComponent.cpp
+++ b/Engine/Source/Runtime/Experimental/GeometryCollectionEngine/Private/GeometryCollection/GeometryCollectionComponent.cpp
@@ -3877,13 +3877,10 @@ PRAGMA_ENABLE_DEPRECATION_WARNINGS
 		SimulationParameters.CollisionGroup = CollisionGroup;
 		SimulationParameters.CollisionSampleFraction = CollisionSampleFraction;
 		SimulationParameters.InitialVelocityType = InitialVelocityType;
-		SimulationParameters.InitialLinearVelocity = InitialLinearVelocity;
-		SimulationParameters.InitialAngularVelocity = InitialAngularVelocity;
-		SimulationParameters.bClearCache = true;
+		SimulationParameters.InitialLinearVelocity = FVector3f(InitialLinearVelocity);
+		SimulationParameters.InitialAngularVelocity = FVector3f(InitialAngularVelocity);
 		SimulationParameters.ObjectType = ObjectType;
 		SimulationParameters.StartAwake = BodyInstance.bStartAwake;
-		SimulationParameters.CacheType = CacheParameters.CacheMode;
-		SimulationParameters.ReverseCacheBeginTime = CacheParameters.ReverseCacheBeginTime;
 		SimulationParameters.bGenerateBreakingData = bNotifyBreaks;
 		SimulationParameters.bGenerateCollisionData = bNotifyCollisions;
 		SimulationParameters.bGenerateTrailingData = bNotifyTrailing;
@@ -3912,6 +3909,14 @@ PRAGMA_ENABLE_DEPRECATION_WARNINGS
 		SimulationParameters.bUseStaticMeshCollisionForTraces = bUseStaticMeshCollisionForTraces;
 		SimulationParameters.bOptimizeConvexes = RestCollection ? RestCollection->bOptimizeConvexes : true;
 
+#if SIMULATIONPARAMETERS_CACHE_PARAMETERS
+PRAGMA_DISABLE_DEPRECATION_WARNINGS
+		SimulationParameters.bClearCache = true;
+		SimulationParameters.CacheType = CacheParameters.CacheMode;
+		SimulationParameters.ReverseCacheBeginTime = CacheParameters.ReverseCacheBeginTime;
+PRAGMA_ENABLE_DEPRECATION_WARNINGS
+#endif
+
 		UPhysicalMaterial* EnginePhysicalMaterial = GetPhysicalMaterial();
 		if (ensure(EnginePhysicalMaterial))
 		{
diff --git a/Engine/Source/Runtime/Experimental/GeometryCollectionEngine/Private/GeometryCollection/GeometryCollectionObject.cpp b/Engine/Source/Runtime/Experimental/GeometryCollectionEngine/Private/GeometryCollection/GeometryCollectionObject.cpp
--- a/Engine/Source/Runtime/Experimental/GeometryCollectionEngine/Private/GeometryCollection/GeometryCollectionObject.cpp
+++ b/Engine/Source/Runtime/Experimental/GeometryCollectionEngine/Private/GeometryCollection/GeometryCollectionObject.cpp
@@ -496,7 +496,9 @@ void UGeometryCollection::GetSharedSimulationParams(FSharedSimulationParameters&
 	const FGeometryCollectionSizeSpecificData& SizeSpecificDefault = GetDefaultSizeSpecificData();
 
 	// we grab the non cached version because this is going to be used to generate the mass attribute which will eventually cache the density value if necessary
-	OutParams.Mass = GetMassOrDensityInternal(OutParams.bMassAsDensity, false);
+	bool bUseMassAsDensity = false;
+	OutParams.Mass = GetMassOrDensityInternal(bUseMassAsDensity, false);
+	OutParams.bMassAsDensity = bUseMassAsDensity;
 	OutParams.MinimumMassClamp = MinimumMassClamp;
 
 	FGeometryCollectionSizeSpecificData InfSize;
