Geometry collection - memory optimization
* packed and re-ordered members in GeometryCollectionComponent, saving 64 bytes
* packed and re-ordered members in GeometryCollectionISMPoolRenderer, saving 16 bytes
* packed and re-ordered members in IPhysicsProxyBase, saving 8 bytes
* packed and re-ordered members in ClusterUnionReplicatedProxyComponent, saving 16 bytes

diff --git a/Engine/Source/Runtime/Experimental/Chaos/Public/Chaos/Framework/PhysicsProxyBase.h b/Engine/Source/Runtime/Experimental/Chaos/Public/Chaos/Framework/PhysicsProxyBase.h
--- a/Engine/Source/Runtime/Experimental/Chaos/Public/Chaos/Framework/PhysicsProxyBase.h
+++ b/Engine/Source/Runtime/Experimental/Chaos/Public/Chaos/Framework/PhysicsProxyBase.h
@@ -92,8 +92,8 @@ public:
 		: Solver(nullptr)
 		, Owner(InOwner)
 		, DirtyIdx(INDEX_NONE)
-		, SyncTimestamp(InProxyTimeStamp)
 		, Type(InType)
+		, SyncTimestamp(InProxyTimeStamp)
 	{}
 
 	UObject* GetOwner() const { return Owner; }
@@ -183,11 +183,13 @@ protected:
 
 private:
 	int32 DirtyIdx;
-	TSharedPtr<FProxyTimestampBase,ESPMode::ThreadSafe> SyncTimestamp;
-	IPhysicsProxyBase* ParentProxy = nullptr;
 protected:
 	/** Proxy type */
 	EPhysicsProxyType Type;
+private:
+	TSharedPtr<FProxyTimestampBase,ESPMode::ThreadSafe> SyncTimestamp;
+	IPhysicsProxyBase* ParentProxy = nullptr;
+protected:
 	int32 InitializedOnStep = INDEX_NONE;
 	int32 IgnoreDataOnStep_Internal = INDEX_NONE;
 
diff --git a/Engine/Source/Runtime/Experimental/GeometryCollectionEngine/Private/GeometryCollection/GeometryCollectionComponent.cpp b/Engine/Source/Runtime/Experimental/GeometryCollectionEngine/Private/GeometryCollection/GeometryCollectionComponent.cpp
--- a/Engine/Source/Runtime/Experimental/GeometryCollectionEngine/Private/GeometryCollection/GeometryCollectionComponent.cpp
+++ b/Engine/Source/Runtime/Experimental/GeometryCollectionEngine/Private/GeometryCollection/GeometryCollectionComponent.cpp
@@ -630,10 +630,11 @@ UGeometryCollectionComponent::UGeometryCollectionComponent(const FObjectInitiali
 	, ReplicationMaxPositionAndVelocityCorrectionLevel(100)
 	, bInitializedRemovalDynamicAttribute(false)
 	, bEnableBoneSelection(false)
-	, ViewLevel(-1)
-	, NavmeshInvalidationTimeSliceIndex(0)
 	, IsObjectDynamic(false)
 	, IsObjectLoading(true)
+	, bIsMoving(false)
+	, ViewLevel(-1)
+	, NavmeshInvalidationTimeSliceIndex(0)
 	, ComponentSpaceTransforms(this)
 	, PhysicsProxy(nullptr)
 #if WITH_EDITOR && WITH_EDITORONLY_DATA
@@ -642,7 +643,6 @@ UGeometryCollectionComponent::UGeometryCollectionComponent(const FObjectInitiali
 #if GEOMETRYCOLLECTION_EDITOR_SELECTION
 	, bIsTransformSelectionModeEnabled(false)
 #endif  // #if GEOMETRYCOLLECTION_EDITOR_SELECTION
-	, bIsMoving(false)
 {
 	// by default tick is registered but disabled, we only need it when we need to update the removal timers
 	// tick will be then enabled only when the root is broken from OnPostPhysicsSync callback
diff --git a/Engine/Source/Runtime/Experimental/GeometryCollectionEngine/Public/GeometryCollection/GeometryCollectionComponent.h b/Engine/Source/Runtime/Experimental/GeometryCollectionEngine/Public/GeometryCollection/GeometryCollectionComponent.h
--- a/Engine/Source/Runtime/Experimental/GeometryCollectionEngine/Public/GeometryCollection/GeometryCollectionComponent.h
+++ b/Engine/Source/Runtime/Experimental/GeometryCollectionEngine/Public/GeometryCollection/GeometryCollectionComponent.h
@@ -1059,10 +1059,6 @@
 	UPROPERTY(meta = (DeprecatedProperty, DeprecationMessage = "Use PhysicalMaterial instead."))
 	float LinearEtherDrag_DEPRECATED;
 
-	/** Uniform angular ether drag. */
-	UPROPERTY(meta=(DeprecatedProperty, DeprecationMessage="Use PhysicalMaterial instead."))
-	float AngularEtherDrag_DEPRECATED;
-
 	/** Physical Properties */
 	UPROPERTY(meta=(DeprecatedProperty, DeprecationMessage="Physical material now derived from render materials, for instance overrides use PhysicalMaterialOverride."))
 	TObjectPtr<const UChaosPhysicalMaterial> PhysicalMaterial_DEPRECATED;
@@ -1455,22 +1451,6 @@
 	FGuid RunTimeDataCollectionGuid;
 #endif
 
-	/** Deprecated for CustomRendererType. */
-	UPROPERTY()
-	TObjectPtr<AGeometryCollectionISMPoolActor> ISMPool_DEPRECATED;
-
-	/** Deprecated for CustomRendererType. */
-	UPROPERTY()
-	bool bAutoAssignISMPool_DEPRECATED = false;
-
-	/** If true, CustomRendererType will be used. If false, CustomRendererType comes from the RestCollection. */
-	UPROPERTY(EditAnywhere, BlueprintReadWrite, Category = "ChaosPhysics|Rendering", meta = (InlineEditConditionToggle))
-	bool bOverrideCustomRenderer = false;
-
-	/** Custom class type that will be used to render the geometry collection instead of using the native rendering. */
-	UPROPERTY(EditAnywhere, BlueprintReadWrite, Category = "ChaosPhysics|Rendering", meta = (editcondition = "bOverrideCustomRenderer", MustImplement = "/Script/GeometryCollectionEngine.GeometryCollectionExternalRenderInterface"))
-	TObjectPtr<UClass> CustomRendererType;
-
 	/** Force the broken state for custom renderer rendering. */
 	uint8 bForceBrokenForCustomRenderer : 1;
 	/** Whether to refresh the custom renderer on physics updates. */
@@ -1482,6 +1462,56 @@
 	/** Force native geometry collection rendering. This is used in the editor fracture mode. */
 	uint8 bForceNativeRenderer : 1;
 
+	/* Per-instance override to enable/disable replication for the geometry collection */
+	UPROPERTY(EditAnywhere, BlueprintReadWrite, Category=Network)
+	bool bEnableReplication;
+
+	/** 
+	 * Enables use of ReplicationAbandonAfterLevel to stop providing network updates to
+	 * clients when the updated particle is of a level higher then specified.
+	 */
+	UPROPERTY(EditAnywhere, BlueprintReadWrite, Category = Network)
+	bool bEnableAbandonAfterLevel;
+
+	/**
+	 * Whether abandoned particles on the client should continue to have collision (i.e.
+	 * still be in the external/internal acceleration structure).
+	 */
+	UPROPERTY(EditAnywhere, BlueprintReadWrite, BlueprintSetter= SetAbandonedParticleCollisionProfileName, Category = Network)
+	FName AbandonedCollisionProfileName;
+
+	/** Deprecated for CustomRendererType. */
+	UPROPERTY()
+	TObjectPtr<AGeometryCollectionISMPoolActor> ISMPool_DEPRECATED;
+
+	/** Custom class type that will be used to render the geometry collection instead of using the native rendering. */
+	UPROPERTY(EditAnywhere, BlueprintReadWrite, Category = "ChaosPhysics|Rendering", meta = (editcondition = "bOverrideCustomRenderer", MustImplement = "/Script/GeometryCollectionEngine.GeometryCollectionExternalRenderInterface"))
+	TObjectPtr<UClass> CustomRendererType;
+
+	/** If true, CustomRendererType will be used. If false, CustomRendererType comes from the RestCollection. */
+	UPROPERTY(EditAnywhere, BlueprintReadWrite, Category = "ChaosPhysics|Rendering", meta = (InlineEditConditionToggle))
+	bool bOverrideCustomRenderer = false;
+
+	/** Deprecated for CustomRendererType. */
+	UPROPERTY()
+	bool bAutoAssignISMPool_DEPRECATED = false;
+
+	void CheckFullyDecayed();
+	bool bAlreadyFullyDecayed = false;
+
+	// todo(chaos): Remove the ability to change this at runtime, as we'll want to use this at cook time instead
+	UPROPERTY(EditAnywhere, BlueprintGetter="GetUseStaticMeshCollisionForTraces", BlueprintSetter="SetUseStaticMeshCollisionForTraces", Category = "Physics")
+	bool bUseStaticMeshCollisionForTraces  = false;
+
+	/**
+	 * If replicating - the cluster level to stop sending corrections for geometry collection chunks.
+	 * recommended for smaller leaf levels when the size of the objects means they are no longer
+	 * gameplay relevant to cut down on required bandwidth to update a collection.
+	 * @see bEnableAbandonAfterLevel
+	 */ 
+	UPROPERTY(meta = (DeprecatedProperty, DeprecationMessage = "GeometryCollection now uses ReplicationAbandonAfterLevel instead of ReplicationAbandonClusterLevel."))
+	int32 ReplicationAbandonClusterLevel_DEPRECATED;
+
 	/** A custom renderer object created from CustomRenderType. */
 	UPROPERTY(SkipSerialization)
 	TScriptInterface<IGeometryCollectionExternalRenderInterface> CustomRenderer;
@@ -1521,24 +1551,6 @@
 	GEOMETRYCOLLECTIONENGINE_API void UpdateRBCollisionEventRegistration();
 	GEOMETRYCOLLECTIONENGINE_API void UpdateGlobalCollisionEventRegistration();
 	GEOMETRYCOLLECTIONENGINE_API void UpdateGlobalRemovalEventRegistration();
-	
-	/* Per-instance override to enable/disable replication for the geometry collection */
-	UPROPERTY(EditAnywhere, BlueprintReadWrite, Category=Network)
-	bool bEnableReplication;
-
-	/** 
-	 * Enables use of ReplicationAbandonAfterLevel to stop providing network updates to
-	 * clients when the updated particle is of a level higher then specified.
-	 */
-	UPROPERTY(EditAnywhere, BlueprintReadWrite, Category = Network)
-	bool bEnableAbandonAfterLevel;
 
 	/**
-	 * Whether abandoned particles on the client should continue to have collision (i.e.
-	 * still be in the external/internal acceleration structure).
-	 */
-	UPROPERTY(EditAnywhere, BlueprintReadWrite, BlueprintSetter= SetAbandonedParticleCollisionProfileName, Category = Network)
-	FName AbandonedCollisionProfileName;
-
-	/**
 	 * A per-level collision profile name. If the name is set to NONE or an invalid collision profile, nothing will be changed.
@@ -1555,14 +1567,5 @@
 	TArray<FName> CollisionProfilePerParticle;
 
 	/**
-	 * If replicating - the cluster level to stop sending corrections for geometry collection chunks.
-	 * recommended for smaller leaf levels when the size of the objects means they are no longer
-	 * gameplay relevant to cut down on required bandwidth to update a collection.
-	 * @see bEnableAbandonAfterLevel
-	 */ 
-	UPROPERTY(meta = (DeprecatedProperty, DeprecationMessage = "GeometryCollection now uses ReplicationAbandonAfterLevel instead of ReplicationAbandonClusterLevel."))
-	int32 ReplicationAbandonClusterLevel_DEPRECATED;
-
-	/**
 	* If replicating - the cluster level after which replication will not happen 
 	* @see bEnableAbandonAfterLevel
@@ -1611,9 +1614,6 @@
 
 	GEOMETRYCOLLECTIONENGINE_API virtual bool ProcessRepData(float DeltaTime, float SimTime);
 
-	void CheckFullyDecayed();
-	bool bAlreadyFullyDecayed = false;
-
 	int32 VersionProcessed = INDEX_NONE;
 	int32 DynamicRepDataVersionProcessed = INDEX_NONE;
 
@@ -1637,16 +1637,21 @@
 	FTransform3f GetCurrentTransform(int32 Index) const;
 	void ComputeCurrentGlobalsMatrices(TArray<FTransform3f>& OutTransforms) const;
 
-	bool bInitializedRemovalDynamicAttribute;
-	bool bEnableBoneSelection;
+	uint8 bInitializedRemovalDynamicAttribute : 1;
+	uint8 bEnableBoneSelection : 1;
+	uint8 IsObjectDynamic : 1;
+	uint8 IsObjectLoading : 1;
+
+	/** True if GeometryCollection transforms have changed from previous tick. */
+	uint8 bIsMoving : 1;
+
 	int ViewLevel;
 
 	uint32 NavmeshInvalidationTimeSliceIndex;
-	bool IsObjectDynamic;
-	bool IsObjectLoading;
 
 	FCollisionFilterData InitialSimFilter;
 	FCollisionFilterData InitialQueryFilter;
+	float CurrentCacheTime;
 	FChaosUserData PhysicsUserData;
 
 #if WITH_EDITORONLY_DATA
@@ -1714,7 +1719,6 @@
 	*/
 	mutable FBox ComponentSpaceBounds;
 
-	float CurrentCacheTime;
 	TArray<bool> EventsPlayed;
 
 	FGeometryCollectionPhysicsProxy* PhysicsProxy;
@@ -1749,10 +1753,6 @@
 	TArray<int32> EmbeddedInstanceIndex;
 #endif
 
-	// todo(chaos): Remove the ability to change this at runtime, as we'll want to use this at cook time instead
-	UPROPERTY(EditAnywhere, BlueprintGetter="GetUseStaticMeshCollisionForTraces", BlueprintSetter="SetUseStaticMeshCollisionForTraces", Category = "Physics")
-	bool bUseStaticMeshCollisionForTraces  = false;
-
 	GEOMETRYCOLLECTIONENGINE_API bool IsEmbeddedGeometryValid() const;
 	GEOMETRYCOLLECTIONENGINE_API void ClearEmbeddedGeometry();
 
@@ -1807,12 +1807,15 @@
 	/** The clusters we need to replicate */
 	TUniquePtr<TSet<Chaos::FPBDRigidClusteredParticleHandle*>> ClustersToRep;
 
+	double LastHardsnapTimeInMs = 0;
+
 	/** One off activation is processed in the same order as server so remember the last one we processed */
 	int32 OneOffActivatedProcessed = 0;
-	double LastHardsnapTimeInMs = 0;
 
-	/** True if GeometryCollection transforms have changed from previous tick. */
-	bool bIsMoving;
+public:
+	/** Uniform angular ether drag. */
+	UPROPERTY(meta=(DeprecatedProperty, DeprecationMessage="Use PhysicalMaterial instead."))
+	float AngularEtherDrag_DEPRECATED;
 
 private:
 	struct FBrokenAndDecayedStates
