Add support for string based anytypes
- Added a generic LogString that can take any string compatible type
- Deprecated the previous LogString
- Added converter classes to support the type policies

diff --git a/Engine/Plugins/Experimental/GeometryCollectionPlugin/Source/GeometryCollectionNodes/Public/Dataflow/GeometryCollectionNodes.h b/Engine/Plugins/Experimental/GeometryCollectionPlugin/Source/GeometryCollectionNodes/Public/Dataflow/GeometryCollectionNodes.h
--- a/Engine/Plugins/Experimental/GeometryCollectionPlugin/Source/GeometryCollectionNodes/Public/Dataflow/GeometryCollectionNodes.h
+++ b/Engine/Plugins/Experimental/GeometryCollectionPlugin/Source/GeometryCollectionNodes/Public/Dataflow/GeometryCollectionNodes.h
@@ -95,9 +95,9 @@ public:
 /**
  *
  * Description for this node
- *
+ * DEPRECATED - use Print node ( core nodes ) 
  */
-USTRUCT()
+USTRUCT(meta = (Deprecated = "5.5"))
 struct FPrintStringDataflowNode : public FDataflowNode
 {
 	GENERATED_USTRUCT_BODY()
@@ -133,9 +133,9 @@ public:
 /**
  *
  * Description for this node
- *
- */
-USTRUCT()
+ * DEPRECATED - use Print node ( core nodes ) 
+  */
+USTRUCT(meta=(Deprecated = "5.5"))
 struct FLogStringDataflowNode : public FDataflowNode
 {
 	GENERATED_USTRUCT_BODY()
diff --git a/Engine/Source/Runtime/Experimental/Dataflow/Core/Private/Dataflow/DataflowCoreNodes.cpp b/Engine/Source/Runtime/Experimental/Dataflow/Core/Private/Dataflow/DataflowCoreNodes.cpp
--- a/Engine/Source/Runtime/Experimental/Dataflow/Core/Private/Dataflow/DataflowCoreNodes.cpp
+++ b/Engine/Source/Runtime/Experimental/Dataflow/Core/Private/Dataflow/DataflowCoreNodes.cpp
@@ -1,8 +1,18 @@
 // Copyright Epic Games, Inc. All Rights Reserved.
 
 #include "Dataflow/DataflowCoreNodes.h"
+#include "Dataflow/DataflowNode.h"
+#include "Dataflow/DataflowNodeFactory.h"
 
-
+namespace Dataflow
+{
+	void RegisterCoreNodes()
+	{
+		DATAFLOW_NODE_REGISTER_CREATION_FACTORY(FDataflowReRouteNode);
+		DATAFLOW_NODE_REGISTER_CREATION_FACTORY(FDataflowBranchNode);
+		DATAFLOW_NODE_REGISTER_CREATION_FACTORY(FDataflowPrintNode);
+	}
+}
 
 FDataflowReRouteNode::FDataflowReRouteNode(const Dataflow::FNodeParameters& Param, FGuid InGuid)
 	: Super(Param, InGuid)
@@ -72,4 +82,17 @@ bool FDataflowBranchNode::OnOutputTypeChanged(const FDataflowOutput* Input)
 		  SetInputConcreteType(&TrueValue, Input->GetType())
 		| SetInputConcreteType(&FalseValue, Input->GetType()) //-V792
 		);
+}
+
+
+FDataflowPrintNode::FDataflowPrintNode(const Dataflow::FNodeParameters& InParam, FGuid InGuid)
+	: FDataflowNode(InParam, InGuid)
+{
+	RegisterInputConnection(&Value);
+}
+
+void FDataflowPrintNode::Evaluate(Dataflow::FContext& Context, const FDataflowOutput* Out) const
+{
+	const FString InValue = GetValue(Context, &Value);
+	UE_LOG(LogTemp, Warning, TEXT("[Dataflow Print] %s"), *InValue);
 }
\ No newline at end of file
diff --git a/Engine/Source/Runtime/Experimental/Dataflow/Core/Private/Dataflow/DataflowNodeFactory.cpp b/Engine/Source/Runtime/Experimental/Dataflow/Core/Private/Dataflow/DataflowNodeFactory.cpp
--- a/Engine/Source/Runtime/Experimental/Dataflow/Core/Private/Dataflow/DataflowNodeFactory.cpp
+++ b/Engine/Source/Runtime/Experimental/Dataflow/Core/Private/Dataflow/DataflowNodeFactory.cpp
@@ -15,12 +15,9 @@ namespace Dataflow
 
 	void FNodeFactory::RegisterDefaultNodes()
 	{
-		// by default always register a reroute node
-		DATAFLOW_NODE_REGISTER_CREATION_FACTORY(FDataflowReRouteNode);
-		DATAFLOW_NODE_REGISTER_CREATION_FACTORY(FDataflowBranchNode);
+		Dataflow::RegisterCoreNodes();
 	}
 
-
 	TSharedPtr<FDataflowNode> FNodeFactory::NewNodeFromRegisteredType(FGraph& Graph, const FNewNodeParameters& Param)
 	{ 
 		if (ClassMap.Contains(Param.Type))
diff --git a/Engine/Source/Runtime/Experimental/Dataflow/Core/Public/Dataflow/DataflowAnyType.h b/Engine/Source/Runtime/Experimental/Dataflow/Core/Public/Dataflow/DataflowAnyType.h
--- a/Engine/Source/Runtime/Experimental/Dataflow/Core/Public/Dataflow/DataflowAnyType.h
+++ b/Engine/Source/Runtime/Experimental/Dataflow/Core/Public/Dataflow/DataflowAnyType.h
@@ -36,3 +36,28 @@ struct FDataflowNumericTypes : public FDataflowAnyType
 	UPROPERTY(EditAnywhere, Category=Value)
 	double Value = 0.0;
 };
+
+USTRUCT()
+struct FDataflowStringTypes : public FDataflowAnyType
+{
+	using FPolicyType = FDataflowStringTypePolicy;
+	using FStorageType = FString;
+
+	GENERATED_USTRUCT_BODY()
+
+	UPROPERTY(EditAnywhere, Category = Value)
+	FString Value;
+};
+
+USTRUCT()
+struct FDataflowStringConvertibleTypes : public FDataflowAnyType
+{
+	using FPolicyType = FDataflowStringConvertibleTypePolicy;
+	using FStorageType = FString;
+
+	GENERATED_USTRUCT_BODY()
+
+	UPROPERTY(EditAnywhere, Category = Value)
+	FString Value;
+};
+
diff --git a/Engine/Source/Runtime/Experimental/Dataflow/Core/Public/Dataflow/DataflowCoreNodes.h b/Engine/Source/Runtime/Experimental/Dataflow/Core/Public/Dataflow/DataflowCoreNodes.h
--- a/Engine/Source/Runtime/Experimental/Dataflow/Core/Public/Dataflow/DataflowCoreNodes.h
+++ b/Engine/Source/Runtime/Experimental/Dataflow/Core/Public/Dataflow/DataflowCoreNodes.h
@@ -61,4 +61,29 @@ public:
 private:
 	virtual bool OnInputTypeChanged(const FDataflowInput* Input) override;
 	virtual bool OnOutputTypeChanged(const FDataflowOutput* Input) override;
-};
\ No newline at end of file
+};
+
+/** 
+* Print value in the log
+* Supports any type comnvertible to a string 
+*/
+USTRUCT()
+struct FDataflowPrintNode : public FDataflowNode
+{
+	GENERATED_USTRUCT_BODY()
+
+	DATAFLOW_NODE_DEFINE_INTERNAL(FDataflowPrintNode, "Print", "Core", "")
+
+public:
+	FDataflowPrintNode(const Dataflow::FNodeParameters& Param, FGuid InGuid = FGuid::NewGuid());
+	virtual void Evaluate(Dataflow::FContext& Context, const FDataflowOutput* Out) const;
+
+public:
+	UPROPERTY(meta = (DataflowInput, DisplayName = "Value"))
+	FDataflowStringConvertibleTypes Value;
+};
+
+namespace Dataflow
+{
+	void RegisterCoreNodes();
+}
diff --git a/Engine/Source/Runtime/Experimental/Dataflow/Core/Public/Dataflow/DataflowInputOutput.h b/Engine/Source/Runtime/Experimental/Dataflow/Core/Public/Dataflow/DataflowInputOutput.h
--- a/Engine/Source/Runtime/Experimental/Dataflow/Core/Public/Dataflow/DataflowInputOutput.h
+++ b/Engine/Source/Runtime/Experimental/Dataflow/Core/Public/Dataflow/DataflowInputOutput.h
@@ -157,8 +157,9 @@ public:
 		TAnyType::FPolicyType::VisitPolicyByType(GetType(),
 			[this, &Context, &InVal](auto SingleTypePolicy)
 			{
-				using FSingleType = typename decltype(SingleTypePolicy)::FReturnType;
-				FSingleType ValueToSet(InVal);
+				using FSingleType = typename decltype(SingleTypePolicy)::FType;
+				FSingleType ValueToSet{};
+				FDataflowConverter<typename TAnyType::FStorageType>::To(InVal, ValueToSet);
 				Context.SetData(CacheKey(), GetProperty(), Forward<FSingleType>(ValueToSet), GetOwningNodeGuid(), GetOwningNodeValueHash(), Dataflow::FTimestamp::Current());
 			});
 	}
@@ -228,9 +229,10 @@ typename TAnyType::FStorageType FDataflowInput::GetValueFromAnyType(Dataflow::FC
 					TAnyType::FPolicyType::VisitPolicyByType(GetType(),
 						[this, &Context, &CacheEntry, &ReturnValue](auto SingleTypePolicy)
 						{
-							using FSingleType = typename decltype(SingleTypePolicy)::FReturnType;
+							using FSingleType = typename decltype(SingleTypePolicy)::FType;
 							FSingleType Default{};
-							ReturnValue = (*CacheEntry)->GetTypedData<FSingleType>(Context, nullptr, Default);
+							const FSingleType& CachedValue = (*CacheEntry)->GetTypedData<FSingleType>(Context, nullptr, Default);
+							FDataflowConverter<typename TAnyType::FStorageType>::From(CachedValue, ReturnValue);
 						});
 				}
 			}
diff --git a/Engine/Source/Runtime/Experimental/Dataflow/Core/Public/Dataflow/DataflowTypePolicy.h b/Engine/Source/Runtime/Experimental/Dataflow/Core/Public/Dataflow/DataflowTypePolicy.h
--- a/Engine/Source/Runtime/Experimental/Dataflow/Core/Public/Dataflow/DataflowTypePolicy.h
+++ b/Engine/Source/Runtime/Experimental/Dataflow/Core/Public/Dataflow/DataflowTypePolicy.h
@@ -27,7 +27,8 @@ struct FDataflowAllTypesPolicy : public IDataflowTypePolicy
 template <typename T>
 struct TDataflowSingleTypePolicy : public IDataflowTypePolicy
 {
-	using FReturnType = T;
+	using FType = T;
+
 	virtual bool SupportsType(FName InType) const override
 	{
 		return SupportsTypeStatic(InType);
@@ -82,11 +83,12 @@ struct TDataflowMultiTypePolicy<>: public IDataflowTypePolicy
 	}
 };
 
+Expose_TNameOf(bool);
+
 template <typename T, typename... TTypes>
 struct TDataflowMultiTypePolicy<T, TTypes...>: public TDataflowMultiTypePolicy<TTypes...>
 {
 	using Super = TDataflowMultiTypePolicy<TTypes...>;
-	using FReturnType = T;
 
 	virtual bool SupportsType(FName InType) const override
 	{
@@ -120,3 +122,105 @@ struct FDataflowNumericTypePolicy :
 	public TDataflowMultiTypePolicy<double, float, int64, uint64, int32, uint32, int16, uint16, int8, uint8>
 {
 };
+
+struct FDataflowStringTypePolicy :
+	public TDataflowMultiTypePolicy<FString, FName>
+{
+};
+
+/**
+* string comvertible types
+* - FString / Fname
+* - Numeric types ( see FDataflowNumericTypePolicy )
+* - bool
+*/
+struct FDataflowStringConvertibleTypePolicy : IDataflowTypePolicy
+{
+	virtual bool SupportsType(FName InType) const override
+	{
+		return SupportsTypeStatic(InType);
+	}
+
+	static bool SupportsTypeStatic(FName InType)
+	{
+		return FDataflowStringTypePolicy::SupportsTypeStatic(InType)
+			|| FDataflowNumericTypePolicy::SupportsTypeStatic(InType)
+			|| TDataflowSingleTypePolicy<bool>::SupportsTypeStatic(InType)
+			;
+	}
+
+	template <typename TVisitor>
+	static bool VisitPolicyByType(FName RequestedType, TVisitor Visitor)
+	{
+		return FDataflowStringTypePolicy::VisitPolicyByType(RequestedType, Visitor)
+			|| FDataflowNumericTypePolicy::VisitPolicyByType(RequestedType, Visitor)
+			|| TDataflowSingleTypePolicy<bool>::VisitPolicyByType(RequestedType, Visitor)
+			;
+	}
+
+	static IDataflowTypePolicy* GetInterface()
+	{
+		static FDataflowStringConvertibleTypePolicy Instance;
+		return &Instance;
+	}
+};
+
+// type Converters
+
+template <typename T>
+struct FDataflowConverter
+{
+	template <typename TFromType>
+	static void From(const TFromType& From, T& To) { To = From; }
+
+	template <typename TToType>
+	static void To(const T& From, TToType& To) { To = From; }
+};
+
+template <>
+struct FDataflowConverter<FString>
+{
+	template <typename TFromType>
+	static void From(const TFromType& From, FString& To)
+	{
+		if constexpr (std::is_same_v<TFromType, FName>)
+		{
+			To = From.ToString();
+		}
+		else if constexpr (std::is_same_v<TFromType, bool>)
+		{
+			To = FString((From == true) ? "True" : "False");
+		}
+		else if constexpr (std::is_convertible_v<TFromType, double>)
+		{
+			To = FString::SanitizeFloat(double(From), 0);
+		}
+		else
+		{
+			To = From;
+		}
+	}
+
+	template <typename TToType>
+	static void To(const FString& From, TToType& To)
+	{
+		if constexpr (std::is_same_v<TToType, FName>)
+		{
+			To = FName(From);
+		}
+		else if constexpr (std::is_same_v<TToType, bool>)
+		{
+			To = From.ToBool();
+		}
+		else if constexpr (std::is_convertible_v<double, TToType>)
+		{
+			double Result = {};
+			LexTryParseString(Result, *From);
+			To = Result;
+		}
+		else
+		{
+			To = From;
+		}
+	}
+};
\ No newline at end of file
