Iris - Ensure if we're trying to write an object without a protocol.

diff --git a/Engine/Source/Runtime/Experimental/Iris/Core/Private/Iris/ReplicationSystem/ReplicationWriter.cpp b/Engine/Source/Runtime/Experimental/Iris/Core/Private/Iris/ReplicationSystem/ReplicationWriter.cpp
--- a/Engine/Source/Runtime/Experimental/Iris/Core/Private/Iris/ReplicationSystem/ReplicationWriter.cpp
+++ b/Engine/Source/Runtime/Experimental/Iris/Core/Private/Iris/ReplicationSystem/ReplicationWriter.cpp
@@ -2032,6 +2032,11 @@ FReplicationWriter::EWriteObjectStatus FReplicationWriter::WriteObjectAndSubObje
 	uint8* ReplicatedObjectStateBuffer = NetRefHandleManager->GetReplicatedObjectStateBufferNoCheck(InternalIndex);
 
 	const bool bIsInitialState = IsInitialState(State);
+	if (InternalIndex != ObjectIndexForOOBAttachment && !ObjectData.Protocol)
+	{
+		ensureMsgf(ObjectData.Protocol, TEXT("Failed to replicate ( InternalIndex: %u ) %s, Protocol: nullptr, InstanceProtocol pointer: %p, HasCachedCreationInfo: %u"), InternalIndex, *NetRefHandle.ToString(), ObjectData.InstanceProtocol, ObjectData.bHasCachedCreationInfo);
+		return EWriteObjectStatus::NoInstanceProtocol;
+	}
 
 	// Filter out changemasks that are not supposed to be replicated to this connection
 	const bool bNeedToFilterChangeMask = (bIsInitialState || Info.HasDirtyChangeMask) && Info.HasChangemaskFilter;
@@ -2144,9 +2149,10 @@ FReplicationWriter::EWriteObjectStatus FReplicationWriter::WriteObjectAndSubObje
 					UE_NET_TRACE_SCOPE(CreationInfo, Writer, Context.GetTraceCollector(), ENetTraceVerbosity::Trace);
 
 					// Warn if we cannot replicate this object
-					if (!ObjectData.InstanceProtocol && !(ObjectData.bHasCachedCreationInfo == 1U))
+					if (!ObjectData.Protocol || (!ObjectData.InstanceProtocol && !(ObjectData.bHasCachedCreationInfo == 1U)))
 					{
-						UE_LOG_REPLICATIONWRITER_WARNING(TEXT("Failed to replicate ( InternalIndex: %u ) %s, ProtocolName: %s, Currently we do not support creating a remote instance when the instance has been detached."), InternalIndex, *NetRefHandle.ToString(), ToCStr(ObjectData.Protocol->DebugName));
+						UE_LOG_REPLICATIONWRITER_WARNING(TEXT("Failed to replicate ( InternalIndex: %u ) %s, ProtocolName: %s, InstanceProtocol pointer: %p, HasCachedCreationInfo: %u"), InternalIndex, *NetRefHandle.ToString(), (ObjectData.Protocol ? ToCStr(ObjectData.Protocol->DebugName) : TEXT("nullptr")), ObjectData.InstanceProtocol, ObjectData.bHasCachedCreationInfo);
+						ensureMsgf(ObjectData.Protocol, TEXT("Failed to replicate ( InternalIndex: %u ) %s, Protocol: nullptr, InstanceProtocol pointer: %p, HasCachedCreationInfo: %u"), InternalIndex, *NetRefHandle.ToString(), ObjectData.InstanceProtocol, ObjectData.bHasCachedCreationInfo);
 						return EWriteObjectStatus::NoInstanceProtocol;
 					}
 
