geometry collection - memory optimization - allocate field data on demand , remove unused arrays and packed existing private members
- saves 320 bytes from the proxy
- cleaned up the buffer command API for the proxy to be clear about what can be called from the game or physics thread

diff --git a/Engine/Source/Programs/HeadlessChaos/Private/GeometryCollection/GeometryCollectionTestSimulationField.cpp b/Engine/Source/Programs/HeadlessChaos/Private/GeometryCollection/GeometryCollectionTestSimulationField.cpp
--- a/Engine/Source/Programs/HeadlessChaos/Private/GeometryCollection/GeometryCollectionTestSimulationField.cpp
+++ b/Engine/Source/Programs/HeadlessChaos/Private/GeometryCollection/GeometryCollectionTestSimulationField.cpp
@@ -56,7 +56,7 @@ namespace GeometryCollectionTest
 		RadialMaskTmp->ExteriorValue = (int32)EObjectStateTypeEnum::Chaos_Object_Kinematic;
 		RadialMaskTmp->SetMaskCondition = ESetMaskConditionType::Field_Set_IFF_NOT_Interior;
 		FName TargetNameTmp = GetFieldPhysicsName(EFieldPhysicsType::Field_DynamicState);
-		Collection->PhysObject->BufferCommand(UnitTest.Solver, { TargetNameTmp, RadialMaskTmp });
+		Collection->PhysObject->BufferFieldCommand_Internal(UnitTest.Solver, { TargetNameTmp, RadialMaskTmp });
 
 		UnitTest.Initialize();
 		UnitTest.Advance();
@@ -102,7 +102,7 @@ namespace GeometryCollectionTest
 		RadialMask->ExteriorValue = (int32)EObjectStateTypeEnum::Chaos_Object_Kinematic;
 		RadialMask->SetMaskCondition = ESetMaskConditionType::Field_Set_IFF_NOT_Interior;
 		FName TargetName = GetFieldPhysicsName(EFieldPhysicsType::Field_DynamicState);
-		Collection->PhysObject->BufferCommand(UnitTest.Solver, { TargetName, RadialMask });
+		Collection->PhysObject->BufferFieldCommand_Internal(UnitTest.Solver, { TargetName, RadialMask });
 
 		{
 			UnitTest.Advance();
@@ -195,7 +195,7 @@ namespace GeometryCollectionTest
 			if (Frame == 1)
 			{
 				FName TargetName = GetFieldPhysicsName(EFieldPhysicsType::Field_DynamicState);
-				Collection->PhysObject->BufferCommand(UnitTest.Solver, { TargetName, RadialMask });
+				Collection->PhysObject->BufferFieldCommand_Internal(UnitTest.Solver, { TargetName, RadialMask });
 			}
 
 			if (Frame >= 2)
@@ -490,7 +490,7 @@ namespace GeometryCollectionTest
 			if (Frame == 31)
 			{
 				EFieldPhysicsType PhysicsType = GetGeometryCollectionPhysicsType(EGeometryCollectionPhysicsTypeEnum::Chaos_CollisionGroup);
-				Collection[1]->PhysObject->BufferCommand(UnitTest.Solver, { PhysicsType, RadialMask });
+				Collection[1]->PhysObject->BufferFieldCommand_Internal(UnitTest.Solver, { PhysicsType, RadialMask });
 			}
 		}
 		// The bottom boxes should have fallen below the ground level, box 2 now on the ground with box 3 on top
diff --git a/Engine/Source/Runtime/Experimental/Chaos/Private/PhysicsProxy/GeometryCollectionPhysicsProxy.cpp b/Engine/Source/Runtime/Experimental/Chaos/Private/PhysicsProxy/GeometryCollectionPhysicsProxy.cpp
--- a/Engine/Source/Runtime/Experimental/Chaos/Private/PhysicsProxy/GeometryCollectionPhysicsProxy.cpp
+++ b/Engine/Source/Runtime/Experimental/Chaos/Private/PhysicsProxy/GeometryCollectionPhysicsProxy.cpp
@@ -499,21 +499,21 @@ FGeometryCollectionPhysicsProxy::FGeometryCollectionPhysicsProxy(
 	const Chaos::EMultiBufferMode BufferMode)
 	: Base(InOwner)
 	, Parameters(SimulationParameters)
+	, WorldTransform_External(FTransform::Identity)
 	, NumTransforms(INDEX_NONE)
 	, NumEffectiveParticles(INDEX_NONE)
 	, BaseParticleIndex(INDEX_NONE)
+	, CollisionParticlesPerObjectFraction(CollisionParticlesPerObjectFractionDefault)
+	, bIsGameThreadWorldTransformDirty(false)
+	, bHasBuiltGeometryOnPT(false)
+	, bHasBuiltGeometryOnGT(false)
 	, IsObjectDynamic(false)
 	, IsObjectLoading(true)
 	, IsObjectDeleting(false)
 	, SimFilter(InSimFilter)
 	, QueryFilter(InQueryFilter)
-	, CollisionParticlesPerObjectFraction(CollisionParticlesPerObjectFractionDefault)
 	, PhysicsThreadCollection(Parameters.RestCollectionShared)
 	, GameThreadCollection(GameThreadCollectionIn)
-	, WorldTransform_External(FTransform::Identity)
-	, bIsGameThreadWorldTransformDirty(false)
-	, bHasBuiltGeometryOnPT(false)
-	, bHasBuiltGeometryOnGT(false)
 	, CollectorGuid(InCollectorGuid)
 {
 	// We rely on a guarded buffer.
@@ -1345,7 +1345,8 @@ void FGeometryCollectionPhysicsProxy::InitializeBodiesPT(Chaos::FPBDRigidsSolver
 
 			Cmd.MetaData.Add(FFieldSystemMetaData::EMetaType::ECommandData_ProcessingResolution, TUniquePtr<FFieldSystemMetaDataProcessingResolution>(ResolutionData));
 			RigidsSolver->GetGeometryCollectionPhysicsProxiesField_Internal().Add(this);
-			Commands.Add(Cmd);
+			FFieldData& FieldData = GetOrCreateFieldData_Internal();
+			FieldData.Commands.Add(Cmd);
 		}
 		Parameters.InitializationCommands.Empty();
 		FieldParameterUpdateCallback(RigidsSolver, false);
@@ -5555,6 +5556,41 @@ void BuildSimulationData(Chaos::FErrorReporter& ErrorReporter, FGeometryCollecti
 // FIELDS
 //==============================================================================
 
+FGeometryCollectionPhysicsProxy::FFieldData& FGeometryCollectionPhysicsProxy::GetOrCreateFieldData_Internal()
+{
+	if (!FieldData_Internal)
+	{
+		FieldData_Internal = MakeUnique<FFieldData>();
+	}
+	return *FieldData_Internal;
+}
+
+void FGeometryCollectionPhysicsProxy::BufferFieldCommand_External(FFieldSystemCommand&& Command)
+{
+	check(IsInGameThread());
+	if (Chaos::FPhysicsSolver* RBDSolver = GetSolver<Chaos::FPhysicsSolver>())
+	{
+		RBDSolver->EnqueueCommandImmediate(
+			[this, RBDSolver, FieldCommand = MoveTemp(Command)]()
+			{
+				BufferFieldCommand_Internal(RBDSolver, FieldCommand);
+			});
+	}
+}
+
+// deprecated - simple forward to BufferFieldCommand_Internal
+void FGeometryCollectionPhysicsProxy::BufferCommand(Chaos::FPBDRigidsSolver* RigidsSolver, const FFieldSystemCommand& Command)
+{
+	BufferFieldCommand_Internal(RigidsSolver, Command);
+}
+
+void FGeometryCollectionPhysicsProxy::BufferFieldCommand_Internal(Chaos::FPBDRigidsSolver* RigidsSolver, const FFieldSystemCommand& Command)
+{
+	check(RigidsSolver != nullptr);
+	RigidsSolver->GetGeometryCollectionPhysicsProxiesField_Internal().Add(this);
+	GetOrCreateFieldData_Internal().Commands.Add(Command);
+}
+
 void FGeometryCollectionPhysicsProxy::FieldParameterUpdateCallback(Chaos::FPBDRigidsSolver* RigidSolver, const bool bUpdateViews)
 {
 	SCOPE_CYCLE_COUNTER(STAT_ParamUpdateField_Object);
@@ -5564,8 +5600,10 @@ void FGeometryCollectionPhysicsProxy::FieldParameterUpdateCallback(Chaos::FPBDRi
 	Chaos::FPBDPositionConstraints PositionTarget;
 	TMap<int32, int32> TargetedParticles;
 
+	FFieldData& FieldData = GetOrCreateFieldData_Internal();
+
 	// Process Particle-Collection commands
-	int32 NumCommands = Commands.Num();
+	int32 NumCommands = FieldData.Commands.Num();
 	if (NumCommands && !RigidSolver->IsShuttingDown() && Collection.GetNumTransforms())
 	{
 		TArray<int32> CommandsToRemove;
@@ -5578,24 +5616,24 @@ void FGeometryCollectionPhysicsProxy::FieldParameterUpdateCallback(Chaos::FPBDRi
 
 		for (int32 CommandIndex = 0; CommandIndex < NumCommands; CommandIndex++)
 		{
-			FFieldSystemCommand& FieldCommand = Commands[CommandIndex];
+			FFieldSystemCommand& FieldCommand = FieldData.Commands[CommandIndex];
 			if (IsParameterFieldValid(FieldCommand) || FieldCommand.PhysicsType == EFieldPhysicsType::Field_InitialLinearVelocity || FieldCommand.PhysicsType == EFieldPhysicsType::Field_InitialAngularVelocity)
 			{
-				if (Chaos::BuildFieldSamplePoints(this, RigidSolver, FieldCommand, ExecutionDatas, PrevResolutionType, PrevFilterType, PrevObjectType, PrevPositionType))
+				if (Chaos::BuildFieldSamplePoints(this, RigidSolver, FieldCommand, FieldData.ExecutionDatas, PrevResolutionType, PrevFilterType, PrevObjectType, PrevPositionType))
 				{
 					const Chaos::FReal TimeSeconds = RigidSolver->GetSolverTime() - FieldCommand.TimeCreation;
 
 					FFieldContext FieldContext(
-						ExecutionDatas,
+						FieldData.ExecutionDatas,
 						FieldCommand.MetaData,
 						TimeSeconds);
 
-					TArray<Chaos::FGeometryParticleHandle*>& ParticleHandles = ExecutionDatas.ParticleHandles[(uint8)EFieldCommandHandlesType::InsideHandles];
+					TArray<Chaos::FGeometryParticleHandle*>& ParticleHandles = FieldData.ExecutionDatas.ParticleHandles[(uint8)EFieldCommandHandlesType::InsideHandles];
 
 					if (FieldCommand.RootNode->Type() == FFieldNodeBase::EFieldType::EField_Int32)
 					{
-						TArray<int32>& FinalResults = ExecutionDatas.IntegerResults[(uint8)EFieldCommandResultType::FinalResult];
-						ResetResultsArray < int32 >(ExecutionDatas.SamplePositions.Num(), FinalResults, 0);
+						TArray<int32>& FinalResults = FieldData.ExecutionDatas.IntegerResults[(uint8)EFieldCommandResultType::FinalResult];
+						ResetResultsArray < int32 >(FieldData.ExecutionDatas.SamplePositions.Num(), FinalResults, 0);
 
 						TFieldArrayView<int32> ResultsView(FinalResults, 0, FinalResults.Num());
 
@@ -5651,14 +5689,14 @@ void FGeometryCollectionPhysicsProxy::FieldParameterUpdateCallback(Chaos::FPBDRi
 						}
 						else
 						{
-							Chaos::FieldIntegerParameterUpdate(RigidSolver, FieldCommand, ExecutionDatas.ParticleHandles[(uint8)EFieldCommandHandlesType::InsideHandles],
+							Chaos::FieldIntegerParameterUpdate(RigidSolver, FieldCommand, FieldData.ExecutionDatas.ParticleHandles[(uint8)EFieldCommandHandlesType::InsideHandles],
 								FieldContext, PositionTarget, TargetedParticles, FinalResults);
 						}
 					}
 					else if (FieldCommand.RootNode->Type() == FFieldNodeBase::EFieldType::EField_FVector)
 					{
-						TArray<FVector>& FinalResults = ExecutionDatas.VectorResults[(uint8)EFieldCommandResultType::FinalResult];
-						ResetResultsArray < FVector >(ExecutionDatas.SamplePositions.Num(), FinalResults, FVector::ZeroVector);
+						TArray<FVector>& FinalResults = FieldData.ExecutionDatas.VectorResults[(uint8)EFieldCommandResultType::FinalResult];
+						ResetResultsArray < FVector >(FieldData.ExecutionDatas.SamplePositions.Num(), FinalResults, FVector::ZeroVector);
 
 						TFieldArrayView<FVector> ResultsView(FinalResults, 0, FinalResults.Num());
 
@@ -5732,8 +5770,8 @@ void FGeometryCollectionPhysicsProxy::FieldParameterUpdateCallback(Chaos::FPBDRi
 					}
 					else if (FieldCommand.RootNode->Type() == FFieldNodeBase::EFieldType::EField_Float)
 					{
-						TArray<float>& FinalResults = ExecutionDatas.ScalarResults[(uint8)EFieldCommandResultType::FinalResult];
-						ResetResultsArray<float>(ExecutionDatas.SamplePositions.Num(), FinalResults, 0.0f);
+						TArray<float>& FinalResults = FieldData.ExecutionDatas.ScalarResults[(uint8)EFieldCommandResultType::FinalResult];
+						ResetResultsArray<float>(FieldData.ExecutionDatas.SamplePositions.Num(), FinalResults, 0.0f);
 
 						TFieldArrayView<float> ResultsView(FinalResults, 0, FinalResults.Num());
 
@@ -5747,7 +5785,7 @@ void FGeometryCollectionPhysicsProxy::FieldParameterUpdateCallback(Chaos::FPBDRi
 		
 		for (int32 Index = CommandsToRemove.Num() - 1; Index >= 0; --Index)
 		{
-			Commands.RemoveAt(CommandsToRemove[Index]);
+			FieldData.Commands.RemoveAt(CommandsToRemove[Index]);
 		}
 	}
 }
@@ -5756,7 +5794,10 @@ void FGeometryCollectionPhysicsProxy::FieldForcesUpdateCallback(Chaos::FPBDRigid
 {
 	SCOPE_CYCLE_COUNTER(STAT_ForceUpdateField_Object);
 	check(RigidSolver);
-	const int32 NumCommands = Commands.Num();
+
+	FFieldData& FieldData = GetOrCreateFieldData_Internal();
+
+	const int32 NumCommands = FieldData.Commands.Num();
 	if (NumCommands && !RigidSolver->IsShuttingDown())
 	{
 		TArray<int32> CommandsToRemove;
@@ -5769,24 +5810,24 @@ void FGeometryCollectionPhysicsProxy::FieldForcesUpdateCallback(Chaos::FPBDRigid
 
 		for (int32 CommandIndex = 0; CommandIndex < NumCommands; CommandIndex++)
 		{
-			const FFieldSystemCommand& FieldCommand = Commands[CommandIndex];
+			const FFieldSystemCommand& FieldCommand = FieldData.Commands[CommandIndex];
 			if (IsForceFieldValid(FieldCommand))
 			{
-				if (Chaos::BuildFieldSamplePoints(this, RigidSolver, FieldCommand, ExecutionDatas, PrevResolutionType, PrevFilterType, PrevObjectType, PrevPositionType))
+				if (Chaos::BuildFieldSamplePoints(this, RigidSolver, FieldCommand, FieldData.ExecutionDatas, PrevResolutionType, PrevFilterType, PrevObjectType, PrevPositionType))
 				{
 					const Chaos::FReal TimeSeconds = RigidSolver->GetSolverTime() - FieldCommand.TimeCreation;
 
 					FFieldContext FieldContext(
-						ExecutionDatas,
+						FieldData.ExecutionDatas,
 						FieldCommand.MetaData,
 						TimeSeconds);
 
-					TArray<Chaos::FGeometryParticleHandle*>& ParticleHandles = ExecutionDatas.ParticleHandles[(uint8)EFieldCommandHandlesType::InsideHandles];
+					TArray<Chaos::FGeometryParticleHandle*>& ParticleHandles = FieldData.ExecutionDatas.ParticleHandles[(uint8)EFieldCommandHandlesType::InsideHandles];
 
 					if (FieldCommand.RootNode->Type() == FFieldNode<FVector>::StaticType())
 					{
-						TArray<FVector>& FinalResults = ExecutionDatas.VectorResults[(uint8)EFieldCommandResultType::FinalResult];
-						ResetResultsArray < FVector >(ExecutionDatas.SamplePositions.Num(), FinalResults, FVector::ZeroVector);
+						TArray<FVector>& FinalResults = FieldData.ExecutionDatas.VectorResults[(uint8)EFieldCommandResultType::FinalResult];
+						ResetResultsArray < FVector >(FieldData.ExecutionDatas.SamplePositions.Num(), FinalResults, FVector::ZeroVector);
 
 						Chaos::FieldVectorForceUpdate(RigidSolver, FieldCommand, ParticleHandles,
 							FieldContext, FinalResults);
@@ -5797,7 +5838,7 @@ void FGeometryCollectionPhysicsProxy::FieldForcesUpdateCallback(Chaos::FPBDRigid
 		}
 		for (int32 Index = CommandsToRemove.Num() - 1; Index >= 0; --Index)
 		{
-			Commands.RemoveAt(CommandsToRemove[Index]);
+			FieldData.Commands.RemoveAt(CommandsToRemove[Index]);
 		}
 	}
 }
diff --git a/Engine/Source/Runtime/Experimental/Chaos/Public/PhysicsProxy/GeometryCollectionPhysicsProxy.h b/Engine/Source/Runtime/Experimental/Chaos/Public/PhysicsProxy/GeometryCollectionPhysicsProxy.h
--- a/Engine/Source/Runtime/Experimental/Chaos/Public/PhysicsProxy/GeometryCollectionPhysicsProxy.h
+++ b/Engine/Source/Runtime/Experimental/Chaos/Public/PhysicsProxy/GeometryCollectionPhysicsProxy.h
@@ -257,23 +257,22 @@ public:
 	const FGeometryCollectionResults* GetConsumerResultsGT() const 
 	{ return PhysToGameInterchange.PeekConsumerBuffer(); }
 
-	/** Enqueue a field \p Command to be processed by \c ProcessCommands() or 
-	 * \c FieldForcesUpdateCallback(). 
-	 */
-	void BufferCommand(Chaos::FPBDRigidsSolver* RigidsSolver, const FFieldSystemCommand& Command)
-	{ 
-		check(RigidsSolver != nullptr);
-		RigidsSolver->GetGeometryCollectionPhysicsProxiesField_Internal().Add(this);
-		Commands.Add(Command); 
-	}
+	/** Enqueue a field \p Command to be processed by \c ProcessCommands() or \c FieldForcesUpdateCallback(). Game thread only */
+	CHAOS_API void BufferFieldCommand_External(FFieldSystemCommand&& Command);
 
-	static CHAOS_API bool NeedToInitializeSharedCollisionStructures(const FGeometryCollection& RestCollection);
-	static CHAOS_API void InitializeSharedCollisionStructures(Chaos::FErrorReporter& ErrorReporter, FGeometryCollection& RestCollection, const FSharedSimulationParameters& SharedParams);
+	UE_DEPRECATED(5.5, "Use BufferFieldCommand_Internal instead when calling on the physics thread or the _external version when calling on the gamethread")
+	CHAOS_API void BufferCommand(Chaos::FPBDRigidsSolver* RigidsSolver, const FFieldSystemCommand& Command);
+
+	/** Enqueue a field \p Command to be processed by \c ProcessCommands() or \c FieldForcesUpdateCallback(). Physics thread only*/
+	CHAOS_API void BufferFieldCommand_Internal(Chaos::FPBDRigidsSolver* RigidsSolver, const FFieldSystemCommand& Command);
 
 	CHAOS_API void FieldForcesUpdateCallback(Chaos::FPBDRigidsSolver* RigidSolver);
 
 	CHAOS_API void FieldParameterUpdateCallback(Chaos::FPBDRigidsSolver* RigidSolver, const bool bUpdateViews = true);
 
+	static CHAOS_API bool NeedToInitializeSharedCollisionStructures(const FGeometryCollection& RestCollection);
+	static CHAOS_API void InitializeSharedCollisionStructures(Chaos::FErrorReporter& ErrorReporter, FGeometryCollection& RestCollection, const FSharedSimulationParameters& SharedParams);
+
 	void UpdateKinematicBodiesCallback(const FParticlesType& InParticles, const float InDt, const float InTime, FKinematicProxy& InKinematicProxy) {}
 	void StartFrameCallback(const float InDt, const float InTime) {}
 	void EndFrameCallback(const float InDt) {}
@@ -645,22 +644,51 @@ private:
 	 */
 	bool PullNonInterpolatableDataFromSinglePhysicsState(const Chaos::FDirtyGeometryCollectionData& BufferData, bool bForcePullXRVW, const TBitArray<>* Seen);
 
-	/* set to true once InitializeBodiesPT has been called*/
-	bool bIsInitializedOnPhysicsThread = false;
-
 	FSimulationParameters Parameters;
-	TArray<FFieldSystemCommand> Commands;
-
-	/** Field Datas stored during evaluation */
-	FFieldExecutionDatas ExecutionDatas;
 
 	TArray<Chaos::FPhysicsObjectUniquePtr> PhysicsObjects;
+
+	// todo : we should probably keep a simulation parameter copy on the game thread instead 
+	FTransform WorldTransform_External;
+	FTransform PreviousWorldTransform_External;
+
 	//
 	//  Proxy State Information
 	//
 	int32 NumTransforms;
 	int32 NumEffectiveParticles;
 	int32 BaseParticleIndex;
+	// Per object collision fraction.
+	float CollisionParticlesPerObjectFraction;
+
+	/** structure that contains the necessary information for processing fields */
+	struct FFieldData
+	{
+		/** field command to execute */
+		TArray<FFieldSystemCommand> Commands;
+
+		/** Field Datas stored during evaluation */
+		FFieldExecutionDatas ExecutionDatas;
+	};
+	// data is allocated on demand on the physics thread only ( see GetOrCreateFieldData_Internal )
+	TUniquePtr<FFieldData> FieldData_Internal;
+
+	FFieldData& GetOrCreateFieldData_Internal();
+
+	EReplicationMode ReplicationMode = EReplicationMode::Unknown;
+
+	uint8 bIsGameThreadWorldTransformDirty : 1;
+	uint8 bHasBuiltGeometryOnPT : 1;
+	uint8 bHasBuiltGeometryOnGT : 1;
+	/* set to true once InitializeBodiesPT has been called*/
+	bool bIsInitializedOnPhysicsThread : 1 = false;
+	//
+	// Buffer Results State Information
+	//
+	bool IsObjectDynamic : 1; // Records current dynamic state
+	bool IsObjectLoading : 1; // Indicate when loaded
+	bool IsObjectDeleting : 1; // Indicate when pending deletion
+
 	TArray<FParticleHandle*> SolverClusterID;
 	TArray<FClusterHandle*> SolverClusterHandles; // make a TArray of the base clase with type
 	TArray<FClusterHandle*> SolverParticleHandles;// make a TArray of base class and join with above
@@ -671,15 +699,6 @@ private:
 	TArray<int32> FromTransformToParticleIndex;
 	TBitArray<> EffectiveParticles;
 
-	//
-	// Buffer Results State Information
-	//
-	bool IsObjectDynamic; // Records current dynamic state
-	bool IsObjectLoading; // Indicate when loaded
-	bool IsObjectDeleting; // Indicate when pending deletion
-
-	EReplicationMode ReplicationMode = EReplicationMode::Unknown;	
-
 	TArray<TUniquePtr<FParticle>> GTParticles;
 	TMap<FParticle*, int32> GTParticlesToTransformGroupIndex;
 	TMap<FParticle*, int32> GTParticlesToInternalClusterUniqueIdx;
@@ -691,10 +710,6 @@ private:
 	const FCollisionFilterData SimFilter;
 	const FCollisionFilterData QueryFilter;
 
-	// This is a subset of the geometry group that are used in the transform hierarchy to represent geometry
-	TArray<FBox> ValidGeometryBoundingBoxes;
-	TArray<int32> ValidGeometryTransformIndices;
-
 	// todo(chaos): Remove this and move to a cook time approach of the SM data based on the GC property
 	FCreateTraceCollisionGeometryCallback CreateTraceCollisionGeometryCallback;
 	
@@ -702,9 +717,6 @@ private:
 	TFunction<void()> PostPhysicsSyncCallback;
 	TFunction<void()> PostParticlesCreatedCallback;
 	
-	// Per object collision fraction.
-	float CollisionParticlesPerObjectFraction;
-
 	// The Simulation data is copied between the game and physics thread. It is 
 	// expected that the two data sets will diverge, based on how the simulation
 	// uses the data, but at the start of the simulation the PhysicsThreadCollection
@@ -712,14 +724,6 @@ private:
 	FGeometryDynamicCollection PhysicsThreadCollection;
 	FGeometryDynamicCollection& GameThreadCollection;
 
-	// todo : we should probably keep a simulation parameter copy on the game thread instead 
-	FTransform WorldTransform_External;
-	FTransform PreviousWorldTransform_External;
-	uint8 bIsGameThreadWorldTransformDirty : 1;
-
-	uint8 bHasBuiltGeometryOnPT : 1;
-	uint8 bHasBuiltGeometryOnGT : 1;
-
 	// Currently this is using triple buffers for game-physics and 
 	// physics-game thread communication, but not for any reason other than this 
 	// is the only implementation we currently have of a guarded buffer - a buffer 
diff --git a/Engine/Source/Runtime/Experimental/GeometryCollectionEngine/Private/GeometryCollection/GeometryCollectionComponent.cpp b/Engine/Source/Runtime/Experimental/GeometryCollectionEngine/Private/GeometryCollection/GeometryCollectionComponent.cpp
--- a/Engine/Source/Runtime/Experimental/GeometryCollectionEngine/Private/GeometryCollection/GeometryCollectionComponent.cpp
+++ b/Engine/Source/Runtime/Experimental/GeometryCollectionEngine/Private/GeometryCollection/GeometryCollectionComponent.cpp
@@ -5781,11 +5781,7 @@ void UGeometryCollectionComponent::DispatchFieldCommand(const FFieldSystemComman
 		FFieldSystemCommand LocalCommand = InCommand;
 		LocalCommand.InitFieldNodes(Solver->GetSolverTime(), Name);
 
-		Solver->EnqueueCommandImmediate([Solver, PhysicsProxy = this->PhysicsProxy, NewCommand = LocalCommand]()
-		{
-			// Pass through nullptr here as geom component commands can never affect other solvers
-			PhysicsProxy->BufferCommand(Solver, NewCommand);
-		});
+		PhysicsProxy->BufferFieldCommand_External(MoveTemp(LocalCommand));
 	}
 }
 
