Cluster union : make sure we always read the right ChildToParent transform on the physics thread  ( submitted from BulkSetChildToParent_External )
- includes nulls in the particle array when getting the particles from the physics objects passed from the game thread

diff --git a/Engine/Source/Runtime/Experimental/Chaos/Private/Chaos/ClusterUnionManager.cpp b/Engine/Source/Runtime/Experimental/Chaos/Private/Chaos/ClusterUnionManager.cpp
--- a/Engine/Source/Runtime/Experimental/Chaos/Private/Chaos/ClusterUnionManager.cpp
+++ b/Engine/Source/Runtime/Experimental/Chaos/Private/Chaos/ClusterUnionManager.cpp
@@ -1413,7 +1413,7 @@ namespace Chaos
 		bool bMadeChanges = false;
 		for (FPBDRigidParticleHandle* Particle : Particles)
 		{
-			if (!ensure(Particle))
+			if (!Particle)
 			{
 				return;
 			}
diff --git a/Engine/Source/Runtime/Experimental/Chaos/Private/Chaos/PhysicsObjectInterface.cpp b/Engine/Source/Runtime/Experimental/Chaos/Private/Chaos/PhysicsObjectInterface.cpp
--- a/Engine/Source/Runtime/Experimental/Chaos/Private/Chaos/PhysicsObjectInterface.cpp
+++ b/Engine/Source/Runtime/Experimental/Chaos/Private/Chaos/PhysicsObjectInterface.cpp
@@ -445,25 +445,27 @@ namespace Chaos
 	}
 
 	template<EThreadContext Id>
-	TArray<TThreadRigidParticle<Id>*> FReadPhysicsObjectInterface<Id>::GetAllRigidParticles(TArrayView<const FConstPhysicsObjectHandle> InObjects)
+	TArray<TThreadRigidParticle<Id>*> FReadPhysicsObjectInterface<Id>::GetAllRigidParticles(TArrayView<const FConstPhysicsObjectHandle> InObjects, bool bIncludeNulls)
 	{
 		TArray<TThreadRigidParticle<Id>*> Particles;
 		Particles.Reserve(InObjects.Num());
 
 		for (const FConstPhysicsObjectHandle& Handle : InObjects)
 		{
-			if (!Handle)
-			{
-				continue;
-			}
+			TThreadRigidParticle<Id>* RigidParticle = nullptr;
 
-			if (TThreadParticle<Id>* Particle = Handle->GetParticle<Id>())
+			if (Handle)
 			{
-				if (TThreadRigidParticle<Id>* RigidParticle = Particle->CastToRigidParticle())
+				if (TThreadParticle<Id>* Particle = Handle->GetParticle<Id>())
 				{
-					Particles.Add(RigidParticle);
+					RigidParticle = Particle->CastToRigidParticle();
 				}
 			}
+
+			if (RigidParticle || bIncludeNulls)
+			{
+				Particles.Add(RigidParticle);
+			}
 		}
 
 		return Particles;
diff --git a/Engine/Source/Runtime/Experimental/Chaos/Private/PhysicsProxy/ClusterUnionPhysicsProxy.cpp b/Engine/Source/Runtime/Experimental/Chaos/Private/PhysicsProxy/ClusterUnionPhysicsProxy.cpp
--- a/Engine/Source/Runtime/Experimental/Chaos/Private/PhysicsProxy/ClusterUnionPhysicsProxy.cpp
+++ b/Engine/Source/Runtime/Experimental/Chaos/Private/PhysicsProxy/ClusterUnionPhysicsProxy.cpp
@@ -850,9 +850,11 @@ namespace Chaos
 			[this, Objects, Transforms, bLock]() mutable
 			{
 				FReadPhysicsObjectInterface_Internal Interface = FPhysicsObjectInternalInterface::GetRead();
-				TArray<FPBDRigidParticleHandle*> Particles = Interface.GetAllRigidParticles(Objects);
+				TArray<FPBDRigidParticleHandle*> Particles = Interface.GetAllRigidParticles(Objects, /* bIncludeNulls */ true);
 				if (ensure(Particles.Num() == Objects.Num()))
 				{
+					Particles.Add(nullptr); // TEST DO NOT SUBMIT !!!!!!!!!!!!!!!!!!!!!!!!!!!!!!!!!!!!!!!!!!!!!!!!!!!!!!!!!!!!!!!!!!!!!!!!!!!!!!!!!!!!!!!!!!!!!!!!!!!!!!!!
+
 					FPBDRigidsEvolutionGBF& Evolution = *static_cast<FPBDRigidsSolver*>(Solver)->GetEvolution();
 					FClusterUnionManager& ClusterUnionManager = Evolution.GetRigidClustering().GetClusterUnionManager();
 					ClusterUnionManager.UpdateClusterUnionParticlesChildToParent(ClusterUnionIndex, Particles, Transforms, bLock);
diff --git a/Engine/Source/Runtime/Experimental/Chaos/Public/Chaos/PhysicsObjectInterface.h b/Engine/Source/Runtime/Experimental/Chaos/Public/Chaos/PhysicsObjectInterface.h
--- a/Engine/Source/Runtime/Experimental/Chaos/Public/Chaos/PhysicsObjectInterface.h
+++ b/Engine/Source/Runtime/Experimental/Chaos/Public/Chaos/PhysicsObjectInterface.h
@@ -79,7 +79,7 @@ namespace Chaos
 		CHAOS_API TThreadKinematicParticle<Id>* GetKinematicParticle(const FConstPhysicsObjectHandle Object);
 		CHAOS_API TThreadRigidParticle<Id>* GetRigidParticle(const FConstPhysicsObjectHandle Object);
 		CHAOS_API TArray<TThreadParticle<Id>*> GetAllParticles(TArrayView<const FConstPhysicsObjectHandle> InObjects);
-		CHAOS_API TArray<TThreadRigidParticle<Id>*> GetAllRigidParticles(TArrayView<const FConstPhysicsObjectHandle> InObjects);
+		CHAOS_API TArray<TThreadRigidParticle<Id>*> GetAllRigidParticles(TArrayView<const FConstPhysicsObjectHandle> InObjects, bool bIncludeNulls = false);
 
 		UE_DEPRECATED(5.3, "GetAllShapes has been deprecated. Please use GetAllThreadShapes instead.")
 		CHAOS_API TArray<FPerShapeData*> GetAllShapes(TArrayView<const FConstPhysicsObjectHandle> InObjects);
