Chaos - fix collision solver buffer allocation for very large constraint counts
- log and return when solver buffer fails to allocate or is incorrectly sized

diff --git a/Engine/Source/Runtime/Experimental/Chaos/Private/Chaos/Collision/PBDCollisionContainerSolver.cpp b/Engine/Source/Runtime/Experimental/Chaos/Private/Chaos/Collision/PBDCollisionContainerSolver.cpp
--- a/Engine/Source/Runtime/Experimental/Chaos/Private/Chaos/Collision/PBDCollisionContainerSolver.cpp
+++ b/Engine/Source/Runtime/Experimental/Chaos/Private/Chaos/Collision/PBDCollisionContainerSolver.cpp
@@ -413,7 +413,7 @@ namespace Chaos
 	{
 		AppliedShockPropagation = FSolverReal(1);
 
-		const int CollisionBufferNum = CalculateCollisionBufferNum(MaxCollisions, CollisionConstraints.Num());
+		const int CollisionBufferNum = CalculateCollisionBufferSize(MaxCollisions, CollisionConstraints.Num());
 		CollisionConstraints.Reset(CollisionBufferNum);
 		bCollisionConstraintPerIterationCollisionDetection.Reset(CollisionBufferNum);
 
@@ -428,11 +428,11 @@ namespace Chaos
 		CollisionSolverManifoldPoints = nullptr;
 	}
 
-	int32 FPBDCollisionContainerSolver::CalculateCollisionBufferNum(const int32 InTightFittingNum, const int32 InCurrentBufferNum) const
+	size_t FPBDCollisionContainerSolver::CalculateCollisionBufferSize(const size_t InTightFittingNum, const size_t InCurrentBufferNum) const
 	{
 		// A buffer over-allocation policy to avoid reallocation every frame in the common case where a pile of objects is dropped
 		// and the number of contacts increases every tick. Used for collision solvers and manifold points
-		int CollisionBufferNum = InTightFittingNum;
+		size_t CollisionBufferNum = InTightFittingNum;
 		if (CollisionBufferNum > InCurrentBufferNum)
 		{
 			CollisionBufferNum = (5 * InTightFittingNum) / 4; // +25%
@@ -484,28 +484,39 @@ namespace Chaos
 		if (MaxCollisionSolverManifoldPoints == 0)
 		{
 			NumCollisionSolvers = 0;
+			return;
 		}
 
-		// Set up the solver buffers
-		if (NumCollisionSolvers > 0)
-		{
-			// Resize the scratch buffer (up to 25% slack)
-			constexpr size_t AlignedSolverSize = Align(sizeof(Private::FPBDCollisionSolver), alignof(Private::FPBDCollisionSolver));
-			constexpr size_t AlignedPointSize = Align(sizeof(Private::FPBDCollisionSolverManifoldPoint), alignof(Private::FPBDCollisionSolverManifoldPoint));
-			const size_t ScratchSize = NumCollisionSolvers * AlignedSolverSize + MaxCollisionSolverManifoldPoints * AlignedPointSize;
-			const size_t ScratchBufferSize = CalculateCollisionBufferNum(ScratchSize, Scratch.BufferSize());
-			Scratch.Reset(ScratchBufferSize);
+		// Resize the scratch buffer (up to 25% slack - see CalculateCollisionBufferGrowSize)
+		constexpr size_t AlignedSolverSize = Align(sizeof(Private::FPBDCollisionSolver), alignof(Private::FPBDCollisionSolver));
+		constexpr size_t AlignedPointSize = Align(sizeof(Private::FPBDCollisionSolverManifoldPoint), alignof(Private::FPBDCollisionSolverManifoldPoint));
+		const size_t ScratchSize = NumCollisionSolvers * AlignedSolverSize + MaxCollisionSolverManifoldPoints * AlignedPointSize;
+		const size_t ScratchBufferSize = CalculateCollisionBufferSize(ScratchSize, Scratch.BufferSize());
+		Scratch.Reset(ScratchBufferSize);
 			
-			if (Scratch.BufferSize() == 0)
-			{
-				UE_LOG(LogChaos, Error, TEXT("FPBDCollisionContainerSolver: failed to allocate scratch buffer of size %lld bytes. NumCollisions=%d, NumManifoldPoints=%d. Collisions will be lost."), ScratchBufferSize, NumCollisionSolvers, MaxCollisionSolverManifoldPoints);
-				NumCollisionSolvers = 0;
-				return;
-			}
+		// Out of memory?
+		if (Scratch.BufferSize() == 0)
+		{
+			UE_LOG(LogChaos, Error, TEXT("FPBDCollisionContainerSolver: failed to allocate scratch buffer of size %lld bytes. NumCollisions=%d, NumManifoldPoints=%d. Collisions will be lost."), ScratchBufferSize, NumCollisionSolvers, MaxCollisionSolverManifoldPoints);
+			MaxCollisionSolverManifoldPoints = 0;
+			NumCollisionSolvers = 0;
+			return;
+		}
 
-			// Allocate scratch space for the collision solvers and manifold points
-			CollisionSolvers = Scratch.AllocArray<Private::FPBDCollisionSolver>(NumCollisionSolvers);
-			CollisionSolverManifoldPoints = Scratch.AllocArray<Private::FPBDCollisionSolverManifoldPoint>(MaxCollisionSolverManifoldPoints);
+		// Allocate scratch space for the collision solvers and manifold points
+		CollisionSolvers = Scratch.AllocArray<Private::FPBDCollisionSolver>(NumCollisionSolvers);
+		CollisionSolverManifoldPoints = Scratch.AllocArray<Private::FPBDCollisionSolverManifoldPoint>(MaxCollisionSolverManifoldPoints);
+
+		// We should never see these errors if the size calculations are correct, but being extra careful...
+		if ((CollisionSolvers == nullptr) || (CollisionSolverManifoldPoints == nullptr))
+		{
+			UE_CLOG((CollisionSolvers == nullptr), LogChaos, Error, TEXT("FPBDCollisionContainerSolver: failed to allocate CollisionSolvers in scratch buffer of size %lld bytes. NumCollisions=%d, NumManifoldPoints=%d. Collisions will be lost."), ScratchBufferSize, NumCollisionSolvers, MaxCollisionSolverManifoldPoints);
+			UE_CLOG((CollisionSolverManifoldPoints == nullptr), LogChaos, Error, TEXT("FPBDCollisionContainerSolver: failed to allocate CollisionSolverManifoldPoints in scratch buffer of size %lld bytes. NumCollisions=%d, NumManifoldPoints=%d. Collisions will be lost."), ScratchBufferSize, NumCollisionSolvers, MaxCollisionSolverManifoldPoints);
+			MaxCollisionSolverManifoldPoints = 0;
+			NumCollisionSolvers = 0;
+			CollisionSolvers = nullptr;
+			CollisionSolverManifoldPoints = nullptr;
+			return;
 		}
 	}
 
@@ -544,8 +555,10 @@ namespace Chaos
 		// All constarints are now added. We can allocate the solver buffers.
 		PrepareSolverBuffer();
 
-		// Make sure have a valid manifold point buffer if we have constraints
-		check((CollisionSolverManifoldPoints != nullptr) || (NumSolvers() == 0));
+		if (CollisionSolverManifoldPoints == nullptr)
+		{
+			return;
+		}
 
 		for (int32 ConstraintIndex = 0, ConstraintEndIndex = NumSolvers(); ConstraintIndex < ConstraintEndIndex; ++ConstraintIndex)
 		{
@@ -563,6 +576,9 @@ namespace Chaos
 			// Set up the solver manifold point buffer pointer
 			const int32 ConstraintManifoldPointMax = CalculateConstraintMaxManifoldPoints(GetConstraint(ConstraintIndex));
 
+			// We should have allocated enough space for all manifold points in PrepareSolverBuffer
+			check(NumCollisionSolverManifoldPoints + ConstraintManifoldPointMax <= MaxCollisionSolverManifoldPoints);
+
 			CollisionSolver.Reset(&CollisionSolverManifoldPoints[NumCollisionSolverManifoldPoints], ConstraintManifoldPointMax);
 			CollisionSolver.SetSolverBodies(*Body0, *Body1);
 
diff --git a/Engine/Source/Runtime/Experimental/Chaos/Public/Chaos/Collision/PBDCollisionContainerSolver.h b/Engine/Source/Runtime/Experimental/Chaos/Public/Chaos/Collision/PBDCollisionContainerSolver.h
--- a/Engine/Source/Runtime/Experimental/Chaos/Public/Chaos/Collision/PBDCollisionContainerSolver.h
+++ b/Engine/Source/Runtime/Experimental/Chaos/Public/Chaos/Collision/PBDCollisionContainerSolver.h
@@ -58,7 +58,7 @@ namespace Chaos
 
 		void CachePrefetchSolver(const int32 ConstraintIndex) const;
 		void AddConstraint(FPBDCollisionConstraint& Constraint);
-		int32 CalculateCollisionBufferNum(const int32 InTightFittingNum, const int32 InCurrentBufferNum) const;
+		size_t CalculateCollisionBufferSize(const size_t InTightFittingNum, const size_t InCurrentBufferNum) const;
 		int32 CalculateConstraintMaxManifoldPoints(const FPBDCollisionConstraint* Constraint) const;
 		void PrepareSolverBuffer();
 		void UpdatePositionShockPropagation(const FReal Dt, const int32 It, const int32 NumIts, const int32 BeginIndex, const int32 EndIndex, const FPBDCollisionSolverSettings& SolverSettings);
