Iris
- Fixed issue with log spam "Failed to replicate ( InternalIndex: )"
- Fixed issue with torn-off objects requesting a flush could get stuck in WaitOnFlush state.

diff --git a/Engine/Source/Runtime/Experimental/Iris/Core/Private/Iris/ReplicationSystem/ReplicationWriter.cpp b/Engine/Source/Runtime/Experimental/Iris/Core/Private/Iris/ReplicationSystem/ReplicationWriter.cpp
--- a/Engine/Source/Runtime/Experimental/Iris/Core/Private/Iris/ReplicationSystem/ReplicationWriter.cpp
+++ b/Engine/Source/Runtime/Experimental/Iris/Core/Private/Iris/ReplicationSystem/ReplicationWriter.cpp
@@ -793,7 +793,7 @@ void FReplicationWriter::UpdateScope(const FNetBitArrayView& UpdatedScope)
 			{
 				// If we have no data to flush, we can stop replication now.
 				const uint32 FlushFlags = GetFlushStatus(Index, Info, Info.FlushFlags);
-				if (GetFlushStatus(Index, Info, Info.FlushFlags) == FlushFlags_None)
+				if (FlushFlags == FlushFlags_None || NetRefHandleManager->GetReplicatedObjectDataNoCheck(Index).bHasCachedCreationInfo == 0U)
 				{
 					StopReplication(Index);
 				}
@@ -1084,11 +1084,13 @@ void FReplicationWriter::HandleDeliveredRecord(const FReplicationRecord::FRecord
 					// Tear-off is marked as a flush
 					if (Info.TearOff)
 					{
+						UE_LOG_REPLICATIONWRITER_CONN(TEXT("ReplicationWriter::HandleDeliveredRecord for ( InternalIndex: %u ) Waiting for flush before tearoff"), InternalIndex);
 						SetState(InternalIndex, EReplicatedObjectState::WaitOnFlush);
 					}
 					// so are objects marked for destroy requiring flush
 					else if (ObjectsPendingDestroy.GetBit(InternalIndex))
 					{
+						UE_LOG_REPLICATIONWRITER_CONN(TEXT("ReplicationWriter::HandleDeliveredRecord for ( InternalIndex: %u ) Waiting for flush before destroy"), InternalIndex);
 						SetState(InternalIndex, EReplicatedObjectState::WaitOnFlush);
 					}
 				}
@@ -1350,7 +1352,7 @@ void FReplicationWriter::HandleDroppedRecord<FReplicationWriter::EReplicatedObje
 		}
 
 		// if we lost changes that are not already retransmitted we update the changemask
-		if (bNeedToResendState | bNeedToResendAttachments)
+		if (bNeedToResendState | bNeedToResendAttachments | Info.TearOff)
 		{
 			if (bNeedToResendState)
 			{
@@ -3008,6 +3010,7 @@ int FReplicationWriter::HandleObjectBatchSuccess(const FBatchInfo& BatchInfo, FR
 			}
 			else
 			{
+				UE_LOG_REPLICATIONWRITER_CONN(TEXT("ReplicationWriter::HandleObjectBatchSuccess for ( InternalIndex: %u ) Waiting for flush before tearoff"), BatchObjectInfo.InternalIndex);
 				SetState(BatchObjectInfo.InternalIndex, EReplicatedObjectState::WaitOnFlush);
 			}
 		}
