Iris
- Presumable fix for servercrhrashes

diff --git a/Engine/Source/Runtime/Experimental/Iris/Core/Private/Iris/ReplicationSystem/ObjectReplicationBridge.cpp b/Engine/Source/Runtime/Experimental/Iris/Core/Private/Iris/ReplicationSystem/ObjectReplicationBridge.cpp
--- a/Engine/Source/Runtime/Experimental/Iris/Core/Private/Iris/ReplicationSystem/ObjectReplicationBridge.cpp
+++ b/Engine/Source/Runtime/Experimental/Iris/Core/Private/Iris/ReplicationSystem/ObjectReplicationBridge.cpp
@@ -603,27 +603,40 @@ void UObjectReplicationBridge::RemoveDependentObject(FNetRefHandle ParentHandle,
 
 bool UObjectReplicationBridge::WriteNetRefHandleCreationInfo(FReplicationBridgeSerializationContext& Context, FNetRefHandle Handle)
 {
-	// Write ProtocolId
-	const UE::Net::FReplicationProtocol* Protocol = GetReplicationSystem()->GetReplicationProtocol(Handle);
-	WriteUint64(Context.SerializationContext.GetBitStreamWriter(), Protocol->ProtocolIdentifier);
-
 	// Write Type header, if there is a cached one, use it!
 	if (TUniquePtr<const FCreationHeader>* CachedHeader = CachedCreationHeaders.Find(Handle))
 	{
+		WriteUint64(Context.SerializationContext.GetBitStreamWriter(), (*CachedHeader)->ProtocolIdentifier);
 		return WriteCreationHeader(Context.SerializationContext, (*CachedHeader).Get());
 	}
 	else
 	{
+		// Write ProtocolId
+		const UE::Net::FReplicationProtocol* Protocol = GetReplicationSystem()->GetReplicationProtocol(Handle);
+		if (!Protocol)
+		{
+			UE_LOG_OBJECTREPLICATIONBRIDGE(Error, TEXT("WriteNetRefHandleCreationInfo: Cannot write creationinfo for %s, since protocol has been detached"), *Handle.ToString());
+			ensureMsgf(Protocol, TEXT("WriteNetRefHandleCreationInfo: Cannot write creationinfo for %s, since protocol has been detached"), *Handle.ToString());
+			return false;
+		}
+		WriteUint64(Context.SerializationContext.GetBitStreamWriter(), Protocol->ProtocolIdentifier);
 		return WriteCreationHeader(Context.SerializationContext, Handle);
 	}
 }
 
 bool UObjectReplicationBridge::CacheNetRefHandleCreationInfo(FNetRefHandle Handle)
 {
-	TUniquePtr<const FCreationHeader> Header(GetCreationHeader(Handle));
+	const UE::Net::FReplicationProtocol* Protocol = GetReplicationSystem()->GetReplicationProtocol(Handle);	
+	if (!Protocol)
+	{
+		return false;
+	}
+
+	TUniquePtr<FCreationHeader> Header(GetCreationHeader(Handle));
 
 	if (Header)
 	{
+		Header->ProtocolIdentifier = Protocol->ProtocolIdentifier;
 		CachedCreationHeaders.Add(Handle, MoveTemp(Header));
 		return true;
 	}
diff --git a/Engine/Source/Runtime/Experimental/Iris/Core/Public/Iris/ReplicationSystem/ObjectReplicationBridge.h b/Engine/Source/Runtime/Experimental/Iris/Core/Public/Iris/ReplicationSystem/ObjectReplicationBridge.h
--- a/Engine/Source/Runtime/Experimental/Iris/Core/Public/Iris/ReplicationSystem/ObjectReplicationBridge.h
+++ b/Engine/Source/Runtime/Experimental/Iris/Core/Public/Iris/ReplicationSystem/ObjectReplicationBridge.h
@@ -149,6 +149,10 @@ public:
 	struct FCreationHeader
 	{
 		virtual ~FCreationHeader() {};
+
+	private:
+		friend UObjectReplicationBridge;
+		uint64 ProtocolIdentifier = uint64(0);
 	};
 	
 	/**
