Geometry collection - memory optimization - remove two TMap from the physics proxy
- saves
- Combine two maps into one for index lookup by GT Particle ( no extra memory as the original Pair used by the single index TMap had padding enough to fit an extra index )
- Removed unused TMap
- keep collector guid only for editor ( 16 more bytes saved )

diff --git a/Engine/Source/Runtime/Experimental/Chaos/Private/PhysicsProxy/GeometryCollectionPhysicsProxy.cpp b/Engine/Source/Runtime/Experimental/Chaos/Private/PhysicsProxy/GeometryCollectionPhysicsProxy.cpp
--- a/Engine/Source/Runtime/Experimental/Chaos/Private/PhysicsProxy/GeometryCollectionPhysicsProxy.cpp
+++ b/Engine/Source/Runtime/Experimental/Chaos/Private/PhysicsProxy/GeometryCollectionPhysicsProxy.cpp
@@ -514,7 +514,9 @@ FGeometryCollectionPhysicsProxy::FGeometryCollectionPhysicsProxy(
 	, QueryFilter(InQueryFilter)
 	, PhysicsThreadCollection(Parameters.RestCollectionShared)
 	, GameThreadCollection(GameThreadCollectionIn)
+#if WITH_EDITORONLY_DATA
 	, CollectorGuid(InCollectorGuid)
+#endif
 {
 	// We rely on a guarded buffer.
 	check(BufferMode == Chaos::EMultiBufferMode::TripleGuarded);
@@ -711,7 +713,7 @@ void FGeometryCollectionPhysicsProxy::CreateGTParticles(TManagedArray<Chaos::FIm
 	TArray<int32> ChildrenToCheckForParentFix;
 	if (!bInitializationTime && !bCreateGTParticleForChildren && !bBuildGeometryForChildrenOnGT)
 	{
-		GTParticlesToTransformGroupIndex.Reserve(NumEffectiveParticles);
+		GTParticleToIndices.Reserve(NumEffectiveParticles);
 	}
 
 	for (int32 ParticleIndex = 0; ParticleIndex < NumEffectiveParticles; ++ParticleIndex)
@@ -731,7 +733,7 @@ void FGeometryCollectionPhysicsProxy::CreateGTParticles(TManagedArray<Chaos::FIm
 		{
 			GTParticles[ParticleIndex] = FParticle::CreateParticle();
 			P = GTParticles[ParticleIndex].Get();
-			GTParticlesToTransformGroupIndex.Add(P, TransformIndex);
+			GTParticleToIndices.FindOrAdd(P).TransformGroupIndex = TransformIndex;
 			GTParticles[ParticleIndex]->SetUniqueIdx(UniqueIdxs[ParticleIndex]);
 #if CHAOS_DEBUG_NAME
 				P->SetDebugName(MakeShared<FString, ESPMode::ThreadSafe>(FString::Printf(TEXT("%s-%d"), *Parameters.Name, TransformIndex)));
@@ -4219,9 +4221,10 @@ bool FGeometryCollectionPhysicsProxy::PullNonInterpolatableDataFromSinglePhysics
 		}
 
 		// internal cluster index map update
+		// make sure to update the index even if the InternalClusterUniqueIdx is INDEX_NONE
+		GTParticleToIndices.FindOrAdd(GTParticles[ParticleIndex].Get()).InternalClusterUniqueId = StateData.InternalClusterUniqueIdx;
 		if (StateData.InternalClusterUniqueIdx > INDEX_NONE)
 		{
-			GTParticlesToInternalClusterUniqueIdx.Add(GTParticles[ParticleIndex].Get(), StateData.InternalClusterUniqueIdx);
 			InternalClusterUniqueIdxToChildrenTransformIndices.FindOrAdd(StateData.InternalClusterUniqueIdx).Add(TransformGroupIndex);
 		}
 	}
@@ -4293,7 +4296,6 @@ bool FGeometryCollectionPhysicsProxy::PullFromPhysicsState(const Chaos::FDirtyGe
 
 	if (NumTransforms > 0)
 	{
-		GTParticlesToInternalClusterUniqueIdx.Reset();
 		InternalClusterUniqueIdxToChildrenTransformIndices.Reset();
 
 		// first: non-interpolatable data (everything besides XRVW OR if no next data exists, everything is non-interpolatable).
diff --git a/Engine/Source/Runtime/Experimental/Chaos/Public/PhysicsProxy/GeometryCollectionPhysicsProxy.h b/Engine/Source/Runtime/Experimental/Chaos/Public/PhysicsProxy/GeometryCollectionPhysicsProxy.h
--- a/Engine/Source/Runtime/Experimental/Chaos/Public/PhysicsProxy/GeometryCollectionPhysicsProxy.h
+++ b/Engine/Source/Runtime/Experimental/Chaos/Public/PhysicsProxy/GeometryCollectionPhysicsProxy.h
@@ -404,9 +404,12 @@ public:
 		if (ChildTransformIndex >= 0 && ChildTransformIndex < GTParticles.Num())
 		{
 			const TUniquePtr<FParticle>& ChildGTParticle = GTParticles[ChildTransformIndex];
-			if (const int32* InternalClusterUniqueIdx = GTParticlesToInternalClusterUniqueIdx.Find(ChildGTParticle.Get()))
+			if (const FGTParticleIndices* Indices = GTParticleToIndices.Find(ChildGTParticle.Get()))
 			{
-				return FGeometryCollectionItemIndex::CreateInternalClusterItemIndex(*InternalClusterUniqueIdx);
+				if (Indices->InternalClusterUniqueId != INDEX_NONE)
+				{
+					return FGeometryCollectionItemIndex::CreateInternalClusterItemIndex(Indices->InternalClusterUniqueId);
+				}
 			}
 		}
 		return FGeometryCollectionItemIndex::CreateInvalidItemIndex();
@@ -424,23 +427,29 @@ public:
 	FGeometryCollectionItemIndex GetItemIndexFromGTParticle_External(const FParticle* GTPParticle) const
 	{
 		// internal cluster have  no representation on the GT, so we use the child GT particle to find the matching internal cluster unique index 
-		if (const int32* InternalClusterUniqueIdx = GTParticlesToInternalClusterUniqueIdx.Find(GTPParticle))
+		if (const FGTParticleIndices* Indices = GTParticleToIndices.Find(GTPParticle))
 		{
-			return FGeometryCollectionItemIndex::CreateInternalClusterItemIndex(*InternalClusterUniqueIdx);
-		}
-		// regular particle that has a matching transform index 
-		if (const int32* TransformGroupIndex = GTParticlesToTransformGroupIndex.Find(GTPParticle))
-		{
-			return FGeometryCollectionItemIndex::CreateTransformItemIndex(*TransformGroupIndex);
+			if (Indices->InternalClusterUniqueId != INDEX_NONE)
+			{
+				return FGeometryCollectionItemIndex::CreateInternalClusterItemIndex(Indices->InternalClusterUniqueId);
+			}
+			// regular particle that has a matching transform index 
+			if (Indices->TransformGroupIndex != INDEX_NONE)
+			{
+				return FGeometryCollectionItemIndex::CreateTransformItemIndex(Indices->TransformGroupIndex);
+			}
 		}
 		return FGeometryCollectionItemIndex::CreateInvalidItemIndex();
 	}
 
 	FGeometryCollectionItemIndex GetItemIndexFromGTParticleNoInternalCluster_External(const FParticle* GTPParticle) const
 	{
-		if (const int32* TransformGroupIndex = GTParticlesToTransformGroupIndex.Find(GTPParticle))
+		if (const FGTParticleIndices* Indices = GTParticleToIndices.Find(GTPParticle))
 		{
-			return FGeometryCollectionItemIndex::CreateTransformItemIndex(*TransformGroupIndex);
+			if (Indices->TransformGroupIndex != INDEX_NONE)
+			{
+				return FGeometryCollectionItemIndex::CreateTransformItemIndex(Indices->TransformGroupIndex);
+			}
 		}
 		return FGeometryCollectionItemIndex::CreateInvalidItemIndex();
 	}
@@ -699,12 +708,20 @@ private:
 	TArray<int32> FromTransformToParticleIndex;
 	TBitArray<> EffectiveParticles;
 
+	// Game thread particles 
 	TArray<TUniquePtr<FParticle>> GTParticles;
-	TMap<FParticle*, int32> GTParticlesToTransformGroupIndex;
-	TMap<FParticle*, int32> GTParticlesToInternalClusterUniqueIdx;
-	TMap<int32, TArray<int32>> InternalClusterUniqueIdxToChildrenTransformIndices;
 
-	TMap<int32, TUniquePtr<FParticle>> GTInternalClustersByUniqueIdx;
+	struct FGTParticleIndices
+	{
+		/** Correponding transform index of the GTparticle */
+		int32 TransformGroupIndex = INDEX_NONE;
+
+		/** unique index of the parent the GT particle is a child of internal cluster  */
+		int32 InternalClusterUniqueId = INDEX_NONE;
+	};
+
+	TMap<FParticle*, FGTParticleIndices> GTParticleToIndices;
+	TMap<int32, TArray<int32>> InternalClusterUniqueIdxToChildrenTransformIndices;
 
 	// These are read on both threads and should not be changed
 	const FCollisionFilterData SimFilter;
@@ -753,8 +770,10 @@ private:
 		return static_cast<ErrorDataType*>(InterpolationData.Get());
 	}
 
+#if WITH_EDITORONLY_DATA
 	// this is used as a unique ID when collecting data from runtime
 	FGuid CollectorGuid;
+#endif
 };
 
 /**
