geometry collection : add support for changing enable gravity flag at runtime

diff --git a/Engine/Source/Runtime/Experimental/Chaos/Private/PhysicsProxy/GeometryCollectionPhysicsProxy.cpp b/Engine/Source/Runtime/Experimental/Chaos/Private/PhysicsProxy/GeometryCollectionPhysicsProxy.cpp
--- a/Engine/Source/Runtime/Experimental/Chaos/Private/PhysicsProxy/GeometryCollectionPhysicsProxy.cpp
+++ b/Engine/Source/Runtime/Experimental/Chaos/Private/PhysicsProxy/GeometryCollectionPhysicsProxy.cpp
@@ -3493,6 +3493,43 @@ void FGeometryCollectionPhysicsProxy::SetMaterialOverrideMassScaleMultiplier_Int
 	}
 }
 
+void FGeometryCollectionPhysicsProxy::SetEnableGravity_External(bool bEnabled)
+{
+	for (TUniquePtr<FParticle>& GTParticle : GTParticles)
+	{
+		GTParticle->SetGravityEnabled(bEnabled);
+	}
+
+	ExecuteOnPhysicsThread(*this, 
+		[this, bEnabled]() 
+		{
+			SetEnableGravity_Internal(bEnabled);
+		});
+}
+
+void FGeometryCollectionPhysicsProxy::SetEnableGravity_Internal(bool bEnabled)
+{
+	Parameters.GravityGroupIndex = bEnabled;
+
+	if (Chaos::FPhysicsSolver* RBDSolver = GetSolver<Chaos::FPhysicsSolver>())
+	{
+		if (Chaos::FPBDRigidsSolver::FPBDRigidsEvolution* Evolution = RBDSolver->GetEvolution())
+		{
+			for (Chaos::FPBDRigidClusteredParticleHandle* Handle : SolverParticleHandles)
+			{
+				if (Handle)
+				{
+					Handle->SetGravityEnabled(bEnabled);
+					if (!Handle->Disabled())
+					{
+						Evolution->WakeParticle(Handle);
+					}
+				}
+			}
+		}
+	}
+}
+
 void FGeometryCollectionPhysicsProxy::SetGravityGroupIndex_External(int32 GravityGroupIndex)
 {
 	ExecuteOnPhysicsThread(*this, 
diff --git a/Engine/Source/Runtime/Experimental/Chaos/Public/PhysicsProxy/GeometryCollectionPhysicsProxy.h b/Engine/Source/Runtime/Experimental/Chaos/Public/PhysicsProxy/GeometryCollectionPhysicsProxy.h
--- a/Engine/Source/Runtime/Experimental/Chaos/Public/PhysicsProxy/GeometryCollectionPhysicsProxy.h
+++ b/Engine/Source/Runtime/Experimental/Chaos/Public/PhysicsProxy/GeometryCollectionPhysicsProxy.h
@@ -519,6 +519,7 @@ public:
 	CHAOS_API void SetDamageModel_External(EDamageModelTypeEnum DamageModel);
 	CHAOS_API void SetUseMaterialDamageModifiers_External(bool bUseMaterialDamageModifiers);
 	CHAOS_API void SetMaterialOverrideMassScaleMultiplier_External(float InMultiplier);
+	CHAOS_API void SetEnableGravity_External(bool EnableGravity_External);
 	CHAOS_API void SetGravityGroupIndex_External(int32 GravityGroupIndex);
 	CHAOS_API void SetOneWayInteractionLevel_External(int32 OneWayInteractionLevel);
 	CHAOS_API void SetPhysicsMaterial_External(const Chaos::FMaterialHandle& MaterialHandle);
@@ -623,6 +624,7 @@ protected:
 	CHAOS_API void SetFilterData_Internal(const FCollisionFilterData& NewSimFilter, const FCollisionFilterData& NewQueryFilter);
 	CHAOS_API void SetPerParticleFilterData_Internal(const TArray<FParticleCollisionFilterData>& PerParticleData);
 	CHAOS_API void SetDamagePropagationData_Internal(bool bEnabled, float BreakDamagePropagationFactor, float ShockDamagePropagationFactor);
+	CHAOS_API void SetEnableGravity_Internal(bool bEnabled);
 	CHAOS_API void SetDamageThresholds_Internal(const TArray<float>& DamageThresholds);
 	CHAOS_API void SetDamageModel_Internal(EDamageModelTypeEnum DamageModel);
 	CHAOS_API void SetUseMaterialDamageModifiers_Internal(bool bUseMaterialDamageModifiers);
diff --git a/Engine/Source/Runtime/Experimental/GeometryCollectionEngine/Private/GeometryCollection/GeometryCollectionComponent.cpp b/Engine/Source/Runtime/Experimental/GeometryCollectionEngine/Private/GeometryCollection/GeometryCollectionComponent.cpp
--- a/Engine/Source/Runtime/Experimental/GeometryCollectionEngine/Private/GeometryCollection/GeometryCollectionComponent.cpp
+++ b/Engine/Source/Runtime/Experimental/GeometryCollectionEngine/Private/GeometryCollection/GeometryCollectionComponent.cpp
@@ -1270,6 +1270,19 @@ void UGeometryCollectionComponent::SetSimulatePhysics(bool bEnabled)
 	}
 }
 
+void UGeometryCollectionComponent::SetEnableGravity(bool bGravityEnabled)
+{
+	if (bGravityEnabled != BodyInstance.bEnableGravity)
+	{
+		BodyInstance.bEnableGravity = bGravityEnabled;
+
+		if (PhysicsProxy)
+		{
+			PhysicsProxy->SetEnableGravity_External(bGravityEnabled);
+		}
+	}
+}
+
 void UGeometryCollectionComponent::AddForce(FVector Force, FName BoneName, bool bAccelChange)
 {
 	ensure(bAccelChange == false); // not supported
diff --git a/Engine/Source/Runtime/Experimental/GeometryCollectionEngine/Public/GeometryCollection/GeometryCollectionComponent.h b/Engine/Source/Runtime/Experimental/GeometryCollectionEngine/Public/GeometryCollection/GeometryCollectionComponent.h
--- a/Engine/Source/Runtime/Experimental/GeometryCollectionEngine/Public/GeometryCollection/GeometryCollectionComponent.h
+++ b/Engine/Source/Runtime/Experimental/GeometryCollectionEngine/Public/GeometryCollection/GeometryCollectionComponent.h
@@ -658,6 +658,7 @@ public:
 	GEOMETRYCOLLECTIONENGINE_API virtual void SetNotifyRigidBodyCollision(bool bNewNotifyRigidBodyCollision) override;
 	GEOMETRYCOLLECTIONENGINE_API virtual bool CanEditSimulatePhysics() override;
 	GEOMETRYCOLLECTIONENGINE_API virtual void SetSimulatePhysics(bool bEnabled) override;
+	GEOMETRYCOLLECTIONENGINE_API virtual void SetEnableGravity(bool bGravityEnabled) override;
 	GEOMETRYCOLLECTIONENGINE_API virtual void AddForce(FVector Force, FName BoneName = NAME_None, bool bAccelChange = false) override;
 	GEOMETRYCOLLECTIONENGINE_API virtual void AddForceAtLocation(FVector Force, FVector WorldLocation, FName BoneName = NAME_None) override;
 	GEOMETRYCOLLECTIONENGINE_API virtual void AddImpulse(FVector Impulse, FName BoneName = NAME_None, bool bVelChange = false) override;
