Geometry collection asset - move cooking of render data  out of the serialize method
- This is required to support async loading and avoid the resource initialization from calling GameThread  checked methods
- Move cooking to both PreSave and PostLoad
THis simplifies the code as PostLoad and PreSave are not called during transactions

diff --git a/Engine/Source/Runtime/Experimental/GeometryCollectionEngine/Private/GeometryCollection/GeometryCollectionObject.cpp b/Engine/Source/Runtime/Experimental/GeometryCollectionEngine/Private/GeometryCollection/GeometryCollectionObject.cpp
--- a/Engine/Source/Runtime/Experimental/GeometryCollectionEngine/Private/GeometryCollection/GeometryCollectionObject.cpp
+++ b/Engine/Source/Runtime/Experimental/GeometryCollectionEngine/Private/GeometryCollection/GeometryCollectionObject.cpp
@@ -11,6 +11,7 @@
 #include "UObject/DestructionObjectVersion.h"
 #include "UObject/UE5MainStreamObjectVersion.h"
 #include "UObject/FortniteMainBranchObjectVersion.h"
+#include "UObject/ObjectSaveContext.h"
 #include "Serialization/ArchiveCountMem.h"
 #include "HAL/IConsoleManager.h"
 #include "Interfaces/ITargetPlatform.h"
@@ -1027,7 +1028,9 @@ void UGeometryCollection::Serialize(FArchive& Ar)
 #if WITH_EDITOR
 		if (Ar.IsSaving() && !Ar.IsTransacting())
 		{
-			EnsureDataIsCooked(false /*bInitResources*/, Ar.IsTransacting(), Ar.IsPersistent(), false /*bAllowCopyFromDDC*/);
+			constexpr bool bAllowCopyFromDDC = false;
+			constexpr bool bIsTransacting = false; // the surrounding if statement garantees that 
+			EnsureSimulationDataIsCooked(bIsTransacting, bAllowCopyFromDDC);
 		}
 #endif
 		if (Ar.IsLoading() || Ar.IsCountingMemory())
@@ -1255,7 +1258,8 @@ void UGeometryCollection::Serialize(FArchive& Ar)
  	if (Ar.IsLoading())
 	{
 		// note: don't allow copy from DDC here, since we've already loaded the data above, and the DDC data does not include any data migrations performed by the load
-		EnsureDataIsCooked(true /*bInitResources*/, Ar.IsTransacting(), Ar.IsPersistent(), false /*bAllowCopyFromDDC*/);
+		constexpr bool bAllowCopyFromDDC = false;
+		EnsureSimulationDataIsCooked(Ar.IsTransacting(), bAllowCopyFromDDC);
 	}
 #endif
 
@@ -1862,6 +1866,13 @@ bool UGeometryCollection::Modify(bool bAlwaysMarkDirty /*= true*/)
 }
 
 void UGeometryCollection::EnsureDataIsCooked(bool bInitResources, bool bIsTransacting, bool bIsPersistant, bool bAllowCopyFromDDC)
+{
+	EnsureSimulationDataIsCooked(bIsTransacting, bAllowCopyFromDDC);
+
+	EnsureRenderDataIsCooked(bInitResources);
+}
+
+void UGeometryCollection::EnsureSimulationDataIsCooked(bool bIsTransacting, bool bAllowCopyFromDDC = true)
 {
 	if (StateGuid != LastBuiltSimulationDataGuid)
 	{
@@ -1886,15 +1897,14 @@ void UGeometryCollection::EnsureDataIsCooked(bool bInitResources, bool bIsTransa
 			CacheMaterialDensity();
 		}
 	}
+}
 
-	// Render data only goes through DDC when loading and saving (bIsPersistant).
+void UGeometryCollection::EnsureRenderDataIsCooked(bool bInitResources)
+{
+	// Render data only goes through DDC when loading and saving ( called from OnPostLoad / OnSave  )
 	// Using DDC during edits isn't worth it especially as we use a continually mutating guid instead of a state hash.
 	// That ensures that all edits are cache misses (slow) and unnecessarily fill up DDC disk space.
-	// TODO: SimulationData currently relies on these calls to update reliably, so we still need to use DDC for edits.
-	//       We could make CreateSimulationData() be reliably called for all edits and then only use DDC for loading and saving.
-	//       If we do that we can combine CreateSimulationDataImp() with CreateRenderDataImp() and FDerivedDataGeometryCollectionCooker
-	//       with FDerivedDataGeometryCollectionRenderDataCooker.
-	if (bIsPersistant && StateGuid != LastBuiltRenderDataGuid)
+	if (StateGuid != LastBuiltRenderDataGuid)
 	{
 		CreateRenderDataImp(/*bCopyFromDDC=*/ bInitResources);
 
@@ -1905,21 +1915,31 @@ void UGeometryCollection::EnsureDataIsCooked(bool bInitResources, bool bIsTransa
 				RenderData->InitResources(*this);
 			}
 		}
-	
+
 		LastBuiltRenderDataGuid = StateGuid;
 	}
 }
+
+#endif
+
+void UGeometryCollection::PreSave(FObjectPreSaveContext SaveContext)
+{
+#if WITH_EDITOR
+	constexpr bool bInitResources = false; 
+	EnsureRenderDataIsCooked(bInitResources);
 #endif
 
+	Super::PreSave(SaveContext);
+}
+
 void UGeometryCollection::PostLoad()
 {
 	Super::PostLoad();
 
-	// Initialize rendering resources.
-	if (FApp::CanEverRender())
-	{
-		InitResources();
-	}
+#if WITH_EDITOR
+	constexpr bool bInitResources = true;
+	EnsureRenderDataIsCooked(bInitResources);
+#endif
 
 #if WITH_EDITORONLY_DATA
 	if (!RootProxy_DEPRECATED.IsNull())
diff --git a/Engine/Source/Runtime/Experimental/GeometryCollectionEngine/Public/GeometryCollection/GeometryCollectionObject.h b/Engine/Source/Runtime/Experimental/GeometryCollectionEngine/Public/GeometryCollection/GeometryCollectionObject.h
--- a/Engine/Source/Runtime/Experimental/GeometryCollectionEngine/Public/GeometryCollection/GeometryCollectionObject.h
+++ b/Engine/Source/Runtime/Experimental/GeometryCollectionEngine/Public/GeometryCollection/GeometryCollectionObject.h
@@ -372,14 +372,12 @@ public:
 	GEOMETRYCOLLECTIONENGINE_API virtual bool Modify(bool bAlwaysMarkDirty = true) override;
 #endif
 	GEOMETRYCOLLECTIONENGINE_API virtual void PostInitProperties() override;
+	GEOMETRYCOLLECTIONENGINE_API virtual void PreSave(FObjectPreSaveContext SaveContext) override;
 	GEOMETRYCOLLECTIONENGINE_API virtual void PostLoad() override;
 	GEOMETRYCOLLECTIONENGINE_API virtual void BeginDestroy() override;
 	/** End UObject Interface */
 
 	GEOMETRYCOLLECTIONENGINE_API void Serialize(FArchive& Ar);
-#if WITH_EDITORONLY_DATA
-	GEOMETRYCOLLECTIONENGINE_API void PostSerialize(const FArchive& Ar);
-#endif
 
 #if WITH_EDITOR
 	GEOMETRYCOLLECTIONENGINE_API void EnsureDataIsCooked(bool bInitResources, bool bIsTransacting, bool bIsPersistant, bool bAllowCopyFromDDC = true);
@@ -844,6 +842,10 @@ private:
 #if WITH_EDITOR
 	GEOMETRYCOLLECTIONENGINE_API void CreateSimulationDataImp(bool bCopyFromDDC);
 	GEOMETRYCOLLECTIONENGINE_API void CreateRenderDataImp(bool bCopyFromDDC);
+
+	GEOMETRYCOLLECTIONENGINE_API void EnsureSimulationDataIsCooked(bool bIsTransacting, bool bAllowCopyFromDDC);
+	GEOMETRYCOLLECTIONENGINE_API void EnsureRenderDataIsCooked(bool bInitResources);
+
 #endif
 
 	/*
