Geometry collection : Make sure to check for potential null pointer in SetEnableGravity_External

diff --git a/Engine/Source/Runtime/Experimental/Chaos/Private/PhysicsProxy/GeometryCollectionPhysicsProxy.cpp b/Engine/Source/Runtime/Experimental/Chaos/Private/PhysicsProxy/GeometryCollectionPhysicsProxy.cpp
--- a/Engine/Source/Runtime/Experimental/Chaos/Private/PhysicsProxy/GeometryCollectionPhysicsProxy.cpp
+++ b/Engine/Source/Runtime/Experimental/Chaos/Private/PhysicsProxy/GeometryCollectionPhysicsProxy.cpp
@@ -3497,7 +3497,10 @@ void FGeometryCollectionPhysicsProxy::SetEnableGravity_External(bool bEnabled)
 {
 	for (TUniquePtr<FParticle>& GTParticle : GTParticles)
 	{
-		GTParticle->SetGravityEnabled(bEnabled);
+		if (GTParticle)
+		{
+			GTParticle->SetGravityEnabled(bEnabled);
+		}
 	}
 
 	ExecuteOnPhysicsThread(*this, 
