Download and install on-demand I/O store chunks using package dependency information

diff --git a/Engine/Source/Runtime/Experimental/IoStore/OnDemand/Private/OnDemandInstallCache.cpp b/Engine/Source/Runtime/Experimental/IoStore/OnDemand/Private/OnDemandInstallCache.cpp
--- a/Engine/Source/Runtime/Experimental/IoStore/OnDemand/Private/OnDemandInstallCache.cpp
+++ b/Engine/Source/Runtime/Experimental/IoStore/OnDemand/Private/OnDemandInstallCache.cpp
@@ -68,7 +68,8 @@ public:
 	virtual TIoStatusOr<FIoMappedRegion> OpenMapped(const FIoChunkId& ChunkId, const FIoReadOptions& Options) override;
 
 	// IOnDemandInstallCache
-	virtual FIoStatus Put(const FIoChunkId& ChunkId, FIoBuffer&& Chunk, const FIoHash& Hash) override;
+	virtual bool ContainsChunk(const FIoHash& Hash) override;
+	virtual FIoStatus PutChunk(FIoBuffer&& Chunk, const FIoHash& Hash) override;
 
 private:
 	bool						Resolve(FIoRequestImpl* Request);
@@ -218,7 +219,14 @@ bool FOnDemandInstallCache::Resolve(FIoRequestImpl* Request)
 	return true;
 }
 
-FIoStatus FOnDemandInstallCache::Put(const FIoChunkId& ChunkId, FIoBuffer&& Chunk, const FIoHash& Hash)
+bool FOnDemandInstallCache::ContainsChunk(const FIoHash& Hash)
+{
+	TStringBuilder<256> Filename;
+	GetChunkFilename(Hash, Filename);
+	return IFileManager::Get().FileSize(*Filename) > 0;
+}
+
+FIoStatus FOnDemandInstallCache::PutChunk(FIoBuffer&& Chunk, const FIoHash& Hash)
 {
 	return WriteChunk(Chunk, Hash);
 }
diff --git a/Engine/Source/Runtime/Experimental/IoStore/OnDemand/Private/OnDemandInstallCache.h b/Engine/Source/Runtime/Experimental/IoStore/OnDemand/Private/OnDemandInstallCache.h
--- a/Engine/Source/Runtime/Experimental/IoStore/OnDemand/Private/OnDemandInstallCache.h
+++ b/Engine/Source/Runtime/Experimental/IoStore/OnDemand/Private/OnDemandInstallCache.h
@@ -19,7 +19,8 @@ class IOnDemandInstallCache
 {
 public:
 	virtual ~IOnDemandInstallCache() = default;
-	virtual FIoStatus Put(const FIoChunkId& ChunkId, FIoBuffer&& Chunk, const FIoHash& Hash) = 0;
+	virtual bool ContainsChunk(const FIoHash& Hash) = 0;
+	virtual FIoStatus PutChunk(FIoBuffer&& Chunk, const FIoHash& Hash) = 0;
 };
 
 struct FOnDemandInstallCacheConfig
diff --git a/Engine/Source/Runtime/Experimental/IoStore/OnDemand/Private/OnDemandIoStore.cpp b/Engine/Source/Runtime/Experimental/IoStore/OnDemand/Private/OnDemandIoStore.cpp
--- a/Engine/Source/Runtime/Experimental/IoStore/OnDemand/Private/OnDemandIoStore.cpp
+++ b/Engine/Source/Runtime/Experimental/IoStore/OnDemand/Private/OnDemandIoStore.cpp
@@ -5,11 +5,13 @@
 #include "OnDemandInstallCache.h"
 #include "OnDemandPackageStoreBackend.h"
 
+#include "Algo/Copy.h"
 #include "Algo/Find.h"
 #include "Algo/RemoveIf.h"
 #include "Algo/Transform.h"
 #include "Async/ManualResetEvent.h"
 #include "Async/UniqueLock.h"
+#include "Containers/RingBuffer.h"
 #include "Containers/StringConv.h"
 #include "HAL/FileManager.h"
 #include "HAL/PlatformFile.h"
@@ -44,11 +46,13 @@ extern FString GIasOnDemandTocExt;
 namespace Private
 {
 
+///////////////////////////////////////////////////////////////////////////////
 static FString GetInstallCacheDirectory()
 {
 	return FPaths::ProjectPersistentDownloadDir() / TEXT("IoStore") / TEXT("InstallCache");
 }
 
+///////////////////////////////////////////////////////////////////////////////
 static void SplitHostUrl(const FStringView& Url, FStringView& OutHost, FStringView& OutRemainder)
 {
 	OutHost = OutRemainder = FStringView();
@@ -65,6 +69,174 @@ static void SplitHostUrl(const FStringView& Url, FStringView& OutHost, FStringVi
 	OutRemainder = Url.RightChop(OutHost.Len());
 }
 
+///////////////////////////////////////////////////////////////////////////////
+using FPackageStoreEntryMap = TMap<FPackageId, const FFilePackageStoreEntry*>;
+
+struct FContainerInstallData
+{
+	FPackageStoreEntryMap	PackageStoreEntries;
+	TSet<FPackageId>		PackageIds;
+	TSet<FIoChunkId>		ResolvedChunks;
+	uint64					TotalSize = 0;
+};
+
+using FInstallData = TMap<FSharedOnDemandContainer, FContainerInstallData>;
+
+FIoStatus BuildInstallData(
+	const TSet<FSharedOnDemandContainer>& Containers,
+	const TSet<FPackageId>& PackageIds,
+	FInstallData& OutInstallData,
+	const TSet<FPackageId>& OutMissing)
+{
+	OutInstallData.Reserve(Containers.Num());
+
+	// Setup the package information for each container
+	for (const FSharedOnDemandContainer& Container : Containers)
+	{
+		if (!Container->Header.IsValid() || Container->Header->PackageIds.IsEmpty())
+		{
+			// The container contains no package data
+			continue;
+		}
+
+		FContainerInstallData& Data = OutInstallData.FindOrAdd(Container);
+		const FIoContainerHeader& Header = *Container->Header;
+		TConstArrayView<FFilePackageStoreEntry> Entries(
+			reinterpret_cast<const FFilePackageStoreEntry*>(Header.StoreEntries.GetData()),
+			Header.PackageIds.Num());
+		
+		Data.PackageStoreEntries.Reserve(Header.PackageIds.Num());
+
+		int32 Idx = 0;
+		for (const FFilePackageStoreEntry& Entry : Entries)
+		{
+			const FPackageId PackageId = Header.PackageIds[Idx++];
+			Data.PackageStoreEntries.Add(PackageId, &Entry);
+		}
+	}
+
+	// Traverse dependencies for each package id
+	using FQueue = TRingBuffer<FPackageId>;
+
+	FQueue Queue;
+	TSet<FPackageId> Visitied;
+	TSet<FPackageId> Missing;
+
+	Visitied.Reserve(PackageIds.Num());
+	Queue.Reserve(PackageIds.Num());
+
+	//Algo::Copy(PackageIds, Queue);
+	for (const FPackageId& PackageId : PackageIds)
+	{
+		Queue.Add(PackageId);
+	}
+
+	while (!Queue.IsEmpty())
+	{
+		FPackageId PackageId = Queue.PopFrontValue();
+
+		bool bIsAlreadyInSet = false;
+		Visitied.Add(PackageId, &bIsAlreadyInSet);
+		if (bIsAlreadyInSet)
+		{
+			continue;
+		}
+
+		bool bFound = false;
+		for (TPair<FSharedOnDemandContainer, FContainerInstallData>& Kv : OutInstallData)
+		{
+			FContainerInstallData& Data = Kv.Value;
+			if (const FFilePackageStoreEntry** Entry = Data.PackageStoreEntries.Find(PackageId))
+			{
+				Data.PackageIds.Add(PackageId);
+
+				// Add all imported packages not already processed
+				const FFilePackageStoreEntry& PackageStoreEntry = **Entry;
+				for (const FPackageId& ImportedPackageId : PackageStoreEntry.ImportedPackages)
+				{
+					if (!Visitied.Contains(ImportedPackageId))
+					{
+						Queue.Add(ImportedPackageId);
+					}
+				}
+
+				bFound = true;
+				break;
+			}
+		}
+
+		if (!bFound)
+		{
+			Missing.Add(PackageId);
+		}
+	}
+
+	for (TPair<FSharedOnDemandContainer, FContainerInstallData>& Kv : OutInstallData)
+	{
+		const FOnDemandContainer& Container = *Kv.Key;
+		FContainerInstallData& Data			= Kv.Value;
+
+		if (Data.PackageIds.IsEmpty())
+		{
+			continue;
+		}
+
+		for (const FPackageId& PackageId : Data.PackageIds)
+		{
+			const FIoChunkId PackageChunkId					= CreatePackageDataChunkId(PackageId);
+			const FOnDemandChunkEntry* PackageChunkEntry	= Container.ChunkEntries.Find(PackageChunkId);
+
+			if (PackageChunkEntry == nullptr)
+			{
+				UE_LOG(LogIoStoreOnDemand, Warning, TEXT("Missing package data"));
+				continue;
+			}
+
+			Data.ResolvedChunks.Add(PackageChunkId);
+			Data.TotalSize += PackageChunkEntry->EncodedSize;
+
+			const EIoChunkType AdditionalPackageChunkTypes[] =
+			{
+				EIoChunkType::BulkData,
+				EIoChunkType::OptionalBulkData,
+				EIoChunkType::MemoryMappedBulkData 
+			};
+
+			for (EIoChunkType ChunkType : AdditionalPackageChunkTypes)
+			{
+				const FIoChunkId ChunkId = CreateIoChunkId(PackageId.Value(), 0, ChunkType);
+				if (const FOnDemandChunkEntry* ChunkEntry = Container.ChunkEntries.Find(ChunkId); ChunkEntry != nullptr)
+				{
+					Data.ResolvedChunks.Add(ChunkId);
+					Data.TotalSize += ChunkEntry->EncodedSize;
+				}
+			}
+		}
+
+		// For now we always download these chunks
+		for (const TPair<FIoChunkId, FOnDemandChunkEntry>& IdEntry : Container.ChunkEntries)
+		{
+			const FIoChunkId& ChunkId				= IdEntry.Key;
+			const FOnDemandChunkEntry& ChunkEntry	= IdEntry.Value;
+
+			switch(ChunkId.GetChunkType())
+			{
+				case EIoChunkType::ExternalFile:
+				case EIoChunkType::ShaderCodeLibrary:
+				case EIoChunkType::ShaderCode:
+				{
+					Data.ResolvedChunks.Add(ChunkId);
+					Data.TotalSize += ChunkEntry.EncodedSize;
+				}
+				default:
+					break;
+			}
+		}
+	}
+
+	return EIoErrorCode::Ok;
+}
+
 } // namespace UE::IoStore::Private
 
 ///////////////////////////////////////////////////////////////////////////////
@@ -231,19 +403,18 @@ FIoStatus FOnDemandIoStore::Initialize()
 				const FString OnDemandTocPath = FPathViews::ChangeExtension(Info.PakFile->PakGetPakFilename(), GIasOnDemandTocExt);
 				if (Ipf.FileExists(*OnDemandTocPath))
 				{
-					MountRequests.Add(OnDemandTocPath, MakeShared<FMountRequest>(FMountRequest
+					FSharedMountRequest MountRequest = MakeShared<FMountRequest>();
+					MountRequest->MountArgs = FOnDemandMountArgs
+					{
+						.MountId	= OnDemandTocPath,
+						.FilePath	= OnDemandTocPath
+					};
+					MountRequest->OnCompleted = [](TIoStatusOr<FOnDemandMountResult> Result)
 					{
-						.MountArgs = FOnDemandMountArgs
-						{
-							.MountId = OnDemandTocPath,
-							.FilePath = OnDemandTocPath
-						},
-						.OnCompleted = [](TIoStatusOr<FOnDemandMountResult> Result)
-						{
-							UE_CLOG(!Result.IsOk(), LogIoStoreOnDemand, Error, TEXT("Failed to mount container, reason '%s'"),
-								*Result.Status().ToString());
-						}
-					}));
+						UE_CLOG(!Result.IsOk(), LogIoStoreOnDemand, Error, TEXT("Failed to mount container, reason '%s'"),
+							*Result.Status().ToString());
+					};
+					MountRequests.Add(OnDemandTocPath, MoveTemp(MountRequest));
 				}
 			}
 
@@ -273,11 +444,10 @@ void FOnDemandIoStore::Mount(FOnDemandMountArgs&& Args, FOnDemandMountCompleted&
 		}
 
 		UE_LOG(LogIoStoreOnDemand, Log, TEXT("Enqueing mount request, MountId='%s'"), *Args.MountId);
-		MountRequest = MakeShared<FMountRequest>(FMountRequest
-		{ 
-			.MountArgs		= MoveTemp(Args),
-			.OnCompleted	= MoveTemp(OnCompleted),
-		});
+		
+		MountRequest = MakeShared<FMountRequest>();
+		MountRequest->MountArgs		= MoveTemp(Args);
+		MountRequest->OnCompleted	= MoveTemp(OnCompleted);
 	}
 
 	TryEnterTickLoop();
@@ -415,21 +585,40 @@ bool FOnDemandIoStore::Tick()
 
 	for (FSharedMountRequest& Request : Requests)
 	{
-		FIoStatus Status = TickMountRequest(*Request);
+		// Tick mount
+		if (FIoStatus Status = TickMountRequest(*Request); !Status.IsOk())
+		{
+			{
+				UE::TUniqueLock Lock(MountRequestMutex);
+				MountRequests.Remove(Request->MountArgs.MountId);
+			}
+			FOnDemandMountCompleted OnCompleted = MoveTemp(Request->OnCompleted);
+			OnCompleted(Status);
+			continue;
+		}
 
+		// Tick install
+		if (EnumHasAnyFlags(Request->MountArgs.Options, EOnDemandMountOptions::Install))
 		{
-			UE::TUniqueLock Lock(MountRequestMutex);
-			MountRequests.Remove(Request->MountArgs.MountId);
-			if (Request->bCancelled)
+			if (FIoStatus Status = TickInstallRequest(*Request); !Status.IsOk())
 			{
-				Status = EIoErrorCode::Cancelled;
+				{
+					UE::TUniqueLock Lock(MountRequestMutex);
+					MountRequests.Remove(Request->MountArgs.MountId);
+				}
+				FOnDemandMountCompleted OnCompleted = MoveTemp(Request->OnCompleted);
+				OnCompleted(Status);
+				continue;;
 			}
 		}
 
-		TStringBuilder<128> Sb;
-		TIoStatusOr<FOnDemandMountResult> MountStatus;
-		if (Status.IsOk())
 		{
+			{
+				UE::TUniqueLock Lock(MountRequestMutex);
+				MountRequests.Remove(Request->MountArgs.MountId);
+			}
+
+			TStringBuilder<128> Sb;
 			UE::TUniqueLock Lock(ContainerMutex);
 			for (const FSharedOnDemandContainer& Container : Request->Containers)
 			{
@@ -466,18 +655,12 @@ bool FOnDemandIoStore::Tick()
 				EnumAddFlags(Container->Flags, EOnDemandContainerFlags::Mounted);
 			}
 
-			MountStatus = TIoStatusOr<FOnDemandMountResult>(FOnDemandMountResult
+			FOnDemandMountCompleted OnCompleted = MoveTemp(Request->OnCompleted);
+			OnCompleted(FOnDemandMountResult
 			{
 				.MountId = Request->MountArgs.MountId
 			});
 		}
-		else
-		{
-			MountStatus = TIoStatusOr<FOnDemandMountResult>(Status);	
-		}
-
-		FOnDemandMountCompleted OnCompleted = MoveTemp(Request->OnCompleted);
-		OnCompleted(MountStatus);
 	}
 
 	return true;
@@ -569,48 +752,288 @@ FIoStatus FOnDemandIoStore::TickMountRequest(FMountRequest& MountRequest)
 		}
 	}
 
-	if (EnumHasAnyFlags(Args.Options, EOnDemandMountOptions::Install))
+	return EIoErrorCode::Ok;
+}
+
+FIoStatus FOnDemandIoStore::TickInstallRequest(FMountRequest& MountRequest)
+{
+	if (InstallCache.IsValid() == false || PackageStoreBackend.IsValid() == false)
+	{
+		return FIoStatusBuilder(EIoErrorCode::InvalidCode) << TEXT("Install cache not configured");
+	}
+
+	if (MountRequest.bCancelled)
 	{
-		if (InstallCache.IsValid() == false || PackageStoreBackend.IsValid() == false)
+		return EIoErrorCode::Cancelled;
+	}
+
+	TAnsiStringBuilder<512> ChunkUrl;
+	FStringView Host, TocRelUrl;
+	Private::SplitHostUrl(MountRequest.MountArgs.Url, Host, TocRelUrl);
+
+	auto GetChunkUrl = [](
+		const FStringView& Host,
+		const FOnDemandContainer& Container,
+		const FOnDemandChunkEntry& Entry,
+		FAnsiStringBuilderBase& OutUrl) -> FAnsiStringBuilderBase&
+	{
+		OutUrl.Reset();
+		if (Host.IsEmpty() == false)
 		{
-			return FIoStatusBuilder(EIoErrorCode::InvalidCode) << TEXT("Install cache not configured");
+			OutUrl << Host;
 		}
 
-		const bool bUsePackageFilter = Args.TagSets.Num() > 0;
+		if (!Container.ChunksDirectory.IsEmpty())
+		{
+			OutUrl << "/" << Container.ChunksDirectory;
+		}
 
-		// List of unique container name, package index list pairs
-		FPackageFilter PackageFilter;
+		const FString HashString = LexToString(Entry.Hash);
+		OutUrl << "/" << HashString.Left(2) << "/" << HashString << ".iochunk";
+
+		return OutUrl;
+	};
 
-		if (bUsePackageFilter && ensure(Args.Toc))
+	auto FetchContainerHeader = [&Host, &GetChunkUrl, &ChunkUrl](
+		const FOnDemandContainer& Container,
+		FIoContainerHeader& Header) -> FIoStatus 
+	{
+		const FIoChunkId			ChunkId = CreateContainerHeaderChunkId(Container.ContainerId);
+		const FOnDemandChunkEntry*	Entry = Container.ChunkEntries.Find(ChunkId);
+
+		if (Entry == nullptr)
 		{
-			const FOnDemandToc& Toc = Args.Toc.GetValue();
-			for (const FString& TagSet : Args.TagSets)
+			return EIoErrorCode::Ok;
+		}
+
+		UE_LOG(LogIoStoreOnDemand, VeryVerbose, TEXT("Fetching container header, ContainerName='%s'"), *Container.Name);
+		TIoStatusOr<FIoBuffer> Response = FHttpClient::Get(GetChunkUrl(Host, Container, *Entry, ChunkUrl).ToView(), 2, EHttpRedirects::Follow);
+		if (Response.IsOk() == false)
+		{
+			FIoStatus Status = FIoStatusBuilder(EIoErrorCode::InvalidCode) << TEXT("Failed to fetch container header chunk from URL");
+			return Status;
+		}
+
+		FMemoryReaderView Ar(Response.ValueOrDie().GetView());
+		Ar << Header;
+		Ar.Close();
+
+		if (Ar.IsError() || Ar.IsCriticalError())
+		{
+			FIoStatus Status = FIoStatusBuilder(EIoErrorCode::FileNotOpen) << TEXT("Failed to serialize container header");
+			return Status;
+		}
+
+		return EIoErrorCode::Ok;
+	};
+
+	TArray<FSharedOnDemandContainer> MountedContainers = GetMountedContainers();
+	TSet<FSharedOnDemandContainer> AllContainers;
+
+	Algo::Copy(MountedContainers, AllContainers);
+	Algo::Copy(MountRequest.Containers, AllContainers);
+
+	// Fetch all container headers
+	for (const FSharedOnDemandContainer& Container : AllContainers)
+	{
+		if (!Container->Header.IsValid())
+		{
+			FSharedContainerHeader Header = MakeShared<FIoContainerHeader>();
+			if (FIoStatus Status = FetchContainerHeader(*Container, *Header); !Status.IsOk())
 			{
-				const FOnDemandTocTagSet* MaybeTagSet = Algo::FindBy(Toc.TagSets, TagSet, &FOnDemandTocTagSet::Tag);
-				if (MaybeTagSet)
+				return Status; 
+			}
+			Container->Header = MoveTemp(Header);
+		}
+	}
+
+	// Parse the tag sets
+	if (MountRequest.MountArgs.TagSets.IsEmpty() == false && MountRequest.TagSets.IsEmpty())
+	{
+		const FOnDemandToc& Toc = MountRequest.MountArgs.Toc.GetValue();
+		for (const FString& Tag : MountRequest.MountArgs.TagSets)
+		{
+			const FOnDemandTocTagSet* TagSet = Algo::FindBy(Toc.TagSets, Tag, &FOnDemandTocTagSet::Tag);
+			if (TagSet == nullptr)
+			{
+				UE_LOG(LogIoStoreOnDemand, Warning, TEXT("Tag set '%s' doesn't exists"), *Tag);
+				continue;
+			}
+
+			FTagSet NewTagSet; 
+			NewTagSet.Tag = TagSet->Tag;
+			for (const FOnDemandTocTagSetPackageList& PackageIndices : TagSet->Packages)
+			{
+				if (MountRequest.Containers.IsEmpty() || int32(PackageIndices.ContainerIndex) > (MountRequest.Containers.Num() - 1))
 				{
-					for (const FOnDemandTocTagSetPackageList& PkgList : MaybeTagSet->Packages)
-					{
-						if (!Toc.Containers.IsValidIndex(PkgList.ContainerIndex))
-						{
-							return FIoStatusBuilder(EIoErrorCode::CorruptToc) << TEXT("Bad container index in tag set");
-						}
+					UE_LOG(LogIoStoreOnDemand, Warning, TEXT("Invalid container index in tag set '%s'"), *Tag);
+					continue;
+				}
+
+				FOnDemandContainer& Container = *MountRequest.Containers[PackageIndices.ContainerIndex];
+				if (Container.Header.IsValid() == false)
+				{
+					UE_LOG(LogIoStoreOnDemand, Warning, TEXT("Tag set specified for invalid container header, Tag='%s', ContainerName='%s'"),
+						*Tag, *Container.UniqueName());
+					continue;
+				}
+
+				NewTagSet.PackageIds.Reserve(PackageIndices.PackageIndicies.Num());
+				for (int32 IdIdx : PackageIndices.PackageIndicies)
+				{
+					NewTagSet.PackageIds.Add(Container.Header->PackageIds[IdIdx]);
+				}
+			}
+
+			MountRequest.TagSets.Emplace(MoveTemp(NewTagSet));
+		}
+	}
 
-						FString ContainerName = OnDemandContainerUniqueName(Args.MountId, Toc.Containers[PkgList.ContainerIndex].ContainerName);
-						PackageFilter.Emplace(MoveTemp(ContainerName), MakeArrayView(PkgList.PackageIndicies));
+	TSet<FPackageId> PackageIdsToInstall;
+	// Install all packages if no tag set was specified
+	if (MountRequest.MountArgs.TagSets.IsEmpty())
+	{
+		for (const FSharedOnDemandContainer& Container : MountRequest.Containers)
+		{
+			if (Container->Header.IsValid())
+			{
+				PackageIdsToInstall.Reserve(Container->Header->PackageIds.Num());
+				for (const FPackageId& PackageId : Container->Header->PackageIds)
+				{
+					PackageIdsToInstall.Add(PackageId);
+				}
+			}
+		}
+	}
+	else
+	{
+		for (const FString& Tag : MountRequest.MountArgs.TagSets)
+		{
+			const FTagSet* TagSet = nullptr;
+			for (const FTagSet& Set : MountRequest.TagSets)
+			{
+				if (Set.Tag == Tag)
+				{
+					PackageIdsToInstall.Reserve(Set.PackageIds.Num());
+					for (const FPackageId& PackageId : Set.PackageIds)
+					{
+						PackageIdsToInstall.Add(PackageId);
 					}
 				}
 			}
 		}
+	}
+
+	// Find all I/O chunks fo the specified list of packages
+	Private::FInstallData InstallData;
+	TSet<FPackageId> Missing;
+
+	FIoStatus Status = BuildInstallData(AllContainers, PackageIdsToInstall, InstallData, Missing);
+	if (Status.IsOk() == false)
+	{
+		return Status;
+	}
+
+	// Check the other I/O backends for missing packge chunks
+	for (const FPackageId& PackageId : Missing)
+	{
+		const FIoChunkId ChunkId = CreatePackageDataChunkId(PackageId);
+		if (FIoDispatcher::Get().DoesChunkExist(ChunkId) == false)
+		{
+			// TODO: Fail the install request or just try to continue
+			UE_LOG(LogIoStoreOnDemand, Warning, TEXT("Missing package chunk '%s'"), *LexToString(ChunkId));
+		}
+	}
+
+	// Download all chunks
+	const int32 MaxConcurrentRequests = 32;
+	int32 ConcurrentRequests = 0;
 
-		FIoStatus Status = InstallContainers(
-			Args.Url, MountRequest.Containers, bUsePackageFilter ? &PackageFilter : nullptr);
-		if (!Status.IsOk())
+	FHttpClientConfig HttpConfig;
+	HttpConfig.MaxConnectionCount = 8;
+	HttpConfig.MaxRetryCount = 2;
+	HttpConfig.Endpoints.Add(FString(Host));
+
+	TUniquePtr<FHttpClient> HttpClient = FHttpClient::Create(MoveTemp(HttpConfig));
+	if (HttpClient.IsValid() == false)
+	{
+		return FIoStatusBuilder(EIoErrorCode::InvalidCode) << TEXT("Failed to initialize HTTP client");
+	}
+
+	uint64 TotalChunkCount		= 0;
+	uint64 TotalBytes			= 0;
+	uint64 DownloadedChunkCount	= 0;
+	uint64 DownloadedBytes		= 0;
+
+	for (const auto& Kv : InstallData)
+	{
+		FOnDemandContainer& Container				= *Kv.Key;
+		const Private::FContainerInstallData& Data	= Kv.Value;
+
+		for (const FIoChunkId& ChunkId : Data.ResolvedChunks)
 		{
-			return Status;
+			const FOnDemandChunkEntry& Entry = Container.ChunkEntries.FindRef(ChunkId);
+			++TotalChunkCount;
+			TotalBytes += Entry.EncodedSize;
+
+			if (InstallCache->ContainsChunk(Entry.Hash))
+			{
+				continue;
+			}
+
+			++DownloadedChunkCount;
+			DownloadedBytes += Entry.EncodedSize;
+
+			FOnDemandChunkEntry* ChunkEntry = Container.ChunkEntries.Find(ChunkId);
+			ConcurrentRequests++;
+			HttpClient->Get(
+				GetChunkUrl(FStringView(), Container, *ChunkEntry, ChunkUrl).ToView(),
+				[this, &Status, ChunkId, ChunkEntry, &ConcurrentRequests]
+				(TIoStatusOr<FIoBuffer> ChunkStatus, uint64 DurationMs) mutable
+				{
+					ConcurrentRequests--;
+					if (!ChunkStatus.IsOk())
+					{
+						Status = ChunkStatus.Status();
+						return;
+					}
+
+					Status = InstallCache->PutChunk(ChunkStatus.ConsumeValueOrDie(), ChunkEntry->Hash);
+				});
+
+			while (ConcurrentRequests >= MaxConcurrentRequests)
+			{
+				HttpClient->Tick();
+			}
+
+			if (Status.IsOk() == false)
+			{
+				return Status;
+			}
+		}
+	}
+
+	while (HttpClient->Tick())
+		;
+
+	// TODO: Only mount what has been installed
+	for (const auto& Kv : InstallData)
+	{
+		FOnDemandContainer& Container				= *Kv.Key;
+		const Private::FContainerInstallData& Data	= Kv.Value;
+		if (!Data.PackageIds.IsEmpty())
+		{
+			check(Container.Header.IsValid());
+			const FIoStatus MountStatus = PackageStoreBackend->Mount(Container.UniqueName(), Container.Header);
+			check(MountStatus.IsOk());
 		}
+
+		EnumAddFlags(Container.Flags, EOnDemandContainerFlags::Installed);
 	}
 
+	UE_LOG(LogIoStoreOnDemand, Log, TEXT("Downloaded %llu (%.2lf MiB) of total %llu (%.2lf MiB) requested chunk(s)'"),
+		DownloadedChunkCount, double(DownloadedBytes) / 1024.0 / 1024.0, TotalChunkCount, double(TotalBytes) / 1024.0 / 1024.0);
+
 	return EIoErrorCode::Ok;
 }
 
@@ -705,186 +1128,6 @@ void FOnDemandIoStore::CreateContainersFromToc(
 	}
 }
 
-FIoStatus FOnDemandIoStore::InstallContainers(
-	const FString& Url,
-	const TConstArrayView<FSharedOnDemandContainer>& ContainersToInstall,
-	const FOnDemandIoStore::FPackageFilter* PackageFilter /*= nullptr*/)
-{
-	FStringView Host, TocRelUrl;
-	Private::SplitHostUrl(Url, Host, TocRelUrl);
-	const FStringView TocPath = FPathViews::GetPath(TocRelUrl);
-	auto AnsiTocPath = StringCast<ANSICHAR>(TocPath.GetData(), TocPath.Len());
-
-	TAnsiStringBuilder<512> ChunkUrl;
-	for (const FSharedOnDemandContainer& Container : ContainersToInstall)
-	{
-		if (EnumHasAnyFlags(Container->Flags, EOnDemandContainerFlags::Installed))
-		{
-			UE_LOG(LogIoStoreOnDemand, Log, TEXT("Container '%s' is already installed, MountId='%s'"),
-				*Container->Name, *Container->MountId);
-			continue;
-		}
-
-		auto GetChunkUrl = [&AnsiTocPath, &Container](
-			const FStringView& Host,
-			const FOnDemandChunkEntry& Entry,
-			FAnsiStringBuilderBase& OutUrl) -> FAnsiStringBuilderBase&
-		{
-			OutUrl.Reset();
-			if (Host.IsEmpty() == false)
-			{
-				OutUrl << Host;
-			}
-
-			if (Container->ChunksDirectory.IsEmpty())
-			{
-				OutUrl << AnsiTocPath << "/" << "chunks";
-			}
-			else
-			{
-				OutUrl << "/" << Container->ChunksDirectory;
-			}
-
-			const FString HashString = LexToString(Entry.Hash);
-			OutUrl << "/" << HashString.Left(2) << "/" << HashString << ".iochunk";
-
-			return OutUrl;
-		};
-
-		UE_LOG(LogIoStoreOnDemand, Log, TEXT("Installing container, Name='%s', MountId='%s'"),
-			*Container->Name, *Container->MountId);
-
-		FIoContainerHeader ContainerHeader;
-
-		const FIoChunkId HeaderChunkId = CreateContainerHeaderChunkId(Container->ContainerId);
-		if (const FOnDemandChunkEntry* Entry = Container->ChunkEntries.Find(HeaderChunkId))
-		{
-			TIoStatusOr<FIoBuffer> Response = FHttpClient::Get(GetChunkUrl(Host, *Entry, ChunkUrl).ToView(), 2, EHttpRedirects::Follow);
-			if (Response.IsOk() == false)
-			{
-				return FIoStatusBuilder(EIoErrorCode::InvalidCode) << TEXT("Failed to fetch container header chunk from URL");
-			}
-
-			FMemoryReaderView Ar(Response.ValueOrDie().GetView());
-			Ar << ContainerHeader;
-			Ar.Close();
-			if (Ar.IsError() || Ar.IsCriticalError())
-			{
-				return FIoStatusBuilder(EIoErrorCode::FileNotOpen) << TEXT("Failed to serialize container header");
-			}
-		}
-
-		if (ContainerHeader.PackageIds.Num())
-		{
-			TConstArrayView<FFilePackageStoreEntry> PackageStoreEntries(
-				reinterpret_cast<const FFilePackageStoreEntry*>(ContainerHeader.StoreEntries.GetData()),
-				ContainerHeader.PackageIds.Num());
-
-			//TODO: Find all referenced chunks, do not apply filter recursively, it should only be for exports
-			
-			if (!PackageFilter)
-			{
-				for (int32 Index = 0; Index < ContainerHeader.PackageIds.Num(); ++Index)
-				{
-					const FPackageId& PackageId = ContainerHeader.PackageIds[Index];
-					const FFilePackageStoreEntry& FilePackageStoreEntry = PackageStoreEntries[Index];
-				}
-			}
-			else
-			{
-				const FString ContainerName = Container->UniqueName();
-				for (const TPair<FString, TConstArrayView<uint32>>& Pair : *PackageFilter)
-				{
-					if (Pair.Key == ContainerName)
-					{
-						for (const uint32 Index : Pair.Value)
-						{
-							if (!ContainerHeader.PackageIds.IsValidIndex(Index))
-							{
-								return FIoStatusBuilder(EIoErrorCode::CorruptToc) << TEXT("Bad pkg index in tagset");
-							}
-
-							const FPackageId& PackageId = ContainerHeader.PackageIds[Index];
-							const FFilePackageStoreEntry& FilePackageStoreEntry = PackageStoreEntries[Index];
-						}
-					}
-				}
-			}
-		}
-
-		// Download all chunks
-		const int32 MaxConcurrentRequests = 32;
-		int32 ConcurrentRequests = 0;
-
-		FHttpClientConfig HttpConfig;
-		HttpConfig.MaxConnectionCount = 8;
-		HttpConfig.MaxRetryCount = 2;
-		HttpConfig.Endpoints.Add(Url);
-
-		TUniquePtr<FHttpClient> HttpClient = FHttpClient::Create(MoveTemp(HttpConfig));
-		if (HttpClient.IsValid() == false)
-		{
-			return FIoStatusBuilder(EIoErrorCode::InvalidCode) << TEXT("Failed to initialize HTTP client");
-		}
-
-		FIoStatus Status(EIoErrorCode::Ok);
-		for (const TPair<FIoChunkId, FOnDemandChunkEntry>& Kv : Container->ChunkEntries)
-		{
-			const FIoChunkId& ChunkId			= Kv.Key;
-			const FOnDemandChunkEntry& TocEntry	= Kv.Value;
-			
-			ConcurrentRequests++;
-			HttpClient->Get(
-				GetChunkUrl(FStringView(), TocEntry, ChunkUrl).ToView(),
-				[this, &Status, ChunkId, ExpectedHash = TocEntry.Hash, &ConcurrentRequests]
-				(TIoStatusOr<FIoBuffer> ChunkStatus, uint64 DurationMs)
-				{
-					ConcurrentRequests--;
-					if (ChunkStatus.IsOk())
-					{
-						if (FIoStatus PutStatus = InstallCache->Put(
-							ChunkId,
-							ChunkStatus.ConsumeValueOrDie(),
-							ExpectedHash); !PutStatus.IsOk())
-						{
-							Status = PutStatus;
-						}
-					}
-					else
-					{
-						Status = ChunkStatus.Status();
-					}
-				});
-
-			while (ConcurrentRequests >= MaxConcurrentRequests)
-			{
-				HttpClient->Tick();
-			}
-
-			if (Status.IsOk() == false)
-			{
-				return Status;
-			}
-		}
-
-		while (HttpClient->Tick());
-
-		if (ContainerHeader.PackageIds.IsEmpty() == false)
-		{
-			if (FIoStatus MountStatus = PackageStoreBackend->Mount(
-				Container->UniqueName(),
-				MoveTemp(ContainerHeader)); !MountStatus.IsOk())
-			{
-				return MountStatus;
-			}
-		}
-
-		EnumAddFlags(Container->Flags, EOnDemandContainerFlags::Installed);
-	}
-
-	return EIoErrorCode::Ok;
-}
-
 TArray<FSharedOnDemandContainer> FOnDemandIoStore::GetMountedContainers()
 {
 	UE::TUniqueLock Lock(ContainerMutex);
diff --git a/Engine/Source/Runtime/Experimental/IoStore/OnDemand/Private/OnDemandIoStore.h b/Engine/Source/Runtime/Experimental/IoStore/OnDemand/Private/OnDemandIoStore.h
--- a/Engine/Source/Runtime/Experimental/IoStore/OnDemand/Private/OnDemandIoStore.h
+++ b/Engine/Source/Runtime/Experimental/IoStore/OnDemand/Private/OnDemandIoStore.h
@@ -12,6 +12,11 @@
 #include "Misc/AES.h"
 #include "Misc/EnumClassFlags.h"
 
+#include <atomic>
+
+struct FIoContainerHeader;
+using FSharedContainerHeader		= TSharedPtr<FIoContainerHeader>;
+
 namespace UE::IoStore
 {
 
@@ -44,7 +49,7 @@ struct FOnDemandChunkEntry
 	uint32	RawSize = 0;
 	uint32	EncodedSize = 0;
 	uint32	BlockOffset = ~uint32(0);
-	uint32	BlockCount = 0; 
+	uint32	BlockCount = 0;
 	uint8	CompressionFormatIndex = 0;
 };
 
@@ -56,6 +61,7 @@ struct FOnDemandContainer
 	FAES::FAESKey			EncryptionKey;
 	FIoHash					PakHash;
 	FChunkEntryMap 			ChunkEntries;
+	FSharedContainerHeader	Header;
 	FString					EncryptionKeyGuid;
 	FString					Name;
 	FString					MountId;
@@ -119,12 +125,19 @@ TConstArrayView<FIoBlockHash> FOnDemandChunkInfo::BlockHashes() const
 ///////////////////////////////////////////////////////////////////////////////
 class FOnDemandIoStore
 {
+	struct FTagSet
+	{
+		FString Tag;
+		TArray<FPackageId> PackageIds;
+	};
+
 	struct FMountRequest
 	{
 		FOnDemandMountArgs					MountArgs;
 		FOnDemandMountCompleted				OnCompleted;
 		TArray<FSharedOnDemandContainer>	Containers;
-		bool								bCancelled = false;
+		TArray<FTagSet>						TagSets;
+		std::atomic_bool					bCancelled{false};
 	};
 
 	using FSharedMountRequest	= TSharedPtr<FMountRequest>;
@@ -152,17 +165,13 @@ private:
 	void					TickLoop();
 	bool					Tick();
 	FIoStatus				TickMountRequest(FMountRequest& MountRequest);
-	void					ConditionallyStartTicking();
+	FIoStatus				TickInstallRequest(FMountRequest& MountRequest);
 	void					OnEncryptionKeyAdded(const FGuid& Id, const FAES::FAESKey& Key);
 	static void				CreateContainersFromToc(
 								FStringView MountId,
 								FStringView TocPath,
 								FOnDemandToc& Toc,
 								TArray<FSharedOnDemandContainer>& Out);
-	FIoStatus				InstallContainers(
-								const FString& Url,
-								const TConstArrayView<FSharedOnDemandContainer>& ContainersToInstall,
-								const FPackageFilter* PackageFilter = nullptr);
 
 	TArray<FSharedOnDemandContainer> GetMountedContainers();
 
diff --git a/Engine/Source/Runtime/Experimental/IoStore/OnDemand/Private/OnDemandPackageStoreBackend.cpp b/Engine/Source/Runtime/Experimental/IoStore/OnDemand/Private/OnDemandPackageStoreBackend.cpp
--- a/Engine/Source/Runtime/Experimental/IoStore/OnDemand/Private/OnDemandPackageStoreBackend.cpp
+++ b/Engine/Source/Runtime/Experimental/IoStore/OnDemand/Private/OnDemandPackageStoreBackend.cpp
@@ -3,6 +3,7 @@
 #include "OnDemandPackageStoreBackend.h"
 
 #include "Algo/Accumulate.h"
+#include "Algo/Find.h"
 #include "Async/Mutex.h"
 #include "Async/UniqueLock.h"
 #include "IO/IoContainerHeader.h"
@@ -19,13 +20,13 @@ class FOnDemandPackageStoreBackend final
 {
 	struct FContainer
 	{
-		FContainer(FString&& Name, FIoContainerHeader&& Header)
-			: ContainerName(MoveTemp(Name))
-			, ContainerHeader(MoveTemp(Header))
+		FContainer(FString&& ContainerName, FSharedContainerHeader ContainerHeader)
+			: Name(MoveTemp(ContainerName))
+			, Header(MoveTemp(ContainerHeader))
 		{ }
 
-		FString				ContainerName;
-		FIoContainerHeader	ContainerHeader;
+		FString					Name;
+		FSharedContainerHeader	Header;
 	};
 
 	using FSharedBackendContext = TSharedPtr<const FPackageStoreBackendContext>;
@@ -39,8 +40,9 @@ public:
 						FOnDemandPackageStoreBackend();
 						virtual ~FOnDemandPackageStoreBackend ();
 
-	virtual FIoStatus	Mount(FString ContainerName, FIoContainerHeader&& ContainerHeader) override;
+	virtual FIoStatus	Mount(FString ContainerName, FSharedContainerHeader ContainerHeader) override;
 	virtual FIoStatus	Unmount(const FString& ContainerName) override;
+	virtual FIoStatus	UnmountAll() override;
 
 	virtual void		OnMounted(TSharedRef<const FPackageStoreBackendContext> Context) override;
 	virtual void 		BeginRead() override;
@@ -77,11 +79,22 @@ FOnDemandPackageStoreBackend::~FOnDemandPackageStoreBackend()
 {
 }
 
-FIoStatus FOnDemandPackageStoreBackend::Mount(FString ContainerName, FIoContainerHeader&& ContainerHeader)
+FIoStatus FOnDemandPackageStoreBackend::Mount(FString ContainerName, FSharedContainerHeader ContainerHeader)
 {
 	{
 		UE::TUniqueLock Lock(Mutex);
-		Containers.Add(MakeUnique<FContainer>(MoveTemp(ContainerName),MoveTemp(ContainerHeader)));
+
+		const FUniqueContainer* Existing =
+			Algo::FindByPredicate(
+				Containers,
+				[&ContainerName](const FUniqueContainer& C) { return C->Name == ContainerName; });
+
+		if (Existing != nullptr)
+		{
+			return EIoErrorCode::Ok;
+		}
+
+		Containers.Add(MakeUnique<FContainer>(MoveTemp(ContainerName), MoveTemp(ContainerHeader)));
 		bNeedsUpdate = true;
 	}
 
@@ -95,6 +108,22 @@ FIoStatus FOnDemandPackageStoreBackend::Mount(FString ContainerName, FIoContaine
 
 FIoStatus FOnDemandPackageStoreBackend::Unmount(const FString& ContainerName)
 {
+	UE::TUniqueLock Lock(Mutex);
+	const int32 NumRemoved = Containers.RemoveAll([&ContainerName](const TUniquePtr<FContainer>& C)
+	{
+		return C->Name == ContainerName;
+	});
+	check(NumRemoved <= 1);
+
+	bNeedsUpdate = NumRemoved > 0;
+	return NumRemoved == 0 ? EIoErrorCode::NotFound : EIoErrorCode::Ok;
+}
+
+FIoStatus FOnDemandPackageStoreBackend::UnmountAll()
+{
+	UE::TUniqueLock Lock(Mutex);
+	Containers.Empty();
+	bNeedsUpdate = true;
 	return EIoErrorCode::Ok;
 }
 
@@ -178,14 +207,14 @@ void FOnDemandPackageStoreBackend::UpdateLookupTables()
 
 	const int32 PackageCount = Algo::TransformAccumulate(
 		Containers,
-		[](const FUniqueContainer& C) { return C->ContainerHeader.PackageIds.Num(); },
+		[](const FUniqueContainer& C) { return C->Header->PackageIds.Num(); },
 		0);
 
 	EntryMap.Reserve(PackageCount);
 
 	for (const FUniqueContainer& Container : Containers)
 	{
-		const FIoContainerHeader& Hdr = Container->ContainerHeader;
+		const FIoContainerHeader& Hdr = *Container->Header;
 		TConstArrayView<FFilePackageStoreEntry> Entries(
 			reinterpret_cast<const FFilePackageStoreEntry*>(Hdr.StoreEntries.GetData()),
 			Hdr.PackageIds.Num());
diff --git a/Engine/Source/Runtime/Experimental/IoStore/OnDemand/Private/OnDemandPackageStoreBackend.h b/Engine/Source/Runtime/Experimental/IoStore/OnDemand/Private/OnDemandPackageStoreBackend.h
--- a/Engine/Source/Runtime/Experimental/IoStore/OnDemand/Private/OnDemandPackageStoreBackend.h
+++ b/Engine/Source/Runtime/Experimental/IoStore/OnDemand/Private/OnDemandPackageStoreBackend.h
@@ -7,6 +7,7 @@
 #include "Templates/SharedPointer.h"
 
 struct FIoContainerHeader;
+using FSharedContainerHeader = TSharedPtr<FIoContainerHeader>;
 
 namespace UE::IoStore
 {
@@ -17,8 +18,9 @@ class IOnDemandPackageStoreBackend
 public:
 	virtual ~IOnDemandPackageStoreBackend() = default;
 
-	virtual FIoStatus Mount(FString ContainerName, FIoContainerHeader&& ContainerHeader) = 0;
+	virtual FIoStatus Mount(FString ContainerName, FSharedContainerHeader ContainerHeader) = 0;
 	virtual FIoStatus Unmount(const FString& ContainerName) = 0;
+	virtual FIoStatus UnmountAll() = 0;
 };
 
 TSharedPtr<IOnDemandPackageStoreBackend> MakeOnDemandPackageStoreBackend();
