cluster union : Remove test code

diff --git a/Engine/Source/Runtime/Experimental/Chaos/Private/PhysicsProxy/ClusterUnionPhysicsProxy.cpp b/Engine/Source/Runtime/Experimental/Chaos/Private/PhysicsProxy/ClusterUnionPhysicsProxy.cpp
--- a/Engine/Source/Runtime/Experimental/Chaos/Private/PhysicsProxy/ClusterUnionPhysicsProxy.cpp
+++ b/Engine/Source/Runtime/Experimental/Chaos/Private/PhysicsProxy/ClusterUnionPhysicsProxy.cpp
@@ -853,8 +853,6 @@ namespace Chaos
 				TArray<FPBDRigidParticleHandle*> Particles = Interface.GetAllRigidParticles(Objects, /* bIncludeNulls */ true);
 				if (ensure(Particles.Num() == Objects.Num()))
 				{
-					Particles.Add(nullptr); // TEST DO NOT SUBMIT !!!!!!!!!!!!!!!!!!!!!!!!!!!!!!!!!!!!!!!!!!!!!!!!!!!!!!!!!!!!!!!!!!!!!!!!!!!!!!!!!!!!!!!!!!!!!!!!!!!!!!!!
-
 					FPBDRigidsEvolutionGBF& Evolution = *static_cast<FPBDRigidsSolver*>(Solver)->GetEvolution();
 					FClusterUnionManager& ClusterUnionManager = Evolution.GetRigidClustering().GetClusterUnionManager();
 					ClusterUnionManager.UpdateClusterUnionParticlesChildToParent(ClusterUnionIndex, Particles, Transforms, bLock);
