Chaos - fix GetConnectedParticles for disconnected joints
- fix GetConnectedParticles making connections through kinematics

diff --git a/Engine/Source/Runtime/Experimental/Chaos/Private/Chaos/PBDRigidsEvolutionGBF.cpp b/Engine/Source/Runtime/Experimental/Chaos/Private/Chaos/PBDRigidsEvolutionGBF.cpp
--- a/Engine/Source/Runtime/Experimental/Chaos/Private/Chaos/PBDRigidsEvolutionGBF.cpp
+++ b/Engine/Source/Runtime/Experimental/Chaos/Private/Chaos/PBDRigidsEvolutionGBF.cpp
@@ -1574,14 +1574,23 @@ TArray<FGeometryParticleHandle*> FPBDRigidsEvolutionGBF::GetConnectedParticles(F
 		{
 			if (FPBDJointConstraintHandle* Joint = Constraint->As<FPBDJointConstraintHandle>())
 			{
-				const TVec3<EJointMotionType>& JointLinearMotion = Joint->GetSettings().LinearMotionTypes;
-				if ((JointLinearMotion[0] == EJointMotionType::Locked) && (JointLinearMotion[1] == EJointMotionType::Locked) && (JointLinearMotion[2] == EJointMotionType::Locked))
+				if (Joint->IsConstraintEnabled() && !Joint->IsConstraintBroken())
 				{
-					FParticlePair JointParticles = Joint->GetConstrainedParticles();
-					FGeometryParticleHandle* OtherParticle = (JointParticles[0] != NextParticle) ? JointParticles[0] : JointParticles[1];
-					if ((OtherParticle != InParticle) && (ConnectedParticles.Find(OtherParticle) == nullptr) && !ExcludeConnections.Contains(OtherParticle->ParticleID()))
+					const TVec3<EJointMotionType>& JointLinearMotion = Joint->GetSettings().LinearMotionTypes;
+					if ((JointLinearMotion[0] == EJointMotionType::Locked) && (JointLinearMotion[1] == EJointMotionType::Locked) && (JointLinearMotion[2] == EJointMotionType::Locked))
 					{
-						ConnectedParticles.Add(OtherParticle, OtherParticle);
+						// Get the other particle on the joint
+						FParticlePair JointParticles = Joint->GetConstrainedParticles();
+						FGeometryParticleHandle* OtherParticle = (JointParticles[0] != NextParticle) ? JointParticles[0] : JointParticles[1];
+
+						// NOTE: We do not generate connections through kinematic particles
+						if ((OtherParticle != nullptr) && (OtherParticle != InParticle) && FGenericParticleHandle(OtherParticle)->IsDynamic())
+						{
+							if ((ConnectedParticles.Find(OtherParticle) == nullptr) && !ExcludeConnections.Contains(OtherParticle->ParticleID()))
+							{
+								ConnectedParticles.Add(OtherParticle, OtherParticle);
+							}
+						}
 					}
 				}
 			}
