Improve how we handle ::ResolveDistributedEndpoint returning false

#rnx

- ResolveDistributedEndpoint will attempt to resolve the list of CDNs we can use until successful. Each time the request fails we wait for a period of time and then submit the request again. If the current project has a fallback url set, then after a certain number of retries (most likely spread over minutes) we give up and return that fallback url. The only time the method can return false is if it is still trying to connect when the FOnDemandIoBackend is told to shutdown, most likely because the process itself is being shutdown.
- Demote the error to a warning, we should expect to see the resolve complete before shutdown but it is not required.
- Add code documentation explaining why the method returned false.

diff --git a/Engine/Source/Runtime/Experimental/IoStore/OnDemand/Private/OnDemandIoDispatcherBackend.cpp b/Engine/Source/Runtime/Experimental/IoStore/OnDemand/Private/OnDemandIoDispatcherBackend.cpp
--- a/Engine/Source/Runtime/Experimental/IoStore/OnDemand/Private/OnDemandIoDispatcherBackend.cpp
+++ b/Engine/Source/Runtime/Experimental/IoStore/OnDemand/Private/OnDemandIoDispatcherBackend.cpp
@@ -1761,8 +1761,13 @@ uint32 FOnDemandIoBackend::Run()
 	{
 		if (ResolveDistributedEndpoint(DistributionUrl) == false)
 		{
-			UE_LOG(LogIas, Error, TEXT("Failed to resolve CDN endpoints from distribution URL"));
+			// ResolveDistributedEndpoint should spin forever until either a valid url is found or
+			// we give up and use a predetermined fallback url. If this returned false then we didn't
+			// have a fallback url but the current process is shutting down so we might as well just
+			// exist the thread early.
+			UE_LOG(LogIas, Warning, TEXT("Failed to resolve CDN endpoints from distribution URL"));
 			BackendStatus.SetHttpEnabled(false);
+
 			return 0;
 		}
 	}
