Fix some minor bugs if the 'ias.HttpPrimaryEndpoint' cvar is used.

#rnx

- In a number of places we assume that the first url in the list is the primary url but we later added the cvar 'ias.HttpPrimaryEndpoint' which allows us to pick any url in the list.
- A lot of this code is due to be refactored in the upcoming months so there is not much point reworking how we track urls etc, so for now a quick and dirty fix is to simply rotate the array of urls until the primary url (which ever index that might be) is at the front, while maintaining the original ordering.
- It was also possible to set a negative index for the primary endpoint, in that case we now log a warning and default to index 0.
- Added a code comment to indicate that the cvar is not dynamic as we don't use it again after the system is initialized.

diff --git a/Engine/Source/Runtime/Experimental/IoStore/OnDemand/Private/OnDemandIoDispatcherBackend.cpp b/Engine/Source/Runtime/Experimental/IoStore/OnDemand/Private/OnDemandIoDispatcherBackend.cpp
--- a/Engine/Source/Runtime/Experimental/IoStore/OnDemand/Private/OnDemandIoDispatcherBackend.cpp
+++ b/Engine/Source/Runtime/Experimental/IoStore/OnDemand/Private/OnDemandIoDispatcherBackend.cpp
@@ -53,6 +53,7 @@ extern FString GIasOnDemandTocExt;
 void LatencyTest(FStringView, FStringView, uint32, TArrayView<int32>);
 
 ///////////////////////////////////////////////////////////////////////////////
+/** Note that GIasHttpPrimaryEndpoint has no effect after initial start up */
 int32 GIasHttpPrimaryEndpoint = 0;
 static FAutoConsoleVariableRef CVar_IasHttpPrimaryEndpoint(
 	TEXT("ias.HttpPrimaryEndpoint"),
@@ -1778,7 +1779,16 @@ uint32 FOnDemandIoBackend::Run()
 		BackendStatus.SetHttpEnabled(false);
 		return 0;
 	}
-	AvailableEps.Current = FMath::Min(GIasHttpPrimaryEndpoint, AvailableEps.Urls.Num() - 1);
+
+	if (GIasHttpPrimaryEndpoint < 0)
+	{
+		UE_LOG(LogIas, Error, TEXT("ias.HttpPrimaryEndpoint should not be set as a negative number, defaulting to 0"));
+		GIasHttpPrimaryEndpoint = 0;
+	}
+
+	// Rotate the list of urls so that the primary endpoint is the first element
+	Algo::Rotate(AvailableEps.Urls, GIasHttpPrimaryEndpoint);
+	AvailableEps.Current = 0;
 
 	BackendStatus.SetHttpEnabled(true);
 	FOnDemandIoBackendStats::Get()->OnHttpConnected();
