Fix 'IasHttpHasEverConnected' statistic so that it is properly recorded.

#rnx

- In a recent refactor we started assuming that once we had resolved the distributed endpoint we were "connected" to a CDN and were calling ::OnHttpConnected. However just because we connected the distributed endpoint service (or in some cases, failed to do so but switched to the fallback CDN url) doesn't mean we can connect to the CDN itself.
- Restored the older logic where we try a quick latency check on the CDN before we enter the main request loop and only signal a connection is the latency check succeeds. If it fails we mark the connection as starting in an error condition state.
- This brings the users connection/disconnection tracking back into sync and we can correctly determine if a user has never successfully contacted a CDN again.
- LatencyTest was changed to return success/fail state and to be present in all builds. We still only log the response times outside of shipping.
-- I was tempted to just remove this and go with the other overload of LatencyTest instead but there might still be some use in recording the initial latency test results so I left it for now. It is likely that this code will face another refactor in the near future anyway.

diff --git a/Engine/Source/Runtime/Experimental/IoStore/OnDemand/Private/OnDemandIoDispatcherBackend.cpp b/Engine/Source/Runtime/Experimental/IoStore/OnDemand/Private/OnDemandIoDispatcherBackend.cpp
--- a/Engine/Source/Runtime/Experimental/IoStore/OnDemand/Private/OnDemandIoDispatcherBackend.cpp
+++ b/Engine/Source/Runtime/Experimental/IoStore/OnDemand/Private/OnDemandIoDispatcherBackend.cpp
@@ -251,18 +251,29 @@ static FAutoConsoleCommand CVar_IasAbandonCache(
 );
 #endif //!UE_BUILD_SHIPPING
 ///////////////////////////////////////////////////////////////////////////////
-#if !UE_BUILD_SHIPPING
-static void LatencyTest(FStringView Url, FStringView Path)
+
+static bool LatencyTest(FStringView Url, FStringView Path)
 {
 	TRACE_CPUPROFILER_EVENT_SCOPE(IasBackend::LatencyTest);
 
 	int32 Results[4] = {};
 	LatencyTest(Url, Path, GIasHttpTimeOutMs, MakeArrayView(Results));
-	UE_LOG(LogIas, Log, TEXT("Endpoint '%s' latency test (ms): %d %d %d %d"),
-		Url.GetData(), Results[0], Results[1], Results[2], Results[3]);
-}
+
+	if (Results[0] >= 0 || Results[1] >= 0 || Results[2] >= 0 || Results[3] >= 0)
+	{
+#if !UE_BUILD_SHIPPING
+		UE_LOG(LogIas, Log, TEXT("Endpoint '%s' latency test (ms): %d %d %d %d"),
+			Url.GetData(), Results[0], Results[1], Results[2], Results[3]);
 #endif // !UE_BUILD_SHIPPING
 
+		return true;
+	}
+	else
+	{
+		return false;
+	}
+}
+
 ///////////////////////////////////////////////////////////////////////////////
 static int32 LatencyTest(TConstArrayView<FString> Urls, FStringView Path, std::atomic_bool& bCancel)
 {
@@ -272,7 +283,7 @@ static int32 LatencyTest(TConstArrayView<FString> Urls, FStringView Path, std::a
 	{
 		int32 LatencyMs = -1;
 		LatencyTest(Urls[Idx], Path, GIasHttpTimeOutMs, MakeArrayView(&LatencyMs, 1));
-		if (LatencyMs > 0)
+		if (LatencyMs >= 0)
 		{
 			return Idx;
 		}
@@ -1790,9 +1801,6 @@ uint32 FOnDemandIoBackend::Run()
 	Algo::Rotate(AvailableEps.Urls, GIasHttpPrimaryEndpoint);
 	AvailableEps.Current = 0;
 
-	BackendStatus.SetHttpEnabled(true);
-	FOnDemandIoBackendStats::Get()->OnHttpConnected();
-
 	FBitWindow HttpErrors;
 	HttpErrors.Reset(GIasHttpErrorSampleCount);
 
@@ -1805,14 +1813,24 @@ uint32 FOnDemandIoBackend::Run()
 		.ReceiveBufferSize = GIasHttpRecvBufKiB >= 0 ? GIasHttpRecvBufKiB << 10 : -1,
 		.bChangeEndpointAfterSuccessfulRetry = GIasHttpChangeEndpointAfterSuccessfulRetry,
 	});
+
 	check(HttpClient.IsValid());
-	HttpClient->SetEndpoint(AvailableEps.Current);
-#if !UE_BUILD_SHIPPING
-	if (AvailableEps.HasCurrent())
+	
+	if (LatencyTest(AvailableEps.GetCurrent(), GetEndpointTestPath()))
 	{
-		LatencyTest(AvailableEps.GetCurrent(), GetEndpointTestPath());
+		HttpClient->SetEndpoint(AvailableEps.Current);
+
+		BackendStatus.SetHttpEnabled(true);
+		FOnDemandIoBackendStats::Get()->OnHttpConnected();
+	}
+	else
+	{
+		BackendStatus.SetHttpError(true);
+
+		AvailableEps.Current = INDEX_NONE;
+		HttpClient->SetEndpoint(INDEX_NONE);
+		HttpErrors.Reset(GIasHttpErrorSampleCount);
 	}
-#endif 
 
 	while (!bStopRequested)
 	{
