ChaosFlesh: Set cached particles as kinematic. Skip cached particles in the evolution and coloring. Make cached simulation much faster

diff --git a/Engine/Plugins/Experimental/ChaosFlesh/Source/ChaosFleshEngine/Private/ChaosCache/FleshComponentCacheAdapter.cpp b/Engine/Plugins/Experimental/ChaosFlesh/Source/ChaosFleshEngine/Private/ChaosCache/FleshComponentCacheAdapter.cpp
--- a/Engine/Plugins/Experimental/ChaosFlesh/Source/ChaosFleshEngine/Private/ChaosCache/FleshComponentCacheAdapter.cpp
+++ b/Engine/Plugins/Experimental/ChaosFlesh/Source/ChaosFleshEngine/Private/ChaosCache/FleshComponentCacheAdapter.cpp
@@ -278,15 +278,32 @@ namespace Chaos
 				const int32 NumParticles = Particles.Size();
 				Softs::FSolverVec3* ParticleXs = &Particles.X(0);
 				Softs::FSolverVec3* ParticleVs = &Particles.V(0);
+				Softs::FSolverReal* ParticleInvMs = &Particles.InvM(0);
+				Softs::FPAndInvM* ParticlePAndInvMs = &Particles.PAndInvM(0);
 				
-				LoadCacheAtTime(InComponent, InTime, NumParticles, true, [&ParticleXs, &ParticleVs](
+				LoadCacheAtTime(InComponent, InTime, NumParticles, true, [&ParticleXs, &ParticleVs, &ParticleInvMs, &ParticlePAndInvMs](
 					const int32 ParticleIndex, const int32 ParticleOffset, const FVector3f& ParticlePosition, const FVector3f& ParticleVelocity)
 				{
 					const int32 GlobalIndex = ParticleIndex + ParticleOffset;
 					
 					ParticleXs[GlobalIndex].Set(ParticlePosition[0], ParticlePosition[1], ParticlePosition[2]);
 					ParticleVs[GlobalIndex].Set(ParticleVelocity[0], ParticleVelocity[1], ParticleVelocity[2]);
+					ParticleInvMs[GlobalIndex] = Softs::FSolverReal(0);
+					ParticlePAndInvMs[GlobalIndex].InvM = Softs::FSolverReal(0);
+					ParticlePAndInvMs[GlobalIndex].P = ParticleXs[GlobalIndex];
 				});
+
+				if(UFleshComponent* FleshComp = CastChecked<UFleshComponent>(InComponent) )
+				{
+					if (const UDeformablePhysicsComponent::FThreadingProxy* PhysicsProxy = FleshComp->GetPhysicsProxy())
+					{
+						if (const Chaos::Softs::FFleshThreadingProxy* Proxy = PhysicsProxy->As<Chaos::Softs::FFleshThreadingProxy>())
+						{
+							FIntVector2 ParticleRange = Proxy->GetSolverParticleRange();
+							Evolution->ActivateParticleRange(ParticleRange[0], false);
+						}
+					}
+				}
 				
 #else // USE_USD_SDK && DO_USD_CACHING
 				FCacheEvaluationContext Context(TickRecord);
diff --git a/Engine/Source/Runtime/Experimental/Chaos/Private/Chaos/Deformable/ChaosDeformableSolver.cpp b/Engine/Source/Runtime/Experimental/Chaos/Private/Chaos/Deformable/ChaosDeformableSolver.cpp
--- a/Engine/Source/Runtime/Experimental/Chaos/Private/Chaos/Deformable/ChaosDeformableSolver.cpp
+++ b/Engine/Source/Runtime/Experimental/Chaos/Private/Chaos/Deformable/ChaosDeformableSolver.cpp
@@ -342,6 +342,8 @@ namespace Chaos::Softs
 					Proxies.Add(Key, TUniquePtr<FThreadingProxy>(Proxy));
 				}
 
+				PrevEvolutionActiveRange = Evolution->ParticlesActiveView().GetActiveRanges();
+
 				if (UninitializedProxys_Internal.Num() != 0)
 				{
 					if (Property.bDoSelfCollision)
@@ -791,7 +793,6 @@ namespace Chaos::Softs
 								{
 									UE_LOG(LogChaosDeformableSolver, Log, TEXT("... Intersections : %d"), Intersections.Num());
 									//GSWeakConstraints->AddExtraConstraints(PositionTargetIndices, PositionTargetWeights, PositionTargetStiffness, PositionTargetSecondIndices, PositionTargetSecondWeights);
-
 								}
 							}
 						}
@@ -1585,7 +1586,7 @@ namespace Chaos::Softs
 			Evolution->ConstraintRules()[ConstraintIndex1] =
 				[this](FSolverParticles& InParticles, const FSolverReal Dt)
 			{
-				this->GSMainConstraint->Apply(InParticles, Dt);
+				this->GSMainConstraint->Apply(InParticles, Dt, 10, false, &(this->Evolution->ParticlesActiveView()));
 			};
 
 
@@ -1638,7 +1639,7 @@ namespace Chaos::Softs
 			Evolution->ConstraintRules()[ConstraintIndex1] =
 				[this](FSolverParticles& InParticles, const FSolverReal Dt)
 				{
-					this->GSMainConstraint->Apply(InParticles, Dt);
+					this->GSMainConstraint->Apply(InParticles, Dt, 10, false, &(this->Evolution->ParticlesActiveView()));
 				};
 
 
@@ -1674,7 +1675,7 @@ namespace Chaos::Softs
 			}
 		}
 
-		GSMainConstraint->InitStaticColor(Evolution->Particles());
+		GSMainConstraint->InitStaticColor(Evolution->Particles(), &(Evolution->ParticlesActiveView()));
 
 		if (Property.bEnablePositionTargets)
 		{
@@ -1733,6 +1734,7 @@ namespace Chaos::Softs
 
 		}
 	}
+
 	void FDeformableSolver::InitializeMuscleActivationVariables()
 	{
 		int32 InitIndex = Evolution->AddConstraintInitRange(1, true);
@@ -1864,6 +1866,33 @@ namespace Chaos::Softs
 
 		EventPreSolve.Broadcast(DeltaTime);
 
+		const TArray<TVector<int32, 2>, TInlineAllocator<8>>& EvolutionActiveRange = Evolution->ParticlesActiveView().GetActiveRanges();
+		bool bActiveRangeAreSame = true;
+		if (EvolutionActiveRange.Num() == PrevEvolutionActiveRange.Num())
+		{
+			for (int32 i = 0; i < PrevEvolutionActiveRange.Num(); i++)
+			{
+				if (PrevEvolutionActiveRange[i]!= EvolutionActiveRange[i])
+				{
+					bActiveRangeAreSame = false;
+					break;
+				}
+			}
+		}
+		else
+		{
+			bActiveRangeAreSame = false;
+		}
+
+		if (!bActiveRangeAreSame)
+		{
+			if (GSMainConstraint)
+			{
+				GSMainConstraint->InitStaticColor(Evolution->Particles(), &(Evolution->ParticlesActiveView()));
+			}
+		}
+
+
 		int32 NumIterations = FMath::Clamp<int32>(Property.NumSolverSubSteps, 0, INT_MAX);
 		if (bEnableSolver && NumIterations)
 		{
diff --git a/Engine/Source/Runtime/Experimental/Chaos/Private/Chaos/GraphColoring.cpp b/Engine/Source/Runtime/Experimental/Chaos/Private/Chaos/GraphColoring.cpp
--- a/Engine/Source/Runtime/Experimental/Chaos/Private/Chaos/GraphColoring.cpp
+++ b/Engine/Source/Runtime/Experimental/Chaos/Private/Chaos/GraphColoring.cpp
@@ -1594,8 +1594,8 @@ TArray<TArray<int32>> Chaos::ComputeNodalColoring(const TArray<TVec4<int32>>& Gr
 }
 
 
-template<typename T>
-TArray<TArray<int32>> Chaos::ComputeNodalColoring(const TArray<TArray<int32>>& Graph, const Chaos::TDynamicParticles<T, 3>& InParticles, const int32 GraphParticlesStart, const int32 GraphParticlesEnd, const TArray<TArray<int32>>& IncidentElements, const TArray<TArray<int32>>& IncidentElementsLocalIndex, TArray<int32>* ParticleColorsOut)
+template<typename T, typename ParticleType>
+TArray<TArray<int32>> Chaos::ComputeNodalColoring(const TArray<TArray<int32>>& Graph, const Chaos::TDynamicParticles<T, 3>& InParticles, const int32 GraphParticlesStart, const int32 GraphParticlesEnd, const TArray<TArray<int32>>& IncidentElements, const TArray<TArray<int32>>& IncidentElementsLocalIndex, const TPBDActiveView<ParticleType>* InParticleActiveView, TArray<int32>* ParticleColorsOut)
 {
 	using namespace Chaos;
 
@@ -1617,49 +1617,102 @@ TArray<TArray<int32>> Chaos::ComputeNodalColoring(const TArray<TArray<int32>>& G
 	TArray<int32> ParticleColors;
 	ParticleColors.Init(INDEX_NONE, GraphParticlesEnd - GraphParticlesStart);
 
-	for (int32 p = 0; p < GraphParticlesEnd - GraphParticlesStart; p++) 
+	if (InParticleActiveView)
 	{
-		int32 ParticleIndex = p + GraphParticlesStart;
-		if (InParticles.InvM(ParticleIndex) != (T)0. && Particle2Incident[ParticleIndex] != INDEX_NONE) 
-		{
-			for (int32 j = 0; j < IncidentElements[Particle2Incident[ParticleIndex]].Num(); j++) 
+		InParticleActiveView->SequentialFor([&ElementColorsSet, &Graph, &GraphParticlesStart, &IncidentElements, &IncidentElementsLocalIndex, &Particle2Incident, &ParticleColors]
+		(const Chaos::TDynamicParticles<T, 3>& InParticles, int32 ParticleIndex)
 			{
-				int32 LocalIndex = IncidentElementsLocalIndex[Particle2Incident[ParticleIndex]][j];
-				int32 e = IncidentElements[Particle2Incident[ParticleIndex]][j];
-				if (!ElementColorsSet[e]) 
+				const int32 p = ParticleIndex - GraphParticlesStart;
+				if (InParticles.InvM(ParticleIndex) != (T)0. && Particle2Incident[ParticleIndex] != INDEX_NONE) 
 				{
-					ElementColorsSet[e] = new TSet<int32>();
+					for (int32 j = 0; j < IncidentElements[Particle2Incident[ParticleIndex]].Num(); j++) 
+					{
+						int32 LocalIndex = IncidentElementsLocalIndex[Particle2Incident[ParticleIndex]][j];
+						int32 e = IncidentElements[Particle2Incident[ParticleIndex]][j];
+						if (!ElementColorsSet[e]) 
+						{
+							ElementColorsSet[e] = new TSet<int32>();
+						}
+					}
+					int32 color_to_use = 0;
+					while (true) 
+					{
+						bool ColorFound = false;
+						for (int32 j = 0; j < IncidentElements[Particle2Incident[ParticleIndex]].Num(); j++) 
+						{
+							int32 e = IncidentElements[Particle2Incident[ParticleIndex]][j];
+							if (ElementColorsSet[e]->Contains(color_to_use)) 
+							{
+								ColorFound = true;
+								break;
+							}
+						}
+						if (!ColorFound) 
+						{
+							ParticleColors[p] = color_to_use;
+							for (int32 j = 0; j < IncidentElements[Particle2Incident[ParticleIndex]].Num(); j++) 
+							{
+								int32 e = IncidentElements[Particle2Incident[ParticleIndex]][j];
+								ElementColorsSet[e]->Emplace(color_to_use);
+								if (ElementColorsSet[e]->Num() == Graph[e].Num())
+								{
+									delete ElementColorsSet[e];
+									ElementColorsSet[e] = nullptr;
+								}
+							}
+							break;
+						}
+						color_to_use++;
+					}
 				}
-			}
-			int32 color_to_use = 0;
-			while (true) 
+			});
+	}
+	else
+	{
+		for (int32 p = 0; p < GraphParticlesEnd - GraphParticlesStart; p++) 
+		{
+			int32 ParticleIndex = p + GraphParticlesStart;
+			if (InParticles.InvM(ParticleIndex) != (T)0. && Particle2Incident[ParticleIndex] != INDEX_NONE) 
 			{
-				bool ColorFound = false;
 				for (int32 j = 0; j < IncidentElements[Particle2Incident[ParticleIndex]].Num(); j++) 
 				{
+					int32 LocalIndex = IncidentElementsLocalIndex[Particle2Incident[ParticleIndex]][j];
 					int32 e = IncidentElements[Particle2Incident[ParticleIndex]][j];
-					if (ElementColorsSet[e]->Contains(color_to_use)) 
+					if (!ElementColorsSet[e]) 
 					{
-						ColorFound = true;
-						break;
+						ElementColorsSet[e] = new TSet<int32>();
 					}
 				}
-				if (!ColorFound) 
+				int32 color_to_use = 0;
+				while (true) 
 				{
-					ParticleColors[p] = color_to_use;
+					bool ColorFound = false;
 					for (int32 j = 0; j < IncidentElements[Particle2Incident[ParticleIndex]].Num(); j++) 
 					{
 						int32 e = IncidentElements[Particle2Incident[ParticleIndex]][j];
-						ElementColorsSet[e]->Emplace(color_to_use);
-						if (ElementColorsSet[e]->Num() == Graph[e].Num())
+						if (ElementColorsSet[e]->Contains(color_to_use)) 
 						{
-							delete ElementColorsSet[e];
-							ElementColorsSet[e] = nullptr;
+							ColorFound = true;
+							break;
 						}
 					}
-					break;
+					if (!ColorFound) 
+					{
+						ParticleColors[p] = color_to_use;
+						for (int32 j = 0; j < IncidentElements[Particle2Incident[ParticleIndex]].Num(); j++) 
+						{
+							int32 e = IncidentElements[Particle2Incident[ParticleIndex]][j];
+							ElementColorsSet[e]->Emplace(color_to_use);
+							if (ElementColorsSet[e]->Num() == Graph[e].Num())
+							{
+								delete ElementColorsSet[e];
+								ElementColorsSet[e] = nullptr;
+							}
+						}
+						break;
+					}
+					color_to_use++;
 				}
-				color_to_use++;
 			}
 		}
 	}
@@ -1808,7 +1861,7 @@ template CHAOS_API void Chaos::ComputeWeakConstraintsColoring<Chaos::FRealSingle
 template CHAOS_API void Chaos::ComputeWeakConstraintsColoring<Chaos::FRealDouble>(const TArray<TArray<int32>>& Indices, const TArray<TArray<int32>>& SecondIndices, const Chaos::TDynamicParticles<Chaos::FRealDouble, 3>& InParticles, TArray<TArray<int32>>& ConstraintsPerColor);
 template CHAOS_API TArray<TArray<int32>> Chaos::ComputeNodalColoring<Chaos::FRealSingle>(const TArray<TVec4<int32>>& Graph, const Chaos::TDynamicParticles<Chaos::FRealSingle, 3>& InParticles, const int32 GraphParticlesStart, const int32 GraphParticlesEnd, const TArray<TArray<int32>>& IncidentElements, const TArray<TArray<int32>>& IncidentElementsLocalIndex);
 template CHAOS_API TArray<TArray<int32>> Chaos::ComputeNodalColoring<Chaos::FRealDouble>(const TArray<TVec4<int32>>& Graph, const Chaos::TDynamicParticles<Chaos::FRealDouble, 3>& InParticles, const int32 GraphParticlesStart, const int32 GraphParticlesEnd, const TArray<TArray<int32>>& IncidentElements, const TArray<TArray<int32>>& IncidentElementsLocalIndex);
-template CHAOS_API TArray<TArray<int32>> Chaos::ComputeNodalColoring<Chaos::FRealSingle>(const TArray<TArray<int32>>& Graph, const Chaos::TDynamicParticles<Chaos::FRealSingle, 3>& InParticles, const int32 GraphParticlesStart, const int32 GraphParticlesEnd, const TArray<TArray<int32>>& IncidentElements, const TArray<TArray<int32>>& IncidentElementsLocalIndex, TArray<int32>* ParticleColorsOut);
-template CHAOS_API TArray<TArray<int32>> Chaos::ComputeNodalColoring<Chaos::FRealDouble>(const TArray<TArray<int32>>& Graph, const Chaos::TDynamicParticles<Chaos::FRealDouble, 3>& InParticles, const int32 GraphParticlesStart, const int32 GraphParticlesEnd, const TArray<TArray<int32>>& IncidentElements, const TArray<TArray<int32>>& IncidentElementsLocalIndex, TArray<int32>* ParticleColorsOut);
+template CHAOS_API TArray<TArray<int32>> Chaos::ComputeNodalColoring<Chaos::FRealSingle>(const TArray<TArray<int32>>& Graph, const Chaos::TDynamicParticles<Chaos::FRealSingle, 3>& InParticles, const int32 GraphParticlesStart, const int32 GraphParticlesEnd, const TArray<TArray<int32>>& IncidentElements, const TArray<TArray<int32>>& IncidentElementsLocalIndex, const TPBDActiveView<Chaos::Softs::FSolverParticles>* InParticleActiveView, TArray<int32>* ParticleColorsOut);
+template CHAOS_API TArray<TArray<int32>> Chaos::ComputeNodalColoring<Chaos::FRealDouble>(const TArray<TArray<int32>>& Graph, const Chaos::TDynamicParticles<Chaos::FRealDouble, 3>& InParticles, const int32 GraphParticlesStart, const int32 GraphParticlesEnd, const TArray<TArray<int32>>& IncidentElements, const TArray<TArray<int32>>& IncidentElementsLocalIndex, const TPBDActiveView<Chaos::TDynamicParticles<Chaos::FRealDouble, 3>>* InParticleActiveView, TArray<int32>* ParticleColorsOut);
 template CHAOS_API void Chaos::ComputeExtraNodalColoring<Chaos::FRealSingle>(const TArray<TArray<int32>>& Graph, const TArray<TArray<int32>>& ExtraGraph, const Chaos::TDynamicParticles<Chaos::FRealSingle, 3>& InParticles, const TArray<TArray<int32>>& IncidentElements, const TArray<TArray<int32>>& ExtraIncidentElements, TArray<int32>& ParticleColors, TArray<TArray<int32>>& ParticlesPerColor);
 template CHAOS_API void Chaos::ComputeExtraNodalColoring<Chaos::FRealDouble>(const TArray<TArray<int32>>& Graph, const TArray<TArray<int32>>& ExtraGraph, const Chaos::TDynamicParticles<Chaos::FRealDouble, 3>& InParticles, const TArray<TArray<int32>>& IncidentElements, const TArray<TArray<int32>>& ExtraIncidentElements, TArray<int32>& ParticleColors, TArray<TArray<int32>>& ParticlesPerColor);
\ No newline at end of file
diff --git a/Engine/Source/Runtime/Experimental/Chaos/Public/Chaos/Deformable/ChaosDeformableSolver.h b/Engine/Source/Runtime/Experimental/Chaos/Public/Chaos/Deformable/ChaosDeformableSolver.h
--- a/Engine/Source/Runtime/Experimental/Chaos/Public/Chaos/Deformable/ChaosDeformableSolver.h
+++ b/Engine/Source/Runtime/Experimental/Chaos/Public/Chaos/Deformable/ChaosDeformableSolver.h
@@ -231,6 +231,7 @@ namespace Chaos::Softs
 		int32 Iteration = 0;
 		bool bSimulationInitialized = false;
 		int32 GroupOffset = 1;
+		TArray<TVector<int32, 2>, TInlineAllocator<8>> PrevEvolutionActiveRange;
 	};
 
 
diff --git a/Engine/Source/Runtime/Experimental/Chaos/Public/Chaos/Deformable/GaussSeidelMainConstraint.h b/Engine/Source/Runtime/Experimental/Chaos/Public/Chaos/Deformable/GaussSeidelMainConstraint.h
--- a/Engine/Source/Runtime/Experimental/Chaos/Public/Chaos/Deformable/GaussSeidelMainConstraint.h
+++ b/Engine/Source/Runtime/Experimental/Chaos/Public/Chaos/Deformable/GaussSeidelMainConstraint.h
@@ -117,7 +117,7 @@ namespace Chaos::Softs
 
 		CHAOS_API void AddDynamicConstraints(const TArray<TArray<int32>>& ExtraConstraints, TArray<TArray<int32>>& ExtraIncidentElements, TArray<TArray<int32>>& ExtraIncidentElementsLocal, bool CheckIncidentElements = false);
 
-		void Apply(ParticleType& Particles, const T Dt, const int32 MaxWriteIters = 10, const bool Write2File = false)
+		void Apply(ParticleType& Particles, const T Dt, const int32 MaxWriteIters = 10, const bool Write2File = false, const TPBDActiveView<FSolverParticles>* InParticleActiveView = nullptr)
 		{
  			PERF_SCOPE(STAT_ChaosGSMainConstraint_Apply);
 			
@@ -141,8 +141,11 @@ namespace Chaos::Softs
 							int32 TaskIndex = CorotatedParams.XPBDCorotatedBatchSize * BatchIndex + BatchSubIndex;
 							if (TaskIndex < ParticlesPerColor[i].Num())
 							{
-								int32 ParticleIndex = ParticlesPerColor[i][TaskIndex];
-								ApplySingleParticle(ParticleIndex, Dt, Particles);
+								if (Particles.InvM(ParticlesPerColor[i][TaskIndex]) != T(0))
+								{
+									int32 ParticleIndex = ParticlesPerColor[i][TaskIndex];
+									ApplySingleParticle(ParticleIndex, Dt, Particles);
+								}
 							}
 						}
 					}, NumBatch < CorotatedParams.XPBDCorotatedBatchThreshold);
@@ -152,15 +155,26 @@ namespace Chaos::Softs
 						
 			if (bDoAcceleration)
 			{
-				AccelerationTechnique(Particles);
+				PERF_SCOPE(STAT_ChaosGSMainConstraint_Acceleration);
+				if (InParticleActiveView)
+				{
+					InParticleActiveView->ParallelFor(AccelerationTechniquePerParticle, CorotatedParams.XPBDCorotatedBatchSize);
+				}
+				else
+				{
+					PhysicsParallelFor(Particles.Size(), [this, &Particles](const int32 ParticleIndex)
+						{
+							this->AccelerationTechniquePerParticle(Particles, ParticleIndex);
+						}, Particles.Size() < 1000);
+				}
 			}
 			
 			CurrentIt ++;
 		}
 
-		void InitStaticColor(const ParticleType& Particles)
+		void InitStaticColor(const ParticleType& Particles, const TPBDActiveView<FSolverParticles>* InParticleActiveView = nullptr)
 		{
-			StaticParticlesPerColor = ComputeNodalColoring(StaticConstraints, Particles, 0, Particles.Size(), StaticIncidentElements, StaticIncidentElementsLocal, &StaticParticleColors);
+			StaticParticlesPerColor = ComputeNodalColoring(StaticConstraints, Particles, 0, Particles.Size(), StaticIncidentElements, StaticIncidentElementsLocal, InParticleActiveView,  &StaticParticleColors);
 			ParticleColors = StaticParticleColors;
 			ParticlesPerColor = StaticParticlesPerColor;
 		}
@@ -304,27 +318,26 @@ namespace Chaos::Softs
 				}
 			};
 
-			AccelerationTechnique = [this](ParticleType& Particles)
+			AccelerationTechniquePerParticle = [this](ParticleType& Particles, int32 ParticleIndex)
 			{
-				PERF_SCOPE(STAT_ChaosGSMainConstraint_Acceleration);
-				PhysicsParallelFor(Particles.Size(), [&](const int32 ParticleIndex)
-					{
-						if (Particles.InvM(ParticleIndex) != T(0) && CurrentIt > SORStart)
-						{
-							Particles.P(ParticleIndex) = OmegaSOR * (Particles.P(ParticleIndex) - this->X_k_1[ParticleIndex]) + this->X_k_1[ParticleIndex];
-						}
-						this->X_k_1[ParticleIndex] = this->X_k[ParticleIndex];
-						this->X_k[ParticleIndex] = Particles.P(ParticleIndex);
-					}, Particles.Size() < 1000);
+				if (Particles.InvM(ParticleIndex) != T(0) && CurrentIt > SORStart)
+				{
+					Particles.P(ParticleIndex) = OmegaSOR * (Particles.P(ParticleIndex) - this->X_k_1[ParticleIndex]) + this->X_k_1[ParticleIndex];
+				}
+				if (Particles.InvM(ParticleIndex) != T(0))
+				{
+					this->X_k_1[ParticleIndex] = this->X_k[ParticleIndex];
+					this->X_k[ParticleIndex] = Particles.P(ParticleIndex);
+				}
 			};
 		}
 
 
 		//Constraints storage:
-		TArray<TArray<int32>> StaticConstraints;
+		TArray<TArray<int32>> StaticConstraints = {};
 		TArray<TArray<int32>> StaticIncidentElements;
 		TArray<TArray<int32>> StaticIncidentElementsLocal;
-		TArray<TArray<int32>> DynamicConstraints;
+		TArray<TArray<int32>> DynamicConstraints = {};
 		TArray<TArray<int32>> DynamicIncidentElements;
 		TArray<TArray<int32>> DynamicIncidentElementsLocal;
 
@@ -348,7 +361,7 @@ namespace Chaos::Softs
 		mutable TArray<Chaos::TVector<T, 3>> xtilde;
 
 		//SOR variables:
-		TFunction<void(ParticleType&)> AccelerationTechnique;
+		TFunction<void(ParticleType&, int32)> AccelerationTechniquePerParticle;
 		mutable TArray<Chaos::TVector<T, 3>> X_k_1;
 		mutable TArray<Chaos::TVector<T, 3>> X_k;
 		int32 CurrentIt = 0;
diff --git a/Engine/Source/Runtime/Experimental/Chaos/Public/Chaos/GraphColoring.h b/Engine/Source/Runtime/Experimental/Chaos/Public/Chaos/GraphColoring.h
--- a/Engine/Source/Runtime/Experimental/Chaos/Public/Chaos/GraphColoring.h
+++ b/Engine/Source/Runtime/Experimental/Chaos/Public/Chaos/GraphColoring.h
@@ -6,6 +6,7 @@
 #include "Chaos/Core.h"
 #include "Chaos/Vector.h"
 #include "Chaos/UniformGrid.h"
+#include "Chaos/PBDActiveView.h"
 
 namespace Chaos
 {
@@ -122,8 +123,8 @@ CHAOS_API void ComputeWeakConstraintsColoring(const TArray<TArray<int32>>& Indic
 template<typename T>
 CHAOS_API TArray<TArray<int32>> ComputeNodalColoring(const TArray<TVec4<int32>>& Graph, const Chaos::TDynamicParticles<T, 3>& InParticles, const int32 GraphParticlesStart, const int32 GraphParticlesEnd, const TArray<TArray<int32>>& IncidentElements, const TArray<TArray<int32>>& IncidentElementsLocalIndex);
 
-template<typename T>
-CHAOS_API TArray<TArray<int32>> ComputeNodalColoring(const TArray<TArray<int32>>& Graph, const Chaos::TDynamicParticles<T, 3>& InParticles, const int32 GraphParticlesStart, const int32 GraphParticlesEnd, const TArray<TArray<int32>>& IncidentElements, const TArray<TArray<int32>>& IncidentElementsLocalIndex, TArray<int32>* ParticleColorsOut = nullptr);
+template<typename T, typename ParticleType>
+CHAOS_API TArray<TArray<int32>> ComputeNodalColoring(const TArray<TArray<int32>>& Graph, const Chaos::TDynamicParticles<T, 3>& InParticles, const int32 GraphParticlesStart, const int32 GraphParticlesEnd, const TArray<TArray<int32>>& IncidentElements, const TArray<TArray<int32>>& IncidentElementsLocalIndex, const Chaos::TPBDActiveView<ParticleType>* InParticleActiveView = nullptr, TArray<int32>* ParticleColorsOut = nullptr);
 
 template<typename T>
 CHAOS_API void ComputeExtraNodalColoring(const TArray<TArray<int32>>& Graph, const TArray<TArray<int32>>& ExtraGraph, const Chaos::TDynamicParticles<T, 3>& InParticles, const TArray<TArray<int32>>& IncidentElements, const TArray<TArray<int32>>& ExtraIncidentElements, TArray<int32>& ParticleColors, TArray<TArray<int32>>& ParticlesPerColor);
