Clean up old UE_DEPRECATED code from Chaos::Softs.

diff --git a/Engine/Source/Runtime/Experimental/Chaos/Private/Chaos/CollectionPropertyFacade.cpp b/Engine/Source/Runtime/Experimental/Chaos/Private/Chaos/CollectionPropertyFacade.cpp
--- a/Engine/Source/Runtime/Experimental/Chaos/Private/Chaos/CollectionPropertyFacade.cpp
+++ b/Engine/Source/Runtime/Experimental/Chaos/Private/Chaos/CollectionPropertyFacade.cpp
@@ -160,10 +160,7 @@ namespace Chaos::Softs
 				{
 					continue;
 				}
-				PRAGMA_DISABLE_DEPRECATION_WARNINGS
-				// TODO: GetFlags needs to return an ECollectionPropertyFlags, not an uint8, but the uint8 getter needs to be deprecated first
-				SetFlags(PropertyIndex, (ECollectionPropertyFlags)InPropertyFacade.GetFlags(InKeyIndex));
-				PRAGMA_ENABLE_DEPRECATION_WARNINGS
+				SetFlags(PropertyIndex, InPropertyFacade.GetFlags(InKeyIndex));
 				// Setting as FVector3f since that is the underlying type
 				SetLowValue(PropertyIndex, InPropertyFacade.GetLowValue<FVector3f>(InKeyIndex));
 				SetHighValue(PropertyIndex, InPropertyFacade.GetHighValue<FVector3f>(InKeyIndex));
@@ -254,17 +251,10 @@ namespace Chaos::Softs
 	}
 
 	void FCollectionPropertyMutableFacade::Append(const TSharedPtr<const FManagedArrayCollection>& InManagedArrayCollection, bool bUpdateExistingProperties)
-	{
-		
+	{		
 		Update(InManagedArrayCollection, ECollectionPropertyUpdateFlags::AppendNewProperties | (bUpdateExistingProperties ? ECollectionPropertyUpdateFlags::UpdateExistingProperties : ECollectionPropertyUpdateFlags::None));
 	}
 
-	void FCollectionPropertyMutableFacade::Append(const FManagedArrayCollection& InManagedArrayCollection)
-	{
-		constexpr ECollectionPropertyUpdateFlags UpdateFlagsAppendNewOnly = ECollectionPropertyUpdateFlags::AppendNewProperties;
-		Update(MakeShared<const FManagedArrayCollection>(InManagedArrayCollection), UpdateFlagsAppendNewOnly);
-	}
-
 	void FCollectionPropertyMutableFacade::Copy(const FManagedArrayCollection& InManagedArrayCollection)
 	{
 		GetManagedArrayCollection()->Reset();
@@ -308,10 +298,7 @@ namespace Chaos::Softs
 				for (int32 InKeyIndex = 0; InKeyIndex < NumInKeys; ++InKeyIndex)
 				{
 					const FString& PropertyName = InPropertyFacade.GetKey(InKeyIndex);
-PRAGMA_DISABLE_DEPRECATION_WARNINGS
-					// TODO: GetFlags needs to return an ECollectionPropertyFlags, not an uint8, but the uint8 getter needs to be deprecated first
-					const ECollectionPropertyFlags PropertyFlags = (ECollectionPropertyFlags)InPropertyFacade.GetFlags(InKeyIndex);
-PRAGMA_ENABLE_DEPRECATION_WARNINGS
+					const ECollectionPropertyFlags PropertyFlags = InPropertyFacade.GetFlags(InKeyIndex);
 					int32 NewPropertyIndex = GetKeyIndex(PropertyName);
 					if (NewPropertyIndex == INDEX_NONE)
 					{
diff --git a/Engine/Source/Runtime/Experimental/Chaos/Private/Chaos/TriangleMesh.cpp b/Engine/Source/Runtime/Experimental/Chaos/Private/Chaos/TriangleMesh.cpp
--- a/Engine/Source/Runtime/Experimental/Chaos/Private/Chaos/TriangleMesh.cpp
+++ b/Engine/Source/Runtime/Experimental/Chaos/Private/Chaos/TriangleMesh.cpp
@@ -908,7 +908,7 @@ const TArray<TVec2<int32>>& FTriangleMesh::GetEdgeToFaces() const
 
 TSet<int32> FTriangleMesh::GetBoundaryPoints()
 {
-	const FSegmentMesh& SegmentMesh = const_cast<const FTriangleMesh*>(this)->GetSegmentMesh();
+	const FSegmentMesh& SegmentMesh = GetSegmentMesh();
 	const TArray<TVec2<int32>>& Edges = SegmentMesh.GetElements();
 	const TArray<TVec2<int32>>& EdgeToFaces = GetEdgeToFaces();
 	TSet<int32> OpenBoundaryPoints;
@@ -1047,7 +1047,7 @@ TArray<FReal> FTriangleMesh::GetCurvatureOnEdges(const TArray<FVec3>& FaceNormal
 {
 	const int32 NumNormals = FaceNormals.Num();
 	check(NumNormals == MElements.Num());
-	const FSegmentMesh& SegmentMesh = const_cast<const FTriangleMesh*>(this)->GetSegmentMesh(); // builds MEdgeToFaces
+	const FSegmentMesh& SegmentMesh = GetSegmentMesh(); // builds MEdgeToFaces
 	TArray<FReal> EdgeAngles;
 	EdgeAngles.SetNumZeroed(MEdgeToFaces.Num());
 	for (int32 EdgeId = 0; EdgeId < MEdgeToFaces.Num(); EdgeId++)
@@ -1074,7 +1074,7 @@ TArray<FReal> FTriangleMesh::GetCurvatureOnEdges(const TConstArrayView<FVec3>& P
 
 TArray<FReal> FTriangleMesh::GetCurvatureOnPoints(const TArray<FReal>& EdgeCurvatures)
 {
-	const FSegmentMesh& SegmentMesh = const_cast<const FTriangleMesh*>(this)->GetSegmentMesh();
+	const FSegmentMesh& SegmentMesh = GetSegmentMesh();
 	const TArray<TVec2<int32>>& Segments = SegmentMesh.GetElements();
 	check(EdgeCurvatures.Num() == Segments.Num());
 
diff --git a/Engine/Source/Runtime/Experimental/Chaos/Public/Chaos/CollectionPropertyFacade.h b/Engine/Source/Runtime/Experimental/Chaos/Public/Chaos/CollectionPropertyFacade.h
--- a/Engine/Source/Runtime/Experimental/Chaos/Public/Chaos/CollectionPropertyFacade.h
+++ b/Engine/Source/Runtime/Experimental/Chaos/Public/Chaos/CollectionPropertyFacade.h
@@ -10,17 +10,6 @@ struct FManagedArrayCollection;
 
 namespace Chaos::Softs
 {
-	enum class UE_DEPRECATED(5.3, "Use ECollectionPropertyFlags instead.") ECollectionPropertyFlag : uint8
-	{
-		None,
-		Enabled = 1 << 0,
-		Animatable = 1 << 1,
-		Dirty = 1 << 7
-	};
-	PRAGMA_DISABLE_DEPRECATION_WARNINGS
-	ENUM_CLASS_FLAGS(ECollectionPropertyFlag)
-	PRAGMA_ENABLE_DEPRECATION_WARNINGS
-
 	/** Property flags, whether properties are enabled, animatable, ...etc. */
 	enum class ECollectionPropertyFlags : uint8
 	{
@@ -89,11 +78,8 @@ namespace Chaos::Softs
 		T GetValue(int32 KeyIndex) const { return GetLowValue<T>(KeyIndex); }
 
 		const FString& GetStringValue(int32 KeyIndex) const { return GetValue<const FString&>(KeyIndex, StringValueArray); }
-		UE_DEPRECATED(5.3, "Use GetStringValue(int32) or GetStringValue(const FString&, const FString&, int32*) instead.")
-		const FString& GetStringValue(int32 KeyIndex, const FString& Default) const { return GetValue<const FString&>(KeyIndex, StringValueArray); }
 
-		UE_DEPRECATED(5.3, "uint8 GetFlags(int32) is deprecated and will soon be replaced by ECollectionPropertyFlags GetFlags(int32).")
-		uint8 GetFlags(int32 KeyIndex) const { return FlagsArray[KeyIndex]; }
+		ECollectionPropertyFlags GetFlags(int32 KeyIndex) const { return (ECollectionPropertyFlags)FlagsArray[KeyIndex]; }
 
 		bool IsEnabled(int32 KeyIndex) const { return HasAnyFlags(KeyIndex, ECollectionPropertyFlags::Enabled); }
 		bool IsAnimatable(int32 KeyIndex) const { return HasAnyFlags(KeyIndex, ECollectionPropertyFlags::Animatable) && !HasAnyFlags(KeyIndex, ECollectionPropertyFlags::Intrinsic); }
@@ -143,12 +129,9 @@ namespace Chaos::Softs
 			return SafeGet(Key, [this](int32 KeyIndex)->FString { return GetStringValue(KeyIndex); }, Default, OutKeyIndex);
 		}
 
-		UE_DEPRECATED(5.3, "uint8 GetFlags(const FString&, uint8, int32*) is deprecated and will soon be replaced by ECollectionPropertyFlags GetFlags(const FString&, uint8, int32*).")
-		uint8 GetFlags(const FString& Key, uint8 Default = 0, int32* OutKeyIndex = nullptr) const
+		ECollectionPropertyFlags GetFlags(const FString& Key, uint8 Default = 0, int32* OutKeyIndex = nullptr) const
 		{
-PRAGMA_DISABLE_DEPRECATION_WARNINGS
-			return SafeGet(Key, [this](int32 KeyIndex)->uint8 { return GetFlags(KeyIndex); }, Default, OutKeyIndex);
-PRAGMA_ENABLE_DEPRECATION_WARNINGS
+			return SafeGet(Key, [this](int32 KeyIndex)->ECollectionPropertyFlags { return GetFlags(KeyIndex); }, (ECollectionPropertyFlags)Default, OutKeyIndex);
 		}
 
 		bool IsEnabled(const FString& Key, bool bDefault = false, int32* OutKeyIndex = nullptr) const
@@ -290,8 +273,6 @@ PRAGMA_ENABLE_DEPRECATION_WARNINGS
 
 		/** SetFlags cannot be used to remove Dirty, StringDirty, Interpolable or Intrinsic flags. Use ClearDirtyFlags to remove dirty flags. */
 		CHAOS_API void SetFlags(int32 KeyIndex, ECollectionPropertyFlags Flags);
-		UE_DEPRECATED(5.3, "Use SetFlags(int32, ECollectionPropertyFlags) instead.")
-		void SetFlags(int32 KeyIndex, uint8 Flags) { return SetFlags(KeyIndex, (ECollectionPropertyFlags)Flags); }
 
 		void SetEnabled(int32 KeyIndex, bool bEnabled) { EnableFlags(KeyIndex, ECollectionPropertyFlags::Enabled, bEnabled); }
 		void SetAnimatable(int32 KeyIndex, bool bAnimatable) { EnableFlags(KeyIndex, ECollectionPropertyFlags::Animatable, bAnimatable); }
@@ -300,8 +281,6 @@ PRAGMA_ENABLE_DEPRECATION_WARNINGS
 		/** Set the intrinsic flag for this property. This flag cannot be removed and implies non Animatable. */
 		void SetIntrinsic(int32 KeyIndex) { EnableFlags(KeyIndex, ECollectionPropertyFlags::Intrinsic, true); }
 		void SetDirty(int32 KeyIndex) { EnableFlags(KeyIndex, ECollectionPropertyFlags::Dirty, true); }
-		UE_DEPRECATED(5.3, "SetDirty can only be set, to unset use ClearDirtyFlags instead.")
-		void SetDirty(int32 KeyIndex, bool bDirty) { EnableFlags(KeyIndex, ECollectionPropertyFlags::Dirty, bDirty); }
 		void SetStringDirty(int32 KeyIndex) { EnableFlags(KeyIndex, ECollectionPropertyFlags::StringDirty, true); }
 		void SetInterpolable(int32 KeyIndex) { EnableFlags(KeyIndex, ECollectionPropertyFlags::Interpolable, true); }
 
@@ -344,8 +323,6 @@ PRAGMA_ENABLE_DEPRECATION_WARNINGS
 		{
 			return SafeSet(Key, [this, Flags](int32 KeyIndex) { SetFlags(KeyIndex, Flags); });
 		}
-		UE_DEPRECATED(5.3, "Use SetFlags(const FString&, ECollectionPropertyFlags) instead.")
-		int32 SetFlags(const FString& Key, uint8 Flags) { return SetFlags(Key, (ECollectionPropertyFlags)Flags); }
 
 		int32 SetEnabled(const FString& Key, bool bEnabled)
 		{
@@ -368,14 +345,6 @@ PRAGMA_ENABLE_DEPRECATION_WARNINGS
 			return SafeSet(Key, [this](int32 KeyIndex) { SetIntrinsic(KeyIndex); });
 		}
 
-		UE_DEPRECATED(5.3, "SetDirty can only be set, to unset use ClearDirtyFlags instead.")
-		int32 SetDirty(const FString& Key, bool bDirty)
-		{
-PRAGMA_DISABLE_DEPRECATION_WARNINGS
-			return SafeSet(Key, [this, bDirty](int32 KeyIndex) { SetDirty(KeyIndex, bDirty); });
-PRAGMA_ENABLE_DEPRECATION_WARNINGS
-		}
-
 		int32 SetDirty(const FString& Key)
 		{
 			return SafeSet(Key, [this](int32 KeyIndex) { SetDirty(KeyIndex); });
@@ -479,9 +448,6 @@ PRAGMA_ENABLE_DEPRECATION_WARNINGS
 		 */
 		CHAOS_API void Append(const TSharedPtr<const FManagedArrayCollection>& InManagedArrayCollection, bool bUpdateExistingProperties);
 
-		UE_DEPRECATED(5.3, "Use SharedPtr version of Append to avoid additional copy.")
-		CHAOS_API void Append(const FManagedArrayCollection& InManagedArrayCollection);
-
 		/**
 		 * Copy all properties and values from an existing collection to this property collection.
 		 * Dirty flags will be copied directly.
@@ -545,8 +511,6 @@ PRAGMA_ENABLE_DEPRECATION_WARNINGS
 // Use this macro to add shorthands for property getters without a key index
 #define UE_CHAOS_DECLARE_INDEXLESS_PROPERTYCOLLECTION_NAME(PropertyName, Type) \
 	inline static const FName PropertyName##Name = TEXT(#PropertyName); \
-	UE_DEPRECATED(5.3, "PropertyName##String is to be removed as to not be confused with GetPropertyName##String().") \
-	static FString PropertyName##String() { return PropertyName##Name.ToString(); } \
 	static bool Is##PropertyName##Enabled(const FCollectionPropertyConstFacade& InPropertyCollection, bool bDefault) \
 	{ \
 		return InPropertyCollection.IsEnabled(PropertyName##Name.ToString(), bDefault); \
@@ -583,13 +547,6 @@ PRAGMA_ENABLE_DEPRECATION_WARNINGS
 	{ \
 		return InPropertyCollection.GetStringValue(PropertyName##Name.ToString(), Default); \
 	} \
-	UE_DEPRECATED(5.3, "GetFlags is being phased out to promote correct dirtying operations.") \
-	uint8 Get##PropertyName##Flags(const FCollectionPropertyConstFacade& InPropertyCollection, uint8 Default) \
-	{ \
-PRAGMA_DISABLE_DEPRECATION_WARNINGS \
-		return InPropertyCollection.GetFlags(PropertyName##Name.ToString(), Default); \
-PRAGMA_ENABLE_DEPRECATION_WARNINGS \
-	}
 
 // Use this macro to add shorthands for property getters and direct access through the declared key index
 #define UE_CHAOS_DECLARE_PROPERTYCOLLECTION_NAME(PropertyName, Type) \
@@ -630,15 +587,6 @@ PRAGMA_ENABLE_DEPRECATION_WARNINGS \
 		checkf(PropertyName##Index != INDEX_NONE, TEXT("The default value getter that sets the property index must be called once prior to calling this function.")); \
 		return PropertyCollection.GetStringValue(PropertyName##Index); \
 	} \
-	UE_DEPRECATED(5.3, "GetFlags is being phased out to promote correct dirtying operations.") \
-	uint8 Get##PropertyName##Flags(const FCollectionPropertyConstFacade& PropertyCollection) const \
-	{ \
-		checkSlow(PropertyName##Index == PropertyCollection.GetKeyIndex(PropertyName##Name.ToString())); \
-		checkf(PropertyName##Index != INDEX_NONE, TEXT("The default value getter that sets the property index must be called once prior to calling this function.")); \
-PRAGMA_DISABLE_DEPRECATION_WARNINGS \
-		return PropertyCollection.GetFlags(PropertyName##Index); \
-PRAGMA_ENABLE_DEPRECATION_WARNINGS \
-	} \
 	bool Is##PropertyName##Enabled(const FCollectionPropertyConstFacade& PropertyCollection) const \
 	{ \
 		checkSlow(PropertyName##Index == PropertyCollection.GetKeyIndex(PropertyName##Name.ToString())); \
diff --git a/Engine/Source/Runtime/Experimental/Chaos/Public/Chaos/NewtonEvolution.h b/Engine/Source/Runtime/Experimental/Chaos/Public/Chaos/NewtonEvolution.h
--- a/Engine/Source/Runtime/Experimental/Chaos/Public/Chaos/NewtonEvolution.h
+++ b/Engine/Source/Runtime/Experimental/Chaos/Public/Chaos/NewtonEvolution.h
@@ -138,16 +138,6 @@ public:
 	FVelocityAndPressureField& GetVelocityAndPressureField(const uint32 GroupId = 0) { check(GroupId < TArrayCollection::Size()); return MGroupVelocityAndPressureFields[GroupId]; }
 	const FVelocityAndPressureField& GetVelocityAndPressureField(const uint32 GroupId = 0) const { check(GroupId < TArrayCollection::Size()); return MGroupVelocityAndPressureFields[GroupId]; }
 
-	UE_DEPRECATED(5.1, "Chaos::Softs::FVelocityField has been renamed FVelocityAndPressureField to match its new behavior.")
-		PRAGMA_DISABLE_DEPRECATION_WARNINGS
-		FVelocityField& GetVelocityField(const uint32 GroupId = 0) { return GetVelocityAndPressureField(GroupId); }
-	PRAGMA_ENABLE_DEPRECATION_WARNINGS
-		UE_DEPRECATED(5.1, "Chaos::Softs::FVelocityField has been renamed FVelocityAndPressureField to match its new behavior.")
-		PRAGMA_DISABLE_DEPRECATION_WARNINGS
-		const FVelocityField& GetVelocityField(const uint32 GroupId = 0) const { return GetVelocityAndPressureField(GroupId); }
-	PRAGMA_ENABLE_DEPRECATION_WARNINGS
-
-
 	int32 GetNewtonIterations() const { return MNumNewtonIterations; }
 	void SetNewtonIterations(const int32 Iterations) { MNumNewtonIterations = Iterations; }
 
diff --git a/Engine/Source/Runtime/Experimental/Chaos/Public/Chaos/PBDAnimDriveConstraint.h b/Engine/Source/Runtime/Experimental/Chaos/Public/Chaos/PBDAnimDriveConstraint.h
--- a/Engine/Source/Runtime/Experimental/Chaos/Public/Chaos/PBDAnimDriveConstraint.h
+++ b/Engine/Source/Runtime/Experimental/Chaos/Public/Chaos/PBDAnimDriveConstraint.h
@@ -6,6 +6,7 @@
 #include "Chaos/Framework/Parallel.h"
 #include "Chaos/PBDSoftsEvolutionFwd.h"
 #include "Chaos/PBDSoftsSolverParticles.h"
+#include "Chaos/SoftsSolverParticlesRange.h"
 #include "Chaos/PBDStiffness.h"
 #include "Chaos/CollectionPropertyFacade.h"
 #include "ChaosStats.h"
@@ -48,33 +49,6 @@ namespace Chaos::Softs
 		{
 		}
 
-		UE_DEPRECATED(5.3, "Use weight map constructor instead.")
-		FPBDAnimDriveConstraint(
-			const int32 InParticleOffset,
-			const int32 InParticleCount,
-			const TArray<FSolverVec3>& InAnimationPositions,  // Use global indexation (will need adding ParticleOffset)
-			const TArray<FSolverVec3>& InAnimationVelocities,  // Use global indexation (will need adding ParticleOffset)
-			const TConstArrayView<FRealSingle>& StiffnessMultipliers,  // Use local indexation
-			const TConstArrayView<FRealSingle>& DampingMultipliers,  // Use local indexation
-			const FCollectionPropertyConstFacade& PropertyCollection
-		)
-			: AnimationPositions(InAnimationPositions)
-			, AnimationVelocities(InAnimationVelocities)
-			, ParticleOffset(InParticleOffset)
-			, ParticleCount(InParticleCount)
-			, Stiffness(
-				FSolverVec2(GetWeightedFloatAnimDriveStiffness(PropertyCollection, 1.f)),
-				StiffnessMultipliers,
-				InParticleCount)
-			, Damping(
-				FSolverVec2(GetWeightedFloatAnimDriveDamping(PropertyCollection, 1.f)),
-				DampingMultipliers,
-				InParticleCount)
-			, AnimDriveStiffnessIndex(PropertyCollection)
-			, AnimDriveDampingIndex(PropertyCollection)
-		{
-		}
-
 		FPBDAnimDriveConstraint(
 			const int32 InParticleOffset,
 			const int32 InParticleCount,
@@ -135,12 +109,6 @@ namespace Chaos::Softs
 			}
 		}
 
-		UE_DEPRECATED(5.3, "Use SetProperties(const FCollectionPropertyConstFacade&, const TMap<FString, TConstArrayView<FRealSingle>>&, FSolverReal) instead.")
-		void SetProperties(const FCollectionPropertyConstFacade& PropertyCollection)
-		{
-			SetProperties(PropertyCollection, TMap<FString, TConstArrayView<FRealSingle>>());
-		}
-
 		inline void SetProperties(const FSolverVec2& InStiffness, const FSolverVec2& InDamping)
 		{
 			Stiffness.SetWeightedValue(InStiffness);
diff --git a/Engine/Source/Runtime/Experimental/Chaos/Public/Chaos/PBDAxialSpringConstraints.h b/Engine/Source/Runtime/Experimental/Chaos/Public/Chaos/PBDAxialSpringConstraints.h
--- a/Engine/Source/Runtime/Experimental/Chaos/Public/Chaos/PBDAxialSpringConstraints.h
+++ b/Engine/Source/Runtime/Experimental/Chaos/Public/Chaos/PBDAxialSpringConstraints.h
@@ -113,38 +113,12 @@ public:
 		, AreaSpringStiffnessIndex(PropertyCollection)
 	{}
 
-	UE_DEPRECATED(5.3, "Use weight map constructor instead.")
-	FPBDAreaSpringConstraints(
-		const FSolverParticles& Particles,
-		int32 InParticleOffset,
-		int32 InParticleCount,
-		const TArray<TVec3<int32>>& InConstraints,
-		const TConstArrayView<FRealSingle>& StiffnessMultipliers,
-		const FCollectionPropertyConstFacade& PropertyCollection,
-		bool bTrimKinematicConstraints)
-		: FPBDAxialSpringConstraints(
-			Particles,
-			InParticleOffset,
-			InParticleCount,
-			InConstraints,
-			StiffnessMultipliers,
-			FSolverVec2(GetWeightedFloatAreaSpringStiffness(PropertyCollection, 1.f)),
-			bTrimKinematicConstraints)
-		, AreaSpringStiffnessIndex(PropertyCollection)
-	{}
-
 	virtual ~FPBDAreaSpringConstraints() override = default;
 
 	CHAOS_API void SetProperties(
 		const FCollectionPropertyConstFacade& PropertyCollection,
 		const TMap<FString, TConstArrayView<FRealSingle>>& WeightMaps);
 
-	UE_DEPRECATED(5.3, "Use SetProperties(const FCollectionPropertyConstFacade&, const TMap<FString, TConstArrayView<FRealSingle>>&, FSolverReal) instead.")
-	void SetProperties(const FCollectionPropertyConstFacade& PropertyCollection)
-	{
-		SetProperties(PropertyCollection, TMap<FString, TConstArrayView<FRealSingle>>());
-	}
-
 private:
 	using FPBDAxialSpringConstraints::Constraints;
 	using FPBDAxialSpringConstraints::Stiffness;
diff --git a/Engine/Source/Runtime/Experimental/Chaos/Public/Chaos/PBDBendingConstraints.h b/Engine/Source/Runtime/Experimental/Chaos/Public/Chaos/PBDBendingConstraints.h
--- a/Engine/Source/Runtime/Experimental/Chaos/Public/Chaos/PBDBendingConstraints.h
+++ b/Engine/Source/Runtime/Experimental/Chaos/Public/Chaos/PBDBendingConstraints.h
@@ -77,40 +77,6 @@ public:
 		InitColor(InParticles);
 	}
 
-	UE_DEPRECATED(5.3, "Use weight map constructor instead.")
-	FPBDBendingConstraints(const FSolverParticles& InParticles,
-		int32 InParticleOffset,
-		int32 InParticleCount,
-		TArray<TVec4<int32>>&& InConstraints,
-		const TConstArrayView<FRealSingle>& StiffnessMultipliers,
-		const TConstArrayView<FRealSingle>& BucklingStiffnessMultipliers,
-		const FCollectionPropertyConstFacade& PropertyCollection,
-		bool bTrimKinematicConstraints = false)
-		: Base(
-			InParticles,
-			InParticleOffset,
-			InParticleCount,
-			MoveTemp(InConstraints),
-			StiffnessMultipliers,
-			TConstArrayView<FRealSingle>(),
-			BucklingStiffnessMultipliers,
-			TConstArrayView<FRealSingle>(),
-			FSolverVec2(GetWeightedFloatBendingElementStiffness(PropertyCollection, 1.f)),
-			FSolverVec2(GetWeightedFloatBucklingRatio(PropertyCollection, 0.f)),  // BucklingRatio is clamped in base class
-			FSolverVec2(GetWeightedFloatBucklingStiffness(PropertyCollection, 1.f)),
-			FSolverVec2((FSolverReal)0.f),
-			ERestAngleConstructionType::Use3DRestAngles,
-			bTrimKinematicConstraints) 
-		, BendingElementStiffnessIndex(PropertyCollection)
-		, BucklingRatioIndex(PropertyCollection)
-		, BucklingStiffnessIndex(PropertyCollection)
-		, FlatnessRatioIndex(PropertyCollection)
-		, RestAngleIndex(PropertyCollection)
-		, RestAngleTypeIndex(PropertyCollection)
-	{
-		InitColor(InParticles);
-	}
-
 	FPBDBendingConstraints(const FSolverParticles& InParticles,
 		int32 ParticleOffset,
 		int32 ParticleCount,
@@ -186,12 +152,6 @@ public:
 		const FCollectionPropertyConstFacade& PropertyCollection,
 		const TMap<FString, TConstArrayView<FRealSingle>>& WeightMaps);
 
-	UE_DEPRECATED(5.3, "Use SetProperties(const FCollectionPropertyConstFacade&, const TMap<FString, TConstArrayView<FRealSingle>>&, FSolverReal) instead.")
-	void SetProperties(const FCollectionPropertyConstFacade& PropertyCollection)
-	{
-		SetProperties(PropertyCollection, TMap<FString, TConstArrayView<FRealSingle>>());
-	}
-
 	template<typename SolverParticlesOrRange>
 	CHAOS_API void Apply(SolverParticlesOrRange& InParticles, const FSolverReal Dt) const;
 
diff --git a/Engine/Source/Runtime/Experimental/Chaos/Public/Chaos/PBDBendingConstraintsBase.h b/Engine/Source/Runtime/Experimental/Chaos/Public/Chaos/PBDBendingConstraintsBase.h
--- a/Engine/Source/Runtime/Experimental/Chaos/Public/Chaos/PBDBendingConstraintsBase.h
+++ b/Engine/Source/Runtime/Experimental/Chaos/Public/Chaos/PBDBendingConstraintsBase.h
@@ -200,25 +200,6 @@ PRAGMA_ENABLE_DEPRECATION_WARNINGS
 			ERestAngleConstructionType::Use3DRestAngles, bTrimKinematicConstraints, MaxStiffness)
 	{}
 
-	UE_DEPRECATED(5.2, "Use one of the other constructors instead.")
-	FPBDBendingConstraintsBase(const FSolverParticles& InParticles, TArray<TVec4<int32>>&& InConstraints, const FSolverReal InStiffness = (FSolverReal)1.)
-		: FPBDBendingConstraintsBase(
-			InParticles,
-			0, 
-			InParticles.Size(),
-			MoveTemp(InConstraints),
-			TConstArrayView<FRealSingle>(),
-			TConstArrayView<FRealSingle>(),
-			TConstArrayView<FRealSingle>(),
-			TConstArrayView<FRealSingle>(),
-			FSolverVec2(InStiffness),
-			FSolverVec2((FSolverReal)0.f),
-			FSolverVec2(InStiffness),
-			FSolverVec2((FSolverReal)0.f),
-			ERestAngleConstructionType::Use3DRestAngles)
-	{
-	}
-
 	virtual ~FPBDBendingConstraintsBase() {}
 
 	// Update stiffness values
@@ -280,21 +261,6 @@ PRAGMA_ENABLE_DEPRECATION_WARNINGS
 		return SafeDivide(Delta, Denom);
 	}
 
-	UE_DEPRECATED(5.1, "Use GetScalingFactor(const FSolverParticles& InParticles, const int32 i, const TStaticArray<FSolverVec3, 4>& Grads, const FSolverReal ExpStiffnessValue, const FSolverReal BucklingRatio, const FSolverReal ExpBucklingValue) instead.")
-	FSolverReal GetScalingFactor(const FSolverParticles& InParticles, const int32 i, const TArray<FSolverVec3>& Grads) const
-	{
-		TStaticArray<FSolverVec3, 4> GradsStaticArray;
-		GradsStaticArray[0] = Grads[0];
-		GradsStaticArray[1] = Grads[1];
-		GradsStaticArray[2] = Grads[2];
-		GradsStaticArray[3] = Grads[3];
-		if (!Stiffness.HasWeightMap())
-		{
-			return GetScalingFactor(InParticles, i, GradsStaticArray, (FSolverReal)Stiffness, (FSolverReal)BucklingStiffness);
-		}
-		return GetScalingFactor(InParticles, i, GradsStaticArray, Stiffness[i], BucklingStiffness[i]);
-	}
-
 	static FSolverReal CalcAngle(const FSolverVec3& P1, const FSolverVec3& P2, const FSolverVec3& P3, const FSolverVec3& P4)
 	{
 		const FSolverVec3 Normal1 = FSolverVec3::CrossProduct(P1 - P3, P2 - P3).GetSafeNormal();
diff --git a/Engine/Source/Runtime/Experimental/Chaos/Public/Chaos/PBDEvolution.h b/Engine/Source/Runtime/Experimental/Chaos/Public/Chaos/PBDEvolution.h
--- a/Engine/Source/Runtime/Experimental/Chaos/Public/Chaos/PBDEvolution.h
+++ b/Engine/Source/Runtime/Experimental/Chaos/Public/Chaos/PBDEvolution.h
@@ -33,10 +33,6 @@ class FPBDEvolution : public TArrayCollection
 		bool InbUsePerParticleDamping = false);
 	~FPBDEvolution() {}
 
-	// Advance one time step. Filter the input time step if specified.
-	UE_DEPRECATED(5.1, "Use AdvanceOneTimeStep(const FSolverReal Dt) instead.")
-	void AdvanceOneTimeStep(const FSolverReal Dt, const bool bSmoothDt) { AdvanceOneTimeStep(Dt); }
-
 	// Advance one time step. Filter the input time step if specified.
 	CHAOS_API void AdvanceOneTimeStep(const FSolverReal Dt);
 
@@ -152,16 +148,6 @@ class FPBDEvolution : public TArrayCollection
 	FVelocityAndPressureField& GetVelocityAndPressureField(const uint32 GroupId = 0) { check(GroupId < TArrayCollection::Size()); return MGroupVelocityAndPressureFields[GroupId]; }
 	const FVelocityAndPressureField& GetVelocityAndPressureField(const uint32 GroupId = 0) const { check(GroupId < TArrayCollection::Size()); return MGroupVelocityAndPressureFields[GroupId]; }
 
-	UE_DEPRECATED(5.1, "Chaos::Softs::FVelocityField has been renamed FVelocityAndPressureField to match its new behavior.")
-	PRAGMA_DISABLE_DEPRECATION_WARNINGS
-	FVelocityField& GetVelocityField(const uint32 GroupId = 0) { return GetVelocityAndPressureField(GroupId); }
-	PRAGMA_ENABLE_DEPRECATION_WARNINGS
-	UE_DEPRECATED(5.1, "Chaos::Softs::FVelocityField has been renamed FVelocityAndPressureField to match its new behavior.")
-	PRAGMA_DISABLE_DEPRECATION_WARNINGS
-	const FVelocityField& GetVelocityField(const uint32 GroupId = 0) const { return GetVelocityAndPressureField(GroupId); }
-	PRAGMA_ENABLE_DEPRECATION_WARNINGS
-
-
 	int32 GetIterations() const { return MNumIterations; }
 	void SetIterations(const int32 Iterations) { MNumIterations = Iterations; }
 
@@ -183,10 +169,6 @@ class FPBDEvolution : public TArrayCollection
 	FSolverReal GetParticleDamping(const uint32 ParticleIndex = 0) const { check(ParticleIndex < this->Particles().Size() && bUsePerParticleDamping); return MParticleDampings[ParticleIndex]; }
 	void SetParticleDamping(const FSolverReal Damping, const uint32 ParticleIndex = 0) { check(ParticleIndex < this->Particles().Size()); if (bUsePerParticleDamping) { MParticleDampings[ParticleIndex] = Damping; } }
 
-
-	UE_DEPRECATED(4.27, "Use GetCollisionStatus() instead")
-	const bool Collided(int32 index) { return MCollided[index]; }
-
 	const TArray<bool>& GetCollisionStatus() { return MCollided; }
 	const TArray<FSolverVec3>& GetCollisionContacts() const { return MCollisionContacts; }
 	const TArray<FSolverVec3>& GetCollisionNormals() const { return MCollisionNormals; }
diff --git a/Engine/Source/Runtime/Experimental/Chaos/Public/Chaos/PBDLongRangeConstraints.h b/Engine/Source/Runtime/Experimental/Chaos/Public/Chaos/PBDLongRangeConstraints.h
--- a/Engine/Source/Runtime/Experimental/Chaos/Public/Chaos/PBDLongRangeConstraints.h
+++ b/Engine/Source/Runtime/Experimental/Chaos/Public/Chaos/PBDLongRangeConstraints.h
@@ -63,31 +63,6 @@ public:
 		, TetherScaleIndex(PropertyCollection)
 	{}
 
-	UE_DEPRECATED(5.3, "Use weight map constructor instead.")
-	FPBDLongRangeConstraints(
-		const FSolverParticles& Particles,
-		const int32 InParticleOffset,
-		const int32 InParticleCount,
-		const TArray<TConstArrayView<TTuple<int32, int32, FRealSingle>>>& InTethers,
-		const TConstArrayView<FRealSingle>& StiffnessMultipliers,
-		const TConstArrayView<FRealSingle>& ScaleMultipliers,
-		const FCollectionPropertyConstFacade& PropertyCollection,
-		FSolverReal MeshScale)
-		: FPBDLongRangeConstraintsBase(
-			Particles,
-			InParticleOffset,
-			InParticleCount,
-			InTethers,
-			StiffnessMultipliers,
-			ScaleMultipliers,
-			FSolverVec2(GetWeightedFloatTetherStiffness(PropertyCollection, 1.f)),
-			FSolverVec2(GetWeightedFloatTetherScale(PropertyCollection, 1.f)),  // Scale clamping done in constructor
-			FPBDStiffness::DefaultPBDMaxStiffness,
-			MeshScale)
-		, TetherStiffnessIndex(PropertyCollection)
-		, TetherScaleIndex(PropertyCollection)
-	{}
-
 	FPBDLongRangeConstraints(
 		const FSolverParticles& Particles,
 		const int32 InParticleOffset,
@@ -121,12 +96,6 @@ public:
 		const TMap<FString, TConstArrayView<FRealSingle>>& WeightMaps,
 		FSolverReal MeshScale);
 
-	UE_DEPRECATED(5.3, "Use SetProperties(const FCollectionPropertyConstFacade&, const TMap<FString, TConstArrayView<FRealSingle>>&, FSolverReal) instead.")
-	void SetProperties(const FCollectionPropertyConstFacade& PropertyCollection)
-	{
-		SetProperties(PropertyCollection, TMap<FString, TConstArrayView<FRealSingle>>(), (FSolverReal)1.);
-	}
-
 	template<typename SolverParticlesOrRange>
 	CHAOS_API void Apply(SolverParticlesOrRange& Particles, const FSolverReal Dt) const;
 
diff --git a/Engine/Source/Runtime/Experimental/Chaos/Public/Chaos/PBDLongRangeConstraintsBase.h b/Engine/Source/Runtime/Experimental/Chaos/Public/Chaos/PBDLongRangeConstraintsBase.h
--- a/Engine/Source/Runtime/Experimental/Chaos/Public/Chaos/PBDLongRangeConstraintsBase.h
+++ b/Engine/Source/Runtime/Experimental/Chaos/Public/Chaos/PBDLongRangeConstraintsBase.h
@@ -21,17 +21,6 @@ public:
 	static constexpr FSolverReal MinTetherScale = (FSolverReal)0.01;
 	static constexpr FSolverReal MaxTetherScale = (FSolverReal)10.;
 
-	enum class UE_DEPRECATED(5.3, "Tether EMode has been replaced with bUseGeodesicTethers.") EMode : uint8
-	{
-		Euclidean,
-		Geodesic,
-
-		// Deprecated modes
-		FastTetherFastLength = Euclidean,
-		AccurateTetherFastLength = Geodesic,
-		AccurateTetherAccurateLength = Geodesic
-	};
-
 	typedef TTuple<int32, int32, FRealSingle> FTether;
 
 	CHAOS_API FPBDLongRangeConstraintsBase(
@@ -70,14 +59,6 @@ PRAGMA_ENABLE_DEPRECATION_WARNINGS
 		TetherScale.SetWeightedValue(InTetherScale.ClampAxes(MinTetherScale, MaxTetherScale) * MeshScale);
 	}
 
-	// Set the stiffness input values used by the constraint
-	UE_DEPRECATED(5.2, "Use SetProperties instead.")
-	void SetStiffness(const FSolverVec2& InStiffness) { Stiffness.SetWeightedValue(InStiffness); }
-
-	// Set the scale low and high value of the scale weight map
-	UE_DEPRECATED(5.2, "Use SetProperties instead.")
-	void SetScale(const FSolverVec2& InScale) { TetherScale.SetWeightedValue(InScale.ClampAxes(MinTetherScale, MaxTetherScale)); }
-
 	// Set stiffness offset and range, as well as the simulation stiffness exponent
 	void ApplyProperties(const FSolverReal Dt, const int32 NumIterations)
 	{
@@ -113,14 +94,6 @@ protected:
 	// Return the minimum number of long range tethers in a batch to process in parallel
 	static CHAOS_API int32 GetMinParallelBatchSize();
 
-	// Return whether the constraint has been setup with a weightmap to interpolate between two low and high values of scales
-	UE_DEPRECATED(5.2, "Use TetherScale.HasWeightMap() instead")
-	bool HasScaleWeightMap() const { return TetherScale.HasWeightMap(); }
-
-	// Part of ApplyProperties to update ScaleTable.
-	UE_DEPRECATED(5.2, "Use TetherScale.ApplyValues() instead")
-	void ApplyScale() { TetherScale.ApplyValues(); }
-
 	// Return a vector representing the amount of segment required for the tether to shrink back to its maximum target length constraint, or zero if the constraint is already met
 	template<typename SolverParticlesOrRange>
 	inline FSolverVec3 GetDelta(const SolverParticlesOrRange& Particles, const FTether& Tether, const FSolverReal InScale) const
@@ -157,12 +130,5 @@ protected:
 	const int32 ParticleCount;
 	FPBDStiffness Stiffness;  // Stiffness weightmap lookup table
 	FPBDWeightMap TetherScale;  // Scale weightmap lookup table
-
-	UE_DEPRECATED(5.2, "Use TetherScale instead")
-	TArray<uint8> ScaleIndices;
-	UE_DEPRECATED(5.2, "Use TetherScale instead")
-	TArray<FSolverReal> ScaleTable;
-	UE_DEPRECATED(5.2, "Use TetherScale instead")
-	FSolverVec2 Scale;
 };
 }  // End namespace Chaos::Softs
diff --git a/Engine/Source/Runtime/Experimental/Chaos/Public/Chaos/PBDSoftsEvolutionFwd.h b/Engine/Source/Runtime/Experimental/Chaos/Public/Chaos/PBDSoftsEvolutionFwd.h
--- a/Engine/Source/Runtime/Experimental/Chaos/Public/Chaos/PBDSoftsEvolutionFwd.h
+++ b/Engine/Source/Runtime/Experimental/Chaos/Public/Chaos/PBDSoftsEvolutionFwd.h
@@ -81,7 +81,6 @@ namespace Chaos
 
 		// Softs solver forces
 		class FVelocityAndPressureField;
-		using FVelocityField UE_DEPRECATED(5.1, "Chaos::Softs::FVelocityField has been renamed FVelocityAndPressureField to match its new behavior.") = FVelocityAndPressureField;
 		class FExternalForces;
 
 		// Linear system solver
diff --git a/Engine/Source/Runtime/Experimental/Chaos/Public/Chaos/PBDSphericalConstraint.h b/Engine/Source/Runtime/Experimental/Chaos/Public/Chaos/PBDSphericalConstraint.h
--- a/Engine/Source/Runtime/Experimental/Chaos/Public/Chaos/PBDSphericalConstraint.h
+++ b/Engine/Source/Runtime/Experimental/Chaos/Public/Chaos/PBDSphericalConstraint.h
@@ -54,26 +54,6 @@ public:
 	{
 	}
 
-	UE_DEPRECATED(5.3, "Use weight map constructor instead.")
-	FPBDSphericalConstraint(
-		const uint32 InParticleOffset,
-		const uint32 InParticleCount,
-		const TArray<FSolverVec3>& InAnimationPositions,  // Use global indexation (will need adding ParticleOffset)
-		const TConstArrayView<FRealSingle>& InSphereRadii,  // Use local indexation
-		const FCollectionPropertyConstFacade& PropertyCollection,
-		FSolverReal MeshScale
-	)
-		: AnimationPositions(InAnimationPositions)
-		, SphereRadii(InSphereRadii)
-		, ParticleOffset(InParticleOffset)
-		, ParticleCount(InParticleCount)
-		, Scale(MeshScale)
-		, MaxDistanceBase((FSolverReal)GetLowMaxDistance(PropertyCollection, 0.f))
-		, MaxDistanceRange((FSolverReal)GetHighMaxDistance(PropertyCollection, 1.f) - MaxDistanceBase)
-		, MaxDistanceIndex(PropertyCollection)
-	{
-	}
-
 	FPBDSphericalConstraint(
 		const uint32 InParticleOffset,
 		const uint32 InParticleCount,
@@ -96,13 +76,6 @@ public:
 		const TMap<FString, TConstArrayView<FRealSingle>>& WeightMaps,
 		FSolverReal MeshScale);
 
-	UE_DEPRECATED(5.3, "Use SetProperties(const FCollectionPropertyConstFacade&, const TMap<FString, TConstArrayView<FRealSingle>>&, FSolverReal) instead.")
-	void SetProperties(const FCollectionPropertyConstFacade& PropertyCollection, FSolverReal MeshScale)
-	{
-		SetProperties(PropertyCollection, TMap<FString, TConstArrayView<FRealSingle>>(), MeshScale);
-	}
-
-
 	template<typename SolverParticlesOrRange>
 	void Apply(SolverParticlesOrRange& Particles, const FSolverReal Dt) const
 	{
@@ -129,21 +102,8 @@ public:
 
 	// Set a new mesh scale
 	void SetScale(FSolverReal InScale) { Scale = InScale; }
-
-	UE_DEPRECATED(5.3, "Use SetScale(FSolverReal) instead.")
-	void SetScale(FSolverReal MaxDistanceScale, FSolverReal MeshScale)
-	{
-		Scale = FMath::Max(MaxDistanceScale, (FSolverReal)0.) * MeshScale;
-	}
-
 	FSolverReal GetScale() const { return Scale; }
 
-	UE_DEPRECATED(5.2, "Use SetScale instead.")
-	void SetSphereRadiiMultiplier(FSolverReal InSphereRadiiMultiplier, FSolverReal MeshScale)
-	{
-		SetScale(InSphereRadiiMultiplier * MeshScale);
-	}
-
 private:
 
 	template<bool bHasMaxDistance, typename SolverParticlesOrRange>
@@ -183,8 +143,6 @@ protected:
 	TConstArrayView<FRealSingle> SphereRadii;  // Use local indexation
 	const int32 ParticleOffset;
 	const int32 ParticleCount;
-	UE_DEPRECATED(5.2, "Use Scale instead.")
-	FSolverReal SphereRadiiMultiplier = 1.f;
 
 private:
 	FSolverReal Scale = (FSolverReal)1.;
@@ -230,35 +188,6 @@ public:
 	{
 	}
 
-	UE_DEPRECATED(5.3, "Use weight map constructor instead.")
-	FPBDSphericalBackstopConstraint(
-		const int32 InParticleOffset,
-		const int32 InParticleCount,
-		const TArray<FSolverVec3>& InAnimationPositions,  // Use global indexation (will need adding ParticleOffset)
-		const TArray<FSolverVec3>& InAnimationNormals,  // Use global indexation (will need adding ParticleOffset)
-		const TConstArrayView<FRealSingle>& InSphereRadii,  // Use local indexation
-		const TConstArrayView<FRealSingle>& InSphereOffsetDistances,  // Use local indexation
-		const FCollectionPropertyConstFacade& PropertyCollection,
-		FSolverReal MeshScale
-	)
-		: AnimationPositions(InAnimationPositions)
-		, AnimationNormals(InAnimationNormals)
-		, SphereRadii(InSphereRadii)
-		, SphereOffsetDistances(InSphereOffsetDistances)
-		, ParticleOffset(InParticleOffset)
-		, ParticleCount(InParticleCount)
-		, Scale(MeshScale)
-		, BackstopRadiusBase((FSolverReal)FMath::Max(GetLowBackstopRadius(PropertyCollection, 0.f), 0.f))
-		, BackstopRadiusRange((FSolverReal)FMath::Max(GetHighBackstopRadius(PropertyCollection, 1.f), 0.f) - BackstopRadiusBase)
-		, BackstopDistanceBase((FSolverReal)GetLowBackstopDistance(PropertyCollection, 0.f))
-		, BackstopDistanceRange((FSolverReal)GetHighBackstopDistance(PropertyCollection, 1.f) - BackstopDistanceBase)
-		, bUseLegacyBackstop(GetUseLegacyBackstop(PropertyCollection, false))  // Only set the legacy backstop in constructor
-		, BackstopDistanceIndex(PropertyCollection)
-		, BackstopRadiusIndex(PropertyCollection)
-		, UseLegacyBackstopIndex(PropertyCollection)
-	{
-	}
-
 	FPBDSphericalBackstopConstraint(
 		const int32 InParticleOffset,
 		const int32 InParticleCount,
@@ -289,12 +218,6 @@ public:
 		const TMap<FString, TConstArrayView<FRealSingle>>& WeightMaps,
 		FSolverReal MeshScale);
 
-	UE_DEPRECATED(5.3, "Use SetProperties(const FCollectionPropertyConstFacade&, const TMap<FString, TConstArrayView<FRealSingle>>&, FSolverReal) instead.")
-	void SetProperties(const FCollectionPropertyConstFacade& PropertyCollection, FSolverReal MeshScale)
-	{
-		SetProperties(PropertyCollection, TMap<FString, TConstArrayView<FRealSingle>>(), MeshScale);
-	}
-
 	void SetEnabled(bool bInEnabled) { bEnabled = bInEnabled; }
 	bool IsEnabled() const { return bEnabled; }
 
@@ -382,26 +305,8 @@ public:
 
 	// Set a new mesh scale
 	void SetScale(FSolverReal InScale) { Scale = InScale; }
-
-	UE_DEPRECATED(5.3, "Use SetScale(FSolverReal) instead.")
-	void SetScale(FSolverReal BackstopScale, FSolverReal MeshScale)
-	{
-		Scale = FMath::Max(BackstopScale, (FSolverReal)0.) * MeshScale;
-	}
-
 	FSolverReal GetScale() const { return Scale; }
 
-	UE_DEPRECATED(5.2, "Use SetScale instead.")
-	void SetSphereRadiiMultiplier(FSolverReal InSphereRadiiMultiplier, FSolverReal MeshScale = (FSolverReal)1.)
-	{
-PRAGMA_DISABLE_DEPRECATION_WARNINGS
-		SetScale(InSphereRadiiMultiplier, MeshScale);
-PRAGMA_ENABLE_DEPRECATION_WARNINGS
-	}
-
-	UE_DEPRECATED(5.2, "Use GetScale() instead.")
-	FSolverReal GetSphereRadiiMultiplier() const { return GetScale(); }
-
 	bool UseLegacyBackstop() const
 	{
 		return bUseLegacyBackstop;
diff --git a/Engine/Source/Runtime/Experimental/Chaos/Public/Chaos/PBDSpringConstraints.h b/Engine/Source/Runtime/Experimental/Chaos/Public/Chaos/PBDSpringConstraints.h
--- a/Engine/Source/Runtime/Experimental/Chaos/Public/Chaos/PBDSpringConstraints.h
+++ b/Engine/Source/Runtime/Experimental/Chaos/Public/Chaos/PBDSpringConstraints.h
@@ -116,38 +116,12 @@ public:
 		, EdgeSpringStiffnessIndex(PropertyCollection)
 	{}
 
-	UE_DEPRECATED(5.3, "Use weight map constructor instead.")
-	FPBDEdgeSpringConstraints(
-		const FSolverParticles& Particles,
-		int32 ParticleOffset,
-		int32 ParticleCount,
-		const TArray<TVec3<int32>>& InConstraints,
-		const TConstArrayView<FRealSingle>& StiffnessMultipliers,
-		const FCollectionPropertyConstFacade& PropertyCollection,
-		bool bTrimKinematicConstraints = false)
-		: FPBDSpringConstraints(
-			Particles,
-			ParticleOffset,
-			ParticleCount,
-			InConstraints,
-			StiffnessMultipliers,
-			FSolverVec2(GetWeightedFloatEdgeSpringStiffness(PropertyCollection, 1.f)),
-			bTrimKinematicConstraints)
-		, EdgeSpringStiffnessIndex(PropertyCollection)
-	{}
-
 	virtual ~FPBDEdgeSpringConstraints() override = default;
 
 	CHAOS_API void SetProperties(
 		const FCollectionPropertyConstFacade& PropertyCollection,
 		const TMap<FString, TConstArrayView<FRealSingle>>& WeightMaps);
 
-	UE_DEPRECATED(5.3, "Use SetProperties(const FCollectionPropertyConstFacade&, const TMap<FString, TConstArrayView<FRealSingle>>&, FSolverReal) instead.")
-	void SetProperties(const FCollectionPropertyConstFacade& PropertyCollection)
-	{
-		SetProperties(PropertyCollection, TMap<FString, TConstArrayView<FRealSingle>>());
-	}
-
 private:
 	using FPBDSpringConstraints::ParticleCount;
 
@@ -196,38 +170,12 @@ public:
 		, BendingSpringStiffnessIndex(PropertyCollection)
 	{}
 
-	UE_DEPRECATED(5.3, "Use weight map constructor instead.")
-	FPBDBendingSpringConstraints(
-		const FSolverParticles& Particles,
-		int32 ParticleOffset,
-		int32 ParticleCount,
-		const TArray<TVec2<int32>>& InConstraints,
-		const TConstArrayView<FRealSingle>& StiffnessMultipliers,
-		const FCollectionPropertyConstFacade& PropertyCollection,
-		bool bTrimKinematicConstraints = false)
-		: FPBDSpringConstraints(
-			Particles,
-			ParticleOffset,
-			ParticleCount,
-			InConstraints,
-			StiffnessMultipliers,
-			FSolverVec2(GetWeightedFloatBendingSpringStiffness(PropertyCollection, 1.f)),
-			bTrimKinematicConstraints)
-		, BendingSpringStiffnessIndex(PropertyCollection)
-	{}
-
 	virtual ~FPBDBendingSpringConstraints() override = default;
 
 	CHAOS_API void SetProperties(
 		const FCollectionPropertyConstFacade& PropertyCollection,
 		const TMap<FString, TConstArrayView<FRealSingle>>& WeightMaps);
 
-	UE_DEPRECATED(5.3, "Use SetProperties(const FCollectionPropertyConstFacade&, const TMap<FString, TConstArrayView<FRealSingle>>&, FSolverReal) instead.")
-	void SetProperties(const FCollectionPropertyConstFacade& PropertyCollection)
-	{
-		SetProperties(PropertyCollection, TMap<FString, TConstArrayView<FRealSingle>>());
-	}
-
 private:
 	using FPBDSpringConstraints::ParticleCount;
 
diff --git a/Engine/Source/Runtime/Experimental/Chaos/Public/Chaos/PBDStiffness.h b/Engine/Source/Runtime/Experimental/Chaos/Public/Chaos/PBDStiffness.h
--- a/Engine/Source/Runtime/Experimental/Chaos/Public/Chaos/PBDStiffness.h
+++ b/Engine/Source/Runtime/Experimental/Chaos/Public/Chaos/PBDStiffness.h
@@ -70,10 +70,6 @@ public:
 	/** Update the weight map table with the current simulation parameters. */
 	inline void ApplyPBDValues(const FSolverReal Dt, const int32 NumIterations);
 
-	/** Update the weight map table with the current simulation parameters. */
-	UE_DEPRECATED(5.2, "Use ApplyPBDValues() instead.")
-	void ApplyValues(const FSolverReal Dt, const int32 NumIterations) { ApplyPBDValues(Dt, NumIterations); }
-
 	/** Update the weight map table with the current simulation parameters. */
 	inline void ApplyXPBDValues(const FSolverReal MaxStiffnesss);
 
diff --git a/Engine/Source/Runtime/Experimental/Chaos/Public/Chaos/TriangleMesh.h b/Engine/Source/Runtime/Experimental/Chaos/Public/Chaos/TriangleMesh.h
--- a/Engine/Source/Runtime/Experimental/Chaos/Public/Chaos/TriangleMesh.h
+++ b/Engine/Source/Runtime/Experimental/Chaos/Public/Chaos/TriangleMesh.h
@@ -194,9 +194,6 @@ namespace Chaos
 		 */
 		CHAOS_API const TArray<TVec2<int32>>& GetEdgeToFaces() const;
 
-		UE_DEPRECATED(5.1, "Non-const access to GetSegmentMesh will be removed. Use const version instead.")
-		FSegmentMesh& GetSegmentMesh() { return const_cast<FSegmentMesh&>(const_cast<const FTriangleMesh*>(this)->GetSegmentMesh()); }
-
 		/**
 		 * @ret Curvature between adjacent faces, specified on edges in radians.
 		 * @param faceNormals - a normal per face.
diff --git a/Engine/Source/Runtime/Experimental/Chaos/Public/Chaos/VelocityField.h b/Engine/Source/Runtime/Experimental/Chaos/Public/Chaos/VelocityField.h
--- a/Engine/Source/Runtime/Experimental/Chaos/Public/Chaos/VelocityField.h
+++ b/Engine/Source/Runtime/Experimental/Chaos/Public/Chaos/VelocityField.h
@@ -119,13 +119,6 @@ public:
 		bool bEnableAerodynamics,
 		const FSolverVec3& SolverWind);
 
-	UE_DEPRECATED(5.3, "Use SetProperties(const FCollectionPropertyConstFacade&, const TMap<FString, TConstArrayView<FRealSingle>>&, FSolverReal, bool) instead.")
-	void SetProperties(const FCollectionPropertyConstFacade& PropertyCollection, FSolverReal WorldScale)
-	{
-		constexpr bool bEnableAerodynamics = true;
-		SetProperties(PropertyCollection, TMap<FString, TConstArrayView<FRealSingle>>(), WorldScale, bEnableAerodynamics);
-	}
-
 	CHAOS_API void SetProperties(
 		const FSolverVec2& Drag,
 		const FSolverVec2& Lift,
@@ -272,8 +265,6 @@ private:
 	UE_CHAOS_DECLARE_PROPERTYCOLLECTION_NAME(WindVelocity, FVector3f);
 };
 
-using FVelocityField UE_DEPRECATED(5.1, "Chaos::Softs::FVelocityField has been renamed FVelocityAndPressureField to match its new behavior.") = FVelocityAndPressureField;
-
 }  // End namespace Chaos::Softs
 
 #if !defined(CHAOS_VELOCITY_FIELD_ISPC_ENABLED_DEFAULT)
diff --git a/Engine/Source/Runtime/Experimental/Chaos/Public/Chaos/XPBDAnisotropicBendingConstraints.h b/Engine/Source/Runtime/Experimental/Chaos/Public/Chaos/XPBDAnisotropicBendingConstraints.h
--- a/Engine/Source/Runtime/Experimental/Chaos/Public/Chaos/XPBDAnisotropicBendingConstraints.h
+++ b/Engine/Source/Runtime/Experimental/Chaos/Public/Chaos/XPBDAnisotropicBendingConstraints.h
@@ -197,12 +197,6 @@ public:
 		const FCollectionPropertyConstFacade& PropertyCollection,
 		const TMap<FString, TConstArrayView<FRealSingle>>& WeightMaps);
 
-	UE_DEPRECATED(5.3, "Use SetProperties(const FCollectionPropertyConstFacade&, const TMap<FString, TConstArrayView<FRealSingle>>&, FSolverReal) instead.")
-	void SetProperties(const FCollectionPropertyConstFacade& PropertyCollection)
-	{
-		SetProperties(PropertyCollection, TMap<FString, TConstArrayView<FRealSingle>>());
-	}
-
 	void ApplyProperties(const FSolverReal /*Dt*/, const int32 /*NumIterations*/)
 	{
 		BucklingRatioWeighted.ApplyValues();
diff --git a/Engine/Source/Runtime/Experimental/Chaos/Public/Chaos/XPBDAxialSpringConstraints.h b/Engine/Source/Runtime/Experimental/Chaos/Public/Chaos/XPBDAxialSpringConstraints.h
--- a/Engine/Source/Runtime/Experimental/Chaos/Public/Chaos/XPBDAxialSpringConstraints.h
+++ b/Engine/Source/Runtime/Experimental/Chaos/Public/Chaos/XPBDAxialSpringConstraints.h
@@ -10,12 +10,6 @@ DECLARE_CYCLE_STAT(TEXT("Chaos XPBD Axial Spring Constraint"), STAT_XPBD_AxialSp
 namespace Chaos::Softs
 {
 
-// Stiffness is in kg/s^2
-UE_DEPRECATED(5.2, "Use FXPBDAxialSpringConstraints::MinStiffness instead.")
-static const FSolverReal XPBDAxialSpringMinStiffness = (FSolverReal)1e-4; // Stiffness below this will be considered 0 since all of our calculations are actually based on 1 / stiffness.
-UE_DEPRECATED(5.2, "Use FXPBDAxialSpringConstraints::MaxStiffness instead.")
-static const FSolverReal XPBDAxialSpringMaxStiffness = (FSolverReal)1e7;
-
 class FXPBDAxialSpringConstraints : public FPBDAxialSpringConstraintsBase
 {
 	typedef FPBDAxialSpringConstraintsBase Base;
@@ -228,26 +222,6 @@ public:
 		, XPBDAreaSpringStiffnessIndex(PropertyCollection)
 	{}
 
-	UE_DEPRECATED(5.3, "Use weight map constructor instead.")
-	FXPBDAreaSpringConstraints(
-		const FSolverParticles& Particles,
-		int32 ParticleOffset,
-		int32 ParticleCount,
-		const TArray<TVec3<int32>>& InConstraints,
-		const TConstArrayView<FRealSingle>& StiffnessMultipliers,
-		const FCollectionPropertyConstFacade& PropertyCollection,
-		bool bTrimKinematicConstraints)
-		: FXPBDAxialSpringConstraints(
-			Particles,
-			ParticleOffset,
-			ParticleCount,
-			InConstraints,
-			StiffnessMultipliers,
-			FSolverVec2(GetWeightedFloatXPBDAreaSpringStiffness(PropertyCollection, MaxStiffness)),
-			bTrimKinematicConstraints)
-		, XPBDAreaSpringStiffnessIndex(PropertyCollection)
-	{}
-
 	virtual ~FXPBDAreaSpringConstraints() override = default;
 
 	void SetProperties(
@@ -277,12 +251,6 @@ public:
 		}
 	}
 
-	UE_DEPRECATED(5.3, "Use SetProperties(const FCollectionPropertyConstFacade&, const TMap<FString, TConstArrayView<FRealSingle>>&, FSolverReal) instead.")
-	void SetProperties(const FCollectionPropertyConstFacade& PropertyCollection)
-	{
-		SetProperties(PropertyCollection, TMap<FString, TConstArrayView<FRealSingle>>());
-	}
-
 private:
 	using FXPBDAxialSpringConstraints::Constraints;
 	using FXPBDAxialSpringConstraints::ParticleOffset;
diff --git a/Engine/Source/Runtime/Experimental/Chaos/Public/Chaos/XPBDBendingConstraints.h b/Engine/Source/Runtime/Experimental/Chaos/Public/Chaos/XPBDBendingConstraints.h
--- a/Engine/Source/Runtime/Experimental/Chaos/Public/Chaos/XPBDBendingConstraints.h
+++ b/Engine/Source/Runtime/Experimental/Chaos/Public/Chaos/XPBDBendingConstraints.h
@@ -9,12 +9,6 @@
 namespace Chaos::Softs
 {
 
-// Stiffness is in kg cm^2 / rad^2 s^2
-UE_DEPRECATED(5.2, "Use FXPBDBendingConstraints::MinStiffness instead.")
-static const FSolverReal XPBDBendMinStiffness = (FSolverReal)0;  // We're not checking against MinStiffness (except when it's constant and == 0)
-UE_DEPRECATED(5.2, "Use FXPBDBendingConstraints::MaxStiffness instead.")
-static const FSolverReal XPBDBendMaxStiffness = (FSolverReal)1e7;
-
 class FXPBDBendingConstraints final : public FPBDBendingConstraintsBase
 {
 	typedef FPBDBendingConstraintsBase Base;
@@ -153,63 +147,6 @@ public:
 		: FXPBDBendingConstraints(InParticles, InParticleOffset, InParticleCount, MoveTemp(InConstraints), WeightMaps, PropertyCollection)
 	{}
 
-	UE_DEPRECATED(5.3, "Use weight map constructor instead.")
-	FXPBDBendingConstraints(const FSolverParticles& InParticles,
-		int32 InParticleOffset,
-		int32 InParticleCount,
-		TArray<TVec4<int32>>&& InConstraints,
-		const TConstArrayView<FRealSingle>& StiffnessMultipliers,
-		const TConstArrayView<FRealSingle>& BucklingStiffnessMultipliers,
-		const TConstArrayView<FRealSingle>& DampingMultipliers,
-		const FCollectionPropertyConstFacade& PropertyCollection,
-		bool bTrimKinematicConstraints = false)
-		: Base(
-			InParticles,
-			InParticleOffset,
-			InParticleCount,
-			MoveTemp(InConstraints),
-			TConstArrayView<FRealSingle>(), // We don't use base stiffness weight maps
-			TConstArrayView<FRealSingle>(),
-			TConstArrayView<FRealSingle>(), // We don't use base stiffness weight maps
-			TConstArrayView<FRealSingle>(),
-			FSolverVec2(GetWeightedFloatXPBDBendingElementStiffness(PropertyCollection, MaxStiffness)),
-			FSolverVec2(GetWeightedFloatXPBDBucklingRatio(PropertyCollection, 0.f)),
-			FSolverVec2(GetWeightedFloatXPBDBucklingStiffness(PropertyCollection, MaxStiffness)),
-			FSolverVec2((FSolverReal)0.f),
-			ERestAngleConstructionType::Use3DRestAngles,
-			true /*bTrimKinematicConstraints*/,
-			MaxStiffness)
-		, XPBDStiffness(
-			FSolverVec2(GetWeightedFloatXPBDBendingElementStiffness(PropertyCollection, MaxStiffness)).ClampAxes(0, MaxStiffness),
-			StiffnessMultipliers,
-			TConstArrayView<TVec2<int32>>(ConstraintSharedEdges),
-			ParticleOffset,
-			ParticleCount)
-		, XPBDBucklingStiffness(
-			FSolverVec2(GetWeightedFloatXPBDBucklingStiffness(PropertyCollection, MaxStiffness)).ClampAxes(0, MaxStiffness),
-			BucklingStiffnessMultipliers,
-			TConstArrayView<TVec2<int32>>(ConstraintSharedEdges),
-			ParticleOffset,
-			ParticleCount)
-		, DampingRatio(
-			FSolverVec2(GetWeightedFloatXPBDBendingElementDamping(PropertyCollection, MinDamping)).ClampAxes(MinDamping, MaxDamping),
-			DampingMultipliers,
-			TConstArrayView<TVec2<int32>>(ConstraintSharedEdges),
-			ParticleOffset,
-			ParticleCount)
-		, XPBDBendingElementStiffnessIndex(PropertyCollection)
-		, XPBDBendingElementDampingIndex(PropertyCollection)
-		, XPBDBucklingRatioIndex(PropertyCollection)
-		, XPBDBucklingStiffnessIndex(PropertyCollection)
-		, XPBDFlatnessRatioIndex(PropertyCollection)
-		, XPBDRestAngleIndex(PropertyCollection)
-		, XPBDRestAngleTypeIndex(PropertyCollection)
-	{
-		Lambdas.Init((FSolverReal)0., Constraints.Num());
-		LambdasDamping.Init((FSolverReal)0., Constraints.Num());
-		InitColor(InParticles);
-	}
-
 	FXPBDBendingConstraints(const FSolverParticles& InParticles,
 		int32 ParticleOffset,
 		int32 ParticleCount,
@@ -343,12 +280,6 @@ public:
 		const FCollectionPropertyConstFacade& PropertyCollection,
 		const TMap<FString, TConstArrayView<FRealSingle>>& WeightMaps);
 
-	UE_DEPRECATED(5.3, "Use SetProperties(const FCollectionPropertyConstFacade&, const TMap<FString, TConstArrayView<FRealSingle>>&, FSolverReal) instead.")
-	void SetProperties(const FCollectionPropertyConstFacade& PropertyCollection)
-	{
-		SetProperties(PropertyCollection, TMap<FString, TConstArrayView<FRealSingle>>());
-	}
-
 	void SetProperties(const FSolverVec2& InStiffness, const FSolverVec2& InBucklingRatio, const FSolverVec2& InBucklingStiffness, const FSolverVec2& InDampingRatio)
 	{
 		XPBDStiffness.SetWeightedValue(InStiffness.ClampAxes(0, MaxStiffness));
diff --git a/Engine/Source/Runtime/Experimental/Chaos/Public/Chaos/XPBDLongRangeConstraints.h b/Engine/Source/Runtime/Experimental/Chaos/Public/Chaos/XPBDLongRangeConstraints.h
--- a/Engine/Source/Runtime/Experimental/Chaos/Public/Chaos/XPBDLongRangeConstraints.h
+++ b/Engine/Source/Runtime/Experimental/Chaos/Public/Chaos/XPBDLongRangeConstraints.h
@@ -11,11 +11,6 @@ DECLARE_CYCLE_STAT(TEXT("Chaos XPBD Long Range Constraint"), STAT_XPBD_LongRange
 namespace Chaos::Softs
 {
 
-UE_DEPRECATED(5.2, "Use FXPBDLongRangeConstraints::MinStiffness instead.")
-static const FSolverReal XPBDLongRangeMinStiffness = (FSolverReal)1e-1;
-UE_DEPRECATED(5.2, "Use FXPBDLongRangeConstraints::MaxStiffness instead.")
-static const FSolverReal XPBDLongRangeMaxStiffness = (FSolverReal)1e7;
-
 class FXPBDLongRangeConstraints final : public FPBDLongRangeConstraintsBase
 {
 public:
@@ -55,36 +50,6 @@ public:
 		Lambdas.Reserve(NumTethers);
 	}
 
-	UE_DEPRECATED(5.3, "Use weight map constructor instead.")
-	FXPBDLongRangeConstraints(
-		const FSolverParticles& Particles,
-		const int32 InParticleOffset,
-		const int32 InParticleCount,
-		const TArray<TConstArrayView<TTuple<int32, int32, FRealSingle>>>& InTethers,
-		const TConstArrayView<FRealSingle>& StiffnessMultipliers,
-		const TConstArrayView<FRealSingle>& ScaleMultipliers,
-		const FCollectionPropertyConstFacade& PropertyCollection)
-	    : FPBDLongRangeConstraintsBase(
-			Particles,
-			InParticleOffset,
-			InParticleCount,
-			InTethers,
-			StiffnessMultipliers,
-			ScaleMultipliers,
-			FSolverVec2(GetWeightedFloatXPBDTetherStiffness(PropertyCollection, MaxStiffness)),
-			FSolverVec2(GetWeightedFloatXPBDTetherScale(PropertyCollection, 1.f)),  // Scale clamping done in constructor
-			MaxStiffness)
-		, XPBDTetherStiffnessIndex(PropertyCollection)
-		, XPBDTetherScaleIndex(PropertyCollection)
-	{
-		NumTethers = 0;
-		for (const TConstArrayView<FTether>& TetherBatch : Tethers)
-		{
-			NumTethers += TetherBatch.Num();
-		}
-		Lambdas.Reserve(NumTethers);
-	}
-
 	FXPBDLongRangeConstraints(
 		const FSolverParticles& Particles,
 		const int32 InParticleOffset,
@@ -156,12 +121,6 @@ public:
 		}
 	}
 
-	UE_DEPRECATED(5.3, "Use SetProperties(const FCollectionPropertyConstFacade&, const TMap<FString, TConstArrayView<FRealSingle>>&, FSolverReal) instead.")
-	void SetProperties(const FCollectionPropertyConstFacade& PropertyCollection, FSolverReal MeshScale)
-	{
-		SetProperties(PropertyCollection, TMap<FString, TConstArrayView<FRealSingle>>(), MeshScale);
-	}
-
 	// Set the stiffness and scale values used by the constraint
 	void SetProperties(const FSolverVec2& InStiffness, const FSolverVec2& InTetherScale, FSolverReal MeshScale = (FSolverReal)1)
 	{
diff --git a/Engine/Source/Runtime/Experimental/Chaos/Public/Chaos/XPBDSpringConstraints.h b/Engine/Source/Runtime/Experimental/Chaos/Public/Chaos/XPBDSpringConstraints.h
--- a/Engine/Source/Runtime/Experimental/Chaos/Public/Chaos/XPBDSpringConstraints.h
+++ b/Engine/Source/Runtime/Experimental/Chaos/Public/Chaos/XPBDSpringConstraints.h
@@ -10,12 +10,6 @@
 namespace Chaos::Softs
 {
 
-// Stiffness is in kg cm /s^2
-UE_DEPRECATED(5.2, "Use FXPBDSpringConstraints::MinStiffness instead.")
-static const FSolverReal XPBDSpringMinStiffness = (FSolverReal)0; // We're not checking against MinStiffness (except when it's constant and == 0)
-UE_DEPRECATED(5.2, "Use FXPBDSpringConstraints::MaxStiffness instead.")
-static const FSolverReal XPBDSpringMaxStiffness = (FSolverReal)1e9;
-
 class FXPBDSpringConstraints : public FPBDSpringConstraintsBase
 {
 	typedef FPBDSpringConstraintsBase Base;
@@ -284,41 +278,12 @@ public:
 		, XPBDEdgeSpringDampingIndex(PropertyCollection)
 	{}
 
-	UE_DEPRECATED(5.3, "Use weight map constructor instead.")
-	FXPBDEdgeSpringConstraints(
-		const FSolverParticles& Particles,
-		int32 ParticleOffset,
-		int32 ParticleCount,
-		const TArray<TVec3<int32>>& InConstraints,
-		const TConstArrayView<FRealSingle>& StiffnessMultipliers,
-		const TConstArrayView<FRealSingle>& DampingMultipliers,
-		const FCollectionPropertyConstFacade& PropertyCollection,
-		bool bTrimKinematicConstraints = false)
-		: FXPBDSpringConstraints(
-			Particles,
-			ParticleOffset,
-			ParticleCount,
-			InConstraints,
-			StiffnessMultipliers,
-			DampingMultipliers,
-			FSolverVec2(GetWeightedFloatXPBDEdgeSpringStiffness(PropertyCollection, MaxStiffness)),
-			FSolverVec2(GetWeightedFloatXPBDEdgeSpringDamping(PropertyCollection, MinDampingRatio)))
-		, XPBDEdgeSpringStiffnessIndex(PropertyCollection)
-		, XPBDEdgeSpringDampingIndex(PropertyCollection)
-	{}
-
 	virtual ~FXPBDEdgeSpringConstraints() override = default;
 
 	CHAOS_API void SetProperties(
 		const FCollectionPropertyConstFacade& PropertyCollection,
 		const TMap<FString, TConstArrayView<FRealSingle>>& WeightMaps);
 
-	UE_DEPRECATED(5.3, "Use SetProperties(const FCollectionPropertyConstFacade&, const TMap<FString, TConstArrayView<FRealSingle>>&, FSolverReal) instead.")
-	void SetProperties(const FCollectionPropertyConstFacade& PropertyCollection)
-	{
-		SetProperties(PropertyCollection, TMap<FString, TConstArrayView<FRealSingle>>());
-	}
-
 private:
 	using FXPBDSpringConstraints::Constraints;
 	using FXPBDSpringConstraints::ParticleOffset;
@@ -414,41 +379,12 @@ public:
 		, XPBDBendingSpringDampingIndex(PropertyCollection)
 	{}
 
-	UE_DEPRECATED(5.3, "Use weight map constructor instead.")
-	FXPBDBendingSpringConstraints(
-		const FSolverParticles& Particles,
-		int32 ParticleOffset,
-		int32 ParticleCount,
-		const TArray<TVec2<int32>>& InConstraints,
-		const TConstArrayView<FRealSingle>& StiffnessMultipliers,
-		const TConstArrayView<FRealSingle>& DampingMultipliers,
-		const FCollectionPropertyConstFacade& PropertyCollection,
-		bool bTrimKinematicConstraints = false)
-		: FXPBDSpringConstraints(
-			Particles,
-			ParticleOffset,
-			ParticleCount,
-			InConstraints,
-			StiffnessMultipliers,
-			DampingMultipliers,
-			FSolverVec2(GetWeightedFloatXPBDBendingSpringStiffness(PropertyCollection, MaxStiffness)),
-			FSolverVec2(GetWeightedFloatXPBDBendingSpringDamping(PropertyCollection, MinDampingRatio)))
-		, XPBDBendingSpringStiffnessIndex(PropertyCollection)
-		, XPBDBendingSpringDampingIndex(PropertyCollection)
-	{}
-
 	virtual ~FXPBDBendingSpringConstraints() override = default;
 
 	CHAOS_API void SetProperties(
 		const FCollectionPropertyConstFacade& PropertyCollection,
 		const TMap<FString, TConstArrayView<FRealSingle>>& WeightMaps);
 
-	UE_DEPRECATED(5.3, "Use SetProperties(const FCollectionPropertyConstFacade&, const TMap<FString, TConstArrayView<FRealSingle>>&, FSolverReal) instead.")
-	void SetProperties(const FCollectionPropertyConstFacade& PropertyCollection)
-	{
-		SetProperties(PropertyCollection, TMap<FString, TConstArrayView<FRealSingle>>());
-	}
-
 private:
 	using FXPBDSpringConstraints::Constraints;
 	using FXPBDSpringConstraints::ParticleOffset;
