Iris
* Made NetRefHandleError enum a bitflag

#rnx

diff --git a/Engine/Source/Runtime/Experimental/Iris/Core/Private/Iris/ReplicationSystem/ObjectReplicationBridge.cpp b/Engine/Source/Runtime/Experimental/Iris/Core/Private/Iris/ReplicationSystem/ObjectReplicationBridge.cpp
--- a/Engine/Source/Runtime/Experimental/Iris/Core/Private/Iris/ReplicationSystem/ObjectReplicationBridge.cpp
+++ b/Engine/Source/Runtime/Experimental/Iris/Core/Private/Iris/ReplicationSystem/ObjectReplicationBridge.cpp
@@ -163,8 +163,6 @@ bool UObjectReplicationBridge::ShouldUseVerboseCsvStats() const
 
 void UObjectReplicationBridge::Initialize(UReplicationSystem* InReplicationSystem)
 {
-	ensureMsgf((GEnsureNetRefHandleError == -1 || GEnsureNetRefHandleError < (int32)UE::Net::ENetRefHandleError::Max), TEXT("GEnsureNetRefHandleError is set to an invalid value: %u"), GEnsureNetRefHandleError);
-
 	Super::Initialize(InReplicationSystem);
 
 	const uint32 CurrentMaxInternalIndex = NetRefHandleManager->GetCurrentMaxInternalNetRefIndex();
@@ -2102,7 +2100,7 @@ void UObjectReplicationBridge::OnErrorWithNetRefHandleReported(UE::Net::ENetRefH
 	// Ensure at the end so the log contains all the relevant information
 	ON_SCOPE_EXIT
 	{
-		if (GEnsureNetRefHandleError==0 || GEnsureNetRefHandleError==(int32)ErrorType)
+		if (GEnsureNetRefHandleError != -1 && (GEnsureNetRefHandleError==0 || EnumHasAnyFlags((ENetRefHandleError)GEnsureNetRefHandleError,ErrorType)))
 		{
 			// Use different calls for every error type to prevent only reporting the first error that occured.
 			switch(ErrorType)
diff --git a/Engine/Source/Runtime/Net/Core/Public/Net/Core/Connection/NetEnums.h b/Engine/Source/Runtime/Net/Core/Public/Net/Core/Connection/NetEnums.h
--- a/Engine/Source/Runtime/Net/Core/Public/Net/Core/Connection/NetEnums.h
+++ b/Engine/Source/Runtime/Net/Core/Public/Net/Core/Connection/NetEnums.h
@@ -123,11 +123,11 @@ namespace UE::Net
  */
 enum class ENetRefHandleError : uint32
 {
-	None = 0,
-	BitstreamCorrupted, // An error while reading detected an irrecoverable bitstream corruption. Client was forced to disconnect
-	ReplicationDisabled, // An error while reading was recoverable but prevents this NetObject from replicating.
-	Max, // always keep last
+	None = 0x00,
+	BitstreamCorrupted = 0x01, // An error while reading detected an irrecoverable bitstream corruption. Client was forced to disconnect
+	ReplicationDisabled = 0x02, // An error while reading was recoverable but prevents this NetObject from replicating.
 };
+ENUM_CLASS_FLAGS(ENetRefHandleError);
 
 inline const TCHAR* LexToString(ENetRefHandleError NetRefHandleError)
 {
@@ -136,7 +136,6 @@ inline const TCHAR* LexToString(ENetRefHandleError NetRefHandleError)
 		case ENetRefHandleError::None: return TEXT("None");
 		case ENetRefHandleError::BitstreamCorrupted: return TEXT("BitstreamCorrupted");
 		case ENetRefHandleError::ReplicationDisabled: return TEXT("ReplicationDisabled");
-		case ENetRefHandleError::Max: return TEXT("Max");
 		default: return TEXT("missing");
 	}
 }
