Dataflow array inputs

#presubmit

diff --git a/Engine/Plugins/ChaosClothAssetEditor/Source/ChaosClothAssetDataflowNodes/Private/ChaosClothAsset/AddWeightMapNode.cpp b/Engine/Plugins/ChaosClothAssetEditor/Source/ChaosClothAssetDataflowNodes/Private/ChaosClothAsset/AddWeightMapNode.cpp
--- a/Engine/Plugins/ChaosClothAssetEditor/Source/ChaosClothAssetDataflowNodes/Private/ChaosClothAsset/AddWeightMapNode.cpp
+++ b/Engine/Plugins/ChaosClothAssetEditor/Source/ChaosClothAssetDataflowNodes/Private/ChaosClothAsset/AddWeightMapNode.cpp
@@ -147,16 +147,12 @@ FChaosClothAssetAddWeightMapNode::FChaosClothAssetAddWeightMapNode(const Dataflo
 	: FDataflowTerminalNode(InParam, InGuid)
 {
 	RegisterInputConnection(&Collection);
-	if (FDataflowInput* const Input = RegisterInputConnection(&InputName.StringValue, GET_MEMBER_NAME_CHECKED(FChaosClothAssetConnectableIStringValue, StringValue)))
-	{
-		Input->SetCanHidePin(true);
-		Input->SetPinIsHidden(true);
-	}
-	if (FDataflowInput* const Input = RegisterInputConnection(&TransferCollection))
-	{
-		Input->SetCanHidePin(true);
-		Input->SetPinIsHidden(true);
-	}
+	RegisterInputConnection(&InputName.StringValue, GET_MEMBER_NAME_CHECKED(FChaosClothAssetConnectableIStringValue, StringValue))
+		.SetCanHidePin(true)
+		.SetPinIsHidden(true);
+	RegisterInputConnection(&TransferCollection)
+		.SetCanHidePin(true)
+		.SetPinIsHidden(true);
 	RegisterOutputConnection(&Collection, &Collection);
 	RegisterOutputConnection(&Name);
 }
diff --git a/Engine/Plugins/ChaosClothAssetEditor/Source/ChaosClothAssetDataflowNodes/Private/ChaosClothAsset/ChaosClothAssetDataflowNodesModule.cpp b/Engine/Plugins/ChaosClothAssetEditor/Source/ChaosClothAssetDataflowNodes/Private/ChaosClothAsset/ChaosClothAssetDataflowNodesModule.cpp
--- a/Engine/Plugins/ChaosClothAssetEditor/Source/ChaosClothAssetDataflowNodes/Private/ChaosClothAsset/ChaosClothAssetDataflowNodesModule.cpp
+++ b/Engine/Plugins/ChaosClothAssetEditor/Source/ChaosClothAssetDataflowNodes/Private/ChaosClothAsset/ChaosClothAssetDataflowNodesModule.cpp
@@ -85,7 +85,9 @@ namespace UE::Chaos::ClothAsset
 			DATAFLOW_NODE_REGISTER_CREATION_FACTORY(FChaosClothAssetDeleteElementNode);
 			DATAFLOW_NODE_REGISTER_CREATION_FACTORY(FChaosClothAssetImportNode);
 			DATAFLOW_NODE_REGISTER_CREATION_FACTORY(FChaosClothAssetMergeClothCollectionsNode);
+			DATAFLOW_NODE_REGISTER_CREATION_FACTORY(FChaosClothAssetMergeClothCollectionsNode_v2);
 			DATAFLOW_NODE_REGISTER_CREATION_FACTORY(FChaosClothAssetProxyDeformerNode);
+			DATAFLOW_NODE_REGISTER_CREATION_FACTORY(FChaosClothAssetProxyDeformerNode_v2);
 			DATAFLOW_NODE_REGISTER_CREATION_FACTORY(FChaosClothAssetRemeshNode);
 			DATAFLOW_NODE_REGISTER_CREATION_FACTORY(FChaosClothAssetReverseNormalsNode);
 			DATAFLOW_NODE_REGISTER_CREATION_FACTORY(FChaosClothAssetSelectionNode);
@@ -126,6 +128,7 @@ namespace UE::Chaos::ClothAsset
 			DATAFLOW_NODE_REGISTER_CREATION_FACTORY(FChaosClothAssetSkeletalMeshImportNode);
 			DATAFLOW_NODE_REGISTER_CREATION_FACTORY(FChaosClothAssetStaticMeshImportNode);
 			DATAFLOW_NODE_REGISTER_CREATION_FACTORY(FChaosClothAssetTerminalNode);
+			DATAFLOW_NODE_REGISTER_CREATION_FACTORY(FChaosClothAssetTerminalNode_v2);
 			DATAFLOW_NODE_REGISTER_CREATION_FACTORY(FChaosClothAssetTransferSkinWeightsNode);
 			DATAFLOW_NODE_REGISTER_CREATION_FACTORY(FChaosClothAssetTransformPositionsNode);
 			DATAFLOW_NODE_REGISTER_CREATION_FACTORY(FChaosClothAssetTransformUVsNode);
diff --git a/Engine/Plugins/ChaosClothAssetEditor/Source/ChaosClothAssetDataflowNodes/Private/ChaosClothAsset/MergeClothCollectionsNode.cpp b/Engine/Plugins/ChaosClothAssetEditor/Source/ChaosClothAssetDataflowNodes/Private/ChaosClothAsset/MergeClothCollectionsNode.cpp
--- a/Engine/Plugins/ChaosClothAssetEditor/Source/ChaosClothAssetDataflowNodes/Private/ChaosClothAsset/MergeClothCollectionsNode.cpp
+++ b/Engine/Plugins/ChaosClothAssetEditor/Source/ChaosClothAssetDataflowNodes/Private/ChaosClothAsset/MergeClothCollectionsNode.cpp
@@ -329,11 +329,175 @@ namespace UE::Chaos::ClothAsset::Private
 	}
 }
 
+FChaosClothAssetMergeClothCollectionsNode_v2::FChaosClothAssetMergeClothCollectionsNode_v2(const Dataflow::FNodeParameters& InParam, FGuid InGuid)
+	: FDataflowNode(InParam, InGuid)
+{
+	// Add two sets of pins to start.
+	AddPins();
+	AddPins();
+	RegisterOutputConnection(&Collection)
+		.SetPassthroughInput(GetConnectionReference(0));
+}
+
+void FChaosClothAssetMergeClothCollectionsNode_v2::Evaluate(Dataflow::FContext& Context, const FDataflowOutput* Out) const
+{
+	if (Out->IsA<FManagedArrayCollection>(&Collection))
+	{
+		using namespace UE::Chaos::ClothAsset;
+		using namespace Chaos::Softs;
+
+		// Evaluate in collection 0.
+		FManagedArrayCollection InCollection = GetValue<FManagedArrayCollection>(Context, GetConnectionReference(0));
+		const TSharedRef<FManagedArrayCollection> ClothCollection = MakeShared<FManagedArrayCollection>(MoveTemp(InCollection));
+
+		// Keep track of whether any of these collections are valid cloth collections
+		FCollectionClothFacade ClothFacade(ClothCollection);
+		bool bAreAnyValid = ClothFacade.IsValid();
+
+		// Make it a valid cloth collection if needed
+		if (!bAreAnyValid)
+		{
+			ClothFacade.DefineSchema();
+		}
+
+		FCollectionPropertyMutableFacade PropertyFacade(ClothCollection);
+		bAreAnyValid |= PropertyFacade.IsValid();
+
+		FCollectionClothSelectionFacade SelectionFacade(ClothCollection);
+		bAreAnyValid |= SelectionFacade.IsValid();
+
+		// Iterate through the inputs and append them to LOD 0
+		for (int32 InputIndex = 1; InputIndex < Collections.Num(); ++InputIndex)
+		{
+			FManagedArrayCollection OtherCollection = GetValue<FManagedArrayCollection>(Context, GetConnectionReference(InputIndex));  // Can't use a const reference here sadly since the facade needs a SharedRef to be created
+			const TSharedRef<const FManagedArrayCollection> OtherClothCollection = MakeShared<const FManagedArrayCollection>(MoveTemp(OtherCollection));
+
+			// Selections need to update with offsets. Gather offsets before appending cloth data.
+			const FCollectionClothSelectionConstFacade OtherSelectionFacade(OtherClothCollection);
+			TMap<FName, int32> GroupNameOffsets;
+			if (OtherSelectionFacade.IsValid())
+			{
+				const TArray<FName> SelectionNames = OtherSelectionFacade.GetNames();
+				for (const FName& SelectionName : SelectionNames)
+				{
+					const FName GroupName = OtherSelectionFacade.GetSelectionGroup(SelectionName);
+					if (!GroupNameOffsets.Find(GroupName))
+					{
+						GroupNameOffsets.Add(GroupName) = ClothCollection->NumElements(GroupName); // NumElements will return zero if the group doesn't exist.
+					}
+				}
+			}
+
+			// Append cloth
+			const FCollectionClothConstFacade OtherClothFacade(OtherClothCollection);
+			if (OtherClothFacade.IsValid() && Private::AreSkeletalMeshesCompatible(*this, ClothFacade, OtherClothFacade))
+			{
+				ClothFacade.Append(OtherClothFacade);
+				bAreAnyValid = true;
+			}
+
+			// Append selections (with offsets)
+			if (OtherSelectionFacade.IsValid())
+			{
+				constexpr bool bUpdateExistingSelections = true; // Want last one wins.
+				SelectionFacade.AppendWithOffsets(OtherSelectionFacade, bUpdateExistingSelections, GroupNameOffsets);
+				bAreAnyValid = true;
+			}
+
+			// Copy properties
+			const FCollectionPropertyConstFacade OtherPropertyFacade(OtherClothCollection);
+			if (OtherPropertyFacade.IsValid())
+			{
+				// Change that boolean to come back to the old behavior
+				static constexpr bool bOverrideProperties = false;
+				if(bOverrideProperties)
+				{
+					constexpr bool bUpdateExistingProperties = true; // Want last one wins.
+					PropertyFacade.Append(OtherClothCollection.ToSharedPtr(), bUpdateExistingProperties);
+				}
+				else
+				{
+					Private::AppendInputProperties(*this, OtherClothFacade, ClothFacade, OtherPropertyFacade, PropertyFacade);
+				}
+				bAreAnyValid = true;
+			}
+		}
+
+		// Set the output
+		if (bAreAnyValid)
+		{
+			// Use the merged cloth collection, but only if there were at least one valid input cloth collections
+			SetValue(Context, MoveTemp(*ClothCollection), &Collection);
+		}
+		else
+		{
+			// Otherwise pass through the first input unchanged
+			ForwardInput(Context, GetConnectionReference(0), &Collection);
+		}
+	}
+}
+
+TArray<Dataflow::FPin> FChaosClothAssetMergeClothCollectionsNode_v2::AddPins()
+{
+	const int32 Index = Collections.AddDefaulted();
+	const FDataflowInput& Input = RegisterInputArrayConnection(GetConnectionReference(Index));
+	return { { Dataflow::FPin::EDirection::INPUT, Input.GetType(), Input.GetName() } };
+}
+
+TArray<Dataflow::FPin> FChaosClothAssetMergeClothCollectionsNode_v2::GetPinsToRemove() const
+{
+	const int32 Index = Collections.Num() - 1;
+	check(Collections.IsValidIndex(Index));
+	if (const FDataflowInput* const Input = FindInput(GetConnectionReference(Index)))
+	{
+		return { { Dataflow::FPin::EDirection::INPUT, Input->GetType(), Input->GetName() } };
+	}
+	return Super::GetPinsToRemove();
+}
+
+void FChaosClothAssetMergeClothCollectionsNode_v2::OnPinRemoved(const Dataflow::FPin& Pin)
+{
+	const int32 Index = Collections.Num() - 1;
+	check(Collections.IsValidIndex(Index));
+#if DO_CHECK
+	const FDataflowInput* const Input = FindInput(GetConnectionReference(Index));
+	check(Input);
+	check(Input->GetName() == Pin.Name);
+	check(Input->GetType() == Pin.Type);
+#endif
+	Collections.SetNum(Index);
+
+	return Super::OnPinRemoved(Pin);
+}
+
+void FChaosClothAssetMergeClothCollectionsNode_v2::Serialize(FArchive& Ar)
+{
+	if (Ar.IsLoading())
+	{
+		check(Collections.Num() > 0);
+		check(FindInput(GetConnectionReference(0)));
+
+		for (int32 Index = 1; Index < Collections.Num(); ++Index)
+		{
+			RegisterInputArrayConnection(GetConnectionReference(Index));
+		}
+	}
+}
+
+Dataflow::TConnectionReference<FManagedArrayCollection> FChaosClothAssetMergeClothCollectionsNode_v2::GetConnectionReference(int32 Index) const
+{
+	return { &Collections[Index], Index, &Collections };
+}
+
+
+
+
 FChaosClothAssetMergeClothCollectionsNode::FChaosClothAssetMergeClothCollectionsNode(const Dataflow::FNodeParameters& InParam, FGuid InGuid)
 	: FDataflowNode(InParam, InGuid)
 {
 	RegisterInputConnection(&Collection);
-	RegisterOutputConnection(&Collection, &Collection);
+	RegisterOutputConnection(&Collection)
+		.SetPassthroughInput(&Collection);
 }
 
 void FChaosClothAssetMergeClothCollectionsNode::Evaluate(Dataflow::FContext& Context, const FDataflowOutput* Out) const
@@ -408,7 +572,7 @@ void FChaosClothAssetMergeClothCollectionsNode::Evaluate(Dataflow::FContext& Con
 			{
 				// Change that boolean to come back to the old behavior
 				static constexpr bool bOverrideProperties = false;
-				if(bOverrideProperties)
+				if (bOverrideProperties)
 				{
 					constexpr bool bUpdateExistingProperties = true; // Want last one wins.
 					PropertyFacade.Append(OtherClothCollection.ToSharedPtr(), bUpdateExistingProperties);
@@ -436,13 +600,13 @@ void FChaosClothAssetMergeClothCollectionsNode::Evaluate(Dataflow::FContext& Con
 	}
 }
 
-Dataflow::FPin FChaosClothAssetMergeClothCollectionsNode::AddPin()
+TArray<Dataflow::FPin> FChaosClothAssetMergeClothCollectionsNode::AddPins()
 {
-	auto AddInput = [this](const FManagedArrayCollection* InCollection) -> Dataflow::FPin
+	auto AddInput = [this](const FManagedArrayCollection* InCollection) -> TArray<Dataflow::FPin>
 	{
 		RegisterInputConnection(InCollection);
 		const FDataflowInput* const Input = FindInput(InCollection);
-		return { Dataflow::FPin::EDirection::INPUT, Input->GetType(), Input->GetName() };
+		return { { Dataflow::FPin::EDirection::INPUT, Input->GetType(), Input->GetName() } };
 	};
 
 	switch (NumInputs)
@@ -455,16 +619,16 @@ Dataflow::FPin FChaosClothAssetMergeClothCollectionsNode::AddPin()
 	default: break;
 	}
 
-	return Super::AddPin();
+	return Super::AddPins();
 }
 
-Dataflow::FPin FChaosClothAssetMergeClothCollectionsNode::GetPinToRemove() const
+TArray<Dataflow::FPin> FChaosClothAssetMergeClothCollectionsNode::GetPinsToRemove() const
 {
-	auto PinToRemove = [this](const FManagedArrayCollection* InCollection) -> Dataflow::FPin
+	auto PinToRemove = [this](const FManagedArrayCollection* InCollection) -> TArray<Dataflow::FPin>
 	{
 		const FDataflowInput* const Input = FindInput(InCollection);
 		check(Input);
-		return { Dataflow::FPin::EDirection::INPUT, Input->GetType(), Input->GetName() };
+		return { { Dataflow::FPin::EDirection::INPUT, Input->GetType(), Input->GetName() } };
 	};
 
 	switch (NumInputs - 1)
@@ -476,7 +640,7 @@ Dataflow::FPin FChaosClothAssetMergeClothCollectionsNode::GetPinToRemove() const
 	case 5: return PinToRemove(&Collection5);
 	default: break;
 	}
-	return Super::GetPinToRemove();
+	return Super::GetPinsToRemove();
 }
 
 void FChaosClothAssetMergeClothCollectionsNode::OnPinRemoved(const Dataflow::FPin& Pin)
@@ -551,7 +715,7 @@ void FChaosClothAssetMergeClothCollectionsNode::Serialize(FArchive& Ar)
 		NumInputs = 1;  // AddPin will increment it again
 		for (int32 InputIndex = 0; InputIndex < NumInputsToAdd; ++InputIndex)
 		{
-			AddPin();
+			AddPins();
 		}
 		check(NumInputsToAdd == NumInputs - 1);
 	}
diff --git a/Engine/Plugins/ChaosClothAssetEditor/Source/ChaosClothAssetDataflowNodes/Private/ChaosClothAsset/ProxyDeformerNode.cpp b/Engine/Plugins/ChaosClothAssetEditor/Source/ChaosClothAssetDataflowNodes/Private/ChaosClothAsset/ProxyDeformerNode.cpp
--- a/Engine/Plugins/ChaosClothAssetEditor/Source/ChaosClothAssetDataflowNodes/Private/ChaosClothAsset/ProxyDeformerNode.cpp
+++ b/Engine/Plugins/ChaosClothAssetEditor/Source/ChaosClothAssetDataflowNodes/Private/ChaosClothAsset/ProxyDeformerNode.cpp
@@ -159,7 +159,7 @@ namespace UE::Chaos::ClothAsset::Private
 		return FPointWeightMap(ClothFacade.GetNumSimVertices3D(), SelectedValue);
 	}
 
-	static TArray<ClothingMeshUtils::FMeshToMeshFilterSet> SelectionsToMeshToMeshFilterSets(const FCollectionClothConstFacade& ClothFacade, const FCollectionClothSelectionConstFacade& SelectionFacade, const TArray<FName>& SelectionNames)
+	static TArray<ClothingMeshUtils::FMeshToMeshFilterSet> SelectionsToMeshToMeshFilterSets_v2(const FCollectionClothConstFacade& ClothFacade, const FCollectionClothSelectionConstFacade& SelectionFacade, const TArray<TPair<FName, FName>>& SelectionNames)
 	{
 		auto GetSimFaceSelection = [&ClothFacade](const FName& SelectionGroup, const TSet<int32>& SelectionSet) -> TSet<int32>
 			{
@@ -229,6 +229,118 @@ namespace UE::Chaos::ClothAsset::Private
 		// Fill up the MeshToMeshFilterSets
 		TArray<ClothingMeshUtils::FMeshToMeshFilterSet> MeshToMeshFilterSets;
 
+		if (SelectionFacade.IsValid())
+		{
+			MeshToMeshFilterSets.Reserve(SelectionNames.Num());
+
+			for (const TPair<FName,FName>& SelectionName : SelectionNames)
+			{
+				if (const TSet<int32>* RenderSelectionSet = SelectionFacade.FindSelectionSet(SelectionName.Get<0>()))
+				{
+					if (const TSet<int32>* SimSelectionSet = SelectionFacade.FindSelectionSet(SelectionName.Get<1>()))
+					{
+						if (RenderSelectionSet->Num() && SimSelectionSet->Num())
+						{
+							const FName RenderSelectionGroup = SelectionFacade.GetSelectionGroup(SelectionName.Get<0>());
+							const FName SimSelectionGroup = SelectionFacade.GetSelectionGroup(SelectionName.Get<1>());
+
+							// Retrieve the sim face selection
+							TSet<int32> SimFaceSelection = GetSimFaceSelection(SimSelectionGroup, *SimSelectionSet);
+							if (!SimFaceSelection.Num())
+							{
+								continue;  // Nothing selected on the simulation side
+							}
+
+							// Retrieve the render vertex selection
+							TSet<int32> RenderVertexSelection = GetRenderVertexSelection(RenderSelectionGroup, *RenderSelectionSet);
+							if (!RenderVertexSelection.Num())
+							{
+								continue;  // Nothing selected on the render side
+							}
+
+							ClothingMeshUtils::FMeshToMeshFilterSet& MeshToMeshFilterSet = MeshToMeshFilterSets.AddDefaulted_GetRef();
+							MeshToMeshFilterSet.SourceTriangles = MoveTemp(SimFaceSelection);
+							MeshToMeshFilterSet.TargetVertices = MoveTemp(RenderVertexSelection);
+						}
+					}
+				}
+			}
+		}
+
+		return MeshToMeshFilterSets;
+	}
+	
+	
+	static TArray<ClothingMeshUtils::FMeshToMeshFilterSet> SelectionsToMeshToMeshFilterSets(const FCollectionClothConstFacade& ClothFacade, const FCollectionClothSelectionConstFacade& SelectionFacade, const TArray<FName>& SelectionNames)
+	{
+		auto GetSimFaceSelection = [&ClothFacade](const FName& SelectionGroup, const TSet<int32>& SelectionSet) -> TSet<int32>
+		{
+			TSet<int32> SimFaceSelection;
+			if (SelectionGroup == ClothCollectionGroup::SimVertices2D)
+			{
+				const TConstArrayView<FIntVector3> SimIndices2D = ClothFacade.GetSimIndices2D();
+				SimFaceSelection.Reserve(SelectionSet.Num());
+
+				for (int32 FaceIndex = 0; FaceIndex < SimIndices2D.Num(); ++FaceIndex)
+				{
+					const FIntVector3& Indices = SimIndices2D[FaceIndex];
+
+					if (SelectionSet.Contains(Indices[0]) &&
+						SelectionSet.Contains(Indices[1]) &&
+						SelectionSet.Contains(Indices[2]))
+					{
+						SimFaceSelection.Add(FaceIndex);
+					}
+				}
+			}
+			else if (SelectionGroup == ClothCollectionGroup::SimVertices3D)
+			{
+				const TConstArrayView<FIntVector3> SimIndices3D = ClothFacade.GetSimIndices3D();
+				SimFaceSelection.Reserve(SelectionSet.Num());
+
+				for (int32 FaceIndex = 0; FaceIndex < SimIndices3D.Num(); ++FaceIndex)
+				{
+					const FIntVector3& Indices = SimIndices3D[FaceIndex];
+
+					if (SelectionSet.Contains(Indices[0]) &&
+						SelectionSet.Contains(Indices[1]) &&
+						SelectionSet.Contains(Indices[2]))
+					{
+						SimFaceSelection.Add(FaceIndex);
+					}
+				}
+			}
+			else if (SelectionGroup == ClothCollectionGroup::SimFaces)
+			{
+				SimFaceSelection = SelectionSet;
+			}
+			return SimFaceSelection;
+		};
+
+		auto GetRenderVertexSelection = [&ClothFacade](const FName& SelectionGroup, const TSet<int32>& SelectionSet) -> TSet<int32>
+		{
+			TSet<int32> RenderVertexSelection;
+			if (SelectionGroup == ClothCollectionGroup::RenderVertices)
+			{
+				RenderVertexSelection = SelectionSet;
+			}
+			else if (SelectionGroup == ClothCollectionGroup::RenderFaces)
+			{
+				RenderVertexSelection.Reserve(SelectionSet.Num() * 3);
+				const TConstArrayView<FIntVector3> RenderIndices = ClothFacade.GetRenderIndices();
+				for (const int32 FaceIndex : SelectionSet)
+				{
+					RenderVertexSelection.Add(RenderIndices[FaceIndex][0]);
+					RenderVertexSelection.Add(RenderIndices[FaceIndex][1]);
+					RenderVertexSelection.Add(RenderIndices[FaceIndex][2]);
+				}
+			}
+			return RenderVertexSelection;
+		};
+
+		// Fill up the MeshToMeshFilterSets
+		TArray<ClothingMeshUtils::FMeshToMeshFilterSet> MeshToMeshFilterSets;
+
 		if (SelectionFacade.IsValid())
 		{
 			MeshToMeshFilterSets.Reserve(SelectionNames.Num());
@@ -280,6 +392,188 @@ namespace UE::Chaos::ClothAsset::Private
 
 }  // End namespace UE::Chaos::ClothAsset::Private
 
+FChaosClothAssetProxyDeformerNode_v2::FChaosClothAssetProxyDeformerNode_v2(const Dataflow::FNodeParameters& InParam, FGuid InGuid)
+	: FDataflowNode(InParam, InGuid)
+{
+	using namespace UE::Chaos::ClothAsset;
+	SimVertexSelection.StringValue = FString();  // An empty selection is an accepted input, but a non existing one isn't
+	SkinningBlendName = ClothCollectionAttribute::RenderDeformerSkinningBlend.ToString();
+
+	RegisterInputConnection(&Collection);
+	RegisterInputConnection(&SimVertexSelection.StringValue, GET_MEMBER_NAME_CHECKED(FChaosClothAssetConnectableIStringValue, StringValue));
+	RegisterOutputConnection(&Collection)
+		.SetPassthroughInput(&Collection);
+	RegisterOutputConnection(&SkinningBlendName);
+
+	// Start with one set of option pins.
+	AddPins();
+}
+
+void FChaosClothAssetProxyDeformerNode_v2::Evaluate(Dataflow::FContext& Context, const FDataflowOutput* Out) const
+{
+	if (Out->IsA<FManagedArrayCollection>(&Collection))
+	{
+		using namespace UE::Chaos::ClothAsset;
+
+		// Update the selection names override
+		SimVertexSelection.StringValue_Override = GetValue<FString>(Context, &SimVertexSelection.StringValue, UE::Chaos::ClothAsset::FWeightMapTools::NotOverridden);
+		for (int32 FilterSetIndex = 0; FilterSetIndex < SelectionFilterSets.Num(); ++FilterSetIndex)
+		{
+			SelectionFilterSets[FilterSetIndex].RenderSelection.StringValue_Override = GetValue<FString>(Context, GetRenderConnectionReference(FilterSetIndex), UE::Chaos::ClothAsset::FWeightMapTools::NotOverridden);
+			SelectionFilterSets[FilterSetIndex].SimSelection.StringValue_Override = GetValue<FString>(Context, GetSimConnectionReference(FilterSetIndex), UE::Chaos::ClothAsset::FWeightMapTools::NotOverridden);
+		}
+
+		// Evaluate in collection
+		FManagedArrayCollection InCollection = GetValue<FManagedArrayCollection>(Context, &Collection);
+		const TSharedRef<FManagedArrayCollection> ClothCollection = MakeShared<FManagedArrayCollection>(MoveTemp(InCollection));
+
+		// Always check for a valid cloth collection/facade/sim mesh to avoid processing non cloth collections or pure render mesh cloth assets
+		FCollectionClothFacade ClothFacade(ClothCollection);
+		if (ClothFacade.IsValid() && ClothFacade.HasValidData())
+		{
+			FCollectionClothSelectionFacade SelectionFacade(ClothCollection);
+
+			// Retrieve the SimVertexSelection name
+			FName SimVertexSelectionName = FName(*GetValue<FString>(Context, &SimVertexSelection.StringValue));
+			if (SimVertexSelectionName != NAME_None && (!SelectionFacade.IsValid() || !SelectionFacade.FindSelectionSet(SimVertexSelectionName)))
+			{
+				FClothDataflowTools::LogAndToastWarning(*this,
+					LOCTEXT("HasSimVertexSelectionHeadline", "Unknown SimVertexSelection."),
+					LOCTEXT("HasSimVertexSelectionDetails", "The specified SimVertexSelection does't exist within the input Cloth Collection."));
+				SimVertexSelectionName = NAME_None;
+			}
+
+			// Add the optional render deformer schema
+			if (!ClothFacade.IsValid(EClothCollectionOptionalSchemas::RenderDeformer))
+			{
+				ClothFacade.DefineSchema(EClothCollectionOptionalSchemas::RenderDeformer);
+			}
+
+			// Create the render weight map for storing the skinning blend weights
+			Private::FDeformerMappingDataGenerator DeformerMappingDataGenerator;
+			DeformerMappingDataGenerator.SimPositions = ClothFacade.GetSimPosition3D();
+			DeformerMappingDataGenerator.SimIndices = ClothFacade.GetSimIndices3D();
+			DeformerMappingDataGenerator.RenderPositions = ClothFacade.GetRenderPosition();
+			DeformerMappingDataGenerator.RenderNormals = ClothFacade.GetRenderNormal();
+			DeformerMappingDataGenerator.RenderIndices = ClothFacade.GetRenderIndices();
+			DeformerMappingDataGenerator.PointWeightMap = Private::SelectionToPointWeightMap(ClothFacade, SelectionFacade, SimVertexSelectionName);
+			DeformerMappingDataGenerator.MeshToMeshFilterSet = Private::SelectionsToMeshToMeshFilterSets_v2(ClothFacade, SelectionFacade, GetSelectionFilterNames(Context));
+			DeformerMappingDataGenerator.RenderDeformerPositionBaryCoordsAndDist = ClothFacade.GetRenderDeformerPositionBaryCoordsAndDist();
+			DeformerMappingDataGenerator.RenderDeformerNormalBaryCoordsAndDist = ClothFacade.GetRenderDeformerNormalBaryCoordsAndDist();
+			DeformerMappingDataGenerator.RenderDeformerTangentBaryCoordsAndDist = ClothFacade.GetRenderDeformerTangentBaryCoordsAndDist();
+			DeformerMappingDataGenerator.RenderDeformerSimIndices3D = ClothFacade.GetRenderDeformerSimIndices3D();
+			DeformerMappingDataGenerator.RenderDeformerWeight = ClothFacade.GetRenderDeformerWeight();
+			DeformerMappingDataGenerator.RenderDeformerSkinningBlend = ClothFacade.GetRenderDeformerSkinningBlend();
+
+			const int32 NumInfluences = DeformerMappingDataGenerator.Generate(bUseSmoothTransition, bUseMultipleInfluences, InfluenceRadius);
+
+			for (int32 RenderPatternIndex = 0; RenderPatternIndex < ClothFacade.GetNumRenderPatterns(); ++RenderPatternIndex)
+			{
+				FCollectionClothRenderPatternFacade RenderPatternFacade = ClothFacade.GetRenderPattern(RenderPatternIndex);
+				RenderPatternFacade.SetRenderDeformerNumInfluences(NumInfluences);
+			}
+		}
+
+		SetValue(Context, MoveTemp(*ClothCollection), &Collection);
+	}
+	else if (Out->IsA<FString>(&SkinningBlendName))
+	{
+		SetValue(Context, SkinningBlendName, &SkinningBlendName);
+	}
+}
+
+TArray<Dataflow::FPin> FChaosClothAssetProxyDeformerNode_v2::AddPins()
+{
+	const int32 Index = SelectionFilterSets.AddDefaulted();
+	TArray<Dataflow::FPin> Pins;
+	Pins.Reserve(2);
+	{
+		const FDataflowInput& Input = RegisterInputArrayConnection(GetRenderConnectionReference(Index), GET_MEMBER_NAME_CHECKED(FChaosClothAssetConnectableIStringValue, StringValue));
+		Pins.Emplace(Dataflow::FPin{ Dataflow::FPin::EDirection::INPUT, Input.GetType(), Input.GetName() });
+	}
+	{
+		const FDataflowInput& Input = RegisterInputArrayConnection(GetSimConnectionReference(Index), GET_MEMBER_NAME_CHECKED(FChaosClothAssetConnectableIStringValue, StringValue));
+		Pins.Emplace(Dataflow::FPin{ Dataflow::FPin::EDirection::INPUT, Input.GetType(), Input.GetName() });
+	}
+	return Pins;
+}
+
+TArray<Dataflow::FPin> FChaosClothAssetProxyDeformerNode_v2::GetPinsToRemove() const
+{
+	const int32 Index = SelectionFilterSets.Num() - 1;
+	check(SelectionFilterSets.IsValidIndex(Index));
+	TArray<Dataflow::FPin> Pins;
+	Pins.Reserve(2);
+	if (const FDataflowInput* const Input = FindInput(GetRenderConnectionReference(Index)))
+	{
+		Pins.Emplace(Dataflow::FPin{ Dataflow::FPin::EDirection::INPUT, Input->GetType(), Input->GetName() });
+	}
+	if (const FDataflowInput* const Input = FindInput(GetSimConnectionReference(Index)))
+	{
+		Pins.Emplace(Dataflow::FPin{ Dataflow::FPin::EDirection::INPUT, Input->GetType(), Input->GetName() });
+	}
+	return Pins;
+}
+
+void FChaosClothAssetProxyDeformerNode_v2::OnPinRemoved(const Dataflow::FPin& Pin)
+{
+	const int32 Index = SelectionFilterSets.Num() - 1;
+	check(SelectionFilterSets.IsValidIndex(Index));
+	const FDataflowInput* const FirstInput = FindInput(GetRenderConnectionReference(Index));
+	const FDataflowInput* const SecondInput = FindInput(GetSimConnectionReference(Index));
+	check(FirstInput || SecondInput);
+	const bool bIsFirstInput = FirstInput && FirstInput->GetName() == Pin.Name;
+	const bool bIsSecondInput = SecondInput && SecondInput->GetName() == Pin.Name;
+	if (bIsFirstInput && !SecondInput || bIsSecondInput && !FirstInput)
+	{
+		// Both inputs removed. Remove array index.
+		SelectionFilterSets.SetNum(Index);
+	}
+	return Super::OnPinRemoved(Pin);
+}
+
+void FChaosClothAssetProxyDeformerNode_v2::Serialize(FArchive& Ar)
+{
+	// Restore the pins when re-loading so they can get properly reconnected
+	if (Ar.IsLoading())
+	{
+		check(SelectionFilterSets.Num() > 0);
+		check(FindInput(GetRenderConnectionReference(0)));
+		check(FindInput(GetSimConnectionReference(0)));
+
+		for (int32 Index = 1; Index < SelectionFilterSets.Num(); ++Index)
+		{
+			RegisterInputArrayConnection(GetRenderConnectionReference(Index), GET_MEMBER_NAME_CHECKED(FChaosClothAssetConnectableIStringValue, StringValue));
+			RegisterInputArrayConnection(GetSimConnectionReference(Index), GET_MEMBER_NAME_CHECKED(FChaosClothAssetConnectableIStringValue, StringValue));
+		}
+	}
+}
+
+TArray<TPair<FName, FName>> FChaosClothAssetProxyDeformerNode_v2::GetSelectionFilterNames(Dataflow::FContext& Context) const
+{
+	TArray<TPair<FName, FName>> SelectionFilterNames;
+	SelectionFilterNames.SetNumUninitialized(SelectionFilterSets.Num());
+
+	for (int32 Index = 0; Index < SelectionFilterSets.Num(); ++Index)
+	{
+		SelectionFilterNames[Index] = MakeTuple<FName, FName>(
+			FName(*GetValue(Context, GetRenderConnectionReference(Index))),
+			FName(*GetValue(Context, GetSimConnectionReference(Index))));
+	}
+	return SelectionFilterNames;
+}
+
+Dataflow::TConnectionReference<FString> FChaosClothAssetProxyDeformerNode_v2::GetRenderConnectionReference(int32 Index) const
+{
+	return { &SelectionFilterSets[Index].RenderSelection.StringValue, Index, &SelectionFilterSets };
+}
+
+Dataflow::TConnectionReference<FString> FChaosClothAssetProxyDeformerNode_v2::GetSimConnectionReference(int32 Index) const
+{
+	return { &SelectionFilterSets[Index].SimSelection.StringValue, Index, &SelectionFilterSets };
+}
+
+
 FChaosClothAssetProxyDeformerNode::FChaosClothAssetProxyDeformerNode(const Dataflow::FNodeParameters& InParam, FGuid InGuid)
 	: FDataflowNode(InParam, InGuid)
 {
@@ -290,7 +584,8 @@ FChaosClothAssetProxyDeformerNode::FChaosClothAssetProxyDeformerNode(const Dataf
 	RegisterInputConnection(&Collection);
 	RegisterInputConnection(&SimVertexSelection.StringValue, GET_MEMBER_NAME_CHECKED(FChaosClothAssetConnectableIStringValue, StringValue));
 	RegisterInputConnection(&SelectionFilterSet0.StringValue, GET_MEMBER_NAME_CHECKED(FChaosClothAssetConnectableIStringValue, StringValue));
-	RegisterOutputConnection(&Collection, &Collection);
+	RegisterOutputConnection(&Collection)
+		.SetPassthroughInput(&Collection);
 	RegisterOutputConnection(&SkinningBlendName);
 }
 
@@ -369,7 +664,7 @@ void FChaosClothAssetProxyDeformerNode::Evaluate(Dataflow::FContext& Context, co
 	}
 }
 
-Dataflow::FPin FChaosClothAssetProxyDeformerNode::AddPin()
+TArray<Dataflow::FPin> FChaosClothAssetProxyDeformerNode::AddPins()
 {
 	check(NumFilterSets > 0);
 	const FChaosClothAssetConnectableStringValue* const SelectionFilterSet = Get1To9SelectionFilterSets()[NumFilterSets - 1];
@@ -378,16 +673,16 @@ Dataflow::FPin FChaosClothAssetProxyDeformerNode::AddPin()
 	++NumFilterSets;
 	const FDataflowInput* const Input = FindInput(SelectionFilterSet);
 	check(Input);
-	return { Dataflow::FPin::EDirection::INPUT, Input->GetType(), Input->GetName() };
+	return { { Dataflow::FPin::EDirection::INPUT, Input->GetType(), Input->GetName() } };
 }
 
-Dataflow::FPin FChaosClothAssetProxyDeformerNode::GetPinToRemove() const
+TArray<Dataflow::FPin> FChaosClothAssetProxyDeformerNode::GetPinsToRemove() const
 {
 	check(NumFilterSets > 1);
 	const FChaosClothAssetConnectableStringValue* const SelectionFilterSet = Get1To9SelectionFilterSets()[NumFilterSets - 2];
 	const FDataflowInput* const Input = FindInput(SelectionFilterSet);
 	check(Input);
-	return { Dataflow::FPin::EDirection::INPUT, Input->GetType(), Input->GetName() };
+	return { { Dataflow::FPin::EDirection::INPUT, Input->GetType(), Input->GetName() } };
 }
 
 void FChaosClothAssetProxyDeformerNode::OnPinRemoved(const Dataflow::FPin& Pin)
@@ -414,7 +709,7 @@ void FChaosClothAssetProxyDeformerNode::Serialize(FArchive& Ar)
 		NumFilterSets = 1;  // Reset to default, add pin will increment it again 
 		for (int32 Index = 0; Index < NumFilterSetsToAdd; ++Index)
 		{
-			AddPin();
+			AddPins();
 		}
 		ensure(NumFilterSetsToAdd == (NumFilterSets - 1));
 	}
diff --git a/Engine/Plugins/ChaosClothAssetEditor/Source/ChaosClothAssetDataflowNodes/Private/ChaosClothAsset/SelectionNode.cpp b/Engine/Plugins/ChaosClothAssetEditor/Source/ChaosClothAssetDataflowNodes/Private/ChaosClothAsset/SelectionNode.cpp
--- a/Engine/Plugins/ChaosClothAssetEditor/Source/ChaosClothAssetDataflowNodes/Private/ChaosClothAsset/SelectionNode.cpp
+++ b/Engine/Plugins/ChaosClothAssetEditor/Source/ChaosClothAssetDataflowNodes/Private/ChaosClothAsset/SelectionNode.cpp
@@ -229,16 +229,12 @@ FChaosClothAssetSelectionNode::FChaosClothAssetSelectionNode(const Dataflow::FNo
 	: FDataflowTerminalNode(InParam, InGuid)
 {
 	RegisterInputConnection(&Collection);
-	if (FDataflowInput* const Input = RegisterInputConnection(&InputName.StringValue, GET_MEMBER_NAME_CHECKED(FChaosClothAssetConnectableIStringValue, StringValue)))
-	{
-		Input->SetCanHidePin(true);
-		Input->SetPinIsHidden(true);
-	}
-	if (FDataflowInput* const Input = RegisterInputConnection(&TransferCollection))
-	{
-		Input->SetCanHidePin(true);
-		Input->SetPinIsHidden(true);
-	}
+	RegisterInputConnection(&InputName.StringValue, GET_MEMBER_NAME_CHECKED(FChaosClothAssetConnectableIStringValue, StringValue))
+		.SetCanHidePin(true)
+		.SetPinIsHidden(true);
+	RegisterInputConnection(&TransferCollection)
+		.SetCanHidePin(true)
+		.SetPinIsHidden(true);
 	RegisterOutputConnection(&Collection, &Collection);
 	RegisterOutputConnection(&Name);
 }
diff --git a/Engine/Plugins/ChaosClothAssetEditor/Source/ChaosClothAssetDataflowNodes/Private/ChaosClothAsset/SelectionToIntMapNode.cpp b/Engine/Plugins/ChaosClothAssetEditor/Source/ChaosClothAssetDataflowNodes/Private/ChaosClothAsset/SelectionToIntMapNode.cpp
--- a/Engine/Plugins/ChaosClothAssetEditor/Source/ChaosClothAssetDataflowNodes/Private/ChaosClothAsset/SelectionToIntMapNode.cpp
+++ b/Engine/Plugins/ChaosClothAssetEditor/Source/ChaosClothAssetDataflowNodes/Private/ChaosClothAsset/SelectionToIntMapNode.cpp
@@ -18,7 +18,7 @@ FChaosClothAssetSelectionToIntMapNode::FChaosClothAssetSelectionToIntMapNode(con
 	RegisterOutputConnection(&Collection, &Collection);
 	RegisterInputConnection(&SelectionName.StringValue, GET_MEMBER_NAME_CHECKED(FChaosClothAssetConnectableIStringValue, StringValue));
 	RegisterInputConnection(&IntMapName.StringValue, GET_MEMBER_NAME_CHECKED(FChaosClothAssetConnectableIOStringValue, StringValue));
-	RegisterOutputConnection(&IntMapName.StringValue, &IntMapName.StringValue, GET_MEMBER_NAME_CHECKED(FChaosClothAssetConnectableIOStringValue, StringValue), GET_MEMBER_NAME_CHECKED(FChaosClothAssetConnectableIOStringValue, StringValue));
+	RegisterOutputConnection(&IntMapName.StringValue, &IntMapName.StringValue, GET_MEMBER_NAME_CHECKED(FChaosClothAssetConnectableIOStringValue, StringValue));
 }
 
 void FChaosClothAssetSelectionToIntMapNode::Evaluate(Dataflow::FContext& Context, const FDataflowOutput* Out) const
diff --git a/Engine/Plugins/ChaosClothAssetEditor/Source/ChaosClothAssetDataflowNodes/Private/ChaosClothAsset/SimulationAerodynamicsConfigNode.cpp b/Engine/Plugins/ChaosClothAssetEditor/Source/ChaosClothAssetDataflowNodes/Private/ChaosClothAsset/SimulationAerodynamicsConfigNode.cpp
--- a/Engine/Plugins/ChaosClothAssetEditor/Source/ChaosClothAssetDataflowNodes/Private/ChaosClothAsset/SimulationAerodynamicsConfigNode.cpp
+++ b/Engine/Plugins/ChaosClothAssetEditor/Source/ChaosClothAssetDataflowNodes/Private/ChaosClothAsset/SimulationAerodynamicsConfigNode.cpp
@@ -10,16 +10,12 @@ FChaosClothAssetSimulationAerodynamicsConfigNode::FChaosClothAssetSimulationAero
 	: FChaosClothAssetSimulationBaseConfigNode(InParam, InGuid)
 {
 	RegisterCollectionConnections();
-	if (FDataflowInput* const Input = RegisterInputConnection(&Drag.WeightMap))
-	{
-		Input->SetCanHidePin(true);
-		Input->SetPinIsHidden(true);
-	}
-	if (FDataflowInput* const Input = RegisterInputConnection(&Lift.WeightMap))
-	{
-		Input->SetCanHidePin(true);
-		Input->SetPinIsHidden(true);
-	}
+	RegisterInputConnection(&Drag.WeightMap)
+		.SetCanHidePin(true)
+		.SetPinIsHidden(true);
+	RegisterInputConnection(&Lift.WeightMap)
+		.SetCanHidePin(true)
+		.SetPinIsHidden(true);
 }
 
 void FChaosClothAssetSimulationAerodynamicsConfigNode::AddProperties(FPropertyHelper& PropertyHelper) const
diff --git a/Engine/Plugins/ChaosClothAssetEditor/Source/ChaosClothAssetDataflowNodes/Private/ChaosClothAsset/SimulationAnimDriveConfigNode.cpp b/Engine/Plugins/ChaosClothAssetEditor/Source/ChaosClothAssetDataflowNodes/Private/ChaosClothAsset/SimulationAnimDriveConfigNode.cpp
--- a/Engine/Plugins/ChaosClothAssetEditor/Source/ChaosClothAssetDataflowNodes/Private/ChaosClothAsset/SimulationAnimDriveConfigNode.cpp
+++ b/Engine/Plugins/ChaosClothAssetEditor/Source/ChaosClothAssetDataflowNodes/Private/ChaosClothAsset/SimulationAnimDriveConfigNode.cpp
@@ -9,16 +9,12 @@ FChaosClothAssetSimulationAnimDriveConfigNode::FChaosClothAssetSimulationAnimDri
 	: FChaosClothAssetSimulationBaseConfigNode(InParam, InGuid)
 {
 	RegisterCollectionConnections();
-	if (FDataflowInput* const Input = RegisterInputConnection(&AnimDriveStiffness.WeightMap))
-	{
-		Input->SetCanHidePin(true);
-		Input->SetPinIsHidden(true);
-	}
-	if (FDataflowInput* const Input = RegisterInputConnection(&AnimDriveDamping.WeightMap))
-	{
-		Input->SetCanHidePin(true);
-		Input->SetPinIsHidden(true);
-	}
+	RegisterInputConnection(&AnimDriveStiffness.WeightMap)
+		.SetCanHidePin(true)
+		.SetPinIsHidden(true);
+	RegisterInputConnection(&AnimDriveDamping.WeightMap)
+		.SetCanHidePin(true)
+		.SetPinIsHidden(true);
 }
 
 void FChaosClothAssetSimulationAnimDriveConfigNode::AddProperties(FPropertyHelper& PropertyHelper) const
diff --git a/Engine/Plugins/ChaosClothAssetEditor/Source/ChaosClothAssetDataflowNodes/Private/ChaosClothAsset/SimulationBackstopConfigNode.cpp b/Engine/Plugins/ChaosClothAssetEditor/Source/ChaosClothAssetDataflowNodes/Private/ChaosClothAsset/SimulationBackstopConfigNode.cpp
--- a/Engine/Plugins/ChaosClothAssetEditor/Source/ChaosClothAssetDataflowNodes/Private/ChaosClothAsset/SimulationBackstopConfigNode.cpp
+++ b/Engine/Plugins/ChaosClothAssetEditor/Source/ChaosClothAssetDataflowNodes/Private/ChaosClothAsset/SimulationBackstopConfigNode.cpp
@@ -9,16 +9,12 @@ FChaosClothAssetSimulationBackstopConfigNode::FChaosClothAssetSimulationBackstop
 	: FChaosClothAssetSimulationBaseConfigNode(InParam, InGuid)
 {
 	RegisterCollectionConnections();
-	if (FDataflowInput* const Input = RegisterInputConnection(&BackstopDistance.WeightMap))
-	{
-		Input->SetCanHidePin(true);
-		Input->SetPinIsHidden(true);
-	}
-	if (FDataflowInput* const Input = RegisterInputConnection(&BackstopRadius.WeightMap))
-	{
-		Input->SetCanHidePin(true);
-		Input->SetPinIsHidden(true);
-	}
+	RegisterInputConnection(&BackstopDistance.WeightMap)
+		.SetCanHidePin(true)
+		.SetPinIsHidden(true);
+	RegisterInputConnection(&BackstopRadius.WeightMap)
+		.SetCanHidePin(true)
+		.SetPinIsHidden(true);
 }
 
 void FChaosClothAssetSimulationBackstopConfigNode::AddProperties(FPropertyHelper& PropertyHelper) const
diff --git a/Engine/Plugins/ChaosClothAssetEditor/Source/ChaosClothAssetDataflowNodes/Private/ChaosClothAsset/SimulationBendingConfigNode.cpp b/Engine/Plugins/ChaosClothAssetEditor/Source/ChaosClothAssetDataflowNodes/Private/ChaosClothAsset/SimulationBendingConfigNode.cpp
--- a/Engine/Plugins/ChaosClothAssetEditor/Source/ChaosClothAssetDataflowNodes/Private/ChaosClothAsset/SimulationBendingConfigNode.cpp
+++ b/Engine/Plugins/ChaosClothAssetEditor/Source/ChaosClothAssetDataflowNodes/Private/ChaosClothAsset/SimulationBendingConfigNode.cpp
@@ -9,76 +9,48 @@ FChaosClothAssetSimulationBendingConfigNode::FChaosClothAssetSimulationBendingCo
 	: FChaosClothAssetSimulationBaseConfigNode(InParam, InGuid)
 {
 	RegisterCollectionConnections();
-	if (FDataflowInput* const Input = RegisterInputConnection(&FlatnessRatio.WeightMap))
-	{
-		Input->SetCanHidePin(true);
-		Input->SetPinIsHidden(true);
-	}
-	if (FDataflowInput* const Input = RegisterInputConnection(&RestAngle.WeightMap))
-	{
-		Input->SetCanHidePin(true);
-		Input->SetPinIsHidden(true);
-	}
-	if (FDataflowInput* const Input = RegisterInputConnection(&BendingStiffness.WeightMap))
-	{
-		Input->SetCanHidePin(true);
-		Input->SetPinIsHidden(true);
-	}
-	if (FDataflowInput* const Input = RegisterInputConnection(&BendingStiffnessWarp.WeightMap))
-	{
-		Input->SetCanHidePin(true);
-		Input->SetPinIsHidden(true);
-	}
-	if (FDataflowInput* const Input = RegisterInputConnection(&BendingStiffnessWeft.WeightMap))
-	{
-		Input->SetCanHidePin(true);
-		Input->SetPinIsHidden(true);
-	}
-	if (FDataflowInput* const Input = RegisterInputConnection(&BendingStiffnessBias.WeightMap))
-	{
-		Input->SetCanHidePin(true);
-		Input->SetPinIsHidden(true);
-	}
-	if (FDataflowInput* const Input = RegisterInputConnection(&BendingDamping.WeightMap))
-	{
-		Input->SetCanHidePin(true);
-		Input->SetPinIsHidden(true);
-	}
-	if (FDataflowInput* const Input = RegisterInputConnection(&BendingAnisoDamping.WeightMap))
-	{
-		Input->SetCanHidePin(true);
-		Input->SetPinIsHidden(true);
-	}
-	if (FDataflowInput* const Input = RegisterInputConnection(&AnisoBucklingRatio.WeightMap))
-	{
-		Input->SetCanHidePin(true);
-		Input->SetPinIsHidden(true);
-	}
-	if (FDataflowInput* const Input = RegisterInputConnection(&BucklingStiffness.WeightMap))
-	{
-		Input->SetCanHidePin(true);
-		Input->SetPinIsHidden(true);
-	}
-	if (FDataflowInput* const Input = RegisterInputConnection(&BucklingStiffnessWarp.WeightMap))
-	{
-		Input->SetCanHidePin(true);
-		Input->SetPinIsHidden(true);
-	}
-	if (FDataflowInput* const Input = RegisterInputConnection(&BucklingStiffnessWeft.WeightMap))
-	{
-		Input->SetCanHidePin(true);
-		Input->SetPinIsHidden(true);
-	}
-	if (FDataflowInput* const Input = RegisterInputConnection(&BucklingStiffnessBias.WeightMap))
-	{
-		Input->SetCanHidePin(true);
-		Input->SetPinIsHidden(true);
-	}
-	if (FDataflowInput* const Input = RegisterInputConnection(&BucklingRatioWeighted.WeightMap))
-	{
-		Input->SetCanHidePin(true);
-		Input->SetPinIsHidden(true);
-	}
+	RegisterInputConnection(&FlatnessRatio.WeightMap)
+		.SetCanHidePin(true)
+		.SetPinIsHidden(true);
+	RegisterInputConnection(&RestAngle.WeightMap)
+		.SetCanHidePin(true)
+		.SetPinIsHidden(true);
+	RegisterInputConnection(&BendingStiffness.WeightMap)
+		.SetCanHidePin(true)
+		.SetPinIsHidden(true);
+	RegisterInputConnection(&BendingStiffnessWarp.WeightMap)
+		.SetCanHidePin(true)
+		.SetPinIsHidden(true);
+	RegisterInputConnection(&BendingStiffnessWeft.WeightMap)
+		.SetCanHidePin(true)
+		.SetPinIsHidden(true);
+	RegisterInputConnection(&BendingStiffnessBias.WeightMap)
+		.SetCanHidePin(true)
+		.SetPinIsHidden(true);
+	RegisterInputConnection(&BendingDamping.WeightMap)
+		.SetCanHidePin(true)
+		.SetPinIsHidden(true);
+	RegisterInputConnection(&BendingAnisoDamping.WeightMap)
+		.SetCanHidePin(true)
+		.SetPinIsHidden(true);
+	RegisterInputConnection(&AnisoBucklingRatio.WeightMap)
+		.SetCanHidePin(true)
+		.SetPinIsHidden(true);
+	RegisterInputConnection(&BucklingStiffness.WeightMap)
+		.SetCanHidePin(true)
+		.SetPinIsHidden(true);
+	RegisterInputConnection(&BucklingStiffnessWarp.WeightMap)
+		.SetCanHidePin(true)
+		.SetPinIsHidden(true);
+	RegisterInputConnection(&BucklingStiffnessWeft.WeightMap)
+		.SetCanHidePin(true)
+		.SetPinIsHidden(true);
+	RegisterInputConnection(&BucklingStiffnessBias.WeightMap)
+		.SetCanHidePin(true)
+		.SetPinIsHidden(true);
+	RegisterInputConnection(&BucklingRatioWeighted.WeightMap)
+		.SetCanHidePin(true)
+		.SetPinIsHidden(true);
 }
 
 void FChaosClothAssetSimulationBendingConfigNode::AddProperties(FPropertyHelper& PropertyHelper) const
diff --git a/Engine/Plugins/ChaosClothAssetEditor/Source/ChaosClothAssetDataflowNodes/Private/ChaosClothAsset/SimulationDampingConfigNode.cpp b/Engine/Plugins/ChaosClothAssetEditor/Source/ChaosClothAssetDataflowNodes/Private/ChaosClothAsset/SimulationDampingConfigNode.cpp
--- a/Engine/Plugins/ChaosClothAssetEditor/Source/ChaosClothAssetDataflowNodes/Private/ChaosClothAsset/SimulationDampingConfigNode.cpp
+++ b/Engine/Plugins/ChaosClothAssetEditor/Source/ChaosClothAssetDataflowNodes/Private/ChaosClothAsset/SimulationDampingConfigNode.cpp
@@ -9,11 +9,9 @@ FChaosClothAssetSimulationDampingConfigNode::FChaosClothAssetSimulationDampingCo
 	: FChaosClothAssetSimulationBaseConfigNode(InParam, InGuid)
 {
 	RegisterCollectionConnections();
-	if (FDataflowInput* const Input = RegisterInputConnection(&DampingCoefficientWeighted.WeightMap))
-	{
-		Input->SetCanHidePin(true);
-		Input->SetPinIsHidden(true);
-	}
+	 RegisterInputConnection(&DampingCoefficientWeighted.WeightMap)
+		.SetCanHidePin(true)
+		.SetPinIsHidden(true);
 }
 
 void FChaosClothAssetSimulationDampingConfigNode::AddProperties(FPropertyHelper& PropertyHelper) const
diff --git a/Engine/Plugins/ChaosClothAssetEditor/Source/ChaosClothAssetDataflowNodes/Private/ChaosClothAsset/SimulationGravityConfigNode.cpp b/Engine/Plugins/ChaosClothAssetEditor/Source/ChaosClothAssetDataflowNodes/Private/ChaosClothAsset/SimulationGravityConfigNode.cpp
--- a/Engine/Plugins/ChaosClothAssetEditor/Source/ChaosClothAssetDataflowNodes/Private/ChaosClothAsset/SimulationGravityConfigNode.cpp
+++ b/Engine/Plugins/ChaosClothAssetEditor/Source/ChaosClothAssetDataflowNodes/Private/ChaosClothAsset/SimulationGravityConfigNode.cpp
@@ -11,11 +11,9 @@ FChaosClothAssetSimulationGravityConfigNode::FChaosClothAssetSimulationGravityCo
 	: FChaosClothAssetSimulationBaseConfigNode(InParam, InGuid)
 {
 	RegisterCollectionConnections();
-	if (FDataflowInput* const Input = RegisterInputConnection(&GravityScaleWeighted.WeightMap))
-	{
-		Input->SetCanHidePin(true);
-		Input->SetPinIsHidden(true);
-	}
+	RegisterInputConnection(&GravityScaleWeighted.WeightMap)
+		.SetCanHidePin(true)
+		.SetPinIsHidden(true);
 }
 
 void FChaosClothAssetSimulationGravityConfigNode::AddProperties(FPropertyHelper& PropertyHelper) const
diff --git a/Engine/Plugins/ChaosClothAssetEditor/Source/ChaosClothAssetDataflowNodes/Private/ChaosClothAsset/SimulationLongRangeAttachmentConfigNode.cpp b/Engine/Plugins/ChaosClothAssetEditor/Source/ChaosClothAssetDataflowNodes/Private/ChaosClothAsset/SimulationLongRangeAttachmentConfigNode.cpp
--- a/Engine/Plugins/ChaosClothAssetEditor/Source/ChaosClothAssetDataflowNodes/Private/ChaosClothAsset/SimulationLongRangeAttachmentConfigNode.cpp
+++ b/Engine/Plugins/ChaosClothAssetEditor/Source/ChaosClothAssetDataflowNodes/Private/ChaosClothAsset/SimulationLongRangeAttachmentConfigNode.cpp
@@ -12,16 +12,12 @@ FChaosClothAssetSimulationLongRangeAttachmentConfigNode::FChaosClothAssetSimulat
 {
 	RegisterCollectionConnections();
 	RegisterInputConnection(&FixedEndWeightMap);
-	if (FDataflowInput* const Input = RegisterInputConnection(&TetherStiffness.WeightMap))
-	{
-		Input->SetCanHidePin(true);
-		Input->SetPinIsHidden(true);
-	}
-	if (FDataflowInput* const Input = RegisterInputConnection(&TetherScale.WeightMap))
-	{
-		Input->SetCanHidePin(true);
-		Input->SetPinIsHidden(true);
-	}
+	RegisterInputConnection(&TetherStiffness.WeightMap)
+		.SetCanHidePin(true)
+		.SetPinIsHidden(true);
+	RegisterInputConnection(&TetherScale.WeightMap)
+		.SetCanHidePin(true)
+		.SetPinIsHidden(true);
 }
 
 void FChaosClothAssetSimulationLongRangeAttachmentConfigNode::AddProperties(FPropertyHelper& PropertyHelper) const
diff --git a/Engine/Plugins/ChaosClothAssetEditor/Source/ChaosClothAssetDataflowNodes/Private/ChaosClothAsset/SimulationMassConfigNode.cpp b/Engine/Plugins/ChaosClothAssetEditor/Source/ChaosClothAssetDataflowNodes/Private/ChaosClothAsset/SimulationMassConfigNode.cpp
--- a/Engine/Plugins/ChaosClothAssetEditor/Source/ChaosClothAssetDataflowNodes/Private/ChaosClothAsset/SimulationMassConfigNode.cpp
+++ b/Engine/Plugins/ChaosClothAssetEditor/Source/ChaosClothAssetDataflowNodes/Private/ChaosClothAsset/SimulationMassConfigNode.cpp
@@ -12,16 +12,12 @@ FChaosClothAssetSimulationMassConfigNode::FChaosClothAssetSimulationMassConfigNo
 	: FChaosClothAssetSimulationBaseConfigNode(InParam, InGuid)
 {
 	RegisterCollectionConnections();
-	if (FDataflowInput* const Input = RegisterInputConnection(&UniformMassWeighted.WeightMap))
-	{
-		Input->SetCanHidePin(true);
-		Input->SetPinIsHidden(true);
-	}
-	if (FDataflowInput* const Input = RegisterInputConnection(&DensityWeighted.WeightMap))
-	{
-		Input->SetCanHidePin(true);
-		Input->SetPinIsHidden(true);
-	}
+	RegisterInputConnection(&UniformMassWeighted.WeightMap)
+		.SetCanHidePin(true)
+		.SetPinIsHidden(true);
+	RegisterInputConnection(&DensityWeighted.WeightMap)
+		.SetCanHidePin(true)
+		.SetPinIsHidden(true);
 }
 
 void FChaosClothAssetSimulationMassConfigNode::AddProperties(FPropertyHelper& PropertyHelper) const
diff --git a/Engine/Plugins/ChaosClothAssetEditor/Source/ChaosClothAssetDataflowNodes/Private/ChaosClothAsset/SimulationMaxDistanceConfigNode.cpp b/Engine/Plugins/ChaosClothAssetEditor/Source/ChaosClothAssetDataflowNodes/Private/ChaosClothAsset/SimulationMaxDistanceConfigNode.cpp
--- a/Engine/Plugins/ChaosClothAssetEditor/Source/ChaosClothAssetDataflowNodes/Private/ChaosClothAsset/SimulationMaxDistanceConfigNode.cpp
+++ b/Engine/Plugins/ChaosClothAssetEditor/Source/ChaosClothAssetDataflowNodes/Private/ChaosClothAsset/SimulationMaxDistanceConfigNode.cpp
@@ -9,11 +9,9 @@ FChaosClothAssetSimulationMaxDistanceConfigNode::FChaosClothAssetSimulationMaxDi
 	: FChaosClothAssetSimulationBaseConfigNode(InParam, InGuid)
 {
 	RegisterCollectionConnections();
-	if (FDataflowInput* const Input = RegisterInputConnection(&MaxDistance.WeightMap))
-	{
-		Input->SetCanHidePin(true);
-		Input->SetPinIsHidden(true);
-	}
+	RegisterInputConnection(&MaxDistance.WeightMap)
+		.SetCanHidePin(true)
+		.SetPinIsHidden(true);
 }
 
 void FChaosClothAssetSimulationMaxDistanceConfigNode::AddProperties(FPropertyHelper& PropertyHelper) const
diff --git a/Engine/Plugins/ChaosClothAssetEditor/Source/ChaosClothAssetDataflowNodes/Private/ChaosClothAsset/SimulationMultiResConfigNode.cpp b/Engine/Plugins/ChaosClothAssetEditor/Source/ChaosClothAssetDataflowNodes/Private/ChaosClothAsset/SimulationMultiResConfigNode.cpp
--- a/Engine/Plugins/ChaosClothAssetEditor/Source/ChaosClothAssetDataflowNodes/Private/ChaosClothAsset/SimulationMultiResConfigNode.cpp
+++ b/Engine/Plugins/ChaosClothAssetEditor/Source/ChaosClothAssetDataflowNodes/Private/ChaosClothAsset/SimulationMultiResConfigNode.cpp
@@ -9,16 +9,12 @@ FChaosClothAssetSimulationMultiResConfigNode::FChaosClothAssetSimulationMultiRes
 	: FChaosClothAssetSimulationBaseConfigNode(InParam, InGuid)
 {
 	RegisterCollectionConnections();
-	if (FDataflowInput* const Input = RegisterInputConnection(&MultiResStiffness.WeightMap))
-	{
-		Input->SetCanHidePin(true);
-		Input->SetPinIsHidden(true);
-	}
-	if (FDataflowInput* const Input = RegisterInputConnection(&MultiResVelocityTargetStiffness.WeightMap))
-	{
-		Input->SetCanHidePin(true);
-		Input->SetPinIsHidden(true);
-	}
+	RegisterInputConnection(&MultiResStiffness.WeightMap)
+		.SetCanHidePin(true)
+		.SetPinIsHidden(true);
+	RegisterInputConnection(&MultiResVelocityTargetStiffness.WeightMap)
+		.SetCanHidePin(true)
+		.SetPinIsHidden(true);
 }
 
 void FChaosClothAssetSimulationMultiResConfigNode::AddProperties(FPropertyHelper& PropertyHelper) const
diff --git a/Engine/Plugins/ChaosClothAssetEditor/Source/ChaosClothAssetDataflowNodes/Private/ChaosClothAsset/SimulationPressureConfigNode.cpp b/Engine/Plugins/ChaosClothAssetEditor/Source/ChaosClothAssetDataflowNodes/Private/ChaosClothAsset/SimulationPressureConfigNode.cpp
--- a/Engine/Plugins/ChaosClothAssetEditor/Source/ChaosClothAssetDataflowNodes/Private/ChaosClothAsset/SimulationPressureConfigNode.cpp
+++ b/Engine/Plugins/ChaosClothAssetEditor/Source/ChaosClothAssetDataflowNodes/Private/ChaosClothAsset/SimulationPressureConfigNode.cpp
@@ -9,11 +9,9 @@ FChaosClothAssetSimulationPressureConfigNode::FChaosClothAssetSimulationPressure
 	: FChaosClothAssetSimulationBaseConfigNode(InParam, InGuid)
 {
 	RegisterCollectionConnections();
-	if (FDataflowInput* const Input = RegisterInputConnection(&Pressure.WeightMap))
-	{
-		Input->SetCanHidePin(true);
-		Input->SetPinIsHidden(true);
-	}
+	RegisterInputConnection(&Pressure.WeightMap)
+		.SetCanHidePin(true)
+		.SetPinIsHidden(true);
 }
 
 void FChaosClothAssetSimulationPressureConfigNode::AddProperties(FPropertyHelper& PropertyHelper) const
diff --git a/Engine/Plugins/ChaosClothAssetEditor/Source/ChaosClothAssetDataflowNodes/Private/ChaosClothAsset/SimulationSelfCollisionConfigNode.cpp b/Engine/Plugins/ChaosClothAssetEditor/Source/ChaosClothAssetDataflowNodes/Private/ChaosClothAsset/SimulationSelfCollisionConfigNode.cpp
--- a/Engine/Plugins/ChaosClothAssetEditor/Source/ChaosClothAssetDataflowNodes/Private/ChaosClothAsset/SimulationSelfCollisionConfigNode.cpp
+++ b/Engine/Plugins/ChaosClothAssetEditor/Source/ChaosClothAssetDataflowNodes/Private/ChaosClothAsset/SimulationSelfCollisionConfigNode.cpp
@@ -11,31 +11,21 @@ FChaosClothAssetSimulationSelfCollisionConfigNode::FChaosClothAssetSimulationSel
 	: FChaosClothAssetSimulationBaseConfigNode(InParam, InGuid)
 {
 	RegisterCollectionConnections();
-	if (FDataflowInput* const Input = RegisterInputConnection(&SelfCollisionLayers.StringValue, GET_MEMBER_NAME_CHECKED(FChaosClothAssetConnectableIStringValue, StringValue)))
-	{
-		Input->SetCanHidePin(true);
-		Input->SetPinIsHidden(true);
-	}
-	if (FDataflowInput* const Input = RegisterInputConnection(&SelfCollisionDisabledFaces.StringValue, GET_MEMBER_NAME_CHECKED(FChaosClothAssetConnectableIStringValue, StringValue)))
-	{
-		Input->SetCanHidePin(true);
-		Input->SetPinIsHidden(true);
-	}
-	if (FDataflowInput* const Input = RegisterInputConnection(&SelfCollisionEnabledKinematicFaces.StringValue, GET_MEMBER_NAME_CHECKED(FChaosClothAssetConnectableIStringValue, StringValue)))
-	{
-		Input->SetCanHidePin(true);
-		Input->SetPinIsHidden(true);
-	}
-	if (FDataflowInput* const Input = RegisterInputConnection(&SelfCollisionKinematicColliderFrictionWeighted.WeightMap))
-	{
-		Input->SetCanHidePin(true);
-		Input->SetPinIsHidden(true);
-	}
-	if (FDataflowInput* const Input = RegisterInputConnection(&SelfCollisionThicknessWeighted.WeightMap))
-	{
-		Input->SetCanHidePin(true);
-		Input->SetPinIsHidden(true);
-	}
+	RegisterInputConnection(&SelfCollisionLayers.StringValue, GET_MEMBER_NAME_CHECKED(FChaosClothAssetConnectableIStringValue, StringValue))
+		.SetCanHidePin(true)
+		.SetPinIsHidden(true);
+	RegisterInputConnection(&SelfCollisionDisabledFaces.StringValue, GET_MEMBER_NAME_CHECKED(FChaosClothAssetConnectableIStringValue, StringValue))
+		.SetCanHidePin(true)
+		.SetPinIsHidden(true);
+	RegisterInputConnection(&SelfCollisionEnabledKinematicFaces.StringValue, GET_MEMBER_NAME_CHECKED(FChaosClothAssetConnectableIStringValue, StringValue))
+		.SetCanHidePin(true)
+		.SetPinIsHidden(true);
+	RegisterInputConnection(&SelfCollisionKinematicColliderFrictionWeighted.WeightMap)
+		.SetCanHidePin(true)
+		.SetPinIsHidden(true);
+	RegisterInputConnection(&SelfCollisionThicknessWeighted.WeightMap)
+		.SetCanHidePin(true)
+		.SetPinIsHidden(true);
 }
 
 void FChaosClothAssetSimulationSelfCollisionConfigNode::AddProperties(FPropertyHelper& PropertyHelper) const
diff --git a/Engine/Plugins/ChaosClothAssetEditor/Source/ChaosClothAssetDataflowNodes/Private/ChaosClothAsset/SimulationStretchConfigNode.cpp b/Engine/Plugins/ChaosClothAssetEditor/Source/ChaosClothAssetDataflowNodes/Private/ChaosClothAsset/SimulationStretchConfigNode.cpp
--- a/Engine/Plugins/ChaosClothAssetEditor/Source/ChaosClothAssetDataflowNodes/Private/ChaosClothAsset/SimulationStretchConfigNode.cpp
+++ b/Engine/Plugins/ChaosClothAssetEditor/Source/ChaosClothAssetDataflowNodes/Private/ChaosClothAsset/SimulationStretchConfigNode.cpp
@@ -9,51 +9,33 @@ FChaosClothAssetSimulationStretchConfigNode::FChaosClothAssetSimulationStretchCo
 	: FChaosClothAssetSimulationBaseConfigNode(InParam, InGuid)
 {
 	RegisterCollectionConnections();
-	if (FDataflowInput* const Input = RegisterInputConnection(&StretchStiffness.WeightMap))
-	{
-		Input->SetCanHidePin(true);
-		Input->SetPinIsHidden(true);
-	}
-	if (FDataflowInput* const Input = RegisterInputConnection(&StretchStiffnessWarp.WeightMap))
-	{
-		Input->SetCanHidePin(true);
-		Input->SetPinIsHidden(true);
-	}
-	if (FDataflowInput* const Input = RegisterInputConnection(&StretchStiffnessWeft.WeightMap))
-	{
-		Input->SetCanHidePin(true);
-		Input->SetPinIsHidden(true);
-	}
-	if (FDataflowInput* const Input = RegisterInputConnection(&StretchStiffnessBias.WeightMap))
-	{
-		Input->SetCanHidePin(true);
-		Input->SetPinIsHidden(true);
-	}
-	if (FDataflowInput* const Input = RegisterInputConnection(&StretchDamping.WeightMap))
-	{
-		Input->SetCanHidePin(true);
-		Input->SetPinIsHidden(true);
-	}
-	if (FDataflowInput* const Input = RegisterInputConnection(&StretchAnisoDamping.WeightMap))
-	{
-		Input->SetCanHidePin(true);
-		Input->SetPinIsHidden(true);
-	}
-	if (FDataflowInput* const Input = RegisterInputConnection(&StretchWarpScale.WeightMap))
-	{
-		Input->SetCanHidePin(true);
-		Input->SetPinIsHidden(true);
-	}
-	if (FDataflowInput* const Input = RegisterInputConnection(&StretchWeftScale.WeightMap))
-	{
-		Input->SetCanHidePin(true);
-		Input->SetPinIsHidden(true);
-	}
-	if (FDataflowInput* const Input = RegisterInputConnection(&AreaStiffness.WeightMap))
-	{
-		Input->SetCanHidePin(true);
-		Input->SetPinIsHidden(true);
-	}
+	RegisterInputConnection(&StretchStiffness.WeightMap)
+		.SetCanHidePin(true)
+		.SetPinIsHidden(true);
+	RegisterInputConnection(&StretchStiffnessWarp.WeightMap)
+		.SetCanHidePin(true)
+		.SetPinIsHidden(true);
+	RegisterInputConnection(&StretchStiffnessWeft.WeightMap)
+		.SetCanHidePin(true)
+		.SetPinIsHidden(true);
+	RegisterInputConnection(&StretchStiffnessBias.WeightMap)
+		.SetCanHidePin(true)
+		.SetPinIsHidden(true);
+	RegisterInputConnection(&StretchDamping.WeightMap)
+		.SetCanHidePin(true)
+		.SetPinIsHidden(true);
+	RegisterInputConnection(&StretchAnisoDamping.WeightMap)
+		.SetCanHidePin(true)
+		.SetPinIsHidden(true);
+	RegisterInputConnection(&StretchWarpScale.WeightMap)
+		.SetCanHidePin(true)
+		.SetPinIsHidden(true);
+	RegisterInputConnection(&StretchWeftScale.WeightMap)
+		.SetCanHidePin(true)
+		.SetPinIsHidden(true);
+	RegisterInputConnection(&AreaStiffness.WeightMap)
+		.SetCanHidePin(true)
+		.SetPinIsHidden(true);
 }
 
 void FChaosClothAssetSimulationStretchConfigNode::AddProperties(FPropertyHelper& PropertyHelper) const
diff --git a/Engine/Plugins/ChaosClothAssetEditor/Source/ChaosClothAssetDataflowNodes/Private/ChaosClothAsset/TerminalNode.cpp b/Engine/Plugins/ChaosClothAssetEditor/Source/ChaosClothAssetDataflowNodes/Private/ChaosClothAsset/TerminalNode.cpp
--- a/Engine/Plugins/ChaosClothAssetEditor/Source/ChaosClothAssetDataflowNodes/Private/ChaosClothAsset/TerminalNode.cpp
+++ b/Engine/Plugins/ChaosClothAssetEditor/Source/ChaosClothAssetDataflowNodes/Private/ChaosClothAsset/TerminalNode.cpp
@@ -87,6 +87,263 @@ namespace UE::Chaos::ClothAsset::Private
 	}
 }
 
+FChaosClothAssetTerminalNode_v2::FChaosClothAssetTerminalNode_v2(const Dataflow::FNodeParameters& InParam, FGuid InGuid)
+	: FDataflowTerminalNode(InParam, InGuid)
+{
+	// Start with Lod0
+	AddPins();
+}
+
+void FChaosClothAssetTerminalNode_v2::SetAssetValue(TObjectPtr<UObject> Asset, Dataflow::FContext& Context) const
+{
+	if (UChaosClothAsset* ClothAsset = Cast<UChaosClothAsset>(Asset.Get()))
+	{
+		using namespace UE::Chaos::ClothAsset;
+
+		if (RefreshAsset.bRefreshAsset)
+		{
+			bClothCollectionChecksumValid = false;
+			RefreshAsset.bRefreshAsset = false;
+		}
+
+		TArray<TSharedRef<FManagedArrayCollection>> InClothCollections = GetCleanedCollectionLodValues(Context);
+		TArray<TSharedRef<FManagedArrayCollection>>& ClothCollections = ClothAsset->GetClothCollections();
+
+		const uint32 PreviousChecksum = ClothColllectionChecksum;
+		const bool bPreviousChecksumsValid = bClothCollectionChecksumValid;
+		ClothColllectionChecksum = Private::CalculateClothChecksum(InClothCollections);
+		bClothCollectionChecksumValid = InClothCollections.Num() > 0;
+
+		if (bPreviousChecksumsValid && PreviousChecksum == ClothColllectionChecksum && Private::PropertyKeysAndSolverTypesMatch(InClothCollections, ClothCollections))
+		{
+			// Cloth and property keys match. Just update property values.
+			check(InClothCollections.Num() == ClothCollections.Num());
+			check(ClothCollections.Num() > 0);
+			for (int32 LODIndex = 0; LODIndex < InClothCollections.Num(); ++LODIndex)
+			{
+				Chaos::Softs::FCollectionPropertyFacade Properties(ClothCollections[LODIndex]);
+				Properties.UpdateProperties(InClothCollections[LODIndex].ToSharedPtr());
+			}
+
+			// Asset must be resaved
+			ClothAsset->MarkPackageDirty();
+
+			return;
+		}
+
+		// Reset the asset's collection
+		ClothCollections.Reset(InClothCollections.Num());
+
+		// Reset the asset's material list
+		ClothAsset->GetMaterials().Reset();
+
+		// Iterate through the LODs
+		FString PhysicsAssetPathName;
+
+		for (int32 LodIndex = 0; LodIndex < InClothCollections.Num(); ++LodIndex)
+		{
+			// New LOD
+			TSharedRef<FManagedArrayCollection>& ClothCollection = ClothCollections.Emplace_GetRef(MakeShared<FManagedArrayCollection>());
+			UE::Chaos::ClothAsset::FCollectionClothFacade ClothFacade(ClothCollection);
+			ClothFacade.DefineSchema();
+
+			const FCollectionClothConstFacade InClothFacade(InClothCollections[LodIndex]);
+			if (!InClothFacade.HasValidRenderData())  // The cloth collection must at least have a render mesh
+			{
+				FClothDataflowTools::LogAndToastWarning(*this, LOCTEXT("InvalidLODHeadline", "Invalid LOD."),
+					FText::Format(
+						LOCTEXT("InvalidLODDetails", "LOD {0} has no valid data."),
+						LodIndex));
+
+				continue;
+			}
+
+			// Copy input LOD to current output LOD
+			ClothFacade.Initialize(InClothFacade);
+
+			// Add this LOD's materials to the asset
+			const int32 NumLodMaterials = ClothFacade.GetNumRenderPatterns();
+
+			TArray<FSkeletalMaterial>& Materials = ClothAsset->GetMaterials();
+			Materials.Reserve(Materials.Num() + NumLodMaterials);
+
+			const TConstArrayView<FString> LodRenderMaterialPathName = ClothFacade.GetRenderMaterialPathName();
+			for (int32 LodMaterialIndex = 0; LodMaterialIndex < NumLodMaterials; ++LodMaterialIndex)
+			{
+				const FString& RenderMaterialPathName = LodRenderMaterialPathName[LodMaterialIndex];
+
+				if (UMaterialInterface* const Material = LoadObject<UMaterialInterface>(ClothAsset, *RenderMaterialPathName, nullptr, LOAD_None, nullptr))
+				{
+					Materials.Emplace(Material, true, false, Material->GetFName());
+				}
+				else
+				{
+					Materials.Emplace();
+				}
+			}
+
+			// Set properties
+			constexpr bool bUpdateExistingProperties = false;
+			Chaos::Softs::FCollectionPropertyMutableFacade(ClothCollection).Append(InClothCollections[LodIndex].ToSharedPtr(), bUpdateExistingProperties);
+
+			// Set selections
+			FCollectionClothSelectionFacade Selection(ClothCollection);
+			FCollectionClothSelectionConstFacade InSelection(InClothCollections[LodIndex]);
+			if (InSelection.IsValid())
+			{
+				Selection.DefineSchema();
+				const TArray<FName> InSelectionNames = InSelection.GetNames();
+				for (const FName& InSelectionName : InSelectionNames)
+				{
+					const FName SelectionGroup = InSelection.GetSelectionGroup(InSelectionName);
+					if (SelectionGroup == ClothCollectionGroup::SimVertices3D ||
+						SelectionGroup == ClothCollectionGroup::SimFaces)
+					{
+						Selection.FindOrAddSelectionSet(InSelectionName, SelectionGroup) = InSelection.GetSelectionSet(InSelectionName);
+					}
+				}
+			}
+
+
+			// Set physics asset only with LOD 0 at the moment
+			if (LodIndex == 0)
+			{
+				using namespace ::Chaos::Softs;
+				PhysicsAssetPathName = InClothFacade.GetPhysicsAssetPathName();
+			}
+		}
+
+		// Make sure that whatever happens there is always at least one empty LOD to avoid crashing the render data
+		if (ClothCollections.Num() < 1)
+		{
+			TSharedRef<FManagedArrayCollection>& ClothCollection = ClothCollections.Emplace_GetRef(MakeShared<FManagedArrayCollection>());
+			FCollectionClothFacade ClothFacade(ClothCollection);
+			ClothFacade.DefineSchema();
+		}
+
+		// Set reference skeleton
+		constexpr bool bRebuildModels = false;  // Avoid rebuilding the asset twice
+		ClothAsset->UpdateSkeletonFromCollection(bRebuildModels);
+
+		// Set physics asset (note: the cloth asset's physics asset is only replaced if a collection path name is found valid)
+		UPhysicsAsset* const PhysicsAsset = !PhysicsAssetPathName.IsEmpty() ?
+			LoadObject<UPhysicsAsset>(ClothAsset, *PhysicsAssetPathName, nullptr, LOAD_None, nullptr) :
+			nullptr;
+		ClothAsset->SetPhysicsAsset(PhysicsAsset);
+
+		// Rebuild the asset static data
+		ClothAsset->Build(&LODTransitionDataCache);
+
+		// Asset must be resaved
+		ClothAsset->MarkPackageDirty();
+	}
+}
+
+TArray<Dataflow::FPin> FChaosClothAssetTerminalNode_v2::AddPins()
+{
+	const int32 Index = CollectionLods.AddDefaulted();
+	const FDataflowInput& Input = RegisterInputArrayConnection(GetConnectionReference(Index));
+	return { { Dataflow::FPin::EDirection::INPUT, Input.GetType(), Input.GetName() } };
+}
+
+TArray<Dataflow::FPin> FChaosClothAssetTerminalNode_v2::GetPinsToRemove() const
+{
+	const int32 Index = CollectionLods.Num() - 1;
+	check(CollectionLods.IsValidIndex(Index));
+	if (const FDataflowInput* const Input = FindInput(GetConnectionReference(Index)))
+	{
+		return { { Dataflow::FPin::EDirection::INPUT, Input->GetType(), Input->GetName() } };
+	}
+	return Super::GetPinsToRemove();
+}
+
+void FChaosClothAssetTerminalNode_v2::OnPinRemoved(const Dataflow::FPin& Pin)
+{
+	const int32 Index = CollectionLods.Num() - 1;
+	check(CollectionLods.IsValidIndex(Index));
+#if DO_CHECK
+	const FDataflowInput* const Input = FindInput(GetConnectionReference(Index));
+	check(Input);
+	check(Input->GetName() == Pin.Name);
+	check(Input->GetType() == Pin.Type);
+#endif
+	CollectionLods.SetNum(Index);
+
+	return Super::OnPinRemoved(Pin);
+}
+
+TArray<TSharedRef<FManagedArrayCollection>> FChaosClothAssetTerminalNode_v2::GetCleanedCollectionLodValues(Dataflow::FContext& Context) const
+{
+	using namespace UE::Chaos::ClothAsset;
+
+	TArray<TSharedRef<FManagedArrayCollection>> CollectionLodValues;
+	CollectionLodValues.Reserve(CollectionLods.Num());
+
+	int32 LastValidLodIndex = INDEX_NONE;
+	for (int32 LodIndex = 0; LodIndex < CollectionLods.Num(); ++LodIndex)
+	{
+		TSharedRef<FManagedArrayCollection>& CollectionLodValue = CollectionLodValues.Emplace_GetRef(MakeShared<FManagedArrayCollection>(GetValue<FManagedArrayCollection>(Context, GetConnectionReference(LodIndex))));
+
+		FCollectionClothFacade ClothFacade(CollectionLodValue);
+		if (ClothFacade.HasValidRenderData())  // The cloth collection must at least have a render mesh
+		{
+			FClothGeometryTools::CleanupAndCompactMesh(CollectionLodValue);
+			LastValidLodIndex = LodIndex;
+		}
+		else
+		{
+			if (LastValidLodIndex >= 0)
+			{
+				ClothFacade.DefineSchema();
+				ClothFacade.Initialize(FCollectionClothConstFacade(CollectionLodValues[LastValidLodIndex]));
+				FClothDataflowTools::LogAndToastWarning(*this,
+					LOCTEXT("InvalidInputLodNHeadline", "Invalid input LOD."),
+					FText::Format(
+						LOCTEXT("InvalidInputLodNDetails",
+							"Invalid or empty input LOD for LOD {0}.\n"
+							"Using the previous valid LOD {1} instead."),
+						LodIndex,
+						LastValidLodIndex));
+			}
+			else
+			{
+				FClothDataflowTools::LogAndToastWarning(*this,
+					LOCTEXT("InvalidInputLod0Headline", "Invalid input LOD 0."),
+					LOCTEXT("InvalidInputLod0Details",
+						"Invalid or empty input LOD for LOD 0.\n"
+						"LOD 0 cannot be empty in order to construct a valid Cloth Asset."));
+				break;
+			}
+		}
+	}
+	return CollectionLodValues;
+}
+
+Dataflow::TConnectionReference<FManagedArrayCollection> FChaosClothAssetTerminalNode_v2::GetConnectionReference(int32 Index) const
+{
+	return { &CollectionLods[Index], Index, &CollectionLods };
+}
+
+void FChaosClothAssetTerminalNode_v2::Serialize(FArchive& Ar)
+{
+	// because we add pins we need to make sure we restore them when loading
+	// to make sure they can get properly reconnected
+
+	if (Ar.IsLoading())
+	{
+		check(CollectionLods.Num() > 0);
+		check(FindInput(GetConnectionReference(0)));
+
+		for (int32 Index = 1; Index < CollectionLods.Num(); ++Index)
+		{
+			RegisterInputArrayConnection(GetConnectionReference(Index));
+		}
+	}
+}
+
+
+
+
 FChaosClothAssetTerminalNode::FChaosClothAssetTerminalNode(const Dataflow::FNodeParameters& InParam, FGuid InGuid)
 	: FDataflowTerminalNode(InParam, InGuid)
 {
@@ -238,13 +495,13 @@ void FChaosClothAssetTerminalNode::SetAssetValue(TObjectPtr<UObject> Asset, Data
 	}
 }
 
-Dataflow::FPin FChaosClothAssetTerminalNode::AddPin()
+TArray<Dataflow::FPin> FChaosClothAssetTerminalNode::AddPins()
 {
-	auto AddInput = [this](const FManagedArrayCollection* Collection) -> Dataflow::FPin
+	auto AddInput = [this](const FManagedArrayCollection* Collection) -> TArray<Dataflow::FPin>
 		{
 			RegisterInputConnection(Collection);
 			const FDataflowInput* const Input = FindInput(Collection);
-			return { Dataflow::FPin::EDirection::INPUT, Input->GetType(), Input->GetName() };
+			return { { Dataflow::FPin::EDirection::INPUT, Input->GetType(), Input->GetName() } };
 		};
 
 	switch (NumLods)
@@ -257,16 +514,16 @@ Dataflow::FPin FChaosClothAssetTerminalNode::AddPin()
 	default: break;
 	}
 
-	return Super::AddPin();
+	return Super::AddPins();
 }
 
-Dataflow::FPin FChaosClothAssetTerminalNode::GetPinToRemove() const
+TArray<Dataflow::FPin> FChaosClothAssetTerminalNode::GetPinsToRemove() const
 {
-	auto PinToRemove = [this](const FManagedArrayCollection* Collection) -> Dataflow::FPin
+	auto PinToRemove = [this](const FManagedArrayCollection* Collection) -> TArray<Dataflow::FPin>
 		{
 			const FDataflowInput* const Input = FindInput(Collection);
 			check(Input);
-			return { Dataflow::FPin::EDirection::INPUT, Input->GetType(), Input->GetName() };
+			return { { Dataflow::FPin::EDirection::INPUT, Input->GetType(), Input->GetName() } };
 		};
 
 	switch (NumLods - 1)
@@ -278,7 +535,7 @@ Dataflow::FPin FChaosClothAssetTerminalNode::GetPinToRemove() const
 	case 5: return PinToRemove(&CollectionLod5);
 	default: break;
 	}
-	return Super::GetPinToRemove();
+	return Super::GetPinsToRemove();
 }
 
 void FChaosClothAssetTerminalNode::OnPinRemoved(const Dataflow::FPin& Pin)
@@ -404,10 +661,9 @@ void FChaosClothAssetTerminalNode::Serialize(FArchive& Ar)
 		NumLods = 1; // need to reset back to default because add pin will increment it again 
 		for (int32 LodIndex = 0; LodIndex < NumLodToAdd; ++LodIndex)
 		{
-			AddPin();
+			AddPins();
 		}
 		ensure(NumLodToAdd == (NumLods - 1));
 	}
 }
-
 #undef LOCTEXT_NAMESPACE
diff --git a/Engine/Plugins/ChaosClothAssetEditor/Source/ChaosClothAssetDataflowNodes/Private/ChaosClothAsset/TransferSkinWeightsNode.cpp b/Engine/Plugins/ChaosClothAssetEditor/Source/ChaosClothAssetDataflowNodes/Private/ChaosClothAsset/TransferSkinWeightsNode.cpp
--- a/Engine/Plugins/ChaosClothAssetEditor/Source/ChaosClothAssetDataflowNodes/Private/ChaosClothAsset/TransferSkinWeightsNode.cpp
+++ b/Engine/Plugins/ChaosClothAssetEditor/Source/ChaosClothAssetDataflowNodes/Private/ChaosClothAsset/TransferSkinWeightsNode.cpp
@@ -716,11 +716,9 @@ FChaosClothAssetTransferSkinWeightsNode::FChaosClothAssetTransferSkinWeightsNode
 {
 	RegisterInputConnection(&Collection);
 	RegisterOutputConnection(&Collection, &Collection);
-	if (FDataflowInput* const Input = RegisterInputConnection(&InpaintMask.WeightMap, GET_MEMBER_NAME_CHECKED(FChaosClothAssetWeightedValueNonAnimatableNoLowHighRange, WeightMap)))
-	{
-		Input->SetCanHidePin(true);
-		Input->SetPinIsHidden(true);
-	}
+	RegisterInputConnection(&InpaintMask.WeightMap, GET_MEMBER_NAME_CHECKED(FChaosClothAssetWeightedValueNonAnimatableNoLowHighRange, WeightMap))
+		.SetCanHidePin(true)
+		.SetPinIsHidden(true);
 }
 
 void FChaosClothAssetTransferSkinWeightsNode::Evaluate(Dataflow::FContext& Context, const FDataflowOutput* Out) const
diff --git a/Engine/Plugins/ChaosClothAssetEditor/Source/ChaosClothAssetDataflowNodes/Private/ChaosClothAsset/WeightMapNode.cpp b/Engine/Plugins/ChaosClothAssetEditor/Source/ChaosClothAssetDataflowNodes/Private/ChaosClothAsset/WeightMapNode.cpp
--- a/Engine/Plugins/ChaosClothAssetEditor/Source/ChaosClothAssetDataflowNodes/Private/ChaosClothAsset/WeightMapNode.cpp
+++ b/Engine/Plugins/ChaosClothAssetEditor/Source/ChaosClothAssetDataflowNodes/Private/ChaosClothAsset/WeightMapNode.cpp
@@ -39,16 +39,12 @@ FChaosClothAssetWeightMapNode::FChaosClothAssetWeightMapNode(const Dataflow::FNo
 	: FDataflowTerminalNode(InParam, InGuid)
 {
 	RegisterInputConnection(&Collection);
-	if (FDataflowInput* const Input = RegisterInputConnection(&InputName.StringValue, GET_MEMBER_NAME_CHECKED(FChaosClothAssetConnectableIStringValue, StringValue)))
-	{
-		Input->SetCanHidePin(true);
-		Input->SetPinIsHidden(true);
-	}
-	if (FDataflowInput* const Input = RegisterInputConnection(&TransferCollection))
-	{
-		Input->SetCanHidePin(true);
-		Input->SetPinIsHidden(true);
-	}
+	RegisterInputConnection(&InputName.StringValue, GET_MEMBER_NAME_CHECKED(FChaosClothAssetConnectableIStringValue, StringValue))
+		.SetCanHidePin(true)
+		.SetPinIsHidden(true);
+	RegisterInputConnection(&TransferCollection)
+		.SetCanHidePin(true)
+		.SetPinIsHidden(true);
 	RegisterOutputConnection(&Collection, &Collection);
 	RegisterOutputConnection(&Name);
 }
diff --git a/Engine/Plugins/ChaosClothAssetEditor/Source/ChaosClothAssetDataflowNodes/Public/ChaosClothAsset/MergeClothCollectionsNode.h b/Engine/Plugins/ChaosClothAssetEditor/Source/ChaosClothAssetDataflowNodes/Public/ChaosClothAsset/MergeClothCollectionsNode.h
--- a/Engine/Plugins/ChaosClothAssetEditor/Source/ChaosClothAssetDataflowNodes/Public/ChaosClothAsset/MergeClothCollectionsNode.h
+++ b/Engine/Plugins/ChaosClothAssetEditor/Source/ChaosClothAssetDataflowNodes/Public/ChaosClothAsset/MergeClothCollectionsNode.h
@@ -8,6 +8,38 @@
 
 /** Merge multiple cloth collections into a single cloth collection of multiple patterns. */
 USTRUCT(Meta = (DataflowCloth))
+struct FChaosClothAssetMergeClothCollectionsNode_v2 : public FDataflowNode
+{
+	GENERATED_USTRUCT_BODY()
+	DATAFLOW_NODE_DEFINE_INTERNAL(FChaosClothAssetMergeClothCollectionsNode_v2, "MergeClothCollections", "Cloth", "Cloth Merge Collection")
+
+public:
+	UPROPERTY()
+	TArray<FManagedArrayCollection> Collections;
+
+	UPROPERTY(Meta = (DataflowOutput, DataflowPassthrough = "Collections[0]"))
+	FManagedArrayCollection Collection;
+
+	FChaosClothAssetMergeClothCollectionsNode_v2(const Dataflow::FNodeParameters& InParam, FGuid InGuid = FGuid::NewGuid());
+
+private:
+	//~ Begin FDataflowNode interface
+	virtual void Evaluate(Dataflow::FContext& Context, const FDataflowOutput* Out) const override;
+	virtual TArray<Dataflow::FPin> AddPins() override;
+	virtual bool CanAddPin() const override { return true; }
+	virtual bool CanRemovePin() const override { return Collections.Num() > 2; }
+	virtual TArray<Dataflow::FPin> GetPinsToRemove() const override;
+	virtual void OnPinRemoved(const Dataflow::FPin& Pin) override;
+	virtual void Serialize(FArchive& Ar) override;
+	//~ End FDataflowNode interface
+
+	Dataflow::TConnectionReference<FManagedArrayCollection> GetConnectionReference(int32 Index) const;
+};
+
+
+
+/** Merge multiple cloth collections into a single cloth collection of multiple patterns. */
+USTRUCT(Meta = (DataflowCloth, Deprecated="5.5"))
 struct FChaosClothAssetMergeClothCollectionsNode : public FDataflowNode
 {
 	GENERATED_USTRUCT_BODY()
@@ -43,10 +75,10 @@ public:
 private:
 	//~ Begin FDataflowNode interface
 	virtual void Evaluate(Dataflow::FContext& Context, const FDataflowOutput* Out) const override;
-	virtual Dataflow::FPin AddPin() override;
+	virtual TArray<Dataflow::FPin> AddPins() override;
 	virtual bool CanAddPin() const override { return NumInputs < MaxInputs; }
 	virtual bool CanRemovePin() const override { return NumInputs > 1; }
-	virtual Dataflow::FPin GetPinToRemove() const override;
+	virtual TArray<Dataflow::FPin> GetPinsToRemove() const override;
 	virtual void OnPinRemoved(const Dataflow::FPin& Pin) override;
 	virtual void Serialize(FArchive& Ar) override;
 	//~ End FDataflowNode interface
diff --git a/Engine/Plugins/ChaosClothAssetEditor/Source/ChaosClothAssetDataflowNodes/Public/ChaosClothAsset/ProxyDeformerNode.h b/Engine/Plugins/ChaosClothAssetEditor/Source/ChaosClothAssetDataflowNodes/Public/ChaosClothAsset/ProxyDeformerNode.h
--- a/Engine/Plugins/ChaosClothAssetEditor/Source/ChaosClothAssetDataflowNodes/Public/ChaosClothAsset/ProxyDeformerNode.h
+++ b/Engine/Plugins/ChaosClothAssetEditor/Source/ChaosClothAssetDataflowNodes/Public/ChaosClothAsset/ProxyDeformerNode.h
@@ -8,6 +8,87 @@
 #include "GeometryCollection/ManagedArrayCollection.h"
 #include "ProxyDeformerNode.generated.h"
 
+/** Selection Filter Set*/
+USTRUCT()
+struct FChaosClothAssetSelectionFilterSet
+{
+	GENERATED_USTRUCT_BODY()
+
+	UPROPERTY(EditAnywhere, Category = "ProxyDeformer")
+	FChaosClothAssetConnectableIStringValue RenderSelection = { TEXT("SelectionRenderFilterSet") };
+
+	UPROPERTY(EditAnywhere, Category = "ProxyDeformer")
+	FChaosClothAssetConnectableIStringValue SimSelection = { TEXT("SelectionSimFilterSet") };
+};
+
+/** Add the proxy deformer information to this cloth collection's render data. */
+USTRUCT(Meta = (DataflowCloth, Experimental))
+struct FChaosClothAssetProxyDeformerNode_v2 : public FDataflowNode
+{
+	GENERATED_USTRUCT_BODY()
+	DATAFLOW_NODE_DEFINE_INTERNAL(FChaosClothAssetProxyDeformerNode_v2, "ProxyDeformer", "Cloth", "Cloth Simulation Proxy Deformer")
+
+public:
+
+	UPROPERTY(Meta = (DataflowInput, DataflowOutput, DataflowPassthrough = "Collection"))
+	FManagedArrayCollection Collection;
+
+	/**
+	 * The name of a selection containing all the dynamic points. Must be of group type SimVertices2D, SimVertices3D, or SimFaces.
+	 * Using an empty (or invalid) selection will make the proxy deformer consider all simulation points as dynamic points,
+	 * and will fully contribute to the render mesh animations (as opposed to using the render mesh skinning for the non dynamic points).
+	 * This selection is usually built from the same weight map set to the MaxDistance config using a WeightMapToSelection node and a very low threshold.
+	 */
+	UPROPERTY(EditAnywhere, Category = "ProxyDeformer")
+	FChaosClothAssetConnectableIStringValue SimVertexSelection;
+
+	/** Selection filter set used to restrict a selection of render vertices to a selection of simulation mesh triangles. Right click and do AddOptionPin to add more selection sets .*/
+	UPROPERTY(EditAnywhere, EditFixedSize, Category = "ProxyDeformer")
+	TArray<FChaosClothAssetSelectionFilterSet> SelectionFilterSets;
+
+	/** Whether using multiple simulation mesh triangles to influence the position of the deformed render vertex. */
+	UPROPERTY(EditAnywhere, Category = "ProxyDeformer")
+	bool bUseMultipleInfluences = false;
+
+	/** The radius around the render vertices to look for all simulation mesh triangles influencing it (AKA SkinningKernelRadius). */
+	UPROPERTY(EditAnywhere, Category = "ProxyDeformer", Meta = (EditCondition = "bUseMultipleInfluences"))
+	float InfluenceRadius = 5.f;
+
+	/**
+	 * Whether to create a smoothed _SkinningBlendWeight render weight map to ease the transition between the deformed part and the skinned part of the render mesh.
+	 * When no transition is created there will be a visible step in the rendered triangles around the edge of the kinematic/dynamic transition of the proxy simulation mesh.
+	 * The _SkinningBlendWeight render weight map is created regardless of the transition being created smooth or not, and can be later adjusted using the weight map tool.
+	 */
+	UPROPERTY(EditAnywhere, Category = "ProxyDeformer")
+	bool bUseSmoothTransition = true;
+
+	/**
+	 * The name of the render mesh weight map generated by this node detailing the contribution of the proxy deformer.
+	 * Value ranges between 0 (fully deformed) and 1 (fully skinned).
+	 * The name of this render mesh weight map cannot be changed and is only provided for further tweaking.
+	 */
+	UPROPERTY(VisibleAnywhere, Category = "ProxyDeformer", Meta = (DataflowOutput))
+	FString SkinningBlendName;
+
+	FChaosClothAssetProxyDeformerNode_v2(const Dataflow::FNodeParameters& InParam, FGuid InGuid = FGuid::NewGuid());
+
+private:
+
+	//~ Begin FDataflowNode interface
+	virtual void Evaluate(Dataflow::FContext& Context, const FDataflowOutput* Out) const override;
+	virtual TArray<Dataflow::FPin> AddPins() override;
+	virtual bool CanAddPin() const override { return true; }
+	virtual bool CanRemovePin() const override { return SelectionFilterSets.Num() > 1; }
+	virtual TArray<Dataflow::FPin> GetPinsToRemove() const override;
+	virtual void OnPinRemoved(const Dataflow::FPin& Pin) override;
+	virtual void Serialize(FArchive& Ar) override;
+	//~ End FDataflowNode interface
+
+	TArray<TPair<FName, FName>> GetSelectionFilterNames(Dataflow::FContext& Context) const;
+	Dataflow::TConnectionReference<FString> GetRenderConnectionReference(int32 Index) const;
+	Dataflow::TConnectionReference<FString> GetSimConnectionReference(int32 Index) const;
+};
+
 /** Add the proxy deformer information to this cloth collection's render data. */
 USTRUCT(Meta = (DataflowCloth))
 struct FChaosClothAssetProxyDeformerNode : public FDataflowNode
@@ -105,10 +186,10 @@ private:
 
 	//~ Begin FDataflowNode interface
 	virtual void Evaluate(Dataflow::FContext& Context, const FDataflowOutput* Out) const override;
-	virtual Dataflow::FPin AddPin() override;
+	virtual TArray<Dataflow::FPin> AddPins() override;
 	virtual bool CanAddPin() const override { return NumFilterSets < MaxNumFilterSets; }
 	virtual bool CanRemovePin() const override { return NumFilterSets > 1; }
-	virtual Dataflow::FPin GetPinToRemove() const override;
+	virtual TArray<Dataflow::FPin> GetPinsToRemove() const override;
 	virtual void OnPinRemoved(const Dataflow::FPin& Pin) override;
 	virtual void Serialize(FArchive& Ar) override;
 	//~ End FDataflowNode interface
diff --git a/Engine/Plugins/ChaosClothAssetEditor/Source/ChaosClothAssetDataflowNodes/Public/ChaosClothAsset/TerminalNode.h b/Engine/Plugins/ChaosClothAssetEditor/Source/ChaosClothAssetDataflowNodes/Public/ChaosClothAsset/TerminalNode.h
--- a/Engine/Plugins/ChaosClothAssetEditor/Source/ChaosClothAssetDataflowNodes/Public/ChaosClothAsset/TerminalNode.h
+++ b/Engine/Plugins/ChaosClothAssetEditor/Source/ChaosClothAssetDataflowNodes/Public/ChaosClothAsset/TerminalNode.h
@@ -19,6 +19,52 @@ struct FChaosClothAssetTerminalNodeRefreshAsset
 
 /** Cloth terminal node to generate a cloth asset from a cloth collection. */
 USTRUCT(Meta = (DataflowCloth, DataflowTerminal))
+struct FChaosClothAssetTerminalNode_v2 : public FDataflowTerminalNode
+{
+	GENERATED_USTRUCT_BODY()
+	DATAFLOW_NODE_DEFINE_INTERNAL(FChaosClothAssetTerminalNode_v2, "ClothAssetTerminal", "Cloth", "Cloth Terminal")  // TODO: Should the category be Terminal instead like all other terminal nodes
+
+public:
+	UPROPERTY()
+	TArray<FManagedArrayCollection> CollectionLods;
+
+	/**
+	 * Refresh the asset even if the ClothCollection hasn't changed.
+	 * Note that it is not required to manually refresh the cloth asset, this is done automatically when there is a change in the Dataflow.
+	 * This function is a developper utility used for debugging.
+	 */
+	UPROPERTY(EditAnywhere, Category = "Cloth Asset Terminal")
+	mutable FChaosClothAssetTerminalNodeRefreshAsset RefreshAsset;
+
+	FChaosClothAssetTerminalNode_v2(const Dataflow::FNodeParameters& InParam, FGuid InGuid = FGuid::NewGuid());
+
+private:
+	//~ Begin FDataflowNode interface
+	virtual void SetAssetValue(TObjectPtr<UObject> Asset, Dataflow::FContext& Context) const override;
+	virtual void Evaluate(Dataflow::FContext& Context, const FDataflowOutput* Out) const override {}
+	virtual TArray<Dataflow::FPin> AddPins() override;
+	virtual bool CanAddPin() const override { return true; }
+	virtual bool CanRemovePin() const override { return CollectionLods.Num() > 1; }
+	virtual TArray<Dataflow::FPin> GetPinsToRemove() const override;
+	virtual void OnPinRemoved(const Dataflow::FPin& Pin) override;
+	virtual void Serialize(FArchive& Ar) override;
+	//~ End FDataflowNode interface
+
+	TArray<TSharedRef<FManagedArrayCollection>> GetCleanedCollectionLodValues(Dataflow::FContext& Context) const;
+	Dataflow::TConnectionReference<FManagedArrayCollection> GetConnectionReference(int32 Index) const;
+
+	UPROPERTY()
+	mutable TArray<FChaosClothAssetLodTransitionDataCache> LODTransitionDataCache;
+
+	// This is for runtime only--used to determine if only properties need to be updated.
+	mutable bool bClothCollectionChecksumValid = false;
+	mutable uint32 ClothColllectionChecksum = 0;
+};
+
+
+
+/** Cloth terminal node to generate a cloth asset from a cloth collection. */
+USTRUCT(Meta = (DataflowCloth, DataflowTerminal, Deprecated = 5.5))
 struct FChaosClothAssetTerminalNode : public FDataflowTerminalNode
 {
 	GENERATED_USTRUCT_BODY()
@@ -62,10 +108,10 @@ private:
 	//~ Begin FDataflowNode interface
 	virtual void SetAssetValue(TObjectPtr<UObject> Asset, Dataflow::FContext& Context) const override;
 	virtual void Evaluate(Dataflow::FContext& Context, const FDataflowOutput* Out) const override {}
-	virtual Dataflow::FPin AddPin() override;
+	virtual TArray<Dataflow::FPin> AddPins() override;
 	virtual bool CanAddPin() const override { return NumLods < MaxLods; }
 	virtual bool CanRemovePin() const override { return NumLods > 1; }
-	virtual Dataflow::FPin GetPinToRemove() const override;
+	virtual TArray<Dataflow::FPin> GetPinsToRemove() const override;
 	virtual void OnPinRemoved(const Dataflow::FPin& Pin) override;
 	virtual void Serialize(FArchive& Ar) override;
 	//~ End FDataflowNode interface
diff --git a/Engine/Plugins/Experimental/GeometryCollectionPlugin/Source/GeometryCollectionNodes/Private/Dataflow/GeometryCollectionMathNodes.cpp b/Engine/Plugins/Experimental/GeometryCollectionPlugin/Source/GeometryCollectionNodes/Private/Dataflow/GeometryCollectionMathNodes.cpp
--- a/Engine/Plugins/Experimental/GeometryCollectionPlugin/Source/GeometryCollectionNodes/Private/Dataflow/GeometryCollectionMathNodes.cpp
+++ b/Engine/Plugins/Experimental/GeometryCollectionPlugin/Source/GeometryCollectionNodes/Private/Dataflow/GeometryCollectionMathNodes.cpp
@@ -928,11 +928,11 @@ void FFloatMathExpressionDataflowNode::Evaluate(Dataflow::FContext& Context, con
 FMathExpressionDataflowNode::FMathExpressionDataflowNode(const Dataflow::FNodeParameters& InParam, FGuid InGuid)
 	: FDataflowNode(InParam, InGuid)
 {
-	FDataflowInput* InputA = RegisterInputConnection(&A);
-	FDataflowInput* InputB = RegisterInputConnection(&B);
-	FDataflowInput* InputC = RegisterInputConnection(&C);
-	FDataflowInput* InputD = RegisterInputConnection(&D);
-	FDataflowOutput* Output = RegisterOutputConnection(&ReturnValue);
+	RegisterInputConnection(&A);
+	RegisterInputConnection(&B);
+	RegisterInputConnection(&C);
+	RegisterInputConnection(&D);
+	RegisterOutputConnection(&ReturnValue);
 }
 
 void FMathExpressionDataflowNode::Evaluate(Dataflow::FContext& Context, const FDataflowOutput* Out) const
diff --git a/Engine/Source/Runtime/Experimental/Dataflow/Core/Private/Dataflow/DataflowConnection.cpp b/Engine/Source/Runtime/Experimental/Dataflow/Core/Private/Dataflow/DataflowConnection.cpp
--- a/Engine/Source/Runtime/Experimental/Dataflow/Core/Private/Dataflow/DataflowConnection.cpp
+++ b/Engine/Source/Runtime/Experimental/Dataflow/Core/Private/Dataflow/DataflowConnection.cpp
@@ -7,6 +7,9 @@
 
 #include UE_INLINE_GENERATED_CPP_BY_NAME(DataflowConnection)
 
+const Dataflow::FConnectionKey Dataflow::FConnectionKey::Invalid = { (uint32)INDEX_NONE, INDEX_NONE, (uint32)INDEX_NONE };
+const Dataflow::FPin Dataflow::FPin::InvalidPin = { Dataflow::FPin::EDirection::NONE, NAME_None, NAME_None };
+
 FDataflowConnection::FDataflowConnection(Dataflow::FPin::EDirection InDirection, FName InType, FName InName, FDataflowNode* InOwningNode, const FProperty* InProperty, FGuid InGuid)
 	: Type(InType)
 	, Name(InName)
diff --git a/Engine/Source/Runtime/Experimental/Dataflow/Core/Private/Dataflow/DataflowCoreNodes.cpp b/Engine/Source/Runtime/Experimental/Dataflow/Core/Private/Dataflow/DataflowCoreNodes.cpp
--- a/Engine/Source/Runtime/Experimental/Dataflow/Core/Private/Dataflow/DataflowCoreNodes.cpp
+++ b/Engine/Source/Runtime/Experimental/Dataflow/Core/Private/Dataflow/DataflowCoreNodes.cpp
@@ -18,7 +18,8 @@ FDataflowReRouteNode::FDataflowReRouteNode(const Dataflow::FNodeParameters& Para
 	: Super(Param, InGuid)
 {
 	RegisterInputConnection(&Value);
-	RegisterOutputConnection(&Value, &Value);
+	RegisterOutputConnection(&Value)
+		.SetPassthroughInput(&Value);
 }
 
 void FDataflowReRouteNode::Evaluate(Dataflow::FContext& Context, const FDataflowOutput* Out) const
diff --git a/Engine/Source/Runtime/Experimental/Dataflow/Core/Private/Dataflow/DataflowInputOutput.cpp b/Engine/Source/Runtime/Experimental/Dataflow/Core/Private/Dataflow/DataflowInputOutput.cpp
--- a/Engine/Source/Runtime/Experimental/Dataflow/Core/Private/Dataflow/DataflowInputOutput.cpp
+++ b/Engine/Source/Runtime/Experimental/Dataflow/Core/Private/Dataflow/DataflowInputOutput.cpp
@@ -102,6 +102,28 @@ void FDataflowInput::FixAndPropagateType(FName InType)
 	}
 }
 
+FDataflowArrayInput::FDataflowArrayInput(int32 InIndex, const Dataflow::FArrayInputParameters& Param)
+	: FDataflowInput(Param)
+	, Index(InIndex)
+	, ElementOffset(Param.InnerOffset)
+	, ArrayProperty(Param.ArrayProperty)
+{}
+
+void* FDataflowArrayInput::RealAddress() const
+{
+	if (void* ContainerRealAddress = Super::RealAddress())
+	{
+		if (ArrayProperty)
+		{
+			if (const void* const AddressAtIndex = ArrayProperty->GetValueAddressAtIndex_Direct(ArrayProperty->Inner, ContainerRealAddress, Index))
+			{
+				return (void*)((size_t)AddressAtIndex + (size_t)ElementOffset);
+			}
+		}
+	}
+	return nullptr;
+}
+
 //
 //
 //  Output
@@ -163,6 +185,27 @@ bool FDataflowOutput::RemoveConnection(FDataflowConnection* InInput)
 	Connections.RemoveSwap((FDataflowInput*)InInput); return true;
 }
 
+FDataflowOutput& FDataflowOutput::SetPassthroughInput(const Dataflow::FConnectionReference& Reference)
+{
+	check(OwningNode);
+	const FDataflowInput* const PassthroughInput = OwningNode->FindInput(Reference);
+	check(PassthroughInput);
+	PassthroughKey = PassthroughInput->GetConnectionKey();
+	return *this;
+}
+
+FDataflowOutput& FDataflowOutput::SetPassthroughInput(const Dataflow::FConnectionKey& Key)
+{
+	check(Key == Dataflow::FConnectionKey::Invalid || !OwningNode || OwningNode->FindInput(Key));
+	PassthroughKey = Key;
+	return *this;
+}
+
+const FDataflowInput* FDataflowOutput::GetPassthroughInput() const
+{
+	return OwningNode ? OwningNode->FindInput(PassthroughKey) : nullptr;
+}
+
 void FDataflowOutput::Invalidate(const Dataflow::FTimestamp& ModifiedTimestamp)
 {
 	for (FDataflowConnection* Con : GetConnections())
@@ -179,9 +222,9 @@ bool FDataflowOutput::Evaluate(Dataflow::FContext& Context) const
 	{
 		return Context.Evaluate(*this);
 	}
-	else if (const FDataflowInput* PassthroughInput = GetPassthroughInput())
+	else if (const FDataflowInput* const PassthroughInput = GetPassthroughInput())
 	{
-		ForwardInput(PassthroughInput->RealAddress(), Context);
+		ForwardInput(PassthroughInput, Context);
 		return true;
 	}
 
@@ -219,23 +262,27 @@ TFuture<bool> FDataflowOutput::EvaluateParallel(Dataflow::FContext& Context) con
 	return Async(EAsyncExecution::TaskGraph, [&]() -> bool { return this->Evaluate(Context); });
 }
 
-const FDataflowInput* FDataflowOutput::GetPassthroughInput() const
+
+void FDataflowOutput::ForwardInput(const Dataflow::FConnectionReference& InputReference, Dataflow::FContext& Context) const
 {
-	return OwningNode ? OwningNode->FindInput(GetPassthroughRealAddress()) : nullptr;
+	if (Property && OwningNode)
+	{
+		const FDataflowInput* InputToForward = OwningNode->FindInput(InputReference);
+		ForwardInput(InputToForward, Context);
+	}
 }
 
-void FDataflowOutput::ForwardInput(const void* InputReference, Dataflow::FContext& Context) const
+void FDataflowOutput::ForwardInput(const FDataflowInput* Input, Dataflow::FContext& Context) const
 {
 	if (Property && OwningNode)
 	{
-		const FDataflowInput* InputToForward = OwningNode->FindInput(InputReference);
-		if (InputToForward->GetConnectedOutputs().Num())
+		if (Input->GetConnectedOutputs().Num())
 		{
-			ensure(InputToForward->GetType() == GetType());
-			ensure(InputToForward->GetConnectedOutputs().Num() == 1);
-			if (const FDataflowOutput* ConnectionOut = InputToForward->GetConnection())
+			ensure(Input->GetType() == GetType());
+			ensure(Input->GetConnectedOutputs().Num() == 1);
+			if (const FDataflowOutput* ConnectionOut = Input->GetConnection())
 			{
-				InputToForward->PullValue(Context);
+				Input->PullValue(Context);
 				Context.SetDataReference(CacheKey(), Property, ConnectionOut->CacheKey());
 			}
 		}
diff --git a/Engine/Source/Runtime/Experimental/Dataflow/Core/Private/Dataflow/DataflowNode.cpp b/Engine/Source/Runtime/Experimental/Dataflow/Core/Private/Dataflow/DataflowNode.cpp
--- a/Engine/Source/Runtime/Experimental/Dataflow/Core/Private/Dataflow/DataflowNode.cpp
+++ b/Engine/Source/Runtime/Experimental/Dataflow/Core/Private/Dataflow/DataflowNode.cpp
@@ -22,6 +22,89 @@ const FLinearColor FDataflowNode::DefaultNodeBodyTintColor = FLinearColor(0.f, 0
 
 const FName FDataflowAnyType::TypeName = TEXT("FDataflowAnyType");
 
+namespace Dataflow::Private
+{
+static uint32 GetArrayElementOffsetFromReference(const FArrayProperty* const ArrayProperty, const Dataflow::FConnectionReference& Reference)
+{
+	check(ArrayProperty);
+	if (const void* const AddressAtIndex = ArrayProperty->GetValueAddressAtIndex_Direct(ArrayProperty->Inner, const_cast<void*>(Reference.ContainerReference), Reference.Index))
+	{
+		check((size_t)Reference.Reference >= (size_t)AddressAtIndex);
+		check((int32)((size_t)Reference.Reference - (size_t)AddressAtIndex) < ArrayProperty->Inner->ElementSize);
+		return (uint32)((size_t)Reference.Reference - (size_t)AddressAtIndex);
+	}
+	return INDEX_NONE;
+}
+
+static const FProperty* FindProperty(const UStruct* Struct, const void* StructValue, const void* InProperty, const FName& PropertyName, TArray<const FProperty*>* OutPropertyChain)
+{
+	const FProperty* Property = nullptr;
+	for (FPropertyValueIterator PropertyIt(FProperty::StaticClass(), Struct, StructValue); PropertyIt; ++PropertyIt)
+	{
+		if (InProperty == PropertyIt.Value() && (PropertyName == NAME_None || PropertyName == PropertyIt.Key()->GetName()))
+		{
+			Property = PropertyIt.Key();
+			if (OutPropertyChain)
+			{
+				PropertyIt.GetPropertyChain(*OutPropertyChain);
+			}
+			break;
+		}
+	}
+	return Property;
+}
+
+static const FProperty& FindPropertyChecked(const UStruct* Struct, const void* StructValue, const void* InProperty, const FName& PropertyName, TArray<const FProperty*>* OutPropertyChain)
+{
+	const FProperty* Property = FindProperty(Struct, StructValue, InProperty, PropertyName, OutPropertyChain);
+	check(Property);
+	return *Property;
+}
+
+static FString GetPinToolTipFromProperty(const FProperty* Property)
+{
+#if WITH_EDITORONLY_DATA
+	check(Property);
+	if (Property->HasMetaData(TEXT("Tooltip")))
+	{
+		const FString ToolTipStr = Property->GetToolTipText(true).ToString();
+		if (ToolTipStr.Len() > 0)
+		{
+			TArray<FString> OutArr;
+			const int32 NumElems = ToolTipStr.ParseIntoArray(OutArr, TEXT(":\r\n"));
+
+			if (NumElems == 2)
+			{
+				return OutArr[1];  // Return tooltip meta text
+			}
+			else if (NumElems == 1)
+			{
+				return OutArr[0];  // Return doc comment
+			}
+		}
+	}
+#endif
+	return "";
+}
+
+static TArray<FString> GetPinMetaDataFromProperty(const FProperty* Property)
+{
+	TArray<FString> MetaDataStrArr;
+#if WITH_EDITORONLY_DATA
+	check(Property);
+	if (Property->HasMetaData(FDataflowNode::DataflowPassthrough))
+	{
+		MetaDataStrArr.Add("Passthrough");
+	}
+	if (Property->HasMetaData(FDataflowNode::DataflowIntrinsic))
+	{
+		MetaDataStrArr.Add("Intrinsic");
+	}
+#endif
+	return MetaDataStrArr;
+}
+};
+
 //
 // Inputs
 //
@@ -48,7 +131,7 @@ void FDataflowNode::AddInput(FDataflowInput* InPtr)
 {
 	if (InPtr)
 	{
-		for (const TPair<int32, FDataflowInput*>& Elem : Inputs)
+		for (const TPair<Dataflow::FConnectionKey, FDataflowInput*>& Elem : ExpandedInputs)
 		{
 			const FDataflowInput* const In = Elem.Value;
 			ensureMsgf(!In->GetName().IsEqual(InPtr->GetName()), TEXT("Add Input Failed: Existing Node input already defined with name (%s)"), *InPtr->GetName().ToString());
@@ -56,17 +139,17 @@ void FDataflowNode::AddInput(FDataflowInput* InPtr)
 
 		check(InPtr->GetOwningNode() == this);
 
-		const uint32 PropertyOffset = InPtr->GetOffset();
-		if (ensure(!Inputs.Contains(PropertyOffset)))
+		const Dataflow::FConnectionKey Key(InPtr->GetOffset(), InPtr->GetContainerIndex(), InPtr->GetContainerElementOffset());
+		if (ensure(!ExpandedInputs.Contains(Key)))
 		{
-			Inputs.Add(PropertyOffset, InPtr);
+			ExpandedInputs.Add(Key, InPtr);
 		}
 	}
 }
 
 FDataflowInput* FDataflowNode::FindInput(FName InName)
 {
-	for (TPair<int32, FDataflowInput*>& Elem : Inputs)
+	for (TPair<Dataflow::FConnectionKey, FDataflowInput*>& Elem : ExpandedInputs)
 	{
 		FDataflowInput* const Con = Elem.Value;
 		if (Con->GetName().IsEqual(InName))
@@ -77,10 +160,9 @@ FDataflowInput* FDataflowNode::FindInput(FName InName)
 	return nullptr;
 }
 
-
 const FDataflowInput* FDataflowNode::FindInput(FName InName) const
 {
-	for (const TPair<int32, FDataflowInput*>& Elem : Inputs)
+	for (const TPair<Dataflow::FConnectionKey, FDataflowInput*>& Elem : ExpandedInputs)
 	{
 		const FDataflowInput* const Con = Elem.Value;
 		if (Con->GetName().IsEqual(InName))
@@ -91,31 +173,75 @@ const FDataflowInput* FDataflowNode::FindInput(FName InName) const
 	return nullptr;
 }
 
-const FDataflowInput* FDataflowNode::FindInput(const void* Reference) const
+const FDataflowInput* FDataflowNode::FindInput(const Dataflow::FConnectionKey& Key) const
 {
-	const int32 Key = GetConnectionOffsetFromReference(Reference);
-	if (const FDataflowInput* const* Con = Inputs.Find(Key))
+	if (const FDataflowInput* const* Con = ExpandedInputs.Find(Key))
 	{
-		check(*Con && (*Con)->RealAddress() == Reference);
+		check(*Con);
 		return *Con;
 	}
 	return nullptr;
 }
 
-FDataflowInput* FDataflowNode::FindInput(void* Reference)
+const FDataflowInput* FDataflowNode::FindInput(const Dataflow::FConnectionReference& Reference) const
+{
+	const Dataflow::FConnectionKey Key = GetKeyFromReference(Reference);
+	if (const FDataflowInput* const Con = FindInput(Key))
+	{
+		check(Con->RealAddress() == Reference.Reference);
+		return Con;
+	}
+	if (Reference.ContainerReference == nullptr && !InputArrayProperties.IsEmpty())
+	{
+		// Search through all connections to see if Reference is the RealAddress of an array property.
+		for (const TPair<Dataflow::FConnectionKey, FDataflowInput*>& Elem : ExpandedInputs)
+		{
+			const FDataflowInput* const Con = Elem.Value;
+			if (Con->RealAddress() == Reference.Reference)
+			{
+				return Con;
+			}
+		}
+	}
+	return nullptr;
+}
+
+FDataflowInput* FDataflowNode::FindInput(const Dataflow::FConnectionKey& Key)
 {
-	const int32 Key = GetConnectionOffsetFromReference(Reference);
-	if (FDataflowInput* const* Con = Inputs.Find(Key))
+	if (FDataflowInput* const* Con = ExpandedInputs.Find(Key))
 	{
-		check(*Con && (*Con)->RealAddress() == Reference);
+		check(*Con);
 		return *Con;
 	}
 	return nullptr;
 }
 
+FDataflowInput* FDataflowNode::FindInput(const Dataflow::FConnectionReference& Reference)
+{
+	const Dataflow::FConnectionKey Key = GetKeyFromReference(Reference);
+	if (FDataflowInput* const Con = FindInput(Key))
+	{
+		check(Con->RealAddress() == Reference.Reference);
+		return Con;
+	}
+	if (Reference.ContainerReference == nullptr && !InputArrayProperties.IsEmpty())
+	{
+		// Search through all connections to see if Reference is the RealAddress of an array property.
+		for (TPair<Dataflow::FConnectionKey, FDataflowInput*>& Elem : ExpandedInputs)
+		{
+			FDataflowInput* const Con = Elem.Value;
+			if (Con->RealAddress() == Reference.Reference)
+			{
+				return Con;
+			}
+		}
+	}
+	return nullptr;
+}
+
 const FDataflowInput* FDataflowNode::FindInput(const FGuid& InGuid) const
 {
-	for (const TPair<int32, FDataflowInput*>& Elem : Inputs)
+	for (const TPair<Dataflow::FConnectionKey, FDataflowInput*>& Elem : ExpandedInputs)
 	{
 		const FDataflowInput* const Con = Elem.Value;
 		if (Con->GetGuid() == InGuid)
@@ -129,23 +255,23 @@ const FDataflowInput* FDataflowNode::FindInput(const FGuid& InGuid) const
 TArray< FDataflowInput* > FDataflowNode::GetInputs() const
 {
 	TArray< FDataflowInput* > Result;
-	Inputs.GenerateValueArray(Result);
+	ExpandedInputs.GenerateValueArray(Result);
 	return Result;
 }
 
 void FDataflowNode::ClearInputs()
 {
-	for (TPair<int32, FDataflowInput*>& Elem : Inputs)
+	for (TPair<Dataflow::FConnectionKey, FDataflowInput*>& Elem : ExpandedInputs)
 	{
 		FDataflowInput* const Con = Elem.Value;
 		delete Con;
 	}
-	Inputs.Reset();
+	ExpandedInputs.Reset();
 }
 
 bool FDataflowNode::HasHideableInputs() const
 {
-	for (const TPair<int32, FDataflowInput*>& Elem : Inputs)
+	for (const TPair<Dataflow::FConnectionKey, FDataflowInput*>& Elem : ExpandedInputs)
 	{
 		const FDataflowInput* const Con = Elem.Value;
 		if (Con->GetCanHidePin())
@@ -158,7 +284,7 @@ bool FDataflowNode::HasHideableInputs() const
 
 bool FDataflowNode::HasHiddenInputs() const
 {
-	for (const TPair<int32, FDataflowInput*>& Elem : Inputs)
+	for (const TPair<Dataflow::FConnectionKey, FDataflowInput*>& Elem : ExpandedInputs)
 	{
 		const FDataflowInput* const Con = Elem.Value;
 		if (Con->GetPinIsHidden())
@@ -247,28 +373,48 @@ const FDataflowOutput* FDataflowNode::FindOutput(uint32 InGuidHash) const
 	return nullptr;
 }
 
-const FDataflowOutput* FDataflowNode::FindOutput(const void* Reference) const
+const FDataflowOutput* FDataflowNode::FindOutput(const Dataflow::FConnectionKey& Key) const
 {
-	const int32 Key = GetConnectionOffsetFromReference(Reference);
-	if (const FDataflowOutput* const* Con = Outputs.Find(Key))
+	if (const FDataflowOutput* const* Con = Outputs.Find(Key.Offset))
 	{
-		check(*Con && (*Con)->RealAddress() == Reference);
+		check(*Con);
 		return *Con;
 	}
 	return nullptr;
 }
 
-FDataflowOutput* FDataflowNode::FindOutput(void* Reference)
+const FDataflowOutput* FDataflowNode::FindOutput(const Dataflow::FConnectionReference& Reference) const
 {
-	const int32 Key = GetConnectionOffsetFromReference(Reference);
-	if (FDataflowOutput* const* Con = Outputs.Find(Key))
+	const Dataflow::FConnectionKey Key = GetKeyFromReference(Reference);
+	if (const FDataflowOutput* const Con = FindOutput(Key))
 	{
-		check(*Con && (*Con)->RealAddress() == Reference);
+		check(Con->RealAddress() == Reference.Reference);
+		return Con;
+	}
+	return nullptr;
+}
+
+FDataflowOutput* FDataflowNode::FindOutput(const Dataflow::FConnectionKey& Key)
+{
+	if (FDataflowOutput* const* Con = Outputs.Find(Key.Offset))
+	{
+		check(*Con);
 		return *Con;
 	}
 	return nullptr;
 }
 
+FDataflowOutput* FDataflowNode::FindOutput(const Dataflow::FConnectionReference& Reference)
+{
+	const Dataflow::FConnectionKey Key = GetKeyFromReference(Reference);
+	if (FDataflowOutput* const Con = FindOutput(Key))
+	{
+		check(Con->RealAddress() == Reference.Reference);
+		return Con;
+	}
+	return nullptr;
+}
+
 const FDataflowOutput* FDataflowNode::FindOutput(const FGuid& InGuid) const
 {
 	for (const TPair<int32, FDataflowOutput*>& Elem : Outputs)
@@ -335,7 +481,7 @@ bool FDataflowNode::HasHiddenOutputs() const
 TArray<Dataflow::FPin> FDataflowNode::GetPins() const
 {
 	TArray<Dataflow::FPin> RetVal;
-	for (const TPair<int32, FDataflowInput*>& Elem : Inputs)
+	for (const TPair<Dataflow::FConnectionKey, FDataflowInput*>& Elem : ExpandedInputs)
 	{
 		const FDataflowInput* const Con = Elem.Value;
 		RetVal.Add({ Dataflow::FPin::EDirection::INPUT,Con->GetType(), Con->GetName(), Con->GetPinIsHidden()});
@@ -352,7 +498,7 @@ void FDataflowNode::UnregisterPinConnection(const Dataflow::FPin& Pin)
 {
 	if (Pin.Direction == Dataflow::FPin::EDirection::INPUT)
 	{
-		for (TMap< int32, FDataflowInput*>::TIterator Iter = Inputs.CreateIterator(); Iter; ++Iter)
+		for (TMap<Dataflow::FConnectionKey, FDataflowInput*>::TIterator Iter = ExpandedInputs.CreateIterator(); Iter; ++Iter)
 		{
 			FDataflowInput* Con = Iter.Value();
 			if (Con->GetName().IsEqual(Pin.Name) && Con->GetType().IsEqual(Pin.Type))
@@ -386,6 +532,15 @@ void FDataflowNode::UnregisterPinConnection(const Dataflow::FPin& Pin)
 
 void FDataflowNode::Invalidate(const Dataflow::FTimestamp& InModifiedTimestamp)
 {
+	if (bPauseInvalidations)
+	{
+		if (PausedModifiedTimestamp < InModifiedTimestamp)
+		{
+			PausedModifiedTimestamp = InModifiedTimestamp;
+		}
+		return;
+	}
+
 	if (LastModifiedTimestamp < InModifiedTimestamp)
 	{
 		LastModifiedTimestamp = InModifiedTimestamp;
@@ -404,30 +559,25 @@ void FDataflowNode::Invalidate(const Dataflow::FTimestamp& InModifiedTimestamp)
 
 const FProperty* FDataflowNode::FindProperty(const UStruct* Struct, const void* InProperty, const FName& PropertyName, TArray<const FProperty*>* OutPropertyChain) const
 {
-	const FProperty* Property = nullptr;
-	for (FPropertyValueIterator PropertyIt(FProperty::StaticClass(), Struct, this); PropertyIt; ++PropertyIt)
-	{
-		if (InProperty == PropertyIt.Value() && (PropertyName == NAME_None || PropertyName == PropertyIt.Key()->GetName()))
-		{
-			Property = PropertyIt.Key();
-			if (OutPropertyChain)
-			{
-				PropertyIt.GetPropertyChain(*OutPropertyChain);
-			}
-			break;
-		}
-	}
-	return Property;
+	return Dataflow::Private::FindProperty(Struct, this, InProperty, PropertyName, OutPropertyChain);
+}
+
+const FProperty& FDataflowNode::FindPropertyChecked(const UStruct* Struct, const void* InProperty, const FName& PropertyName, TArray<const FProperty*>* OutPropertyChain) const
+{
+	return Dataflow::Private::FindPropertyChecked(Struct, this, InProperty, PropertyName, OutPropertyChain);
 }
 
 const FProperty* FDataflowNode::FindProperty(const UStruct* Struct, const FName& PropertyFullName, TArray<const FProperty*>* OutPropertyChain) const
 {
+	// If PropertyFullName corresponds with an array property, it will contain a [ContainerIndex]. We don't care about which element in the array we're in--the FProperty will be the same.
+	const FString PropertyFullNameStringIndexNone = StripContainerIndexFromPropertyFullName(PropertyFullName.ToString());
+
 	const FProperty* Property = nullptr;
 	for (FPropertyValueIterator PropertyIt(FProperty::StaticClass(), Struct, this); PropertyIt; ++PropertyIt)
 	{
 		TArray<const FProperty*> PropertyChain;
 		PropertyIt.GetPropertyChain(PropertyChain);
-		if (GetPropertyFullName(PropertyChain) == PropertyFullName)
+		if (GetPropertyFullNameString(PropertyChain, INDEX_NONE) == PropertyFullNameStringIndexNone)
 		{
 			Property = PropertyIt.Key();
 			if (OutPropertyChain)
@@ -471,130 +621,221 @@ uint32 FDataflowNode::GetConnectionOffsetFromReference(const void* Reference) co
 	return (uint32)((size_t)Reference - (size_t)this);
 }
 
-FString FDataflowNode::GetPropertyFullNameString(const TConstArrayView<const FProperty*>& PropertyChain)
+Dataflow::FConnectionKey FDataflowNode::GetKeyFromReference(const Dataflow::FConnectionReference& Reference) const
+{
+	Dataflow::FConnectionKey Key;
+	Key.Offset = Reference.ContainerReference ? GetConnectionOffsetFromReference(Reference.ContainerReference) : GetConnectionOffsetFromReference(Reference.Reference);
+	Key.ContainerIndex = Reference.Index;
+	Key.ContainerElementOffset = INDEX_NONE;
+	if (const FArrayProperty* const* ArrayProperty = InputArrayProperties.Find(Key.Offset))
+	{
+		Key.ContainerElementOffset = Dataflow::Private::GetArrayElementOffsetFromReference(*ArrayProperty, Reference);
+	}
+	return Key;
+}
+
+FString FDataflowNode::GetPropertyFullNameString(const TConstArrayView<const FProperty*>& PropertyChain, int32 ContainerIndex)
 {
 	FString PropertyFullName;
-	for (const FProperty* const Property : PropertyChain)
+	bool bFoundArrayProperty = false;
+	for (int32 Index = PropertyChain.Num()-1; Index >= 0 ; --Index)
 	{
-		const FString PropertyName = Property->GetName();
+		const FProperty* const Property = PropertyChain[Index];
+		FString PropertyName = Property->GetName();
+		if (const FArrayProperty* const ArrayProperty = CastField<FArrayProperty>(Property))
+		{
+			if (ContainerIndex != INDEX_NONE)
+			{
+				check(!bFoundArrayProperty); // We only expect to find one array to substitute in.
+				bFoundArrayProperty = true;
+				PropertyName = FString::Format(TEXT("{0}[{1}]"), { PropertyName, ContainerIndex });
+			}
+
+			// Skip the next property. It has the same name as the container (e.g., otherwise you'll get MyFloatArray[5].MyFloatArray)
+			--Index;
+		}
+
 		PropertyFullName = PropertyFullName.IsEmpty() ?
 			PropertyName :
-			FString::Format(TEXT("{0}.{1}"), { PropertyName, PropertyFullName });
+			FString::Format(TEXT("{0}.{1}"), { PropertyFullName, PropertyName });
 	}
 	return PropertyFullName;
 }
 
-FName FDataflowNode::GetPropertyFullName(const TArray<const FProperty*>& PropertyChain)
+FName FDataflowNode::GetPropertyFullName(const TArray<const FProperty*>& PropertyChain, int32 ContainerIndex)
 {
-	const FString PropertyFullName = GetPropertyFullNameString(TConstArrayView<const FProperty*>(PropertyChain));
+	const FString PropertyFullName = GetPropertyFullNameString(TConstArrayView<const FProperty*>(PropertyChain), ContainerIndex);
 	return FName(*PropertyFullName);
 }
 
-FText FDataflowNode::GetPropertyDisplayNameText(const TArray<const FProperty*>& PropertyChain)
+FString FDataflowNode::StripContainerIndexFromPropertyFullName(const FString& InPropertyFullName)
+{
+	FString PropertyFullName = InPropertyFullName;
+	FString PropertyFullNameStripped;
+
+	int32 OpenSquareBracketIndex, CloseSquareBracketIndex;
+	while((PropertyFullName.FindChar('[', OpenSquareBracketIndex) && PropertyFullName.FindChar(']', CloseSquareBracketIndex)
+		&& OpenSquareBracketIndex < CloseSquareBracketIndex))
+	{
+		if (CloseSquareBracketIndex > OpenSquareBracketIndex + 1 && PropertyFullName.Mid(OpenSquareBracketIndex + 1, CloseSquareBracketIndex - OpenSquareBracketIndex - 1).IsNumeric())
+		{
+			// number within brackets. remove it.
+			PropertyFullNameStripped += PropertyFullName.Left(OpenSquareBracketIndex);
+		}
+		else
+		{
+			// We found some other brackets like [foo] or []. These didn't come from our ContainerIndex. Just leave them and move on.
+			PropertyFullNameStripped += PropertyFullName.Left(CloseSquareBracketIndex + 1);
+		}
+		PropertyFullName.RightChopInline(CloseSquareBracketIndex + 1);
+	}
+	PropertyFullNameStripped += PropertyFullName;
+	return PropertyFullNameStripped;
+}
+
+FText FDataflowNode::GetPropertyDisplayNameText(const TArray<const FProperty*>& PropertyChain, int32 ContainerIndex)
 {
 #if WITH_EDITORONLY_DATA  // GetDisplayNameText() is only available if WITH_EDITORONLY_DATA
 	FText PropertyText;
-	for (const FProperty* const Property : PropertyChain)
+	bool bFoundArrayProperty = false;
+	for (int32 Index = PropertyChain.Num() - 1; Index >= 0; --Index)
 	{
+		const FProperty* const Property = PropertyChain[Index];
+		FText PropertyDisplayName = Property->GetDisplayNameText();
+		if (const FArrayProperty* const ArrayProperty = CastField<FArrayProperty>(Property))
+		{
+			if (ContainerIndex != INDEX_NONE)
+			{
+				check(!bFoundArrayProperty); // We only expect to find one array to substitute in.
+				bFoundArrayProperty = true;	
+				static const FTextFormat TextFormatContainer(NSLOCTEXT("DataflowNode", "PropertyDisplayNameTextContainer", "{0}[{1}]"));
+				PropertyDisplayName = FText::Format(TextFormatContainer, PropertyDisplayName, ContainerIndex);
+			}
+
+			// Skip ElemProperty. Otherwise you get names like "MyFloatArray[0].MyFloatArray" when you just want "MyFloatArray[0]"
+			--Index;
+		}
+
 		static const FTextFormat TextFormat(NSLOCTEXT("DataflowNode", "PropertyDisplayNameTextConcatenator", "{0}.{1}"));
 		PropertyText = PropertyText.IsEmpty() ?
-			Property->GetDisplayNameText() :
-			FText::Format(TextFormat, Property->GetDisplayNameText(), PropertyText);
+			PropertyDisplayName :
+			FText::Format(TextFormat, PropertyText, PropertyDisplayName);
 	}
+
 	return PropertyText;
 #else
-	return FText::FromName(GetPropertyFullName(PropertyChain));
+	return FText::FromName(GetPropertyFullName(PropertyChain, ContainerIndex));
 #endif
 }
 
-bool FDataflowNode::InitConnectionParametersFromPropertyReference(const FStructOnScope& StructOnScope, const void* PropertyRef, const FName& PropertyName, Dataflow::FConnectionParameters& OutParams)
+void FDataflowNode::InitConnectionParametersFromPropertyReference(const FStructOnScope& StructOnScope, const void* PropertyRef, const FName& PropertyName, Dataflow::FConnectionParameters& OutParams)
 {
-		if (const UStruct* Struct = StructOnScope.GetStruct())
-		{
-			TArray<const FProperty*> PropertyChain;
-			const FProperty* const Property = FindProperty(Struct, PropertyRef, PropertyName, &PropertyChain);
-			if (ensure(Property && PropertyChain.Num()))
-			{
-				FString ExtendedType;
-				const FString CPPType = Property->GetCPPType(&ExtendedType);
-				OutParams.Type = FName(CPPType + ExtendedType);
-				OutParams.Name = GetPropertyFullName(PropertyChain);
-				OutParams.Property = Property;
-				OutParams.Owner = this;
-				OutParams.Offset = GetConnectionOffsetFromReference(PropertyRef);
-				check(OutParams.Offset == GetPropertyOffset(PropertyChain));
-				return true;
-			}
-		}
-	return false;
+	const UStruct* Struct = StructOnScope.GetStruct();
+	check(Struct);
+	TArray<const FProperty*> PropertyChain;
+	const FProperty& Property = FindPropertyChecked(Struct, PropertyRef, PropertyName, &PropertyChain);
+	check(PropertyChain.Num());
+	FString ExtendedType;
+	const FString CPPType = Property.GetCPPType(&ExtendedType);
+	OutParams.Type = FName(CPPType + ExtendedType);
+	OutParams.Name = GetPropertyFullName(PropertyChain);
+	OutParams.Property = &Property;
+	OutParams.Owner = this;
+	OutParams.Offset = GetConnectionOffsetFromReference(PropertyRef);
+	check(OutParams.Offset == GetPropertyOffset(PropertyChain));
 }
 
-FDataflowInput* FDataflowNode::RegisterInputConnectionInternal(const void* InProperty, const FName& PropertyName)
+FDataflowInput& FDataflowNode::RegisterInputConnectionInternal(const Dataflow::FConnectionReference& Reference, const FName& PropertyName)
 {
-	if (TUniquePtr<FStructOnScope> ScriptOnStruct = TUniquePtr<FStructOnScope>(NewStructOnScope()))
-	{
-		Dataflow::FInputParameters InputParams;
-		if (InitConnectionParametersFromPropertyReference(*ScriptOnStruct, InProperty, PropertyName, InputParams))
-		{
-			FDataflowInput* const Input = new FDataflowInput(InputParams);
-			check(Input->RealAddress() == InProperty);
-			AddInput(Input);
-			check(FindInput(InProperty) == Input);
-			return Input;
-		}
-	}
-	return nullptr;
+	TUniquePtr<FStructOnScope> ScriptOnStruct = TUniquePtr<FStructOnScope>(NewStructOnScope());
+	check(ScriptOnStruct);
+	Dataflow::FInputParameters InputParams;
+	InitConnectionParametersFromPropertyReference(*ScriptOnStruct, Reference.Reference, PropertyName, InputParams);
+	FDataflowInput* const Input = new FDataflowInput(InputParams);
+	check(Input->RealAddress() == Reference.Reference);
+	AddInput(Input);
+	check(FindInput(Reference) == Input);
+	return *Input;
 }
 
-void FDataflowNode::UnregisterInputConnection(const void* InProperty, const FName& PropertyName)
+FDataflowInput& FDataflowNode::RegisterInputArrayConnectionInternal(const Dataflow::FConnectionReference& Reference, const FName& ElementPropertyName,
+	const FName& ArrayPropertyName)
 {
-	if (TUniquePtr<FStructOnScope> ScriptOnStruct = TUniquePtr<FStructOnScope>(NewStructOnScope()))
+	TUniquePtr<FStructOnScope> ScriptOnStruct = TUniquePtr<FStructOnScope>(NewStructOnScope());
+	check(ScriptOnStruct);
+	const UStruct* Struct = ScriptOnStruct->GetStruct();
+	check(Struct);
+	Dataflow::FArrayInputParameters InputParams;
+	InputParams.Owner = this;
+
+	// Find the Array property.
+	TArray<const FProperty*> ArrayPropertyChain;
+	for (FPropertyValueIterator PropertyIt(FArrayProperty::StaticClass(), Struct, this); PropertyIt; ++PropertyIt)
 	{
-		if (const UStruct* const Struct = ScriptOnStruct->GetStruct())
+		if (Reference.ContainerReference == PropertyIt.Value() && (ArrayPropertyName == NAME_None || ArrayPropertyName == PropertyIt.Key()->GetName()))
 		{
-			TArray<const FProperty*> PropertyChain;
-			const FProperty* const Property =
-				FindProperty(Struct, InProperty, PropertyName, &PropertyChain);
-			if (ensure(Property && PropertyChain.Num()))
-			{
-				const uint32 Offset = GetPropertyOffset(PropertyChain);
-				if (FDataflowInput* const* const Input = Inputs.Find(Offset))
-				{
-					Inputs.Remove(Offset);
-					delete *Input;
-
-					// Invalidate graph as this input might have had connections
-					Invalidate();
-				}
-			}
+			InputParams.ArrayProperty = CastFieldChecked<FArrayProperty>(PropertyIt.Key());
+			InputParams.Offset = GetConnectionOffsetFromReference(Reference.ContainerReference);
+			PropertyIt.GetPropertyChain(ArrayPropertyChain);
+			break;
 		}
 	}
-}
 
-FDataflowOutput* FDataflowNode::RegisterOutputConnectionInternal(const void* InProperty, const void* Passthrough, const FName& PropertyName, const FName& PassthroughName)
-{
-	if (TUniquePtr<FStructOnScope> ScriptOnStruct = TUniquePtr<FStructOnScope>(NewStructOnScope()))
+	check(InputParams.ArrayProperty);
+
+	// Find the element property
+	TArray<const FProperty*> PropertyChain;
+	const void* const AddressAtIndex = InputParams.ArrayProperty->GetValueAddressAtIndex_Direct(InputParams.ArrayProperty->Inner, const_cast<void*>(Reference.ContainerReference), Reference.Index);
+	if (AddressAtIndex == Reference.Reference && (ElementPropertyName == NAME_None || ElementPropertyName == InputParams.ArrayProperty->Inner->GetName()))
 	{
-		Dataflow::FOutputParameters OutputParams;
-		if (InitConnectionParametersFromPropertyReference(*ScriptOnStruct, InProperty, PropertyName, OutputParams))
-		{
-			FDataflowOutput* OutputConnection = new FDataflowOutput(OutputParams);
-			check(OutputConnection->RealAddress() == InProperty);
+		InputParams.Property = InputParams.ArrayProperty->Inner;
+		PropertyChain = { InputParams.ArrayProperty->Inner };
+	}
+	else if (const FStructProperty* const InnerStruct = CastField<FStructProperty>(InputParams.ArrayProperty->Inner))
+	{
+		InputParams.Property = &Dataflow::Private::FindPropertyChecked(InnerStruct->Struct, AddressAtIndex, Reference.Reference, ElementPropertyName, &PropertyChain);
+		PropertyChain.Add(InnerStruct);
+	}
 
-			TArray<const FProperty*> PassthroughPropertyChain;
-			if (FindProperty(ScriptOnStruct->GetStruct(), Passthrough, PassthroughName, &PassthroughPropertyChain))
-			{
-				const uint32 PassthroughOffset = GetConnectionOffsetFromReference(Passthrough);
-				check(PassthroughOffset == GetPropertyOffset(PassthroughPropertyChain));
-				OutputConnection->SetPassthroughOffset(PassthroughOffset);
-			}
+	check(InputParams.Property);
 
-			AddOutput(OutputConnection);
-			check(FindOutput(InProperty) == OutputConnection);
-			return OutputConnection;
-		}
+	PropertyChain.Append(ArrayPropertyChain);
+	FString ExtendedType;
+	const FString CPPType = InputParams.Property->GetCPPType(&ExtendedType);
+	InputParams.Type = FName(CPPType + ExtendedType);
+	InputParams.Name = GetPropertyFullName(PropertyChain, Reference.Index);
+	InputParams.InnerOffset = Dataflow::Private::GetArrayElementOffsetFromReference(InputParams.ArrayProperty, Reference);
+
+	InputArrayProperties.Emplace(InputParams.Offset, InputParams.ArrayProperty);
+
+	FDataflowInput* const Input = new FDataflowArrayInput(Reference.Index, InputParams);
+	AddInput(Input);
+	check(FindInput(Reference) == Input);
+	return *Input;
+}
+
+void FDataflowNode::UnregisterInputConnection(const Dataflow::FConnectionKey& Key)
+{
+	if (ExpandedInputs.Remove(Key))
+	{
+		// Invalidate graph as this input might have had connections
+		Invalidate();
 	}
-	return nullptr;
+}
+
+FDataflowOutput& FDataflowNode::RegisterOutputConnectionInternal(const Dataflow::FConnectionReference& Reference, const FName& PropertyName)
+{
+	TUniquePtr<FStructOnScope> ScriptOnStruct = TUniquePtr<FStructOnScope>(NewStructOnScope());
+	check(ScriptOnStruct);
+	Dataflow::FOutputParameters OutputParams;
+	InitConnectionParametersFromPropertyReference(*ScriptOnStruct, Reference.Reference, PropertyName, OutputParams);
+	FDataflowOutput* OutputConnection = new FDataflowOutput(OutputParams);
+	check(OutputConnection->RealAddress() == Reference.Reference);
+
+	AddOutput(OutputConnection);
+	check(FindOutput(Reference) == OutputConnection);
+	check(FindOutput(OutputConnection->GetConnectionKey()) == OutputConnection);
+	return *OutputConnection;
 }
 
 uint32 FDataflowNode::GetValueHash()
@@ -686,14 +927,13 @@ bool FDataflowNode::ValidateConnections()
 			{
 				const FProperty* const Property = PropertyIt.Key();
 				check(Property);
+				TArray<const FProperty*> PropertyChain;
+				PropertyIt.GetPropertyChain(PropertyChain);
+				const FName PropName(GetPropertyFullName(PropertyChain));
 
 				if (Property->HasMetaData(FDataflowNode::DataflowInput))
 				{
-					TArray<const FProperty*> PropertyChain;
-					PropertyIt.GetPropertyChain(PropertyChain);
-					const FName PropName(GetPropertyFullName(PropertyChain));
-
-					if (!FindInput(PropName))
+					if (!FindInput(PropertyIt.Value()))
 					{
 						UE_LOG(LogChaos, Warning, TEXT("Missing dataflow RegisterInputConnection in constructor for (%s:%s)"), *GetName().ToString(), *PropName.ToString())
 							bHasValidConnections = false;
@@ -701,26 +941,16 @@ bool FDataflowNode::ValidateConnections()
 				}
 				if (Property->HasMetaData(FDataflowNode::DataflowOutput))
 				{
-					TArray<const FProperty*> PropertyChain;
-					PropertyIt.GetPropertyChain(PropertyChain);
-					const FName PropName(GetPropertyFullName(PropertyChain));
-
-					const FDataflowOutput* const OutputConnection = FindOutput(PropName);
+					const FDataflowOutput* const OutputConnection = FindOutput(PropertyIt.Value());
 					if(!OutputConnection)
 					{
 						UE_LOG(LogChaos, Warning, TEXT("Missing dataflow RegisterOutputConnection in constructor for (%s:%s)"), *GetName().ToString(),*PropName.ToString());
 						bHasValidConnections = false;
 					}
+
 					// If OutputConnection is valid, validate passthrough connections if they exist
 					else if (const FString* PassthroughName = Property->FindMetaData(FDataflowNode::DataflowPassthrough))
 					{
-						void* PassthroughConnectionAddress = OutputConnection->GetPassthroughRealAddress();
-						if (PassthroughConnectionAddress == nullptr)
-						{
-							UE_LOG(LogChaos, Warning, TEXT("Missing DataflowPassthrough registration for (%s:%s)"), *GetName().ToString(), *PropName.ToString());
-							bHasValidConnections = false;
-						}
-
 						// Assume passthrough name is relative to current property name.
 						FString FullPassthroughName;
 						if (PropertyChain.Num() <= 1)
@@ -729,33 +959,39 @@ bool FDataflowNode::ValidateConnections()
 						}
 						else
 						{
-							FullPassthroughName = FString::Format(TEXT("{0}.{1}"), { GetPropertyFullNameString(TConstArrayView<const FProperty*>(&PropertyChain[1], PropertyChain.Num() - 1)), *PassthroughName});
+							FullPassthroughName = FString::Format(TEXT("{0}.{1}"), { GetPropertyFullNameString(TConstArrayView<const FProperty*>(&PropertyChain[1], PropertyChain.Num() - 1)), *PassthroughName });
+						}
+
+						const FDataflowInput* const PassthroughConnectionInput = OutputConnection->GetPassthroughInput();
+						if (PassthroughConnectionInput == nullptr)
+						{
+							UE_LOG(LogChaos, Warning, TEXT("Missing DataflowPassthrough registration for (%s:%s)"), *GetName().ToString(), *PropName.ToString());
+							bHasValidConnections = false;
 						}
 
-						const FDataflowInput* PassthroughConnectionInput = FindInput(FName(FullPassthroughName));
-						const FDataflowInput* PassthroughConnectionInputFromArg = FindInput(PassthroughConnectionAddress);
+						const FDataflowInput* const PassthroughConnectionInputFromMetadata = FindInput(FName(FullPassthroughName));
 
-						if(PassthroughConnectionInputFromArg != PassthroughConnectionInput)
+						if(PassthroughConnectionInput != PassthroughConnectionInputFromMetadata)
 						{
-							UE_LOG(LogChaos, Warning, TEXT("Mismatch in declared and registered DataflowPassthrough connection; (%s:%s vs %s)"), *GetName().ToString(), *PropName.ToString(), *PassthroughConnectionInputFromArg->GetName().ToString());
+							UE_LOG(LogChaos, Warning, TEXT("Mismatch in declared and registered DataflowPassthrough connection; (%s:%s vs %s)"), *GetName().ToString(), *FullPassthroughName, *PassthroughConnectionInput->GetName().ToString());
 							bHasValidConnections = false;
 						}
 
-						if(!PassthroughConnectionInput)
+						if(!PassthroughConnectionInputFromMetadata)
 						{
-							UE_LOG(LogChaos, Warning, TEXT("Incorrect DataflowPassthrough Connection set for (%s:%s)"), *GetName().ToString(),*PropName.ToString());
+							UE_LOG(LogChaos, Warning, TEXT("Incorrect DataflowPassthrough Connection set for (%s:%s)"), *GetName().ToString(), *PropName.ToString());
 							bHasValidConnections = false;
 						}
 
 						else if(OutputConnection->GetType() != PassthroughConnectionInput->GetType())
 						{
-							UE_LOG(LogChaos, Warning, TEXT("DataflowPassthrough connection types mismatch for (%s:%s)"), *GetName().ToString(),*PropName.ToString());
+							UE_LOG(LogChaos, Warning, TEXT("DataflowPassthrough connection types mismatch for (%s:%s)"), *GetName().ToString(), *PropName.ToString());
 							bHasValidConnections = false;
 						}
 					}
-					else if(OutputConnection->GetPassthroughRealAddress()) 
+					else if(OutputConnection->GetPassthroughInput()) 
 					{
-						UE_LOG(LogChaos, Warning, TEXT("Missing DataflowPassthrough declaration for (%s:%s)"), *GetName().ToString(),*PropName.ToString());
+						UE_LOG(LogChaos, Warning, TEXT("Missing DataflowPassthrough declaration for (%s:%s)"), *GetName().ToString(), *PropName.ToString());
 						bHasValidConnections = false;
 					}
 				}
@@ -773,8 +1009,25 @@ FString FDataflowNode::GetToolTip()
 	return FactoryParameters.ToolTip;
 }
 
-FText FDataflowNode::GetPinDisplayName(const FName& PropertyFullName)
+FText FDataflowNode::GetPinDisplayName(const FName& PropertyFullName, const Dataflow::FPin::EDirection Direction)
 {
+	int32 ContainerIndex = INDEX_NONE;
+
+	if (Direction == Dataflow::FPin::EDirection::INPUT)
+	{
+		if (const FDataflowInput* const Input = FindInput(PropertyFullName))
+		{
+			ContainerIndex = Input->GetContainerIndex();
+		}
+	}
+	else if (Direction == Dataflow::FPin::EDirection::OUTPUT)
+	{
+		if (const FDataflowOutput* const Output = FindOutput(PropertyFullName))
+		{
+			ContainerIndex = Output->GetContainerIndex();
+		}
+	}
+
 	if (const TUniquePtr<FStructOnScope> ScriptOnStruct = TUniquePtr<FStructOnScope>(NewStructOnScope()))
 	{
 		if (const UStruct* const Struct = ScriptOnStruct->GetStruct())
@@ -782,7 +1035,7 @@ FText FDataflowNode::GetPinDisplayName(const FName& PropertyFullName)
 			TArray<const FProperty*> PropertyChain;
 			if (FindProperty(Struct, PropertyFullName, &PropertyChain))
 			{
-				return GetPropertyDisplayNameText(PropertyChain);
+				return GetPropertyDisplayNameText(PropertyChain, ContainerIndex);
 			}
 		}
 	}
@@ -790,33 +1043,36 @@ FText FDataflowNode::GetPinDisplayName(const FName& PropertyFullName)
 	return FText();
 }
 
-FString FDataflowNode::GetPinToolTip(const FName& PropertyFullName)
+FString FDataflowNode::GetPinToolTip(const FName& PropertyFullName, const Dataflow::FPin::EDirection Direction)
 {
 #if WITH_EDITORONLY_DATA
-	if (const TUniquePtr<FStructOnScope> ScriptOnStruct = TUniquePtr<FStructOnScope>(NewStructOnScope()))
+	if (Direction == Dataflow::FPin::EDirection::INPUT)
+	{
+		if (const FDataflowInput* const Input = FindInput(PropertyFullName))
+		{
+			if (const FProperty* const Property = Input->GetProperty())
+			{
+				return Dataflow::Private::GetPinToolTipFromProperty(Property);
+			}
+		}
+	}
+	else if (Direction == Dataflow::FPin::EDirection::OUTPUT)
+	{
+		if (const FDataflowOutput* const Output = FindOutput(PropertyFullName))
+		{
+			if (const FProperty* const Property = Output->GetProperty())
+			{
+				return Dataflow::Private::GetPinToolTipFromProperty(Property);
+			}
+		}
+	}
+	else if (const TUniquePtr<FStructOnScope> ScriptOnStruct = TUniquePtr<FStructOnScope>(NewStructOnScope()))
 	{
 		if (const UStruct* const Struct = ScriptOnStruct->GetStruct())
 		{
 			if (const FProperty* const Property = FindProperty(Struct, PropertyFullName))
 			{
-				if (Property->HasMetaData(TEXT("Tooltip")))
-				{
-					const FString ToolTipStr = Property->GetToolTipText(true).ToString();
-					if (ToolTipStr.Len() > 0)
-					{
-						TArray<FString> OutArr;
-						const int32 NumElems = ToolTipStr.ParseIntoArray(OutArr, TEXT(":\r\n"));
-
-						if (NumElems == 2)
-						{
-							return OutArr[1];  // Return tooltip meta text
-						}
-						else if (NumElems == 1)
-						{
-							return OutArr[0];  // Return doc comment
-						}
-					}
-				}
+				return Dataflow::Private::GetPinToolTipFromProperty(Property);
 			}
 		}
 	}
@@ -825,26 +1081,36 @@ FString FDataflowNode::GetPinToolTip(const FName& PropertyFullName)
 	return "";
 }
 
-TArray<FString> FDataflowNode::GetPinMetaData(const FName& PropertyFullName)
+TArray<FString> FDataflowNode::GetPinMetaData(const FName& PropertyFullName, const Dataflow::FPin::EDirection Direction)
 {
 #if WITH_EDITORONLY_DATA
-	if (const TUniquePtr<FStructOnScope> ScriptOnStruct = TUniquePtr<FStructOnScope>(NewStructOnScope()))
+	if (Direction == Dataflow::FPin::EDirection::INPUT)
+	{
+		if (const FDataflowInput* const Input = FindInput(PropertyFullName))
+		{
+			if (const FProperty* const Property = Input->GetProperty())
+			{
+				return Dataflow::Private::GetPinMetaDataFromProperty(Property);
+			}
+		}
+	}
+	else if (Direction == Dataflow::FPin::EDirection::OUTPUT)
+	{
+		if (const FDataflowOutput* const Output = FindOutput(PropertyFullName))
+		{
+			if (const FProperty* const Property = Output->GetProperty())
+			{
+				return Dataflow::Private::GetPinMetaDataFromProperty(Property);
+			}
+		}
+	}
+	else if (const TUniquePtr<FStructOnScope> ScriptOnStruct = TUniquePtr<FStructOnScope>(NewStructOnScope()))
 	{
 		if (const UStruct* const Struct = ScriptOnStruct->GetStruct())
 		{
 			if (const FProperty* const Property = FindProperty(Struct, PropertyFullName))
 			{
-				TArray<FString> MetaDataStrArr;
-				if (Property->HasMetaData(FDataflowNode::DataflowPassthrough))
-				{
-					MetaDataStrArr.Add("Passthrough");
-				}
-				if (Property->HasMetaData(FDataflowNode::DataflowIntrinsic))
-				{
-					MetaDataStrArr.Add("Intrinsic");
-				}
-
-				return MetaDataStrArr;
+				return Dataflow::Private::GetPinMetaDataFromProperty(Property);
 			}
 		}
 	}
@@ -864,8 +1130,7 @@ void FDataflowNode::CopyNodeProperties(const TSharedPtr<FDataflowNode> CopyFromD
 	this->SerializeInternal(ArReader);
 }
 
-
-void FDataflowNode::ForwardInput(Dataflow::FContext& Context, const void* InputReference, const void* Reference) const
+void FDataflowNode::ForwardInput(Dataflow::FContext& Context, const Dataflow::FConnectionReference& InputReference, const Dataflow::FConnectionReference& Reference) const
 {
 	if (const FDataflowOutput* Output = FindOutput(Reference))
 	{
@@ -873,7 +1138,7 @@ void FDataflowNode::ForwardInput(Dataflow::FContext& Context, const void* InputR
 		{
 			// we need to pull the value first so the upstream of the graph evaluate 
 			Input->PullValue(Context);
-			Output->ForwardInput(InputReference, Context);
+			Output->ForwardInput(Input, Context);
 		}
 		else
 		{
@@ -907,7 +1172,7 @@ bool FDataflowNode::TrySetConnectionType(FDataflowConnection* Connection, FName
 	return false;
 }
 
-bool FDataflowNode::SetInputConcreteType(void* InputReference, FName NewType)
+bool FDataflowNode::SetInputConcreteType(const Dataflow::FConnectionReference& InputReference, FName NewType)
 {
 	if (FDataflowInput* Input = FindInput(InputReference))
 	{
@@ -919,7 +1184,7 @@ bool FDataflowNode::SetInputConcreteType(void* InputReference, FName NewType)
 	return false;
 }
 
-bool FDataflowNode::SetOutputConcreteType(void* OutputReference, FName NewType)
+bool FDataflowNode::SetOutputConcreteType(const Dataflow::FConnectionReference& OutputReference, FName NewType)
 {
 	if (FDataflowOutput* Output = FindOutput(OutputReference))
 	{
@@ -930,8 +1195,3 @@ bool FDataflowNode::SetOutputConcreteType(void* OutputReference, FName NewType)
 	}
 	return false;
 }
-
-
-
-
-
diff --git a/Engine/Source/Runtime/Experimental/Dataflow/Core/Public/Dataflow/DataflowConnection.h b/Engine/Source/Runtime/Experimental/Dataflow/Core/Public/Dataflow/DataflowConnection.h
--- a/Engine/Source/Runtime/Experimental/Dataflow/Core/Public/Dataflow/DataflowConnection.h
+++ b/Engine/Source/Runtime/Experimental/Dataflow/Core/Public/Dataflow/DataflowConnection.h
@@ -15,6 +15,8 @@
 
 class FProperty;
 struct FDataflowNode;
+struct FDataflowInput;
+struct FDataflowOutput;
 
 namespace Dataflow
 {
@@ -33,6 +35,13 @@ namespace Dataflow
 		FName Type;
 		FName Name;
 		bool bHidden = false;
+
+		bool operator==(const FPin& Other) const
+		{
+			return Direction == Other.Direction && Type == Other.Type && Name == Other.Name && bHidden == Other.bHidden;
+		}
+		
+		DATAFLOWCORE_API static const FPin InvalidPin;
 	};
 
 	struct FConnectionParameters
@@ -54,6 +63,63 @@ namespace Dataflow
 		FGuid Guid;
 	};
 
+
+	// Do not hold onto FConnectionReference when Reference is dynamically allocated (e.g., when using array inputs).
+	// Use FConnectionKey instead.
+	struct FConnectionReference
+	{
+		const void* Reference;
+		int32 Index = INDEX_NONE;
+		const void* ContainerReference = nullptr;
+
+		FConnectionReference(const void* InReference, int32 InIndex = INDEX_NONE, const void* InContainerReference = nullptr)
+			: Reference(InReference)
+			, Index(InIndex)
+			, ContainerReference(InContainerReference)
+		{}
+	};
+
+	template<typename T>
+	struct TConnectionReference : public FConnectionReference
+	{
+		TConnectionReference(const T* InReference, int32 InIndex = INDEX_NONE, const void* InContainerReference = nullptr)
+			: FConnectionReference(InReference, InIndex, InContainerReference)
+		{}
+	};
+
+	class FConnectionKey
+	{
+	public:
+		FConnectionKey() = default;
+
+		bool operator==(const FConnectionKey& Other) const
+		{
+			return Offset == Other.Offset && ContainerIndex == Other.ContainerIndex && ContainerElementOffset == Other.ContainerElementOffset;
+		}
+
+		friend uint32 GetTypeHash(const FConnectionKey& Key)
+		{
+			return HashCombineFast(HashCombineFast(GetTypeHash(Key.Offset), GetTypeHash(Key.ContainerIndex)), GetTypeHash(Key.ContainerElementOffset));
+		}
+
+		static const FConnectionKey Invalid;
+
+	private:
+		friend struct ::FDataflowConnection;
+		friend struct ::FDataflowInput;
+		friend struct ::FDataflowOutput;
+		friend struct ::FDataflowNode;
+		FConnectionKey(uint32 InOffset, int32 InContainerIndex, uint32 InContainerElementOffset)
+			: Offset(InOffset)
+			, ContainerIndex(InContainerIndex)
+			, ContainerElementOffset(InContainerElementOffset)
+		{}
+
+		uint32 Offset = INDEX_NONE;
+		int32 ContainerIndex = INDEX_NONE;
+		uint32 ContainerElementOffset = INDEX_NONE;
+	};
+
 	class FGraph;
 }
 
@@ -116,16 +182,22 @@ public:
 
 	Dataflow::FPin::EDirection GetDirection() const { return Direction; }
 	uint32 GetOffset() const { return Offset; }
+	virtual int32 GetContainerIndex() const { return INDEX_NONE; }
+	virtual uint32 GetContainerElementOffset() const { return INDEX_NONE; }
+	Dataflow::FConnectionKey GetConnectionKey() const 
+	{		
+		return Dataflow::FConnectionKey(GetOffset(), GetContainerIndex(), GetContainerElementOffset());
+	}
 
 	FName GetType() const { return Type; }
 
 	FGuid GetGuid() const { return Guid; }
 	void SetGuid(FGuid InGuid) { Guid = InGuid; }
-
+	
 	FName GetName() const { return Name; }
 	void SetName(FName InName) { Name = InName; }
 
-	void* RealAddress() const { ensure(OwningNode);  return (void*)((size_t)OwningNode + (size_t)GetOffset()); };
+	virtual void* RealAddress() const { ensure(OwningNode);  return (void*)((size_t)OwningNode + (size_t)GetOffset()); };
 	Dataflow::FContextCacheKey CacheKey() const { return GetTypeHash(Guid); };
 
 	virtual bool AddConnection(FDataflowConnection* In) { return false; };
@@ -140,15 +212,23 @@ public:
 	template<class T>
 	bool IsA(const T* InVar) const
 	{
-		return (size_t)OwningNode + (size_t)GetOffset() == (size_t)InVar;
+		return (size_t)RealAddress() == (size_t)InVar;
 	}
 
 	virtual void Invalidate(const Dataflow::FTimestamp& ModifiedTimestamp = Dataflow::FTimestamp::Current()) {};
 
 	bool GetCanHidePin() const { return bCanHidePin; }
 	bool GetPinIsHidden() const { return bCanHidePin && bPinIsHidden; }
-	void SetCanHidePin(bool bInCanHidePin) { bCanHidePin = bInCanHidePin; }
-	void SetPinIsHidden(bool bInPinIsHidden) { bPinIsHidden = bInPinIsHidden; }
+	FDataflowConnection& SetCanHidePin(bool bInCanHidePin) 
+	{
+		bCanHidePin = bInCanHidePin; 
+		return *this; 
+	}
+	FDataflowConnection& SetPinIsHidden(bool bInPinIsHidden)
+	{
+		bPinIsHidden = bInPinIsHidden;
+		return *this;
+	}
 
 private:
 	void InitFromType();
diff --git a/Engine/Source/Runtime/Experimental/Dataflow/Core/Public/Dataflow/DataflowInputOutput.h b/Engine/Source/Runtime/Experimental/Dataflow/Core/Public/Dataflow/DataflowInputOutput.h
--- a/Engine/Source/Runtime/Experimental/Dataflow/Core/Public/Dataflow/DataflowInputOutput.h
+++ b/Engine/Source/Runtime/Experimental/Dataflow/Core/Public/Dataflow/DataflowInputOutput.h
@@ -25,6 +25,15 @@ namespace Dataflow
 			: FConnectionParameters(InType, InName, InOwner, InProperty, InOffset, InGuid)
 		{}
 	};
+
+	struct FArrayInputParameters : public FInputParameters
+	{
+		FArrayInputParameters()
+			: FInputParameters()
+		{}
+		const FArrayProperty* ArrayProperty = nullptr;
+		uint32 InnerOffset = INDEX_NONE;
+	};
 }
 
 USTRUCT()
@@ -80,6 +89,25 @@ public:
 	virtual void Invalidate(const Dataflow::FTimestamp& ModifiedTimestamp = Dataflow::FTimestamp::Current()) override;
 };
 
+USTRUCT()
+struct FDataflowArrayInput : public FDataflowInput
+{
+	GENERATED_USTRUCT_BODY()
+
+private:
+	int32 Index;
+	uint32 ElementOffset; // Offset to Property inside an array element
+	const FArrayProperty* ArrayProperty = nullptr;
+	// uint32 Offset; // On base class. This is the Offset to ArrayProperty from OwningNode.
+
+public:
+	explicit FDataflowArrayInput(int32 InIndex = INDEX_NONE, const Dataflow::FArrayInputParameters& Param = {});
+
+	DATAFLOWCORE_API virtual void* RealAddress() const override;
+	virtual int32 GetContainerIndex() const override { return Index; }
+	virtual uint32 GetContainerElementOffset() const override { return ElementOffset; }
+};
+
 //
 // Output
 //
@@ -92,6 +120,7 @@ namespace Dataflow
 		{}
 	};
 }
+
 USTRUCT()
 struct FDataflowOutput : public FDataflowConnection
 {
@@ -101,8 +130,11 @@ struct FDataflowOutput : public FDataflowConnection
 	
 	TArray< FDataflowInput* > Connections;
 
+	UE_DEPRECATED(5.5, "Use PassthroughKey instead")
 	uint32 PassthroughOffset = INDEX_NONE;
 
+	Dataflow::FConnectionKey PassthroughKey;
+
 protected:
 	friend struct FDataflowInput;
 	virtual void FixAndPropagateType(FName InType) override;
@@ -117,6 +149,14 @@ public:
 
 	DATAFLOWCORE_API FDataflowOutput(const Dataflow::FOutputParameters& Param = {});
 
+	PRAGMA_DISABLE_DEPRECATION_WARNINGS
+	~FDataflowOutput() = default;
+	FDataflowOutput(const FDataflowOutput&) = delete; 
+	FDataflowOutput(FDataflowOutput&&) = delete;
+	FDataflowOutput& operator=(const FDataflowOutput&) = delete;
+	FDataflowOutput& operator=(FDataflowOutput&&) = delete;
+	PRAGMA_ENABLE_DEPRECATION_WARNINGS
+
 	DATAFLOWCORE_API TArray<FDataflowInput*>& GetConnections();
 	DATAFLOWCORE_API const TArray<FDataflowInput*>& GetConnections() const;
 	bool HasAnyConnections() const { return !Connections.IsEmpty(); }
@@ -128,16 +168,22 @@ public:
 
 	DATAFLOWCORE_API virtual bool RemoveConnection(FDataflowConnection* InInput) override;
 
-	virtual FORCEINLINE void SetPassthroughOffset(const uint32 InPassthroughOffset)
+	UE_DEPRECATED(5.5, "Use SetPassthroughInput instead")
+	virtual void SetPassthroughOffset(const uint32 InPassthroughOffset)
 	{
-		PassthroughOffset = InPassthroughOffset;
+		SetPassthroughInput(Dataflow::FConnectionKey(InPassthroughOffset, INDEX_NONE, INDEX_NONE));
 	}
 
+	DATAFLOWCORE_API FDataflowOutput& SetPassthroughInput(const Dataflow::FConnectionReference& Reference);
+	DATAFLOWCORE_API FDataflowOutput& SetPassthroughInput(const Dataflow::FConnectionKey& Key);
+
+	DATAFLOWCORE_API const FDataflowInput* GetPassthroughInput() const;
+
 	virtual FORCEINLINE void* GetPassthroughRealAddress() const
 	{
-		if(PassthroughOffset != INDEX_NONE)
+		if(const FDataflowInput* const PassthroughInput = GetPassthroughInput())
 		{
-			return (void*)((size_t)OwningNode + (size_t)PassthroughOffset);
+			return PassthroughInput->RealAddress();
 		}
 		return nullptr;
 	}
@@ -182,7 +228,9 @@ public:
 
 	// there's no need for a templatized version as the parameter will not be used
 	// the method do check if the type of the input is the same as the output type though 
-	DATAFLOWCORE_API void ForwardInput(const void* InputReference, Dataflow::FContext& Context) const;
+	DATAFLOWCORE_API void ForwardInput(const Dataflow::FConnectionReference& InputReference, Dataflow::FContext& Context) const;
+	DATAFLOWCORE_API void ForwardInput(const FDataflowInput* Input, Dataflow::FContext& Context) const;
+
 
 	DATAFLOWCORE_API bool EvaluateImpl(Dataflow::FContext& Context) const;
 	
@@ -193,7 +241,15 @@ public:
 	DATAFLOWCORE_API virtual void Invalidate(const Dataflow::FTimestamp& ModifiedTimestamp = Dataflow::FTimestamp::Current()) override;
 
 private:
-	DATAFLOWCORE_API const FDataflowInput* GetPassthroughInput() const;
+};
+
+template<>
+struct TStructOpsTypeTraits<FDataflowOutput> : public  TStructOpsTypeTraitsBase2<FDataflowOutput>
+{
+	enum
+	{
+		WithCopy = false,
+	};
 };
 
 template<typename T>
diff --git a/Engine/Source/Runtime/Experimental/Dataflow/Core/Public/Dataflow/DataflowNode.h b/Engine/Source/Runtime/Experimental/Dataflow/Core/Public/Dataflow/DataflowNode.h
--- a/Engine/Source/Runtime/Experimental/Dataflow/Core/Public/Dataflow/DataflowNode.h
+++ b/Engine/Source/Runtime/Experimental/Dataflow/Core/Public/Dataflow/DataflowNode.h
@@ -24,6 +24,7 @@ namespace Dataflow {
 }
 
 
+
 /**
 * FNode
 *		Base class for node based evaluation within the Dataflow graph. 
@@ -44,6 +45,7 @@ struct FDataflowNode
 	FName Name;
 	Dataflow::FTimestamp LastModifiedTimestamp;
 
+	UE_DEPRECATED(5.5, "Inputs type has changed and has been made private")
 	TMap< int, FDataflowInput* > Inputs;
 	TMap< int, FDataflowOutput* > Outputs;
 
@@ -64,7 +66,12 @@ struct FDataflowNode
 	{
 	}
 
+	PRAGMA_DISABLE_DEPRECATION_WARNINGS
 	virtual ~FDataflowNode() { ClearInputs(); ClearOutputs(); }
+	// Warning: FDataflowNodes aren't actually safe to copy/move yet. These are here to disable deprecation warnings from the implicit operators that were getting created anyways.
+	// (Deleting the operators would require tagging all derived classes since this is a USTRUCT, so also not doing that here).
+	FDataflowNode& operator=(const FDataflowNode&) = default;
+	PRAGMA_ENABLE_DEPRECATION_WARNINGS
 
 	FGuid GetGuid() const { return Guid; }
 	FName GetName() const { return Name; }
@@ -78,9 +85,9 @@ struct FDataflowNode
 	virtual FName GetCategory() const { return ""; }
 	virtual FString GetTags() const { return ""; }
 	DATAFLOWCORE_API virtual FString GetToolTip();
-	DATAFLOWCORE_API FString GetPinToolTip(const FName& PropertyName);
-	DATAFLOWCORE_API FText GetPinDisplayName(const FName& PropertyName);
-	DATAFLOWCORE_API TArray<FString> GetPinMetaData(const FName& PropertyName);
+	DATAFLOWCORE_API FString GetPinToolTip(const FName& PropertyName, const Dataflow::FPin::EDirection Direction = Dataflow::FPin::EDirection::NONE);
+	DATAFLOWCORE_API FText GetPinDisplayName(const FName& PropertyName, const Dataflow::FPin::EDirection Direction = Dataflow::FPin::EDirection::NONE);
+	DATAFLOWCORE_API TArray<FString> GetPinMetaData(const FName& PropertyName, const Dataflow::FPin::EDirection Direction = Dataflow::FPin::EDirection::NONE);
 	virtual TArray<Dataflow::FRenderingParameter> GetRenderParameters() const { return GetRenderParametersImpl(); }
 	// Copy node property values from another node
 	UE_DEPRECATED(5.4, "FDataflowNode::CopyNodeProperties is deprecated.")
@@ -97,14 +104,43 @@ struct FDataflowNode
 
 	DATAFLOWCORE_API TArray<Dataflow::FPin> GetPins() const;
 
+	UE_DEPRECATED(5.5, "Use AddPins method instead")
+	virtual Dataflow::FPin AddPin() { return Dataflow::FPin::InvalidPin; }
 	/** Override this function to add the AddOptionPin functionality to the node's context menu. */
-	virtual Dataflow::FPin AddPin() { return { Dataflow::FPin::EDirection::NONE, NAME_None, NAME_None }; }
+	virtual TArray<Dataflow::FPin> AddPins()
+	{
+PRAGMA_DISABLE_DEPRECATION_WARNINGS
+		const Dataflow::FPin DeprecatedAddPin = AddPin();
+PRAGMA_ENABLE_DEPRECATION_WARNINGS
+		if(DeprecatedAddPin == Dataflow::FPin::InvalidPin)
+		{
+			return TArray<Dataflow::FPin>();
+		}
+		return { DeprecatedAddPin };
+	}
 	/** Override this function to add the AddOptionPin functionality to the node's context menu. */
 	virtual bool CanAddPin() const { return false; }
-	/** Override this function to add the RemoveOptionPin functionality to the node's context menu. */
-	virtual Dataflow::FPin GetPinToRemove() const { return { Dataflow::FPin::EDirection::NONE, NAME_None, NAME_None }; }
-	UE_DEPRECATED(5.4, "Use GetPinToRemove and OnPinRemoved instead.")
-	virtual Dataflow::FPin RemovePin() { return GetPinToRemove(); }
+	UE_DEPRECATED(5.5, "Use GetPinsToRemove method instead")
+	virtual Dataflow::FPin GetPinToRemove() const { return Dataflow::FPin::InvalidPin; }
+	UE_DEPRECATED(5.4, "Use GetPinsToRemove and OnPinRemoved instead.")
+	virtual Dataflow::FPin RemovePin()
+	{
+		PRAGMA_DISABLE_DEPRECATION_WARNINGS
+		return GetPinToRemove();
+		PRAGMA_ENABLE_DEPRECATION_WARNINGS
+	}
+	/** Override this function to add the RemoveOptionPin functionality to the node's context menu. OnPinRemoved will be called in this order.*/
+	virtual TArray<Dataflow::FPin> GetPinsToRemove() const
+	{
+PRAGMA_DISABLE_DEPRECATION_WARNINGS
+		const Dataflow::FPin DeprecatedRemovePin = GetPinToRemove();
+PRAGMA_ENABLE_DEPRECATION_WARNINGS
+		if (DeprecatedRemovePin == Dataflow::FPin::InvalidPin)
+		{
+			return TArray<Dataflow::FPin>();
+		}
+		return { DeprecatedRemovePin };
+	}
 	/** 
 	 * Override this to update any bookkeeping when a pin is being removed.
 	 * This will be called before the pin is unregistered as an input.
@@ -113,6 +149,7 @@ struct FDataflowNode
 	/** Override this function to add the RemoveOptionPin functionality to the node's context menu. */
 	virtual bool CanRemovePin() const { return false; }
 
+
 	DATAFLOWCORE_API bool InputSupportsType(FName Name, FName Type) const;
 	DATAFLOWCORE_API bool OutputSupportsType(FName Name, FName Type) const;
 
@@ -123,9 +160,13 @@ struct FDataflowNode
 	DATAFLOWCORE_API bool HasHiddenInputs() const;
 
 	DATAFLOWCORE_API FDataflowInput* FindInput(FName Name);
-	DATAFLOWCORE_API FDataflowInput* FindInput(void* Reference);
+	DATAFLOWCORE_API FDataflowInput* FindInput(const Dataflow::FConnectionKey& Key);
+	/**This version can find array inputs if only the Reference is supplied by searching through all inputs */
+	DATAFLOWCORE_API FDataflowInput* FindInput(const Dataflow::FConnectionReference& Reference);
 	DATAFLOWCORE_API const FDataflowInput* FindInput(FName Name) const;
-	DATAFLOWCORE_API const FDataflowInput* FindInput(const void* Reference) const;
+	DATAFLOWCORE_API const FDataflowInput* FindInput(const Dataflow::FConnectionKey& Key) const;
+	/**This version can find array inputs if only the Reference is supplied by searching through all inputs */
+	DATAFLOWCORE_API const FDataflowInput* FindInput(const Dataflow::FConnectionReference& Reference) const;
 	DATAFLOWCORE_API const FDataflowInput* FindInput(const FGuid& InGuid) const;
 
 	DATAFLOWCORE_API virtual void AddOutput(FDataflowOutput* InPtr);
@@ -137,13 +178,15 @@ struct FDataflowNode
 
 	DATAFLOWCORE_API FDataflowOutput* FindOutput(FName Name);
 	DATAFLOWCORE_API FDataflowOutput* FindOutput(uint32 GuidHash);
-	DATAFLOWCORE_API FDataflowOutput* FindOutput(void* Reference);
+	DATAFLOWCORE_API FDataflowOutput* FindOutput(const Dataflow::FConnectionKey& Key);
+	DATAFLOWCORE_API FDataflowOutput* FindOutput(const Dataflow::FConnectionReference& Reference);
 	DATAFLOWCORE_API const FDataflowOutput* FindOutput(FName Name) const;
 	DATAFLOWCORE_API const FDataflowOutput* FindOutput(uint32 GuidHash) const;
-	DATAFLOWCORE_API const FDataflowOutput* FindOutput(const void* Reference) const;
+	DATAFLOWCORE_API const FDataflowOutput* FindOutput(const Dataflow::FConnectionKey& Key) const;
+	DATAFLOWCORE_API const FDataflowOutput* FindOutput(const Dataflow::FConnectionReference& Reference) const;
 	DATAFLOWCORE_API const FDataflowOutput* FindOutput(const FGuid& InGuid) const;
 
-	/** Return a property's byte offset from the dataflow base node address using the full property name (must includes its parent struct property names). */
+	/** Return a property's byte offset from the dataflow base node address using the full property name (must includes its parent struct property names). Does not work with array properties.*/
 	uint32 GetPropertyOffset(const FName& PropertyFullName) const;
 
 	static DATAFLOWCORE_API const FName DataflowInput;
@@ -171,37 +214,75 @@ struct FDataflowNode
 	virtual const UScriptStruct* TypedScriptStruct() const { return nullptr; }
 
 	/** Register the Input and Outputs after the creation in the factory. Use PropertyName to disambiguate a struct name from its first property. */
-	template <typename T>
-	FDataflowInput* RegisterInputConnection(const T* Reference, const FName& PropertyName = NAME_None)
+	template<typename T>
+	FDataflowInput& RegisterInputConnection(const Dataflow::TConnectionReference<T>& Reference, const FName& PropertyName = NAME_None)
 	{
-		FDataflowInput* Input = RegisterInputConnectionInternal(Reference, PropertyName);
+		FDataflowInput& Input = RegisterInputConnectionInternal(Reference, PropertyName);
 		if constexpr (std::is_base_of_v<FDataflowAnyType, T>)
 		{
-			if (Input)
-			{
-				Input->SetTypePolicy(T::FPolicyType::GetInterface());
-			}
+			Input.SetTypePolicy(T::FPolicyType::GetInterface());
 		}
 		return Input;
 	}
 
 	template <typename T>
-	FDataflowOutput* RegisterOutputConnection(const T* Reference, const T* Passthrough = nullptr, const FName& PropertyName = NAME_None, const FName& PassthroughName = NAME_None)
+	FDataflowInput& RegisterInputConnection(const T* Reference, const FName& PropertyName = NAME_None)
+	{
+		return RegisterInputConnection(Dataflow::TConnectionReference<T>(Reference), PropertyName);
+	}
+
+	template <typename T>
+	FDataflowOutput& RegisterOutputConnection(const Dataflow::TConnectionReference<T>& Reference, const Dataflow::TConnectionReference<T>& Passthrough = Dataflow::TConnectionReference<T>(nullptr), const FName& PropertyName = NAME_None)
 	{
-		FDataflowOutput* Output = RegisterOutputConnectionInternal(Reference, Passthrough, PropertyName, PassthroughName);
+		FDataflowOutput& Output = RegisterOutputConnectionInternal(Reference, PropertyName);
 		if constexpr (std::is_base_of_v<FDataflowAnyType, T>)
 		{
-			if (Output)
-			{
-				Output->SetTypePolicy(T::FPolicyType::GetInterface());
-			}
+			Output.SetTypePolicy(T::FPolicyType::GetInterface());
+		}
+
+		if (Passthrough.Reference != nullptr)
+		{
+			Output.SetPassthroughInput(Passthrough);
 		}
 		return Output;
 	}
 
+	template <typename T>
+	FDataflowOutput& RegisterOutputConnection(const T* Reference, const T* Passthrough = nullptr, const FName& PropertyName = NAME_None)
+	{
+		return RegisterOutputConnection(Dataflow::TConnectionReference<T>(Reference), Dataflow::TConnectionReference<T>(Passthrough), PropertyName);
+	}
+
+	template <typename T>
+	UE_DEPRECATED(5.5, "PassthroughName is no longer needed to register output connections")
+	FDataflowOutput& RegisterOutputConnection(const T* Reference, const T* Passthrough, const FName& PropertyName, const FName& PassthroughName)
+	{
+		return RegisterOutputConnection(Dataflow::TConnectionReference<T>(Reference), Dataflow::TConnectionReference<T>(Passthrough), PropertyName);
+	}
+
+	template<typename T>
+	FDataflowInput& RegisterInputArrayConnection(const Dataflow::TConnectionReference<T>& Reference, const FName& ElementPropertyName = NAME_None,
+		const FName& ArrayPropertyName = NAME_None)
+	{
+		FDataflowInput& Input = RegisterInputArrayConnectionInternal(Reference, ElementPropertyName, ArrayPropertyName);
+		if constexpr (std::is_base_of_v<FDataflowAnyType, T>)
+		{
+			Input.SetTypePolicy(T::FPolicyType::GetInterface());
+		}
+		return Input;
+	}
 
 	/** Unregister the input connection if one exists matching this property, and then invalidate the graph. */
-	DATAFLOWCORE_API void UnregisterInputConnection(const void* Property, const FName& PropertyName = NAME_None);
+	void UnregisterInputConnection(const Dataflow::FConnectionReference& Reference)
+	{
+		UnregisterInputConnection(GetKeyFromReference(Reference));
+	}
+	UE_DEPRECATED(5.5, "PropertyName is no longer needed to unregister connections")
+	void UnregisterInputConnection(const void* Reference, const FName& PropertyName)
+	{
+		return UnregisterInputConnection(GetKeyFromReference(Reference));
+	}
+	DATAFLOWCORE_API void UnregisterInputConnection(const Dataflow::FConnectionKey& Key);
 	/** Unregister the connection if one exists matching this pin, then invalidate the graph. */
 	DATAFLOWCORE_API void UnregisterPinConnection(const Dataflow::FPin& Pin);
 
@@ -221,24 +302,42 @@ struct FDataflowNode
 	*						*Reference will be used as the default if the input is not connected. 
 	*/
 	template<class T, typename = std::enable_if_t<!std::is_base_of_v<FDataflowAnyType, T>>>
-	const T& GetValue(Dataflow::FContext& Context, const T* Reference) const
+	const T& GetValue(Dataflow::FContext& Context, const Dataflow::TConnectionReference<T>& Reference) const
 	{
 		checkSlow(FindInput(Reference));
-		return FindInput(Reference)->template GetValue<T>(Context, *Reference);
+		return FindInput(Reference)->template GetValue<T>(Context, *static_cast<const T*>(Reference.Reference));
+	}
+
+	template<typename TAnyType, typename = std::enable_if_t<std::is_base_of_v<FDataflowAnyType, TAnyType>>>
+	typename TAnyType::FStorageType GetValue(Dataflow::FContext& Context, const Dataflow::TConnectionReference<TAnyType>& Reference) const
+	{
+		checkSlow(Reference.Reference && FindInput(Reference));
+		return FindInput(Reference)->template GetValueFromAnyType<TAnyType>(Context, static_cast<const TAnyType*>(Reference.Reference)->Value);
+	}
+
+	template<class T, typename = std::enable_if_t<!std::is_base_of_v<FDataflowAnyType, T>>>
+	const T& GetValue(Dataflow::FContext& Context, const T* Reference) const
+	{
+		return GetValue(Context, Dataflow::TConnectionReference<T>(Reference));
 	}
 
 	template<typename TAnyType, typename = std::enable_if_t<std::is_base_of_v<FDataflowAnyType, TAnyType>>>
 	typename TAnyType::FStorageType GetValue(Dataflow::FContext& Context, const TAnyType* Reference) const
 	{
-		checkSlow(Reference && FindInput(Reference));
-		return FindInput(Reference)->template GetValueFromAnyType<TAnyType>(Context, Reference->Value);
+		return GetValue(Context, Dataflow::TConnectionReference<TAnyType>(Reference));
 	}
 	
 	template<class T>
-	TFuture<const T&> GetValueParallel(Dataflow::FContext& Context, const T* Reference) const
+	TFuture<const T&> GetValueParallel(Dataflow::FContext& Context, const Dataflow::TConnectionReference<T>& Reference) const
 	{
 		checkSlow(FindInput(Reference));
-		return FindInput(Reference)->template GetValueParallel<T>(Context, *Reference);
+		return FindInput(Reference)->template GetValueParallel<T>(Context, *static_cast<const T*>(Reference.Reference));
+	}
+
+	template<class T>
+	TFuture<const T&> GetValueParallel(Dataflow::FContext& Context, const T* Reference) const
+	{
+		return GetValueParallel(Context, Dataflow::TConnectionReference<T>(Reference));
 	}
 
 	/**
@@ -251,12 +350,17 @@ struct FDataflowNode
 	*   @param Reference : Pointer to a member of this node that corresponds with the output to set.
 	*   @param Default : Default value if the input is not connected.
 	*/
-	template<class T> const T& GetValue(Dataflow::FContext& Context, const T* Reference, const T& Default) const
+	template<class T> const T& GetValue(Dataflow::FContext& Context, const Dataflow::TConnectionReference<T>& Reference, const T& Default) const
 	{
 		checkSlow(FindInput(Reference));
 		return FindInput(Reference)->template GetValue<T>(Context, Default);
 	}
 
+	template<class T> const T& GetValue(Dataflow::FContext& Context, const T* Reference, const T& Default) const
+	{
+		return GetValue(Context, Dataflow::TConnectionReference<T>(Reference), Default);
+	}
+
 	/**
 	*   SetValue(...)
 	*
@@ -310,7 +414,7 @@ struct FDataflowNode
 	*   @param InputReference : Pointer to a input member of this node that needs to be forwarded
 	*   @param Reference : Pointer to a member of this node that corresponds with the output to set.
 	*/
-	DATAFLOWCORE_API void ForwardInput(Dataflow::FContext& Context, const void* InputReference, const void* Reference) const;
+	DATAFLOWCORE_API void ForwardInput(Dataflow::FContext& Context, const Dataflow::FConnectionReference& InputReference, const Dataflow::FConnectionReference& Reference) const;
 
 	/**
 	*   IsConnected(...)
@@ -319,12 +423,36 @@ struct FDataflowNode
 	*
 	*   @param Reference : Pointer to a member of this node that corresponds with the input.
 	*/
-	template<class T> bool IsConnected(const T* Reference) const
+	bool IsConnected(const Dataflow::FConnectionReference& Reference) const
 	{
 		checkSlow(FindInput(Reference));
 		return (FindInput(Reference)->GetConnection() != nullptr);
 	}
 
+	template<typename T>
+	bool IsConnected(const T* Reference) const
+	{
+		return IsConnected(Dataflow::FConnectionReference(Reference));
+	}
+
+	void PauseInvalidations()
+	{
+		if (!bPauseInvalidations)
+		{
+			bPauseInvalidations = true;
+			PausedModifiedTimestamp = Dataflow::FTimestamp::Invalid;
+		}
+	}
+
+	void ResumeInvalidations()
+	{
+		if (bPauseInvalidations)
+		{
+			bPauseInvalidations = false;
+			Invalidate(PausedModifiedTimestamp);
+		}
+	}
+
 	DATAFLOWCORE_API void Invalidate(const Dataflow::FTimestamp& ModifiedTimestamp = Dataflow::FTimestamp::Current());
 
 	virtual void OnInvalidate() {}
@@ -368,32 +496,32 @@ struct FDataflowNode
 	// returns true if the type was changed successfully
 	// only unset datatype connection will be set a new type 
 	DATAFLOWCORE_API bool TrySetConnectionType(FDataflowConnection* Connection, FName NewType);
-
 protected:
 	virtual bool OnInputTypeChanged(const FDataflowInput* Input) { return false; };
 	virtual bool OnOutputTypeChanged(const FDataflowOutput* Output) { return false; }
 
 	// returns true if the input type was changed successfully
-	DATAFLOWCORE_API bool SetInputConcreteType(void* OutputReference, FName NewType);
+	DATAFLOWCORE_API bool SetInputConcreteType(const Dataflow::FConnectionReference& InputReference, FName NewType);
 
 	// returns true if the output type was changed successfully
-	DATAFLOWCORE_API bool SetOutputConcreteType(void* InputReference, FName NewType);
-
-	DATAFLOWCORE_API FDataflowInput* RegisterInputConnectionInternal(const void* Property, const FName& PropertyName = NAME_None);
-	DATAFLOWCORE_API FDataflowOutput* RegisterOutputConnectionInternal(
-		const void* Property,
-		const void* Passthrough = nullptr,
-		const FName& PropertyName = NAME_None,
-		const FName& PassthroughName = NAME_None);
+	DATAFLOWCORE_API bool SetOutputConcreteType(const Dataflow::FConnectionReference& OutputReference, FName NewType);
 
+	DATAFLOWCORE_API FDataflowInput& RegisterInputConnectionInternal(const Dataflow::FConnectionReference& Reference, const FName& PropertyName = NAME_None);
+	DATAFLOWCORE_API FDataflowOutput& RegisterOutputConnectionInternal(const Dataflow::FConnectionReference& Reference, const FName& PropertyName = NAME_None);
+	DATAFLOWCORE_API FDataflowInput& RegisterInputArrayConnectionInternal(const Dataflow::FConnectionReference& Reference, const FName& ElementPropertyName = NAME_None,
+		const FName& ArrayPropertyName = NAME_None);
 
 private:
-	bool InitConnectionParametersFromPropertyReference(const FStructOnScope& StructOnScope, const void* PropertyRef, const FName& PropertyName, Dataflow::FConnectionParameters& OutParams);
-	static FString GetPropertyFullNameString(const TConstArrayView<const FProperty*>& PropertyChain);
-	static FName GetPropertyFullName(const TArray<const FProperty*>& PropertyChain);
-	static FText GetPropertyDisplayNameText(const TArray<const FProperty*>& PropertyChain);
+
+	void InitConnectionParametersFromPropertyReference(const FStructOnScope& StructOnScope, const void* PropertyRef, const FName& PropertyName, Dataflow::FConnectionParameters& OutParams);
+	// This will add [ContainerIndex] to any array it finds unless ContainerIndex == INDEX_NONE.
+	static FString GetPropertyFullNameString(const TConstArrayView<const FProperty*>& PropertyChain, int32 ContainerIndex = INDEX_NONE);
+	static FName GetPropertyFullName(const TArray<const FProperty*>& PropertyChain, int32 ContainerIndex = INDEX_NONE);
+	static FText GetPropertyDisplayNameText(const TArray<const FProperty*>& PropertyChain, int32 ContainerIndex = INDEX_NONE);
+	static FString StripContainerIndexFromPropertyFullName(const FString& PropertyFullName);
 	static uint32 GetPropertyOffset(const TArray<const FProperty*>& PropertyChain);
 	uint32 GetConnectionOffsetFromReference(const void* Reference) const;
+	Dataflow::FConnectionKey GetKeyFromReference(const Dataflow::FConnectionReference& Reference) const;
 
 	/**
 	* Find a property using the property address and name (not including its parent struct property names).
@@ -401,23 +529,54 @@ private:
 	* its first child property, then the parent will be returned.
 	*/
 	const FProperty* FindProperty(const UStruct* Struct, const void* Property, const FName& PropertyName, TArray<const FProperty*>* OutPropertyChain = nullptr) const;
+	const FProperty& FindPropertyChecked(const UStruct* Struct, const void* Property, const FName& PropertyName, TArray<const FProperty*>* OutPropertyChain = nullptr) const;
 
 	/** Find a property using the property full name (must includes its parent struct property names). */
 	const FProperty* FindProperty(const UStruct* Struct, const FName& PropertyFullName, TArray<const FProperty*>* OutPropertyChain = nullptr) const;
 
 	virtual TArray<Dataflow::FRenderingParameter> GetRenderParametersImpl() const { return TArray<Dataflow::FRenderingParameter>(); }
 
-
 	bool bHasValidConnections = true;
+	TMap<Dataflow::FConnectionKey, FDataflowInput*> ExpandedInputs;
+	TMap<uint32, const FArrayProperty*> InputArrayProperties; // Used to calculate ContainerElementOffsets
 
 protected:
+	bool bPauseInvalidations = false;
+	Dataflow::FTimestamp PausedModifiedTimestamp = Dataflow::FTimestamp::Invalid; // When unpausing invalidations, Invalidate will be called with this timestamp.
 	FOnNodeInvalidated OnNodeInvalidatedDelegate;
 
 };
 
-
 namespace Dataflow
 {
+
+class FDataflowNodePauseInvalidationScope
+{
+public:
+	explicit FDataflowNodePauseInvalidationScope(FDataflowNode* InNode)
+		:Node(InNode)
+	{
+		if (Node)
+		{
+			Node->PauseInvalidations();
+		}
+	}
+
+	~FDataflowNodePauseInvalidationScope()
+	{
+		if (Node)
+		{
+			Node->ResumeInvalidations();
+		}
+	}
+
+	FDataflowNodePauseInvalidationScope() = delete;
+	FDataflowNodePauseInvalidationScope(const FDataflowNodePauseInvalidationScope&) = delete;
+	FDataflowNodePauseInvalidationScope(FDataflowNodePauseInvalidationScope&&) = delete;
+private:
+	FDataflowNode* Node;
+};
+
 	//
 	// Use these macros to register dataflow nodes. 
 	//
diff --git a/Engine/Source/Runtime/Experimental/Dataflow/Engine/Private/Dataflow/DataflowEdNode.cpp b/Engine/Source/Runtime/Experimental/Dataflow/Engine/Private/Dataflow/DataflowEdNode.cpp
--- a/Engine/Source/Runtime/Experimental/Dataflow/Engine/Private/Dataflow/DataflowEdNode.cpp
+++ b/Engine/Source/Runtime/Experimental/Dataflow/Engine/Private/Dataflow/DataflowEdNode.cpp
@@ -25,6 +25,35 @@
 
 DEFINE_LOG_CATEGORY_STATIC(DATAFLOWNODE_LOG, Error, All);
 
+namespace Dataflow::Private
+{
+static Dataflow::FPin::EDirection EdPinDirectionToDataflowDirection(EEdGraphPinDirection EdDirection)
+{
+	if (EdDirection == EEdGraphPinDirection::EGPD_Input)
+	{
+		return Dataflow::FPin::EDirection::INPUT;
+	}
+	if (EdDirection == EEdGraphPinDirection::EGPD_Output)
+	{
+		return Dataflow::FPin::EDirection::OUTPUT;
+	}
+	return Dataflow::FPin::EDirection::NONE;
+}
+
+static EEdGraphPinDirection DataflowDirectionToEdPinDirection(Dataflow::FPin::EDirection Direction)
+{
+	if (Direction == Dataflow::FPin::EDirection::INPUT)
+	{
+		return EEdGraphPinDirection::EGPD_Input;
+	}
+	if (Direction == Dataflow::FPin::EDirection::OUTPUT)
+	{
+		return EEdGraphPinDirection::EGPD_Output;
+	}
+	return EEdGraphPinDirection::EGPD_MAX;
+}
+}
+
 UDataflowEdNode::UDataflowEdNode(const FObjectInitializer& ObjectInitializer)
 	: Super(ObjectInitializer)
 {
@@ -118,16 +147,8 @@ void UDataflowEdNode::AllocateDefaultPins()
 			{
 				for (const Dataflow::FPin& Pin : DataflowNode->GetPins())
 				{
-					if (Pin.Direction == Dataflow::FPin::EDirection::INPUT)
-					{
-						UEdGraphPin* const EdPin = CreatePin(EEdGraphPinDirection::EGPD_Input, Pin.Type, Pin.Name);
-						EdPin->bHidden = Pin.bHidden;
-					}					
-					if (Pin.Direction == Dataflow::FPin::EDirection::OUTPUT)
-					{
-						UEdGraphPin* const EdPin = CreatePin(EEdGraphPinDirection::EGPD_Output, Pin.Type, Pin.Name);
-						EdPin->bHidden = Pin.bHidden;
-					}
+					UEdGraphPin* const EdPin = CreatePin(Dataflow::Private::DataflowDirectionToEdPinDirection(Pin.Direction), Pin.Type, Pin.Name);
+					EdPin->bHidden = Pin.bHidden;
 				}
 			}
 		}
@@ -197,24 +218,13 @@ void UDataflowEdNode::UpdatePinsFromDataflowNode()
 
 				for (const Dataflow::FPin& Pin : DataflowNode->GetPins())
 				{
-					if (Pin.Direction == Dataflow::FPin::EDirection::INPUT)
+					const EEdGraphPinDirection EdDirection = Dataflow::Private::DataflowDirectionToEdPinDirection(Pin.Direction);
+					UEdGraphPin* EdPin = FindPin(Pin.Name, EdDirection);
+					if (!EdPin)
 					{
-						UEdGraphPin* EdPin = FindPin(Pin.Name, EEdGraphPinDirection::EGPD_Input);
-						if(!EdPin)
-						{
-							EdPin = CreatePin(EEdGraphPinDirection::EGPD_Input, Pin.Type, Pin.Name);
-						}
-						EdPin->bHidden = Pin.bHidden;
-					}
-					if (Pin.Direction == Dataflow::FPin::EDirection::OUTPUT)
-					{
-						UEdGraphPin* EdPin = FindPin(Pin.Name, EEdGraphPinDirection::EGPD_Output);
-						if(!EdPin)
-						{
-							EdPin = CreatePin(EEdGraphPinDirection::EGPD_Output, Pin.Type, Pin.Name);
-						}
-						EdPin->bHidden = Pin.bHidden;
+						EdPin = CreatePin(EdDirection, Pin.Type, Pin.Name);
 					}
+					EdPin->bHidden = Pin.bHidden;
 				}
 			}
 		}
@@ -231,19 +241,19 @@ void UDataflowEdNode::AddOptionPin()
 
 		if (const TSharedPtr<FDataflowNode> DataflowNode = DataflowGraph->FindBaseNode(DataflowNodeGuid))
 		{
-			const Dataflow::FPin Pin = DataflowNode->AddPin();
-			switch (Pin.Direction)
+			const TArray<Dataflow::FPin> AddedPins = DataflowNode->AddPins();
+			for (const Dataflow::FPin& Pin : AddedPins)
 			{
-			case Dataflow::FPin::EDirection::INPUT:
-				CreatePin(EEdGraphPinDirection::EGPD_Input, Pin.Type, Pin.Name);
-				ReconstructNode();
-				break;
-			case Dataflow::FPin::EDirection::OUTPUT:
-				CreatePin(EEdGraphPinDirection::EGPD_Output, Pin.Type, Pin.Name);
-				ReconstructNode();
-				break;
-			default:
-				break;  // Add pin isn't implemented on this node
+				switch (Pin.Direction)
+				{
+				case Dataflow::FPin::EDirection::INPUT:
+				case Dataflow::FPin::EDirection::OUTPUT:
+					CreatePin(Dataflow::Private::DataflowDirectionToEdPinDirection(Pin.Direction), Pin.Type, Pin.Name);
+					ReconstructNode();
+					break;
+				default:
+					break;  // Add pin isn't implemented on this node
+				}
 			}
 		}
 
@@ -265,29 +275,25 @@ void UDataflowEdNode::RemoveOptionPin()
 
 		if (const TSharedPtr<FDataflowNode> DataflowNode = DataflowGraph->FindBaseNode(DataflowNodeGuid))
 		{
-			const Dataflow::FPin Pin = DataflowNode->GetPinToRemove();
-			switch (Pin.Direction)
+			const TArray<Dataflow::FPin> RemovePins = DataflowNode->GetPinsToRemove();
+			Dataflow::FDataflowNodePauseInvalidationScope PauseInvalidationScope(DataflowNode.Get()); // Don't call invalidations per pin. Nodes may not evaluate correctly until all pins have been removed.
+			for (const Dataflow::FPin& Pin : RemovePins)
 			{
-			case Dataflow::FPin::EDirection::INPUT:
-				if (UEdGraphPin* const EdPin = FindPin(Pin.Name, EEdGraphPinDirection::EGPD_Input))
-				{
-					constexpr bool bNotifyNodes = true;
-					EdPin->BreakAllPinLinks(bNotifyNodes);
-					RemovePin(EdPin);
-					ReconstructNode();
-				}
-				break;
-			case Dataflow::FPin::EDirection::OUTPUT:
-				if (UEdGraphPin* const EdPin = FindPin(Pin.Name, EEdGraphPinDirection::EGPD_Output))
+				switch (Pin.Direction)
 				{
-					constexpr bool bNotifyNodes = true;
-					EdPin->BreakAllPinLinks(bNotifyNodes);
-					RemovePin(EdPin);
-					ReconstructNode();
+				case Dataflow::FPin::EDirection::INPUT:
+				case Dataflow::FPin::EDirection::OUTPUT:
+					if (UEdGraphPin* const EdPin = FindPin(Pin.Name, Dataflow::Private::DataflowDirectionToEdPinDirection(Pin.Direction)))
+					{
+						constexpr bool bNotifyNodes = true;
+						EdPin->BreakAllPinLinks(bNotifyNodes);
+						RemovePin(EdPin);
+						ReconstructNode();
+					}
+					break;
+				default:
+					break;  // Add pin isn't implemented on this node
 				}
-				break;
-			default:
-				break;  // Add pin isn't implemented on this node
 			}
 		}
 
@@ -481,7 +487,7 @@ FText UDataflowEdNode::GetPinDisplayName(const UEdGraphPin* Pin) const
 		{
 			if (TSharedPtr<FDataflowNode> DataflowNode = DataflowGraph->FindBaseNode(DataflowNodeGuid))
 			{
-				FText DisplayName = DataflowNode->GetPinDisplayName(Pin->PinName);
+				const FText DisplayName = DataflowNode->GetPinDisplayName(Pin->PinName, Dataflow::Private::EdPinDirectionToDataflowDirection(Pin->Direction));
 				if (!DisplayName.IsEmpty())
 				{
 					return DisplayName;
@@ -502,7 +508,7 @@ void UDataflowEdNode::GetPinHoverText(const UEdGraphPin& Pin, FString& HoverText
 			{		
 				FString MetaDataStr;
 				
-				TArray<FString> PinMetaData = DataflowNode->GetPinMetaData(Pin.PinName);
+				TArray<FString> PinMetaData = DataflowNode->GetPinMetaData(Pin.PinName, Dataflow::Private::EdPinDirectionToDataflowDirection(Pin.Direction));
 			
 				if (Pin.Direction == EGPD_Input && PinMetaData.Contains(FDataflowNode::DataflowIntrinsic.ToString()))
 				{
@@ -519,7 +525,8 @@ void UDataflowEdNode::GetPinHoverText(const UEdGraphPin& Pin, FString& HoverText
 					NameStr.Appendf(TEXT(" %s"), *MetaDataStr);
 				}
 
-				HoverTextOut.Appendf(TEXT("%s\n%s\n\n%s"), *NameStr, *Pin.PinType.PinCategory.ToString(), *DataflowNode->GetPinToolTip(Pin.PinName));
+				HoverTextOut.Appendf(TEXT("%s\n%s\n\n%s"), *NameStr, *Pin.PinType.PinCategory.ToString(),
+					*DataflowNode->GetPinToolTip(Pin.PinName, Dataflow::Private::EdPinDirectionToDataflowDirection(Pin.Direction)));
 			}
 		}
 	}
