Initialise centre of mass on FConvex construction in paths where we could exit early to avoid non-deterministic cooks for empty or small convex hulls.

diff --git a/Engine/Source/Runtime/Experimental/Chaos/Public/Chaos/Convex.h b/Engine/Source/Runtime/Experimental/Chaos/Public/Chaos/Convex.h
--- a/Engine/Source/Runtime/Experimental/Chaos/Public/Chaos/Convex.h
+++ b/Engine/Source/Runtime/Experimental/Chaos/Public/Chaos/Convex.h
@@ -118,6 +118,7 @@ namespace Chaos
 
 		FConvex(const TArray<FVec3Type>& InVertices, const FReal InMargin, FConvexBuilder::EBuildMethod BuildMethod = FConvexBuilder::EBuildMethod::Default)
 		    : FImplicitObject(EImplicitObject::IsConvex | EImplicitObject::HasBoundingBox, ImplicitObjectType::Convex)
+			, CenterOfMass(FVec3Type(0.f))
 		{
 			const int32 NumVertices = InVertices.Num();
 			if (NumVertices == 0)
