New chunked array keep memory over physics steps.
No allocation are required anymore at every step.
This is saving around 27% on Gather Solver

diff --git a/Engine/Source/Runtime/Experimental/Chaos/Public/Chaos/Evolution/SolverBodyContainer.h b/Engine/Source/Runtime/Experimental/Chaos/Public/Chaos/Evolution/SolverBodyContainer.h
--- a/Engine/Source/Runtime/Experimental/Chaos/Public/Chaos/Evolution/SolverBodyContainer.h
+++ b/Engine/Source/Runtime/Experimental/Chaos/Public/Chaos/Evolution/SolverBodyContainer.h
@@ -6,6 +6,7 @@
 #include "Chaos/ObjectPool.h"
 #include "Chaos/ParticleHandle.h"
 #include "Chaos/Particle/ParticleUtilities.h"
+#include "Chaos/Framework/ChunkedArray.h"
 
 #include "Containers/Array.h"
 #include "Containers/Map.h"
@@ -40,7 +41,7 @@ namespace Chaos
 		static const int32 BodyArrayChunkSize = NumBodiesPerChunk * sizeof(FSolverBody);
 
 	public:
-		using FSolverBodyArray = TChunkedArray<FSolverBody, BodyArrayChunkSize>;
+		using FSolverBodyArray = Private::TChaosChunkedArray<FSolverBody, BodyArrayChunkSize>;
 
 		FSolverBodyContainer()
 			: bLocked(true)
@@ -51,8 +52,8 @@ namespace Chaos
 		// @param MaxBodies The number of bodies that will get added to the container. The container asserts if we attempt to add more than this.
 		inline void Reset(int MaxBodies)
 		{
-			SolverBodies.Empty(MaxBodies);
-			Particles.Empty(MaxBodies);
+			SolverBodies.Reset();
+			Particles.Reset(MaxBodies);
 			ParticleToIndexMap.Reset();
 			bLocked = false;
 		}
diff --git a/Engine/Source/Runtime/Experimental/Chaos/Public/Chaos/Framework/ChunkedArray.h b/Engine/Source/Runtime/Experimental/Chaos/Public/Chaos/Framework/ChunkedArray.h
new file mode 100644
--- /dev/null
+++ b/Engine/Source/Runtime/Experimental/Chaos/Public/Chaos/Framework/ChunkedArray.h
@@ -0,0 +1,23 @@
+// Copyright Epic Games, Inc. All Rights Reserved.
+#pragma once
+
+#include "Containers/ChunkedArray.h"
+
+namespace Chaos
+{
+	namespace Private
+	{
+		template<typename InElementType, uint32 TargetBytesPerChunk = 16384, typename AllocatorType = FDefaultAllocator >
+		class TChaosChunkedArray : public TChunkedArray<InElementType, TargetBytesPerChunk, AllocatorType>
+		{
+		public:
+			// This function will set the number of element to 0, but won't allocate more memory neither won't free memory,
+			// and won't call any destructor.
+			void Reset()
+			{
+				this->NumElements = 0;
+			}
+		};
+
+	}
+}
\ No newline at end of file
