ChaosDebugDraw - fix race condition when creating a thread context for parallel-for

The race would happen if the owning thread is helping with tasks, which would be setting/unsetting the Frame pointer on that thread's context. We we reading from the parent thread context to get the Frame to use for the task thread on the task thread, which is not safe if the parent thread is running tasks. The solution is to copy the Frame pointer on the parent thread and pass it to the task, rather than passing the parent thread's context.

diff --git a/Engine/Source/Runtime/Experimental/Chaos/Private/Chaos/Framework/Parallel.cpp b/Engine/Source/Runtime/Experimental/Chaos/Private/Chaos/Framework/Parallel.cpp
--- a/Engine/Source/Runtime/Experimental/Chaos/Private/Chaos/Framework/Parallel.cpp
+++ b/Engine/Source/Runtime/Experimental/Chaos/Private/Chaos/Framework/Parallel.cpp
@@ -67,7 +67,7 @@ void Chaos::PhysicsParallelFor(int32 InNum, TFunctionRef<void(int32)> InCallable
 #endif
 
 #if CHAOS_DEBUG_DRAW
-	const ChaosDD::Private::FChaosDDContext& ParentDDContext = ChaosDD::Private::FChaosDDContext::Get();
+	const ChaosDD::Private::FChaosDDTaskParentContext ParentDDContext;
 #else
 	int32 ParentDDContext = 0;
 #endif
diff --git a/Engine/Source/Runtime/Experimental/Chaos/Private/ChaosDebugDraw/ChaosDDContext.cpp b/Engine/Source/Runtime/Experimental/Chaos/Private/ChaosDebugDraw/ChaosDDContext.cpp
--- a/Engine/Source/Runtime/Experimental/Chaos/Private/ChaosDebugDraw/ChaosDDContext.cpp
+++ b/Engine/Source/Runtime/Experimental/Chaos/Private/ChaosDebugDraw/ChaosDDContext.cpp
@@ -29,7 +29,7 @@ namespace ChaosDD::Private
 	void FChaosDDTimelineContext::BeginFrame(const FChaosDDTimelinePtr& InTimeline, double InTime, double InDt)
 	{
 		FChaosDDContext& Context = FChaosDDContext::Get();
-		ParentFrame = Context.Frame;
+		PreviousFrame = Context.Frame;
 
 		if (InTimeline.IsValid())
 		{
@@ -52,8 +52,8 @@ namespace ChaosDD::Private
 		}
 
 		FChaosDDContext& Context = FChaosDDContext::Get();
-		Context.Frame = ParentFrame;
-		ParentFrame.Reset();
+		Context.Frame = PreviousFrame;
+		PreviousFrame.Reset();
 	}
 
 	FChaosDDScopeTimelineContext::FChaosDDScopeTimelineContext(const FChaosDDTimelinePtr& InTimeline, double InTime, double InDt)
@@ -72,29 +72,31 @@ namespace ChaosDD::Private
 	//
 	//
 
-	void FChaosDDTaskContext::BeginThread(const FChaosDDContext& InParentDDContext)
+	FChaosDDTaskParentContext::FChaosDDTaskParentContext()
+		: Frame(FChaosDDContext::Get().Frame)
 	{
-		FChaosDDContext& Context = FChaosDDContext::Get();
-		ParentFrame = Context.Frame;
+	}
 
-		if (InParentDDContext.Frame.IsValid())
-		{
-			Context.Frame = InParentDDContext.Frame;
-		}
-		else
-		{
-			Context.Frame.Reset();
-		}
+	void FChaosDDTaskContext::BeginThread(const FChaosDDTaskParentContext& InParentDDContext)
+	{
+		// NOTE: (UE-216178) We used to pass a reference to the parent FChaosDDContext directly to the
+		// child thread and pulled the FramePointer from it in BeginThread. That is not safe because 
+		// the parent thread may also be helping with tasks and so the Frame on that context will be 
+		// getting set/unset. Instead we copy the Frame pointer on the parent thread and pass it in.
+
+		FChaosDDContext& Context = FChaosDDContext::Get();
+		PreviousFrame = Context.Frame;
+		Context.Frame = InParentDDContext.Frame;
 	}
 
 	void FChaosDDTaskContext::EndThread()
 	{
 		FChaosDDContext& Context = FChaosDDContext::Get();
-		Context.Frame = ParentFrame;
-		ParentFrame.Reset();
+		Context.Frame = PreviousFrame;
+		PreviousFrame.Reset();
 	}
 
-	FChaosDDScopeTaskContext::FChaosDDScopeTaskContext(const FChaosDDContext& InParentDDContext)
+	FChaosDDScopeTaskContext::FChaosDDScopeTaskContext(const FChaosDDTaskParentContext& InParentDDContext)
 	{
 		Context.BeginThread(InParentDDContext);
 	}
diff --git a/Engine/Source/Runtime/Experimental/Chaos/Public/ChaosDebugDraw/ChaosDDContext.h b/Engine/Source/Runtime/Experimental/Chaos/Public/ChaosDebugDraw/ChaosDDContext.h
--- a/Engine/Source/Runtime/Experimental/Chaos/Public/ChaosDebugDraw/ChaosDDContext.h
+++ b/Engine/Source/Runtime/Experimental/Chaos/Public/ChaosDebugDraw/ChaosDDContext.h
@@ -29,8 +29,9 @@ namespace ChaosDD::Private
 
 	private:
 		friend class FChaosDDScene;
-		friend class FChaosDDTimelineContext;
 		friend class FChaosDDTaskContext;
+		friend class FChaosDDTaskParentContext;
+		friend class FChaosDDTimelineContext;
 
 		const FChaosDDFramePtr& GetFrame() const
 		{
@@ -64,6 +65,10 @@ namespace ChaosDD::Private
 
 	//
 	// Initializes the FChaosDDContext for a thread that owns a timeline
+	// 
+	// This starts a new frame (debug draw buffer) and sets up the FChaosDDContext for this thread.
+	// The active context should be accessed via FChaosDDContext::GetWriter(). FChaosDDTimelineContext
+	// is not directly used other than to instantiate.
 	//
 	class CHAOS_API FChaosDDTimelineContext
 	{
@@ -73,7 +78,7 @@ namespace ChaosDD::Private
 
 	private:
 		FChaosDDTimelinePtr Timeline;
-		FChaosDDFramePtr ParentFrame;
+		FChaosDDFramePtr PreviousFrame;
 	};
 
 	//
@@ -90,19 +95,45 @@ namespace ChaosDD::Private
 	};
 
 	//
-	// Initializes the FChaosDDContext for a task.
+	// Used to propagate a debug draw context to a child thread.
+	// To use: 
+	//		- put a FChaosDDTaskParentContext on the stack on the parent thread
+	//		- pass the FChaosDDTaskParentContext to the child thread
+	//		- put FChaosDDScopeTaskContext(ParentContext) on the child thread
+	// (Search for FChaosDDScopeTaskContext for examples.)
+	//
+	class CHAOS_API FChaosDDTaskParentContext
+	{
+	public:
+		FChaosDDTaskParentContext();
+	private:
+		friend class FChaosDDTaskContext;
+		FChaosDDFramePtr Frame;
+	};
+
+	//
+	// Initializes the FChaosDDContext for a task thread.
 	// Assumes that the task is kicked off from a thread that has an active debug draw context,
 	// which should be passed into this context. Any debug draws from the task will go to the 
 	// same frame as the parent context.
+	// 
+	// The active context is accessed via FChaosDDContext::GetWriter() (and not this object).
+	// 
+	// NOTE: This is only intended to be used for tasks which will be awaited before the end
+	// of the frame (truly asynchronous tasks would need their own timeline, or just set up
+	// a context that writes to the global frame)
+	// 
+	// @todo(chaos): ChaosDDFrame should track how many contexts it is referenced by and
+	// assert that it is not active when we end the frame.
 	//
 	class CHAOS_API FChaosDDTaskContext
 	{
 	public:
-		void BeginThread(const FChaosDDContext& InParentDDContext);
+		void BeginThread(const FChaosDDTaskParentContext& InParentDDContext);
 		void EndThread();
 
 	private:
-		FChaosDDFramePtr ParentFrame;
+		FChaosDDFramePtr PreviousFrame;
 	};
 
 	//
@@ -111,7 +142,7 @@ namespace ChaosDD::Private
 	class CHAOS_API FChaosDDScopeTaskContext
 	{
 	public:
-		FChaosDDScopeTaskContext(const FChaosDDContext& InParentDDContext);
+		FChaosDDScopeTaskContext(const FChaosDDTaskParentContext& InParentDDContext);
 		~FChaosDDScopeTaskContext();
 
 	private:
diff --git a/Engine/Source/Runtime/Experimental/Chaos/Public/ChaosDebugDraw/ChaosDDTypes.h b/Engine/Source/Runtime/Experimental/Chaos/Public/ChaosDebugDraw/ChaosDDTypes.h
--- a/Engine/Source/Runtime/Experimental/Chaos/Public/ChaosDebugDraw/ChaosDDTypes.h
+++ b/Engine/Source/Runtime/Experimental/Chaos/Public/ChaosDebugDraw/ChaosDDTypes.h
@@ -22,9 +22,10 @@ namespace ChaosDD::Private
 	class FChaosDDFrameWriter;
 	class FChaosDDGlobalFrame;
 	class FChaosDDScene;
+	class FChaosDDTaskContext;
+	class FChaosDDTaskParentContext;
 	class FChaosDDTimeline;
 	class FChaosDDTimelineContext;
-	class FChaosDDTaskContext;
 
 	class IChaosDDRenderer;
 
