Rudeimentary base-16 support.

#rnx

diff --git a/Engine/Source/Runtime/Experimental/IoStore/HttpClient/Private/Misc.inl b/Engine/Source/Runtime/Experimental/IoStore/HttpClient/Private/Misc.inl
--- a/Engine/Source/Runtime/Experimental/IoStore/HttpClient/Private/Misc.inl
+++ b/Engine/Source/Runtime/Experimental/IoStore/HttpClient/Private/Misc.inl
@@ -252,8 +252,11 @@ static int32 FindMessageTerminal(const char* Data, uint32 Length)
 }
 
 ////////////////////////////////////////////////////////////////////////////////
+template <uint32 Base=10>
 static int64 CrudeToInt(FAnsiStringView View)
 {
+	static_assert(Base == 10 || Base == 16);
+
 	// FCStringAnsi::* is not used to mitigate any locale hiccups. By
 	// initialising 'Value' with MSB set we can detect cases where View did not
 	// start with digits. This works as we won't be using this on huge numbers.
@@ -263,9 +266,19 @@ static int64 CrudeToInt(FAnsiStringView View)
 		uint32 Digit = c - '0';
 		if (Digit > 9u)
 		{
-			break;
+			if (Base != 16)
+			{
+				break;
+			}
+
+			Digit = (c | 0x20) - 'a';
+			if (Digit > uint32('f' - 'a'))
+			{
+				break;
+			}
+			Digit += 10;
 		}
-		Value *= 10;
+		Value *= Base;
 		Value += Digit;
 	}
 	return Value;
diff --git a/Engine/Source/Runtime/Experimental/IoStore/HttpClient/Private/Test.inl b/Engine/Source/Runtime/Experimental/IoStore/HttpClient/Private/Test.inl
--- a/Engine/Source/Runtime/Experimental/IoStore/HttpClient/Private/Test.inl
+++ b/Engine/Source/Runtime/Experimental/IoStore/HttpClient/Private/Test.inl
@@ -75,6 +75,14 @@ static void MiscTest()
 	check(CrudeToInt("9") == 9);
 	check(CrudeToInt("493") == 493);
 
+	check(CrudeToInt<16>("56") == 0x56);
+	check(CrudeToInt<16>("1") == 0x01);
+	check(CrudeToInt<16>("9") == 0x09);
+	check(CrudeToInt<16>("a") == 0x0a); check(CrudeToInt<16>("A") == 0x0a);
+	check(CrudeToInt<16>("f") == 0x0f); check(CrudeToInt<16>("F") == 0x0f);
+	check(CrudeToInt<16>("g") < 0);
+	check(CrudeToInt<16>("49e") == 0x49e);
+
 	FUrlOffsets UrlOut;
 
 	check(ParseUrl("", UrlOut) == -1);
