Fixed null context pointer when loading certs.

#rnx

diff --git a/Engine/Source/Runtime/Experimental/IoStore/HttpClient/Private/Peer.inl b/Engine/Source/Runtime/Experimental/IoStore/HttpClient/Private/Peer.inl
--- a/Engine/Source/Runtime/Experimental/IoStore/HttpClient/Private/Peer.inl
+++ b/Engine/Source/Runtime/Experimental/IoStore/HttpClient/Private/Peer.inl
@@ -43,7 +43,20 @@ FCertRoots::~FCertRoots()
 ////////////////////////////////////////////////////////////////////////////////
 FCertRoots::FCertRoots(FMemoryView PemData)
 {
-	SSL_CTX* Context = SSL_CTX_new(TLS_client_method());
+	if (static bool InitOnce = false; !InitOnce)
+	{
+		// While OpenSSL will lazily initialise itself, the defaults used will fail
+		// initialisation on some platforms. So we have a go here. We do not register
+		// anything for clean-up as we do not know if anyone else has done so.
+		uint64 InitOpts = OPENSSL_INIT_NO_ATEXIT;
+		OPENSSL_init_ssl(InitOpts, nullptr);
+		InitOnce = true;
+	}
+
+	auto* Method = TLS_client_method();
+	SSL_CTX* Context = SSL_CTX_new(Method);
+	checkf(Context != nullptr, TEXT("ERR_get_error() == %d"), ERR_get_error());
+
 	SSL_CTX_set_options(Context, SSL_OP_NO_SSLv2|SSL_OP_NO_SSLv3);
 
 	const void* Data = PemData.GetData();
diff --git a/Engine/Source/Runtime/Experimental/IoStore/OnDemand/Private/IoStoreOnDemand.cpp b/Engine/Source/Runtime/Experimental/IoStore/OnDemand/Private/IoStoreOnDemand.cpp
--- a/Engine/Source/Runtime/Experimental/IoStore/OnDemand/Private/IoStoreOnDemand.cpp
+++ b/Engine/Source/Runtime/Experimental/IoStore/OnDemand/Private/IoStoreOnDemand.cpp
@@ -295,16 +295,33 @@ static void LoadCaCerts()
 {
 	using namespace UE::IoStore::HTTP;
 
+	// The following config option is used when staging to copy root certs PEM
+	const TCHAR* CertSection = TEXT("/Script/Engine.NetworkSettings");
+	const TCHAR* CertKey = TEXT("n.VerifyPeer");
+	bool bExpectCerts = false;
+	if (GConfig != nullptr)
+	{
+		GConfig->GetBool(CertSection, CertKey, bExpectCerts, GEngineIni);
+	}
+
+	// Open the certs file
 	IFileManager& Ifm = IFileManager::Get();
 	FString PemPath = FPaths::EngineContentDir() / TEXT("Certificates/ThirdParty/cacert.pem");
 	FArchive* Reader = Ifm.CreateFileReader(*PemPath);
-	check(Reader != nullptr)
+	if (!bExpectCerts && Reader == nullptr)
+	{
+		UE_LOG(LogIas, Warning, TEXT("Unable to load '%s'. Maybe it wasn't staged? Ensure '[%s]/%s=true' when staging"), *PemPath, CertSection, CertKey);
+		return;
+	}
+	checkf(Reader != nullptr, TEXT("%s/%s==true but '%s' could not be loaded"), CertSection, CertKey, *PemPath);
 
+	// Buffer certificate data
 	uint32 Size = uint32(Reader->TotalSize());
 	FIoBuffer PemData(Size);
 	FMutableMemoryView PemView = PemData.GetMutableView();
 	Reader->Serialize(PemView.GetData(), Size);
 
+	// Load the certs
 	FCertRoots CaRoots(PemData.GetView());
 
 	uint32 NumCerts = CaRoots.Num();
