Geometry collection : Memory optimization : Remove simplicials from dynamic collection avoiding extra expensive copies and extra memory
this saves a total 16 bytes per bone ( Physics thread + game thread dyanmic collection )  per geometry collection component

diff --git a/Engine/Source/Programs/HeadlessChaos/Private/GeometryCollection/GeometryCollectionTestCollisionResolution.cpp b/Engine/Source/Programs/HeadlessChaos/Private/GeometryCollection/GeometryCollectionTestCollisionResolution.cpp
--- a/Engine/Source/Programs/HeadlessChaos/Private/GeometryCollection/GeometryCollectionTestCollisionResolution.cpp
+++ b/Engine/Source/Programs/HeadlessChaos/Private/GeometryCollection/GeometryCollectionTestCollisionResolution.cpp
@@ -50,10 +50,6 @@ namespace GeometryCollectionTest
 		{
 			// validate that Simplicials are null when CollisionType==Chaos_Volumetric
 			EXPECT_EQ(Collection->DynamicCollection->GetNumTransforms(), 4);
-			EXPECT_EQ(Collection->DynamicCollection->Simplicials[0], nullptr);
-			EXPECT_EQ(Collection->DynamicCollection->Simplicials[1], nullptr);
-			EXPECT_EQ(Collection->DynamicCollection->Simplicials[2], nullptr);
-			EXPECT_EQ(Collection->DynamicCollection->Simplicials[3], nullptr);
 			EXPECT_EQ(UnitTest.Solver->GetParticles().GetGeometryCollectionParticles().CollisionParticles(0), nullptr);
 
 			const FReal MaxRestingSeparation = -UnitTest.Solver->GetEvolution()->GetGravityForces().GetAcceleration(0).Z * UnitTest.Dt * UnitTest.Dt;	// PBD resting separation will be up to this
@@ -95,12 +91,16 @@ namespace GeometryCollectionTest
 			UnitTest.Advance();
 		}
 		{
+			using FSimplicialPtr = TUniquePtr<FCollisionStructureManager::FSimplicial>;
+			TManagedArray<FSimplicialPtr>* Simplicials = Collection->RestCollection->FindAttribute<FSimplicialPtr>(FGeometryDynamicCollection::SimplicialsAttribute, FGeometryCollection::TransformAttribute);
+			check(Simplicials);
+
 			// validate that Simplicials are null when CollisionType==Chaos_Volumetric
 			EXPECT_EQ(Collection->DynamicCollection->GetNumTransforms(), 4);
-			EXPECT_TRUE(Collection->DynamicCollection->Simplicials[0]!=nullptr);
+			EXPECT_TRUE((*Simplicials)[0]!=nullptr);
 			EXPECT_TRUE(UnitTest.Solver->GetParticles().GetGeometryCollectionParticles().CollisionParticles(0)!=nullptr);
-			EXPECT_TRUE(Collection->DynamicCollection->Simplicials[0]->Size() == UnitTest.Solver->GetParticles().GetGeometryCollectionParticles().CollisionParticles(0)->Size());
-			EXPECT_TRUE(Collection->DynamicCollection->Simplicials[0]->Size() != 0);
+			EXPECT_TRUE((*Simplicials)[0]->Size() == UnitTest.Solver->GetParticles().GetGeometryCollectionParticles().CollisionParticles(0)->Size());
+			EXPECT_TRUE((*Simplicials)[0]->Size() != 0);
 
 			EXPECT_LT(FMath::Abs(Collection->RestCollection->Transform[0].GetTranslation().Z - (Radius + 10.f)), KINDA_SMALL_NUMBER);
 			EXPECT_NEAR(Collection->DynamicCollection->GetTransform(0).GetTranslation().Z, Radius, KINDA_SMALL_NUMBER); 
@@ -327,17 +327,21 @@ namespace GeometryCollectionTest
 			UnitTest.Advance();
 		}
 		{
+			using FSimplicialPtr = TUniquePtr<FCollisionStructureManager::FSimplicial>;
+			TManagedArray<FSimplicialPtr>* Simplicials = Collection->RestCollection->FindAttribute<FSimplicialPtr>(FGeometryDynamicCollection::SimplicialsAttribute, FGeometryCollection::TransformAttribute);
+			check(Simplicials);
+
 			// validate simplicials and implicits are configured correctly
 			EXPECT_EQ(Collection->DynamicCollection->GetNumTransforms(), 4);
-			EXPECT_TRUE(Collection->DynamicCollection->Simplicials[0] != nullptr);
+			EXPECT_TRUE((*Simplicials)[0] != nullptr);
 			EXPECT_TRUE(UnitTest.Solver->GetParticles().GetGeometryCollectionParticles().CollisionParticles(0) != nullptr);
-			EXPECT_TRUE(Collection->DynamicCollection->Simplicials[0]->Size() == UnitTest.Solver->GetParticles().GetGeometryCollectionParticles().CollisionParticles(0)->Size());
-			EXPECT_TRUE(Collection->DynamicCollection->Simplicials[0]->Size() != 0);
+			EXPECT_TRUE((*Simplicials)[0]->Size() == UnitTest.Solver->GetParticles().GetGeometryCollectionParticles().CollisionParticles(0)->Size());
+			EXPECT_TRUE((*Simplicials)[0]->Size() != 0);
 			// The following test has been disabled because from now we remove the Implicits after initialization to free up some memory
 			// EXPECT_TRUE(Collection->DynamicCollection->GetAttribute<Chaos::FImplicitObjectPtr>(FGeometryDynamicCollection::ImplicitsAttribute, FTransformCollection::TransformGroup)[0]->GetType() == (int32)Chaos::ImplicitObjectType::LevelSet);
 
 			EXPECT_EQ(CollectionStaticSphere->DynamicCollection->GetNumTransforms(), 4);
-			EXPECT_TRUE(CollectionStaticSphere->DynamicCollection->Simplicials[0] == nullptr);
+			EXPECT_TRUE(UnitTest.Solver->GetParticles().GetGeometryCollectionParticles().CollisionParticles(1) == nullptr);
 			// The following test has been disabled because from now we remove the Implicits after initialization to free up some memory
 			// EXPECT_TRUE(CollectionStaticSphere->DynamicCollection->GetAttribute<Chaos::FImplicitObjectPtr>(FGeometryDynamicCollection::ImplicitsAttribute, FTransformCollection::TransformGroup)[0]->GetType() == (int32)Chaos::ImplicitObjectType::Sphere);
 
diff --git a/Engine/Source/Programs/HeadlessChaos/Private/GeometryCollection/GeometryCollectionTestFramework.cpp b/Engine/Source/Programs/HeadlessChaos/Private/GeometryCollection/GeometryCollectionTestFramework.cpp
--- a/Engine/Source/Programs/HeadlessChaos/Private/GeometryCollection/GeometryCollectionTestFramework.cpp
+++ b/Engine/Source/Programs/HeadlessChaos/Private/GeometryCollection/GeometryCollectionTestFramework.cpp
@@ -181,6 +181,12 @@ namespace GeometryCollectionTest
 			Chaos::FErrorReporter ErrorReporter;
 			BuildSimulationData(ErrorReporter, *RestCollection.Get(), SimulationParams.Shared);
 
+			SimulationParams.bUseSimplicialsWhenAvailable =
+				SimulationParams.RestCollectionShared
+				&& SimulationParams.RestCollectionShared->HasAttribute(FGeometryDynamicCollection::SimplicialsAttribute, FTransformCollection::TransformGroup)
+				&& SimulationParams.Shared.SizeSpecificData[0].CollisionShapesData.Num()
+				&& (SimulationParams.Shared.SizeSpecificData[0].CollisionShapesData[0].CollisionType == ECollisionTypeEnum::Chaos_Surface_Volumetric);
+
 			FGeometryCollectionPhysicsProxy::InitializeDynamicCollection(*DynamicCollection, *RestCollection, SimulationParams);
 		}
 
diff --git a/Engine/Source/Runtime/Experimental/Chaos/Private/GeometryCollectionProxyData.cpp b/Engine/Source/Runtime/Experimental/Chaos/Private/GeometryCollectionProxyData.cpp
--- a/Engine/Source/Runtime/Experimental/Chaos/Private/GeometryCollectionProxyData.cpp
+++ b/Engine/Source/Runtime/Experimental/Chaos/Private/GeometryCollectionProxyData.cpp
@@ -153,7 +153,6 @@ PRAGMA_ENABLE_DEPRECATION_WARNINGS
 	// Transform Group
 	AddExternalAttribute<bool>(FGeometryDynamicCollection::ActiveAttribute, FTransformCollection::TransformGroup, Active);
 	AddExternalAttribute<uint8>(FGeometryDynamicCollection::DynamicStateAttribute, FTransformCollection::TransformGroup, DynamicState);
-	AddExternalAttribute(SimplicialsAttribute, FTransformCollection::TransformGroup, Simplicials);
 	AddExternalAttribute(SimulatableParticlesAttribute, FGeometryCollection::TransformGroup, SimulatableParticles);
 	AddExternalAttribute(InternalClusterParentTypeAttribute, FGeometryCollection::TransformGroup, InternalClusterParentType);
 }
@@ -167,7 +166,6 @@ FGeometryDynamicCollection::FGeometryDynamicCollection(TSharedPtr<const FGeometr
 	// Transform Group
 	AddExternalAttribute<bool>(FGeometryDynamicCollection::ActiveAttribute, FTransformCollection::TransformGroup, Active);
 	AddExternalAttribute<uint8>(FGeometryDynamicCollection::DynamicStateAttribute, FTransformCollection::TransformGroup, DynamicState);
-	AddExternalAttribute(SimplicialsAttribute, FTransformCollection::TransformGroup, Simplicials);
 	AddExternalAttribute(SimulatableParticlesAttribute, FGeometryCollection::TransformGroup, SimulatableParticles);
 	AddExternalAttribute(InternalClusterParentTypeAttribute, FGeometryCollection::TransformGroup, InternalClusterParentType);
 }
diff --git a/Engine/Source/Runtime/Experimental/Chaos/Private/PhysicsProxy/GeometryCollectionPhysicsProxy.cpp b/Engine/Source/Runtime/Experimental/Chaos/Private/PhysicsProxy/GeometryCollectionPhysicsProxy.cpp
--- a/Engine/Source/Runtime/Experimental/Chaos/Private/PhysicsProxy/GeometryCollectionPhysicsProxy.cpp
+++ b/Engine/Source/Runtime/Experimental/Chaos/Private/PhysicsProxy/GeometryCollectionPhysicsProxy.cpp
@@ -644,31 +644,6 @@ void FGeometryCollectionPhysicsProxy::Initialize(Chaos::FPBDRigidsEvolutionBase
 		PhysicsThreadCollection.CopyMatchingAttributesFrom(DynamicCollection, MakeArrayView(SkipList, SkipListSize));
 		PhysicsThreadCollection.CopyInitialVelocityAttributesFrom(DynamicCollection);
 
-		// Copy simplicials.
-		// TODO: Ryan - Should we just transfer ownership of the SimplicialsAttribute from the DynamicCollection to
-		// the PhysicsThreadCollection?
-		{
-			if (FGeometryCollection::AreCollisionParticlesEnabled()
-				&& DynamicCollection.HasAttribute(DynamicCollection.SimplicialsAttribute, FTransformCollection::TransformGroup))
-			{
-				const auto& SourceSimplicials = DynamicCollection.GetAttribute<TUniquePtr<FSimplicial>>(
-					DynamicCollection.SimplicialsAttribute, FTransformCollection::TransformGroup);
-				for (int32 Index = PhysicsThreadCollection.NumElements(FTransformCollection::TransformGroup) - 1; 0 <= Index; Index--)
-				{
-					PhysicsThreadCollection.Simplicials[Index].Reset(
-						SourceSimplicials[Index] ? SourceSimplicials[Index]->NewCopy() : nullptr);
-				}
-			}
-			else
-			{
-				for (int32 Index = PhysicsThreadCollection.NumElements(FTransformCollection::TransformGroup) - 1; 0 <= Index; Index--)
-				{
-					PhysicsThreadCollection.Simplicials[Index].Reset();
-				}
-			}
-		}
-
-
 		if (Parameters.EnableClustering)
 		{
 			// make sure we set Activate the right way when clustering is enabled ( only root should be enabled at start ) 
@@ -973,35 +948,6 @@ void FGeometryCollectionPhysicsProxy::InitializeDynamicCollection(FGeometryDynam
 		InitialVelocityFacade.Fill(FVector3f(Params.InitialLinearVelocity), FVector3f(Params.InitialAngularVelocity));
 	}
 
-	// process simplicials
-	{
-		// CVar defined in BodyInstance but pertinent here as we will need to copy simplicials in the case that this is set.
-		// Original CVar is read-only so taking a static ptr here is fine as the value cannot be changed
-		static IConsoleVariable* AnalyticDisableCVar = IConsoleManager::Get().FindConsoleVariable(TEXT("p.IgnoreAnalyticCollisionsOverride"));
-		static const bool bAnalyticsDisabled = (AnalyticDisableCVar && AnalyticDisableCVar->GetBool());
-
-		if (FGeometryCollection::AreCollisionParticlesEnabled()
-			&& RestCollection.HasAttribute(DynamicCollection.SimplicialsAttribute, FTransformCollection::TransformGroup)
-			&& Params.Shared.SizeSpecificData[0].CollisionShapesData.Num()
-			&& (Params.Shared.SizeSpecificData[0].CollisionShapesData[0].CollisionType == ECollisionTypeEnum::Chaos_Surface_Volumetric || bAnalyticsDisabled))
-		{
-			const auto& RestSimplicials = RestCollection.GetAttribute<TUniquePtr<FSimplicial>>(
-				DynamicCollection.SimplicialsAttribute, FTransformCollection::TransformGroup);
-			for (int32 Index = DynamicCollection.NumElements(FTransformCollection::TransformGroup) - 1; 0 <= Index; Index--)
-			{
-				DynamicCollection.Simplicials[Index].Reset(
-					RestSimplicials[Index] ? RestSimplicials[Index]->NewCopy() : nullptr);
-			}
-		}
-		else
-		{
-			for (int32 Index = DynamicCollection.NumElements(FTransformCollection::TransformGroup) - 1; 0 <= Index; Index--)
-			{
-				DynamicCollection.Simplicials[Index].Reset();
-			}
-		}
-	}
-
 	// Process Activity
 	{
 		const int32 NumTransforms = DynamicCollection.SimulatableParticles.Num();
@@ -1222,7 +1168,13 @@ void FGeometryCollectionPhysicsProxy::InitializeBodiesPT(Chaos::FPBDRigidsSolver
 		const TManagedArray<bool>& SimulatableParticles = DynamicCollection.SimulatableParticles;
 		const TManagedArray<FTransform>& MassToLocal = RestCollection->GetAttribute<FTransform>(MassToLocalAttributeName, FTransformCollection::TransformGroup);
 		const TManagedArray<Chaos::FImplicitObjectPtr>& Implicits = DynamicCollection.GetAttribute<Chaos::FImplicitObjectPtr>(FGeometryDynamicCollection::ImplicitsAttribute, FTransformCollection::TransformGroup);
-		const TManagedArray<TUniquePtr<FCollisionStructureManager::FSimplicial>>& Simplicials = DynamicCollection.Simplicials;
+
+		using FSimplicialUniquePtr = TUniquePtr<FCollisionStructureManager::FSimplicial>;
+		const TManagedArray<FSimplicialUniquePtr>* Simplicials = nullptr;
+		if (Parameters.bUseSimplicialsWhenAvailable)
+		{
+			Simplicials = RestCollection->FindAttribute<FSimplicialUniquePtr>(DynamicCollection.SimplicialsAttribute, FTransformCollection::TransformGroup);
+		}
 
 		// In PushToPhysicsState, we're going to compute a relative transform from Parameters.PrevWorldTransform
 		// to a particle's current world transform to get its relative transform. Then, that relative transform is
@@ -1275,7 +1227,7 @@ void FGeometryCollectionPhysicsProxy::InitializeBodiesPT(Chaos::FPBDRigidsSolver
 				PopulateSimulatedParticle(
 					Handle,
 					Parameters.Shared,
-					Simplicials[TransformGroupIndex].Get(),
+					Simplicials? (*Simplicials)[TransformGroupIndex].Get(): nullptr,
 					(bBuildGeometryForChildrenOnPT || (TransformGroupIndex == Parameters.InitialRootIndex)) ? Implicits[TransformGroupIndex] : nullptr,
 					SimFilter,
 					QueryFilter,
@@ -1687,7 +1639,13 @@ void FGeometryCollectionPhysicsProxy::CreateChildrenGeometry_Internal()
 			check(NumTransforms == PhysicsThreadCollection.NumElements(FTransformCollection::TransformGroup));
 			const TManagedArray<FTransform>& MassToLocal = RestCollection->GetAttribute<FTransform>(MassToLocalAttributeName, FTransformCollection::TransformGroup);
 			const TManagedArray<Chaos::FImplicitObjectPtr>& Implicits = PhysicsThreadCollection.GetAttribute<Chaos::FImplicitObjectPtr>(FGeometryDynamicCollection::ImplicitsAttribute, FTransformCollection::TransformGroup);
-			const TManagedArray<TUniquePtr<FCollisionStructureManager::FSimplicial>>& Simplicials = PhysicsThreadCollection.Simplicials;
+
+			using FSimplicialUniquePtr = TUniquePtr<FCollisionStructureManager::FSimplicial>;
+			const TManagedArray<FSimplicialUniquePtr>* Simplicials = nullptr;
+			if (Parameters.bUseSimplicialsWhenAvailable)
+			{
+				Simplicials = RestCollection->FindAttribute<FSimplicialUniquePtr>(FGeometryDynamicCollection::SimplicialsAttribute, FTransformCollection::TransformGroup);
+			}
 
 			// In PushToPhysicsState, we're going to compute a relative transform from Parameters.PrevWorldTransform
 			// to a particle's current world transform to get its relative transform. Then, that relative transform is
@@ -1709,7 +1667,8 @@ void FGeometryCollectionPhysicsProxy::CreateChildrenGeometry_Internal()
 					if (ensure(Handle) && Handle->GetGeometry() == nullptr)
 					{
 						const FTransform WorldTransform = MassToLocal[TransformIndex] * Transform[TransformIndex] * Parameters.WorldTransform;
-						SetImplicitToPTParticles(Handle, Parameters.Shared, Simplicials[TransformIndex].Get(), Implicits[TransformIndex], SimFilter, QueryFilter, WorldTransform, CollisionParticlesPerObjectFraction);
+						const FCollisionStructureManager::FSimplicial* Simplicial = Simplicials ? (*Simplicials)[TransformIndex].Get() : nullptr;
+						SetImplicitToPTParticles(Handle, Parameters.Shared, Simplicial, Implicits[TransformIndex], SimFilter, QueryFilter, WorldTransform, CollisionParticlesPerObjectFraction);
 						check(Handle->GetGeometry() != nullptr);
 					}
 				}
@@ -1885,7 +1844,13 @@ Chaos::TPBDGeometryCollectionParticleHandle<Chaos::FReal, 3>* FGeometryCollectio
 	TManagedArray<uint8>& DynamicState = DynamicCollection.DynamicState;
 	const TManagedArray<FTransform>& MassToLocal = Parameters.RestCollectionShared->GetAttribute<FTransform>(MassToLocalAttributeName, FTransformCollection::TransformGroup);
 	const TManagedArray<Chaos::FImplicitObjectPtr>& Implicits = DynamicCollection.GetAttribute<Chaos::FImplicitObjectPtr>(FGeometryDynamicCollection::ImplicitsAttribute, FTransformCollection::TransformGroup);
-	const TManagedArray<TUniquePtr<FCollisionStructureManager::FSimplicial>>& Simplicials = DynamicCollection.Simplicials;
+
+	using FSimplicialUniquePtr = TUniquePtr<FCollisionStructureManager::FSimplicial>;
+	const TManagedArray<FSimplicialUniquePtr>* Simplicials = nullptr;
+	if (Parameters.bUseSimplicialsWhenAvailable)
+	{
+		Simplicials = Parameters.RestCollectionShared->FindAttribute<FSimplicialUniquePtr>(DynamicCollection.SimplicialsAttribute, FTransformCollection::TransformGroup);
+	}
 	Chaos::Facades::FCollectionAnchoringFacade AnchoringFacade(DynamicCollection);
 
 	//If we are a root particle use the world transform, otherwise set the relative transform
@@ -1934,7 +1899,7 @@ Chaos::TPBDGeometryCollectionParticleHandle<Chaos::FReal, 3>* FGeometryCollectio
 	PopulateSimulatedParticle(
 		Handle,
 		Parameters.Shared,
-		Simplicials[CollectionClusterIndex].Get(),
+		Simplicials? (*Simplicials)[CollectionClusterIndex].Get(): nullptr,
 		(bBuildGeometryForChildrenOnPT || (CollectionClusterIndex == Parameters.InitialRootIndex)) ? Implicits[CollectionClusterIndex] : nullptr,
 		SimFilter,
 		QueryFilter,
diff --git a/Engine/Source/Runtime/Experimental/Chaos/Public/GeometryCollection/GeometryCollectionSimulationCoreTypes.h b/Engine/Source/Runtime/Experimental/Chaos/Public/GeometryCollection/GeometryCollectionSimulationCoreTypes.h
--- a/Engine/Source/Runtime/Experimental/Chaos/Public/GeometryCollection/GeometryCollectionSimulationCoreTypes.h
+++ b/Engine/Source/Runtime/Experimental/Chaos/Public/GeometryCollection/GeometryCollectionSimulationCoreTypes.h
@@ -214,6 +214,7 @@ PRAGMA_ENABLE_DEPRECATION_WARNINGS
 		, bEnableStrainOnCollision(true)
 		, bUseStaticMeshCollisionForTraces(false)
 		, bOptimizeConvexes(true)
+		, bUseSimplicialsWhenAvailable(false)
 		, bUseDamagePropagation(false)
 		, BreakDamagePropagationFactor(1.0f)
 		, ShockDamagePropagationFactor(0.0f)
@@ -286,6 +287,7 @@ PRAGMA_ENABLE_DEPRECATION_WARNINGS
 		, bEnableStrainOnCollision(Other.bEnableStrainOnCollision)
 		, bUseStaticMeshCollisionForTraces(Other.bUseStaticMeshCollisionForTraces)
 		, bOptimizeConvexes(Other.bOptimizeConvexes)
+		, bUseSimplicialsWhenAvailable(Other.bUseSimplicialsWhenAvailable)
 		, bUseDamagePropagation(Other.bUseDamagePropagation)
 		, BreakDamagePropagationFactor(Other.BreakDamagePropagationFactor)
 		, ShockDamagePropagationFactor(Other.ShockDamagePropagationFactor)
@@ -392,7 +394,8 @@ PRAGMA_ENABLE_DEPRECATION_WARNINGS
 
 	bool bEnableStrainOnCollision : 1;
 	bool bUseStaticMeshCollisionForTraces : 1;
-	bool bOptimizeConvexes : 1 = true;
+	bool bOptimizeConvexes : 1;
+	bool bUseSimplicialsWhenAvailable : 1;
 
 	bool bUseDamagePropagation : 1;
 	float BreakDamagePropagationFactor;
diff --git a/Engine/Source/Runtime/Experimental/Chaos/Public/GeometryCollectionProxyData.h b/Engine/Source/Runtime/Experimental/Chaos/Public/GeometryCollectionProxyData.h
--- a/Engine/Source/Runtime/Experimental/Chaos/Public/GeometryCollectionProxyData.h
+++ b/Engine/Source/Runtime/Experimental/Chaos/Public/GeometryCollectionProxyData.h
@@ -118,7 +118,9 @@ public:
 	TManagedArray<uint8> DynamicState; 
 	static_assert(sizeof(EObjectStateTypeEnum) <= sizeof(uint8)); // DynamicState must fit  EObjectStateTypeEnum
 
+	UE_DEPRECATED(5.5, "Simplicials array is optional and should be queried using the FindAttribute method")
 	TManagedArray<TUniquePtr<FCollisionStructureManager::FSimplicial>> Simplicials;
+
 	TManagedArray<bool> SimulatableParticles;
 
 	UE_DEPRECATED(5.4, "CollisionStructureID attribute is no longer supported")
diff --git a/Engine/Source/Runtime/Experimental/GeometryCollectionEngine/Private/GeometryCollection/GeometryCollectionComponent.cpp b/Engine/Source/Runtime/Experimental/GeometryCollectionEngine/Private/GeometryCollection/GeometryCollectionComponent.cpp
--- a/Engine/Source/Runtime/Experimental/GeometryCollectionEngine/Private/GeometryCollection/GeometryCollectionComponent.cpp
+++ b/Engine/Source/Runtime/Experimental/GeometryCollectionEngine/Private/GeometryCollection/GeometryCollectionComponent.cpp
@@ -3855,6 +3855,11 @@ static FORCEINLINE_DEBUGGABLE int32 ComputeParticleLevel(Chaos::FPBDRigidCluster
 
 void UGeometryCollectionComponent::RegisterAndInitializePhysicsProxy()
 {
+	// CVar defined in BodyInstance but pertinent here as we will need to copy simplicials in the case that this is set.
+	// Original CVar is read-only so taking a static ptr here is fine as the value cannot be changed
+	static IConsoleVariable* AnalyticDisableCVar = IConsoleManager::Get().FindConsoleVariable(TEXT("p.IgnoreAnalyticCollisionsOverride"));
+	static const bool bAnalyticsDisabled = (AnalyticDisableCVar && AnalyticDisableCVar->GetBool());
+
 	FSimulationParameters SimulationParameters;
 	{
 #if !(UE_BUILD_SHIPPING || UE_BUILD_TEST)
@@ -3862,6 +3867,7 @@ void UGeometryCollectionComponent::RegisterAndInitializePhysicsProxy()
 #endif
 		EClusterConnectionTypeEnum ClusterCollectionType = ClusterConnectionType_DEPRECATED;
 		float ConnectionGraphBoundsFilteringMargin = 0;
+		bool bUseSimplicialWhenAvailable = false;
 		if (RestCollection)
 		{
 			RestCollection->GetSharedSimulationParams(SimulationParameters.Shared);
@@ -3873,6 +3879,13 @@ PRAGMA_ENABLE_DEPRECATION_WARNINGS
 			SimulationParameters.InitialRootIndex = RestCollection->GetRootIndex();
 			ClusterCollectionType = RestCollection->ClusterConnectionType;
 			ConnectionGraphBoundsFilteringMargin = RestCollection->ConnectionGraphBoundsFilteringMargin;
+			bUseSimplicialWhenAvailable =
+				FGeometryCollection::AreCollisionParticlesEnabled()
+				&& SimulationParameters.RestCollectionShared
+				&& SimulationParameters.RestCollectionShared->HasAttribute(FGeometryDynamicCollection::SimplicialsAttribute, FTransformCollection::TransformGroup)
+				&& SimulationParameters.Shared.SizeSpecificData[0].CollisionShapesData.Num()
+				&& (SimulationParameters.Shared.SizeSpecificData[0].CollisionShapesData[0].CollisionType == ECollisionTypeEnum::Chaos_Surface_Volumetric || bAnalyticsDisabled)
+				;
 		}
 		SimulationParameters.Simulating = BodyInstance.bSimulatePhysics;
 		SimulationParameters.EnableClustering = EnableClustering;
@@ -3921,6 +3934,7 @@ PRAGMA_ENABLE_DEPRECATION_WARNINGS
 		SimulationParameters.bEnableStrainOnCollision = bEnableDamageFromCollision;
 		SimulationParameters.bUseStaticMeshCollisionForTraces = bUseStaticMeshCollisionForTraces;
 		SimulationParameters.bOptimizeConvexes = RestCollection ? RestCollection->bOptimizeConvexes : true;
+		SimulationParameters.bUseSimplicialsWhenAvailable = bUseSimplicialWhenAvailable;
 
 #if SIMULATIONPARAMETERS_CACHE_PARAMETERS
 PRAGMA_DISABLE_DEPRECATION_WARNINGS
