Stream state should also report a content status.

#rnx

diff --git a/Engine/Source/Runtime/Experimental/IoStore/HttpClient/Private/Api.inl b/Engine/Source/Runtime/Experimental/IoStore/HttpClient/Private/Api.inl
--- a/Engine/Source/Runtime/Experimental/IoStore/HttpClient/Private/Api.inl
+++ b/Engine/Source/Runtime/Experimental/IoStore/HttpClient/Private/Api.inl
@@ -190,6 +190,7 @@ FTicketStatus::EId FTicketStatus::GetId() const
 	switch (Activity->State)
 	{
 	case FActivity::EState::RecvMessage:	return EId::Response;
+	case FActivity::EState::RecvStream:
 	case FActivity::EState::RecvDone:		return EId::Content;
 	case FActivity::EState::Cancelled:		return EId::Cancelled;
 	case FActivity::EState::Failed:			return EId::Error;
