Removed unused member.

#rnx

diff --git a/Engine/Source/Runtime/Experimental/IoStore/HttpClient/Private/Activity.inl b/Engine/Source/Runtime/Experimental/IoStore/HttpClient/Private/Activity.inl
--- a/Engine/Source/Runtime/Experimental/IoStore/HttpClient/Private/Activity.inl
+++ b/Engine/Source/Runtime/Experimental/IoStore/HttpClient/Private/Activity.inl
@@ -92,7 +92,6 @@ struct alignas(16) FActivity
 	};
 	UPTRINT				SinkParam;
 	FTicketSink			Sink;
-	FSocket				Socket;
 	FResponseInternal	Response;
 	FBuffer				Buffer;
 };
