Report maximum bytes remaining for chunk-transferred content.

We use the bytes remaining to decide when to dispatch the next request. There are cases where we do not know the size of the data an activity is expecting; the HTTP message. During this phase a maximum size is reported. Another case is chunk-transferred content which is of an indefinite length - as we do not know how much data is left to receive, we report a very large amount.

#rnx

diff --git a/Engine/Source/Runtime/Experimental/IoStore/HttpClient/Private/Activity.inl b/Engine/Source/Runtime/Experimental/IoStore/HttpClient/Private/Activity.inl
--- a/Engine/Source/Runtime/Experimental/IoStore/HttpClient/Private/Activity.inl
+++ b/Engine/Source/Runtime/Experimental/IoStore/HttpClient/Private/Activity.inl
@@ -65,8 +65,8 @@ struct alignas(16) FActivity
 		Build,
 		Send,
 		RecvMessage,
-		RecvContent,
 		RecvStream,
+		RecvContent,
 		RecvDone,
 		Completed,
 		Cancelled,
@@ -100,7 +100,7 @@ struct alignas(16) FActivity
 static void Activity_TraceStateNames()
 {
 	FAnsiStringView StateNames[] = {
-		"None", "Build", "Send", "RecvMessage", "RecvContent", "RecvStream",
+		"None", "Build", "Send", "RecvMessage", "RecvStream", "RecvContent",
 		"RecvDone", "Completed", "Cancelled", "Failed", "$",
 	};
 	static_assert(UE_ARRAY_COUNT(StateNames) == int32(FActivity::EState::_Num) + 1);
@@ -142,8 +142,8 @@ static int32 Activity_Rewind(FActivity* Activity)
 ////////////////////////////////////////////////////////////////////////////////
 static uint32 Activity_RemainingKiB(FActivity* Activity)
 {
-	if (Activity->State < FActivity::EState::RecvContent) return MAX_uint32;
-	if (Activity->State > FActivity::EState::RecvContent) return 0;
+	if (Activity->State <= FActivity::EState::RecvStream)  return MAX_uint32;
+	if (Activity->State >  FActivity::EState::RecvContent) return 0;
 
 	uint32 ContentLength = uint32(Activity->Response.ContentLength);
 	check(Activity->StateParam <= ContentLength);
