Separated read call from the logic that receives the data.

#rnx

diff --git a/Engine/Source/Runtime/Experimental/IoStore/HttpClient/Private/Loop.inl b/Engine/Source/Runtime/Experimental/IoStore/HttpClient/Private/Loop.inl
--- a/Engine/Source/Runtime/Experimental/IoStore/HttpClient/Private/Loop.inl
+++ b/Engine/Source/Runtime/Experimental/IoStore/HttpClient/Private/Loop.inl
@@ -59,6 +59,45 @@ static FOutcome DoSend(FActivity* Activity, FPeerType& Peer)
 	return FOutcome::Ok(Result);
 }
 
+////////////////////////////////////////////////////////////////////////////////
+static FOutcome DoRecvPeer(FActivity* Activity, FPeerType& Peer, int32& MaxRecvSize, int32 Size)
+{
+	Size = FMath::Min(Size, MaxRecvSize);
+	check(Size >= 0);
+	if (Size == 0)
+	{
+		return FOutcome::Waiting();
+	}
+
+	Trace(Activity, ETrace::StateChange, Activity->State);
+
+	FMutableMemoryView DestView = Activity->Dest->GetMutableView();
+	char* Cursor = (char*)(DestView.GetData()) + Activity->StateParam;
+	check(Size + Activity->StateParam <= DestView.GetSize());
+
+	FOutcome Outcome = Peer.Recv(Cursor, Size);
+
+	if (Outcome.IsWaiting())
+	{
+		return Outcome;
+	}
+
+	if (Outcome.IsError())
+	{
+		Activity_SetError(Activity, Outcome.GetMessage().GetData());
+		return Outcome;
+	}
+
+	check(Outcome.IsOk());
+
+	int32 Result = Outcome.GetResult();
+	check(Result <= MaxRecvSize);
+	Activity->StateParam += Result;
+	MaxRecvSize -= Result;
+
+	return Outcome;
+}
+
 ////////////////////////////////////////////////////////////////////////////////
 static FOutcome DoRecvMessage(FActivity* Activity, FPeerType& Peer)
 {
@@ -267,48 +306,20 @@ static FOutcome DoRecvContent(FActivity* Activity, FPeerType& Peer, int32& MaxRe
 {
 	TRACE_CPUPROFILER_EVENT_SCOPE(IasHttp::DoRecvContent);
 
-	FResponseInternal& Response = Activity->Response;
-	FMutableMemoryView DestView = Activity->Dest->GetMutableView();
-
 	while (true)
 	{
+		const FResponseInternal& Response = Activity->Response;
 		int32 Size = (Response.ContentLength - Activity->StateParam);
 		if (Size == 0)
 		{
 			break;
 		}
 
-		Size = FMath::Min(Size, MaxRecvSize);
-		check(Size >= 0);
-		if (Size == 0)
+		FOutcome Outcome = DoRecvPeer(Activity, Peer, MaxRecvSize, Size);
+		if (!Outcome.IsOk())
 		{
-			return FOutcome::Waiting();
-		}
-
-		Trace(Activity, ETrace::StateChange, Activity->State);
-
-		char* Cursor = (char*)(DestView.GetData()) + Activity->StateParam;
-
-		FOutcome Outcome = Peer.Recv(Cursor, Size);
-
-		if (Outcome.IsWaiting())
-		{
-			return Outcome;
-		}
-
-		if (Outcome.IsError())
-		{
-			Activity_SetError(Activity, Outcome.GetMessage().GetData());
 			return Outcome;
 		}
-
-		check(Outcome.IsOk());
-
-		int32 Result = Outcome.GetResult();
-		check(Result <= MaxRecvSize);
-
-		Activity->StateParam += Result;
-		MaxRecvSize -= Result;
 	}
 
 #if IAS_HTTP_WITH_PERF
