Enabling IoStore HTTP UE Mime Types for Zen Store connection

diff --git a/Engine/Source/Runtime/Experimental/IoStore/HttpClient/Private/Api.inl b/Engine/Source/Runtime/Experimental/IoStore/HttpClient/Private/Api.inl
--- a/Engine/Source/Runtime/Experimental/IoStore/HttpClient/Private/Api.inl
+++ b/Engine/Source/Runtime/Experimental/IoStore/HttpClient/Private/Api.inl
@@ -34,11 +34,9 @@ FRequest&& FRequest::Accept(EMimeType MimeType)
 	case EMimeType::Binary:				return Header("Accept", "application/octet-stream");
 	case EMimeType::Json:				return Header("Accept", "application/json");
 	case EMimeType::Xml:				return Header("Accept", "application/xml");
-	/* UE_CUSTOM_MIME_TYPES
 	case EMimeType::CbObject:			return Header("Accept", "application/x-ue-cb");
 	case EMimeType::CbPackage:			return Header("Accept", "application/x-ue-pkg");
 	case EMimeType::CompressedBuffer:	return Header("Accept", "application/x-ue-comp");
-	*/
 	}
 
 	return MoveTemp(*this);
