CSV stats for Iris object creation and rep notify detailed breakdown.

diff --git a/Engine/Source/Runtime/Experimental/Iris/Core/Private/Iris/Core/IrisProfiler.cpp b/Engine/Source/Runtime/Experimental/Iris/Core/Private/Iris/Core/IrisProfiler.cpp
new file mode 100644
--- /dev/null
+++ b/Engine/Source/Runtime/Experimental/Iris/Core/Private/Iris/Core/IrisProfiler.cpp
@@ -0,0 +1,66 @@
+// Copyright Epic Games, Inc. All Rights Reserved.
+
+#include "Iris/Core/IrisProfiler.h"
+#include "HAL/IConsoleManager.h"
+
+CSV_DEFINE_CATEGORY(IrisClient, IRIS_CLIENT_PROFILER_ENABLE);
+CSV_DEFINE_CATEGORY(IrisClientDetailObjectCreate, IRIS_CLIENT_PROFILER_ENABLE);
+CSV_DEFINE_CATEGORY(IrisClientDetailRepNotify, IRIS_CLIENT_PROFILER_ENABLE);
+CSV_DEFINE_CATEGORY(IrisClientDetailRPC, IRIS_CLIENT_PROFILER_ENABLE);
+
+namespace UE::Net::Private
+{
+
+static bool bEnableDetailedClientProfiler = 0;
+static FAutoConsoleVariableRef CVarEnableDetailedClientProfilerRef(TEXT("net.Iris.EnableDetailedClientProfiler"), bEnableDetailedClientProfiler, TEXT("Generates detailed CSV Iris stats (client only)."), ECVF_Default );
+
+}
+
+namespace UE::Net
+{
+
+#if IRIS_CLIENT_PROFILER_ENABLE
+
+void FClientProfiler::RecordObjectCreate(FName ObjectName, bool bIsSubObject)
+{
+	FCsvProfiler::RecordCustomStat("ClientObjectCreate", CSV_CATEGORY_INDEX(IrisClient), 1, ECsvCustomStatOp::Accumulate);
+	if (!bIsSubObject)
+	{
+		FCsvProfiler::RecordCustomStat("ClientObjectCreateRoot", CSV_CATEGORY_INDEX(IrisClient), 1, ECsvCustomStatOp::Accumulate);
+	}
+
+	if (UE::Net::Private::bEnableDetailedClientProfiler)
+	{
+		if (!bIsSubObject)
+		{
+			FCsvProfiler::RecordCustomStat(ObjectName, CSV_CATEGORY_INDEX(IrisClientDetailObjectCreate), 1, ECsvCustomStatOp::Accumulate);
+		}
+	}
+}
+
+void FClientProfiler::RecordRepNotify(FName RepNotifyName)
+{
+	FCsvProfiler::RecordCustomStat("RepNotifyCount", CSV_CATEGORY_INDEX(IrisClient), 1, ECsvCustomStatOp::Accumulate);
+	if (UE::Net::Private::bEnableDetailedClientProfiler)
+	{
+		FCsvProfiler::RecordCustomStat(RepNotifyName, CSV_CATEGORY_INDEX(IrisClientDetailRepNotify), 1, ECsvCustomStatOp::Accumulate);
+	}
+}
+
+void FClientProfiler::RecordRPC(FName RPCName)
+{
+	FCsvProfiler::RecordCustomStat("CallCountRPC", CSV_CATEGORY_INDEX(IrisClient), 1, ECsvCustomStatOp::Accumulate);
+	if (UE::Net::Private::bEnableDetailedClientProfiler)
+	{
+		FCsvProfiler::RecordCustomStat(RPCName, CSV_CATEGORY_INDEX(IrisClientDetailRPC), 1, ECsvCustomStatOp::Accumulate);
+	}
+}
+
+bool FClientProfiler::IsCapturing()
+{
+	return FCsvProfiler::Get()->IsCapturing();
+}
+
+#endif
+
+}
\ No newline at end of file
diff --git a/Engine/Source/Runtime/Experimental/Iris/Core/Private/Iris/IrisConfig.cpp b/Engine/Source/Runtime/Experimental/Iris/Core/Private/Iris/IrisConfig.cpp
--- a/Engine/Source/Runtime/Experimental/Iris/Core/Private/Iris/IrisConfig.cpp
+++ b/Engine/Source/Runtime/Experimental/Iris/Core/Private/Iris/IrisConfig.cpp
@@ -40,6 +40,3 @@ EReplicationSystem GetUseIrisReplicationCmdlineValue()
 
 // Enable Iris category by default on servers
 CSV_DEFINE_CATEGORY(Iris, WITH_SERVER_CODE);
-
-// Enable IrisClient category on the client only
-CSV_DEFINE_CATEGORY(IrisClient, !WITH_SERVER_CODE);
\ No newline at end of file
diff --git a/Engine/Source/Runtime/Experimental/Iris/Core/Private/Iris/ReplicationState/PropertyReplicationState.cpp b/Engine/Source/Runtime/Experimental/Iris/Core/Private/Iris/ReplicationState/PropertyReplicationState.cpp
--- a/Engine/Source/Runtime/Experimental/Iris/Core/Private/Iris/ReplicationState/PropertyReplicationState.cpp
+++ b/Engine/Source/Runtime/Experimental/Iris/Core/Private/Iris/ReplicationState/PropertyReplicationState.cpp
@@ -19,8 +19,6 @@
 
 DEFINE_LOG_CATEGORY_STATIC(LogIrisRepNotify, Warning, All);
 
-CSV_DECLARE_CATEGORY_EXTERN(IrisClient);
-
 namespace UE::Net
 {
 
@@ -479,7 +477,9 @@ void FPropertyReplicationState::CallRepNotifies(void* RESTRICT DstData, const FC
 
 				if (bShouldCallRepNotify)
 				{
-					CSV_CUSTOM_STAT(IrisClient, RepNotifyCount, 1, ECsvCustomStatOp::Accumulate);
+#if IRIS_CLIENT_PROFILER_ENABLE
+					UE::Net::FClientProfiler::RecordRepNotify(RepNotifyFunction->GetFName());
+#endif
 
 					// We only want to call RepNotify once for c-arrays
 					LastPropertyWithRepNotify = Property;
diff --git a/Engine/Source/Runtime/Experimental/Iris/Core/Private/Iris/ReplicationSystem/NetBlob/NetRPC.cpp b/Engine/Source/Runtime/Experimental/Iris/Core/Private/Iris/ReplicationSystem/NetBlob/NetRPC.cpp
--- a/Engine/Source/Runtime/Experimental/Iris/Core/Private/Iris/ReplicationSystem/NetBlob/NetRPC.cpp
+++ b/Engine/Source/Runtime/Experimental/Iris/Core/Private/Iris/ReplicationSystem/NetBlob/NetRPC.cpp
@@ -22,6 +22,7 @@
 #include "Iris/Serialization/ObjectNetSerializer.h"
 #include "Iris/Serialization/NetBitStreamUtil.h"
 #include "Iris/Core/IrisLog.h"
+#include "Iris/Core/IrisProfiler.h"
 #include "Containers/ArrayView.h"
 #include "ProfilingDebugging/CsvProfiler.h"
 
@@ -568,6 +569,10 @@ void FNetRPC::CallFunction(FNetRPCCallContext& CallContext)
 
 	// Call function
 	{
+#if IRIS_CLIENT_PROFILER_ENABLE
+		UE::Net::FClientProfiler::RecordRPC(Function->GetFName());
+#endif
+
 		UE::Net::FScopedNetContextRPC CallingRPC;
 		Object->ProcessEvent(const_cast<UFunction*>(Function), FunctionParameters);
 	}
diff --git a/Engine/Source/Runtime/Experimental/Iris/Core/Private/Iris/ReplicationSystem/ReplicationReader.cpp b/Engine/Source/Runtime/Experimental/Iris/Core/Private/Iris/ReplicationSystem/ReplicationReader.cpp
--- a/Engine/Source/Runtime/Experimental/Iris/Core/Private/Iris/ReplicationSystem/ReplicationReader.cpp
+++ b/Engine/Source/Runtime/Experimental/Iris/Core/Private/Iris/ReplicationSystem/ReplicationReader.cpp
@@ -933,13 +933,6 @@ void FReplicationReader::ReadObjectInBatch(FNetSerializationContext& Context, FN
 	
 		// Get Bridge
 		FReplicationBridgeSerializationContext BridgeContext(Context, Parameters.ConnectionId);
-
-		CSV_CUSTOM_STAT(IrisClient, ClientObjectCreate, 1, ECsvCustomStatOp::Accumulate);
-		if (!bIsSubObject)
-		{
-			CSV_CUSTOM_STAT(IrisClient, ClientObjectCreateRoot, 1, ECsvCustomStatOp::Accumulate);
-		}
-
 		const FReplicationBridgeCreateNetRefHandleResult CreateResult = ReplicationBridge->CallCreateNetRefHandleFromRemote(RootObjectOfSubObject, IncompleteHandle, BridgeContext);
 		FNetRefHandle NetRefHandle = CreateResult.NetRefHandle;
 		if (!NetRefHandle.IsValid())
@@ -964,6 +957,17 @@ void FReplicationReader::ReadObjectInBatch(FNetSerializationContext& Context, FN
 		FReplicatedObjectInfo& ObjectInfo = StartReplication(InternalIndex);
 
 		ObjectInfo.bIsDeltaCompressionEnabled = bIsDeltaCompressed;
+
+#if IRIS_CLIENT_PROFILER_ENABLE
+		if (UE::Net::FClientProfiler::IsCapturing())
+		{
+			UObject* Object = NetRefHandleManager->GetReplicatedObjectInstance(InternalIndex);
+			if (Object)
+			{
+				UE::Net::FClientProfiler::RecordObjectCreate(Object->GetClass()->GetFName(), bIsSubObject);
+			}
+		}
+#endif
 	}
 	else
 	{
@@ -2055,6 +2059,7 @@ void FReplicationReader::ResolveAndDispatchUnresolvedReferences()
 		ResolveAndDispatchUnresolvedReferencesForObject(Context, InternalIndex);
 	}
 
+#if IRIS_CLIENT_PROFILER_ENABLE
 	CSV_CUSTOM_STAT(IrisClient, HotUnresolvedHandleCache, HotUnresolvedHandleCache.Num(), ECsvCustomStatOp::Set);
 	CSV_CUSTOM_STAT(IrisClient, ColdUnresolvedHandleCache, ColdUnresolvedHandleCache.Num(), ECsvCustomStatOp::Set);
 
@@ -2067,6 +2072,7 @@ void FReplicationReader::ResolveAndDispatchUnresolvedReferences()
 		VisitedUnresolvedHandles.GetAllocatedSize() +
 		InternalObjectsToResolve.GetAllocatedSize());
 	CSV_CUSTOM_STAT(IrisClient, UnresolvedHandleBufferSizes, TotalCacheSize, ECsvCustomStatOp::Set);
+#endif
 
 	if (NumHandlesPendingResolveLastUpdate != VisitedUnresolvedHandles.Num() || ObjectsWithAttachmentPendingResolve.Num() > 0)
 	{
diff --git a/Engine/Source/Runtime/Experimental/Iris/Core/Public/Iris/Core/IrisProfiler.h b/Engine/Source/Runtime/Experimental/Iris/Core/Public/Iris/Core/IrisProfiler.h
--- a/Engine/Source/Runtime/Experimental/Iris/Core/Public/Iris/Core/IrisProfiler.h
+++ b/Engine/Source/Runtime/Experimental/Iris/Core/Public/Iris/Core/IrisProfiler.h
@@ -3,6 +3,7 @@
 #pragma once
 
 #include "HAL/Platform.h"
+#include "ProfilingDebugging/CsvProfiler.h"
 
 #ifndef IRIS_PROFILER_ENABLE
 #	if (UE_BUILD_SHIPPING)
@@ -58,3 +59,29 @@
 #else
 #	define IRIS_PROFILER_SCOPE_VERBOSE(x)
 #endif
+
+#ifndef IRIS_CLIENT_PROFILER_ENABLE
+#	define IRIS_CLIENT_PROFILER_ENABLE (!WITH_SERVER_CODE && CSV_PROFILER)
+#endif
+
+namespace UE::Net
+{
+
+#if IRIS_CLIENT_PROFILER_ENABLE
+
+class FClientProfiler
+{
+public:
+
+	/** Record profiler events. */
+	static void RecordObjectCreate(FName ObjectName, bool bIsSubObject);
+	static void RecordRepNotify(FName RepNotifyName);
+	static void RecordRPC(FName RPCName);
+
+	/** Return true if capturing events. */
+	static bool IsCapturing();
+};
+
+#endif
+
+}
\ No newline at end of file
