Iris
* Added PushModelDirtyObjects and NetForceUpdateObjects csv data

#rnx

diff --git a/Engine/Source/Runtime/Experimental/Iris/Core/Private/Iris/ReplicationSystem/DirtyNetObjectTracker.cpp b/Engine/Source/Runtime/Experimental/Iris/Core/Private/Iris/ReplicationSystem/DirtyNetObjectTracker.cpp
--- a/Engine/Source/Runtime/Experimental/Iris/Core/Private/Iris/ReplicationSystem/DirtyNetObjectTracker.cpp
+++ b/Engine/Source/Runtime/Experimental/Iris/Core/Private/Iris/ReplicationSystem/DirtyNetObjectTracker.cpp
@@ -7,7 +7,10 @@
 #include "Iris/Core/IrisProfiler.h"
 #include "Iris/ReplicationSystem/ReplicationSystem.h"
 #include "Iris/ReplicationSystem/ReplicationSystemInternal.h"
+
+#include "ProfilingDebugging/CsvProfiler.h"
 #include "Traits/IntType.h"
+
 #include <atomic>
 
 // Don't compile verbose logs in Shipping builds
@@ -156,6 +159,10 @@ void FDirtyNetObjectTracker::MarkNetObjectDirty(FInternalNetRefIndex NetObjectIn
 		return;
 	}
 
+#if UE_NET_IRIS_CSV_STATS
+	PushModelDirtyObjectsCount += (DirtyNetObjects.IsBitSet(NetObjectIndex) ? 0 : 1);
+#endif
+
 	const uint32 BitOffset = NetObjectIndex;
 	const StorageType BitMask = StorageType(1) << (BitOffset & (StorageTypeBitCount - 1));
 
@@ -169,6 +176,10 @@ void FDirtyNetObjectTracker::MarkNetObjectDirty(FInternalNetRefIndex NetObjectIn
 
 void FDirtyNetObjectTracker::ForceNetUpdate(FInternalNetRefIndex NetObjectIndex)
 {
+#if UE_NET_IRIS_CSV_STATS
+	ForceNetUpdateObjectsCount += (ForceNetUpdateObjects.IsBitSet(NetObjectIndex)?0:1);
+#endif
+
 	ForceNetUpdateObjects.SetBit(NetObjectIndex);
 
 	// Flag the object dirty so we update his filters too
@@ -223,6 +234,19 @@ void FDirtyNetObjectTracker::ReconcilePolledList(const FNetBitArrayView& Objects
 	std::atomic_thread_fence(std::memory_order_seq_cst);
 }
 
+#if UE_NET_IRIS_CSV_STATS
+void FDirtyNetObjectTracker::ReportCSVStats()
+{
+	CSV_CUSTOM_STAT(Iris, PushModelDirtyObjects, PushModelDirtyObjectsCount, ECsvCustomStatOp::Set);
+	CSV_CUSTOM_STAT(Iris, ForceNetUpdateObjects, ForceNetUpdateObjectsCount, ECsvCustomStatOp::Set);
+
+	PushModelDirtyObjectsCount = 0;
+	ForceNetUpdateObjectsCount = 0;
+}
+#endif
+
+#pragma region GlobalFunctions
+
 void MarkNetObjectStateDirty(uint32 ReplicationSystemId, FInternalNetRefIndex NetObjectIndex)
 {
 	if (UReplicationSystem* ReplicationSystem = GetReplicationSystem(ReplicationSystemId))
@@ -241,4 +265,6 @@ void ForceNetUpdate(uint32 ReplicationSystemId, FInternalNetRefIndex NetObjectIn
 	}
 }
 
-}
+#pragma endregion
+
+} // end namespace UE::Net::Private
diff --git a/Engine/Source/Runtime/Experimental/Iris/Core/Private/Iris/ReplicationSystem/DirtyNetObjectTracker.h b/Engine/Source/Runtime/Experimental/Iris/Core/Private/Iris/ReplicationSystem/DirtyNetObjectTracker.h
--- a/Engine/Source/Runtime/Experimental/Iris/Core/Private/Iris/ReplicationSystem/DirtyNetObjectTracker.h
+++ b/Engine/Source/Runtime/Experimental/Iris/Core/Private/Iris/ReplicationSystem/DirtyNetObjectTracker.h
@@ -6,6 +6,7 @@
 #include "Net/Core/DirtyNetObjectTracker/GlobalDirtyNetObjectTracker.h"
 
 #include "Iris/IrisConfig.h"
+#include "Iris/Core/IrisCsv.h"
 
 namespace UE::Net::Private
 {
@@ -58,6 +59,10 @@ public:
 	/** Reset the global list and look at the final polled list and clear any flags for objects that got polled */
 	void ReconcilePolledList(const FNetBitArrayView& ObjectsPolled);
 
+#if UE_NET_IRIS_CSV_STATS
+	void ReportCSVStats();
+#endif
+
 	/** Returns the list of objects that are dirty this frame or were dirty in previous frames but not cleaned up at that time. */
 	const FNetBitArrayView GetAccumulatedDirtyNetObjects() const { return MakeNetBitArrayView(AccumulatedDirtyNetObjects); }
 
@@ -108,6 +113,11 @@ private:
 #if UE_NET_THREAD_SAFETY_CHECK
 	std::atomic_bool bIsExternalAccessAllowed = false;
 #endif
+
+#if UE_NET_IRIS_CSV_STATS
+	int32 PushModelDirtyObjectsCount = 0;
+	int32 ForceNetUpdateObjectsCount = 0;
+#endif
 };
 
 /**
diff --git a/Engine/Source/Runtime/Experimental/Iris/Core/Public/Iris/ReplicationSystem/ObjectReplicationBridge.h b/Engine/Source/Runtime/Experimental/Iris/Core/Public/Iris/ReplicationSystem/ObjectReplicationBridge.h
--- a/Engine/Source/Runtime/Experimental/Iris/Core/Public/Iris/ReplicationSystem/ObjectReplicationBridge.h
+++ b/Engine/Source/Runtime/Experimental/Iris/Core/Public/Iris/ReplicationSystem/ObjectReplicationBridge.h
@@ -192,7 +192,7 @@ public:
 
 public:
 
-	// Debug functions exposed that are triggered via console commands
+	// Debug functions that are triggered via console commands
 	void PrintDynamicFilterClassConfig() const;
 	void PrintReplicatedObjects(uint32 ArgTraits) const;
 	void PrintAlwaysRelevantObjects(uint32 ArgTraits) const;
