Added support for direct s3 upload from IASTool
Fixed double blocking in isolated output sentry.
#tests Isolated output sentry in gamedev, with s3 direct upload enabled and disabled.

diff --git a/Engine/Source/Runtime/Experimental/IoStore/OnDemand/Private/Tool/ChunkPlugin.cpp b/Engine/Source/Runtime/Experimental/IoStore/OnDemand/Private/Tool/ChunkPlugin.cpp
--- a/Engine/Source/Runtime/Experimental/IoStore/OnDemand/Private/Tool/ChunkPlugin.cpp
+++ b/Engine/Source/Runtime/Experimental/IoStore/OnDemand/Private/Tool/ChunkPlugin.cpp
@@ -17,20 +17,176 @@
 #include "Misc/Paths.h"
 #include "Misc/PathViews.h"
 #include "Misc/StringBuilder.h"
+#include "S3/S3Client.h"
 #include "Serialization/JsonReader.h"
 #include "Serialization/JsonSerializer.h"
 #include "Serialization/LargeMemoryWriter.h"
 #include "Serialization/MemoryReader.h"
+#include "UploadQueue.h"
 
 namespace UE::IoStore::Tool
 {
 
+
+
+struct FS3Params
+{
+public:
+	FString S3ServiceUrl;
+	FString S3Bucket;
+	FString S3BucketPrefix;
+	FString S3Region;
+	FString S3AccessKey;
+	FString S3SecretKey;
+	FString S3SessionToken;
+	int32 S3MaxConcurrentUploads;
+
+	FS3Params(const FContext& Context)
+	{
+		S3ServiceUrl = FString(Context.Get<FStringView>(TEXT("-ServiceUrl"), FString()));
+		S3Bucket = FString(Context.Get<FStringView>(TEXT("-Bucket"), FString()));
+		S3Region = FString(Context.Get<FStringView>(TEXT("-Region"), FString()));
+		S3AccessKey = FString(Context.Get<FStringView>(TEXT("-AccessKey"), FString()));
+		S3SecretKey = FString(Context.Get<FStringView>(TEXT("-SecretKey"), FString()));
+		S3SessionToken = FString(Context.Get<FStringView>(TEXT("-SessionToken"), FString()));
+
+
+		S3BucketPrefix = FString(Context.Get<FStringView>(TEXT("-BucketPrefix"), FString()));
+		S3MaxConcurrentUploads = Context.Get<int32>(TEXT("-MaxConcurrentUploads"), 10);
+	}
+
+	bool IsValid()
+	{
+		return !S3AccessKey.IsEmpty() && !S3SecretKey.IsEmpty() && !S3SessionToken.IsEmpty();
+	}
+};
+
+
 ////////////////////////////////////////////////////////////////////////////////
 struct FChunkPluginSettings
 {
 	TMap<FString, TArray<FString>> PackageSets;
 };
 
+class FChunkWriterInterface
+{
+public:
+	virtual ~FChunkWriterInterface() {}
+	virtual FIoStatus WriteChunk(const FString& RelativeDir, FIoBuffer Chunk, const FIoHash& Hash) = 0;
+
+	virtual bool Flush() = 0;
+};
+
+
+
+class FS3ChunkWriter : public FChunkWriterInterface
+{
+public:
+	FS3ChunkWriter(const FS3Params& S3Params) : 
+		BucketPrefix(S3Params.S3BucketPrefix),
+		Client(FS3ClientConfig({ S3Params.S3Region, S3Params.S3ServiceUrl }), FS3ClientCredentials(S3Params.S3AccessKey, S3Params.S3SecretKey, S3Params.S3SessionToken)),
+		UploadQueue(Client, S3Params.S3Bucket, S3Params.S3MaxConcurrentUploads)
+	{
+	}
+
+	virtual ~FS3ChunkWriter()
+	{
+		// this should have already been done but do it now just in case
+		UploadQueue.Flush();
+	}
+
+
+	virtual FIoStatus WriteChunk(const FString& RelativeDirectory, FIoBuffer Chunk, const FIoHash& Hash) override
+	{
+		const FString HashString = LexToString(Hash);
+
+		TStringBuilder<256> Key;
+		Key << BucketPrefix << TEXT("/")
+			<< RelativeDirectory
+			<< TEXT("/") << HashString.Left(2)
+			<< TEXT("/") << HashString
+			<< TEXT(".iochunk");
+
+		if (UploadQueue.Enqueue(Key, Chunk) == false)
+		{
+			return FIoStatus(EIoErrorCode::WriteError, TEXT("Failed to upload chunk"));
+		}
+
+		return FIoStatus(EIoErrorCode::Ok);
+	}
+
+	virtual bool Flush() override
+	{
+		if (!UploadQueue.Flush())
+		{
+			return false;
+		}
+		return true;
+	}
+
+
+private:
+	FString BucketPrefix;
+	UE::FS3Client Client;
+	FUploadQueue UploadQueue;
+};
+
+
+class FDiskChunkWriter : public FChunkWriterInterface
+{
+public:
+	FDiskChunkWriter(const FString& InOutputFolder)
+	{
+		OutputFolder = InOutputFolder;
+		
+	}
+
+	virtual ~FDiskChunkWriter()
+	{
+	}
+
+	virtual  FIoStatus WriteChunk(const FString& Directory, FIoBuffer Chunk, const FIoHash& Hash) override 
+	{
+		IFileManager& FileMgr = IFileManager::Get();
+		const FString HashString = LexToString(Hash);
+
+		TStringBuilder<256> Sb;
+		Sb << OutputFolder << TEXT("/") << Directory << TEXT("/") << HashString.Left(2);
+
+		bool bTree = true;
+		if (FileMgr.MakeDirectory(Sb.ToString(), bTree) == false)
+		{
+			return FIoStatusBuilder(EIoErrorCode::WriteError)
+				<< TEXT("Failed to create directory '")
+				<< FString(Sb.ToString())
+				<< TEXT("'");
+		}
+
+		Sb << TEXT("/") << HashString << TEXT(".iochunk");
+
+		if (TUniquePtr<FArchive> Ar(IFileManager::Get().CreateFileWriter(Sb.ToString())); Ar.IsValid())
+		{
+			UE_LOG(LogIoStore, Display, TEXT("Writing file '%s' (%.2lf KiB)"),
+				Sb.ToString(), double(Chunk.GetView().GetSize()) / 1024.0);
+			Ar->Serialize((void*)Chunk.GetView().GetData(), Chunk.GetView().GetSize());
+
+			return EIoErrorCode::Ok;
+		}
+
+		return FIoStatusBuilder(EIoErrorCode::WriteError)
+			<< TEXT("Failed to write file '")
+			<< FString(Sb.ToString())
+			<< TEXT("'");
+	}
+
+	bool Flush() override
+	{
+		return true;
+	}
+private:
+	FString OutputFolder;
+};
+
 ////////////////////////////////////////////////////////////////////////////////
 static FIoStatus WriteChunk(const FString& Directory, FMemoryView Chunk, const FIoHash& Hash)
 {
@@ -66,11 +222,7 @@ static FIoStatus WriteChunk(const FString& Directory, FMemoryView Chunk, const F
 		<< TEXT("'");
 }
 
-////////////////////////////////////////////////////////////////////////////////
-static FIoStatus WriteChunk(const FString& Directory, FMemoryView Chunk)
-{
-	return WriteChunk(Directory, Chunk, FIoHash::HashBuffer(Chunk));
-}
+
 
 ////////////////////////////////////////////////////////////////////////////////
 static int32 ChunkPluginCommandEntry(const FContext& Context)
@@ -82,18 +234,21 @@ static int32 ChunkPluginCommandEntry(const FContext& Context)
 	const FString BuildVersion			= FString(Context.Get<FStringView>(TEXT("-BuildVersion"), FString()));
 	const FString OnDemandTocName		= FString(Context.Get<FStringView>(TEXT("-OnDemandTocName"), FString()));
 	const FString InputFolder			= FString(Context.Get<FStringView>(TEXT("-InputFolder"), FString()));
-	const FString OutputFolder			= FString(Context.Get<FStringView>(TEXT("-OutputFolder"), FString()));
+	const FString InOutputFolder			= FString(Context.Get<FStringView>(TEXT("-OutputFolder"), FString()));
 	const FString IntermediateFolder	= FString(Context.Get<FStringView>(TEXT("-IntermediateFolder"), FString()));
 	FString SettingsFile				= FString(Context.Get<FStringView>(TEXT("-SettingsFile"), FString()));
+	
+	FS3Params S3Params(Context);
+
 	const bool bIncludeSigPak			= Context.Get<bool>(TEXT("-IncludeSigPak"), false);
 	const bool bDeleteContainerFiles	= !Context.Get<bool>(TEXT("-KeepContainerFiles"), false);
+	FString OutputFolder				= InOutputFolder;
 	FString ContainerFolder				= InputFolder;
-	FString IoStoreOutputFolder			= OutputFolder / TEXT("iostore");
-	FString ChunksOutputFolder			= IoStoreOutputFolder / TEXT("chunks");
+	FString IoStoreRelativeFolder		= TEXT("iostore");
+	FString ChunksRelativeFolder		= IoStoreRelativeFolder / TEXT("chunks");
 
 	FPaths::NormalizeDirectoryName(ContainerFolder);
-	FPaths::NormalizeDirectoryName(IoStoreOutputFolder);
-	FPaths::NormalizeDirectoryName(ChunksOutputFolder);
+	FPaths::NormalizeDirectoryName(OutputFolder);
 	FPaths::NormalizeFilename(SettingsFile);
 
 	UE_LOG(LogIoStore, Display, TEXT("I/O store chunk plugin:"));
@@ -110,6 +265,18 @@ static int32 ChunkPluginCommandEntry(const FContext& Context)
 
 	IFileManager& FileMgr = IFileManager::Get();
 
+	TUniquePtr<FChunkWriterInterface> ChunkWriter;
+	bool bUseS3 = false;
+	if (S3Params.IsValid())
+	{
+		ChunkWriter = MakeUnique<FS3ChunkWriter>(S3Params);
+	}
+	else
+	{
+		ChunkWriter = MakeUnique<FDiskChunkWriter>(OutputFolder);
+	}
+	check(ChunkWriter != nullptr)
+
 	FChunkPluginSettings Settings;
 	if (!SettingsFile.IsEmpty())
 	{
@@ -165,13 +332,6 @@ static int32 ChunkPluginCommandEntry(const FContext& Context)
 		}
 	}
 
-	
-	if (FileMgr.MakeDirectory(*ChunksOutputFolder, true) == false)
-	{
-		UE_LOG(LogIoStore, Error, TEXT("Failed to create directory '%s'"), *ChunksOutputFolder);
-		return -1;
-	}
-
 	TMap<FGuid, FAES::FAESKey> EncryptionKeys;
 	if (FileMgr.DirectoryExists(*ContainerFolder) == false)
 	{
@@ -342,8 +502,10 @@ static int32 ChunkPluginCommandEntry(const FContext& Context)
 				return -1;
 			}
 
-			if (FIoStatus WriteStatus = WriteChunk(ChunksOutputFolder, ReadResult.IoBuffer.GetView(), ChunkHash);
-				WriteStatus.IsOk() == false)
+
+
+			FIoStatus WriteStatus = ChunkWriter->WriteChunk(ChunksRelativeFolder, ReadResult.IoBuffer, ChunkHash);
+			if (WriteStatus.IsOk() == false)
 			{
 				UE_LOG(LogIoStore, Error, TEXT("%s"), *WriteStatus.ToString());
 				return -1;
@@ -423,8 +585,12 @@ static int32 ChunkPluginCommandEntry(const FContext& Context)
 
 			FMemoryView Chunk = MakeMemoryView(FileData.GetData(), FileData.Num());
 			const FIoHash ChunkHash = FIoHash::HashBuffer(Chunk);
-			if (FIoStatus WriteStatus = WriteChunk(ChunksOutputFolder, Chunk, ChunkHash);
-				WriteStatus.IsOk() == false)
+
+			FIoBuffer Buffer(FIoBuffer::Clone, Chunk);
+
+			FIoStatus WriteStatus = ChunkWriter->WriteChunk(ChunksRelativeFolder, Buffer, ChunkHash);
+
+			if (WriteStatus.IsOk() == false)
 			{
 				UE_LOG(LogIoStore, Error, TEXT("%s"), *WriteStatus.ToString());
 				return -1;
@@ -447,7 +613,9 @@ static int32 ChunkPluginCommandEntry(const FContext& Context)
 		FString Filename = FPathViews::SetExtension(OnDemandTocName, TEXT(".uondemandtoc"));
 		Filename.ToLowerInline();
 
-		const FString TocPath = IoStoreOutputFolder / Filename;
+		// todo, we actually want this file in the base directory
+		// const FString TocPath = OutputFolder / Filename; // like this yo
+		const FString TocPath = OutputFolder / IoStoreRelativeFolder / Filename;
 		if (TUniquePtr<FArchive> Ar(IFileManager::Get().CreateFileWriter(*TocPath)); Ar.IsValid())
 		{
 			*Ar << OnDemandToc;
@@ -516,10 +684,20 @@ static int32 ChunkPluginCommandEntry(const FContext& Context)
 		}
 	}
 
+
+	if (!ChunkWriter->Flush())
+	{
+		UE_LOG(LogIoStore, Error, TEXT("Writer error: Failed to upload chunk(s)"));
+		return 1; 
+	}
+
+
 	return 0;
 }
 
 ////////////////////////////////////////////////////////////////////////////////
+extern FArgumentSet S3Arguments;
+
 static FCommand ChunkPluginCommand(
 	ChunkPluginCommandEntry,
 	TEXT("ChunkPlugin"),
@@ -535,6 +713,9 @@ static FCommand ChunkPluginCommand(
 		TArgument<FStringView>(TEXT("-ErrorOutput"),		TEXT("Error output.")),
 		TArgument<bool>(TEXT("-IncludeSigPak"),				TEXT("Include .sig and .pak file in the uondemandtoc")),
 		TArgument<bool>(TEXT("-KeepContainerFiles"),		TEXT("Should we keep the container files after processing them.")),
+		TArgument<FStringView>(TEXT("-BucketPrefix"),		TEXT("Path to prefix to bucket objects")),
+		TArgument<int32>(TEXT("-MaxConcurrentUploads"),		TEXT("Number of simultaneous uploads")),
+		S3Arguments,
 	}
 );
 
diff --git a/Engine/Source/Runtime/Experimental/IoStore/OnDemand/Private/Tool/Upload.cpp b/Engine/Source/Runtime/Experimental/IoStore/OnDemand/Private/Tool/Upload.cpp
--- a/Engine/Source/Runtime/Experimental/IoStore/OnDemand/Private/Tool/Upload.cpp
+++ b/Engine/Source/Runtime/Experimental/IoStore/OnDemand/Private/Tool/Upload.cpp
@@ -17,11 +17,12 @@
 #include "Misc/KeyChainUtilities.h"
 #include "Misc/Paths.h"
 #include "Misc/PathViews.h"
+#include "S3/S3Client.h"
 #include "Serialization/CompactBinarySerialization.h"
 #include "Serialization/CompactBinaryWriter.h"
 #include "Serialization/MemoryReader.h"
 #include "Serialization/LargeMemoryWriter.h"
-#include "S3/S3Client.h"
+#include "UploadQueue.h"
 
 ////////////////////////////////////////////////////////////////////////////////
 namespace UE::IoStore
@@ -352,139 +353,7 @@ static void WriteConfigFile(
 	}
 }
 
-////////////////////////////////////////////////////////////////////////////////
-class FUploadQueue
-{
-public:
-	FUploadQueue(FS3Client& Client, const FString& Bucket, int32 ThreadCount);
-	bool Enqueue(FStringView Key, FIoBuffer Payload);
-	bool Flush();
-
-private:
-	void ThreadEntry();
-
-	struct FQueueEntry
-	{
-		FString Key;
-		FIoBuffer Payload;
-	};
-
-	FS3Client& Client;
-	TArray<TFuture<void>> Threads;
-	FString Bucket;
-	FCriticalSection CriticalSection;
-	TQueue<FQueueEntry> Queue;
-	FEventRef WakeUpEvent;
-	FEventRef UploadCompleteEvent;
-	std::atomic_int32_t ConcurrentUploads{0};
-	std::atomic_int32_t ActiveThreadCount{0};
-	std::atomic_int32_t ErrorCount{0};
-	std::atomic_bool bCompleteAdding{false};
-};
-
-FUploadQueue::FUploadQueue(FS3Client& InClient, const FString& InBucket, int32 ThreadCount)
-	: Client(InClient)
-	, Bucket(InBucket)
-{
-	ActiveThreadCount = ThreadCount;
-	for (int32 Idx = 0; Idx < ThreadCount; ++Idx)
-	{
-		Threads.Add(AsyncThread([this]()
-		{
-			ThreadEntry();
-		}));
-	}
-}
-
-bool FUploadQueue::Enqueue(FStringView Key, FIoBuffer Payload)
-{
-	if (ActiveThreadCount == 0)
-	{
-		return false;
-	}
-
-	for(;;)
-	{
-		bool bEnqueued = false;
-		{
-			FScopeLock _(&CriticalSection);
-			if (ConcurrentUploads < Threads.Num())
-			{
-				bEnqueued = Queue.Enqueue(FQueueEntry {FString(Key), Payload});
-			}
-		}
-
-		if (bEnqueued)
-		{
-			WakeUpEvent->Trigger();
-			break;
-		}
-
-		UploadCompleteEvent->Wait();
-	}
-
-	return true;
-}
-
-void FUploadQueue::ThreadEntry()
-{
-	for(;;)
-	{
-		FQueueEntry Entry;
-		bool bDequeued = false;
-		{
-			FScopeLock _(&CriticalSection);
-			bDequeued = Queue.Dequeue(Entry);
-		}
 
-		if (!bDequeued)
-		{
-			if (bCompleteAdding)
-			{
-				break;
-			}
-			WakeUpEvent->Wait();
-			continue;
-		}
-
-		ConcurrentUploads++;
-		const FS3PutObjectResponse Response = Client.TryPutObject(FS3PutObjectRequest{Bucket, Entry.Key, Entry.Payload.GetView()});
-		ConcurrentUploads--;
-		UploadCompleteEvent->Trigger();
-
-		if (Response.IsOk())
-		{
-			UE_LOG(LogIas, Log, TEXT("Uploaded chunk '%s/%s/%s'"), *Client.GetConfig().ServiceUrl, *Bucket, *Entry.Key);
-		}
-		else
-		{
-			TStringBuilder<256> ErrorResponse;
-			Response.GetErrorResponse(ErrorResponse);
-
-			UE_LOG(LogIas, Warning, TEXT("Failed to upload chunk '%s/%s/%s' (%s)"), *Client.GetConfig().ServiceUrl, *Bucket, *Entry.Key, ErrorResponse.ToString());
-			ErrorCount++;
-			break;
-		}
-	}
-
-	ActiveThreadCount--;
-}
-
-bool FUploadQueue::Flush()
-{
-	bCompleteAdding = true;
-	for (int32 Idx = 0; Idx < Threads.Num(); ++Idx)
-	{
-		WakeUpEvent->Trigger();
-	}
-
-	for (TFuture<void>& Thread : Threads)
-	{
-		Thread.Wait();
-	}
-
-	return ErrorCount == 0;
-}
 
 ////////////////////////////////////////////////////////////////////////////////
 static TIoStatusOr<FUploadResult> UploadContainerFiles(
diff --git a/Engine/Source/Runtime/Experimental/IoStore/OnDemand/Private/Tool/UploadQueue.cpp b/Engine/Source/Runtime/Experimental/IoStore/OnDemand/Private/Tool/UploadQueue.cpp
new file mode 100644
--- /dev/null
+++ b/Engine/Source/Runtime/Experimental/IoStore/OnDemand/Private/Tool/UploadQueue.cpp
@@ -0,0 +1,126 @@
+// Copyright Epic Games, Inc. All Rights Reserved.
+
+#if defined(UE_WITH_IAS_TOOL)
+
+#include "UploadQueue.h"
+
+#include <Containers/UnrealString.h>
+#include <CoreGlobals.h>
+#include <Misc/OutputDeviceRedirector.h>
+#include <IO/IoStoreOnDemand.h>
+#include "S3/S3Client.h"
+
+
+namespace UE::IoStore::Tool {
+
+
+FUploadQueue::FUploadQueue(UE::FS3Client& InClient, const FString& InBucket, int32 ThreadCount)
+	: Client(InClient)
+	, Bucket(InBucket)
+{
+	ActiveThreadCount = ThreadCount;
+	for (int32 Idx = 0; Idx < ThreadCount; ++Idx)
+	{
+		Threads.Add(AsyncThread([this]()
+			{
+				ThreadEntry();
+			}));
+	}
+}
+
+bool FUploadQueue::Enqueue(FStringView Key, FIoBuffer Payload)
+{
+	if (ActiveThreadCount == 0)
+	{
+		return false;
+	}
+
+	for (;;)
+	{
+		bool bEnqueued = false;
+		{
+			FScopeLock _(&CriticalSection);
+			if (ConcurrentUploads < Threads.Num())
+			{
+				bEnqueued = Queue.Enqueue(FQueueEntry{ FString(Key), Payload });
+			}
+		}
+
+		if (bEnqueued)
+		{
+			WakeUpEvent->Trigger();
+			break;
+		}
+
+		UploadCompleteEvent->Wait();
+	}
+
+	return true;
+}
+
+void FUploadQueue::ThreadEntry()
+{
+	for (;;)
+	{
+		FQueueEntry Entry;
+		bool bDequeued = false;
+		{
+			FScopeLock _(&CriticalSection);
+			bDequeued = Queue.Dequeue(Entry);
+		}
+
+		if (!bDequeued)
+		{
+			if (bCompleteAdding)
+			{
+				break;
+			}
+			WakeUpEvent->Wait();
+			continue;
+		}
+
+		ConcurrentUploads++;
+		const FS3PutObjectResponse Response = Client.TryPutObject(FS3PutObjectRequest{ Bucket, Entry.Key, Entry.Payload.GetView() });
+		ConcurrentUploads--;
+		UploadCompleteEvent->Trigger();
+
+		if (Response.IsOk())
+		{
+			UE_LOG(LogIoStoreOnDemand, Log, TEXT("Uploaded chunk '%s/%s/%s'"), *Client.GetConfig().ServiceUrl, *Bucket, *Entry.Key);
+		}
+		else
+		{
+			TStringBuilder<256> ErrorResponse;
+			Response.GetErrorResponse(ErrorResponse);
+
+			UE_LOG(LogIoStoreOnDemand, Warning, TEXT("Failed to upload chunk '%s/%s/%s' (%s)"), *Client.GetConfig().ServiceUrl, *Bucket, *Entry.Key, ErrorResponse.ToString());
+			ErrorCount++;
+			break;
+		}
+	}
+
+	ActiveThreadCount--;
+}
+
+bool FUploadQueue::Flush()
+{
+	bCompleteAdding = true;
+	for (int32 Idx = 0; Idx < Threads.Num(); ++Idx)
+	{
+		WakeUpEvent->Trigger();
+	}
+
+	for (TFuture<void>& Thread : Threads)
+	{
+		Thread.Wait();
+	}
+
+	return ErrorCount == 0;
+}
+
+} // namespace UE::IoStore::Tool
+
+
+
+
+#endif // UE_WITH_IAS_TOOL
diff --git a/Engine/Source/Runtime/Experimental/IoStore/OnDemand/Private/Tool/UploadQueue.h b/Engine/Source/Runtime/Experimental/IoStore/OnDemand/Private/Tool/UploadQueue.h
new file mode 100644
--- /dev/null
+++ b/Engine/Source/Runtime/Experimental/IoStore/OnDemand/Private/Tool/UploadQueue.h
@@ -0,0 +1,57 @@
+// Copyright Epic Games, Inc. All Rights Reserved.
+
+#pragma once
+
+#if defined(UE_WITH_IAS_TOOL)
+
+
+#include "IO/IoBuffer.h"
+
+#include <Containers/Array.h>
+#include <Containers/StringConv.h>
+#include <Containers/StringView.h>
+#include <String/LexFromString.h>
+
+
+
+namespace UE
+{
+	class FS3Client;
+}
+
+namespace UE::IoStore::Tool {
+
+////////////////////////////////////////////////////////////////////////////////
+// s3 upload helper.... 
+class FUploadQueue
+{
+public:
+	FUploadQueue(UE::FS3Client& Client, const FString& Bucket, int32 ThreadCount);
+	bool Enqueue(FStringView Key, FIoBuffer Payload);
+	bool Flush();
+
+private:
+	void ThreadEntry();
+
+	struct FQueueEntry
+	{
+		FString Key;
+		FIoBuffer Payload;
+	};
+
+	UE::FS3Client& Client;
+	TArray<TFuture<void>> Threads;
+	FString Bucket;
+	FCriticalSection CriticalSection;
+	TQueue<FQueueEntry> Queue;
+	FEventRef WakeUpEvent;
+	FEventRef UploadCompleteEvent;
+	std::atomic_int32_t ConcurrentUploads{ 0 };
+	std::atomic_int32_t ActiveThreadCount{ 0 };
+	std::atomic_int32_t ErrorCount{ 0 };
+	std::atomic_bool bCompleteAdding{ false };
+};
+
+} // namespace UE::IoStore::Tool
+
+#endif // UE_WITH_IAS_TOOL
