Dataflow : Remove serialized render state from the dataflow node.

diff --git a/Engine/Source/Runtime/Experimental/Dataflow/Engine/Private/Dataflow/DataflowObject.cpp b/Engine/Source/Runtime/Experimental/Dataflow/Engine/Private/Dataflow/DataflowObject.cpp
--- a/Engine/Source/Runtime/Experimental/Dataflow/Engine/Private/Dataflow/DataflowObject.cpp
+++ b/Engine/Source/Runtime/Experimental/Dataflow/Engine/Private/Dataflow/DataflowObject.cpp
@@ -87,20 +87,6 @@ void UDataflow::PostLoad()
 		{
 			EdNode->SetEnabledState(ENodeEnabledState::Disabled);
 		}
-		else
-		{
-			if (DataflowEdNode)
-			{
-				if (DataflowEdNode->ShouldRenderNode())
-				{
-					RenderTargets.Add(DataflowEdNode);
-				}
-				if (DataflowEdNode->ShouldWireframeRenderNode())
-				{
-					WireframeRenderTargets.Add(DataflowEdNode);
-				}
-			}
-		}
 	}
 #endif
 
diff --git a/Engine/Source/Runtime/Experimental/Dataflow/Engine/Public/Dataflow/DataflowEdNode.h b/Engine/Source/Runtime/Experimental/Dataflow/Engine/Public/Dataflow/DataflowEdNode.h
--- a/Engine/Source/Runtime/Experimental/Dataflow/Engine/Public/Dataflow/DataflowEdNode.h
+++ b/Engine/Source/Runtime/Experimental/Dataflow/Engine/Public/Dataflow/DataflowEdNode.h
@@ -104,13 +104,13 @@ public:
 
 private:
 
-	UPROPERTY()
+	UPROPERTY(Transient)
 	bool bRenderInAssetEditor = false;
 
-	UPROPERTY()
+	UPROPERTY(Transient)
 	bool bRenderWireframeInAssetEditor = false;
 
-	UPROPERTY()
+	UPROPERTY(Transient)
 	bool bCanEnableRenderWireframe = true;
 
 	void OnPinConnection(const UEdGraphPin& Pin, const FEdGraphPinType& Type);
