Iris
- Fixes to ensure that members for both object and subobjects affected by changing conditionals such as owning connection, physics and role always are marked dirty when condition is re-enabled.

diff --git a/Engine/Source/Runtime/Experimental/Iris/Core/Private/Iris/ReplicationSystem/Conditionals/ReplicationConditionals.cpp b/Engine/Source/Runtime/Experimental/Iris/Core/Private/Iris/ReplicationSystem/Conditionals/ReplicationConditionals.cpp
--- a/Engine/Source/Runtime/Experimental/Iris/Core/Private/Iris/ReplicationSystem/Conditionals/ReplicationConditionals.cpp
+++ b/Engine/Source/Runtime/Experimental/Iris/Core/Private/Iris/ReplicationSystem/Conditionals/ReplicationConditionals.cpp
@@ -13,16 +13,26 @@
 #include "Iris/ReplicationSystem/Filtering/NetObjectGroups.h"
 #include "Iris/ReplicationSystem/Filtering/ReplicationFiltering.h"
 #include "Iris/Serialization/InternalNetSerializers.h"
+#include "Iris/ReplicationSystem/ReplicationWriter.h"
 #include "Containers/ArrayView.h"
 #include "UObject/CoreNetTypes.h"
 #include "Net/Core/NetHandle/NetHandleManager.h"
 #include "Net/Core/PropertyConditions/RepChangedPropertyTracker.h"
 #include "Net/Core/PropertyConditions/PropertyConditions.h"
 #include "Net/Core/Trace/NetDebugName.h"
+#include "HAL/IConsoleManager.h"
+
+DEFINE_LOG_CATEGORY_STATIC(LogIrisConditionals, Log, All);
 
 namespace UE::Net::Private
 {
 
+static bool bEnableUpdateObjectsWithDirtyConditionals = true;
+static FAutoConsoleVariableRef CVarEnableUpdateObjectsWithDirtyConditionals(
+	TEXT("net.Iris.EnableUpdateObjectsWithDirtyConditionals"),
+	bEnableUpdateObjectsWithDirtyConditionals,
+	TEXT("Enable the updating subobjects with conditionals."));
+
 FReplicationConditionals::FReplicationConditionals()
 {
 }
@@ -47,6 +57,7 @@ void FReplicationConditionals::Init(FReplicationConditionalsInitParams& Params)
 
 	PerObjectInfos.SetNumZeroed(MaxInternalNetRefIndex);
 	ConnectionInfos.SetNum(MaxConnectionCount + 1U);
+	ObjectsWithDirtyLifetimeConditionals.Init(Params.MaxInternalNetRefIndex);
 }
 
 void FReplicationConditionals::OnMaxInternalNetRefIndexIncreased(FInternalNetRefIndex NewMaxInternalIndex)
@@ -55,6 +66,8 @@ void FReplicationConditionals::OnMaxInternalNetRefIndexIncreased(FInternalNetRef
 
 	PerObjectInfos.SetNumZeroed(NewMaxInternalIndex);
 
+	ObjectsWithDirtyLifetimeConditionals.SetNumBits(NewMaxInternalIndex);
+
 	for (FPerConnectionInfo& ConnectionInfo : ConnectionInfos)
 	{
 		// Resize the netobject list for valid connections
@@ -82,17 +95,33 @@ bool FReplicationConditionals::SetConditionConnectionFilter(FInternalNetRefIndex
 	FPerObjectInfo* ObjectInfo = GetPerObjectInfo(ObjectIndex);
 	if (ObjectInfo->AutonomousConnectionId != AutonomousConnectionId)
 	{
+		UE_LOG(LogIrisConditionals, Verbose, TEXT("SetConditionConnectionFilter %s. AutonomousConnectionId: %u"), *NetRefHandleManager->PrintObjectFromIndex(ObjectIndex), AutonomousConnectionId);
+
 		const uint32 ConnIdForBaselineInvalidation = (bEnable ? ConnectionId : ObjectInfo->AutonomousConnectionId);
 		ObjectInfo->AutonomousConnectionId = uint16(AutonomousConnectionId);
 
+		// TODO: Make sure we also invalidate baselines for subobjects
 		BaselineInvalidationTracker->InvalidateBaselines(ObjectIndex, ConnIdForBaselineInvalidation);
 
 		MarkRemoteRoleDirty(ObjectIndex);
+		// Mark object as having dirty global conditional that should be evaluated before next send
+		ObjectsWithDirtyLifetimeConditionals.SetBit(ObjectIndex);
 	}
 
 	return true;
 }
 
+void FReplicationConditionals::SetOwningConnection(FInternalNetRefIndex ObjectIndex, uint32 ConnectionId)
+{
+	if (ReplicationFiltering->GetOwningConnection(ObjectIndex) != ConnectionId)
+	{
+		UE_LOG(LogIrisConditionals, Verbose, TEXT("SetOwningConnection on object %u. Connection: %u"), ObjectIndex, ConnectionId);
+
+		// Mark object as having dirty global conditional that should be evaluated before next send
+		ObjectsWithDirtyLifetimeConditionals.SetBit(ObjectIndex);
+	}
+}
+
 void FReplicationConditionals::AddConnection(uint32 ConnectionId)
 {
 	// Init connection info
@@ -120,8 +149,14 @@ bool FReplicationConditionals::SetCondition(FInternalNetRefIndex ObjectIndex, ER
 		FPerObjectInfo* ObjectInfo = GetPerObjectInfo(ObjectIndex);
 		if (bEnable && !ObjectInfo->bRepPhysics)
 		{
+			UE_LOG(LogIrisConditionals, Verbose, TEXT("SetCondition object %s. EReplicationCondition::ReplicatePhysics: %u"), *NetRefHandleManager->PrintObjectFromIndex(ObjectIndex), bEnable ? 1U : 0U);
+
 			// We only care to track this change if the condition is enabled.
+			// TODO: Make sure we also invalidate baselines for subobjects
 			BaselineInvalidationTracker->InvalidateBaselines(ObjectIndex, BaselineInvalidationTracker->InvalidateBaselineForAllConnections);
+
+			// Mark object as having dirty global conditional that should be evaluated before next send
+			ObjectsWithDirtyLifetimeConditionals.SetBit(ObjectIndex);
 		}
 		ObjectInfo->bRepPhysics = bEnable ? 1U : 0U;
 		return true;
@@ -316,6 +351,7 @@ bool FReplicationConditionals::SetPropertyCustomCondition(FInternalNetRefIndex O
 					// If a condition is enabled we also mark the corresponding regular changemask as dirty.
 					FNetBitArrayView MemberChangeMask = GetMemberChangeMask(Fragment.ExternalSrcBuffer, StateDescriptor);
 					FReplicationStateHeader& Header = GetReplicationStateHeader(Fragment.ExternalSrcBuffer, StateDescriptor);
+
 					MarkDirty(Header, MemberChangeMask, ChangeMaskDescriptor);
 
 					// Enabled conditions causes new properties to be replicated which most likely have incorrect values at the receiving end.
@@ -469,6 +505,10 @@ bool FReplicationConditionals::SetPropertyDynamicCondition(FInternalNetRefIndex
 void FReplicationConditionals::Update()
 {
 	UpdateObjectsInScope();
+	if (bEnableUpdateObjectsWithDirtyConditionals)
+	{
+		UpdateAndResetObjectsWithDirtyConditionals();
+	}
 }
 
 void FReplicationConditionals::GetChildSubObjectsToReplicate(uint32 ReplicatingConnectionId, const FConditionalsMask& LifetimeConditionals,  const FInternalNetRefIndex ParentObjectIndex, FSubObjectsToReplicateArray& OutSubObjectsToReplicate)
@@ -491,11 +531,10 @@ void FReplicationConditionals::GetChildSubObjectsToReplicate(uint32 ReplicatingC
 			// Append child subobjects that fulfill the condition
 			for (uint32 ArrayIndex = 0; ArrayIndex < SubObjectsInfo.NumSubObjects; ++ArrayIndex)
 			{
+				const FInternalNetRefIndex SubObjectIndex = SubObjectsInfo.ChildSubObjects[ArrayIndex];
 				const ELifetimeCondition LifeTimeCondition = (ELifetimeCondition)SubObjectsInfo.SubObjectLifeTimeConditions[ArrayIndex];
 				if (LifeTimeCondition == COND_NetGroup)
 				{
-					const FInternalNetRefIndex SubObjectIndex = SubObjectsInfo.ChildSubObjects[ArrayIndex];
-
 					uint32 GroupCount = 0U;
 					if (const FNetObjectGroupHandle* GroupMemberships = NetObjectGroups->GetGroupMemberships(SubObjectIndex, GroupCount))
 					{
@@ -525,14 +564,22 @@ void FReplicationConditionals::GetChildSubObjectsToReplicate(uint32 ReplicatingC
 								break;
 							}
 						}
+
+						if (!bShouldReplicateSubObject)
+						{
+							UE_LOG(LogIrisConditionals, VeryVerbose, TEXT("%s Filtered out by COND_NetGroup"), *NetRefHandleManager->PrintObjectFromIndex(SubObjectIndex));
+						}
 					}
 				}
 				else if (LifetimeConditionals.IsConditionEnabled(LifeTimeCondition))
 				{
-					const FInternalNetRefIndex SubObjectIndex = SubObjectsInfo.ChildSubObjects[ArrayIndex];
 					GetChildSubObjectsToReplicate(ReplicatingConnectionId, LifetimeConditionals, SubObjectIndex, OutSubObjectsToReplicate);
 					OutSubObjectsToReplicate.Add(SubObjectIndex);
 				}
+				else
+				{
+					UE_LOG(LogIrisConditionals, VeryVerbose, TEXT("%s Filtered out by %s"), *NetRefHandleManager->PrintObjectFromIndex(SubObjectIndex), *UEnum::GetValueAsString(LifeTimeCondition));
+				}
 			}
 		}
 	}
@@ -563,6 +610,7 @@ bool FReplicationConditionals::ApplyConditionalsToChangeMask(uint32 ReplicatingC
 	if (EnumHasAnyFlags(Protocol->ProtocolTraits, EReplicationProtocolTraits::HasLifetimeConditionals))
 	{
 		const FConditionalsMask LifetimeConditionals = GetLifetimeConditionals(ReplicatingConnectionId, ParentObjectIndex, bIsInitialState);
+
 		FConditionalsMask PrevLifeTimeConditions = ConnectionInfos[ReplicatingConnectionId].ObjectConditionals[ObjectIndex];
 		if (PrevLifeTimeConditions.IsUninitialized())
 		{
@@ -596,6 +644,8 @@ bool FReplicationConditionals::ApplyConditionalsToChangeMask(uint32 ReplicatingC
 				{
 					if (!PrevLifeTimeConditions.IsConditionEnabled(Condition))
 					{
+						UE_LOG(LogIrisConditionals, Verbose, TEXT("Dirtying member %s %s:%s due to condition %s"), *NetRefHandleManager->PrintObjectFromIndex(ObjectIndex), ToCStr(StateDescriptor->DebugName), ToCStr(StateDescriptor->MemberDebugDescriptors[MemberIt].DebugName), *UEnum::GetValueAsString(Condition));
+
 						const FReplicationStateMemberChangeMaskDescriptor& ChangeMaskDescriptor = ChangeMaskDescriptors[MemberIt];
 						for (uint32 BitIt = ChangeMaskBitOffset + ChangeMaskDescriptor.BitOffset, BitEndIt = BitIt + ChangeMaskDescriptor.BitCount; BitIt != BitEndIt; ++BitIt)
 						{
@@ -607,10 +657,11 @@ bool FReplicationConditionals::ApplyConditionalsToChangeMask(uint32 ReplicatingC
 				else
 				{
 					const FReplicationStateMemberChangeMaskDescriptor& ChangeMaskDescriptor = ChangeMaskDescriptors[MemberIt];
-					for (uint32 BitIt = ChangeMaskBitOffset + ChangeMaskDescriptor.BitOffset, BitEndIt = BitIt + ChangeMaskDescriptor.BitCount; BitIt != BitEndIt; ++BitIt)
+					if (ChangeMask.IsAnyBitSet(ChangeMaskBitOffset + ChangeMaskDescriptor.BitOffset, ChangeMaskDescriptor.BitCount))
 					{
-						bMaskWasModified |= ChangeMask.GetBit(BitIt);
-						ChangeMask.ClearBit(BitIt);
+						UE_LOG(LogIrisConditionals, VeryVerbose, TEXT("Filtering out member %s %s:%s due to condition %s"), *NetRefHandleManager->PrintObjectFromIndex(ObjectIndex), ToCStr(StateDescriptor->DebugName), ToCStr(StateDescriptor->MemberDebugDescriptors[MemberIt].DebugName), *UEnum::GetValueAsString(Condition));
+						ChangeMask.ClearBits(ChangeMaskBitOffset + ChangeMaskDescriptor.BitOffset, ChangeMaskDescriptor.BitCount);
+						bMaskWasModified = true;
 					}
 				}
 			}
@@ -644,6 +695,7 @@ bool FReplicationConditionals::ApplyConditionalsToChangeMask(uint32 ReplicatingC
 						{
 							if (!PrevLifeTimeConditions.IsConditionEnabled(Condition))
 							{
+								UE_LOG(LogIrisConditionals, Verbose, TEXT("Dirtying member %s %s:%s due to condition %s"), *NetRefHandleManager->PrintObjectFromIndex(ObjectIndex), ToCStr(StateDescriptor->DebugName), ToCStr(StateDescriptor->MemberDebugDescriptors[MemberIt].DebugName), *UEnum::GetValueAsString(Condition));
 								const FReplicationStateMemberChangeMaskDescriptor& ChangeMaskDescriptor = ChangeMaskDescriptors[MemberIt];
 								for (uint32 BitIt = CurrentChangeMaskBitOffset + ChangeMaskDescriptor.BitOffset, BitEndIt = BitIt + ChangeMaskDescriptor.BitCount; BitIt != BitEndIt; ++BitIt)
 								{
@@ -655,10 +707,12 @@ bool FReplicationConditionals::ApplyConditionalsToChangeMask(uint32 ReplicatingC
 						else
 						{
 							const FReplicationStateMemberChangeMaskDescriptor& ChangeMaskDescriptor = ChangeMaskDescriptors[MemberIt];
-							for (uint32 BitIt = CurrentChangeMaskBitOffset + ChangeMaskDescriptor.BitOffset, BitEndIt = BitIt + ChangeMaskDescriptor.BitCount; BitIt != BitEndIt; ++BitIt)
+							if (ChangeMask.IsAnyBitSet(CurrentChangeMaskBitOffset + ChangeMaskDescriptor.BitOffset, ChangeMaskDescriptor.BitCount))
 							{
-								bMaskWasModified |= ChangeMask.GetBit(BitIt);
-								ChangeMask.ClearBit(BitIt);
+								UE_LOG(LogIrisConditionals, VeryVerbose, TEXT("Filtering out member %s %s:%s due to condition %s"), *NetRefHandleManager->PrintObjectFromIndex(ObjectIndex), ToCStr(StateDescriptor->DebugName), ToCStr(StateDescriptor->MemberDebugDescriptors[MemberIt].DebugName), *UEnum::GetValueAsString(Condition));
+
+								ChangeMask.ClearBits(CurrentChangeMaskBitOffset + ChangeMaskDescriptor.BitOffset, ChangeMaskDescriptor.BitCount);
+								bMaskWasModified = true;
 							}
 						}
 					}
@@ -755,9 +809,39 @@ void FReplicationConditionals::UpdateObjectsInScope()
 				}
 			}
 		}
+
 	}
 }
 
+void FReplicationConditionals::UpdateAndResetObjectsWithDirtyConditionals()
+{
+	IRIS_PROFILER_SCOPE(FReplicationConditionals_UpdateAndResetObjectsWithDirtyConditionals);
+
+	const FNetBitArray& ValidConnections = ReplicationConnections->GetValidConnections();
+
+	// We do not expect many objects with dirty global lifetime conditionals each frame
+	const uint32 MaxBatchObjectCount = 128U;
+	FInternalNetRefIndex ObjectIndices[MaxBatchObjectCount];
+
+	const uint32 BitCount = ~0U;
+	for (uint32 ObjectCount, StartIndex = 0; (ObjectCount = ObjectsWithDirtyLifetimeConditionals.GetSetBitIndices(StartIndex, BitCount, ObjectIndices, MaxBatchObjectCount)) > 0; )
+	{
+		for (uint32 ConnectionId : ValidConnections)
+		{
+			UE::Net::Private::FReplicationConnection* Connection = ReplicationConnections->GetConnection(ConnectionId);
+			Connection->ReplicationWriter->UpdateDirtyGlobalLifetimeConditionals(MakeArrayView(ObjectIndices, ObjectCount));
+		}
+
+		StartIndex = ObjectIndices[ObjectCount - 1] + 1U;
+		if ((StartIndex == ObjectsWithDirtyLifetimeConditionals.GetNumBits()) | (ObjectCount < MaxBatchObjectCount))
+		{
+			break;
+		}
+	}	
+
+	ObjectsWithDirtyLifetimeConditionals.ClearAllBits();
+}
+
 FReplicationConditionals::FConditionalsMask FReplicationConditionals::GetLifetimeConditionals(uint32 ReplicatingConnectionId, FInternalNetRefIndex ParentObjectIndex, bool bIsInitialState) const
 {
 	FConditionalsMask ConditionalsMask{0};
@@ -784,7 +868,7 @@ FReplicationConditionals::FConditionalsMask FReplicationConditionals::GetLifetim
 	ConditionalsMask.SetConditionEnabled(COND_SimulatedOnlyNoReplay, bRoleSimulated);
 	ConditionalsMask.SetConditionEnabled(COND_SimulatedOrPhysicsNoReplay, bRoleSimulated | bRepPhysics);
 	ConditionalsMask.SetConditionEnabled(COND_SkipReplay, true);
-
+	
 	return ConditionalsMask;
 }
 
diff --git a/Engine/Source/Runtime/Experimental/Iris/Core/Private/Iris/ReplicationSystem/Conditionals/ReplicationConditionals.h b/Engine/Source/Runtime/Experimental/Iris/Core/Private/Iris/ReplicationSystem/Conditionals/ReplicationConditionals.h
--- a/Engine/Source/Runtime/Experimental/Iris/Core/Private/Iris/ReplicationSystem/Conditionals/ReplicationConditionals.h
+++ b/Engine/Source/Runtime/Experimental/Iris/Core/Private/Iris/ReplicationSystem/Conditionals/ReplicationConditionals.h
@@ -30,7 +30,7 @@ struct FReplicationConditionalsInitParams
 {
 	const FNetRefHandleManager* NetRefHandleManager = nullptr;
 	const FReplicationFiltering* ReplicationFiltering = nullptr;
-	const FReplicationConnections* ReplicationConnections = nullptr;
+	FReplicationConnections* ReplicationConnections = nullptr;
 	const FNetObjectGroups* NetObjectGroups = nullptr;
 	FDeltaCompressionBaselineInvalidationTracker* BaselineInvalidationTracker = nullptr;
 	FInternalNetRefIndex MaxInternalNetRefIndex = 0;
@@ -60,6 +60,7 @@ public:
 
 	bool SetConditionConnectionFilter(FInternalNetRefIndex ObjectIndex, EReplicationCondition Condition, uint32 ConnectionId, bool bEnable);
 	bool SetCondition(FInternalNetRefIndex ObjectIndex, EReplicationCondition Condition, bool bEnable);
+	void SetOwningConnection(FInternalNetRefIndex ObjectIndex, uint32 ConnectionId);
 
 	// For property custom conditions only
 	void InitPropertyCustomConditions(FInternalNetRefIndex ObjectIndex);
@@ -71,6 +72,16 @@ public:
 
 	void Update();
 
+	struct FConditionalsMask
+	{
+		bool IsUninitialized() const { return ConditionalsMask == 0; }
+		bool IsConditionEnabled(int Condition) const { checkSlow(Condition < 16); return ConditionalsMask & (uint16(1) << unsigned(Condition)); }
+		bool SetConditionEnabled(int Condition, bool bEnabled) { checkSlow(Condition < 16); return ConditionalsMask |= (uint16(bEnabled ? 1 : 0) << unsigned(Condition)); }
+
+		// Each LifetimeCondition is represented in this member via (1U << ELifetimeCondition)
+		uint16 ConditionalsMask;
+	};
+
 	bool ApplyConditionalsToChangeMask(uint32 ReplicatingConnectionId, bool bIsInitialState, FInternalNetRefIndex ParentObjectIndex, FInternalNetRefIndex ObjectIndex, uint32* ChangeMaskData, const uint32* ConditionalChangeMaskData, const FReplicationProtocol* Protocol);
 
 	using FSubObjectsToReplicateArray = TArray<FInternalNetRefIndex, TInlineAllocator<32>>;
@@ -84,16 +95,6 @@ private:
 		uint16 bRepPhysics : 1;
 	};
 
-	struct FConditionalsMask
-	{
-		bool IsUninitialized() const { return ConditionalsMask == 0; }
-		bool IsConditionEnabled(int Condition) const { return ConditionalsMask & (uint16(1) << unsigned(Condition)); }
-		bool SetConditionEnabled(int Condition, bool bEnabled) { return ConditionalsMask |= (uint16(bEnabled ? 1 : 0) << unsigned(Condition)); }
-
-		// Each LifetimeCondition is represented in this member via (1U << ELifetimeCondition)
-		uint16 ConditionalsMask;
-	};
-
 	struct FPerConnectionInfo
 	{
 		TArray<FConditionalsMask> ObjectConditionals;
@@ -113,6 +114,7 @@ private:
 
 private:
 	void UpdateObjectsInScope();
+	void UpdateAndResetObjectsWithDirtyConditionals();
 
 	FConditionalsMask GetLifetimeConditionals(uint32 ReplicatingConnectionId, FInternalNetRefIndex ParentObjectIndex, bool bInitialState) const;
 
@@ -135,13 +137,14 @@ private:
 
 	const FNetRefHandleManager* NetRefHandleManager = nullptr;
 	const FReplicationFiltering* ReplicationFiltering = nullptr;
-	const FReplicationConnections* ReplicationConnections = nullptr;
+	FReplicationConnections* ReplicationConnections = nullptr;
 	FDeltaCompressionBaselineInvalidationTracker* BaselineInvalidationTracker = nullptr;
 	const FNetObjectGroups* NetObjectGroups = nullptr;
 
 	TArray<FPerObjectInfo> PerObjectInfos;
 	TArray<FPerConnectionInfo> ConnectionInfos;
 	TMap<FInternalNetRefIndex, FObjectDynamicConditions> DynamicConditions;
+	FNetBitArray ObjectsWithDirtyLifetimeConditionals;
 
 	FInternalNetRefIndex MaxInternalNetRefIndex = 0;
 	uint32 MaxConnectionCount = 0;
diff --git a/Engine/Source/Runtime/Experimental/Iris/Core/Private/Iris/ReplicationSystem/DeltaCompression/DeltaCompressionBaselineInvalidationTracker.h b/Engine/Source/Runtime/Experimental/Iris/Core/Private/Iris/ReplicationSystem/DeltaCompression/DeltaCompressionBaselineInvalidationTracker.h
--- a/Engine/Source/Runtime/Experimental/Iris/Core/Private/Iris/ReplicationSystem/DeltaCompression/DeltaCompressionBaselineInvalidationTracker.h
+++ b/Engine/Source/Runtime/Experimental/Iris/Core/Private/Iris/ReplicationSystem/DeltaCompression/DeltaCompressionBaselineInvalidationTracker.h
@@ -44,6 +44,7 @@ public:
 
 	void Init(FDeltaCompressionBaselineInvalidationTrackerInitParams& InitParams);
 
+	// It is up to calling code to also do this for eventual subobjects.
 	void InvalidateBaselines(FInternalNetRefIndex ObjectIndex, uint32 ConnId);
 
 	// Returns an array of objects with enabled conditions such that DC baselines need to be invalidated.
diff --git a/Engine/Source/Runtime/Experimental/Iris/Core/Private/Iris/ReplicationSystem/ReplicationSystem.cpp b/Engine/Source/Runtime/Experimental/Iris/Core/Private/Iris/ReplicationSystem/ReplicationSystem.cpp
--- a/Engine/Source/Runtime/Experimental/Iris/Core/Private/Iris/ReplicationSystem/ReplicationSystem.cpp
+++ b/Engine/Source/Runtime/Experimental/Iris/Core/Private/Iris/ReplicationSystem/ReplicationSystem.cpp
@@ -1125,6 +1125,9 @@ void UReplicationSystem::SetOwningNetConnection(FNetRefHandle Handle, uint32 Con
 		return;
 	}
 
+	FReplicationConditionals& Conditionals = Impl->ReplicationSystemInternal.GetConditionals();
+	Conditionals.SetOwningConnection(ObjectInternalIndex, ConnectionId);
+
 	FReplicationFiltering& Filtering = Impl->ReplicationSystemInternal.GetFiltering();
 	Filtering.SetOwningConnection(ObjectInternalIndex, ConnectionId);
 }
diff --git a/Engine/Source/Runtime/Experimental/Iris/Core/Private/Iris/ReplicationSystem/ReplicationWriter.cpp b/Engine/Source/Runtime/Experimental/Iris/Core/Private/Iris/ReplicationSystem/ReplicationWriter.cpp
--- a/Engine/Source/Runtime/Experimental/Iris/Core/Private/Iris/ReplicationSystem/ReplicationWriter.cpp
+++ b/Engine/Source/Runtime/Experimental/Iris/Core/Private/Iris/ReplicationSystem/ReplicationWriter.cpp
@@ -847,6 +847,19 @@ void FReplicationWriter::UpdateScope(const FNetBitArrayView& UpdatedScope)
 	ObjectsInScope.Combine(ObjectsPendingDestroy, FNetBitArrayBase::AndNotOp);
 }
 
+void FReplicationWriter::UpdateDirtyGlobalLifetimeConditionals(TArrayView<FInternalNetRefIndex> ObjectsWithDirtyConditionals)
+{
+	for (FInternalNetRefIndex InternalObjectIndex : ObjectsWithDirtyConditionals)
+	{
+		FReplicationInfo& Info = ReplicatedObjects[InternalObjectIndex];
+		if (Info.GetState() != EReplicatedObjectState::Invalid)
+		{
+			// We just set a flag that we will process if the object is going to be replicated
+			Info.HasDirtyConditionals = 1;
+		}
+	}
+}
+
 void FReplicationWriter::InternalUpdateDirtyChangeMasks(const FChangeMaskCache& CachedChangeMasks, EFlushFlags ExtraFlushFlags, bool bMarkForTearOff)
 {
 	//IRIS_PROFILER_SCOPE(FReplicationWriter_UpdateDirtyChangeMasks);
@@ -2054,8 +2067,8 @@ FReplicationWriter::EWriteObjectStatus FReplicationWriter::WriteObjectAndSubObje
 		return EWriteObjectStatus::NoInstanceProtocol;
 	}
 
-	// Filter out changemasks that are not supposed to be replicated to this connection
-	const bool bNeedToFilterChangeMask = (bIsInitialState || Info.HasDirtyChangeMask) && Info.HasChangemaskFilter;
+	// Objects affected by conditionals might need to modify the changemask
+	const bool bNeedToFilterChangeMask = (bIsInitialState || Info.HasDirtyChangeMask || Info.HasDirtyConditionals) && Info.HasChangemaskFilter;
 	if (bNeedToFilterChangeMask)
 	{
 		ApplyFilterToChangeMask(OutBatchInfo.ParentInternalIndex, InternalIndex, Info, ObjectData.Protocol, ReplicatedObjectStateBuffer, bIsInitialState);
@@ -2067,6 +2080,23 @@ FReplicationWriter::EWriteObjectStatus FReplicationWriter::WriteObjectAndSubObje
 #endif
 	}
 
+	// Even if root is not affected by dirty conditionals one of our subobjects might be.
+	if (!Info.IsSubObject && Info.HasDirtyConditionals)
+	{
+		for (uint32 SubObjectIndex : NetRefHandleManager->GetSubObjects(InternalIndex))
+		{
+			FReplicationInfo& SubObjectInfo = GetReplicationInfo(SubObjectIndex);
+			const EReplicatedObjectState SubObjectState = SubObjectInfo.GetState();
+			if (SubObjectInfo.HasChangemaskFilter && (SubObjectState > EReplicatedObjectState::Invalid && SubObjectState < EReplicatedObjectState::PendingDestroy))
+			{
+				MarkObjectDirty(SubObjectIndex, "UpdateGlobalConditional");
+				SubObjectInfo.HasDirtyChangeMask = 1;
+				Info.HasDirtySubObjects = 1;
+			}
+		}
+	}
+	Info.HasDirtyConditionals = 0;
+
 	const bool bIsObjectIndexForAttachment = IsObjectIndexForOOBAttachment(InternalIndex);
 	const bool bHasState = (bIsInitialState || Info.HasDirtyChangeMask) && !!(WriteObjectFlags & EWriteObjectFlag::WriteObjectFlag_State);
 	const bool bHasAttachments = (Info.HasAttachments || bIsObjectIndexForAttachment);
@@ -3449,9 +3479,9 @@ void FReplicationWriter::ApplyFilterToChangeMask(uint32 ParentInternalIndex, uin
 {
 	const uint32* ConditionalChangeMaskPointer = (EnumHasAnyFlags(Protocol->ProtocolTraits, EReplicationProtocolTraits::HasConditionalChangeMask) ? reinterpret_cast<const uint32*>(InternalStateBuffer + Protocol->GetConditionalChangeMaskOffset()) : static_cast<const uint32*>(nullptr));
 	const bool bChangeMaskWasModified = ReplicationConditionals->ApplyConditionalsToChangeMask(Parameters.ConnectionId, bIsInitialState, ParentInternalIndex, InternalIndex, Info.GetChangeMaskStoragePointer(), ConditionalChangeMaskPointer, Protocol);
-	if (bChangeMaskWasModified && !MakeNetBitArrayView(Info.GetChangeMaskStoragePointer(), Info.ChangeMaskBitCount).IsAnyBitSet())
+	if (bChangeMaskWasModified)
 	{
-		Info.HasDirtyChangeMask = 0;
+		Info.HasDirtyChangeMask = MakeNetBitArrayView(Info.GetChangeMaskStoragePointer(), Info.ChangeMaskBitCount).IsAnyBitSet() ? 1 : 0;
 	}
 }
 
diff --git a/Engine/Source/Runtime/Experimental/Iris/Core/Private/Iris/ReplicationSystem/ReplicationWriter.h b/Engine/Source/Runtime/Experimental/Iris/Core/Private/Iris/ReplicationSystem/ReplicationWriter.h
--- a/Engine/Source/Runtime/Experimental/Iris/Core/Private/Iris/ReplicationSystem/ReplicationWriter.h
+++ b/Engine/Source/Runtime/Experimental/Iris/Core/Private/Iris/ReplicationSystem/ReplicationWriter.h
@@ -130,6 +130,7 @@ public:
 				uint64 LastAckedBaselineIndex : 2;						// Last acknowledged baseline index which we can use for deltacompresion
 				uint64 PendingBaselineIndex : 2;						// Baseline index pending acknowledgment from client
 				uint64 FlushFlags : 3;									// Flags indicating what we are waiting for when flushing
+				uint64 HasDirtyConditionals : 1;						// If this flag is set, we must update conditionals.
 			};
 		};
 
@@ -161,6 +162,9 @@ public:
 	// Called if an object first being teared-off/flushed and then explicitly destroyed before it has been removed from scope
 	void NotifyDestroyedObjectPendingEndReplication(FInternalNetRefIndex ObjectInternalIndex);
 
+	// Called to propagate changes to global lifetime conditionals
+	void UpdateDirtyGlobalLifetimeConditionals(TArrayView<FInternalNetRefIndex> ObjectsWithDirtyConditionals);
+
 	// Propagate dirty changemasks
 	void UpdateDirtyChangeMasks(const FChangeMaskCache& CachedChangeMasks) { InternalUpdateDirtyChangeMasks(CachedChangeMasks, EFlushFlags::FlushFlags_None, false); }
 
