Dataflow : Add back context caching to the creation pattern.

diff --git a/Engine/Plugins/Experimental/ChaosFlesh/Source/ChaosFleshEngine/Private/ChaosFlesh/FleshAsset.cpp b/Engine/Plugins/Experimental/ChaosFlesh/Source/ChaosFleshEngine/Private/ChaosFlesh/FleshAsset.cpp
--- a/Engine/Plugins/Experimental/ChaosFlesh/Source/ChaosFleshEngine/Private/ChaosFlesh/FleshAsset.cpp
+++ b/Engine/Plugins/Experimental/ChaosFlesh/Source/ChaosFleshEngine/Private/ChaosFlesh/FleshAsset.cpp
@@ -6,6 +6,7 @@
 #include "ChaosFlesh/FleshAsset.h"
 #include "ChaosFlesh/FleshCollection.h"
 #include "Components/SkeletalMeshComponent.h"
+#include "Dataflow/DataflowContent.h"
 #include "Engine/SkeletalMesh.h"
 #include "GeometryCollection/TransformCollection.h"
 
@@ -74,7 +75,7 @@ void UFleshAsset::Serialize(FArchive& Ar)
 
 TObjectPtr<UDataflowBaseContent> UFleshAsset::CreateDataflowContent()
 {
-	TObjectPtr<UDataflowSkeletalContent> SkeletalContent = NewObject<UDataflowSkeletalContent>();
+	TObjectPtr<UDataflowSkeletalContent> SkeletalContent = DataflowContextHelpers::CreateNewDataflowContent<UDataflowSkeletalContent>(this);
 
 	SkeletalContent->SetDataflowOwner(this);
 	SkeletalContent->SetTerminalAsset(this);
diff --git a/Engine/Plugins/Experimental/Dataflow/Source/DataflowEditor/Private/Dataflow/AssetDefinition_DataflowContext.cpp b/Engine/Plugins/Experimental/Dataflow/Source/DataflowEditor/Private/Dataflow/AssetDefinition_DataflowContext.cpp
--- a/Engine/Plugins/Experimental/Dataflow/Source/DataflowEditor/Private/Dataflow/AssetDefinition_DataflowContext.cpp
+++ b/Engine/Plugins/Experimental/Dataflow/Source/DataflowEditor/Private/Dataflow/AssetDefinition_DataflowContext.cpp
@@ -18,235 +18,6 @@
 
 #define LOCTEXT_NAMESPACE "AssetActions_DataflowContext"
 
-bool bDataflowEnableContextCaching = true;
-FAutoConsoleVariableRef CVARDataflowEnableContextCaching(TEXT("p.Dataflow.Editor.ContextCaching"), bDataflowEnableContextCaching, 
-	TEXT("Allow the Dataflow editor to crate and use a pre-evaluated graph when the dataflow editor is re-opened.[def:true]"));
-
-namespace DataflowContextDefinitionHelpers
-{
-
-	/*
-	* BindContextToGraph
-	* 
-	* On load the Property in the cache will be a nullptr
-	* This block will rebind the property in the cache to the 
-	* the Property on the UDataflow asset. If there is ever an
-	* mis-match, just return false to indicate that the binding
-	* has failed.
-	*/
-	bool BindContextToGraph(TObjectPtr<UObject>& Asset, UDataflow* DataflowAsset)
-	{
-		using namespace Dataflow;
-
-		TSharedPtr<FGraph> Dataflow = DataflowAsset->GetDataflow();
-		if (!Dataflow) return false;
-
-		UDataflowBaseContent* BaseContent = Cast< UDataflowBaseContent>(Asset.Get());
-		if (!BaseContent) return false;
-
-		TSharedPtr<FEngineContext>& Context = BaseContent->GetDataflowContext();
-		if (!Context) return false;
-
-		TSet<FContextCacheKey> Keys; Context->GetKeys(Keys);
-		for (FContextCacheKey Key : Keys)
-		{
-			bool bValidKey = false;
-
-			if (TUniquePtr<FContextCacheElementBase>* Data = Context->GetBaseData(Key))
-			{
-				if ((*Data) && !(*Data)->GetProperty())
-				{
-					if (TSharedPtr<FDataflowNode> Node = Dataflow->FindBaseNode((*Data)->GetNodeGuid()))
-					{
-						if (FDataflowOutput* Output = Node->FindOutput(Key))
-						{
-							if (const FProperty* Property = Output->GetProperty())
-							{
-								(*Data)->SetProperty(Property);
-								bValidKey = true;
-							}
-						}
-					}
-				}
-				else
-				{
-					bValidKey = true;
-				}
-			}
-
-			if (!bValidKey)
-			{
-				return false;
-			}
-		}
-		return true;
-	}
-
-	/*
-	* ValidateCachedNodeHash
-	* 
-	* Check that the hashes stored in the cache reflect the hash of the nodes 
-	* properties.
-	*
-	*/
-	bool ValidateCachedNodeHash(TObjectPtr<UObject>& Asset, UDataflow* DataflowAsset)
-	{
-		using namespace Dataflow;
-
-		TSharedPtr<FGraph> Dataflow = DataflowAsset->GetDataflow();
-		if (!Dataflow) return false;
-
-		UDataflowBaseContent* BaseContent = Cast< UDataflowBaseContent>(Asset.Get());
-		if (!BaseContent) return false;
-
-		TSharedPtr<FEngineContext>& Context = BaseContent->GetDataflowContext();
-		if (!Context) return false;
-
-		TSet<FContextCacheKey> Keys; Context->GetKeys(Keys);
-		for (FContextCacheKey Key : Keys)
-		{
-			bool bValidKey = false;
-			if (TUniquePtr<FContextCacheElementBase>* Data = Context->GetBaseData(Key))
-			{
-				if (!(*Data)) return false;
-				if (!(*Data)->GetProperty()) return false;
-				if(TSharedPtr<FDataflowNode> Node = Dataflow->FindBaseNode((*Data)->GetNodeGuid()))
-				{
-					if ((*Data)->GetNodeHash() != Node->GetValueHash())
-					{
-						return false;
-					}
-					bValidKey = true;
-				}
-			}
-			if (!bValidKey)
-			{
-				return false;
-			}
-		}
-
-
-		return true;
-	}
-
-	bool ResetCacheTimestamp(TObjectPtr<UObject>& Asset, UDataflow* DataflowAsset)
-	{
-		using namespace Dataflow;
-
-		UDataflowBaseContent* BaseContent = Cast< UDataflowBaseContent>(Asset.Get());
-		if (!BaseContent) return false;
-		if (!BaseContent->GetDataflowContext()) return false;
-
-		Dataflow::FTimestamp NewTimestamp = DataflowAsset->GetRenderingTimestamp().Value + 1;
-		BaseContent->SetLastModifiedTimestamp(NewTimestamp, false /*bMakeDirty*/);
-
-		TSharedPtr<FEngineContext>& Context = BaseContent->GetDataflowContext();
-		TSet<FContextCacheKey> Keys; Context->GetKeys(Keys);
-		for (FContextCacheKey Key : Keys)
-		{
-			if (TUniquePtr<FContextCacheElementBase>* Data = Context->GetBaseData(Key))
-			{
-				if (*Data)
-				{
-					(*Data)->SetTimestamp(NewTimestamp);
-				}
-				else
-				{
-					return false;
-				}
-			}
-			else
-			{
-				return false;
-			}
-		}
-		Context->Timestamp = NewTimestamp;
-
-		return true;
-	}
-
-	/*
-	* CreateNewDataflowContent
-	*
-	*/	
-	template<class T>
-	TObjectPtr<T> CreateNewDataflowContent(const TObjectPtr<UObject>& ContentOwner)
-	{
-		check(ContentOwner.Get());
-		bool bNeedsNewAsset = true;
-		TObjectPtr<UObject> Asset = nullptr;
-
-		UDataflow* DataflowAsset = UE::Dataflow::Private::GetDataflowAssetFrom(ContentOwner);
-		if (bDataflowEnableContextCaching)
-		{
-			// Setup an asset that lives in the content broswer.
-			FContentBrowserModule& ContentBrowserModule = FModuleManager::LoadModuleChecked<FContentBrowserModule>("ContentBrowser");
-
-			UClass* DataflowClass = T::StaticClass();
-
-			const FString AssetPackageName = ContentOwner->GetOutermost()->GetName();
-			const FString AssetDefaultPath = FPackageName::GetLongPackagePath(AssetPackageName);
-
-			FString PackageName = FString::Printf(TEXT("%s/Cache/Dataflow/DataflowContext_%s"), *AssetDefaultPath, *ContentOwner.GetName());
-			if (DataflowAsset)
-			{
-				PackageName = FString::Printf(TEXT("%s_%s"), *PackageName, *DataflowAsset->GetName());
-			}
-
-			UPackage* Package = FindPackage(nullptr, *PackageName);
-			if (!Package)
-			{
-				Package = CreatePackage(*PackageName);
-			}
-
-
-			Asset = StaticLoadObject(DataflowClass, Package, *PackageName);
-			if (Asset && DataflowAsset)
-			{
-				// Validate the loaded cache
-				bNeedsNewAsset =
-					!BindContextToGraph(Asset, DataflowAsset) ||
-					!ValidateCachedNodeHash(Asset, DataflowAsset) ||
-					!ResetCacheTimestamp(Asset, DataflowAsset);
-			}
-
-			if (!Asset || bNeedsNewAsset)
-			{
-				const FName AssetName(FPackageName::GetLongPackageAssetName(PackageName));
-				Asset = NewObject<UObject>(Package, DataflowClass, AssetName, RF_Public | RF_Standalone | RF_Transactional);
-
-				Asset->MarkPackageDirty();
-				FAssetRegistryModule::AssetCreated(Asset);
-
-				if (UDataflowBaseContent* BaseContent = Cast< UDataflowBaseContent>(Asset.Get()))
-				{
-					BaseContent->SetDataflowOwner(ContentOwner);
-					BaseContent->SetDataflowAsset(DataflowAsset);
-				}
-			}
-		}
-		else
-		{
-			Asset = NewObject<T>(ContentOwner, T::StaticClass());
-			if (UDataflowBaseContent* BaseContent = Cast< UDataflowBaseContent>(Asset.Get()))
-			{
-				BaseContent->SetDataflowOwner(ContentOwner);
-				BaseContent->SetDataflowAsset(DataflowAsset);
-			}
-		}
-
-
-		if (UDataflowBaseContent* BaseContent = Cast< UDataflowBaseContent>(Asset.Get()))
-		{
-			BaseContent->SetDataflowOwner(ContentOwner);
-		}
-		return Cast<T>(Asset.Get());
-	}
-
-	template TObjectPtr<UDataflowBaseContent> CreateNewDataflowContent(const TObjectPtr<UObject>& ContentOwner);
-	template TObjectPtr<UDataflowSkeletalContent> CreateNewDataflowContent(const TObjectPtr<UObject>& ContentOwner);
-}
-
 
 namespace UE::Dataflow::DataflowContext
 {
diff --git a/Engine/Plugins/Experimental/Dataflow/Source/DataflowEditor/Private/Dataflow/DataflowEditor.cpp b/Engine/Plugins/Experimental/Dataflow/Source/DataflowEditor/Private/Dataflow/DataflowEditor.cpp
--- a/Engine/Plugins/Experimental/Dataflow/Source/DataflowEditor/Private/Dataflow/DataflowEditor.cpp
+++ b/Engine/Plugins/Experimental/Dataflow/Source/DataflowEditor/Private/Dataflow/DataflowEditor.cpp
@@ -2,12 +2,15 @@
 
 #include "Dataflow/DataflowEditor.h"
 
+#include "AssetRegistry/AssetRegistryModule.h"
 #include "Animation/Skeleton.h"
+#include "ContentBrowserModule.h"
 #include "Dataflow/AssetDefinition_DataflowContext.h"
 #include "Dataflow/DataflowContent.h"
 #include "Dataflow/DataflowEditorToolkit.h"
 #include "Dataflow/DataflowEditorUtil.h"
 #include "Engine/SkeletalMesh.h"
+#include "Modules/ModuleManager.h"
 
 #include UE_INLINE_GENERATED_CPP_BY_NAME(DataflowEditor)
 
@@ -33,7 +36,7 @@ void UDataflowEditor::Initialize(const TArray<TObjectPtr<UObject>>& InObjects)
 		{
 			if(UDataflow* DataflowAsset = Cast<UDataflow>(ContentOwner))
 			{
-				EditorContent = DataflowContextDefinitionHelpers::CreateNewDataflowContent<UDataflowBaseContent>(ContentOwner);
+				EditorContent = DataflowContextHelpers::CreateNewDataflowContent<UDataflowBaseContent>(ContentOwner);
 				EditorContent->SetDataflowOwner(DataflowAsset);
 				EditorContent->SetDataflowAsset(DataflowAsset);
 			}
@@ -42,7 +45,17 @@ void UDataflowEditor::Initialize(const TArray<TObjectPtr<UObject>>& InObjects)
 				if(IDataflowContentOwner* EditorContentOwner = Cast<IDataflowContentOwner>(ContentOwner))
 				{
 					EditorContent = EditorContentOwner->BuildDataflowContent();
-					RequiredObjects.Add(EditorContent->GetDataflowAsset());
+					if (EditorContent)
+					{
+						if (EditorContent->IsSaved())
+						{
+							// Setup an asset that lives in the content broswer.
+							FModuleManager::LoadModuleChecked<FContentBrowserModule>("ContentBrowser");
+							FAssetRegistryModule::AssetCreated(EditorContent);
+						}
+
+						RequiredObjects.Add(EditorContent->GetDataflowAsset());
+					}
 				}
 			}
 		}
diff --git a/Engine/Plugins/Experimental/Dataflow/Source/DataflowEditor/Public/Dataflow/AssetDefinition_DataflowContext.h b/Engine/Plugins/Experimental/Dataflow/Source/DataflowEditor/Public/Dataflow/AssetDefinition_DataflowContext.h
--- a/Engine/Plugins/Experimental/Dataflow/Source/DataflowEditor/Public/Dataflow/AssetDefinition_DataflowContext.h
+++ b/Engine/Plugins/Experimental/Dataflow/Source/DataflowEditor/Public/Dataflow/AssetDefinition_DataflowContext.h
@@ -6,15 +6,6 @@
 
 #include "AssetDefinition_DataflowContext.generated.h"
 
-class UDataflowBaseContent;
-
-namespace DataflowContextDefinitionHelpers
-{
-	// Return true if we should proceed, false if we should re-open the dialog
-	template<class T>
-	DATAFLOWEDITOR_API TObjectPtr<T> CreateNewDataflowContent(const TObjectPtr<UObject>& ContentOwner);
-}
-
 
 
 UCLASS()
diff --git a/Engine/Source/Runtime/Experimental/Dataflow/Engine/Private/Dataflow/DataflowContent.cpp b/Engine/Source/Runtime/Experimental/Dataflow/Engine/Private/Dataflow/DataflowContent.cpp
--- a/Engine/Source/Runtime/Experimental/Dataflow/Engine/Private/Dataflow/DataflowContent.cpp
+++ b/Engine/Source/Runtime/Experimental/Dataflow/Engine/Private/Dataflow/DataflowContent.cpp
@@ -9,10 +9,251 @@
 #include "Dataflow/DataflowObject.h"
 #include "Engine/SkeletalMesh.h"
 #include "GameFramework/Actor.h"
+#include "Misc/PackageName.h"
 #include "PreviewScene.h"
 
 #include UE_INLINE_GENERATED_CPP_BY_NAME(DataflowContent)
 
+bool bDataflowEnableContextCaching = true;
+FAutoConsoleVariableRef CVARDataflowEnableContextCaching(TEXT("p.Dataflow.Editor.ContextCaching"), bDataflowEnableContextCaching,
+	TEXT("Allow the Dataflow editor to crate and use a pre-evaluated graph when the dataflow editor is re-opened.[def:true]"));
+
+namespace DataflowContextHelpers
+{
+
+	/*
+	* BindContextToGraph
+	*
+	* On load the Property in the cache will be a nullptr
+	* This block will rebind the property in the cache to the
+	* the Property on the UDataflow asset. If there is ever an
+	* mis-match, just return false to indicate that the binding
+	* has failed.
+	*/
+	bool BindContextToGraph(TObjectPtr<UObject>& Asset, UDataflow* DataflowAsset)
+	{
+		using namespace Dataflow;
+
+		TSharedPtr<FGraph> Dataflow = DataflowAsset->GetDataflow();
+		if (!Dataflow) return false;
+
+		UDataflowBaseContent* BaseContent = Cast< UDataflowBaseContent>(Asset.Get());
+		if (!BaseContent) return false;
+
+		TSharedPtr<FEngineContext>& Context = BaseContent->GetDataflowContext();
+		if (!Context) return false;
+
+		TSet<FContextCacheKey> Keys; Context->GetKeys(Keys);
+		for (FContextCacheKey Key : Keys)
+		{
+			bool bValidKey = false;
+
+			if (TUniquePtr<FContextCacheElementBase>* Data = Context->GetBaseData(Key))
+			{
+				if ((*Data) && !(*Data)->GetProperty())
+				{
+					if (TSharedPtr<FDataflowNode> Node = Dataflow->FindBaseNode((*Data)->GetNodeGuid()))
+					{
+						if (FDataflowOutput* Output = Node->FindOutput(Key))
+						{
+							if (const FProperty* Property = Output->GetProperty())
+							{
+								(*Data)->SetProperty(Property);
+								bValidKey = true;
+							}
+						}
+					}
+				}
+				else
+				{
+					bValidKey = true;
+				}
+			}
+
+			if (!bValidKey)
+			{
+				return false;
+			}
+		}
+		return true;
+	}
+
+	/*
+	* ValidateCachedNodeHash
+	*
+	* Check that the hashes stored in the cache reflect the hash of the nodes
+	* properties.
+	*
+	*/
+	bool ValidateCachedNodeHash(TObjectPtr<UObject>& Asset, UDataflow* DataflowAsset)
+	{
+		using namespace Dataflow;
+
+		TSharedPtr<FGraph> Dataflow = DataflowAsset->GetDataflow();
+		if (!Dataflow) return false;
+
+		UDataflowBaseContent* BaseContent = Cast< UDataflowBaseContent>(Asset.Get());
+		if (!BaseContent) return false;
+
+		TSharedPtr<FEngineContext>& Context = BaseContent->GetDataflowContext();
+		if (!Context) return false;
+
+		TSet<FContextCacheKey> Keys; Context->GetKeys(Keys);
+		for (FContextCacheKey Key : Keys)
+		{
+			bool bValidKey = false;
+			if (TUniquePtr<FContextCacheElementBase>* Data = Context->GetBaseData(Key))
+			{
+				if (!(*Data)) return false;
+				if (!(*Data)->GetProperty()) return false;
+				if (TSharedPtr<FDataflowNode> Node = Dataflow->FindBaseNode((*Data)->GetNodeGuid()))
+				{
+					if ((*Data)->GetNodeHash() != Node->GetValueHash())
+					{
+						return false;
+					}
+					bValidKey = true;
+				}
+			}
+			if (!bValidKey)
+			{
+				return false;
+			}
+		}
+
+
+		return true;
+	}
+
+	bool ResetCacheTimestamp(TObjectPtr<UObject>& Asset, UDataflow* DataflowAsset)
+	{
+		using namespace Dataflow;
+
+		UDataflowBaseContent* BaseContent = Cast< UDataflowBaseContent>(Asset.Get());
+		if (!BaseContent) return false;
+		if (!BaseContent->GetDataflowContext()) return false;
+
+		Dataflow::FTimestamp NewTimestamp = DataflowAsset->GetRenderingTimestamp().Value + 1;
+		BaseContent->SetLastModifiedTimestamp(NewTimestamp, false /*bMakeDirty*/);
+
+		TSharedPtr<FEngineContext>& Context = BaseContent->GetDataflowContext();
+		TSet<FContextCacheKey> Keys; Context->GetKeys(Keys);
+		for (FContextCacheKey Key : Keys)
+		{
+			if (TUniquePtr<FContextCacheElementBase>* Data = Context->GetBaseData(Key))
+			{
+				if (*Data)
+				{
+					(*Data)->SetTimestamp(NewTimestamp);
+				}
+				else
+				{
+					return false;
+				}
+			}
+			else
+			{
+				return false;
+			}
+		}
+		Context->Timestamp = NewTimestamp;
+
+		return true;
+	}
+
+	UDataflow* GetDataflowAssetFrom(UObject* InObject)
+	{
+		if (UClass* Class = InObject->GetClass())
+		{
+			if (FProperty* Property = Class->FindPropertyByName(FName("DataflowAsset")))
+			{
+				return *Property->ContainerPtrToValuePtr<UDataflow*>(InObject);
+			}
+		}
+		return nullptr;
+	}
+
+	/*
+	* CreateNewDataflowContent
+	*
+	*/
+	template<class T>
+	TObjectPtr<T> CreateNewDataflowContent(const TObjectPtr<UObject>& ContentOwner)
+	{
+		check(ContentOwner.Get());
+		bool bNeedsNewAsset = true;
+		TObjectPtr<UObject> Asset = nullptr;
+
+		UDataflow* DataflowAsset = GetDataflowAssetFrom(ContentOwner);
+		if (bDataflowEnableContextCaching)
+		{
+			UClass* DataflowClass = T::StaticClass();
+
+			const FString AssetPackageName = ContentOwner->GetOutermost()->GetName();
+			const FString AssetDefaultPath = FPackageName::GetLongPackagePath(AssetPackageName);
+
+			FString PackageName = FString::Printf(TEXT("%s/Cache/Dataflow/DataflowContext_%s"), *AssetDefaultPath, *ContentOwner.GetName());
+			if (DataflowAsset)
+			{
+				PackageName = FString::Printf(TEXT("%s_%s"), *PackageName, *DataflowAsset->GetName());
+			}
+
+			UPackage* Package = FindPackage(nullptr, *PackageName);
+			if (!Package)
+			{
+				Package = CreatePackage(*PackageName);
+			}
+
+
+			Asset = StaticLoadObject(DataflowClass, Package, *PackageName);
+			if (Asset && DataflowAsset)
+			{
+				// Validate the loaded cache
+				bNeedsNewAsset =
+					!BindContextToGraph(Asset, DataflowAsset) ||
+					!ValidateCachedNodeHash(Asset, DataflowAsset) ||
+					!ResetCacheTimestamp(Asset, DataflowAsset);
+			}
+
+			if (!Asset || bNeedsNewAsset)
+			{
+				const FName AssetName(FPackageName::GetLongPackageAssetName(PackageName));
+				Asset = NewObject<UObject>(Package, DataflowClass, AssetName, RF_Public | RF_Standalone | RF_Transactional);
+
+				Asset->MarkPackageDirty();
+
+				if (UDataflowBaseContent* BaseContent = Cast< UDataflowBaseContent>(Asset.Get()))
+				{
+					BaseContent->SetIsSaved(true);
+					BaseContent->SetDataflowOwner(ContentOwner);
+					BaseContent->SetDataflowAsset(DataflowAsset);
+				}
+			}
+		}
+		else
+		{
+			Asset = NewObject<T>(ContentOwner, T::StaticClass());
+			if (UDataflowBaseContent* BaseContent = Cast< UDataflowBaseContent>(Asset.Get()))
+			{
+				BaseContent->SetIsSaved(false);
+				BaseContent->SetDataflowOwner(ContentOwner);
+				BaseContent->SetDataflowAsset(DataflowAsset);
+			}
+		}
+
+
+		if (UDataflowBaseContent* BaseContent = Cast< UDataflowBaseContent>(Asset.Get()))
+		{
+			BaseContent->SetDataflowOwner(ContentOwner);
+		}
+		return Cast<T>(Asset.Get());
+	}
+
+	template DATAFLOWENGINE_API TObjectPtr<UDataflowBaseContent> CreateNewDataflowContent(const TObjectPtr<UObject>& ContentOwner);
+	template DATAFLOWENGINE_API TObjectPtr<UDataflowSkeletalContent> CreateNewDataflowContent(const TObjectPtr<UObject>& ContentOwner);
+}
+
+
 void UDataflowContextObject::AddReferencedObjects(UObject* InThis, FReferenceCollector& Collector)
 {
 	UDataflowContextObject* This = CastChecked<UDataflowContextObject>(InThis);
diff --git a/Engine/Source/Runtime/Experimental/Dataflow/Engine/Public/Dataflow/DataflowContent.h b/Engine/Source/Runtime/Experimental/Dataflow/Engine/Public/Dataflow/DataflowContent.h
--- a/Engine/Source/Runtime/Experimental/Dataflow/Engine/Public/Dataflow/DataflowContent.h
+++ b/Engine/Source/Runtime/Experimental/Dataflow/Engine/Public/Dataflow/DataflowContent.h
@@ -29,6 +29,14 @@ namespace Dataflow
 	};
 }
 
+namespace DataflowContextHelpers
+{
+	// Return a new(or saved) content that can store the execution state of the graph. 
+	template<class T>
+	DATAFLOWENGINE_API TObjectPtr<T> CreateNewDataflowContent(const TObjectPtr<UObject>& ContentOwner);
+}
+
+
 /** 
  * Context object used for selection/rendering 
  */
@@ -174,6 +182,10 @@ public:
 	/** Content Serialization */
 	virtual void Serialize(FArchive& Ar);
 
+	/* Context cache saving */
+	bool IsSaved() const { return bIsSaved; }
+	void SetIsSaved(bool bInSaved) { bIsSaved = bInSaved; }
+
 	//~ UObject interface
 	static void AddReferencedObjects(UObject* InThis, FReferenceCollector& Collector);
 
@@ -196,6 +208,9 @@ protected:
     /** Dirty flag to trigger rendering. Do we need that? since when accessing the member by non const ref we will not dirty it */
 	UPROPERTY()
 	bool bIsDirty = true;
+
+	/** Saved as a cached context. Will be automatically saved to a cache directory if true. Use the pvar p.Dataflow.Editor.ContextCaching to enable. [def:false] */
+	bool bIsSaved = false;
 };
 
 /** 
