GeometryCollection : Conversion from SkeletalMesh
- Fixed insertion issue for debug builds where infos and raw names were misaligned.

diff --git a/Engine/Source/Runtime/Experimental/Chaos/Private/GeometryCollection/Facades/CollectionTransformFacade.cpp b/Engine/Source/Runtime/Experimental/Chaos/Private/GeometryCollection/Facades/CollectionTransformFacade.cpp
--- a/Engine/Source/Runtime/Experimental/Chaos/Private/GeometryCollection/Facades/CollectionTransformFacade.cpp
+++ b/Engine/Source/Runtime/Experimental/Chaos/Private/GeometryCollection/Facades/CollectionTransformFacade.cpp
@@ -361,10 +361,13 @@ namespace GeometryCollection::Facades
 
 	CHAOS_API void FCollectionTransformFacade::EnforceSingleRoot(const FString & RootName)
 	{
+		check(Collection);
+
 		TArray<int32> Roots = GetRootIndices();
 		if (Roots.Num() > 1)
 		{
 			int32 Idx = TransformAttribute.AddElements(1);
+			ParentAttribute.Modify()[Idx] = INDEX_NONE;
 
 			if (HasBoneNameAttribute())
 			{
diff --git a/Engine/Source/Runtime/Experimental/GeometryCollectionEngine/Private/GeometryCollection/GeometryCollectionEngineConversion.cpp b/Engine/Source/Runtime/Experimental/GeometryCollectionEngine/Private/GeometryCollection/GeometryCollectionEngineConversion.cpp
--- a/Engine/Source/Runtime/Experimental/GeometryCollectionEngine/Private/GeometryCollection/GeometryCollectionEngineConversion.cpp
+++ b/Engine/Source/Runtime/Experimental/GeometryCollectionEngine/Private/GeometryCollection/GeometryCollectionEngineConversion.cpp
@@ -1714,17 +1714,19 @@ void FGeometryCollectionEngineConversion::ConvertActorToGeometryCollection(const
 
 void FGeometryCollectionEngineConversion::ConvertCollectionToSkeleton(const FManagedArrayCollection& InCollection, USkeleton* OutSkeleton, TArray<int32>& OutIndexRemap)
 {
-	GeometryCollection::Facades::FCollectionTransformFacade Transforms(InCollection);
+	FManagedArrayCollection LocalCollection = InCollection;
+	GeometryCollection::Facades::FCollectionTransformFacade Transforms(LocalCollection);
 	if (Transforms.IsValid() && Transforms.HasBoneNameAttribute() && OutSkeleton)
 	{
+		Transforms.EnforceSingleRoot("root");
+
 		OutIndexRemap.Init(INDEX_NONE, Transforms.Num());
 		auto AddMapping = [&OutIndexRemap](int32 A, int32 B)
 		{
 			OutIndexRemap[A] = B;
 		};
 
-		FReferenceSkeletonModifier Edit(OutSkeleton);
-		auto AddChildren = [&Transforms, &Edit, &AddMapping](const TArray<int32>& CollectionChildren)
+		auto AddChildren = [&Transforms, &OutSkeleton, &AddMapping](const TArray<int32>& CollectionChildren)
 		{
 			TQueue<int32> Children;
 			auto Enqueue = [&Children](const TArray<int32>& List)
@@ -1740,13 +1742,16 @@ void FGeometryCollectionEngineConversion::ConvertCollectionToSkeleton(const FMan
 				Children.Dequeue(CurrentIndex);
 
 				int32 CollectionParentIndex = (*Transforms.GetParents())[CurrentIndex];
-				int32 SkeletionParentIndex = Edit.FindBoneIndex(FName((*Transforms.FindBoneNames())[CollectionParentIndex]));
+				int32 SkeletionParentIndex = OutSkeleton->GetReferenceSkeleton().FindBoneIndex(FName((*Transforms.FindBoneNames())[CollectionParentIndex]));
 
 				FName BoneName = FName((*Transforms.FindBoneNames())[CurrentIndex]);
 				FTransform Transform = FTransform((*Transforms.FindTransforms())[CurrentIndex]);
 				FMeshBoneInfo Info(BoneName, BoneName.ToString(), SkeletionParentIndex);
-				Edit.Add(Info, Transform, true /*bAllowMultipleRoots*/);
-				AddMapping(CurrentIndex, Edit.GetReferenceSkeleton().GetRawBoneNum());
+				{
+					FReferenceSkeletonModifier Edit(OutSkeleton);
+					Edit.Add(Info, Transform, true /*bAllowMultipleRoots*/);
+				}
+				AddMapping(CurrentIndex, OutSkeleton->GetReferenceSkeleton().GetRawBoneNum());
 
 				Enqueue((*Transforms.FindChildren())[CurrentIndex].Array());
 			}
@@ -1762,8 +1767,11 @@ void FGeometryCollectionEngineConversion::ConvertCollectionToSkeleton(const FMan
 				FName BoneName = FName((*Transforms.FindBoneNames())[i]);
 				FTransform Transform = FTransform((*Transforms.FindTransforms())[i]);
 				FMeshBoneInfo Info(BoneName, BoneName.ToString(), INDEX_NONE);
-				Edit.Add(Info, Transform, true /*bAllowMultipleRoots*/);
-				AddMapping(i,Edit.GetReferenceSkeleton().GetNum());
+				{
+					FReferenceSkeletonModifier Edit(OutSkeleton);
+					Edit.Add(Info, Transform, true /*bAllowMultipleRoots*/);
+				}
+				AddMapping(i,OutSkeleton->GetReferenceSkeleton().GetNum());
 
 				if ((*Transforms.FindChildren())[i].Num())
 				{
