Sequencer: remove IMovieScenePlayer usage in OnBindingIDsUpdated and animation constraints

diff --git a/Engine/Plugins/Animation/ControlRig/Source/ControlRig/Private/Constraints/ControlRigTransformableHandle.cpp b/Engine/Plugins/Animation/ControlRig/Source/ControlRig/Private/Constraints/ControlRigTransformableHandle.cpp
--- a/Engine/Plugins/Animation/ControlRig/Source/ControlRig/Private/Constraints/ControlRigTransformableHandle.cpp
+++ b/Engine/Plugins/Animation/ControlRig/Source/ControlRig/Private/Constraints/ControlRigTransformableHandle.cpp
@@ -627,13 +627,13 @@ bool UTransformableControlHandle::AddTransformKeys(const TArray<FFrameNumber>& I
 }
 
 //for control rig need to check to see if the control rig is different then we may need to update it based upon what we are now bound to
-void UTransformableControlHandle::ResolveBoundObjects(FMovieSceneSequenceID LocalSequenceID, IMovieScenePlayer& Player, UObject* SubObject)
+void UTransformableControlHandle::ResolveBoundObjects(FMovieSceneSequenceID LocalSequenceID, TSharedRef<UE::MovieScene::FSharedPlaybackState> SharedPlaybackState, UObject* SubObject)
 {
 	if (const UControlRig* InControlRig = Cast<UControlRig>(SubObject))
 	{
 		if (ControlRig != InControlRig)
 		{
-			for (const TWeakObjectPtr<> ParentObject : ConstraintBindingID.ResolveBoundObjects(LocalSequenceID, Player))
+			for (const TWeakObjectPtr<> ParentObject : ConstraintBindingID.ResolveBoundObjects(LocalSequenceID, SharedPlaybackState))
 			{
 				const UObject* Bindable = FControlRigObjectBinding::GetBindableObject(ParentObject.Get());
 				if (InControlRig->GetObjectBinding() && InControlRig->GetObjectBinding()->GetBoundObject() == Bindable)
diff --git a/Engine/Plugins/Animation/ControlRig/Source/ControlRig/Private/Sequencer/MovieSceneControlRigParameterSection.cpp b/Engine/Plugins/Animation/ControlRig/Source/ControlRig/Private/Sequencer/MovieSceneControlRigParameterSection.cpp
--- a/Engine/Plugins/Animation/ControlRig/Source/ControlRig/Private/Sequencer/MovieSceneControlRigParameterSection.cpp
+++ b/Engine/Plugins/Animation/ControlRig/Source/ControlRig/Private/Sequencer/MovieSceneControlRigParameterSection.cpp
@@ -822,7 +822,7 @@ UMovieSceneControlRigParameterSection::UMovieSceneControlRigParameterSection() :
 #endif
 }
 
-void UMovieSceneControlRigParameterSection::OnBindingIDsUpdated(const TMap<UE::MovieScene::FFixedObjectBindingID, UE::MovieScene::FFixedObjectBindingID>& OldFixedToNewFixedMap, FMovieSceneSequenceID LocalSequenceID, const FMovieSceneSequenceHierarchy* Hierarchy, IMovieScenePlayer& Player)
+void UMovieSceneControlRigParameterSection::OnBindingIDsUpdated(const TMap<UE::MovieScene::FFixedObjectBindingID, UE::MovieScene::FFixedObjectBindingID>& OldFixedToNewFixedMap, FMovieSceneSequenceID LocalSequenceID, TSharedRef<UE::MovieScene::FSharedPlaybackState> SharedPlaybackState)
 {
 	for (FConstraintAndActiveChannel& ConstraintChannel : ConstraintsChannels)
 	{
@@ -830,11 +830,11 @@ void UMovieSceneControlRigParameterSection::OnBindingIDsUpdated(const TMap<UE::M
 		{
 			if (TransformConstraint->ChildTRSHandle)
 			{
-				TransformConstraint->ChildTRSHandle->OnBindingIDsUpdated(OldFixedToNewFixedMap, LocalSequenceID, Hierarchy, Player);
+				TransformConstraint->ChildTRSHandle->OnBindingIDsUpdated(OldFixedToNewFixedMap, LocalSequenceID, SharedPlaybackState);
 			}
 			if (TransformConstraint->ParentTRSHandle)
 			{
-				TransformConstraint->ParentTRSHandle->OnBindingIDsUpdated(OldFixedToNewFixedMap, LocalSequenceID, Hierarchy, Player);
+				TransformConstraint->ParentTRSHandle->OnBindingIDsUpdated(OldFixedToNewFixedMap, LocalSequenceID, SharedPlaybackState);
 			}
 		}
 	}
diff --git a/Engine/Plugins/Animation/ControlRig/Source/ControlRig/Private/Sequencer/MovieSceneControlRigParameterTemplate.cpp b/Engine/Plugins/Animation/ControlRig/Source/ControlRig/Private/Sequencer/MovieSceneControlRigParameterTemplate.cpp
--- a/Engine/Plugins/Animation/ControlRig/Source/ControlRig/Private/Sequencer/MovieSceneControlRigParameterTemplate.cpp
+++ b/Engine/Plugins/Animation/ControlRig/Source/ControlRig/Private/Sequencer/MovieSceneControlRigParameterTemplate.cpp
@@ -1240,7 +1240,7 @@ struct FControlRigParameterExecutionToken : IMovieSceneExecutionToken
 		
 		FMovieSceneSequenceID SequenceID = Operand.SequenceID;
 		TArrayView<TWeakObjectPtr<>> BoundObjects = Player.FindBoundObjects(Operand);
-			const UMovieSceneSequence* Sequence = Player.State.FindSequence(Operand.SequenceID);
+		const UMovieSceneSequence* Sequence = Player.State.FindSequence(Operand.SequenceID);
 		UControlRig* ControlRig = nullptr;
 		UObject* BoundObject = BoundObjects.Num() > 0 ? BoundObjects[0].Get() : nullptr;
 		if (Sequence && BoundObject)
@@ -1412,6 +1412,7 @@ struct FControlRigParameterExecutionToken : IMovieSceneExecutionToken
 				}
 				if (BoundObject)
 				{
+					TSharedRef<UE::MovieScene::FSharedPlaybackState> SharedPlaybackState = Player.GetSharedPlaybackState();
 					for (FConstraintAndActiveValue& ConstraintValue : ConstraintsValues)
 					{
 						UMovieSceneControlRigParameterSection* NonConstSection = const_cast<UMovieSceneControlRigParameterSection*>(Section);
@@ -1424,7 +1425,7 @@ struct FControlRigParameterExecutionToken : IMovieSceneExecutionToken
 							{
 								TransformConstraint->InitConstraint(BoundObject->GetWorld());
 							}
-							ConstraintValue.Constraint->ResolveBoundObjects(Operand.SequenceID, Player, ControlRig);
+							ConstraintValue.Constraint->ResolveBoundObjects(Operand.SequenceID, SharedPlaybackState, ControlRig);
 							ConstraintValue.Constraint->SetActive(ConstraintValue.Value);
 						}
 					}
@@ -1438,7 +1439,7 @@ struct FControlRigParameterExecutionToken : IMovieSceneExecutionToken
 						{
 							if (TransformConstraint->ParentTRSHandle)
 							{
-								TransformConstraint->ParentTRSHandle->ResolveBoundObjects(Operand.SequenceID, Player, ControlRig);
+								TransformConstraint->ParentTRSHandle->ResolveBoundObjects(Operand.SequenceID, SharedPlaybackState, ControlRig);
 								TransformConstraint->EnsurePrimaryDependency(BoundObject->GetWorld());
 							}
 						}
diff --git a/Engine/Plugins/Animation/ControlRig/Source/ControlRig/Public/Constraints/ControlRigTransformableHandle.h b/Engine/Plugins/Animation/ControlRig/Source/ControlRig/Public/Constraints/ControlRigTransformableHandle.h
--- a/Engine/Plugins/Animation/ControlRig/Source/ControlRig/Public/Constraints/ControlRigTransformableHandle.h
+++ b/Engine/Plugins/Animation/ControlRig/Source/ControlRig/Public/Constraints/ControlRigTransformableHandle.h
@@ -78,7 +78,7 @@ public:
 		const bool bLocal = true) const override;
 
 	/** Resolve the bound objects so that any object it references are resolved and correctly set up*/
-	virtual void ResolveBoundObjects(FMovieSceneSequenceID LocalSequenceID, IMovieScenePlayer& Player, UObject* SubObject = nullptr) override;
+	virtual void ResolveBoundObjects(FMovieSceneSequenceID LocalSequenceID, TSharedRef<UE::MovieScene::FSharedPlaybackState> SharedPlaybackState, UObject* SubObject = nullptr) override;
 
 	/** Make a duplicate of myself with this outer*/
 	virtual UTransformableHandle* Duplicate(UObject* NewOuter) const override;
@@ -148,4 +148,4 @@ private:
 
 	/** @todo document */
 	FRigControlElement* GetControlElement() const;
-};
\ No newline at end of file
+};
diff --git a/Engine/Plugins/Animation/ControlRig/Source/ControlRig/Public/Sequencer/MovieSceneControlRigParameterSection.h b/Engine/Plugins/Animation/ControlRig/Source/ControlRig/Public/Sequencer/MovieSceneControlRigParameterSection.h
--- a/Engine/Plugins/Animation/ControlRig/Source/ControlRig/Public/Sequencer/MovieSceneControlRigParameterSection.h
+++ b/Engine/Plugins/Animation/ControlRig/Source/ControlRig/Public/Sequencer/MovieSceneControlRigParameterSection.h
@@ -519,7 +519,7 @@ protected:
 	virtual void PostEditImport() override;
 	virtual void PostLoad() override;
 	virtual float GetTotalWeightValue(FFrameTime InTime) const override;
-	virtual void OnBindingIDsUpdated(const TMap<UE::MovieScene::FFixedObjectBindingID, UE::MovieScene::FFixedObjectBindingID>& OldFixedToNewFixedMap, FMovieSceneSequenceID LocalSequenceID, const FMovieSceneSequenceHierarchy* Hierarchy, IMovieScenePlayer& Player) override;
+	virtual void OnBindingIDsUpdated(const TMap<UE::MovieScene::FFixedObjectBindingID, UE::MovieScene::FFixedObjectBindingID>& OldFixedToNewFixedMap, FMovieSceneSequenceID LocalSequenceID, TSharedRef<UE::MovieScene::FSharedPlaybackState> SharedPlaybackState) override;
 	virtual void GetReferencedBindings(TArray<FGuid>& OutBindings) override;
 	virtual void PreSave(FObjectPreSaveContext SaveContext) override;
 
diff --git a/Engine/Plugins/Media/MediaCompositing/Source/MediaCompositing/Private/MovieScene/MovieSceneMediaSection.cpp b/Engine/Plugins/Media/MediaCompositing/Source/MediaCompositing/Private/MovieScene/MovieSceneMediaSection.cpp
--- a/Engine/Plugins/Media/MediaCompositing/Source/MediaCompositing/Private/MovieScene/MovieSceneMediaSection.cpp
+++ b/Engine/Plugins/Media/MediaCompositing/Source/MediaCompositing/Private/MovieScene/MovieSceneMediaSection.cpp
@@ -103,17 +103,16 @@ void UMovieSceneMediaSection::MigrateFrameTimes(FFrameRate SourceRate, FFrameRat
 	}
 }
 
-void UMovieSceneMediaSection::OnBindingIDsUpdated(const TMap<UE::MovieScene::FFixedObjectBindingID, UE::MovieScene::FFixedObjectBindingID>& OldFixedToNewFixedMap, FMovieSceneSequenceID LocalSequenceID, const FMovieSceneSequenceHierarchy* Hierarchy, IMovieScenePlayer& Player)
+void UMovieSceneMediaSection::OnBindingIDsUpdated(const TMap<UE::MovieScene::FFixedObjectBindingID, UE::MovieScene::FFixedObjectBindingID>& OldFixedToNewFixedMap, FMovieSceneSequenceID LocalSequenceID, TSharedRef<UE::MovieScene::FSharedPlaybackState> SharedPlaybackState)
 {
-	UE::MovieScene::FFixedObjectBindingID FixedBindingID = 
-		MediaSourceProxyBindingID.ResolveToFixed(LocalSequenceID, Player);
+	UE::MovieScene::FFixedObjectBindingID FixedBindingID = MediaSourceProxyBindingID.ResolveToFixed(LocalSequenceID, SharedPlaybackState);
 
 	if (OldFixedToNewFixedMap.Contains(FixedBindingID))
 	{
 		Modify();
 
-		MediaSourceProxyBindingID =
-			OldFixedToNewFixedMap[FixedBindingID].ConvertToRelative(LocalSequenceID, Hierarchy);
+		const FMovieSceneSequenceHierarchy* Hierarchy = SharedPlaybackState->GetHierarchy();
+		MediaSourceProxyBindingID = OldFixedToNewFixedMap[FixedBindingID].ConvertToRelative(LocalSequenceID, Hierarchy);
 	}
 }
 
diff --git a/Engine/Plugins/Media/MediaCompositing/Source/MediaCompositing/Public/MovieSceneMediaSection.h b/Engine/Plugins/Media/MediaCompositing/Source/MediaCompositing/Public/MovieSceneMediaSection.h
--- a/Engine/Plugins/Media/MediaCompositing/Source/MediaCompositing/Public/MovieSceneMediaSection.h
+++ b/Engine/Plugins/Media/MediaCompositing/Source/MediaCompositing/Public/MovieSceneMediaSection.h
@@ -97,7 +97,7 @@ protected:
 	virtual UMovieSceneSection* SplitSection(FQualifiedFrameTime SplitTime, bool bDeleteKeys) override;
 	virtual TOptional<FFrameTime> GetOffsetTime() const override;
 	virtual void MigrateFrameTimes(FFrameRate SourceRate, FFrameRate DestinationRate) override;
-	virtual void OnBindingIDsUpdated(const TMap<UE::MovieScene::FFixedObjectBindingID, UE::MovieScene::FFixedObjectBindingID>& OldFixedToNewFixedMap, FMovieSceneSequenceID LocalSequenceID, const FMovieSceneSequenceHierarchy* Hierarchy, IMovieScenePlayer& Player) override;
+	virtual void OnBindingIDsUpdated(const TMap<UE::MovieScene::FFixedObjectBindingID, UE::MovieScene::FFixedObjectBindingID>& OldFixedToNewFixedMap, FMovieSceneSequenceID LocalSequenceID, TSharedRef<UE::MovieScene::FSharedPlaybackState> SharedPlaybackState) override;
 	virtual void GetReferencedBindings(TArray<FGuid>& OutBindings) override;
 
 
diff --git a/Engine/Source/Editor/Sequencer/Private/SequencerUtilities.cpp b/Engine/Source/Editor/Sequencer/Private/SequencerUtilities.cpp
--- a/Engine/Source/Editor/Sequencer/Private/SequencerUtilities.cpp
+++ b/Engine/Source/Editor/Sequencer/Private/SequencerUtilities.cpp
@@ -3336,13 +3336,15 @@ void FSequencerUtilities::UpdateBindingIDs(TSharedRef<ISequencer> Sequencer, FGu
 	TMap<UE::MovieScene::FFixedObjectBindingID, UE::MovieScene::FFixedObjectBindingID> OldFixedToNewFixedMap;
 	OldFixedToNewFixedMap.Add(UE::MovieScene::FFixedObjectBindingID(OldGuid, FocusedGuid), UE::MovieScene::FFixedObjectBindingID(NewGuid, FocusedGuid));
 
+	TSharedRef<UE::MovieScene::FSharedPlaybackState> SharedPlaybackState = Sequencer->GetSharedPlaybackState();
+
 	if (UMovieScene* MovieScene = Sequencer->GetRootMovieSceneSequence()->GetMovieScene())
 	{
 		for (UMovieSceneSection* Section : MovieScene->GetAllSections())
 		{
 			if (Section)
 			{
-				Section->OnBindingIDsUpdated(OldFixedToNewFixedMap, Sequencer->GetRootTemplateID(), Hierarchy, *Sequencer);
+				Section->OnBindingIDsUpdated(OldFixedToNewFixedMap, Sequencer->GetRootTemplateID(), SharedPlaybackState);
 			}
 		}
 	}
@@ -3359,7 +3361,7 @@ void FSequencerUtilities::UpdateBindingIDs(TSharedRef<ISequencer> Sequencer, FGu
 					{
 						if (Section)
 						{
-							Section->OnBindingIDsUpdated(OldFixedToNewFixedMap, Pair.Key, Hierarchy, *Sequencer);
+							Section->OnBindingIDsUpdated(OldFixedToNewFixedMap, Pair.Key, SharedPlaybackState);
 						}
 					}
 				}
diff --git a/Engine/Source/Runtime/Experimental/Animation/Constraints/Private/TransformConstraint.cpp b/Engine/Source/Runtime/Experimental/Animation/Constraints/Private/TransformConstraint.cpp
--- a/Engine/Source/Runtime/Experimental/Animation/Constraints/Private/TransformConstraint.cpp
+++ b/Engine/Source/Runtime/Experimental/Animation/Constraints/Private/TransformConstraint.cpp
@@ -594,18 +594,18 @@ bool UTickableTransformConstraint::HasBoundObjects() const
 	return false;
 }
 
-void UTickableTransformConstraint::ResolveBoundObjects(FMovieSceneSequenceID LocalSequenceID, IMovieScenePlayer& Player, UObject* SubObject)
+void UTickableTransformConstraint::ResolveBoundObjects(FMovieSceneSequenceID LocalSequenceID, TSharedRef<UE::MovieScene::FSharedPlaybackState> SharedPlaybackState, UObject* SubObject)
 {
 	// update dependencies if the constraint becomes valid once resolved 
 	FConstraintDependencyScope Scope(this);
 	
 	if (ChildTRSHandle && ChildTRSHandle->HasBoundObjects())
 	{
-		ChildTRSHandle->ResolveBoundObjects(LocalSequenceID, Player, SubObject);
+		ChildTRSHandle->ResolveBoundObjects(LocalSequenceID, SharedPlaybackState, SubObject);
 	}
 	if (ParentTRSHandle && ParentTRSHandle->HasBoundObjects())
 	{
-		ParentTRSHandle->ResolveBoundObjects(LocalSequenceID, Player, SubObject);
+		ParentTRSHandle->ResolveBoundObjects(LocalSequenceID, SharedPlaybackState, SubObject);
 	}
 }
 
@@ -2542,4 +2542,4 @@ bool FHandleDependencyChecker::HasDependency(const UTransformableHandle& InHandl
 	}
 	
 	return false;
-}
\ No newline at end of file
+}
diff --git a/Engine/Source/Runtime/Experimental/Animation/Constraints/Private/TransformableHandle.cpp b/Engine/Source/Runtime/Experimental/Animation/Constraints/Private/TransformableHandle.cpp
--- a/Engine/Source/Runtime/Experimental/Animation/Constraints/Private/TransformableHandle.cpp
+++ b/Engine/Source/Runtime/Experimental/Animation/Constraints/Private/TransformableHandle.cpp
@@ -6,6 +6,7 @@
 #include "Components/SceneComponent.h"
 #include "GameFramework/Actor.h"
 #include "Engine/Engine.h"
+#include "EntitySystem/MovieSceneSharedPlaybackState.h"
 #include "MovieSceneSection.h"
 #include "TransformableHandleUtils.h"
 #include "Channels/MovieSceneFloatChannel.h"
@@ -48,14 +49,16 @@ bool UTransformableHandle::HasBoundObjects() const
 	return ConstraintBindingID.IsValid();
 }
 
-void UTransformableHandle::OnBindingIDsUpdated(const TMap<UE::MovieScene::FFixedObjectBindingID, UE::MovieScene::FFixedObjectBindingID>& OldFixedToNewFixedMap, FMovieSceneSequenceID LocalSequenceID, const FMovieSceneSequenceHierarchy* Hierarchy, IMovieScenePlayer& Player)
+void UTransformableHandle::OnBindingIDsUpdated(const TMap<UE::MovieScene::FFixedObjectBindingID, UE::MovieScene::FFixedObjectBindingID>& OldFixedToNewFixedMap, FMovieSceneSequenceID LocalSequenceID, TSharedRef<UE::MovieScene::FSharedPlaybackState> SharedPlaybackState)
 {
 	if (ConstraintBindingID.IsValid())
 	{
-		UE::MovieScene::FFixedObjectBindingID FixedBindingID = ConstraintBindingID.ResolveToFixed(LocalSequenceID, Player);
+		UE::MovieScene::FFixedObjectBindingID FixedBindingID = ConstraintBindingID.ResolveToFixed(LocalSequenceID, SharedPlaybackState);
 		if (OldFixedToNewFixedMap.Contains(FixedBindingID))
 		{
 			Modify();
+
+			const FMovieSceneSequenceHierarchy* Hierarchy = SharedPlaybackState->GetHierarchy();
 			ConstraintBindingID = OldFixedToNewFixedMap[FixedBindingID].ConvertToRelative(LocalSequenceID, Hierarchy);
 		}
 	}
@@ -534,14 +537,14 @@ bool UTransformableComponentHandle::AddTransformKeys(const TArray<FFrameNumber>&
 	return true;
 }
 
-void UTransformableComponentHandle::ResolveBoundObjects(FMovieSceneSequenceID LocalSequenceID, IMovieScenePlayer& Player, UObject* SubObject)
+void UTransformableComponentHandle::ResolveBoundObjects(FMovieSceneSequenceID LocalSequenceID, TSharedRef<UE::MovieScene::FSharedPlaybackState> SharedPlaybackState, UObject* SubObject)
 {
 	// in the context of blueprints being recompiled (cf. AActor::RerunConstructionScripts()) or with spawnable objects, the component has to be updated.
 	// as this can be called after it has been destroyed, we get the component even if it's pending kill otherwise Get() will return a nullptr.
 	static constexpr bool bEvenIfPendingKill = true;
 	const USceneComponent* ComponentEvenIfPendingKill = Component.Get(bEvenIfPendingKill);
 	
-	for (const TWeakObjectPtr<>& ParentObject : ConstraintBindingID.ResolveBoundObjects(LocalSequenceID, Player))
+	for (const TWeakObjectPtr<>& ParentObject : ConstraintBindingID.ResolveBoundObjects(LocalSequenceID, SharedPlaybackState))
 	{
 		if (const AActor* Actor = Cast<AActor>(ParentObject.Get()))
 		{
diff --git a/Engine/Source/Runtime/Experimental/Animation/Constraints/Public/ConstraintsManager.h b/Engine/Source/Runtime/Experimental/Animation/Constraints/Public/ConstraintsManager.h
--- a/Engine/Source/Runtime/Experimental/Animation/Constraints/Public/ConstraintsManager.h
+++ b/Engine/Source/Runtime/Experimental/Animation/Constraints/Public/ConstraintsManager.h
@@ -9,9 +9,10 @@
 #include "Engine/Level.h"
 #include "ConstraintsManager.generated.h"
 
-class IMovieScenePlayer;
 class UTickableConstraint;
 
+namespace UE::MovieScene { struct FSharedPlaybackState; }
+
 /** 
  * FConstraintTickFunction
  * Represents the interface of constraint as a tick function. This allows both to evaluate a constraint in the
@@ -93,7 +94,7 @@ public:
 	/** If true it contains objects bound to an external system, like sequencer so we don't do certain things, like remove constraints when they don't resolve*/
 	CONSTRAINTS_API virtual bool HasBoundObjects() const PURE_VIRTUAL(HasBoundObjects, return false;);
 	/** Resolve the bound objects so that any object it references are resovled and correctly set up*/
-	CONSTRAINTS_API virtual void ResolveBoundObjects(FMovieSceneSequenceID LocalSequenceID, IMovieScenePlayer& Player, UObject* SubObject = nullptr)  PURE_VIRTUAL(ResolveBoundObjects);
+	CONSTRAINTS_API virtual void ResolveBoundObjects(FMovieSceneSequenceID LocalSequenceID, TSharedRef<UE::MovieScene::FSharedPlaybackState> SharedPlaybackState, UObject* SubObject = nullptr)  PURE_VIRTUAL(ResolveBoundObjects);
 
 	/** @todo document */
 	CONSTRAINTS_API virtual uint32 GetTargetHash() const PURE_VIRTUAL(GetTargetHash, return 0;);
diff --git a/Engine/Source/Runtime/Experimental/Animation/Constraints/Public/TransformConstraint.h b/Engine/Source/Runtime/Experimental/Animation/Constraints/Public/TransformConstraint.h
--- a/Engine/Source/Runtime/Experimental/Animation/Constraints/Public/TransformConstraint.h
+++ b/Engine/Source/Runtime/Experimental/Animation/Constraints/Public/TransformConstraint.h
@@ -4,6 +4,7 @@
 
 #include "Constraint.h"
 #include "ConstraintsManager.h"
+#include "Evaluation/IMovieScenePlaybackCapability.h"
 #include "Kismet/BlueprintFunctionLibrary.h"
 
 #include "Delegates/DelegateCombinations.h"
@@ -44,7 +45,7 @@ public:
 	CONSTRAINTS_API virtual bool HasBoundObjects() const override;
 	
 	/** Resolve the bound objects so that any object it references are resovled and correctly set up*/
-	CONSTRAINTS_API virtual void ResolveBoundObjects(FMovieSceneSequenceID LocalSequenceID, IMovieScenePlayer& Player, UObject* SubObject = nullptr) override;
+	CONSTRAINTS_API virtual void ResolveBoundObjects(FMovieSceneSequenceID LocalSequenceID, TSharedRef<UE::MovieScene::FSharedPlaybackState> SharedPlaybackState, UObject* SubObject = nullptr) override;
 
 	/** Whether or not it's valid for example it may not be fully loaded*/
 	virtual bool IsValid(const bool bDeepCheck = true) const override;
@@ -534,4 +535,4 @@ struct FHandleDependencyChecker
 	bool HasDependency(const UTransformableHandle& InHandle, const UTransformableHandle& InParentToCheck) const;		
 private:
 	TWeakObjectPtr<UWorld> WeakWorld = nullptr;
-};
\ No newline at end of file
+};
diff --git a/Engine/Source/Runtime/Experimental/Animation/Constraints/Public/TransformableHandle.h b/Engine/Source/Runtime/Experimental/Animation/Constraints/Public/TransformableHandle.h
--- a/Engine/Source/Runtime/Experimental/Animation/Constraints/Public/TransformableHandle.h
+++ b/Engine/Source/Runtime/Experimental/Animation/Constraints/Public/TransformableHandle.h
@@ -7,18 +7,22 @@
 #include "MovieSceneObjectBindingID.h"
 #include "TransformableHandle.generated.h"
 
-struct FTickFunction;
-class USceneComponent;
-struct FMovieSceneFloatChannel;
-struct FMovieSceneDoubleChannel;
 class UMovieSceneSection;
+class USceneComponent;
+enum class EMovieSceneTransformChannel : uint32;
 struct FFRameNumber;
 struct FFrameRate;
-enum class EMovieSceneTransformChannel : uint32;
-class IMovieScenePlayer;
-namespace UE { namespace MovieScene {struct FFixedObjectBindingID;} }
+struct FMovieSceneDoubleChannel;
+struct FMovieSceneFloatChannel;
 struct FMovieSceneSequenceHierarchy;
 struct FMovieSceneSequenceID;
+struct FTickFunction;
+
+namespace UE::MovieScene
+{
+	struct FFixedObjectBindingID;
+	struct FSharedPlaybackState;
+}
 
 UENUM()
 enum class EHandleEvent : uint8
@@ -68,14 +72,14 @@ public:
 	CONSTRAINTS_API virtual bool HasBoundObjects() const;
 
 	/** Resolve the bound objects so that any object it references are resolved and correctly set up*/
-	CONSTRAINTS_API virtual void ResolveBoundObjects(FMovieSceneSequenceID LocalSequenceID, IMovieScenePlayer& Player, UObject* SubObject = nullptr) PURE_VIRTUAL(ResolveBoundObjects);
+	CONSTRAINTS_API virtual void ResolveBoundObjects(FMovieSceneSequenceID LocalSequenceID, TSharedRef<UE::MovieScene::FSharedPlaybackState> SharedPlaybackState, UObject* SubObject = nullptr) PURE_VIRTUAL(ResolveBoundObjects);
 
 	/** Make a duplicate of myself with this outer*/
 	CONSTRAINTS_API virtual UTransformableHandle* Duplicate(UObject* NewOuter) const PURE_VIRTUAL(Duplicate, return nullptr;);
 
 
 	/** Fix up Binding in case it has changed*/
-	CONSTRAINTS_API void OnBindingIDsUpdated(const TMap<UE::MovieScene::FFixedObjectBindingID, UE::MovieScene::FFixedObjectBindingID>& OldFixedToNewFixedMap, FMovieSceneSequenceID LocalSequenceID, const FMovieSceneSequenceHierarchy* Hierarchy, IMovieScenePlayer& Player);
+	CONSTRAINTS_API void OnBindingIDsUpdated(const TMap<UE::MovieScene::FFixedObjectBindingID, UE::MovieScene::FFixedObjectBindingID>& OldFixedToNewFixedMap, FMovieSceneSequenceID LocalSequenceID, TSharedRef<UE::MovieScene::FSharedPlaybackState> SharedPlaybackState);
 	
 	/** Perform any special ticking needed for this handle, by default it does nothing, todo need to see if we need to tick control rig also*/
 	virtual void TickTarget() const {};
@@ -202,7 +206,7 @@ public:
 		const bool bLocal = true) const override;
 
 	/** Resolve the bound objects so that any object it references are resovled and correctly set up*/
-	CONSTRAINTS_API virtual void ResolveBoundObjects(FMovieSceneSequenceID LocalSequenceID, IMovieScenePlayer& Player, UObject* SubObject = nullptr) override;
+	CONSTRAINTS_API virtual void ResolveBoundObjects(FMovieSceneSequenceID LocalSequenceID, TSharedRef<UE::MovieScene::FSharedPlaybackState> SharedPlaybackState, UObject* SubObject = nullptr) override;
 
 	/** Make a duplicate of myself with this outer*/
 	CONSTRAINTS_API virtual UTransformableHandle* Duplicate(UObject* NewOuter) const override;
diff --git a/Engine/Source/Runtime/MovieScene/Public/MovieSceneSection.h b/Engine/Source/Runtime/MovieScene/Public/MovieSceneSection.h
--- a/Engine/Source/Runtime/MovieScene/Public/MovieSceneSection.h
+++ b/Engine/Source/Runtime/MovieScene/Public/MovieSceneSection.h
@@ -40,7 +40,7 @@ class IMovieSceneEasingFunction;
 class IMovieScenePlayer;
 class UMovieSceneEntitySystemLinker;
 class UObject;
-namespace UE { namespace MovieScene { class ISectionEventHandler; } }
+enum class ECookOptimizationFlags;
 struct FEasingComponentData;
 struct FFrame;
 struct FFrameRate;
@@ -54,19 +54,15 @@ struct FMovieSceneSequenceID;
 struct FPropertyChangedEvent;
 struct FQualifiedFrameTime;
 
-enum class ECookOptimizationFlags;
-
-namespace UE
-{
-namespace MovieScene
+namespace UE::MovieScene
 {
+	class ISectionEventHandler;
 	struct FEntityImportParams;
 	struct FFixedObjectBindingID;
 	struct FImportedEntity;
-}
+	struct FSharedPlaybackState;
 }
 
-
 /** Enumeration defining how a section's channel proxy behaves. */
 enum class EMovieSceneChannelProxyType : uint8
 {
@@ -592,11 +588,8 @@ public:
 	/* Migrate the frame times of the movie scene section from the source frame rate to the destination frame rate */
 	virtual void MigrateFrameTimes(FFrameRate SourceRate, FFrameRate DestinationRate) {}
 
-	/**
-	 * When guid bindings are updated to allow this section to fix-up any internal bindings
-	 *
-	 */
-	virtual void OnBindingIDsUpdated(const TMap<UE::MovieScene::FFixedObjectBindingID, UE::MovieScene::FFixedObjectBindingID>& OldFixedToNewFixedMap, FMovieSceneSequenceID LocalSequenceID, const FMovieSceneSequenceHierarchy* Hierarchy, IMovieScenePlayer& Player) {}
+	/** When guid bindings are updated to allow this section to fix-up any internal bindings */
+	virtual void OnBindingIDsUpdated(const TMap<UE::MovieScene::FFixedObjectBindingID, UE::MovieScene::FFixedObjectBindingID>& OldFixedToNewFixedMap, FMovieSceneSequenceID LocalSequenceID, TSharedRef<UE::MovieScene::FSharedPlaybackState> SharedPlaybackState) {}
 
 	/** Get the referenced bindings for this section */
 	virtual void GetReferencedBindings(TArray<FGuid>& OutBindings) {}
diff --git a/Engine/Source/Runtime/MovieSceneTracks/Private/Sections/MovieScene3DConstraintSection.cpp b/Engine/Source/Runtime/MovieSceneTracks/Private/Sections/MovieScene3DConstraintSection.cpp
--- a/Engine/Source/Runtime/MovieSceneTracks/Private/Sections/MovieScene3DConstraintSection.cpp
+++ b/Engine/Source/Runtime/MovieSceneTracks/Private/Sections/MovieScene3DConstraintSection.cpp
@@ -14,14 +14,15 @@ UMovieScene3DConstraintSection::UMovieScene3DConstraintSection( const FObjectIni
 	bSupportsInfiniteRange = true;
 }
 
-void UMovieScene3DConstraintSection::OnBindingIDsUpdated(const TMap<UE::MovieScene::FFixedObjectBindingID, UE::MovieScene::FFixedObjectBindingID>& OldFixedToNewFixedMap, FMovieSceneSequenceID LocalSequenceID, const FMovieSceneSequenceHierarchy* Hierarchy, IMovieScenePlayer& Player)
+void UMovieScene3DConstraintSection::OnBindingIDsUpdated(const TMap<UE::MovieScene::FFixedObjectBindingID, UE::MovieScene::FFixedObjectBindingID>& OldFixedToNewFixedMap, FMovieSceneSequenceID LocalSequenceID, TSharedRef<UE::MovieScene::FSharedPlaybackState> SharedPlaybackState)
 {
-	UE::MovieScene::FFixedObjectBindingID FixedBindingID = ConstraintBindingID.ResolveToFixed(LocalSequenceID, Player);
+	UE::MovieScene::FFixedObjectBindingID FixedBindingID = ConstraintBindingID.ResolveToFixed(LocalSequenceID, SharedPlaybackState);
 
 	if (OldFixedToNewFixedMap.Contains(FixedBindingID))
 	{
 		Modify();
 
+		const FMovieSceneSequenceHierarchy* Hierarchy = SharedPlaybackState->GetHierarchy();
 		SetConstraintBindingID(OldFixedToNewFixedMap[FixedBindingID].ConvertToRelative(LocalSequenceID, Hierarchy));
 	}
 }
diff --git a/Engine/Source/Runtime/MovieSceneTracks/Private/Sections/MovieScene3DTransformSection.cpp b/Engine/Source/Runtime/MovieSceneTracks/Private/Sections/MovieScene3DTransformSection.cpp
--- a/Engine/Source/Runtime/MovieSceneTracks/Private/Sections/MovieScene3DTransformSection.cpp
+++ b/Engine/Source/Runtime/MovieSceneTracks/Private/Sections/MovieScene3DTransformSection.cpp
@@ -412,7 +412,7 @@ UMovieScene3DTransformSection::UMovieScene3DTransformSection(const FObjectInitia
 	Scale[2].SetDefault(1.f);
 }
 
-void UMovieScene3DTransformSection::OnBindingIDsUpdated(const TMap<UE::MovieScene::FFixedObjectBindingID, UE::MovieScene::FFixedObjectBindingID>& OldFixedToNewFixedMap, FMovieSceneSequenceID LocalSequenceID, const FMovieSceneSequenceHierarchy* Hierarchy, IMovieScenePlayer& Player)
+void UMovieScene3DTransformSection::OnBindingIDsUpdated(const TMap<UE::MovieScene::FFixedObjectBindingID, UE::MovieScene::FFixedObjectBindingID>& OldFixedToNewFixedMap, FMovieSceneSequenceID LocalSequenceID, TSharedRef<UE::MovieScene::FSharedPlaybackState> SharedPlaybackState)
 {
 	if (Constraints)
 	{
@@ -423,7 +423,7 @@ void UMovieScene3DTransformSection::OnBindingIDsUpdated(const TMap<UE::MovieScen
 				//Don't do child's we do that in the system, needed for duplication
 				if (TransformConstraint->ParentTRSHandle)
 				{
-					TransformConstraint->ParentTRSHandle->OnBindingIDsUpdated(OldFixedToNewFixedMap, LocalSequenceID, Hierarchy, Player);
+					TransformConstraint->ParentTRSHandle->OnBindingIDsUpdated(OldFixedToNewFixedMap, LocalSequenceID, SharedPlaybackState);
 				}
 			}
 		}
diff --git a/Engine/Source/Runtime/MovieSceneTracks/Private/Sections/MovieSceneActorReferenceSection.cpp b/Engine/Source/Runtime/MovieSceneTracks/Private/Sections/MovieSceneActorReferenceSection.cpp
--- a/Engine/Source/Runtime/MovieSceneTracks/Private/Sections/MovieSceneActorReferenceSection.cpp
+++ b/Engine/Source/Runtime/MovieSceneTracks/Private/Sections/MovieSceneActorReferenceSection.cpp
@@ -148,14 +148,15 @@ void UMovieSceneActorReferenceSection::PostLoad()
 	}
 }
 
-void UMovieSceneActorReferenceSection::OnBindingIDsUpdated(const TMap<UE::MovieScene::FFixedObjectBindingID, UE::MovieScene::FFixedObjectBindingID>& OldFixedToNewFixedMap, FMovieSceneSequenceID LocalSequenceID, const FMovieSceneSequenceHierarchy* Hierarchy, IMovieScenePlayer& Player)
+void UMovieSceneActorReferenceSection::OnBindingIDsUpdated(const TMap<UE::MovieScene::FFixedObjectBindingID, UE::MovieScene::FFixedObjectBindingID>& OldFixedToNewFixedMap, FMovieSceneSequenceID LocalSequenceID, TSharedRef<UE::MovieScene::FSharedPlaybackState> SharedPlaybackState)
 {
-	UE::MovieScene::FFixedObjectBindingID DefaultFixedBindingID = ActorReferenceData.GetDefault().Object.ResolveToFixed(LocalSequenceID, Player);
+	UE::MovieScene::FFixedObjectBindingID DefaultFixedBindingID = ActorReferenceData.GetDefault().Object.ResolveToFixed(LocalSequenceID, SharedPlaybackState);
 
 	if (OldFixedToNewFixedMap.Contains(DefaultFixedBindingID))
 	{
 		Modify();
 
+		const FMovieSceneSequenceHierarchy* Hierarchy = SharedPlaybackState->GetHierarchy();
 		FMovieSceneActorReferenceKey NewDefaultValue = ActorReferenceData.GetDefault();
 		NewDefaultValue.Object = OldFixedToNewFixedMap[DefaultFixedBindingID].ConvertToRelative(LocalSequenceID, Hierarchy);
 
@@ -164,12 +165,13 @@ void UMovieSceneActorReferenceSection::OnBindingIDsUpdated(const TMap<UE::MovieS
 
 	for (FMovieSceneActorReferenceKey& Key : ActorReferenceData.GetData().GetValues())
 	{
-		UE::MovieScene::FFixedObjectBindingID KeyFixedBindingID = Key.Object.ResolveToFixed(LocalSequenceID, Player);
+		UE::MovieScene::FFixedObjectBindingID KeyFixedBindingID = Key.Object.ResolveToFixed(LocalSequenceID, SharedPlaybackState);
 
 		if (OldFixedToNewFixedMap.Contains(KeyFixedBindingID))
 		{
 			Modify();
 
+			const FMovieSceneSequenceHierarchy* Hierarchy = SharedPlaybackState->GetHierarchy();
 			Key.Object = OldFixedToNewFixedMap[KeyFixedBindingID].ConvertToRelative(LocalSequenceID, Hierarchy);
 		}
 	}
diff --git a/Engine/Source/Runtime/MovieSceneTracks/Private/Sections/MovieSceneCameraCutSection.cpp b/Engine/Source/Runtime/MovieSceneTracks/Private/Sections/MovieSceneCameraCutSection.cpp
--- a/Engine/Source/Runtime/MovieSceneTracks/Private/Sections/MovieSceneCameraCutSection.cpp
+++ b/Engine/Source/Runtime/MovieSceneTracks/Private/Sections/MovieSceneCameraCutSection.cpp
@@ -38,14 +38,15 @@ UMovieSceneCameraCutSection::UMovieSceneCameraCutSection(const FObjectInitialize
 	SetBlendType(EMovieSceneBlendType::Absolute);
 }
 
-void UMovieSceneCameraCutSection::OnBindingIDsUpdated(const TMap<UE::MovieScene::FFixedObjectBindingID, UE::MovieScene::FFixedObjectBindingID>& OldFixedToNewFixedMap, FMovieSceneSequenceID LocalSequenceID, const FMovieSceneSequenceHierarchy* Hierarchy, IMovieScenePlayer& Player)
+void UMovieSceneCameraCutSection::OnBindingIDsUpdated(const TMap<UE::MovieScene::FFixedObjectBindingID, UE::MovieScene::FFixedObjectBindingID>& OldFixedToNewFixedMap, FMovieSceneSequenceID LocalSequenceID, TSharedRef<UE::MovieScene::FSharedPlaybackState> SharedPlaybackState)
 {
-	UE::MovieScene::FFixedObjectBindingID FixedBindingID = CameraBindingID.ResolveToFixed(LocalSequenceID, Player);
+	UE::MovieScene::FFixedObjectBindingID FixedBindingID = CameraBindingID.ResolveToFixed(LocalSequenceID, SharedPlaybackState);
 
 	if (OldFixedToNewFixedMap.Contains(FixedBindingID))
 	{
 		Modify();
 
+		const FMovieSceneSequenceHierarchy* Hierarchy = SharedPlaybackState->GetHierarchy();
 		CameraBindingID = OldFixedToNewFixedMap[FixedBindingID].ConvertToRelative(LocalSequenceID, Hierarchy);
 	}
 }
diff --git a/Engine/Source/Runtime/MovieSceneTracks/Private/Systems/MovieSceneConstraintSystem.cpp b/Engine/Source/Runtime/MovieSceneTracks/Private/Systems/MovieSceneConstraintSystem.cpp
--- a/Engine/Source/Runtime/MovieSceneTracks/Private/Systems/MovieSceneConstraintSystem.cpp
+++ b/Engine/Source/Runtime/MovieSceneTracks/Private/Systems/MovieSceneConstraintSystem.cpp
@@ -172,6 +172,7 @@ struct FEvaluateConstraintChannels
 		}
 
 		const FSequenceInstance& TargetInstance = InstanceRegistry->GetInstance(InstanceHandle);
+		TSharedRef<FSharedPlaybackState> SharedPlaybackState = TargetInstance.GetSharedPlaybackState();
 
 		bool Result = false;
 		ConstraintAndActiveChannel->ActiveChannel.Evaluate(FrameTime, Result);
@@ -182,7 +183,7 @@ struct FEvaluateConstraintChannels
 			TransformConstraint->InitConstraint(World);
 
 			// this has to be done once the constraint initialized
-			TransformConstraint->ResolveBoundObjects(TargetInstance.GetSequenceID(), *TargetInstance.GetPlayer());
+			TransformConstraint->ResolveBoundObjects(TargetInstance.GetSequenceID(), SharedPlaybackState);
 			
 			if (UTransformableComponentHandle* ComponentHandle = Cast<UTransformableComponentHandle>(TransformConstraint->ChildTRSHandle))
 			{
@@ -197,7 +198,7 @@ struct FEvaluateConstraintChannels
 		}
 		else
 		{
-			Constraint->ResolveBoundObjects(TargetInstance.GetSequenceID(), *TargetInstance.GetPlayer());	
+			Constraint->ResolveBoundObjects(TargetInstance.GetSequenceID(), SharedPlaybackState);	
 		}
 	}
 
diff --git a/Engine/Source/Runtime/MovieSceneTracks/Public/Sections/MovieScene3DConstraintSection.h b/Engine/Source/Runtime/MovieSceneTracks/Public/Sections/MovieScene3DConstraintSection.h
--- a/Engine/Source/Runtime/MovieSceneTracks/Public/Sections/MovieScene3DConstraintSection.h
+++ b/Engine/Source/Runtime/MovieSceneTracks/Public/Sections/MovieScene3DConstraintSection.h
@@ -50,7 +50,7 @@ public:
 
 	//~ UMovieSceneSection interface
 
-	virtual void OnBindingIDsUpdated(const TMap<UE::MovieScene::FFixedObjectBindingID, UE::MovieScene::FFixedObjectBindingID>& OldFixedToNewFixedMap, FMovieSceneSequenceID LocalSequenceID, const FMovieSceneSequenceHierarchy* Hierarchy, IMovieScenePlayer& Player) override;
+	virtual void OnBindingIDsUpdated(const TMap<UE::MovieScene::FFixedObjectBindingID, UE::MovieScene::FFixedObjectBindingID>& OldFixedToNewFixedMap, FMovieSceneSequenceID LocalSequenceID, TSharedRef<UE::MovieScene::FSharedPlaybackState> SharedPlaybackState) override;
 	
 	virtual void GetReferencedBindings(TArray<FGuid>& OutBindings) override;
 
diff --git a/Engine/Source/Runtime/MovieSceneTracks/Public/Sections/MovieScene3DTransformSection.h b/Engine/Source/Runtime/MovieSceneTracks/Public/Sections/MovieScene3DTransformSection.h
--- a/Engine/Source/Runtime/MovieSceneTracks/Public/Sections/MovieScene3DTransformSection.h
+++ b/Engine/Source/Runtime/MovieSceneTracks/Public/Sections/MovieScene3DTransformSection.h
@@ -234,7 +234,7 @@ public:
 	
 	MOVIESCENETRACKS_API virtual bool ShowCurveForChannel(const void *Channel) const override;
 	MOVIESCENETRACKS_API virtual void SetBlendType(EMovieSceneBlendType InBlendType) override;
-	MOVIESCENETRACKS_API virtual void OnBindingIDsUpdated(const TMap<UE::MovieScene::FFixedObjectBindingID, UE::MovieScene::FFixedObjectBindingID>& OldFixedToNewFixedMap, FMovieSceneSequenceID LocalSequenceID, const FMovieSceneSequenceHierarchy* Hierarchy, IMovieScenePlayer& Player) override;
+	MOVIESCENETRACKS_API virtual void OnBindingIDsUpdated(const TMap<UE::MovieScene::FFixedObjectBindingID, UE::MovieScene::FFixedObjectBindingID>& OldFixedToNewFixedMap, FMovieSceneSequenceID LocalSequenceID, TSharedRef<UE::MovieScene::FSharedPlaybackState> SharedPlaybackState) override;
 	MOVIESCENETRACKS_API virtual void GetReferencedBindings(TArray<FGuid>& OutBindings) override;
 	MOVIESCENETRACKS_API virtual void PreSave(FObjectPreSaveContext SaveContext) override;
 
diff --git a/Engine/Source/Runtime/MovieSceneTracks/Public/Sections/MovieSceneActorReferenceSection.h b/Engine/Source/Runtime/MovieSceneTracks/Public/Sections/MovieSceneActorReferenceSection.h
--- a/Engine/Source/Runtime/MovieSceneTracks/Public/Sections/MovieSceneActorReferenceSection.h
+++ b/Engine/Source/Runtime/MovieSceneTracks/Public/Sections/MovieSceneActorReferenceSection.h
@@ -194,7 +194,7 @@ public:
 	virtual void PostLoad() override;
 
 	//~ UMovieSceneSection interface
-	virtual void OnBindingIDsUpdated(const TMap<UE::MovieScene::FFixedObjectBindingID, UE::MovieScene::FFixedObjectBindingID>& OldFixedToNewFixedMap, FMovieSceneSequenceID LocalSequenceID, const FMovieSceneSequenceHierarchy* Hierarchy, IMovieScenePlayer& Player) override;
+	virtual void OnBindingIDsUpdated(const TMap<UE::MovieScene::FFixedObjectBindingID, UE::MovieScene::FFixedObjectBindingID>& OldFixedToNewFixedMap, FMovieSceneSequenceID LocalSequenceID, TSharedRef<UE::MovieScene::FSharedPlaybackState> SharedPlaybackState) override;
 
 	const FMovieSceneActorReferenceData& GetActorReferenceData() const { return ActorReferenceData; }
 
diff --git a/Engine/Source/Runtime/MovieSceneTracks/Public/Sections/MovieSceneCameraCutSection.h b/Engine/Source/Runtime/MovieSceneTracks/Public/Sections/MovieSceneCameraCutSection.h
--- a/Engine/Source/Runtime/MovieSceneTracks/Public/Sections/MovieSceneCameraCutSection.h
+++ b/Engine/Source/Runtime/MovieSceneTracks/Public/Sections/MovieSceneCameraCutSection.h
@@ -52,7 +52,7 @@ public:
 	}
 
 	//~ UMovieSceneSection interface
-	virtual void OnBindingIDsUpdated(const TMap<UE::MovieScene::FFixedObjectBindingID, UE::MovieScene::FFixedObjectBindingID>& OldFixedToNewFixedMap, FMovieSceneSequenceID LocalSequenceID, const FMovieSceneSequenceHierarchy* Hierarchy, IMovieScenePlayer& Player) override;
+	virtual void OnBindingIDsUpdated(const TMap<UE::MovieScene::FFixedObjectBindingID, UE::MovieScene::FFixedObjectBindingID>& OldFixedToNewFixedMap, FMovieSceneSequenceID LocalSequenceID, TSharedRef<UE::MovieScene::FSharedPlaybackState> SharedPlaybackState) override;
 	virtual void GetReferencedBindings(TArray<FGuid>& OutBindings) override;
 
 	/** ~UObject interface */
