Geometry collection - memory optimization - avoid over allocation when creating attributes

diff --git a/Engine/Source/Runtime/Experimental/Chaos/Public/GeometryCollection/ManagedArrayCollection.h b/Engine/Source/Runtime/Experimental/Chaos/Public/GeometryCollection/ManagedArrayCollection.h
--- a/Engine/Source/Runtime/Experimental/Chaos/Public/GeometryCollection/ManagedArrayCollection.h
+++ b/Engine/Source/Runtime/Experimental/Chaos/Public/GeometryCollection/ManagedArrayCollection.h
@@ -671,6 +671,7 @@ private:
 			, ManagedArray(SharedManagedArray.Get())
 		{
 			MakeUniqueForWrite();
+			ManagedArray->Reserve(InitialSize); // make sure we have the exact size to avoid over-allocation
 			ManagedArray->Resize(InitialSize);
 		};
 
@@ -687,6 +688,7 @@ private:
 			, SharedManagedArray(nullptr)
 			, ManagedArray(ExternalArrayPtr)
 		{
+			ManagedArray->Reserve(InitialSize); // make sure we have the exact size to avoid over-allocation
 			ManagedArray->Resize(InitialSize);
 		};
 
