Dataflow : fix template type fix code that would only work with anytype
- this is fallout of anytype imporvement that were made no so long ago
- now properly set the new type

diff --git a/Engine/Source/Runtime/Experimental/Dataflow/Core/Private/Dataflow/DataflowInputOutput.cpp b/Engine/Source/Runtime/Experimental/Dataflow/Core/Private/Dataflow/DataflowInputOutput.cpp
--- a/Engine/Source/Runtime/Experimental/Dataflow/Core/Private/Dataflow/DataflowInputOutput.cpp
+++ b/Engine/Source/Runtime/Experimental/Dataflow/Core/Private/Dataflow/DataflowInputOutput.cpp
@@ -89,8 +89,9 @@ void FDataflowInput::PullValue(Dataflow::FContext& Context) const
 void FDataflowInput::FixAndPropagateType(FName InType)
 {
 	check(InType.ToString().StartsWith(Type.ToString()));
-	const bool bSuccess = OwningNode->TrySetConnectionType(this, InType);
-	if (bSuccess)
+	// first try to fix as a legal change to give a chance for the node to proapagte the type changes if necessary
+	const bool bNeedToPropagate = OwningNode->TrySetConnectionType(this, InType);
+	if (bNeedToPropagate)
 	{
 		if (FDataflowReRouteNode* ReRouteNode = OwningNode->AsType<FDataflowReRouteNode>())
 		{
@@ -100,6 +101,12 @@ void FDataflowInput::FixAndPropagateType(FName InType)
 			}
 		}
 	}
+	else
+	{
+		// Simply fix the type  ( at this point we should not be a anytype )
+		check(!bIsAnyType);
+		Type = InType;
+	}
 }
 
 FDataflowArrayInput::FDataflowArrayInput(int32 InIndex, const Dataflow::FArrayInputParameters& Param)
@@ -292,8 +299,9 @@ void FDataflowOutput::ForwardInput(const FDataflowInput* Input, Dataflow::FConte
 void FDataflowOutput::FixAndPropagateType(FName InType)
 {
 	check(InType.ToString().StartsWith(Type.ToString()));
-	const bool bSuccess = OwningNode->TrySetConnectionType(this, InType);
-	if (bSuccess)
+	// first try to fix as a legal change to give a chance for the node to proapagte the type changes if necessary
+	const bool bNeedToPropagate = OwningNode->TrySetConnectionType(this, InType);
+	if (bNeedToPropagate)
 	{
 		if (FDataflowReRouteNode* ReRouteNode = OwningNode->AsType<FDataflowReRouteNode>())
 		{
@@ -306,4 +314,10 @@ void FDataflowOutput::FixAndPropagateType(FName InType)
 			}
 		}
 	}
+	else
+	{
+		// Simply fix the type  ( at this point we should not be a anytype )
+		check(!bIsAnyType);
+		Type = InType;
+	}
 }
\ No newline at end of file
