Modular vehicle option to manually override simulation order

diff --git a/Engine/Plugins/Experimental/ChaosModularVehicle/Source/ChaosModularVehicleEngine/Private/ChaosModularVehicle/ModularVehicleBuilder.cpp b/Engine/Plugins/Experimental/ChaosModularVehicle/Source/ChaosModularVehicleEngine/Private/ChaosModularVehicle/ModularVehicleBuilder.cpp
--- a/Engine/Plugins/Experimental/ChaosModularVehicle/Source/ChaosModularVehicleEngine/Private/ChaosModularVehicle/ModularVehicleBuilder.cpp
+++ b/Engine/Plugins/Experimental/ChaosModularVehicle/Source/ChaosModularVehicleEngine/Private/ChaosModularVehicle/ModularVehicleBuilder.cpp
@@ -22,6 +22,8 @@ void FModularVehicleBuilder::GenerateSimTree(UModularVehicleBaseComponent* Modul
 		{
 			SimModuleTree->SetAnimationEnabled(RequiresAnimation);
 
+			SimModuleTree->SetSimTreeProcessingOrder(ModularVehicle->GetTreeProcessingOrder());
+
 			// Physics thread takes ownership of the tree from here
 			ModularVehicle->VehicleSimulationPT->Initialize(SimModuleTree);
 		}
diff --git a/Engine/Plugins/Experimental/ChaosModularVehicle/Source/ChaosModularVehicleEngine/Public/ChaosModularVehicle/ModularVehicleBaseComponent.h b/Engine/Plugins/Experimental/ChaosModularVehicle/Source/ChaosModularVehicleEngine/Public/ChaosModularVehicle/ModularVehicleBaseComponent.h
--- a/Engine/Plugins/Experimental/ChaosModularVehicle/Source/ChaosModularVehicleEngine/Public/ChaosModularVehicle/ModularVehicleBaseComponent.h
+++ b/Engine/Plugins/Experimental/ChaosModularVehicle/Source/ChaosModularVehicleEngine/Public/ChaosModularVehicle/ModularVehicleBaseComponent.h
@@ -110,6 +110,8 @@ public:
 
 	APlayerController* GetPlayerController() const;
 	bool IsLocallyControlled() const;
+	void SetTreeProcessingOrder(Chaos::ESimTreeProcessingOrder TreeProcessingOrderIn) { TreeProcessingOrder = TreeProcessingOrderIn; }
+	Chaos::ESimTreeProcessingOrder GetTreeProcessingOrder() { return TreeProcessingOrder; }
 
 	virtual void GetLifetimeReplicatedProps(TArray<FLifetimeProperty>& OutLifetimeProps) const override;
 
@@ -368,6 +370,7 @@ private:
 	int32 ClusteringCount = 0;
 
 	bool bIsLocallyControlled;
+	Chaos::ESimTreeProcessingOrder TreeProcessingOrder = Chaos::ESimTreeProcessingOrder::LeafFirst;
 
 	TArray<FModuleAnimationSetup> ModuleAnimationSetups;
 
diff --git a/Engine/Source/Runtime/Experimental/ChaosVehicles/ChaosVehiclesCore/Private/SimModule/SimModuleTree.cpp b/Engine/Source/Runtime/Experimental/ChaosVehicles/ChaosVehiclesCore/Private/SimModule/SimModuleTree.cpp
--- a/Engine/Source/Runtime/Experimental/ChaosVehicles/ChaosVehiclesCore/Private/SimModule/SimModuleTree.cpp
+++ b/Engine/Source/Runtime/Experimental/ChaosVehicles/ChaosVehiclesCore/Private/SimModule/SimModuleTree.cpp
@@ -248,7 +248,7 @@ void FSimModuleTree::SimulateNode(float DeltaTime, FAllInputs& Inputs, int NodeI
 {
 	if (ISimulationModuleBase* Module = AccessSimModule(NodeIndex))
 	{
-		if (SimTreeProcessingOrder == ESimTreeProcessingOrder::RootFirst)
+		if (SimTreeProcessingOrder == ESimTreeProcessingOrder::RootFirst || SimTreeProcessingOrder == ESimTreeProcessingOrder::ManualOverride)
 		{
 			if (Module->IsEnabled())
 			{
@@ -261,9 +261,12 @@ void FSimModuleTree::SimulateNode(float DeltaTime, FAllInputs& Inputs, int NodeI
 			}
 		}
 
-		for (int ChildIdx : GetChildren(NodeIndex))
+		if (SimTreeProcessingOrder != ESimTreeProcessingOrder::ManualOverride)
 		{
-			SimulateNode(DeltaTime, Inputs, ChildIdx, PhysicsProxy);
+			for (int ChildIdx : GetChildren(NodeIndex))
+			{
+				SimulateNode(DeltaTime, Inputs, ChildIdx, PhysicsProxy);
+			}
 		}
 
 		if (SimTreeProcessingOrder == ESimTreeProcessingOrder::LeafFirst)
diff --git a/Engine/Source/Runtime/Experimental/ChaosVehicles/ChaosVehiclesCore/Public/SimModule/SimModuleTree.h b/Engine/Source/Runtime/Experimental/ChaosVehicles/ChaosVehiclesCore/Public/SimModule/SimModuleTree.h
--- a/Engine/Source/Runtime/Experimental/ChaosVehicles/ChaosVehiclesCore/Public/SimModule/SimModuleTree.h
+++ b/Engine/Source/Runtime/Experimental/ChaosVehicles/ChaosVehiclesCore/Public/SimModule/SimModuleTree.h
@@ -20,8 +20,9 @@ namespace Chaos
 
 	enum ESimTreeProcessingOrder : int8
 	{
-		LeafFirst = 0,
-		RootFirst = 1
+		ManualOverride = 0,	// User calls simulate on the child modules
+		LeafFirst = 1,		// modules simulation from the leaf first
+		RootFirst = 2		// modules simulate from the root first
 	};
 
 	struct FPendingModuleAdds
