More thorough handling of FIoBuffer user supplies.
Skip all the FIoBuffer stuff if there is no content.
Strategic handling of reading ahead.

#rnx

diff --git a/Engine/Source/Runtime/Experimental/IoStore/HttpClient/Private/Loop.inl b/Engine/Source/Runtime/Experimental/IoStore/HttpClient/Private/Loop.inl
--- a/Engine/Source/Runtime/Experimental/IoStore/HttpClient/Private/Loop.inl
+++ b/Engine/Source/Runtime/Experimental/IoStore/HttpClient/Private/Loop.inl
@@ -226,19 +226,28 @@ static FOutcome DoRecvMessage(FActivity* Activity, FPeerType& Peer)
 		FTicketStatus& SinkArg = *(FTicketStatus*)Activity;
 		Activity->Sink(SinkArg);
 	}
+	Activity->NoContent |= (ContentLength == 0);
 
+	// Check the user gave us a destination for content
+	FIoBuffer& Dest = *(Activity->Dest);
 	if (Activity->NoContent == 0)
 	{
-		if (Activity->Dest == PriorDest)
+		if (&Dest == PriorDest)
 		{
 			Activity_SetError(Activity, "User did not provide a destination buffer");
 			return FOutcome::Error(Activity->ErrorReason);
 		}
 
 		// The user seems to have forgotten something. Let's help them along
-		if (Activity->Dest->GetSize() == 0)
+		if (int32 DestSize = int32(Dest.GetSize()); DestSize == 0)
 		{
-			*Activity->Dest = FIoBuffer(ContentLength);
+			Dest = FIoBuffer(ContentLength);
+		}
+		else if (DestSize < ContentLength)
+		{
+			// todo: support piece-wise transfer of content (a la chunked).
+			Activity_SetError(Activity, "Destination buffer too small");
+			return FOutcome::Error(Activity->ErrorReason);
 		}
 	}
 
@@ -269,34 +278,22 @@ static FOutcome DoRecvMessage(FActivity* Activity, FPeerType& Peer)
 	auto NextState = bStreamed ? FActivity::EState::RecvStream : FActivity::EState::RecvContent;
 	Activity_ChangeState(Activity, NextState, AlreadyReceived);
 
+	// Copy any of the content we may have already received.
 	if (AlreadyReceived == 0)
 	{
 		return FOutcome::Ok();
 	}
 
-	FMutableMemoryView DestView = Activity->Dest->GetMutableView();
-	const char* Cursor = BufferRight - AlreadyReceived;
-	if (!bStreamed || AlreadyReceived < DestView.GetSize())
+	// This ordinarily doesn't happen due to the way higher levels pipeline
+	// requests. It can however occur with chunked transfers.
+	if (AlreadyReceived > Dest.GetSize())
 	{
-		::memcpy(DestView.GetData(), Cursor, AlreadyReceived);
-		return FOutcome::Ok();
+		Dest = FIoBuffer(AlreadyReceived);
 	}
 
-#if 1
-	check(false); // not implemented yet
-#else
-	do
-	{
-		uint32 Size = FMath::Min<uint32>(DestView.GetSize(), AlreadyReceived);
-		::memcpy(DestView.GetData(), Cursor, Size);
-
-		/* send sink */
-
-		AlreadyReceived -= Size;
-		Cursor += Size;
-	}
-	while (AlreadyReceived);
-#endif
+	FMutableMemoryView DestView = Dest.GetMutableView();
+	const char* Cursor = BufferRight - AlreadyReceived;
+	::memcpy(DestView.GetData(), Cursor, AlreadyReceived);
 
 	return FOutcome::Ok();
 }
@@ -1498,9 +1495,6 @@ FTicket FEventLoop::Send(FRequest&& Request, FTicketSink Sink, UPTRINT SinkParam
 	return Impl->Send(Activity);
 }
 
-
-
-
 // }}}
 
 } // namespace UE::IoStore::HTTP
