Send sends immediately, don't wait for acknowledgements or assume batch of small writes (we only make one)

#rnx

diff --git a/Engine/Source/Runtime/Experimental/IoStore/HttpClient/Private/Socket.inl b/Engine/Source/Runtime/Experimental/IoStore/HttpClient/Private/Socket.inl
--- a/Engine/Source/Runtime/Experimental/IoStore/HttpClient/Private/Socket.inl
+++ b/Engine/Source/Runtime/Experimental/IoStore/HttpClient/Private/Socket.inl
@@ -39,6 +39,7 @@
 #	include <arpa/inet.h>
 #	include <fcntl.h>
 #	include <netdb.h>
+#	include <netinet/tcp.h>
 #	include <poll.h>
 #	include <sys/select.h>
 #	include <sys/socket.h>
@@ -142,6 +143,9 @@ bool FSocket::Create()
 		return false;
 	}
 
+	int32 Yes = 1;
+	setsockopt(Socket, IPPROTO_TCP, TCP_NODELAY, &(char&)(Yes), sizeof(Yes));
+
 	Trace(Socket, ETrace::SocketCreate);
 	return true;
 }
