Use IAD tagsets instead of individual asset references to support partial GFP install.

Also fix some FOnDemandIoStore issues
- Support FOnDemandToc being loaded by game code
- Some support for OnDemandIoStore cache directory having different names depending on target - still needs more work to support forking servers.
- Fix BuildInstallData not returning missing packages
- Make sure BuildInstallData installs always downloaded chunks for a container without any packages (not sure this can actually occur)
- Disable pak mount callbacks - these don't work right and causes errors without global HTTP endpoint configured.
- Fix FOnDemandIoStore not being created without a global HTTP endpoint configured.

diff --git a/Engine/Plugins/Runtime/GameFeatures/Source/GameFeatures/Private/GameFeaturePluginStateMachine.cpp b/Engine/Plugins/Runtime/GameFeatures/Source/GameFeatures/Private/GameFeaturePluginStateMachine.cpp
--- a/Engine/Plugins/Runtime/GameFeatures/Source/GameFeatures/Private/GameFeaturePluginStateMachine.cpp
+++ b/Engine/Plugins/Runtime/GameFeatures/Source/GameFeatures/Private/GameFeaturePluginStateMachine.cpp
@@ -98,7 +98,7 @@ namespace UE::GameFeatures
 		TEXT("Enable to make block deactivation until all dependencies are deactivated. Warning - this can lead to failure to unload"));
 
 	static TAutoConsoleVariable<bool> CVarEnableAssetStreaming(TEXT("GameFeaturePlugin.EnableAssetStreaming"),
-		false,
+		true,
 		TEXT("Enable experimental asset streaming"));
 
 	bool ShouldDeferLocalizationDataLoad()
@@ -2296,6 +2296,8 @@ struct FGameFeaturePluginState_WaitingForDependencies : public FTransitionDepend
 	}
 };
 
+// TODO: this is mostly going to look like the download state with a few modifications. Should probably just have a
+// common base with it
 struct FGameFeaturePluginState_AssetDependencyStreaming : public FGameFeaturePluginState
 {
 	FGameFeaturePluginState_AssetDependencyStreaming(FGameFeaturePluginStateMachineProperties& InStateProperties) 
@@ -2310,6 +2312,7 @@ struct FGameFeaturePluginState_AssetDependencyStreaming : public FGameFeaturePlu
 
 	void Cleanup()
 	{
+		IInstallBundleManager::InstallBundleCompleteDelegate.RemoveAll(this);
 		Result = MakeValue();
 		PendingBundleDownloads.Empty();
 	}
@@ -2356,10 +2359,16 @@ struct FGameFeaturePluginState_AssetDependencyStreaming : public FGameFeaturePlu
 	{
 		Cleanup();
 
+		if (StateProperties.GetPluginProtocol() != EGameFeaturePluginProtocol::InstallBundle)
+		{
+			// For now, we don't support asset streaming with file protocol
+			return;
+		}
+
 		// TODO: Install Bundles need to move away from FNames for identifiers, this is currently just bloating up the name table
 		// when using dynamic GFPs
-		TArray<FGameFeaturePluginDependency> AssetDependencies = UGameFeaturesSubsystem::Get().FindPluginAssetDependencies(StateProperties.PluginInstalledFilename);
-		TValueOrError<TArray<FName>, FString> MaybeAssetInstallBundles = UGameFeaturesSubsystem::Get().GetPolicy().GetStreamingAssetInstallBundles(AssetDependencies);
+		TValueOrError<TArray<FName>, FString> MaybeAssetInstallBundles = 
+			UGameFeaturesSubsystem::Get().GetPolicy().GetStreamingAssetInstallBundles(StateProperties.PluginIdentifier.GetFullPluginURL());
 		if (MaybeAssetInstallBundles.HasError())
 		{
 			Result = GetErrorResult(MaybeAssetInstallBundles.GetError());
@@ -2372,22 +2381,24 @@ struct FGameFeaturePluginState_AssetDependencyStreaming : public FGameFeaturePlu
 			return;
 		}
 
-		// Respect DoNotDownload flag
-		// TODO: move DoNotDownload flag to base protocol options? We could now have a file: GFP that needs to stream dependencies
-		if (StateProperties.ProtocolOptions.HasSubtype<FInstallBundlePluginProtocolOptions>())
+		const FInstallBundlePluginProtocolOptions& Options = StateProperties.ProtocolOptions.GetSubtype<FInstallBundlePluginProtocolOptions>();
+		if (Options.bDoNotDownload)
 		{
-			const FInstallBundlePluginProtocolOptions& Options = StateProperties.ProtocolOptions.GetSubtype<FInstallBundlePluginProtocolOptions>();
-			if (Options.bDoNotDownload)
-			{
-				Result = GetErrorResult(TEXT("GFPStateMachine.DownloadNotAllowed"));
-				return;
-			}
+			Result = GetErrorResult(TEXT("GFPStateMachine.DownloadNotAllowed"));
+			return;
+		}
+
+		FInstallBundlePluginProtocolMetaData& ProtocolData = StateProperties.ProtocolMetadata.GetSubtype<FInstallBundlePluginProtocolMetaData>();
+		for (FName Bundle : AssetInstallBundles)
+		{
+			// Add to metadata list for later cleanup
+			ProtocolData.InstallBundles.AddUnique(Bundle);
 		}
 
 		TSharedPtr<IInstallBundleManager> BundleManager = IInstallBundleManager::GetPlatformInstallBundleManager();
 
-		// TODO: where do flags come from? The dependency being streamed from?
-		EInstallBundleRequestFlags InstallFlags = EInstallBundleRequestFlags::None;
+		EInstallBundleRequestFlags InstallFlags = InstallFlags = UseAsyncLoading() ?
+			(Options.InstallBundleFlags | EInstallBundleRequestFlags::AsyncMount) : Options.InstallBundleFlags;
 		TValueOrError<FInstallBundleRequestInfo, EInstallBundleResult> MaybeRequestInfo = BundleManager->RequestUpdateContent(AssetInstallBundles, InstallFlags);
 		if (MaybeRequestInfo.HasError())
 		{
@@ -2399,7 +2410,10 @@ struct FGameFeaturePluginState_AssetDependencyStreaming : public FGameFeaturePlu
 
 		FInstallBundleRequestInfo RequestInfo = MaybeRequestInfo.StealValue();
 		PendingBundleDownloads = MoveTemp(RequestInfo.BundlesEnqueued);
-		IInstallBundleManager::InstallBundleCompleteDelegate.AddRaw(this, &FGameFeaturePluginState_AssetDependencyStreaming::OnInstallBundleCompleted);
+		if (!PendingBundleDownloads.IsEmpty())
+		{
+			IInstallBundleManager::InstallBundleCompleteDelegate.AddRaw(this, &FGameFeaturePluginState_AssetDependencyStreaming::OnInstallBundleCompleted);
+		}
 
 		// TODO: how to apply pausing? bUserPauseDownload is protocol specific
 		//if (Options.bUserPauseDownload)
diff --git a/Engine/Plugins/Runtime/GameFeatures/Source/GameFeatures/Private/GameFeaturesSubsystem.cpp b/Engine/Plugins/Runtime/GameFeatures/Source/GameFeatures/Private/GameFeaturesSubsystem.cpp
--- a/Engine/Plugins/Runtime/GameFeatures/Source/GameFeatures/Private/GameFeaturesSubsystem.cpp
+++ b/Engine/Plugins/Runtime/GameFeatures/Source/GameFeatures/Private/GameFeaturesSubsystem.cpp
@@ -2206,7 +2206,6 @@ bool UGameFeaturesSubsystem::GetGameFeaturePluginDetailsInternal(const FString&
 		const FStringView NameField = TEXTVIEW("Name");
 		const FStringView EnabledField = TEXTVIEW("Enabled");
 		const FStringView ActivateField = TEXTVIEW("Activate");
-		const FStringView AssetReferencesField = TEXTVIEW("AssetReferences");
 		for (const TSharedPtr<FJsonValue>& PluginElement : *PluginsArray)
 		{
 			if (!PluginElement)
@@ -2242,24 +2241,8 @@ bool UGameFeaturesSubsystem::GetGameFeaturePluginDetailsInternal(const FString&
 			bool bElementActivate = false;
 			ElementObject->TryGetBoolField(ActivateField, bElementActivate);
 
-			TArray<FString> AssetReferences;
-			const TArray<TSharedPtr<FJsonValue>>* AssetRefsArray = nullptr;
-			ElementObject->TryGetArrayField(AssetReferencesField, AssetRefsArray);
-			if (AssetRefsArray)
-			{
-				for (const TSharedPtr<FJsonValue>& AssetRefElement : *AssetRefsArray)
-				{
-					FString AssetRef;
-					if (AssetRefElement && AssetRefElement->TryGetString(AssetRef))
-					{
-						AssetReferences.Add(MoveTemp(AssetRef));
-					}
-				}
-			}
-
 			FGameFeaturePluginReferenceDetails& RefDetails = OutPluginDetails.PluginDependencies.Emplace_GetRef();
 			RefDetails.PluginName = MoveTemp(DependencyName);
-			RefDetails.AssetReferences = MoveTemp(AssetReferences);
 			RefDetails.bShouldActivate = bElementActivate;
 		}
 	}
@@ -2762,21 +2745,6 @@ void UGameFeaturesSubsystem::FinishTermination(UGameFeaturePluginStateMachine* M
 	TerminalGameFeaturePluginStateMachines.RemoveSwap(Machine);
 }
 
-TArray<FGameFeaturePluginDependency> UGameFeaturesSubsystem::FindPluginAssetDependencies(const FString& PluginDescriptorFilename)
-{
-	FGameFeaturePluginDetails Details;
-	ensure(GetGameFeaturePluginDetailsInternal(PluginDescriptorFilename, Details));
-
-	TArray<FGameFeaturePluginDependency> OutDeps;
-	OutDeps.Reserve(Details.PluginDependencies.Num());
-	for (FGameFeaturePluginReferenceDetails& RefDetails : Details.PluginDependencies)
-	{
-		OutDeps.Emplace(MoveTemp(RefDetails.PluginName), MoveTemp(RefDetails.AssetReferences));
-	}
-
-	return OutDeps;
-}
-
 bool UGameFeaturesSubsystem::FindOrCreatePluginDependencyStateMachines(const FString& PluginURL, const FGameFeaturePluginStateMachineProperties& InStateProperties, TArray<UGameFeaturePluginStateMachine*>& OutDependencyMachines)
 {
 	const FString& PluginFilename = InStateProperties.PluginInstalledFilename;
diff --git a/Engine/Plugins/Runtime/GameFeatures/Source/GameFeatures/Public/GameFeaturesProjectPolicies.h b/Engine/Plugins/Runtime/GameFeatures/Source/GameFeatures/Public/GameFeaturesProjectPolicies.h
--- a/Engine/Plugins/Runtime/GameFeatures/Source/GameFeatures/Public/GameFeaturesProjectPolicies.h
+++ b/Engine/Plugins/Runtime/GameFeatures/Source/GameFeatures/Public/GameFeaturesProjectPolicies.h
@@ -46,9 +46,8 @@ public:
 	// returns The dependency URL or an error if the dependency could not be resolved
 	virtual TValueOrError<FString, FString> ResolvePluginDependency(const FString& PluginURL, const FString& DependencyName) const;
 
-	// Called to resolve install bundles for stream asset dependencies, input maps a plugin dependency to referenced assets 
-	// (See UGameFeaturesSubsystem::FindPluginAssetDependencies)
-	virtual TValueOrError<TArray<FName>, FString> GetStreamingAssetInstallBundles(const TArray<FGameFeaturePluginDependency>& AssetDependencies) const { return MakeValue(); }
+	// Called to resolve install bundles for streaming asset dependencies
+	virtual TValueOrError<TArray<FName>, FString> GetStreamingAssetInstallBundles(const FString& PluginURL) const { return MakeValue(); }
 
 	// Called by code that explicitly wants to load a specific plugin
 	// (e.g., when using a fast cook a game might want to allow explicitly loaded game feature plugins)
diff --git a/Engine/Plugins/Runtime/GameFeatures/Source/GameFeatures/Public/GameFeaturesSubsystem.h b/Engine/Plugins/Runtime/GameFeatures/Source/GameFeatures/Public/GameFeaturesSubsystem.h
--- a/Engine/Plugins/Runtime/GameFeatures/Source/GameFeatures/Public/GameFeaturesSubsystem.h
+++ b/Engine/Plugins/Runtime/GameFeatures/Source/GameFeatures/Public/GameFeaturesSubsystem.h
@@ -233,7 +233,6 @@ void GAMEFEATURES_API LexFromString(EGameFeatureTargetState& Value, const TCHAR*
 struct FGameFeaturePluginReferenceDetails
 {
 	FString PluginName;
-	TArray<FString> AssetReferences;
 	bool bShouldActivate = false;
 };
 
@@ -386,9 +385,6 @@ struct FGameFeatureInfo
 	EGameFeaturePluginState CurrentState;
 };
 
-// Key - plugin reference, Value - any specific assets imported
-using FGameFeaturePluginDependency = TPair<FString, TArray<FString>>;
-
 /** The manager subsystem for game features */
 UCLASS()
 class GAMEFEATURES_API UGameFeaturesSubsystem : public UEngineSubsystem
@@ -731,9 +727,6 @@ private:
 	void FinishTermination(UGameFeaturePluginStateMachine* Machine);
 	friend class UGameFeaturePluginStateMachine;
 
-	TArray<FGameFeaturePluginDependency> FindPluginAssetDependencies(const FString& PluginDescriptorFilename);
-	friend struct FGameFeaturePluginState_AssetDependencyStreaming;
-
 	/** Handler for when a state machine requests its dependencies. Returns false if the dependencies could not be read */
 	bool FindOrCreatePluginDependencyStateMachines(const FString& PluginURL, const FGameFeaturePluginStateMachineProperties& InStateProperties, TArray<UGameFeaturePluginStateMachine*>& OutDependencyMachines);
 	template <typename> friend struct FTransitionDependenciesGameFeaturePluginState;
diff --git a/Engine/Source/Runtime/Experimental/IoStore/OnDemand/Private/OnDemandIoStore.cpp b/Engine/Source/Runtime/Experimental/IoStore/OnDemand/Private/OnDemandIoStore.cpp
--- a/Engine/Source/Runtime/Experimental/IoStore/OnDemand/Private/OnDemandIoStore.cpp
+++ b/Engine/Source/Runtime/Experimental/IoStore/OnDemand/Private/OnDemandIoStore.cpp
@@ -22,6 +22,7 @@
 #include "IO/PackageStore.h"
 #include "Misc/CommandLine.h"
 #include "Misc/CoreDelegatesInternal.h"
+#include "Misc/CoreMisc.h"
 #include "Misc/EncryptionKeyManager.h"
 #include "Misc/Guid.h"
 #include "Misc/Paths.h"
@@ -49,7 +50,40 @@ namespace Private
 ///////////////////////////////////////////////////////////////////////////////
 static FString GetInstallCacheDirectory()
 {
-	return FPaths::ProjectPersistentDownloadDir() / TEXT("IoStore") / TEXT("InstallCache");
+	FString DirName;
+
+	if (IsRunningDedicatedServer())
+	{
+#if 0 // TODO: Fix this for forking servers
+		if (!FForkProcessHelper::IsForkRequested())
+		{
+			DirName = TEXT("InstallCacheServer");
+		}
+		else
+		{
+			if (!FForkProcessHelper::IsForkedChildProcess())
+			{
+				UE_LOG(LogIoStoreOnDemand, Fatal, TEXT("Attempting to create IOStore cache before forking!"));
+			}
+
+			DirName = FString::Printf(TEXT("InstallCacheServer-%u"), FPlatformProcess::GetCurrentProcessId());
+		}
+#else
+		DirName = TEXT("InstallCacheServer");
+#endif
+	}
+#if WITH_EDITOR
+	else if (GIsEditor)
+	{
+		DirName = TEXT("InstallCacheEditor");
+	}
+#endif //if WITH_EDITOR
+	else
+	{
+		DirName = TEXT("InstallCache");
+	}
+
+	return FPaths::ProjectPersistentDownloadDir() / TEXT("IoStore") / DirName;
 }
 
 ///////////////////////////////////////////////////////////////////////////////
@@ -86,20 +120,21 @@ FIoStatus BuildInstallData(
 	const TSet<FSharedOnDemandContainer>& Containers,
 	const TSet<FPackageId>& PackageIds,
 	FInstallData& OutInstallData,
-	const TSet<FPackageId>& OutMissing)
+	TSet<FPackageId>& OutMissing)
 {
 	OutInstallData.Reserve(Containers.Num());
 
 	// Setup the package information for each container
 	for (const FSharedOnDemandContainer& Container : Containers)
 	{
+		FContainerInstallData& Data = OutInstallData.FindOrAdd(Container);
+
 		if (!Container->Header.IsValid() || Container->Header->PackageIds.IsEmpty())
 		{
 			// The container contains no package data
 			continue;
 		}
 
-		FContainerInstallData& Data = OutInstallData.FindOrAdd(Container);
 		const FIoContainerHeader& Header = *Container->Header;
 		TConstArrayView<FFilePackageStoreEntry> Entries(
 			reinterpret_cast<const FFilePackageStoreEntry*>(Header.StoreEntries.GetData()),
@@ -120,7 +155,6 @@ FIoStatus BuildInstallData(
 
 	FQueue Queue;
 	TSet<FPackageId> Visitied;
-	TSet<FPackageId> Missing;
 
 	Visitied.Reserve(PackageIds.Num());
 	Queue.Reserve(PackageIds.Num());
@@ -167,7 +201,7 @@ FIoStatus BuildInstallData(
 
 		if (!bFound)
 		{
-			Missing.Add(PackageId);
+			OutMissing.Add(PackageId);
 		}
 	}
 
@@ -176,11 +210,6 @@ FIoStatus BuildInstallData(
 		const FOnDemandContainer& Container = *Kv.Key;
 		FContainerInstallData& Data			= Kv.Value;
 
-		if (Data.PackageIds.IsEmpty())
-		{
-			continue;
-		}
-
 		for (const FPackageId& PackageId : Data.PackageIds)
 		{
 			const FIoChunkId PackageChunkId					= CreatePackageDataChunkId(PackageId);
@@ -348,6 +377,7 @@ FIoStatus FOnDemandIoStore::Initialize()
 		}
 	}
 
+#if 0
 	OnMountPakHandle = FCoreInternalDelegates::GetOnPakMountOperation().AddLambda(
 		[this](EMountOperation Operation, const TCHAR* ContainerPath, int32 Order) -> void
 		{
@@ -422,6 +452,7 @@ FIoStatus FOnDemandIoStore::Initialize()
 			TickLoop();
 		}
 	}
+#endif
 
 	return EIoErrorCode::Ok;
 }
@@ -675,9 +706,9 @@ FIoStatus FOnDemandIoStore::TickMountRequest(FMountRequest& MountRequest)
 	Private::SplitHostUrl(Args.Url, Host, TocRelUrl);
 	const FStringView TocPath = FPathViews::GetPath(TocRelUrl);
 
-	if (Args.Toc.IsSet())
+	if (Args.Toc)
 	{
-		CreateContainersFromToc(Args.MountId, TocPath, Args.Toc.GetValue(), MountRequest.Containers);
+		CreateContainersFromToc(Args.MountId, TocPath, *Args.Toc, MountRequest.Containers);
 	}
 	else if (Args.FilePath.IsEmpty() == false)
 	{
@@ -691,8 +722,7 @@ FIoStatus FOnDemandIoStore::TickMountRequest(FMountRequest& MountRequest)
 			return TocStatus.Status();
 		}
 
-		FOnDemandToc Toc = TocStatus.ConsumeValueOrDie();
-		Args.Toc.Emplace(MoveTemp(Toc));
+		Args.Toc = MakeShared<FOnDemandToc>(TocStatus.ConsumeValueOrDie());
 
 		CreateContainersFromToc(Args.MountId, TocPath, *Args.Toc, MountRequest.Containers);
 	}
@@ -709,8 +739,7 @@ FIoStatus FOnDemandIoStore::TickMountRequest(FMountRequest& MountRequest)
 			return TocStatus.Status();
 		}
 
-		FOnDemandToc Toc = TocStatus.ConsumeValueOrDie();
-		Args.Toc.Emplace(MoveTemp(Toc));
+		Args.Toc = MakeShared<FOnDemandToc>(TocStatus.ConsumeValueOrDie());
 
 		CreateContainersFromToc(Args.MountId, TocPath, *Args.Toc, MountRequest.Containers);
 	}
@@ -850,7 +879,7 @@ FIoStatus FOnDemandIoStore::TickInstallRequest(FMountRequest& MountRequest)
 	// Parse the tag sets
 	if (MountRequest.MountArgs.TagSets.IsEmpty() == false && MountRequest.TagSets.IsEmpty())
 	{
-		const FOnDemandToc& Toc = MountRequest.MountArgs.Toc.GetValue();
+		const FOnDemandToc& Toc = *MountRequest.MountArgs.Toc;
 		for (const FString& Tag : MountRequest.MountArgs.TagSets)
 		{
 			const FOnDemandTocTagSet* TagSet = Algo::FindBy(Toc.TagSets, Tag, &FOnDemandTocTagSet::Tag);
diff --git a/Engine/Source/Runtime/Experimental/IoStore/OnDemand/Public/IO/IoStoreOnDemand.h b/Engine/Source/Runtime/Experimental/IoStore/OnDemand/Public/IO/IoStoreOnDemand.h
--- a/Engine/Source/Runtime/Experimental/IoStore/OnDemand/Public/IO/IoStoreOnDemand.h
+++ b/Engine/Source/Runtime/Experimental/IoStore/OnDemand/Public/IO/IoStoreOnDemand.h
@@ -208,9 +208,9 @@ struct FOnDemandToc
 
 	static FGuid VersionGuid;
 
-	static TIoStatusOr<FOnDemandToc> LoadFromFile(const FString& FilePath, bool bValidate);
-	static TIoStatusOr<FOnDemandToc> LoadFromUrl(FAnsiStringView Url, uint32 RetryCount = 0, bool bFollowRedirects = false);
-	static TIoStatusOr<FOnDemandToc> LoadFromUrl(FStringView Url, uint32 RetryCount = 0, bool bFollowRedirects = false);
+	UE_API static TIoStatusOr<FOnDemandToc> LoadFromFile(const FString& FilePath, bool bValidate);
+	UE_API static TIoStatusOr<FOnDemandToc> LoadFromUrl(FAnsiStringView Url, uint32 RetryCount = 0, bool bFollowRedirects = false);
+	UE_API static TIoStatusOr<FOnDemandToc> LoadFromUrl(FStringView Url, uint32 RetryCount = 0, bool bFollowRedirects = false);
 };
 
 UE_API bool LoadFromCompactBinary(FCbFieldView Field, FOnDemandToc& OutToc);
@@ -251,7 +251,7 @@ ENUM_CLASS_FLAGS(EOnDemandMountOptions);
 struct FOnDemandMountArgs
 {
 	/** Mount an already serialized TOC. */
-	TOptional<FOnDemandToc> Toc;
+	TSharedPtr<FOnDemandToc> Toc;
 	/** Mandatory ID to be used for unmounting all container file(s) included in the TOC. */
 	FString MountId;
 	/** Download the TOC from the specified URL. */
@@ -261,7 +261,7 @@ struct FOnDemandMountArgs
 	/** Used with EOnDemandMountOptions::Install to fliter requested content */
 	TArray<FString> TagSets;
 	/** Mount options. */
-	EOnDemandMountOptions Options;
+	EOnDemandMountOptions Options = EOnDemandMountOptions::None;
 };
 
 struct FOnDemandMountResult
diff --git a/Engine/Source/Runtime/InstallBundleManager/Public/InstallBundleTypes.h b/Engine/Source/Runtime/InstallBundleManager/Public/InstallBundleTypes.h
--- a/Engine/Source/Runtime/InstallBundleManager/Public/InstallBundleTypes.h
+++ b/Engine/Source/Runtime/InstallBundleManager/Public/InstallBundleTypes.h
@@ -284,6 +284,7 @@ struct FInstallBundleSourceUpdateBundleInfo
 	FDateTime LastAccessTime = FDateTime::MinValue(); // If cached, used to decide eviction order
 	EInstallBundleInstallState BundleContentState = EInstallBundleInstallState::NotInstalled; // Whether this bundle is up to date
 	bool bIsCached = false; // Whether this bundle should be cached if this source has a bundle cache
+	// TODO: this appears to be unused
 	bool bUseChunkDBs = false; // Whether this bundle should attempt to use ChunkDBs at all
 };
 
