Iris - FieldOfViewNetObjectPrioritizer and tests.

The prioritizer consists of an inner sphere, outer sphere, a cone and a line of sight. Radii, field of view for cone, line of sight width and priorities are all configurable. Priorities cover the shape they're configured for except for the cone where the priority is linearly interpolated from the inner cone length to the full cone length. A 2D greyscale image is produced to memory by the VisualizeFieldOfViewNetObjectPrioritizer test.

diff --git a/Engine/Plugins/Runtime/ReplicationSystemTestPlugin/Source/Private/Tests/ReplicationSystem/Prioritization/TestFieldOfViewNetObjectPrioritizer.cpp b/Engine/Plugins/Runtime/ReplicationSystemTestPlugin/Source/Private/Tests/ReplicationSystem/Prioritization/TestFieldOfViewNetObjectPrioritizer.cpp
new file mode 100644
--- /dev/null
+++ b/Engine/Plugins/Runtime/ReplicationSystemTestPlugin/Source/Private/Tests/ReplicationSystem/Prioritization/TestFieldOfViewNetObjectPrioritizer.cpp
@@ -0,0 +1,227 @@
+// Copyright Epic Games, Inc. All Rights Reserved.
+
+#include "Tests/ReplicationSystem/Prioritization/TestFieldOfViewNetObjectPrioritizer.h"
+#include "Tests/ReplicationSystem/Prioritization/TestNetObjectPrioritizerFixture.h"
+
+
+UFieldOfViewNetObjectPrioritizerTestConfig::UFieldOfViewNetObjectPrioritizerTestConfig()
+{
+	// Intentioally left blank. Do not remove.
+}
+
+UFieldOfViewNetObjectPrioritizerForConeTestConfig::UFieldOfViewNetObjectPrioritizerForConeTestConfig()
+: Super()
+{
+	InnerSpherePriority = 0.0f;
+	OuterSpherePriority = 0.0f;
+	LineOfSightPriority = 0.0f;
+}
+
+namespace UE::Net::Private
+{
+
+class FTestFieldOfViewNetObjectPrioritizer : public FTestNetObjectPrioritizerFixture
+{
+protected:
+	using Super = FTestNetObjectPrioritizerFixture;
+
+	virtual void SetUp() override
+	{
+		Super::SetUp();
+		InitFoVNetObjectPrioritizer();
+	}
+
+	virtual void TearDown() override
+	{
+		FoVPrioritizerHandle = InvalidNetObjectPrioritizerHandle;
+		Super::TearDown();
+	}
+
+private:
+	// Called by Super::SetUp()
+	virtual void GetPrioritizerDefinitions(TArray<FNetObjectPrioritizerDefinition>& InPrioritizerDefinitions) override
+	{
+		// Standard FieldOfView prioritizer
+		{
+			FNetObjectPrioritizerDefinition& PrioritizerDef = InPrioritizerDefinitions.Emplace_GetRef();
+			PrioritizerDef.PrioritizerName = FName("FoVPrioritizer");
+			PrioritizerDef.ClassName = FName("/Script/IrisCore.FieldOfViewNetObjectPrioritizer");
+			PrioritizerDef.ConfigClassName = FName("/Script/ReplicationSystemTestPlugin.FieldOfViewNetObjectPrioritizerTestConfig");
+		}
+
+		// Prioritizer for cone test
+		{
+			FNetObjectPrioritizerDefinition& PrioritizerDef = InPrioritizerDefinitions.Emplace_GetRef();
+			PrioritizerDef.PrioritizerName = FName("ConePrioritizer");
+			PrioritizerDef.ClassName = FName("/Script/IrisCore.FieldOfViewNetObjectPrioritizer");
+			PrioritizerDef.ConfigClassName = FName("/Script/ReplicationSystemTestPlugin.FieldOfViewNetObjectPrioritizerForConeTestConfig");
+		}
+	}
+
+	void InitFoVNetObjectPrioritizer()
+	{
+		FoVPrioritizerHandle = Server->ReplicationSystem->GetPrioritizerHandle(FName("FoVPrioritizer"));
+	}
+
+protected:
+	FNetObjectPrioritizerHandle FoVPrioritizerHandle = InvalidNetObjectPrioritizerHandle;
+	static constexpr float SixtyDegreeFovInRadians = FMath::DegreesToRadians(60.0f);
+};
+
+}
+
+namespace UE::Net::Private
+{
+
+UE_NET_TEST_FIXTURE(FTestFieldOfViewNetObjectPrioritizer, ViewPositionSameAsObjectGivesHighestPriority)
+{
+	const UFieldOfViewNetObjectPrioritizerTestConfig* PrioritizerConfig = GetDefault<UFieldOfViewNetObjectPrioritizerTestConfig>();
+
+	const FVector TestPosition(-1000, 5000, -10000);
+
+	TArray<FVector> WorldLocations;
+	WorldLocations.Add(TestPosition);
+
+	FReplicationView View = MakeReplicationView(TestPosition, FVector::ForwardVector, FMath::DegreesToRadians(60.0f));
+	const FPrioritizationResult& Result = PrioritizeWorldLocations(View, FoVPrioritizerHandle, WorldLocations);
+
+	// Can adjust to max all priorities in the config, but typically one sets the inner sphere priority to the max wanted priority.
+	const float MaxPriority = PrioritizerConfig->InnerSpherePriority;
+	UE_NET_ASSERT_EQ(Result.Priorities[0], MaxPriority);
+}
+
+UE_NET_TEST_FIXTURE(FTestFieldOfViewNetObjectPrioritizer, ViewPositionFarAwayFromObjectGivesLowPriority)
+{
+	const UFieldOfViewNetObjectPrioritizerTestConfig* PrioritizerConfig = GetDefault<UFieldOfViewNetObjectPrioritizerTestConfig>();
+
+	const FVector TestPosition(-1000, 5000, -10000);
+
+	TArray<FVector> WorldLocations;
+	WorldLocations.Add(TestPosition);
+
+	// Can adjust to max all priorities in the config, but typically one sets the inner sphere priority to the max wanted priority.
+	const float MaxPriority = PrioritizerConfig->InnerSpherePriority;
+
+	// One view
+	{
+		FReplicationView View = MakeReplicationView(TestPosition, FVector::ForwardVector, SixtyDegreeFovInRadians);
+
+		const FPrioritizationResult& Result = PrioritizeWorldLocations(View, FoVPrioritizerHandle, WorldLocations);
+		UE_NET_ASSERT_EQ(Result.Priorities[0], MaxPriority);
+	}
+
+	// Two views
+	{
+		FReplicationView View = MakeReplicationView(TestPosition, FVector::ForwardVector, SixtyDegreeFovInRadians);
+		View.Views.Add(MakeReplicationView(FVector::ZeroVector, FVector::ForwardVector, SixtyDegreeFovInRadians).Views[0]);
+
+		const FPrioritizationResult& Result = PrioritizeWorldLocations(View, FoVPrioritizerHandle, WorldLocations);
+		UE_NET_ASSERT_EQ(Result.Priorities[0], MaxPriority);
+	}
+
+	// Three views
+	{
+		FReplicationView View = MakeReplicationView(-TestPosition, FVector::ForwardVector, SixtyDegreeFovInRadians);
+		View.Views.Add(MakeReplicationView(TestPosition, FVector::ForwardVector, SixtyDegreeFovInRadians).Views[0]);
+		View.Views.Add(MakeReplicationView(FVector::ZeroVector, FVector::ForwardVector, SixtyDegreeFovInRadians).Views[0]);
+
+		const FPrioritizationResult& Result = PrioritizeWorldLocations(View, FoVPrioritizerHandle, WorldLocations);
+		UE_NET_ASSERT_EQ(Result.Priorities[0], MaxPriority);
+	}
+}
+
+UE_NET_TEST_FIXTURE(FTestFieldOfViewNetObjectPrioritizer, ViewPositionVeryFarAwayFromObjectGivesLowestPriority)
+{
+	const UFieldOfViewNetObjectPrioritizerTestConfig* PrioritizerConfig = GetDefault<UFieldOfViewNetObjectPrioritizerTestConfig>();
+
+	const FVector TestPosition(0, 0, 0);
+
+	TArray<FVector> WorldLocations;
+	WorldLocations.Add(TestPosition);
+
+	const float OutsideLength = FMath::Max(PrioritizerConfig->ConeLength, PrioritizerConfig->OuterSphereRadius) + 1.0f;
+	const float MinPriority = PrioritizerConfig->OutsidePriority;
+
+	// Test one view
+	{
+		FReplicationView View = MakeReplicationView(TestPosition + FVector(OutsideLength, 0, 0), FVector::ForwardVector, SixtyDegreeFovInRadians);
+
+		const FPrioritizationResult& Result = PrioritizeWorldLocations(View, FoVPrioritizerHandle, WorldLocations);
+		UE_NET_ASSERT_EQ(Result.Priorities[0], MinPriority);
+	}
+
+	// Test two views
+	{
+		FReplicationView View = MakeReplicationView(TestPosition + FVector::ForwardVector*OutsideLength, FVector::ForwardVector, SixtyDegreeFovInRadians);
+		// Add second view even further away
+		View.Views.Add(MakeReplicationView(TestPosition + FVector::ForwardVector*(2.0f*OutsideLength), FVector::ForwardVector, SixtyDegreeFovInRadians).Views[0]);
+
+		const FPrioritizationResult& Result = PrioritizeWorldLocations(View, FoVPrioritizerHandle, WorldLocations);
+		UE_NET_ASSERT_EQ(Result.Priorities[0], MinPriority);
+	}
+
+	// Test more than two views.
+	{
+		FReplicationView View = MakeReplicationView(TestPosition + FVector::ForwardVector*OutsideLength, FVector::ForwardVector, SixtyDegreeFovInRadians);
+		// Add a couple of views
+		View.Views.Add(MakeReplicationView(TestPosition + FVector::ForwardVector*(2.0f*OutsideLength), FVector::ForwardVector, SixtyDegreeFovInRadians).Views[0]);
+		View.Views.Add(MakeReplicationView(TestPosition - FVector::ForwardVector*(4.0f*OutsideLength), FVector::ForwardVector, SixtyDegreeFovInRadians).Views[0]);
+
+		const FPrioritizationResult& Result = PrioritizeWorldLocations(View, FoVPrioritizerHandle, WorldLocations);
+		UE_NET_ASSERT_EQ(Result.Priorities[0], MinPriority);
+	}
+}
+
+UE_NET_TEST_FIXTURE(FTestFieldOfViewNetObjectPrioritizer, ViewPositionInConeGivesReasonablePriority)
+{
+	const UFieldOfViewNetObjectPrioritizerForConeTestConfig* PrioritizerConfig = GetDefault<UFieldOfViewNetObjectPrioritizerForConeTestConfig>();
+	
+	const FNetObjectPrioritizerHandle ConePrioritizerHandle = Server->ReplicationSystem->GetPrioritizerHandle(FName("ConePrioritizer"));
+
+	TArray<FVector> WorldLocations;
+	// Start location is where the cone starts scaling the priority down to MinConePriority
+	WorldLocations.Add(FVector(0, 0, 0));
+	WorldLocations.Add(FVector(PrioritizerConfig->InnerConeLength, 0, 0));
+	// Something between inner and total length, closer to the inner
+	WorldLocations.Add(FVector(PrioritizerConfig->InnerConeLength + 0.25f*(PrioritizerConfig->ConeLength - PrioritizerConfig->InnerConeLength), 0, 0));
+	// Something between inner and total length, closer to the total
+	WorldLocations.Add(FVector(PrioritizerConfig->ConeLength - 0.25f*(PrioritizerConfig->ConeLength - PrioritizerConfig->InnerConeLength), 0, 0));
+	// End location is where the cone ends
+	WorldLocations.Add(FVector(PrioritizerConfig->ConeLength, 0, 0));
+	WorldLocations.Add(FVector(PrioritizerConfig->ConeLength + 10.0f, 0, 0));
+
+	FReplicationView View = MakeReplicationView(FVector::ZeroVector, FVector::ForwardVector, SixtyDegreeFovInRadians);
+	const FPrioritizationResult& Result = PrioritizeWorldLocations(View, ConePrioritizerHandle, WorldLocations);
+
+	UE_NET_ASSERT_EQ(Result.Priorities[0], PrioritizerConfig->MaxConePriority);
+	UE_NET_ASSERT_TRUE(FMath::IsNearlyEqual(Result.Priorities[1], PrioritizerConfig->MaxConePriority, 0.01f));
+	UE_NET_ASSERT_LT(Result.Priorities[2], Result.Priorities[1]);
+	UE_NET_ASSERT_GT(Result.Priorities[2], Result.Priorities[3]);
+	UE_NET_ASSERT_LT(Result.Priorities[3], Result.Priorities[2]);
+	UE_NET_ASSERT_GT(Result.Priorities[3], Result.Priorities[4]);
+	UE_NET_ASSERT_TRUE(FMath::IsNearlyEqual(Result.Priorities[4], PrioritizerConfig->MinConePriority, 0.01f));
+	UE_NET_ASSERT_EQ(Result.Priorities[5], PrioritizerConfig->OutsidePriority);
+}
+
+UE_NET_TEST_FIXTURE(FTestFieldOfViewNetObjectPrioritizer, VisualizeFieldOfViewNetObjectPrioritizer)
+{
+	if (!FPlatformMisc::IsDebuggerPresent())
+	{
+		UE_NET_LOG("Not running VisualizeFieldOfViewNetObjectPrioritizer due to debugger not present.");
+		return;
+	}
+
+	const UFieldOfViewNetObjectPrioritizerTestConfig* PrioritizerConfig = GetDefault<UFieldOfViewNetObjectPrioritizerTestConfig>();
+
+	// The produced image needs to be viewed in a debugger capable of it or stored to a file and viewed in an external viewer.
+	FVisualizationParams VisualizationParams;
+	VisualizationParams.View = MakeReplicationView(FVector::ZeroVector, FVector::ForwardVector, SixtyDegreeFovInRadians);
+	VisualizationParams.PrioritizationBox.Min = FVector(-1, -1, 0)*(PrioritizerConfig->OuterSphereRadius + 100.0f);
+	VisualizationParams.PrioritizationBox.Max = FVector(PrioritizerConfig->ConeLength + 100.0f, PrioritizerConfig->OuterSphereRadius + 100.0f, 0.0f);
+	// 20 units per pixel means 100/20=5 pixels per meter.
+	VisualizationParams.UnitsPerPixel = 20.0f;
+	const FNetObjectPrioritizerImage& Image = Super::Visualize(FoVPrioritizerHandle, VisualizationParams);
+	UE_NET_ASSERT_FALSE(Image.GreyScaleData.IsEmpty());
+}
+
+}
diff --git a/Engine/Plugins/Runtime/ReplicationSystemTestPlugin/Source/Private/Tests/ReplicationSystem/Prioritization/TestFieldOfViewNetObjectPrioritizer.h b/Engine/Plugins/Runtime/ReplicationSystemTestPlugin/Source/Private/Tests/ReplicationSystem/Prioritization/TestFieldOfViewNetObjectPrioritizer.h
new file mode 100644
--- /dev/null
+++ b/Engine/Plugins/Runtime/ReplicationSystemTestPlugin/Source/Private/Tests/ReplicationSystem/Prioritization/TestFieldOfViewNetObjectPrioritizer.h
@@ -0,0 +1,27 @@
+// Copyright Epic Games, Inc. All Rights Reserved.
+
+#pragma once
+
+#include "Iris/ReplicationSystem/Prioritization/FieldOfViewNetObjectPrioritizer.h"
+#include "TestFieldOfViewNetObjectPrioritizer.generated.h"
+
+UCLASS()
+class UFieldOfViewNetObjectPrioritizerTestConfig : public UFieldOfViewNetObjectPrioritizerConfig
+{
+	GENERATED_BODY()
+
+public:
+	UFieldOfViewNetObjectPrioritizerTestConfig();
+
+};
+
+// The cone test config will zero out all non-cone priorities.
+UCLASS()
+class UFieldOfViewNetObjectPrioritizerForConeTestConfig : public UFieldOfViewNetObjectPrioritizerConfig
+{
+	GENERATED_BODY()
+
+public:
+	UFieldOfViewNetObjectPrioritizerForConeTestConfig();
+
+};
diff --git a/Engine/Plugins/Runtime/ReplicationSystemTestPlugin/Source/Private/Tests/ReplicationSystem/Prioritization/TestNetObjectPrioritizer.cpp b/Engine/Plugins/Runtime/ReplicationSystemTestPlugin/Source/Private/Tests/ReplicationSystem/Prioritization/TestNetObjectPrioritizer.cpp
--- a/Engine/Plugins/Runtime/ReplicationSystemTestPlugin/Source/Private/Tests/ReplicationSystem/Prioritization/TestNetObjectPrioritizer.cpp
+++ b/Engine/Plugins/Runtime/ReplicationSystemTestPlugin/Source/Private/Tests/ReplicationSystem/Prioritization/TestNetObjectPrioritizer.cpp
@@ -13,10 +13,10 @@
 namespace UE::Net::Private
 {
 
-class FTestNetPrioritizerFixture : public FReplicationSystemTestFixture
+class FTestNetObjectPrioritizerAPIFixture : public FReplicationSystemTestFixture
 {
 public:
-	FTestNetPrioritizerFixture() : MockNetObjectPrioritizer(nullptr) {}
+	FTestNetObjectPrioritizerAPIFixture() : MockNetObjectPrioritizer(nullptr) {}
 
 protected:
 	virtual void SetUp() override
@@ -85,19 +85,19 @@ private:
 	TArray<FNetObjectPrioritizerDefinition> OriginalPrioritizerDefinitions;
 };
 
-UE_NET_TEST_FIXTURE(FTestNetPrioritizerFixture, CanGetPrioritizerHandle)
+UE_NET_TEST_FIXTURE(FTestNetObjectPrioritizerAPIFixture, CanGetPrioritizerHandle)
 {
 	UE_NET_ASSERT_NE(MockPrioritizerHandle, InvalidNetObjectPrioritizerHandle);
 }
 
-UE_NET_TEST_FIXTURE(FTestNetPrioritizerFixture, PrioritizerInitWasCalled)
+UE_NET_TEST_FIXTURE(FTestNetObjectPrioritizerAPIFixture, PrioritizerInitWasCalled)
 {
 	const UMockNetObjectPrioritizer::FFunctionCallStatus& FunctionCallStatus = MockNetObjectPrioritizer->GetFunctionCallStatus();
 	UE_NET_ASSERT_EQ(FunctionCallStatus.CallCounts.Init, 1U);
 	UE_NET_ASSERT_EQ(FunctionCallStatus.SuccessfulCallCounts.Init, 1U);
 }
 
-UE_NET_TEST_FIXTURE(FTestNetPrioritizerFixture, PrioritizerAddObjectSucceedsAndRemoveObjectIsCalledWhenObjectIsDestroyed)
+UE_NET_TEST_FIXTURE(FTestNetObjectPrioritizerAPIFixture, PrioritizerAddObjectSucceedsAndRemoveObjectIsCalledWhenObjectIsDestroyed)
 {
 	UTestReplicatedIrisObject* TestObject = CreateObject(0, 0);
 	UE_NET_ASSERT_NE(TestObject, nullptr);
@@ -143,7 +143,7 @@ UE_NET_TEST_FIXTURE(FTestNetPrioritizerFixture, PrioritizerAddObjectSucceedsAndR
 	}
 }
 
-UE_NET_TEST_FIXTURE(FTestNetPrioritizerFixture, PrioritizerAddObjectFailsAndRemoveObjectIsNotCalledWhenObjectIsDestroyed)
+UE_NET_TEST_FIXTURE(FTestNetObjectPrioritizerAPIFixture, PrioritizerAddObjectFailsAndRemoveObjectIsNotCalledWhenObjectIsDestroyed)
 {
 	UTestReplicatedIrisObject* TestObject = CreateObject(0, 0);
 	UE_NET_ASSERT_NE(TestObject, nullptr);
@@ -187,7 +187,7 @@ UE_NET_TEST_FIXTURE(FTestNetPrioritizerFixture, PrioritizerAddObjectFailsAndRemo
 	}
 }
 
-UE_NET_TEST_FIXTURE(FTestNetPrioritizerFixture, SwitchingPrioritizersCallsRemoveObjectOnPreviousPrioritizer)
+UE_NET_TEST_FIXTURE(FTestNetObjectPrioritizerAPIFixture, SwitchingPrioritizersCallsRemoveObjectOnPreviousPrioritizer)
 {
 	UTestReplicatedIrisObject* TestObject = CreateObject(0, 0);
 	UE_NET_ASSERT_NE(TestObject, nullptr);
@@ -232,7 +232,7 @@ UE_NET_TEST_FIXTURE(FTestNetPrioritizerFixture, SwitchingPrioritizersCallsRemove
 	}
 }
 
-UE_NET_TEST_FIXTURE(FTestNetPrioritizerFixture, PrioritizeIsNotCalledWhenNoObjectsAreAddedToIt)
+UE_NET_TEST_FIXTURE(FTestNetObjectPrioritizerAPIFixture, PrioritizeIsNotCalledWhenNoObjectsAreAddedToIt)
 {
 	UTestReplicatedIrisObject* TestObject = CreateObject(0, 0);
 	UE_NET_ASSERT_NE(TestObject, nullptr);
@@ -253,7 +253,7 @@ UE_NET_TEST_FIXTURE(FTestNetPrioritizerFixture, PrioritizeIsNotCalledWhenNoObjec
 	}
 }
 
-UE_NET_TEST_FIXTURE(FTestNetPrioritizerFixture, PrioritizeIsNotCalledWhenThereAreNoConnections)
+UE_NET_TEST_FIXTURE(FTestNetObjectPrioritizerAPIFixture, PrioritizeIsNotCalledWhenThereAreNoConnections)
 {
 	UTestReplicatedIrisObject* TestObject = CreateObject(0, 0);
 	UE_NET_ASSERT_NE(TestObject, nullptr);
@@ -281,7 +281,7 @@ UE_NET_TEST_FIXTURE(FTestNetPrioritizerFixture, PrioritizeIsNotCalledWhenThereAr
 	}
 }
 
-UE_NET_TEST_FIXTURE(FTestNetPrioritizerFixture, PrioritizeIsCalledWhenThereAreConnections)
+UE_NET_TEST_FIXTURE(FTestNetObjectPrioritizerAPIFixture, PrioritizeIsCalledWhenThereAreConnections)
 {
 	// Setup a couple of connections with valid views
 	{
@@ -331,7 +331,7 @@ UE_NET_TEST_FIXTURE(FTestNetPrioritizerFixture, PrioritizeIsCalledWhenThereAreCo
 	}
 }
 
-UE_NET_TEST_FIXTURE(FTestNetPrioritizerFixture, UpdateObjectsIsCalledForDirtyObject)
+UE_NET_TEST_FIXTURE(FTestNetObjectPrioritizerAPIFixture, UpdateObjectsIsCalledForDirtyObject)
 {
 	UTestReplicatedIrisObject* TestObject1 = CreateObject(0, 0);
 	UE_NET_ASSERT_NE(TestObject1, nullptr);
@@ -369,7 +369,7 @@ UE_NET_TEST_FIXTURE(FTestNetPrioritizerFixture, UpdateObjectsIsCalledForDirtyObj
 	}
 }
 
-UE_NET_TEST_FIXTURE(FTestNetPrioritizerFixture, UpdateObjectsIsNotCalledWhenNoObjectIsDirty)
+UE_NET_TEST_FIXTURE(FTestNetObjectPrioritizerAPIFixture, UpdateObjectsIsNotCalledWhenNoObjectIsDirty)
 {
 	UTestReplicatedIrisObject* TestObject1 = CreateObject(0, 0);
 	UE_NET_ASSERT_NE(TestObject1, nullptr);
@@ -404,7 +404,7 @@ UE_NET_TEST_FIXTURE(FTestNetPrioritizerFixture, UpdateObjectsIsNotCalledWhenNoOb
 	}
 }
 
-UE_NET_TEST_FIXTURE(FTestNetPrioritizerFixture, NativeIrisObjectGetsPriorityFromStart)
+UE_NET_TEST_FIXTURE(FTestNetObjectPrioritizerAPIFixture, NativeIrisObjectGetsPriorityFromStart)
 {
 	UTestPrioritizationNativeIrisObject* Object = CreateObject<UTestPrioritizationNativeIrisObject>();
 
@@ -428,7 +428,7 @@ UE_NET_TEST_FIXTURE(FTestNetPrioritizerFixture, NativeIrisObjectGetsPriorityFrom
 	UE_NET_ASSERT_EQ(Priority, ObjectPriority);
 }
 
-UE_NET_TEST_FIXTURE(FTestNetPrioritizerFixture, NativeIrisObjectGetsUpdatedPriority)
+UE_NET_TEST_FIXTURE(FTestNetObjectPrioritizerAPIFixture, NativeIrisObjectGetsUpdatedPriority)
 {
 	UTestPrioritizationNativeIrisObject* Object = CreateObject<UTestPrioritizationNativeIrisObject>();
 
@@ -455,7 +455,7 @@ UE_NET_TEST_FIXTURE(FTestNetPrioritizerFixture, NativeIrisObjectGetsUpdatedPrior
 	UE_NET_ASSERT_EQ(Priority, UpdatedPriority);
 }
 
-UE_NET_TEST_FIXTURE(FTestNetPrioritizerFixture, ObjectGetsPriorityFromStart)
+UE_NET_TEST_FIXTURE(FTestNetObjectPrioritizerAPIFixture, ObjectGetsPriorityFromStart)
 {
 	UTestPrioritizationObject* Object = CreateObject<UTestPrioritizationObject>();
 
@@ -479,7 +479,7 @@ UE_NET_TEST_FIXTURE(FTestNetPrioritizerFixture, ObjectGetsPriorityFromStart)
 	UE_NET_ASSERT_EQ(Priority, ObjectPriority);
 }
 
-UE_NET_TEST_FIXTURE(FTestNetPrioritizerFixture, ObjectGetsUpdatedPriority)
+UE_NET_TEST_FIXTURE(FTestNetObjectPrioritizerAPIFixture, ObjectGetsUpdatedPriority)
 {
 	UTestPrioritizationObject* Object = CreateObject<UTestPrioritizationObject>();
 
diff --git a/Engine/Plugins/Runtime/ReplicationSystemTestPlugin/Source/Private/Tests/ReplicationSystem/Prioritization/TestNetObjectPrioritizerFixture.cpp b/Engine/Plugins/Runtime/ReplicationSystemTestPlugin/Source/Private/Tests/ReplicationSystem/Prioritization/TestNetObjectPrioritizerFixture.cpp
new file mode 100644
--- /dev/null
+++ b/Engine/Plugins/Runtime/ReplicationSystemTestPlugin/Source/Private/Tests/ReplicationSystem/Prioritization/TestNetObjectPrioritizerFixture.cpp
@@ -0,0 +1,313 @@
+// Copyright Epic Games, Inc. All Rights Reserved.
+
+#include "Tests/ReplicationSystem/Prioritization/TestNetObjectPrioritizerFixture.h"
+#include "Tests/ReplicationSystem/Prioritization/TestPrioritizationObject.h"
+#include "Iris/ReplicationSystem/NetRefHandleManager.h"
+#include "Iris/ReplicationSystem/ReplicationSystemInternal.h"
+
+namespace UE::Net
+{
+
+void FTestNetObjectPrioritizerFixture::SetUp()
+{
+	PrioritizerDefinitions.Reset();
+	GetPrioritizerDefinitions(PrioritizerDefinitions);
+
+	InitNetObjectPrioritizerDefinitions();
+
+	FReplicationSystemServerClientTestFixture::SetUp();
+}
+
+void FTestNetObjectPrioritizerFixture::TearDown()
+{
+	FReplicationSystemServerClientTestFixture::TearDown();
+	RestoreNetObjectPrioritizerDefinitions();
+}
+
+void FTestNetObjectPrioritizerFixture::GetPrioritizerDefinitions(TArray<FNetObjectPrioritizerDefinition>& InPrioritizerDefinitions)
+{
+}
+
+void FTestNetObjectPrioritizerFixture::InitNetObjectPrioritizerDefinitions()
+{
+	const UClass* NetObjectPrioritizerDefinitionsClass = UNetObjectPrioritizerDefinitions::StaticClass();
+	const FProperty* DefinitionsProperty = NetObjectPrioritizerDefinitionsClass->FindPropertyByName("NetObjectPrioritizerDefinitions");
+	check(DefinitionsProperty != nullptr);
+
+	// Save CDO state.
+	UNetObjectPrioritizerDefinitions* DefaultPrioritizerDefinitions = GetMutableDefault<UNetObjectPrioritizerDefinitions>();
+	DefinitionsProperty->CopyCompleteValue(&OriginalPrioritizerDefinitions, (void*)(UPTRINT(DefaultPrioritizerDefinitions) + DefinitionsProperty->GetOffset_ForInternal()));
+
+	// Modify definitions to include the desired prioritizers.
+	DefinitionsProperty->CopyCompleteValue((void*)(UPTRINT(DefaultPrioritizerDefinitions) + DefinitionsProperty->GetOffset_ForInternal()), &PrioritizerDefinitions);
+}
+
+void FTestNetObjectPrioritizerFixture::RestoreNetObjectPrioritizerDefinitions()
+{
+	const UClass* NetObjectPrioritizerDefinitionsClass = UNetObjectPrioritizerDefinitions::StaticClass();
+	const FProperty* DefinitionsProperty = NetObjectPrioritizerDefinitionsClass->FindPropertyByName("NetObjectPrioritizerDefinitions");
+
+	UNetObjectPrioritizerDefinitions* DefaultPrioritizerDefinitions = GetMutableDefault<UNetObjectPrioritizerDefinitions>();
+	DefinitionsProperty->CopyCompleteValue((void*)(UPTRINT(DefaultPrioritizerDefinitions) + DefinitionsProperty->GetOffset_ForInternal()), &OriginalPrioritizerDefinitions);
+}
+
+FReplicationView FTestNetObjectPrioritizerFixture::MakeReplicationView(const FVector& ViewPos, const FVector& ViewDir, float ViewRadians)
+{
+	FReplicationView ReplicationView;
+	FReplicationView::FView& View = ReplicationView.Views.Emplace_GetRef();
+	View.Pos = ViewPos;
+	View.Dir = ViewDir;
+	View.FoVRadians = ViewRadians;
+	return ReplicationView;
+}
+
+FTestNetObjectPrioritizerFixture::FPrioritizationResult FTestNetObjectPrioritizerFixture::PrioritizeWorldLocations(const FReplicationView& View, FNetObjectPrioritizerHandle PrioritizerHandle, TConstArrayView<FVector> WorldLocations)
+{
+	// Need to set the world location update function to be able to set world locations.
+	check(Server->GetReplicationBridge()->GetExternalWorldLocationUpdateFunctor() == nullptr);
+
+	FReplicationSystemTestClient* Client = CreateClient();
+
+	Server->GetReplicationSystem()->SetReplicationView(Client->ConnectionIdOnServer, View);
+
+	Server->GetReplicationBridge()->SetExternalWorldLocationUpdateFunctor([&](FNetRefHandle NetHandle, const UObject* Object, FVector& OutLocation, float& OutCullDistance)
+		{
+			if (const UTestPrioritizationWithWorldLocationObject* WorldLocObject = Cast<UTestPrioritizationWithWorldLocationObject>(Object))
+			{
+				OutLocation = WorldLocObject->GetWorldLocation();
+				OutCullDistance = WorldLocObject->GetNetCullDistance();
+			}
+		});
+
+	// Create objects but do not begin replication yet
+	TArray<TObjectPtr<UTestPrioritizationWithWorldLocationObject>> WorldLocObjects;
+	for (const FVector& WorldLocation : WorldLocations)
+	{
+		UTestPrioritizationWithWorldLocationObject* Object = NewObject<UTestPrioritizationWithWorldLocationObject>();
+		Object->SetWorldLocation(WorldLocation);
+		WorldLocObjects.Add(Object);
+	}
+
+	// Begin replication for all created objects
+	{
+		UObjectReplicationBridge::FCreateNetRefHandleParams CreateNetRefHandleParams = { .bNeedsWorldLocationUpdate = true };
+
+		for (UTestPrioritizationWithWorldLocationObject* Object : WorldLocObjects)
+		{
+			Object->NetRefHandle = Server->GetReplicationBridge()->BeginReplication(Object, CreateNetRefHandleParams);
+			Server->ReplicationSystem->SetPrioritizer(Object->NetRefHandle, PrioritizerHandle);
+		}
+	}
+
+	FPrioritizationResult TestResult;
+	TestResult.Priorities.Reserve(WorldLocObjects.Num());
+
+	// Prioritize.
+	{
+		Private::FReplicationPrioritization& ReplicationPrioritization = Server->GetReplicationSystem()->GetReplicationSystemInternal()->GetPrioritization();
+		Private::FNetRefHandleManager& NetRefHandleManager = Server->GetReplicationSystem()->GetReplicationSystemInternal()->GetNetRefHandleManager();
+
+		Server->PreSendUpdate();
+
+		{
+			// Retrieve the priorities. These are floats for internal indices so we need to map the NetRefHandles to internal indices to construct the return value.
+			TConstArrayView<float> Priorities = ReplicationPrioritization.GetPrioritiesForConnection(Client->ConnectionIdOnServer);
+			for (UTestPrioritizationWithWorldLocationObject* Object : WorldLocObjects)
+			{
+				float Priority = 0.0f;
+				Private::FInternalNetRefIndex NetRefIndex = NetRefHandleManager.GetInternalIndex(Object->NetRefHandle);
+				if (NetRefIndex != Private::FNetRefHandleManager::InvalidInternalIndex)
+				{
+					Priority = Priorities[NetRefIndex];
+				}
+
+				TestResult.Priorities.Add(Priority);
+			}
+		}
+
+		Server->PostSendUpdate();
+	}
+
+	// Restore/reset world location update.
+	Server->GetReplicationBridge()->SetExternalWorldLocationUpdateFunctor({});
+
+	// Cleanup
+	{
+		DestroyClient(Client);
+
+		for (UTestPrioritizationWithWorldLocationObject* Object : WorldLocObjects)
+		{
+			Server->ReplicationBridge->EndReplication(Object, EEndReplicationFlags::Destroy);
+			Object->MarkAsGarbage();
+		}
+
+		WorldLocObjects.Empty();
+	}
+
+	return TestResult;
+}
+
+FTestNetObjectPrioritizerFixture::FPrioritizationResult FTestNetObjectPrioritizerFixture::PrioritizeObjects(const FReplicationView& View, TConstArrayView<FNetRefHandle> NetRefHandles)
+{
+	FReplicationSystemTestClient* Client = CreateClient();
+
+	FPrioritizationResult TestResult;
+	TestResult.Priorities.Reserve(NetRefHandles.Num());
+
+	// Prioritize.
+	{
+		Private::FReplicationPrioritization& ReplicationPrioritization = Server->GetReplicationSystem()->GetReplicationSystemInternal()->GetPrioritization();
+		Private::FNetRefHandleManager& NetRefHandleManager = Server->GetReplicationSystem()->GetReplicationSystemInternal()->GetNetRefHandleManager();
+
+		Server->PreSendUpdate();
+
+		{
+			// Retrieve the priorities. These are floats for internal indices so we need to map the NetRefHandles to internal indices to construct the return value.
+			TConstArrayView<float> Priorities = ReplicationPrioritization.GetPrioritiesForConnection(Client->ConnectionIdOnServer);
+			for (const FNetRefHandle& NetRefHandle : NetRefHandles)
+			{
+				float Priority = 0.0f;
+				Private::FInternalNetRefIndex NetRefIndex = NetRefHandleManager.GetInternalIndex(NetRefHandle);
+				if (NetRefIndex != Private::FNetRefHandleManager::InvalidInternalIndex)
+				{
+					Priority = Priorities[NetRefIndex];
+				}
+
+				TestResult.Priorities.Add(Priority);
+			}
+		}
+
+		Server->PostSendUpdate();
+	}
+
+	DestroyClient(Client);
+
+	return TestResult;
+}
+
+FNetObjectPrioritizerImage FTestNetObjectPrioritizerFixture::Visualize(FNetObjectPrioritizerHandle PrioritizerHandle, const FVisualizationParams& Params)
+{
+
+	const double LocationDelta = Params.UnitsPerPixel;
+
+	// Need to set the world location update function to be able to set world locations.
+	check(Server->GetReplicationBridge()->GetExternalWorldLocationUpdateFunctor() == nullptr);
+
+	FReplicationSystemTestClient* Client = CreateClient();
+
+	Server->GetReplicationSystem()->SetReplicationView(Client->ConnectionIdOnServer, Params.View);
+
+	Server->GetReplicationBridge()->SetExternalWorldLocationUpdateFunctor([&](FNetRefHandle NetHandle, const UObject* Object, FVector& OutLocation, float& OutCullDistance)
+		{
+			if (const UTestPrioritizationWithWorldLocationObject* WorldLocObject = Cast<UTestPrioritizationWithWorldLocationObject>(Object))
+			{
+				OutLocation = WorldLocObject->GetWorldLocation();
+				OutCullDistance = WorldLocObject->GetNetCullDistance();
+			}
+		});
+
+	FNetObjectPrioritizerImage Image;
+	Image.ImageWidth = IntCastChecked<uint32>(int64((Params.PrioritizationBox.Max.X + LocationDelta - Params.PrioritizationBox.Min.X)/LocationDelta));
+	Image.ImageHeight = IntCastChecked<uint32>(int64((Params.PrioritizationBox.Max.Y + LocationDelta - Params.PrioritizationBox.Min.Y)/LocationDelta));
+	Image.GreyScaleData.SetNumZeroed(IntCastChecked<int64>(Image.ImageWidth*Image.ImageHeight));
+	
+	// Create objects for one horizontal line. UObject creation is slow.
+	TArray<TObjectPtr<UTestPrioritizationWithWorldLocationObject>> WorldLocObjects;
+	WorldLocObjects.Reserve(IntCastChecked<int32>(Image.ImageWidth));
+
+	for (double X = Params.PrioritizationBox.Min.X; X <= Params.PrioritizationBox.Max.X; X += LocationDelta)
+	{
+		UTestPrioritizationWithWorldLocationObject* Object = NewObject<UTestPrioritizationWithWorldLocationObject>(GetTransientPackage(), FName(), RF_Transient);
+		WorldLocObjects.Add(Object);
+	}
+
+	// Begin replication for all created objects
+	{
+		UObjectReplicationBridge::FCreateNetRefHandleParams CreateNetRefHandleParams = { .bNeedsWorldLocationUpdate = true };
+
+		for (UTestPrioritizationWithWorldLocationObject* Object : WorldLocObjects)
+		{
+			Object->NetRefHandle = Server->GetReplicationBridge()->BeginReplication(Object, CreateNetRefHandleParams);
+			Server->ReplicationSystem->SetPrioritizer(Object->NetRefHandle, PrioritizerHandle);
+		}
+	}
+
+	FPrioritizationResult PrioResult;
+	PrioResult.Priorities.Reserve(WorldLocObjects.Num());
+
+	// Loop through the entire box, creating one image line at a time.
+	int32 ImageX = 0;
+	int32 ImageY = 0;
+	for (double Z = Params.PrioritizationBox.Min.Z; Z <= Params.PrioritizationBox.Max.Z; Z += LocationDelta)
+	{
+		ImageY = 0;
+		for (double Y = Params.PrioritizationBox.Min.Y; Y <= Params.PrioritizationBox.Max.Y; Y += LocationDelta, ++ImageY)
+		{
+			ImageX = 0;
+			for (double X = Params.PrioritizationBox.Min.X; X <= Params.PrioritizationBox.Max.X; X += LocationDelta, ++ImageX)
+			{
+				UTestPrioritizationWithWorldLocationObject* Object = WorldLocObjects[ImageX];
+				Object->SetWorldLocation(FVector(X, Y, Z));
+				Server->ReplicationSystem->ForceNetUpdate(Object->NetRefHandle);
+			}
+
+			// Prioritize.
+			{
+				Private::FReplicationPrioritization& ReplicationPrioritization = Server->GetReplicationSystem()->GetReplicationSystemInternal()->GetPrioritization();
+				Private::FNetRefHandleManager& NetRefHandleManager = Server->GetReplicationSystem()->GetReplicationSystemInternal()->GetNetRefHandleManager();
+
+				Server->PreSendUpdate();
+
+				{
+					// Retrieve the priorities. These are floats for internal indices so we need to map the NetRefHandles to internal indices to construct the return value.
+					TConstArrayView<float> Priorities = ReplicationPrioritization.GetPrioritiesForConnection(Client->ConnectionIdOnServer);
+					for (UTestPrioritizationWithWorldLocationObject* Object : WorldLocObjects)
+					{
+						float Priority = 0.0f;
+						Private::FInternalNetRefIndex NetRefIndex = NetRefHandleManager.GetInternalIndex(Object->NetRefHandle);
+						if (NetRefIndex != Private::FNetRefHandleManager::InvalidInternalIndex)
+						{
+							Priority = Priorities[NetRefIndex];
+						}
+
+						PrioResult.Priorities.Add(Priority);
+					}
+				}
+
+				Server->PostSendUpdate();
+			}
+
+			// Update image data.
+			for (uint32 Pixel = 0, PixelEnd = Image.ImageWidth; Pixel < PixelEnd; ++Pixel)
+			{
+				const uint32 PixelIndex = ImageY * Image.ImageWidth + Pixel;
+				Image.GreyScaleData[PixelIndex] = FMath::Max(Image.GreyScaleData[PixelIndex], static_cast<uint8>(int32(255*FMath::Clamp(PrioResult.Priorities[PixelIndex], 0.0f, 1.0f))));
+			}
+		}
+	}
+
+	// Restore/reset world location update.
+	Server->GetReplicationBridge()->SetExternalWorldLocationUpdateFunctor({});
+
+	// Cleanup
+	{
+		DestroyClient(Client);
+
+		for (UTestPrioritizationWithWorldLocationObject* Object : WorldLocObjects)
+		{
+			Server->ReplicationBridge->EndReplication(Object, EEndReplicationFlags::Destroy);
+			Object->MarkAsGarbage();
+		}
+
+		WorldLocObjects.Empty();
+	}
+
+	if (Params.bGarbageCollectObjects)
+	{
+		CollectGarbage(GARBAGE_COLLECTION_KEEPFLAGS);
+	}
+
+	return Image;
+}
+
+}
diff --git a/Engine/Plugins/Runtime/ReplicationSystemTestPlugin/Source/Private/Tests/ReplicationSystem/Prioritization/TestNetObjectPrioritizerFixture.h b/Engine/Plugins/Runtime/ReplicationSystemTestPlugin/Source/Private/Tests/ReplicationSystem/Prioritization/TestNetObjectPrioritizerFixture.h
new file mode 100644
--- /dev/null
+++ b/Engine/Plugins/Runtime/ReplicationSystemTestPlugin/Source/Private/Tests/ReplicationSystem/Prioritization/TestNetObjectPrioritizerFixture.h
@@ -0,0 +1,78 @@
+// Copyright Epic Games, Inc. All Rights Reserved.
+
+#pragma once
+
+#include "Tests/ReplicationSystem/ReplicationSystemServerClientTestFixture.h"
+#include "Iris/ReplicationSystem/Prioritization/NetObjectPrioritizer.h"
+#include "Iris/ReplicationSystem/Prioritization/NetObjectPrioritizerDefinitions.h"
+#include "Containers/ContainersFwd.h"
+#include "UObject/NameTypes.h"
+
+namespace UE::Net
+{
+
+struct FNetObjectPrioritizerImage
+{
+	// Width of image, in pixels.
+	uint32 ImageWidth = 0;
+	// Height of image, in pixels.
+	uint32 ImageHeight = 0;
+	// Image data where the value is mapped linearly from priority to value. 255 indicates a priority of 1.0f and above and 0.0f indicates a priority of 0.0f.
+	TArray64<uint8> GreyScaleData;
+};
+
+class FTestNetObjectPrioritizerFixture : public FReplicationSystemServerClientTestFixture
+{
+protected:
+	virtual void SetUp() override;
+
+	virtual void TearDown() override;
+
+	/** Subclasses need to override to get their prioritizers setup properly. */
+	virtual void GetPrioritizerDefinitions(TArray<FNetObjectPrioritizerDefinition>& InPrioritizerDefinitions);
+
+	static FReplicationView MakeReplicationView(const FVector& ViewPos, const FVector& ViewDir, float ViewRadians);
+
+	struct FPrioritizationResult
+	{
+		// Array with priorities for each object tested
+		TArray<float> Priorities;
+	};
+
+	/** Creates objects with specified world locations and prioritizes them using prioritizer named PrioritizerName. */
+	FPrioritizationResult PrioritizeWorldLocations(const FReplicationView& View, FNetObjectPrioritizerHandle PrioritizerHandle, TConstArrayView<FVector> WorldLocations);
+
+	/** Prioritizes objects with passed NetRefHandles. It is assumed the handles already have the appropritate prioritizer set. */
+	FPrioritizationResult PrioritizeObjects(const FReplicationView& View, TConstArrayView<FNetRefHandle> NetRefHandles);
+
+	struct FVisualizationParams
+	{
+		// View to be used to perform the prioritization and visualization.
+		FReplicationView View;
+		// The bounding box to prioritize. Keep it 2D unless you really have use for 3D prioritization for the image. X coords will be used for image width, Y coords will be image height.
+		FBox PrioritizationBox;
+		// Scaling factor. 1 UE unit is 1 cm. Unless very large images is desired or a very small box to be prioritized a value of 100 or more may be useful. 
+		float UnitsPerPixel = 100.0f;
+		// Visualization requires a lot of UObjects to be created. It's wise to perform garbage collection after the visualization has finished.
+		bool bGarbageCollectObjects = true;
+	};
+
+	/**
+	 * Produce grey scale image depicting how the prioritizer with the given name prioritizes. For the purpose of visualization priority values will be clamped to [0, 1.0f]
+	 * and produce values in range [0, 255] such that a grey scale image is produced.
+	 * If you want a clean picture make sure not to add any objects to the system prior to calling Visualize.
+	 * NOTE: Very slow.
+	 */
+	FNetObjectPrioritizerImage Visualize(FNetObjectPrioritizerHandle PrioritizerHandle, const FVisualizationParams& Params);
+
+private:
+	void InitNetObjectPrioritizerDefinitions();
+
+	void RestoreNetObjectPrioritizerDefinitions();
+
+private:
+	TArray<FNetObjectPrioritizerDefinition> OriginalPrioritizerDefinitions;
+	TArray<FNetObjectPrioritizerDefinition> PrioritizerDefinitions;
+};
+
+}
diff --git a/Engine/Plugins/Runtime/ReplicationSystemTestPlugin/Source/Private/Tests/ReplicationSystem/Prioritization/TestPrioritizationObject.h b/Engine/Plugins/Runtime/ReplicationSystemTestPlugin/Source/Private/Tests/ReplicationSystem/Prioritization/TestPrioritizationObject.h
--- a/Engine/Plugins/Runtime/ReplicationSystemTestPlugin/Source/Private/Tests/ReplicationSystem/Prioritization/TestPrioritizationObject.h
+++ b/Engine/Plugins/Runtime/ReplicationSystemTestPlugin/Source/Private/Tests/ReplicationSystem/Prioritization/TestPrioritizationObject.h
@@ -61,8 +61,6 @@ private:
 	FPriorityStateWrapper PriorityState;
 };
 
-
-// Property replication variant of above
 UCLASS()
 class UTestPrioritizationObject : public UReplicatedTestObject
 {
@@ -77,4 +75,36 @@ protected:
 private:
 	UPROPERTY(Replicated)
 	float NetTest_Priority;
-};
\ No newline at end of file
+};
+
+UCLASS()
+class UTestPrioritizationWithWorldLocationObject : public UTestPrioritizationObject
+{
+	GENERATED_BODY()
+
+public:
+	void SetWorldLocation(const FVector& InWorldLocation)
+	{
+		WorldLocation = InWorldLocation;
+	}
+
+	const FVector& GetWorldLocation() const
+	{
+		return WorldLocation;
+	}
+
+	void SetNetCullDistance(float InNetCullDistance)
+	{
+		NetCullDistance = InNetCullDistance;
+	}
+
+	float GetNetCullDistance() const
+	{
+		return NetCullDistance;
+	}
+
+private:
+	FVector WorldLocation = FVector::ZeroVector;
+	float NetCullDistance = 0.0f;
+};
+
diff --git a/Engine/Source/Runtime/Experimental/Iris/Core/Private/Iris/ReplicationSystem/Prioritization/FieldOfViewNetObjectPrioritizer.cpp b/Engine/Source/Runtime/Experimental/Iris/Core/Private/Iris/ReplicationSystem/Prioritization/FieldOfViewNetObjectPrioritizer.cpp
new file mode 100644
--- /dev/null
+++ b/Engine/Source/Runtime/Experimental/Iris/Core/Private/Iris/ReplicationSystem/Prioritization/FieldOfViewNetObjectPrioritizer.cpp
@@ -0,0 +1,279 @@
+// Copyright Epic Games, Inc. All Rights Reserved.
+
+#include "Iris/ReplicationSystem/Prioritization/FieldOfViewNetObjectPrioritizer.h"
+#include "Iris/Core/IrisCsv.h"
+#include "Iris/Core/IrisProfiler.h"
+#include "Math/UnrealMathUtility.h"
+#include "Misc/MemStack.h"
+
+void UFieldOfViewNetObjectPrioritizer::Init(FNetObjectPrioritizerInitParams& Params)
+{
+	checkf(Params.Config != nullptr, TEXT("Need config to operate."));
+	Config = TStrongObjectPtr<UFieldOfViewNetObjectPrioritizerConfig>(CastChecked<UFieldOfViewNetObjectPrioritizerConfig>(Params.Config));
+
+	Super::Init(Params);
+}
+
+void UFieldOfViewNetObjectPrioritizer::Deinit()
+{
+	Super::Deinit();
+
+	Config.Reset();
+}
+
+void UFieldOfViewNetObjectPrioritizer::Prioritize(FNetObjectPrioritizationParams& PrioritizationParams)
+{
+	IRIS_CSV_PROFILER_SCOPE(Iris, UFieldOfViewNetObjectPrioritizer_Prioritize);
+	
+	FMemStack& Mem = FMemStack::Get();
+	FMemMark MemMark(Mem);
+
+	// Trade-off memory/performance
+	constexpr uint32 MaxBatchObjectCount = 1024U;
+
+	uint32 BatchObjectCount = FMath::Min((PrioritizationParams.ObjectCount + 3U) & ~3U, MaxBatchObjectCount);
+	FBatchParams BatchParams;
+	SetupBatchParams(BatchParams, PrioritizationParams, BatchObjectCount, Mem);
+
+	for (uint32 ObjectIt = 0U, ObjectEndIt = PrioritizationParams.ObjectCount; ObjectIt < ObjectEndIt; )
+	{
+		const uint32 CurrentBatchObjectCount = FMath::Min(ObjectEndIt - ObjectIt, MaxBatchObjectCount);
+
+		BatchParams.ObjectCount = CurrentBatchObjectCount;
+		PrepareBatch(BatchParams, PrioritizationParams, ObjectIt);
+		PrioritizeBatch(BatchParams);
+		FinishBatch(BatchParams, PrioritizationParams, ObjectIt);
+
+		ObjectIt += CurrentBatchObjectCount;
+	}
+}
+
+void UFieldOfViewNetObjectPrioritizer::PrepareBatch(FBatchParams& BatchParams, const FNetObjectPrioritizationParams& PrioritizationParams, uint32 ObjectIndexStartOffset)
+{
+	IRIS_PROFILER_SCOPE(UFieldOfViewNetObjectPrioritizer_PrepareBatch);
+	const float* ExternalPriorities = PrioritizationParams.Priorities;
+	const uint32* ObjectIndices = PrioritizationParams.ObjectIndices;
+	const FNetObjectPrioritizationInfo* PrioritizationInfos = PrioritizationParams.PrioritizationInfos;
+
+	float* LocalPriorities = BatchParams.Priorities;
+	VectorRegister4* Positions = BatchParams.Positions;
+
+	// Copy priorities.
+	{
+		uint32 LocalObjIt = 0;
+		for (uint32 ObjIt = ObjectIndexStartOffset, ObjEndIt = ObjIt + BatchParams.ObjectCount; ObjIt != ObjEndIt; ++ObjIt, ++LocalObjIt)
+		{
+			const uint32 ObjectIndex = ObjectIndices[ObjIt];
+			LocalPriorities[LocalObjIt] = ExternalPriorities[ObjectIndex];
+		}
+	}
+
+	// Copy positions. 
+	uint32 LocalObjIt = 0;
+	{
+		for (uint32 ObjIt = ObjectIndexStartOffset, ObjEndIt = ObjIt + BatchParams.ObjectCount; ObjIt != ObjEndIt; ++ObjIt, ++LocalObjIt)
+		{
+			const uint32 ObjectIndex = ObjectIndices[ObjIt];
+			const FObjectLocationInfo& Info = static_cast<const FObjectLocationInfo&>(PrioritizationInfos[ObjectIndex]);
+			Positions[LocalObjIt] = GetLocation(Info);
+		}
+	}
+
+	// Make sure we have a multiple of four valid entries.
+	for (uint32 ObjIt = ObjectIndexStartOffset, ObjEndIt = (ObjIt + 3U) & ~3U; ObjIt != ObjEndIt; ++ObjIt, ++LocalObjIt)
+	{
+		LocalPriorities[LocalObjIt] = 0.0f;
+		Positions[LocalObjIt] = VectorZero();
+	}
+}
+
+void UFieldOfViewNetObjectPrioritizer::PrioritizeBatch(FBatchParams& BatchParams)
+{
+	IRIS_PROFILER_SCOPE(UFieldOfViewNetObjectPrioritizer_PrioritizeBatch);
+
+	TArray<VectorRegister4, TInlineAllocator<16>> ViewPositions;
+	TArray<VectorRegister4, TInlineAllocator<16>> ViewDirs;
+	for (const UE::Net::FReplicationView::FView& View : BatchParams.View.Views)
+	{
+		const FVector& ViewPos = View.Pos;
+		const FVector& ViewDir = View.Dir;
+		ViewPositions.Add(VectorLoadFloat3_W0(&ViewPos));
+		ViewDirs.Add(VectorLoadFloat3_W0(&ViewDir));
+	}
+
+	const VectorRegister4* Positions = BatchParams.Positions;
+	float* Priorities = BatchParams.Priorities;
+	const int ViewCount = BatchParams.View.Views.Num();
+	for (uint32 ObjIt = 0, ObjEndIt = BatchParams.ObjectCount; ObjIt < ObjEndIt; ObjIt += 4)
+	{
+		VectorRegister4 Priorities0123 = VectorMax(VectorLoadAligned(Priorities + ObjIt), BatchParams.PriorityCalculationConstants.OutsidePriority);
+		// As the cone and line of sight capsule are view direction dependent there's not a lot we can do to optimize multi-view calculations.
+		for (int ViewIt = 0, ViewEndIt = ViewCount; ViewIt < ViewEndIt; ++ViewIt)
+		{
+			const VectorRegister4 ViewPos = ViewPositions[ViewIt];
+			const VectorRegister4 ViewDir = ViewDirs[ViewIt];
+
+			// Object directions
+			const VectorRegister4 ObjectDir0 = VectorSubtract(Positions[ObjIt + 0], ViewPos);
+			const VectorRegister4 ObjectDir1 = VectorSubtract(Positions[ObjIt + 1], ViewPos);
+			const VectorRegister4 ObjectDir2 = VectorSubtract(Positions[ObjIt + 2], ViewPos);
+			const VectorRegister4 ObjectDir3 = VectorSubtract(Positions[ObjIt + 3], ViewPos);
+
+			// Calculate squared distances to ViewPos
+			const VectorRegister4 DistSqrToViewPos0 = VectorDot4(ObjectDir0, ObjectDir0);
+			const VectorRegister4 DistSqrToViewPos1 = VectorDot4(ObjectDir1, ObjectDir1);
+			const VectorRegister4 DistSqrToViewPos2 = VectorDot4(ObjectDir2, ObjectDir2);
+			const VectorRegister4 DistSqrToViewPos3 = VectorDot4(ObjectDir3, ObjectDir3);
+
+			// Assemble all distances into a single vector
+			const VectorRegister4 DistSqrToViewPos0101 = VectorSwizzle(VectorCombineHigh(DistSqrToViewPos0, DistSqrToViewPos1), 0, 2, 1, 3);
+			const VectorRegister4 DistSqrToViewPos2323 = VectorSwizzle(VectorCombineHigh(DistSqrToViewPos2, DistSqrToViewPos3), 0, 2, 1, 3);
+			const VectorRegister4 DistSqrToViewPos0123 = VectorCombineHigh(DistSqrToViewPos0101, DistSqrToViewPos2323);
+			const VectorRegister4 DistToViewPos0123 = VectorSqrt(DistSqrToViewPos0123);
+
+			// Project object direction onto cone center axis to retrieve the distance on the cone
+			const VectorRegister4 ConeDist0 = VectorDot4(ObjectDir0, ViewDir);
+			const VectorRegister4 ConeDist1 = VectorDot4(ObjectDir1, ViewDir);
+			const VectorRegister4 ConeDist2 = VectorDot4(ObjectDir2, ViewDir);
+			const VectorRegister4 ConeDist3 = VectorDot4(ObjectDir3, ViewDir);
+
+			// $IRIS TODO Worth doing VectorMaskBits on ConeDistInRangeMask before doing the more expensive cone checks?
+			// Calculate the distance to the center axis at the cone distances
+			VectorRegister4 DistSqrToConeCenterAxis0 = VectorSubtract(Positions[ObjIt + 0], VectorMultiply(ConeDist0, ViewDir));
+			VectorRegister4 DistSqrToConeCenterAxis1 = VectorSubtract(Positions[ObjIt + 1], VectorMultiply(ConeDist1, ViewDir));
+			VectorRegister4 DistSqrToConeCenterAxis2 = VectorSubtract(Positions[ObjIt + 2], VectorMultiply(ConeDist2, ViewDir));
+			VectorRegister4 DistSqrToConeCenterAxis3 = VectorSubtract(Positions[ObjIt + 3], VectorMultiply(ConeDist3, ViewDir));
+			DistSqrToConeCenterAxis0 = VectorDot4(DistSqrToConeCenterAxis0, DistSqrToConeCenterAxis0);
+			DistSqrToConeCenterAxis1 = VectorDot4(DistSqrToConeCenterAxis1, DistSqrToConeCenterAxis1);
+			DistSqrToConeCenterAxis2 = VectorDot4(DistSqrToConeCenterAxis2, DistSqrToConeCenterAxis2);
+			DistSqrToConeCenterAxis3 = VectorDot4(DistSqrToConeCenterAxis3, DistSqrToConeCenterAxis3);
+
+			// Assemble all cone distances into a single vector
+			const VectorRegister4 ConeDist0101 = VectorSwizzle(VectorCombineHigh(ConeDist0, ConeDist1), 0, 2, 1, 3);
+			const VectorRegister4 ConeDist2323 = VectorSwizzle(VectorCombineHigh(ConeDist2, ConeDist3), 0, 2, 1, 3);
+			const VectorRegister4 ConeDist0123 = VectorCombineHigh(ConeDist0101, ConeDist2323);
+
+			// Assemble all distances to cone center axis into a single vector
+			const VectorRegister4 DistSqrToConeCenterAxis0101 = VectorSwizzle(VectorCombineHigh(DistSqrToConeCenterAxis0, DistSqrToConeCenterAxis1), 0, 2, 1, 3);
+			const VectorRegister4 DistSqrToConeCenterAxis2323 = VectorSwizzle(VectorCombineHigh(DistSqrToConeCenterAxis2, DistSqrToConeCenterAxis3), 0, 2, 1, 3);
+			const VectorRegister4 DistSqrToConeCenterAxis0123 = VectorCombineHigh(DistSqrToConeCenterAxis0101, DistSqrToConeCenterAxis2323);
+
+			// Validate that the cone distances fall into the valid range [0, OuterConeLength]
+			const VectorRegister4 ConeDistGEZeroMask = VectorCompareGE(ConeDist0123, VectorZeroVectorRegister());
+			const VectorRegister4 ConeDistLEDistMask = VectorCompareLE(ConeDist0123, BatchParams.PriorityCalculationConstants.ConeLength);
+			const VectorRegister4 ConeDistInRangeMask = VectorBitwiseAnd(ConeDistGEZeroMask, ConeDistLEDistMask);
+
+			// Validate that projected points fall within the cone radius at their respective distances
+			VectorRegister4 ConeRadiusAtDistSqr = VectorMultiply(ConeDist0123, BatchParams.PriorityCalculationConstants.ConeRadiusFactor);
+			ConeRadiusAtDistSqr = VectorMultiply(ConeRadiusAtDistSqr, ConeRadiusAtDistSqr);
+
+			// Calculate priorities. Go through the different shapes and always take the max of the current priorities and the shape priorities, assuming the object point is inside it.
+			// No assumptions are being made with regards to the different shapes' priorities- always take the max of all viable priorities.
+
+			// Cone priorities.
+			// Use the distance to the view pos for the cone priorities. It's more correct versus using the distance along the center axis but costs a VectorSqrt.
+			// InsideConeMask are for positions inside the cone. However the cone priorities starts getting lower at InnerConeLength. For distances shorter than InnerConeLength the priority is InnerConePriority.
+			const VectorRegister4 InsideConeMask = VectorBitwiseAnd(VectorCompareLE(DistSqrToConeCenterAxis0123, ConeRadiusAtDistSqr), ConeDistInRangeMask);
+			const VectorRegister4 InsideInnerConeMask = VectorCompareLE(ConeDist0123, BatchParams.PriorityCalculationConstants.InnerConeLength);
+			const VectorRegister4 ConeLengthFactor = VectorMultiply(VectorSubtract(DistToViewPos0123, BatchParams.PriorityCalculationConstants.InnerConeLength), BatchParams.PriorityCalculationConstants.InvConeLengthDiff);
+			VectorRegister4 ConePriorities0123 = VectorMultiplyAdd(ConeLengthFactor, BatchParams.PriorityCalculationConstants.ConePriorityDiff, BatchParams.PriorityCalculationConstants.InnerConePriority);
+			ConePriorities0123 = VectorSelect(InsideInnerConeMask, BatchParams.PriorityCalculationConstants.InnerConePriority, ConePriorities0123);
+			ConePriorities0123 = VectorBitwiseAnd(ConePriorities0123, InsideConeMask);
+
+			// Line of sight priorities.
+			const VectorRegister4 InsideLineOfSightMask = VectorBitwiseAnd(VectorCompareLE(DistSqrToConeCenterAxis0123, BatchParams.PriorityCalculationConstants.LineOfSightRadiusSqr), ConeDistInRangeMask);
+			const VectorRegister4 LoSPriorities0123 = VectorBitwiseAnd(InsideLineOfSightMask, BatchParams.PriorityCalculationConstants.LineOfSightPriority);
+
+			// Outer sphere priorities
+			const VectorRegister4 InsideOuterSphereMask = VectorCompareLE(DistSqrToViewPos0123, BatchParams.PriorityCalculationConstants.OuterSphereRadiusSqr);
+			const VectorRegister4 OuterSpherePriorities0123 = VectorBitwiseAnd(InsideOuterSphereMask, BatchParams.PriorityCalculationConstants.OuterSpherePriority);
+
+			// Inner sphere priorities
+			const VectorRegister4 InsideInnerSphereMask = VectorCompareLE(DistSqrToViewPos0123, BatchParams.PriorityCalculationConstants.InnerSphereRadiusSqr);
+			const VectorRegister4 InsideSpherePriorities0123 = VectorBitwiseAnd(InsideInnerSphereMask, BatchParams.PriorityCalculationConstants.InnerSpherePriority);
+
+			// Compute max value of all priorities
+			Priorities0123 = VectorMax(Priorities0123, VectorMax(ConePriorities0123, LoSPriorities0123));
+			Priorities0123 = VectorMax(Priorities0123, VectorMax(OuterSpherePriorities0123, InsideSpherePriorities0123));
+		}
+
+		// Store our calculated priority which takes into account the priorities prior to the above calculations.
+		VectorStoreAligned(Priorities0123, Priorities + ObjIt);
+	}
+}
+
+void UFieldOfViewNetObjectPrioritizer::FinishBatch(const FBatchParams& BatchParams, FNetObjectPrioritizationParams& PrioritizationParams, uint32 ObjectIndexStartOffset)
+{
+	IRIS_PROFILER_SCOPE(UFieldOfViewNetObjectPrioritizer_FinishBatch);
+	float* ExternalPriorities = PrioritizationParams.Priorities;
+	const uint32* ObjectIndices = PrioritizationParams.ObjectIndices;
+
+	const float* LocalPriorities = BatchParams.Priorities;
+
+	// Update the object priority array
+	uint32 LocalObjIt = 0;
+	for (uint32 ObjIt = ObjectIndexStartOffset, ObjEndIt = ObjIt + BatchParams.ObjectCount; ObjIt != ObjEndIt; ++ObjIt, ++LocalObjIt)
+	{
+		const uint32 ObjectIndex = ObjectIndices[ObjIt];
+		ExternalPriorities[ObjectIndex] = LocalPriorities[LocalObjIt];
+	}
+}
+
+void UFieldOfViewNetObjectPrioritizer::SetupBatchParams(FBatchParams& OutBatchParams, const FNetObjectPrioritizationParams& PrioritizationParams, uint32 MaxBatchObjectCount, FMemStackBase& Mem)
+{
+	OutBatchParams.View = PrioritizationParams.View;
+	OutBatchParams.ConnectionId = PrioritizationParams.ConnectionId;
+	OutBatchParams.Positions = static_cast<VectorRegister4*>(Mem.Alloc(MaxBatchObjectCount*sizeof(VectorRegister4), alignof(VectorRegister4)));
+	OutBatchParams.Priorities = static_cast<float*>(Mem.Alloc(MaxBatchObjectCount*sizeof(float), 16));
+
+	SetupCalculationConstants(OutBatchParams.PriorityCalculationConstants);
+
+	FMemory::Memzero(OutBatchParams.Positions, MaxBatchObjectCount*sizeof(VectorRegister4));
+}
+
+void UFieldOfViewNetObjectPrioritizer::SetupCalculationConstants(FPriorityCalculationConstants& OutConstants)
+{
+	// Cone constants
+	VectorRegister4 InnerConeLength = VectorSetFloat1(Config->InnerConeLength);
+	VectorRegister4 ConeLength = VectorSetFloat1(Config->ConeLength);
+	VectorRegister4 ConeLengthDiff = VectorSubtract(ConeLength, InnerConeLength);
+	VectorRegister4 InvConeLengthDiff = VectorReciprocalAccurate(ConeLengthDiff);
+	VectorRegister4 ConeRadius = VectorSetFloat1(Config->ConeLength*FMath::Tan(0.5f*FMath::DegreesToRadians(Config->ConeFieldOfViewDegrees)));
+	VectorRegister4 ConeRadiusFactor = VectorDivide(ConeRadius, ConeLength);
+	VectorRegister4 InnerConePriority = VectorSetFloat1(Config->MaxConePriority); 
+	VectorRegister4 OuterConePriority = VectorSetFloat1(Config->MinConePriority);
+	VectorRegister4 ConePriorityDiff = VectorSubtract(OuterConePriority, InnerConePriority);
+
+	// Inner and outer sphere constants
+	VectorRegister4 InnerSphereRadiusSqr = VectorSetFloat1(FMath::Square(Config->InnerSphereRadius));
+	VectorRegister4 OuterSphereRadiusSqr = VectorSetFloat1(FMath::Square(Config->OuterSphereRadius));
+	VectorRegister4 InnerSpherePriority = VectorSetFloat1(Config->InnerSpherePriority);
+	VectorRegister4 OuterSpherePriority = VectorSetFloat1(Config->OuterSpherePriority);
+
+	// Line of sight constants
+	VectorRegister4 LineOfSightRadiusSqr = VectorSetFloat1(FMath::Square(0.5f*Config->LineOfSightWidth));
+	VectorRegister4 LineOfSightPriority = VectorSetFloat1(Config->LineOfSightPriority);
+
+	// Outside priority
+	VectorRegister4 OutsidePriority = VectorSetFloat1(Config->OutsidePriority);
+
+	// Store all constants
+	OutConstants.InnerConeLength = InnerConeLength;
+	OutConstants.ConeLength = ConeLength;
+	OutConstants.ConeLengthDiff = ConeLengthDiff;
+	OutConstants.InvConeLengthDiff = InvConeLengthDiff;
+	OutConstants.ConeRadiusFactor = ConeRadiusFactor;
+	OutConstants.InnerConePriority = InnerConePriority;
+	OutConstants.OuterConePriority = OuterConePriority;
+	OutConstants.ConePriorityDiff = ConePriorityDiff;
+
+	OutConstants.InnerSphereRadiusSqr = InnerSphereRadiusSqr;
+	OutConstants.OuterSphereRadiusSqr = OuterSphereRadiusSqr;
+	OutConstants.InnerSpherePriority = InnerSpherePriority;
+	OutConstants.OuterSpherePriority = OuterSpherePriority;
+
+	OutConstants.LineOfSightRadiusSqr = LineOfSightRadiusSqr;
+	OutConstants.LineOfSightPriority = LineOfSightPriority;
+
+	OutConstants.OutsidePriority = OutsidePriority;
+}
diff --git a/Engine/Source/Runtime/Experimental/Iris/Core/Private/Iris/ReplicationSystem/Prioritization/ReplicationPrioritization.h b/Engine/Source/Runtime/Experimental/Iris/Core/Private/Iris/ReplicationSystem/Prioritization/ReplicationPrioritization.h
--- a/Engine/Source/Runtime/Experimental/Iris/Core/Private/Iris/ReplicationSystem/Prioritization/ReplicationPrioritization.h
+++ b/Engine/Source/Runtime/Experimental/Iris/Core/Private/Iris/ReplicationSystem/Prioritization/ReplicationPrioritization.h
@@ -4,8 +4,9 @@
 
 #include "CoreTypes.h"
 #include "Containers/Array.h"
-#include "Net/Core/NetBitArray.h"
+#include "Containers/ArrayView.h"
 #include "Iris/ReplicationSystem/Prioritization/NetObjectPrioritizer.h"
+#include "Net/Core/NetBitArray.h"
 #include "UObject/StrongObjectPtr.h"
 
 class UNetObjectPrioritizerDefinitions;
@@ -20,6 +21,9 @@ namespace UE::Net
 
 		typedef uint32 FInternalNetRefIndex;
 	}
+
+	// For testing
+	class FTestNetObjectPrioritizerFixture;
 }
 
 namespace UE::Net::Private
@@ -69,6 +73,11 @@ private:
 	void InitPrioritizers();
 
 private:
+	friend UE::Net::FTestNetObjectPrioritizerFixture;
+
+	// For testing
+	TConstArrayView<float> GetPrioritiesForConnection(uint32 ConnectionId) const;
+
 	struct FPrioritizerInfo
 	{
 		TStrongObjectPtr<UNetObjectPrioritizer> Prioritizer;
@@ -107,4 +116,9 @@ private:
 	uint32 HasNewObjectsWithStaticPriority : 1;
 };
 
+inline TConstArrayView<float> FReplicationPrioritization::GetPrioritiesForConnection(uint32 ConnectionId) const
+{
+	return MakeArrayView(ConnectionInfos[ConnectionId].Priorities);
+}
+
 }
diff --git a/Engine/Source/Runtime/Experimental/Iris/Core/Public/Iris/ReplicationSystem/Prioritization/FieldOfViewNetObjectPrioritizer.h b/Engine/Source/Runtime/Experimental/Iris/Core/Public/Iris/ReplicationSystem/Prioritization/FieldOfViewNetObjectPrioritizer.h
new file mode 100644
--- /dev/null
+++ b/Engine/Source/Runtime/Experimental/Iris/Core/Public/Iris/ReplicationSystem/Prioritization/FieldOfViewNetObjectPrioritizer.h
@@ -0,0 +1,136 @@
+// Copyright Epic Games, Inc. All Rights Reserved.
+
+#pragma once
+
+#include "Iris/ReplicationSystem/Prioritization/LocationBasedNetObjectPrioritizer.h"
+#include "Net/Core/NetBitArray.h"
+#include "Math/VectorRegister.h"
+#include "UObject/StrongObjectPtr.h"
+#include "FieldOfViewNetObjectPrioritizer.generated.h"
+
+class FMemStackBase;
+
+UCLASS(transient, config=Engine)
+class UFieldOfViewNetObjectPrioritizerConfig : public UNetObjectPrioritizerConfig
+{
+	GENERATED_BODY()
+
+public:
+	/* Inner sphere radius that will set InnerSpherePriority on objects in it. */
+	UPROPERTY(Config)
+	float InnerSphereRadius = 3000.0f;
+
+	/** Priority for objects inside the inner sphere */
+	UPROPERTY(Config)
+	float InnerSpherePriority = 1.0f;
+
+	/* Outer sphere radius that will set OuterSpherePriority on objects in it. */
+	UPROPERTY(Config)
+	float OuterSphereRadius = 10000.0f;
+
+	/** Priority for objects inside the outer sphere */
+	UPROPERTY(Config)
+	float OuterSpherePriority = 0.2f;
+
+	/* The field of view used to form the cone. */
+	UPROPERTY(Config)
+	float ConeFieldOfViewDegrees = 45.0f;
+
+	/** Give max cone priority up to this length of the cone. */
+	UPROPERTY(Config)
+	float InnerConeLength = 3000.0f;
+
+	/** Total cone length. */
+	UPROPERTY(Config)
+	float ConeLength = 20000.0f;
+
+	/** The minimum priority to set for objects within the cone. The priority falls off linearly from InnerConeLength to ConeLength. */
+	UPROPERTY(Config)
+	float MinConePriority = 0.2f;
+
+	/** The maximum priority to set for objects within the cone, up to InnerConeLength. The priority falls off linearly from InnerConeLength to ConeLength. */
+	UPROPERTY(Config)
+	float MaxConePriority = 1.0f;
+
+	/** Line of sight length is same as cone length but the width needs to be specified. */
+	UPROPERTY(Config)
+	float LineOfSightWidth = 200.0f;
+
+	/** Priority for objects in line of sight. */
+	UPROPERTY(Config)
+	float LineOfSightPriority = 1.0f;
+	
+	/** Priority outside the field of view cone, spheres and line of sight */
+	UPROPERTY(Config)
+	float OutsidePriority = 0.1f;
+};
+
+/**
+ * The FieldOfViewNetObjectPrioritizer prioritizes objects based on a cone, derived from a configured field of view and length. The priorities withing the cone are calculated
+ * according to configured cone priorities falling off linearly based on distance to the view location. Additionally there's a inner sphere with a most likely high priority and and larger circle with most likely less priority than the inner circle. 
+ * The circle priorities aren't scaled based on distance to view location, they're fixed. Finally there's a line of sight capsule with a configured width with a fixed priority
+ * for objects within it. The final result is the maximum priority based on these four geometric shapes. An image visualization is generated by code in TestFieldOfViewNetObjectPrioritizer.cpp.
+ */
+UCLASS(Transient, MinimalAPI)
+class UFieldOfViewNetObjectPrioritizer : public ULocationBasedNetObjectPrioritizer
+{
+	GENERATED_BODY()
+
+protected:
+	// UNetObjectPrioritizer interface
+	IRISCORE_API virtual void Init(FNetObjectPrioritizerInitParams& Params) override;
+	IRISCORE_API virtual void Deinit() override;
+	IRISCORE_API virtual void Prioritize(FNetObjectPrioritizationParams&) override;
+
+protected:
+	struct FPriorityCalculationConstants
+	{
+		// Cone constants
+		VectorRegister InnerConeLength;
+		VectorRegister ConeLength;
+		// OuterConeLength - InnerConeLength
+		VectorRegister ConeLengthDiff;
+		// 1/ConeLengthDiff
+		VectorRegister InvConeLengthDiff;
+		// ConeLength/ConeRadius for cone radius calculations at various distances.
+		VectorRegister ConeRadiusFactor;
+		VectorRegister InnerConePriority;
+		VectorRegister OuterConePriority;
+		// OuterConePriority - InnerConePriority 
+		VectorRegister ConePriorityDiff;
+
+		// Inner and outer sphere constants
+		VectorRegister InnerSphereRadiusSqr;
+		VectorRegister OuterSphereRadiusSqr;
+		VectorRegister InnerSpherePriority;
+		VectorRegister OuterSpherePriority;
+
+		// Line of sight constants
+		VectorRegister LineOfSightRadiusSqr;
+		VectorRegister LineOfSightPriority;
+
+		// Outside priority
+		VectorRegister OutsidePriority;
+	};
+
+	struct FBatchParams
+	{
+		FPriorityCalculationConstants PriorityCalculationConstants;
+		UE::Net::FReplicationView View;
+		uint32 ConnectionId;
+
+		uint32 ObjectCount;
+		VectorRegister* Positions;
+		/** 16-byte aligned pointer to priorities. */
+		float* Priorities;
+	};
+
+protected:
+	void PrepareBatch(FBatchParams& BatchParams, const FNetObjectPrioritizationParams& PrioritizationParams, uint32 ObjectIndexStartOffset);
+	void PrioritizeBatch(FBatchParams& BatchParams);
+	void FinishBatch(const FBatchParams& BatchParams, FNetObjectPrioritizationParams& PrioritizationParams, uint32 ObjectIndexStartOffset);
+	void SetupBatchParams(FBatchParams& OutBatchParams, const FNetObjectPrioritizationParams& PrioritizationParams, uint32 MaxBatchObjectCount, FMemStackBase& Mem);
+	void SetupCalculationConstants(FPriorityCalculationConstants& OutConstants);
+
+	TStrongObjectPtr<UFieldOfViewNetObjectPrioritizerConfig> Config;
+};
