Fixed compile errors

diff --git a/Engine/Source/Runtime/Experimental/Iris/Core/Private/Iris/ReplicationSystem/Prioritization/FieldOfViewNetObjectPrioritizer.cpp b/Engine/Source/Runtime/Experimental/Iris/Core/Private/Iris/ReplicationSystem/Prioritization/FieldOfViewNetObjectPrioritizer.cpp
--- a/Engine/Source/Runtime/Experimental/Iris/Core/Private/Iris/ReplicationSystem/Prioritization/FieldOfViewNetObjectPrioritizer.cpp
+++ b/Engine/Source/Runtime/Experimental/Iris/Core/Private/Iris/ReplicationSystem/Prioritization/FieldOfViewNetObjectPrioritizer.cpp
@@ -56,7 +56,7 @@ void UFieldOfViewNetObjectPrioritizer::PrepareBatch(FBatchParams& BatchParams, c
 	const FNetObjectPrioritizationInfo* PrioritizationInfos = PrioritizationParams.PrioritizationInfos;
 
 	float* LocalPriorities = BatchParams.Priorities;
-	VectorRegister4* Positions = BatchParams.Positions;
+	VectorRegister* Positions = BatchParams.Positions;
 
 	// Copy priorities.
 	{
@@ -91,8 +91,8 @@ void UFieldOfViewNetObjectPrioritizer::PrioritizeBatch(FBatchParams& BatchParams
 {
 	IRIS_PROFILER_SCOPE(UFieldOfViewNetObjectPrioritizer_PrioritizeBatch);
 
-	TArray<VectorRegister4, TInlineAllocator<16>> ViewPositions;
-	TArray<VectorRegister4, TInlineAllocator<16>> ViewDirs;
+	TArray<VectorRegister, TInlineAllocator<16>> ViewPositions;
+	TArray<VectorRegister, TInlineAllocator<16>> ViewDirs;
 	for (const UE::Net::FReplicationView::FView& View : BatchParams.View.Views)
 	{
 		const FVector& ViewPos = View.Pos;
@@ -101,70 +101,70 @@ void UFieldOfViewNetObjectPrioritizer::PrioritizeBatch(FBatchParams& BatchParams
 		ViewDirs.Add(VectorLoadFloat3_W0(&ViewDir));
 	}
 
-	const VectorRegister4* Positions = BatchParams.Positions;
+	const VectorRegister* Positions = BatchParams.Positions;
 	float* Priorities = BatchParams.Priorities;
 	const int ViewCount = BatchParams.View.Views.Num();
 	for (uint32 ObjIt = 0, ObjEndIt = BatchParams.ObjectCount; ObjIt < ObjEndIt; ObjIt += 4)
 	{
-		VectorRegister4 Priorities0123 = VectorMax(VectorLoadAligned(Priorities + ObjIt), BatchParams.PriorityCalculationConstants.OutsidePriority);
+		VectorRegister Priorities0123 = VectorMax(VectorLoadAligned(Priorities + ObjIt), BatchParams.PriorityCalculationConstants.OutsidePriority);
 		// As the cone and line of sight capsule are view direction dependent there's not a lot we can do to optimize multi-view calculations.
 		for (int ViewIt = 0, ViewEndIt = ViewCount; ViewIt < ViewEndIt; ++ViewIt)
 		{
-			const VectorRegister4 ViewPos = ViewPositions[ViewIt];
-			const VectorRegister4 ViewDir = ViewDirs[ViewIt];
+			const VectorRegister ViewPos = ViewPositions[ViewIt];
+			const VectorRegister ViewDir = ViewDirs[ViewIt];
 
 			// Object directions
-			const VectorRegister4 ObjectDir0 = VectorSubtract(Positions[ObjIt + 0], ViewPos);
-			const VectorRegister4 ObjectDir1 = VectorSubtract(Positions[ObjIt + 1], ViewPos);
-			const VectorRegister4 ObjectDir2 = VectorSubtract(Positions[ObjIt + 2], ViewPos);
-			const VectorRegister4 ObjectDir3 = VectorSubtract(Positions[ObjIt + 3], ViewPos);
+			const VectorRegister ObjectDir0 = VectorSubtract(Positions[ObjIt + 0], ViewPos);
+			const VectorRegister ObjectDir1 = VectorSubtract(Positions[ObjIt + 1], ViewPos);
+			const VectorRegister ObjectDir2 = VectorSubtract(Positions[ObjIt + 2], ViewPos);
+			const VectorRegister ObjectDir3 = VectorSubtract(Positions[ObjIt + 3], ViewPos);
 
 			// Calculate squared distances to ViewPos
-			const VectorRegister4 DistSqrToViewPos0 = VectorDot4(ObjectDir0, ObjectDir0);
-			const VectorRegister4 DistSqrToViewPos1 = VectorDot4(ObjectDir1, ObjectDir1);
-			const VectorRegister4 DistSqrToViewPos2 = VectorDot4(ObjectDir2, ObjectDir2);
-			const VectorRegister4 DistSqrToViewPos3 = VectorDot4(ObjectDir3, ObjectDir3);
+			const VectorRegister DistSqrToViewPos0 = VectorDot4(ObjectDir0, ObjectDir0);
+			const VectorRegister DistSqrToViewPos1 = VectorDot4(ObjectDir1, ObjectDir1);
+			const VectorRegister DistSqrToViewPos2 = VectorDot4(ObjectDir2, ObjectDir2);
+			const VectorRegister DistSqrToViewPos3 = VectorDot4(ObjectDir3, ObjectDir3);
 
 			// Assemble all distances into a single vector
-			const VectorRegister4 DistSqrToViewPos0101 = VectorSwizzle(VectorCombineHigh(DistSqrToViewPos0, DistSqrToViewPos1), 0, 2, 1, 3);
-			const VectorRegister4 DistSqrToViewPos2323 = VectorSwizzle(VectorCombineHigh(DistSqrToViewPos2, DistSqrToViewPos3), 0, 2, 1, 3);
-			const VectorRegister4 DistSqrToViewPos0123 = VectorCombineHigh(DistSqrToViewPos0101, DistSqrToViewPos2323);
-			const VectorRegister4 DistToViewPos0123 = VectorSqrt(DistSqrToViewPos0123);
+			const VectorRegister DistSqrToViewPos0101 = VectorSwizzle(VectorCombineHigh(DistSqrToViewPos0, DistSqrToViewPos1), 0, 2, 1, 3);
+			const VectorRegister DistSqrToViewPos2323 = VectorSwizzle(VectorCombineHigh(DistSqrToViewPos2, DistSqrToViewPos3), 0, 2, 1, 3);
+			const VectorRegister DistSqrToViewPos0123 = VectorCombineHigh(DistSqrToViewPos0101, DistSqrToViewPos2323);
+			const VectorRegister DistToViewPos0123 = VectorSqrt(DistSqrToViewPos0123);
 
 			// Project object direction onto cone center axis to retrieve the distance on the cone
-			const VectorRegister4 ConeDist0 = VectorDot4(ObjectDir0, ViewDir);
-			const VectorRegister4 ConeDist1 = VectorDot4(ObjectDir1, ViewDir);
-			const VectorRegister4 ConeDist2 = VectorDot4(ObjectDir2, ViewDir);
-			const VectorRegister4 ConeDist3 = VectorDot4(ObjectDir3, ViewDir);
+			const VectorRegister ConeDist0 = VectorDot4(ObjectDir0, ViewDir);
+			const VectorRegister ConeDist1 = VectorDot4(ObjectDir1, ViewDir);
+			const VectorRegister ConeDist2 = VectorDot4(ObjectDir2, ViewDir);
+			const VectorRegister ConeDist3 = VectorDot4(ObjectDir3, ViewDir);
 
 			// $IRIS TODO Worth doing VectorMaskBits on ConeDistInRangeMask before doing the more expensive cone checks?
 			// Calculate the distance to the center axis at the cone distances
-			VectorRegister4 DistSqrToConeCenterAxis0 = VectorSubtract(Positions[ObjIt + 0], VectorMultiply(ConeDist0, ViewDir));
-			VectorRegister4 DistSqrToConeCenterAxis1 = VectorSubtract(Positions[ObjIt + 1], VectorMultiply(ConeDist1, ViewDir));
-			VectorRegister4 DistSqrToConeCenterAxis2 = VectorSubtract(Positions[ObjIt + 2], VectorMultiply(ConeDist2, ViewDir));
-			VectorRegister4 DistSqrToConeCenterAxis3 = VectorSubtract(Positions[ObjIt + 3], VectorMultiply(ConeDist3, ViewDir));
+			VectorRegister DistSqrToConeCenterAxis0 = VectorSubtract(Positions[ObjIt + 0], VectorMultiply(ConeDist0, ViewDir));
+			VectorRegister DistSqrToConeCenterAxis1 = VectorSubtract(Positions[ObjIt + 1], VectorMultiply(ConeDist1, ViewDir));
+			VectorRegister DistSqrToConeCenterAxis2 = VectorSubtract(Positions[ObjIt + 2], VectorMultiply(ConeDist2, ViewDir));
+			VectorRegister DistSqrToConeCenterAxis3 = VectorSubtract(Positions[ObjIt + 3], VectorMultiply(ConeDist3, ViewDir));
 			DistSqrToConeCenterAxis0 = VectorDot4(DistSqrToConeCenterAxis0, DistSqrToConeCenterAxis0);
 			DistSqrToConeCenterAxis1 = VectorDot4(DistSqrToConeCenterAxis1, DistSqrToConeCenterAxis1);
 			DistSqrToConeCenterAxis2 = VectorDot4(DistSqrToConeCenterAxis2, DistSqrToConeCenterAxis2);
 			DistSqrToConeCenterAxis3 = VectorDot4(DistSqrToConeCenterAxis3, DistSqrToConeCenterAxis3);
 
 			// Assemble all cone distances into a single vector
-			const VectorRegister4 ConeDist0101 = VectorSwizzle(VectorCombineHigh(ConeDist0, ConeDist1), 0, 2, 1, 3);
-			const VectorRegister4 ConeDist2323 = VectorSwizzle(VectorCombineHigh(ConeDist2, ConeDist3), 0, 2, 1, 3);
-			const VectorRegister4 ConeDist0123 = VectorCombineHigh(ConeDist0101, ConeDist2323);
+			const VectorRegister ConeDist0101 = VectorSwizzle(VectorCombineHigh(ConeDist0, ConeDist1), 0, 2, 1, 3);
+			const VectorRegister ConeDist2323 = VectorSwizzle(VectorCombineHigh(ConeDist2, ConeDist3), 0, 2, 1, 3);
+			const VectorRegister ConeDist0123 = VectorCombineHigh(ConeDist0101, ConeDist2323);
 
 			// Assemble all distances to cone center axis into a single vector
-			const VectorRegister4 DistSqrToConeCenterAxis0101 = VectorSwizzle(VectorCombineHigh(DistSqrToConeCenterAxis0, DistSqrToConeCenterAxis1), 0, 2, 1, 3);
-			const VectorRegister4 DistSqrToConeCenterAxis2323 = VectorSwizzle(VectorCombineHigh(DistSqrToConeCenterAxis2, DistSqrToConeCenterAxis3), 0, 2, 1, 3);
-			const VectorRegister4 DistSqrToConeCenterAxis0123 = VectorCombineHigh(DistSqrToConeCenterAxis0101, DistSqrToConeCenterAxis2323);
+			const VectorRegister DistSqrToConeCenterAxis0101 = VectorSwizzle(VectorCombineHigh(DistSqrToConeCenterAxis0, DistSqrToConeCenterAxis1), 0, 2, 1, 3);
+			const VectorRegister DistSqrToConeCenterAxis2323 = VectorSwizzle(VectorCombineHigh(DistSqrToConeCenterAxis2, DistSqrToConeCenterAxis3), 0, 2, 1, 3);
+			const VectorRegister DistSqrToConeCenterAxis0123 = VectorCombineHigh(DistSqrToConeCenterAxis0101, DistSqrToConeCenterAxis2323);
 
 			// Validate that the cone distances fall into the valid range [0, OuterConeLength]
-			const VectorRegister4 ConeDistGEZeroMask = VectorCompareGE(ConeDist0123, VectorZeroVectorRegister());
-			const VectorRegister4 ConeDistLEDistMask = VectorCompareLE(ConeDist0123, BatchParams.PriorityCalculationConstants.ConeLength);
-			const VectorRegister4 ConeDistInRangeMask = VectorBitwiseAnd(ConeDistGEZeroMask, ConeDistLEDistMask);
+			const VectorRegister ConeDistGEZeroMask = VectorCompareGE(ConeDist0123, VectorZeroVectorRegister());
+			const VectorRegister ConeDistLEDistMask = VectorCompareLE(ConeDist0123, BatchParams.PriorityCalculationConstants.ConeLength);
+			const VectorRegister ConeDistInRangeMask = VectorBitwiseAnd(ConeDistGEZeroMask, ConeDistLEDistMask);
 
 			// Validate that projected points fall within the cone radius at their respective distances
-			VectorRegister4 ConeRadiusAtDistSqr = VectorMultiply(ConeDist0123, BatchParams.PriorityCalculationConstants.ConeRadiusFactor);
+			VectorRegister ConeRadiusAtDistSqr = VectorMultiply(ConeDist0123, BatchParams.PriorityCalculationConstants.ConeRadiusFactor);
 			ConeRadiusAtDistSqr = VectorMultiply(ConeRadiusAtDistSqr, ConeRadiusAtDistSqr);
 
 			// Calculate priorities. Go through the different shapes and always take the max of the current priorities and the shape priorities, assuming the object point is inside it.
@@ -173,24 +173,24 @@ void UFieldOfViewNetObjectPrioritizer::PrioritizeBatch(FBatchParams& BatchParams
 			// Cone priorities.
 			// Use the distance to the view pos for the cone priorities. It's more correct versus using the distance along the center axis but costs a VectorSqrt.
 			// InsideConeMask are for positions inside the cone. However the cone priorities starts getting lower at InnerConeLength. For distances shorter than InnerConeLength the priority is InnerConePriority.
-			const VectorRegister4 InsideConeMask = VectorBitwiseAnd(VectorCompareLE(DistSqrToConeCenterAxis0123, ConeRadiusAtDistSqr), ConeDistInRangeMask);
-			const VectorRegister4 InsideInnerConeMask = VectorCompareLE(ConeDist0123, BatchParams.PriorityCalculationConstants.InnerConeLength);
-			const VectorRegister4 ConeLengthFactor = VectorMultiply(VectorSubtract(DistToViewPos0123, BatchParams.PriorityCalculationConstants.InnerConeLength), BatchParams.PriorityCalculationConstants.InvConeLengthDiff);
-			VectorRegister4 ConePriorities0123 = VectorMultiplyAdd(ConeLengthFactor, BatchParams.PriorityCalculationConstants.ConePriorityDiff, BatchParams.PriorityCalculationConstants.InnerConePriority);
+			const VectorRegister InsideConeMask = VectorBitwiseAnd(VectorCompareLE(DistSqrToConeCenterAxis0123, ConeRadiusAtDistSqr), ConeDistInRangeMask);
+			const VectorRegister InsideInnerConeMask = VectorCompareLE(ConeDist0123, BatchParams.PriorityCalculationConstants.InnerConeLength);
+			const VectorRegister ConeLengthFactor = VectorMultiply(VectorSubtract(DistToViewPos0123, BatchParams.PriorityCalculationConstants.InnerConeLength), BatchParams.PriorityCalculationConstants.InvConeLengthDiff);
+			VectorRegister ConePriorities0123 = VectorMultiplyAdd(ConeLengthFactor, BatchParams.PriorityCalculationConstants.ConePriorityDiff, BatchParams.PriorityCalculationConstants.InnerConePriority);
 			ConePriorities0123 = VectorSelect(InsideInnerConeMask, BatchParams.PriorityCalculationConstants.InnerConePriority, ConePriorities0123);
 			ConePriorities0123 = VectorBitwiseAnd(ConePriorities0123, InsideConeMask);
 
 			// Line of sight priorities.
-			const VectorRegister4 InsideLineOfSightMask = VectorBitwiseAnd(VectorCompareLE(DistSqrToConeCenterAxis0123, BatchParams.PriorityCalculationConstants.LineOfSightRadiusSqr), ConeDistInRangeMask);
-			const VectorRegister4 LoSPriorities0123 = VectorBitwiseAnd(InsideLineOfSightMask, BatchParams.PriorityCalculationConstants.LineOfSightPriority);
+			const VectorRegister InsideLineOfSightMask = VectorBitwiseAnd(VectorCompareLE(DistSqrToConeCenterAxis0123, BatchParams.PriorityCalculationConstants.LineOfSightRadiusSqr), ConeDistInRangeMask);
+			const VectorRegister LoSPriorities0123 = VectorBitwiseAnd(InsideLineOfSightMask, BatchParams.PriorityCalculationConstants.LineOfSightPriority);
 
 			// Outer sphere priorities
-			const VectorRegister4 InsideOuterSphereMask = VectorCompareLE(DistSqrToViewPos0123, BatchParams.PriorityCalculationConstants.OuterSphereRadiusSqr);
-			const VectorRegister4 OuterSpherePriorities0123 = VectorBitwiseAnd(InsideOuterSphereMask, BatchParams.PriorityCalculationConstants.OuterSpherePriority);
+			const VectorRegister InsideOuterSphereMask = VectorCompareLE(DistSqrToViewPos0123, BatchParams.PriorityCalculationConstants.OuterSphereRadiusSqr);
+			const VectorRegister OuterSpherePriorities0123 = VectorBitwiseAnd(InsideOuterSphereMask, BatchParams.PriorityCalculationConstants.OuterSpherePriority);
 
 			// Inner sphere priorities
-			const VectorRegister4 InsideInnerSphereMask = VectorCompareLE(DistSqrToViewPos0123, BatchParams.PriorityCalculationConstants.InnerSphereRadiusSqr);
-			const VectorRegister4 InsideSpherePriorities0123 = VectorBitwiseAnd(InsideInnerSphereMask, BatchParams.PriorityCalculationConstants.InnerSpherePriority);
+			const VectorRegister InsideInnerSphereMask = VectorCompareLE(DistSqrToViewPos0123, BatchParams.PriorityCalculationConstants.InnerSphereRadiusSqr);
+			const VectorRegister InsideSpherePriorities0123 = VectorBitwiseAnd(InsideInnerSphereMask, BatchParams.PriorityCalculationConstants.InnerSpherePriority);
 
 			// Compute max value of all priorities
 			Priorities0123 = VectorMax(Priorities0123, VectorMax(ConePriorities0123, LoSPriorities0123));
@@ -223,39 +223,39 @@ void UFieldOfViewNetObjectPrioritizer::SetupBatchParams(FBatchParams& OutBatchPa
 {
 	OutBatchParams.View = PrioritizationParams.View;
 	OutBatchParams.ConnectionId = PrioritizationParams.ConnectionId;
-	OutBatchParams.Positions = static_cast<VectorRegister4*>(Mem.Alloc(MaxBatchObjectCount*sizeof(VectorRegister4), alignof(VectorRegister4)));
+	OutBatchParams.Positions = static_cast<VectorRegister*>(Mem.Alloc(MaxBatchObjectCount*sizeof(VectorRegister), alignof(VectorRegister)));
 	OutBatchParams.Priorities = static_cast<float*>(Mem.Alloc(MaxBatchObjectCount*sizeof(float), 16));
 
 	SetupCalculationConstants(OutBatchParams.PriorityCalculationConstants);
 
-	FMemory::Memzero(OutBatchParams.Positions, MaxBatchObjectCount*sizeof(VectorRegister4));
+	FMemory::Memzero(OutBatchParams.Positions, MaxBatchObjectCount*sizeof(VectorRegister));
 }
 
 void UFieldOfViewNetObjectPrioritizer::SetupCalculationConstants(FPriorityCalculationConstants& OutConstants)
 {
 	// Cone constants
-	VectorRegister4 InnerConeLength = VectorSetFloat1(Config->InnerConeLength);
-	VectorRegister4 ConeLength = VectorSetFloat1(Config->ConeLength);
-	VectorRegister4 ConeLengthDiff = VectorSubtract(ConeLength, InnerConeLength);
-	VectorRegister4 InvConeLengthDiff = VectorReciprocalAccurate(ConeLengthDiff);
-	VectorRegister4 ConeRadius = VectorSetFloat1(Config->ConeLength*FMath::Tan(0.5f*FMath::DegreesToRadians(Config->ConeFieldOfViewDegrees)));
-	VectorRegister4 ConeRadiusFactor = VectorDivide(ConeRadius, ConeLength);
-	VectorRegister4 InnerConePriority = VectorSetFloat1(Config->MaxConePriority); 
-	VectorRegister4 OuterConePriority = VectorSetFloat1(Config->MinConePriority);
-	VectorRegister4 ConePriorityDiff = VectorSubtract(OuterConePriority, InnerConePriority);
+	VectorRegister InnerConeLength = VectorSetFloat1(Config->InnerConeLength);
+	VectorRegister ConeLength = VectorSetFloat1(Config->ConeLength);
+	VectorRegister ConeLengthDiff = VectorSubtract(ConeLength, InnerConeLength);
+	VectorRegister InvConeLengthDiff = VectorReciprocalAccurate(ConeLengthDiff);
+	VectorRegister ConeRadius = VectorSetFloat1(Config->ConeLength*FMath::Tan(0.5f*FMath::DegreesToRadians(Config->ConeFieldOfViewDegrees)));
+	VectorRegister ConeRadiusFactor = VectorDivide(ConeRadius, ConeLength);
+	VectorRegister InnerConePriority = VectorSetFloat1(Config->MaxConePriority); 
+	VectorRegister OuterConePriority = VectorSetFloat1(Config->MinConePriority);
+	VectorRegister ConePriorityDiff = VectorSubtract(OuterConePriority, InnerConePriority);
 
 	// Inner and outer sphere constants
-	VectorRegister4 InnerSphereRadiusSqr = VectorSetFloat1(FMath::Square(Config->InnerSphereRadius));
-	VectorRegister4 OuterSphereRadiusSqr = VectorSetFloat1(FMath::Square(Config->OuterSphereRadius));
-	VectorRegister4 InnerSpherePriority = VectorSetFloat1(Config->InnerSpherePriority);
-	VectorRegister4 OuterSpherePriority = VectorSetFloat1(Config->OuterSpherePriority);
+	VectorRegister InnerSphereRadiusSqr = VectorSetFloat1(FMath::Square(Config->InnerSphereRadius));
+	VectorRegister OuterSphereRadiusSqr = VectorSetFloat1(FMath::Square(Config->OuterSphereRadius));
+	VectorRegister InnerSpherePriority = VectorSetFloat1(Config->InnerSpherePriority);
+	VectorRegister OuterSpherePriority = VectorSetFloat1(Config->OuterSpherePriority);
 
 	// Line of sight constants
-	VectorRegister4 LineOfSightRadiusSqr = VectorSetFloat1(FMath::Square(0.5f*Config->LineOfSightWidth));
-	VectorRegister4 LineOfSightPriority = VectorSetFloat1(Config->LineOfSightPriority);
+	VectorRegister LineOfSightRadiusSqr = VectorSetFloat1(FMath::Square(0.5f*Config->LineOfSightWidth));
+	VectorRegister LineOfSightPriority = VectorSetFloat1(Config->LineOfSightPriority);
 
 	// Outside priority
-	VectorRegister4 OutsidePriority = VectorSetFloat1(Config->OutsidePriority);
+	VectorRegister OutsidePriority = VectorSetFloat1(Config->OutsidePriority);
 
 	// Store all constants
 	OutConstants.InnerConeLength = InnerConeLength;
