[Chaos] Changing Kinematic Targets to store its rotation on single precision

This is needed because a while a go we updated our particle sto store their rotation on single precision, so we need the kinematic targets to match so equals comparison can work as expected.

As part of this work, KinematicTarget is no longer a template class as it didn't needed to be anymore

diff --git a/Engine/Source/Runtime/Core/Public/UObject/FortniteMainBranchObjectVersion.h b/Engine/Source/Runtime/Core/Public/UObject/FortniteMainBranchObjectVersion.h
--- a/Engine/Source/Runtime/Core/Public/UObject/FortniteMainBranchObjectVersion.h
+++ b/Engine/Source/Runtime/Core/Public/UObject/FortniteMainBranchObjectVersion.h
@@ -501,6 +501,9 @@ struct FFortniteMainBranchObjectVersion
 		// Upgraded movie scene 'dynamic bindings' to use the new Custom Bindings system
 		LevelSequenceUpgradeDynamicBindings,
 		
+		// Changed the precision for the stored rotation on kinematic targets to match the precision used in particles 
+		ChaosStoreKinematicTargetRotationAsSinglePrecision,
+		
 		// -----<new versions can be added above this line>-------------------------------------------------
 		VersionPlusOne,
 		LatestVersion = VersionPlusOne - 1
diff --git a/Engine/Source/Runtime/Engine/Public/Physics/ImmediatePhysics/ImmediatePhysicsChaos/ImmediatePhysicsCore_Chaos.h b/Engine/Source/Runtime/Engine/Public/Physics/ImmediatePhysics/ImmediatePhysicsChaos/ImmediatePhysicsCore_Chaos.h
--- a/Engine/Source/Runtime/Engine/Public/Physics/ImmediatePhysics/ImmediatePhysicsChaos/ImmediatePhysicsCore_Chaos.h
+++ b/Engine/Source/Runtime/Engine/Public/Physics/ImmediatePhysics/ImmediatePhysicsChaos/ImmediatePhysicsCore_Chaos.h
@@ -14,12 +14,12 @@ namespace Chaos
 	class FImplicitObject;
 	class FBasicBroadPhase;
 	class FBasicCollisionDetector;
+	class FKinematicTarget;
 	class FPBDJointConstraintHandle;
 	class FPBDJointConstraints;
 	class FPerShapeData;
 	template<class T> class TArrayCollectionArray;
 	struct FKinematicGeometryParticleParameters;
-	template<typename T, int D> class TKinematicTarget;
 	template<typename T> class TPBDConstraintIslandRule;
 	struct FPBDRigidParticleParameters;
 	class FPBDRigidsSOAs;
@@ -36,7 +36,7 @@ namespace ImmediatePhysics_Chaos
 	using EActorType = ImmediatePhysics_Shared::EActorType;
 	using EForceType = ImmediatePhysics_Shared::EForceType;
 
-	using FKinematicTarget = Chaos::TKinematicTarget<FReal, 3>;
+	using FKinematicTarget = Chaos::FKinematicTarget;
 }
 
 struct FBodyInstance;
diff --git a/Engine/Source/Runtime/Experimental/Chaos/Private/Chaos/Evolution/PBDMinEvolution.cpp b/Engine/Source/Runtime/Experimental/Chaos/Private/Chaos/Evolution/PBDMinEvolution.cpp
--- a/Engine/Source/Runtime/Experimental/Chaos/Private/Chaos/Evolution/PBDMinEvolution.cpp
+++ b/Engine/Source/Runtime/Experimental/Chaos/Private/Chaos/Evolution/PBDMinEvolution.cpp
@@ -304,9 +304,9 @@ namespace Chaos
 				return;
 			}
 
-			TKinematicTarget<FReal, 3>& KinematicTarget = Particle.KinematicTarget();
+			FKinematicTarget& KinematicTarget = Particle.KinematicTarget();
 			const FVec3 CurrentX = Particle.GetX();
-			const FRotation3 CurrentR = Particle.GetR();
+			const FRotation3f CurrentR = Particle.GetR();
 			constexpr FReal MinDt = 1e-6f;
 
 			bool bMoved = false;
@@ -330,23 +330,23 @@ namespace Chaos
 				// Move to kinematic target and update velocities to match
 				// Target positions only need to be processed once, and we reset the velocity next frame (if no new target is set)
 				FVec3 NewX;
-				FRotation3 NewR;
+				FRotation3f NewR;
 				if (IsLastStep)
 				{
-					NewX = KinematicTarget.GetTarget().GetLocation();
-					NewR = KinematicTarget.GetTarget().GetRotation();
+					NewX = KinematicTarget.GetPosition();
+					NewR = KinematicTarget.GetRotation();
 					KinematicTarget.SetMode(EKinematicTargetMode::Reset);
 				}
 				else
 				{
 					// as a reminder, stepfraction is the remaing fraction of the step from the remaining steps
 					// for total of 4 steps and current step of 2, this will be 1/3 ( 1 step passed, 3 steps remains )
-					NewX = FVec3::Lerp(CurrentX, KinematicTarget.GetTarget().GetLocation(), StepFraction);
-					NewR = FRotation3::Slerp(CurrentR, KinematicTarget.GetTarget().GetRotation(), decltype(FQuat::X)(StepFraction));
+					NewX = FVec3::Lerp(CurrentX, KinematicTarget.GetPosition(), StepFraction);
+					NewR = FRotation3f::Slerp(CurrentR, KinematicTarget.GetRotation(), decltype(FQuat4f::X)(StepFraction));
 				}
 
 				const bool bPositionChanged = !FVec3::IsNearlyEqual(NewX, CurrentX, UE_SMALL_NUMBER);
-				const bool bRotationChanged = !FRotation3::IsNearlyEqual(NewR, CurrentR, UE_SMALL_NUMBER);
+				const bool bRotationChanged = !FRotation3f::IsNearlyEqual(NewR, CurrentR, UE_KINDA_SMALL_NUMBER);
 				bMoved = bPositionChanged || bRotationChanged;
 				FVec3 NewV = FVec3(0);
 				FVec3 NewW = FVec3(0);
@@ -362,7 +362,7 @@ namespace Chaos
 					}
 				}
 				Particle.SetX(NewX);
-				Particle.SetR(NewR);
+				Particle.SetRf(NewR);
 				Particle.SetV(NewV);
 				Particle.SetW(NewW);
 
@@ -374,7 +374,7 @@ namespace Chaos
 				// Move based on velocity
 				bMoved = true;
 				Particle.SetX(Particle.GetX() + Particle.GetV() * Dt);
-				Particle.SetR(FRotation3::IntegrateRotationWithAngularVelocity(Particle.GetR(), Particle.GetW(), Dt));
+				Particle.SetRf(FRotation3f::IntegrateRotationWithAngularVelocity(Particle.GetRf(), Particle.GetWf(), static_cast<FRealSingle>(Dt)));
 				break;
 			}
 			}
@@ -382,9 +382,9 @@ namespace Chaos
 			// Set positions and previous velocities if we can
 			// Note: At present kinematics are in fact rigid bodies
 			Particle.SetP(Particle.GetX());
-			Particle.SetQ(Particle.GetR());
+			Particle.SetQf(Particle.GetRf());
 			Particle.SetPreV(Particle.GetV());
-			Particle.SetPreW(Particle.GetW());
+			Particle.SetPreWf(Particle.GetWf());
 
 			if (bMoved)
 			{
@@ -404,9 +404,9 @@ namespace Chaos
 		(FTransientKinematicGeometryParticleHandle& Particle, const int32 ParticleIndex)
 		-> void
 		{
-			TKinematicTarget<FReal, 3>& KinematicTarget = Particle.KinematicTarget();
+			FKinematicTarget& KinematicTarget = Particle.KinematicTarget();
 			const FVec3 CurrentX = Particle.GetX();
-			const FRotation3 CurrentR = Particle.GetR();
+			const FRotation3f CurrentR = Particle.GetR();
 			constexpr FReal MinDt = 1e-6f;
 
 			bool bMoved = false;
@@ -430,26 +430,26 @@ namespace Chaos
 				// Move to kinematic target and update velocities to match
 				// Target positions only need to be processed once, and we reset the velocity next frame (if no new target is set)
 				FVec3 NewX;
-				FRotation3 NewR;
+				FRotation3f NewR;
 				if (IsLastStep)
 				{
-					NewX = KinematicTarget.GetTarget().GetLocation();
-					NewR = KinematicTarget.GetTarget().GetRotation();
+					NewX = KinematicTarget.GetPosition();
+					NewR = KinematicTarget.GetRotation();
 					KinematicTarget.SetMode(EKinematicTargetMode::Reset);
 				}
 				else
 				{
 					// as a reminder, stepfraction is the remaing fraction of the step from the remaining steps
 					// for total of 4 steps and current step of 2, this will be 1/3 ( 1 step passed, 3 steps remains )
-					NewX = FVec3::Lerp(CurrentX, KinematicTarget.GetTarget().GetLocation(), StepFraction);
-					NewR = FRotation3::Slerp(CurrentR, KinematicTarget.GetTarget().GetRotation(), decltype(FQuat::X)(StepFraction));
+					NewX = FVec3::Lerp(CurrentX, KinematicTarget.GetPosition(), StepFraction);
+					NewR = FRotation3f::Slerp(CurrentR, KinematicTarget.GetRotation(), decltype(FQuat4f::X)(StepFraction));
 				}
 
 				const bool bPositionChanged = !FVec3::IsNearlyEqual(NewX, CurrentX, UE_SMALL_NUMBER);
 				const bool bRotationChanged = !FRotation3::IsNearlyEqual(NewR, CurrentR, UE_SMALL_NUMBER);
 				bMoved = bPositionChanged || bRotationChanged;
 				FVec3 NewV = FVec3(0);
-				FVec3 NewW = FVec3(0);
+				FVec3f NewW = FVec3(0);
 				if (Dt > MinDt)
 				{
 					if (bPositionChanged)
@@ -458,13 +458,13 @@ namespace Chaos
 					}
 					if (bRotationChanged)
 					{
-						NewW = FRotation3::CalculateAngularVelocity(CurrentR, NewR, Dt);
+						NewW = FRotation3f::CalculateAngularVelocity(CurrentR, NewR, static_cast<FRealSingle>(Dt));
 					}
 				}
 				Particle.SetX(NewX);
-				Particle.SetR(NewR);
+				Particle.SetRf(NewR);
 				Particle.SetV(NewV);
-				Particle.SetW(NewW);
+				Particle.SetWf(NewW);
 
 				break;
 			}
diff --git a/Engine/Source/Runtime/Experimental/Chaos/Private/Chaos/Island/IslandManager.cpp b/Engine/Source/Runtime/Experimental/Chaos/Private/Chaos/Island/IslandManager.cpp
--- a/Engine/Source/Runtime/Experimental/Chaos/Private/Chaos/Island/IslandManager.cpp
+++ b/Engine/Source/Runtime/Experimental/Chaos/Private/Chaos/Island/IslandManager.cpp
@@ -125,7 +125,7 @@ namespace Chaos::Private
 				const FKinematicTarget& KinematicTarget = Kinematic->KinematicTarget();
 				if (KinematicTarget.GetMode() == EKinematicTargetMode::Position)
 				{
-					bIsStationary = (Kinematic->GetX() - KinematicTarget.GetTargetPosition()).IsZero() && (Kinematic->GetR() * KinematicTarget.GetTargetRotation().Inverse()).IsIdentity();
+					bIsStationary = (Kinematic->GetX() - KinematicTarget.GetPosition()).IsZero() && (Kinematic->GetRf() * KinematicTarget.GetRotation().Inverse()).IsIdentity();
 				}
 				else
 				{
diff --git a/Engine/Source/Runtime/Experimental/Chaos/Private/Chaos/KinematicGeometryParticles.cpp b/Engine/Source/Runtime/Experimental/Chaos/Private/Chaos/KinematicGeometryParticles.cpp
--- a/Engine/Source/Runtime/Experimental/Chaos/Private/Chaos/KinematicGeometryParticles.cpp
+++ b/Engine/Source/Runtime/Experimental/Chaos/Private/Chaos/KinematicGeometryParticles.cpp
@@ -24,7 +24,7 @@ namespace Chaos
 
 template class Chaos::TKinematicGeometryParticlesImp<Chaos::FReal, 3, Chaos::EGeometryParticlesSimType::RigidBodySim>;
 template class Chaos::TKinematicGeometryParticlesImp<Chaos::FReal, 3, Chaos::EGeometryParticlesSimType::Other>;
-template class Chaos::TKinematicTarget<Chaos::FReal, 3>;
+class FKinematicTarget;
 
 // Re-enable when double precision is able to compile
 //template Chaos::TKinematicGeometryParticlesImp<double, 3>;
diff --git a/Engine/Source/Runtime/Experimental/Chaos/Private/Chaos/PBDRigidsEvolutionGBF.cpp b/Engine/Source/Runtime/Experimental/Chaos/Private/Chaos/PBDRigidsEvolutionGBF.cpp
--- a/Engine/Source/Runtime/Experimental/Chaos/Private/Chaos/PBDRigidsEvolutionGBF.cpp
+++ b/Engine/Source/Runtime/Experimental/Chaos/Private/Chaos/PBDRigidsEvolutionGBF.cpp
@@ -927,9 +927,9 @@ void FPBDRigidsEvolutionGBF::ApplyKinematicTargets(const FReal Dt, const FReal S
 	const auto& ApplyParticleKinematicTarget =
 	[Dt, StepFraction, IsLastStep](FTransientPBDRigidParticleHandle& Particle, const int32 ParticleIndex) -> void
 	{
-		TKinematicTarget<FReal, 3>& KinematicTarget = Particle.KinematicTarget();
+		FKinematicTarget& KinematicTarget = Particle.KinematicTarget();
 		const FVec3 CurrentX = Particle.GetX();
-		const FRotation3 CurrentR = Particle.GetR();
+		const FRotation3f CurrentR = Particle.GetRf();
 		constexpr FReal MinDt = 1e-6f;
 
 		bool bMoved = false;
@@ -954,26 +954,26 @@ void FPBDRigidsEvolutionGBF::ApplyKinematicTargets(const FReal Dt, const FReal S
 			// Move to kinematic target and update velocities to match
 			// Target positions only need to be processed once, and we reset the velocity next frame (if no new target is set)
 			FVec3 NewX;
-			FRotation3 NewR;
+			FRotation3f NewR;
 			if (IsLastStep)
 			{
-				NewX = KinematicTarget.GetTarget().GetLocation();
-				NewR = KinematicTarget.GetTarget().GetRotation();
+				NewX = KinematicTarget.GetPosition();
+				NewR = KinematicTarget.GetRotation();
 				KinematicTarget.SetMode(EKinematicTargetMode::Reset);
 			}
 			else
 			{
 				// as a reminder, stepfraction is the remaing fraction of the step from the remaining steps
 				// for total of 4 steps and current step of 2, this will be 1/3 ( 1 step passed, 3 steps remains )
-				NewX = FVec3::Lerp(CurrentX, KinematicTarget.GetTarget().GetLocation(), StepFraction);
-				NewR = FRotation3::Slerp(CurrentR, KinematicTarget.GetTarget().GetRotation(), decltype(FQuat::X)(StepFraction));
+				NewX = FVec3::Lerp(CurrentX, KinematicTarget.GetPosition(), StepFraction);
+				NewR = FRotation3f::Slerp(CurrentR, KinematicTarget.GetRotation(), decltype(FQuat4f::X)(StepFraction));
 			}
 
 			const bool bPositionChanged = !FVec3::IsNearlyEqual(NewX, CurrentX, UE_SMALL_NUMBER);
-			const bool bRotationChanged = !FRotation3::IsNearlyEqual(NewR, CurrentR, UE_SMALL_NUMBER);
+			const bool bRotationChanged = !FRotation3::IsNearlyEqual(NewR, CurrentR, UE_KINDA_SMALL_NUMBER);
 			bMoved = bPositionChanged || bRotationChanged;
 			FVec3 NewV = FVec3(0);
-			FVec3 NewW = FVec3(0);
+			FVec3f NewW = FVec3f(0);
 			if (Dt > MinDt)
 			{
 				if (bPositionChanged)
@@ -982,13 +982,13 @@ void FPBDRigidsEvolutionGBF::ApplyKinematicTargets(const FReal Dt, const FReal S
 				}
 				if (bRotationChanged)
 				{
-					NewW = FRotation3::CalculateAngularVelocity(CurrentR, NewR, Dt);
+					NewW = FRotation3f::CalculateAngularVelocity(CurrentR, NewR, static_cast<FRealSingle>(Dt));
 				}
 			}
 			Particle.SetX(NewX);
-			Particle.SetR(NewR);
+			Particle.SetRf(NewR);
 			Particle.SetV(NewV);
-			Particle.SetW(NewW);
+			Particle.SetWf(NewW);
 			Particle.SetIsMovingKinematic();
 
 			break;
@@ -1430,7 +1430,7 @@ void FPBDRigidsEvolutionGBF::SetParticleKinematicTarget(FGeometryParticleHandle*
 	{
 		if (NewKinematicTarget.GetMode() == EKinematicTargetMode::Position)
 		{
-			SetParticleTransform(ParticleHandle, NewKinematicTarget.GetTargetPosition(), NewKinematicTarget.GetTargetRotation(), false);
+			SetParticleTransform(ParticleHandle, NewKinematicTarget.GetPosition(), NewKinematicTarget.GetRotation(), false);
 		}
 	}
 }
diff --git a/Engine/Source/Runtime/Experimental/Chaos/Private/PhysicsProxy/GeometryCollectionPhysicsProxy.cpp b/Engine/Source/Runtime/Experimental/Chaos/Private/PhysicsProxy/GeometryCollectionPhysicsProxy.cpp
--- a/Engine/Source/Runtime/Experimental/Chaos/Private/PhysicsProxy/GeometryCollectionPhysicsProxy.cpp
+++ b/Engine/Source/Runtime/Experimental/Chaos/Private/PhysicsProxy/GeometryCollectionPhysicsProxy.cpp
@@ -3632,7 +3632,7 @@ void FGeometryCollectionPhysicsProxy::SetClusteredParticleKinematicTarget_Intern
 	const Chaos::EObjectStateType ObjectState = Handle->ObjectState();
 	if (ensure(ObjectState == Chaos::EObjectStateType::Kinematic))
 	{
-		Chaos::TKinematicTarget<Chaos::FReal, 3> NewKinematicTarget;
+		Chaos::FKinematicTarget NewKinematicTarget;
 		NewKinematicTarget.SetTargetMode(NewWorldTransform);
 
 		if (Chaos::FPhysicsSolver* RBDSolver = GetSolver<Chaos::FPhysicsSolver>())
diff --git a/Engine/Source/Runtime/Experimental/Chaos/Private/PhysicsProxy/SingleParticlePhysicsProxy.cpp b/Engine/Source/Runtime/Experimental/Chaos/Private/PhysicsProxy/SingleParticlePhysicsProxy.cpp
--- a/Engine/Source/Runtime/Experimental/Chaos/Private/PhysicsProxy/SingleParticlePhysicsProxy.cpp
+++ b/Engine/Source/Runtime/Experimental/Chaos/Private/PhysicsProxy/SingleParticlePhysicsProxy.cpp
@@ -105,7 +105,7 @@ void PushToPhysicsStateImp(const Chaos::FDirtyPropertiesManager& Manager, Chaos:
 			// Update world-space cached state like the bounds
 			// @todo(chaos): do we need to do this here? It should be done in Integrate and ApplyKinematicTarget so only really Statics need this...
 			const bool bHasKinematicTarget = (NewKinematicTargetGT != nullptr) && (NewKinematicTargetGT->GetMode() == EKinematicTargetMode::Position);
-			const FRigidTransform3 WorldTransform = !bHasKinematicTarget ? FRigidTransform3(Handle->GetX(), Handle->GetR()) : NewKinematicTargetGT->GetTarget();
+			const FRigidTransform3 WorldTransform = !bHasKinematicTarget ? FRigidTransform3(Handle->GetX(), Handle->GetR()) : NewKinematicTargetGT->GetTransform();
 			Handle->UpdateWorldSpaceState(WorldTransform, FVec3(0));
 
 			Evolution.DirtyParticle(*Handle);
diff --git a/Engine/Source/Runtime/Experimental/Chaos/Public/Chaos/KinematicGeometryParticles.h b/Engine/Source/Runtime/Experimental/Chaos/Public/Chaos/KinematicGeometryParticles.h
--- a/Engine/Source/Runtime/Experimental/Chaos/Public/Chaos/KinematicGeometryParticles.h
+++ b/Engine/Source/Runtime/Experimental/Chaos/Public/Chaos/KinematicGeometryParticles.h
@@ -49,8 +49,8 @@ class TKinematicGeometryParticlesImp : public TGeometryParticlesImp<T, d, SimTyp
 	const TVector<FRealSingle, d> GetWf(const int32 Index) const { return MW[Index]; }
 	void SetWf(const int32 Index, const TVector<FRealSingle, d>& InW) { MW[Index] = InW; }
 
-	const TKinematicTarget<T, d>& KinematicTarget(const int32 Index) const { return KinematicTargets[Index]; }
-	TKinematicTarget<T, d>& KinematicTarget(const int32 Index) { return KinematicTargets[Index]; }
+	const FKinematicTarget& KinematicTarget(const int32 Index) const { return KinematicTargets[Index]; }
+	FKinematicTarget& KinematicTarget(const int32 Index) { return KinematicTargets[Index]; }
 
 	FString ToString(int32 index) const
 	{
@@ -111,12 +111,12 @@ class TKinematicGeometryParticlesImp : public TGeometryParticlesImp<T, d, SimTyp
 
 	FORCEINLINE TArray<TVector<FRealSingle, d>>& AllV() { return MV; }
 	FORCEINLINE TArray<TVector<FRealSingle, d>>& AllW() { return MW; }
-	FORCEINLINE TArray<TKinematicTarget<T, d>>& AllKinematicTargets() { return KinematicTargets; }
+	FORCEINLINE TArray<FKinematicTarget>& AllKinematicTargets() { return KinematicTargets; }
 
   private:
 	TArrayCollectionArray<TVector<FRealSingle, d>> MV;
 	TArrayCollectionArray<TVector<FRealSingle, d>> MW;
-	TArrayCollectionArray<TKinematicTarget<T, d>> KinematicTargets;
+	TArrayCollectionArray<FKinematicTarget> KinematicTargets;
 };
 
 template <typename T, int d, EGeometryParticlesSimType SimType>
diff --git a/Engine/Source/Runtime/Experimental/Chaos/Public/Chaos/KinematicTargets.h b/Engine/Source/Runtime/Experimental/Chaos/Public/Chaos/KinematicTargets.h
--- a/Engine/Source/Runtime/Experimental/Chaos/Public/Chaos/KinematicTargets.h
+++ b/Engine/Source/Runtime/Experimental/Chaos/Public/Chaos/KinematicTargets.h
@@ -5,6 +5,7 @@
 
 #include "Chaos/ChaosArchive.h"
 #include "Chaos/Transform.h"
+#include "UObject/FortniteMainBranchObjectVersion.h"
 #include "UObject/FortniteReleaseBranchCustomObjectVersion.h"
 #include "UObject/PhysicsObjectVersion.h"
 
@@ -21,27 +22,31 @@ namespace Chaos
 		Velocity,		/** Particle is moved based on velocity and angular velocity, mode remains as "Velocity" until changed. */
 	};
 
+	class FKinematicTarget;
+
+	template<class T, int d>
+	using TKinematicTarget UE_DEPRECATED(5.5, "Deprecated. this class is to be deleted, use class FKinematicTarget instead") = FKinematicTarget;
+
 	/**
 	 * Data used to integrate kinematic bodies
 	 */
-	template<class T, int d>
-	class TKinematicTarget
+	class FKinematicTarget
 	{
 	public:
 
-		static TKinematicTarget<T, d> MakePositionTarget(const TRigidTransform<T, d>& InTransform)
+		static FKinematicTarget MakePositionTarget(const FRigidTransform3& InTransform)
 		{
-			return TKinematicTarget<T, d>(InTransform);
+			return FKinematicTarget(InTransform);
 		}
 
-		static TKinematicTarget<T, d> MakePositionTarget(const TVector<T, d>& InPosition, const TRotation<T, d>& InRotation)
+		static FKinematicTarget MakePositionTarget(const FVec3& InPosition, const FRotation3f& InRotation)
 		{
-			return TKinematicTarget<T, d>(InPosition, InRotation);
+			return FKinematicTarget(InPosition, InRotation);
 		}
 
-		TKinematicTarget()
+		FKinematicTarget()
 			: Position(0)
-			, Rotation(TRotation<T, d>::FromIdentity())
+			, Rotation(FRotation3f::FromIdentity())
 			, Mode(EKinematicTargetMode::None)
 		{
 		}
@@ -52,25 +57,37 @@ namespace Chaos
 		/** Get the kinematic target mode */
 		EKinematicTargetMode GetMode() const { return Mode; }
 
+		UE_DEPRECATED(5.5, "This method is Deprecated and it will be removed in a future release. Use GetTransform instead")
+		/** Get the target transform (asserts if not in Position mode) */
+		FRigidTransform3 GetTarget() const { check(Mode == EKinematicTargetMode::Position); return {Position, Rotation}; }
+
+		UE_DEPRECATED(5.5, "This method is Deprecated and it will be removed in a future release. Use GetPosition instead")
+		/** Get the target position (asserts if not in Position mode) */
+		FVec3 GetTargetPosition() const { check(Mode == EKinematicTargetMode::Position); return Position; }
+
+		UE_DEPRECATED(5.5, "This method is Deprecated and it will be removed in a future release. Use GetRotation instead")
+		/** Get the target rotation (asserts if not in Position mode) */
+		FRotation3 GetTargetRotation() const { check(Mode == EKinematicTargetMode::Position); return Rotation; }
+
 		/** Get the target transform (asserts if not in Position mode) */
-		TRigidTransform<T, d> GetTarget() const { check(Mode == EKinematicTargetMode::Position); return {Position, Rotation}; }
+		FRigidTransform3 GetTransform() const { check(Mode == EKinematicTargetMode::Position); return {Position, Rotation}; }
 
 		/** Get the target position (asserts if not in Position mode) */
-		TVector<T, d> GetTargetPosition() const { check(Mode == EKinematicTargetMode::Position); return Position; }
+		FVec3 GetPosition() const { check(Mode == EKinematicTargetMode::Position); return Position; }
 
 		/** Get the target rotation (asserts if not in Position mode) */
-		TRotation<T, d> GetTargetRotation() const { check(Mode == EKinematicTargetMode::Position); return Rotation; }
+		FRotation3f GetRotation() const { check(Mode == EKinematicTargetMode::Position); return Rotation; }
 
 		/** Clear the kinematic target */
 		void Clear()
 		{
-			Position = TVector<T, d>();
-			Rotation = TRotation<T, d>();
+			Position = FVec3();
+			Rotation = FRotation3f();
 			Mode = EKinematicTargetMode::None;
 		}
 
 		/** Use transform target mode and set the transform target */
-		void SetTargetMode(const TVector<T, d>& X, const TRotation<T, d>& R)
+		void SetTargetMode(const FVec3& X, const FRotation3f& R)
 		{
 			Position = X;
 			Rotation = R;
@@ -78,7 +95,7 @@ namespace Chaos
 		}
 
 		/** Use transform target mode and set the transform target */
-		void SetTargetMode(const TRigidTransform<T, d>& InTarget)
+		void SetTargetMode(const FRigidTransform3& InTarget)
 		{
 			Position = InTarget.GetLocation();
 			Rotation = InTarget.GetRotation();
@@ -91,31 +108,54 @@ namespace Chaos
 		// For internal use only
 		void SetMode(EKinematicTargetMode InMode) { Mode = InMode; }
 
-		friend FChaosArchive& operator<<(FChaosArchive& Ar, TKinematicTarget<T, d>& KinematicTarget)
+		friend FChaosArchive& operator<<(FChaosArchive& Ar, FKinematicTarget& KinematicTarget)
 		{
 			Ar.UsingCustomVersion(FFortniteReleaseBranchCustomObjectVersion::GUID);
 			Ar.UsingCustomVersion(FPhysicsObjectVersion::GUID);
+			Ar.UsingCustomVersion(FFortniteMainBranchObjectVersion::GUID);
 
 			const bool bRemovedScaleFN = (Ar.CustomVer(FFortniteReleaseBranchCustomObjectVersion::GUID) >= FFortniteReleaseBranchCustomObjectVersion::ChaosKinematicTargetRemoveScale);
 			const bool bRemovedScaleUE4 = (Ar.CustomVer(FPhysicsObjectVersion::GUID) >= FPhysicsObjectVersion::ChaosKinematicTargetRemoveScale);
+			const bool bRotationStoredAsSinglePrecision = (Ar.CustomVer(FFortniteMainBranchObjectVersion::GUID) >= FFortniteMainBranchObjectVersion::ChaosStoreKinematicTargetRotationAsSinglePrecision);
 
 			if (bRemovedScaleFN || bRemovedScaleUE4)
 			{
-				Ar << KinematicTarget.Position << KinematicTarget.Rotation << KinematicTarget.Mode;
+				Ar << KinematicTarget.Position;
+
+				if (bRotationStoredAsSinglePrecision)
+				{
+					Ar << KinematicTarget.Rotation;
+				}
+				else
+				{
+					FRotation3 RotationDoublePrecision;
+					if (Ar.IsLoading())
+					{
+						Ar << RotationDoublePrecision;
+						KinematicTarget.Rotation = RotationDoublePrecision;
+					}
+					else
+					{
+						RotationDoublePrecision = KinematicTarget.Rotation;
+						Ar << RotationDoublePrecision;
+					}
+				}
+				
+				Ar << KinematicTarget.Mode;
 			}
 			else
 			{
 				FRigidTransform3 Transform;
 				Ar << Transform << KinematicTarget.Mode;
 
-				KinematicTarget.Position = TVec3<T>(Transform.GetLocation());
-				KinematicTarget.Rotation = TRotation3<T>(Transform.GetRotation());
+				KinematicTarget.Position = FVec3(Transform.GetLocation());
+				KinematicTarget.Rotation = TRotation3<FRealSingle>(Transform.GetRotation());
 			}
 
 			return Ar;
 		}
 
-		bool IsEqual(const TKinematicTarget& other) const
+		bool IsEqual(const FKinematicTarget& other) const
 		{
 			return (
 				Mode == other.Mode &&
@@ -130,7 +170,7 @@ namespace Chaos
 			return IsEqual(other.KinematicTarget());
 		}
 
-		bool operator==(const TKinematicTarget& other) const
+		bool operator==(const FKinematicTarget& other) const
 		{
 			return IsEqual(other);
 		}
@@ -144,22 +184,24 @@ namespace Chaos
 		}
 
 	private:
-		explicit TKinematicTarget(const TRigidTransform<T, d>& InTransform)
+		explicit FKinematicTarget(const FRigidTransform3& InTransform)
 			: Position(InTransform.GetTranslation())
 			, Rotation(InTransform.GetRotation())
 			, Mode(EKinematicTargetMode::Position)
 		{
 		}
 
-		TKinematicTarget(const TVector<T, d>& InPosition, const TRotation<T, d>& InRotation)
+		FKinematicTarget(const FVec3& InPosition, const FRotation3f& InRotation)
 			: Position(InPosition)
 			, Rotation(InRotation)
 			, Mode(EKinematicTargetMode::Position)
 		{
 		}
 
-		TVector<T, d> Position;
-		TRotation<T, d> Rotation;
+		FVec3 Position;
+		FRotation3f Rotation;
 		EKinematicTargetMode Mode;
 	};
+
 }
+
diff --git a/Engine/Source/Runtime/Experimental/Chaos/Public/Chaos/ParticleDirtyFlags.h b/Engine/Source/Runtime/Experimental/Chaos/Public/Chaos/ParticleDirtyFlags.h
--- a/Engine/Source/Runtime/Experimental/Chaos/Public/Chaos/ParticleDirtyFlags.h
+++ b/Engine/Source/Runtime/Experimental/Chaos/Public/Chaos/ParticleDirtyFlags.h
@@ -52,8 +52,6 @@ struct FParticleID
 	}
 };
 
-using FKinematicTarget = TKinematicTarget<FReal, 3>;
-
 enum class EResimType: uint8;
 enum class ESleepType: uint8;
 
diff --git a/Engine/Source/Runtime/Experimental/Chaos/Public/Chaos/ParticleHandle.h b/Engine/Source/Runtime/Experimental/Chaos/Public/Chaos/ParticleHandle.h
--- a/Engine/Source/Runtime/Experimental/Chaos/Public/Chaos/ParticleHandle.h
+++ b/Engine/Source/Runtime/Experimental/Chaos/Public/Chaos/ParticleHandle.h
@@ -937,13 +937,13 @@ public:
 		SetW(Velocities.W());
 	}
 
-	void SetKinematicTarget(const TKinematicTarget<T, d>& InKinematicTarget, bool bInvalidate = true)
+	void SetKinematicTarget(const FKinematicTarget& InKinematicTarget, bool bInvalidate = true)
 	{
 		KinematicGeometryParticles->KinematicTarget(ParticleIdx) = InKinematicTarget;
 	}
 
-	const TKinematicTarget<T, d>& KinematicTarget() const { return KinematicGeometryParticles->KinematicTarget(ParticleIdx); }
-	TKinematicTarget<T, d>& KinematicTarget() { return KinematicGeometryParticles->KinematicTarget(ParticleIdx); }
+	const FKinematicTarget& KinematicTarget() const { return KinematicGeometryParticles->KinematicTarget(ParticleIdx); }
+	FKinematicTarget& KinematicTarget() { return KinematicGeometryParticles->KinematicTarget(ParticleIdx); }
 
 	//Really only useful when using a transient handle
 	const TKinematicGeometryParticleHandleImp<T, d, true>* Handle() const { return KinematicGeometryParticles->Handle(ParticleIdx); }
diff --git a/Engine/Source/Runtime/Experimental/Chaos/Public/PhysicsProxy/SingleParticlePhysicsProxy.h b/Engine/Source/Runtime/Experimental/Chaos/Public/PhysicsProxy/SingleParticlePhysicsProxy.h
--- a/Engine/Source/Runtime/Experimental/Chaos/Public/PhysicsProxy/SingleParticlePhysicsProxy.h
+++ b/Engine/Source/Runtime/Experimental/Chaos/Public/PhysicsProxy/SingleParticlePhysicsProxy.h
@@ -1132,7 +1132,7 @@ public:
 			{
 				// NOTE: using ClearKinematicTarget() here would just clean the dirty flag, but we actually
 				// want to make sure the kinematic target mode is set to "None", which is how it's default constructed.
-				SetKinematicTarget(Chaos::TKinematicTarget<Chaos::FReal, 3>(), bInvalidate);
+				SetKinematicTarget(Chaos::FKinematicTarget(), bInvalidate);
 			}
 		}
 
