[Chaos] Adding a null check to prevent a crash if a solver command function becomes unbound.

Also added ensures to try to track down the underlaying issue. The ensures should tell us if the issue is still happenning (without crashing), and if the sim callback object was created with an unbound function from the get go. (In case it is not, it means we could have a memeory stomp).

diff --git a/Engine/Source/Runtime/Experimental/Chaos/Public/Chaos/SimCallbackObject.h b/Engine/Source/Runtime/Experimental/Chaos/Public/Chaos/SimCallbackObject.h
--- a/Engine/Source/Runtime/Experimental/Chaos/Public/Chaos/SimCallbackObject.h
+++ b/Engine/Source/Runtime/Experimental/Chaos/Public/Chaos/SimCallbackObject.h
@@ -424,13 +424,17 @@ public:
 		: Func(MoveTemp(InFunc))
 		, Func2(nullptr)
 		, bFuncHasTimeParameters(false)
-	{}
+	{
+		ensureMsgf(Func, TEXT("Created a sim callback object with an unbound function. This command will not be executed."));
+	}
 
 	FSimCallbackCommandObject(TUniqueFunction<void(FReal DeltaTime, FReal SimTime)>&& InFunc)
 		: Func(nullptr)
 		, Func2(MoveTemp(InFunc))
 		, bFuncHasTimeParameters(true)
-	{}
+	{
+		ensureMsgf(Func2, TEXT("Created a sim callback object with an unbound function. This command will not be executed."));
+	}
 
 	virtual void FreeOutputData_External(FSimCallbackOutput* Output)
 	{
@@ -462,12 +466,18 @@ private:
 	virtual void OnPreSimulate_Internal() override
 	{
 		if (!bFuncHasTimeParameters)
-		{ 
-			Func();
+		{
+			if (ensureMsgf(Func, TEXT("The function of this sim callback object became unbound. This should not happen. This command will not be executed.")))
+			{
+				Func();
+			}
 		}
 		else
 		{
-			Func2(GetDeltaTime_Internal(), GetSimTime_Internal());
+			if (ensureMsgf(Func2, TEXT("The function of this sim callback object became unbound. This should not happen. This command will not be executed.")))
+			{
+				Func2(GetDeltaTime_Internal(), GetSimTime_Internal());
+			}
 		}
 		
 	}
