Completely disable Iris client detailed stats in a shipping build by default.

diff --git a/Engine/Source/Runtime/Experimental/Iris/Core/Private/Iris/Core/IrisProfiler.cpp b/Engine/Source/Runtime/Experimental/Iris/Core/Private/Iris/Core/IrisProfiler.cpp
--- a/Engine/Source/Runtime/Experimental/Iris/Core/Private/Iris/Core/IrisProfiler.cpp
+++ b/Engine/Source/Runtime/Experimental/Iris/Core/Private/Iris/Core/IrisProfiler.cpp
@@ -8,11 +8,18 @@ CSV_DEFINE_CATEGORY(IrisClientDetailObjectCreate, IRIS_CLIENT_PROFILER_ENABLE);
 CSV_DEFINE_CATEGORY(IrisClientDetailRepNotify, IRIS_CLIENT_PROFILER_ENABLE);
 CSV_DEFINE_CATEGORY(IrisClientDetailRPC, IRIS_CLIENT_PROFILER_ENABLE);
 
+// Detailed CSV stats are disabled in shipping by default.
+#ifndef IRIS_CLIENT_PROFILER_DETAILED
+#	define IRIS_CLIENT_PROFILER_DETAILED (!UE_BUILD_SHIPPING)
+#endif
+
 namespace UE::Net::Private
 {
 
+#if IRIS_CLIENT_PROFILER_DETAILED
 static bool bEnableDetailedClientProfiler = 0;
 static FAutoConsoleVariableRef CVarEnableDetailedClientProfilerRef(TEXT("net.Iris.EnableDetailedClientProfiler"), bEnableDetailedClientProfiler, TEXT("Generates detailed CSV Iris stats (client only)."), ECVF_Default );
+#endif
 
 }
 
@@ -29,6 +36,7 @@ void FClientProfiler::RecordObjectCreate(FName ObjectName, bool bIsSubObject)
 		FCsvProfiler::RecordCustomStat("ClientObjectCreateRoot", CSV_CATEGORY_INDEX(IrisClient), 1, ECsvCustomStatOp::Accumulate);
 	}
 
+#if IRIS_CLIENT_PROFILER_DETAILED
 	if (UE::Net::Private::bEnableDetailedClientProfiler)
 	{
 		if (!bIsSubObject)
@@ -36,24 +44,31 @@ void FClientProfiler::RecordObjectCreate(FName ObjectName, bool bIsSubObject)
 			FCsvProfiler::RecordCustomStat(ObjectName, CSV_CATEGORY_INDEX(IrisClientDetailObjectCreate), 1, ECsvCustomStatOp::Accumulate);
 		}
 	}
+#endif
 }
 
 void FClientProfiler::RecordRepNotify(FName RepNotifyName)
 {
 	FCsvProfiler::RecordCustomStat("RepNotifyCount", CSV_CATEGORY_INDEX(IrisClient), 1, ECsvCustomStatOp::Accumulate);
+
+#if IRIS_CLIENT_PROFILER_DETAILED
 	if (UE::Net::Private::bEnableDetailedClientProfiler)
 	{
 		FCsvProfiler::RecordCustomStat(RepNotifyName, CSV_CATEGORY_INDEX(IrisClientDetailRepNotify), 1, ECsvCustomStatOp::Accumulate);
 	}
+#endif
 }
 
 void FClientProfiler::RecordRPC(FName RPCName)
 {
 	FCsvProfiler::RecordCustomStat("CallCountRPC", CSV_CATEGORY_INDEX(IrisClient), 1, ECsvCustomStatOp::Accumulate);
+
+#if IRIS_CLIENT_PROFILER_DETAILED
 	if (UE::Net::Private::bEnableDetailedClientProfiler)
 	{
 		FCsvProfiler::RecordCustomStat(RPCName, CSV_CATEGORY_INDEX(IrisClientDetailRPC), 1, ECsvCustomStatOp::Accumulate);
 	}
+#endif
 }
 
 bool FClientProfiler::IsCapturing()
