Fix CIS deprecation warnings from DataflowNode change.

diff --git a/Engine/Source/Runtime/Experimental/Dataflow/Core/Public/Dataflow/DataflowNode.h b/Engine/Source/Runtime/Experimental/Dataflow/Core/Public/Dataflow/DataflowNode.h
--- a/Engine/Source/Runtime/Experimental/Dataflow/Core/Public/Dataflow/DataflowNode.h
+++ b/Engine/Source/Runtime/Experimental/Dataflow/Core/Public/Dataflow/DataflowNode.h
@@ -33,9 +33,11 @@ namespace Dataflow {
 *        is stored on the FContext. The Evaluate is const to allow support for multithreaded
 *        evaluation. 
 */
+PRAGMA_DISABLE_DEPRECATION_WARNINGS
 USTRUCT()
 struct FDataflowNode
 {
+PRAGMA_ENABLE_DEPRECATION_WARNINGS
 	GENERATED_USTRUCT_BODY()
 
 	friend class Dataflow::FGraph;
