Dataflow - Fixed ensure when forwarding unconnected inputs during evaluation.

#rnx
#tests Tested against the nodes that caused the ensures.

diff --git a/Engine/Plugins/ChaosClothAssetEditor/Source/ChaosClothAssetDataflowNodes/Private/ChaosClothAsset/MergeClothCollectionsNode.cpp b/Engine/Plugins/ChaosClothAssetEditor/Source/ChaosClothAssetDataflowNodes/Private/ChaosClothAsset/MergeClothCollectionsNode.cpp
--- a/Engine/Plugins/ChaosClothAssetEditor/Source/ChaosClothAssetDataflowNodes/Private/ChaosClothAsset/MergeClothCollectionsNode.cpp
+++ b/Engine/Plugins/ChaosClothAssetEditor/Source/ChaosClothAssetDataflowNodes/Private/ChaosClothAsset/MergeClothCollectionsNode.cpp
@@ -432,7 +432,7 @@ void FChaosClothAssetMergeClothCollectionsNode_v2::Evaluate(Dataflow::FContext&
 		else
 		{
 			// Otherwise pass through the first input unchanged
-			ForwardInput(Context, GetConnectionReference(0), &Collection);
+			SafeForwardInput(Context, GetConnectionReference(0), &Collection);
 		}
 	}
 }
diff --git a/Engine/Plugins/Experimental/GeometryCollectionPlugin/Source/GeometryCollectionNodes/Private/Dataflow/GeometryCollectionAssetNodes.cpp b/Engine/Plugins/Experimental/GeometryCollectionPlugin/Source/GeometryCollectionNodes/Private/Dataflow/GeometryCollectionAssetNodes.cpp
--- a/Engine/Plugins/Experimental/GeometryCollectionPlugin/Source/GeometryCollectionNodes/Private/Dataflow/GeometryCollectionAssetNodes.cpp
+++ b/Engine/Plugins/Experimental/GeometryCollectionPlugin/Source/GeometryCollectionNodes/Private/Dataflow/GeometryCollectionAssetNodes.cpp
@@ -95,10 +95,10 @@ void FGeometryCollectionTerminalDataflowNode::SetAssetValue(TObjectPtr<UObject>
 void FGeometryCollectionTerminalDataflowNode::Evaluate(Dataflow::FContext& Context) const
 {
 	// simply forward all inputs to corresponding outputs
-	ForwardInput(Context, &Collection, &Collection);
-	ForwardInput(Context, &Materials, &Materials);
-	ForwardInput(Context, &MaterialInstances, &MaterialInstances);
-	ForwardInput(Context, &InstancedMeshes, &InstancedMeshes);
+	SafeForwardInput(Context, &Collection, &Collection);
+	SafeForwardInput(Context, &Materials, &Materials);
+	SafeForwardInput(Context, &MaterialInstances, &MaterialInstances);
+	SafeForwardInput(Context, &InstancedMeshes, &InstancedMeshes);
 }
 
 // ===========================================================================================================================
diff --git a/Engine/Source/Runtime/Experimental/Dataflow/Core/Public/Dataflow/DataflowNode.h b/Engine/Source/Runtime/Experimental/Dataflow/Core/Public/Dataflow/DataflowNode.h
--- a/Engine/Source/Runtime/Experimental/Dataflow/Core/Public/Dataflow/DataflowNode.h
+++ b/Engine/Source/Runtime/Experimental/Dataflow/Core/Public/Dataflow/DataflowNode.h
@@ -407,17 +407,46 @@ PRAGMA_ENABLE_DEPRECATION_WARNINGS
 	/**
 	*   ForwardInput(...)
 	*
-	*   Forward an input to this output
-	*	this will not cache the value itself but cache a reference to the input connection cache entry
-	*	this is memory efficient and do not require a runtime copy of the data
-	*	input and output references must match in type
+	*   Forward an input to this output.
+	*   This will not cache the value itself but cache a reference to the input connection cache entry.
+	*   This is memory efficient and do not require a runtime copy of the data.
+	*   Input and output references must match in type.
+	*   Note that forwarding an input never sets a default value when no input is connected, use SafeForwardInput instead.
 	*
 	*   @param Context : The evaluation context that holds the data store.
-	*   @param InputReference : Pointer to a input member of this node that needs to be forwarded
+	*   @param InputReference : Pointer to a input member of this node that needs to be forwarded.
 	*   @param Reference : Pointer to a member of this node that corresponds with the output to set.
 	*/
 	DATAFLOWCORE_API void ForwardInput(Dataflow::FContext& Context, const Dataflow::FConnectionReference& InputReference, const Dataflow::FConnectionReference& Reference) const;
 
+	/**
+	*   SafeForwardInput(...)
+	*
+	*   Forward an input to this output or set a default value if no input is connected.
+	*   This is more memory efficient when an input is connected than setting the value.
+	*   Input and output references must match in type.
+	*
+	*   @param Context : The evaluation context that holds the data store.
+	*   @param InputReference : Pointer to a input member of this node that needs to be forwarded.
+	*   @param Reference : Pointer to a member of this node that corresponds with the output to set.
+	*/
+	template<class T>
+	void SafeForwardInput(Dataflow::FContext& Context, const Dataflow::FConnectionReference& InputReference, const T* Reference) const
+	{
+		if (IsConnected(InputReference))
+		{
+			ForwardInput(Context, InputReference, Dataflow::TConnectionReference<T>(Reference));
+		}
+		else if constexpr (std::is_base_of_v<FDataflowAnyType, T>)
+		{
+			SetValue(Context, reinterpret_cast<const T*>(InputReference.Reference)->Value, Reference);
+		}
+		else
+		{
+			SetValue(Context, *reinterpret_cast<const T*>(InputReference.Reference), Reference);
+		}
+	}
+
 	/**
 	*   IsConnected(...)
 	*
