Dataflow - Trivial change, use static_cast instead of reinterpret_cast in the new SafeForwardInput function.

#rnx

diff --git a/Engine/Source/Runtime/Experimental/Dataflow/Core/Public/Dataflow/DataflowNode.h b/Engine/Source/Runtime/Experimental/Dataflow/Core/Public/Dataflow/DataflowNode.h
--- a/Engine/Source/Runtime/Experimental/Dataflow/Core/Public/Dataflow/DataflowNode.h
+++ b/Engine/Source/Runtime/Experimental/Dataflow/Core/Public/Dataflow/DataflowNode.h
@@ -439,11 +439,11 @@ PRAGMA_ENABLE_DEPRECATION_WARNINGS
 		}
 		else if constexpr (std::is_base_of_v<FDataflowAnyType, T>)
 		{
-			SetValue(Context, reinterpret_cast<const T*>(InputReference.Reference)->Value, Reference);
+			SetValue(Context, static_cast<const T*>(InputReference.Reference)->Value, Reference);
 		}
 		else
 		{
-			SetValue(Context, *reinterpret_cast<const T*>(InputReference.Reference), Reference);
+			SetValue(Context, *static_cast<const T*>(InputReference.Reference), Reference);
 		}
 	}
 
