Iris
* Added ConsumeNetMetrics to ActorReplicationBridge.
* It can be used to collect internal Iris metrics and send them to an AnalyticsProvider

#rnx

diff --git a/Engine/Plugins/Runtime/ReplicationSystemTestPlugin/Source/Private/Tests/ReplicationSystem/TestReplicationSystem.cpp b/Engine/Plugins/Runtime/ReplicationSystemTestPlugin/Source/Private/Tests/ReplicationSystem/TestReplicationSystem.cpp
--- a/Engine/Plugins/Runtime/ReplicationSystemTestPlugin/Source/Private/Tests/ReplicationSystem/TestReplicationSystem.cpp
+++ b/Engine/Plugins/Runtime/ReplicationSystemTestPlugin/Source/Private/Tests/ReplicationSystem/TestReplicationSystem.cpp
@@ -7,6 +7,7 @@
 #include "Iris/ReplicationSystem/ReplicationSystemInternal.h"
 #include "Iris/ReplicationSystem/NetTokenStore.h"
 #include "Iris/Core/IrisLog.h"
+#include "Iris/Metrics/NetMetrics.h"
 #include "Misc/ScopeExit.h"
 
 namespace UE::Net::Private
@@ -1945,6 +1946,52 @@ UE_NET_TEST_FIXTURE(FReplicationSystemServerClientTestFixture, TestPartialDequan
 
 }
 
+UE_NET_TEST_FIXTURE(FReplicationSystemServerClientTestFixture, TestNetMetric)
+{
+	{
+		FNetMetric Metric(50.0);
+		UE_NET_ASSERT_TRUE(Metric.GetDataType() == FNetMetric::EDataType::Double);
+	}
+
+	{
+		FNetMetric Metric(50.f);
+		UE_NET_ASSERT_TRUE(Metric.GetDataType() == FNetMetric::EDataType::Double);
+	}
+
+	{
+		FNetMetric Metric;
+		UE_NET_ASSERT_TRUE(Metric.GetDataType() == FNetMetric::EDataType::None);
+		float Value = 100.f;
+		Metric.Set(Value);
+		UE_NET_ASSERT_TRUE(Metric.GetDataType() == FNetMetric::EDataType::Double);
+	}
+
+	{
+		FNetMetric Metric(5U);
+		UE_NET_ASSERT_TRUE(Metric.GetDataType() == FNetMetric::EDataType::Unsigned);
+	}
+
+	{
+		FNetMetric Metric;
+		UE_NET_ASSERT_TRUE(Metric.GetDataType() == FNetMetric::EDataType::None);
+		uint32 Value = 100U;
+		Metric.Set(Value);
+		UE_NET_ASSERT_TRUE(Metric.GetDataType() == FNetMetric::EDataType::Unsigned);
+	}
+
+	{
+		FNetMetric Metric(-5);
+		UE_NET_ASSERT_TRUE(Metric.GetDataType() == FNetMetric::EDataType::Signed);
+	}
+
+	{
+		FNetMetric Metric;
+		UE_NET_ASSERT_TRUE(Metric.GetDataType() == FNetMetric::EDataType::None);
+		Metric.Set(5);
+		UE_NET_ASSERT_TRUE(Metric.GetDataType() == FNetMetric::EDataType::Signed);
+	}
+}
+
 } // end namespace UE::Net::Private
 
 
diff --git a/Engine/Source/Runtime/Engine/Private/Net/Iris/ReplicationSystem/ActorReplicationBridge.cpp b/Engine/Source/Runtime/Engine/Private/Net/Iris/ReplicationSystem/ActorReplicationBridge.cpp
--- a/Engine/Source/Runtime/Engine/Private/Net/Iris/ReplicationSystem/ActorReplicationBridge.cpp
+++ b/Engine/Source/Runtime/Engine/Private/Net/Iris/ReplicationSystem/ActorReplicationBridge.cpp
@@ -23,6 +23,7 @@
 #include "Iris/Serialization/NetBitStreamWriter.h"
 #include "Iris/Serialization/ObjectNetSerializer.h"
 #include "Iris/Serialization/IrisObjectReferencePackageMap.h"
+#include "Iris/Metrics/NetMetrics.h"
 #include "Engine/Engine.h"
 #include "Engine/EngineTypes.h"
 #include "Engine/Level.h"
@@ -30,6 +31,7 @@
 #include "Engine/NetDriver.h"
 #include "Engine/Level.h"
 #include "Engine/World.h"
+#include "AnalyticsEventAttribute.h"
 #include "GameFramework/Actor.h"
 #include "GameFramework/PlayerController.h"
 #include "HAL/LowLevelMemStats.h"
@@ -43,6 +45,7 @@
 #include "ProfilingDebugging/AssetMetadataTrace.h"
 #include "Templates/Casts.h"
 #include "UObject/Package.h"
+
 #include <limits>
 
 #define UE_LOG_ACTORREPLICATIONBRIDGE(Category, Format, ...)	UE_LOG(LogIrisBridge, Category, TEXT("ActorReplicationBridge(%u)::") Format, GetReplicationSystem()->GetId(), ##__VA_ARGS__)
@@ -1396,6 +1399,52 @@ FString UActorReplicationBridge::PrintConnectionInfo(uint32 ConnectionId) const
 	}
 }
 
+void UActorReplicationBridge::ConsumeNetMetrics(TArray<FAnalyticsEventAttribute>& OutAttrs)
+{
+	using namespace UE::Net;
+
+	FNetMetrics MetricsCollector;
+
+	GetReplicationSystem()->CollectNetMetrics(MetricsCollector);
+
+	//$IRIS TODO: Add ResetNetMetrics call if the RepSystem starts caching cumulative stats
+
+	const TMap<FName, FNetMetric>& Metrics = MetricsCollector.GetMetrics();
+
+	for (auto& It : Metrics)
+	{
+		FName MetricName = It.Key;
+		const FNetMetric& Metric = It.Value;
+
+		switch (Metric.GetDataType())
+		{
+			case FNetMetric::EDataType::Signed:
+			{
+				const int32 Value = Metric.GetSigned();
+				OutAttrs.Add(FAnalyticsEventAttribute(MetricName.ToString(), Value));
+			} break;
+
+			case FNetMetric::EDataType::Unsigned:
+			{
+				const uint32 Value = Metric.GetUnsigned();
+				OutAttrs.Add(FAnalyticsEventAttribute(MetricName.ToString(), Value));
+			} break;
+
+			case FNetMetric::EDataType::Double:
+			{
+				const double Value = Metric.GetDouble();
+				OutAttrs.Add(FAnalyticsEventAttribute(MetricName.ToString(), Value));
+			} break;
+
+			default:
+			{
+				checkNoEntry();
+			} break;
+		}
+	}
+
+}
+
 #else //!UE_WITH_IRIS
 
 UActorReplicationBridge::UActorReplicationBridge() = default;
diff --git a/Engine/Source/Runtime/Engine/Private/Net/Iris/ReplicationSystem/ReplicationSystemUtil.cpp b/Engine/Source/Runtime/Engine/Private/Net/Iris/ReplicationSystem/ReplicationSystemUtil.cpp
--- a/Engine/Source/Runtime/Engine/Private/Net/Iris/ReplicationSystem/ReplicationSystemUtil.cpp
+++ b/Engine/Source/Runtime/Engine/Private/Net/Iris/ReplicationSystem/ReplicationSystemUtil.cpp
@@ -68,6 +68,11 @@ UReplicationSystem* FReplicationSystemUtil::GetReplicationSystem(const AActor* A
 	return NetDriver ? NetDriver->GetReplicationSystem() : nullptr;
 }
 
+UReplicationSystem* FReplicationSystemUtil::GetReplicationSystem(const UNetDriver* NetDriver)
+{
+	return NetDriver ? NetDriver->GetReplicationSystem() : nullptr;
+}
+
 UActorReplicationBridge* FReplicationSystemUtil::GetActorReplicationBridge(const AActor* Actor)
 {
 	if (UReplicationSystem* ReplicationSystem = GetReplicationSystem(Actor))
diff --git a/Engine/Source/Runtime/Engine/Public/Net/Iris/ReplicationSystem/ActorReplicationBridge.h b/Engine/Source/Runtime/Engine/Public/Net/Iris/ReplicationSystem/ActorReplicationBridge.h
--- a/Engine/Source/Runtime/Engine/Public/Net/Iris/ReplicationSystem/ActorReplicationBridge.h
+++ b/Engine/Source/Runtime/Engine/Public/Net/Iris/ReplicationSystem/ActorReplicationBridge.h
@@ -9,6 +9,7 @@
 
 class UNetDriver;
 class UIrisObjectReferencePackageMap;
+struct FAnalyticsEventAttribute;
 
 #if UE_WITH_IRIS
 
@@ -90,6 +91,13 @@ public:
 	/** Updates the level group for an actor that changed levels */
 	void ActorChangedLevel(const AActor* Actor, const ULevel* PreviousLevel);
 
+	/**
+	 * Add relevant network metrics gathered since the last call to ConsumeNetMetrics.
+	 * Any periodic stat will be reset here too.
+	 * @param OutAttrs A list of Name/Value pairings that will be sent to an AnalyticsProvider
+	 */
+	ENGINE_API void ConsumeNetMetrics(TArray<FAnalyticsEventAttribute>& OutAttrs);
+
 	using UObjectReplicationBridge::EndReplication;
 
 protected:
diff --git a/Engine/Source/Runtime/Engine/Public/Net/Iris/ReplicationSystem/ReplicationSystemUtil.h b/Engine/Source/Runtime/Engine/Public/Net/Iris/ReplicationSystem/ReplicationSystemUtil.h
--- a/Engine/Source/Runtime/Engine/Public/Net/Iris/ReplicationSystem/ReplicationSystemUtil.h
+++ b/Engine/Source/Runtime/Engine/Public/Net/Iris/ReplicationSystem/ReplicationSystemUtil.h
@@ -19,6 +19,7 @@ class FRepChangedPropertyTracker;
 class UReplicationSystem;
 class UWorld;
 class UNetConnection;
+class UNetDriver;
 
 namespace UE::Net
 {
@@ -34,6 +35,9 @@ struct FReplicationSystemUtil
 	/** Returns the UReplicationSystem for the main NetDriver assigned to the Actor. Note that an Actor may be replicated by multiple ReplicationSystems. May return null. */
 	ENGINE_API static UReplicationSystem* GetReplicationSystem(const AActor* Actor);
 
+	/** Returns the UReplicationSystem for a specific NetDriver. May return null. */
+	ENGINE_API static UReplicationSystem* GetReplicationSystem(const UNetDriver* NetDriver);
+
 	/** Returns the UActorReplicationBridge of the UReplicationSystem belogning to the main NetDriver assigned to the Actor. May return null. */
 	ENGINE_API static UActorReplicationBridge* GetActorReplicationBridge(const AActor* Actor);
 
diff --git a/Engine/Source/Runtime/Experimental/Iris/Core/Private/Iris/ReplicationSystem/ReplicationSystem.cpp b/Engine/Source/Runtime/Experimental/Iris/Core/Private/Iris/ReplicationSystem/ReplicationSystem.cpp
--- a/Engine/Source/Runtime/Experimental/Iris/Core/Private/Iris/ReplicationSystem/ReplicationSystem.cpp
+++ b/Engine/Source/Runtime/Experimental/Iris/Core/Private/Iris/ReplicationSystem/ReplicationSystem.cpp
@@ -28,6 +28,8 @@
 #include "Iris/Serialization/NetSerializer.h"
 #include "Iris/Serialization/IrisObjectReferencePackageMap.h"
 
+#include "Iris/Metrics/NetMetrics.h"
+
 #include "HAL/IConsoleManager.h"
 #include "ProfilingDebugging/CsvProfiler.h"
 #include "UObject/UObjectGlobals.h"
@@ -62,6 +64,9 @@ public:
 	FNetBitArray ConnectionsPendingPostTickDispatchSend;
 	EReplicationSystemSendPass CurrentSendPass = EReplicationSystemSendPass::Invalid;
 
+	FName MetricNameTotalRootObjects;
+	FName MetricNameTotalSubObjects;
+
 	explicit FReplicationSystemImpl(UReplicationSystem* InReplicationSystem, const UReplicationSystem::FReplicationSystemParams& Params)
 	: ReplicationSystem(InReplicationSystem)
 	, ReplicationSystemInternal(
@@ -73,6 +78,8 @@ public:
 			.MaxReplicationWriterObjectCount = Params.MaxReplicationWriterObjectCount,
 		}))
 	{
+		MetricNameTotalRootObjects = TEXT("TotalSubObjects");
+		MetricNameTotalSubObjects = TEXT("TotalRootObjects");
 	}
 
 	~FReplicationSystemImpl()
@@ -639,6 +646,20 @@ public:
 		const FNetBitArray& ValidConnections = Connections.GetValidConnections();
 		ValidConnections.ForAllSetBits(UpdateUnresolvableReferenceTracking);
 	}
+
+	void CollectNetMetrics(UE::Net::FNetMetrics& OutNetMetrics) const
+	{
+		using namespace UE::Net;
+
+		const FNetRefHandleManager& NetRefHandleManager = ReplicationSystemInternal.GetNetRefHandleManager();
+
+		const uint32 TotalNetObjects = NetRefHandleManager.GetActiveObjectCount();
+		const uint32 TotalSubObjects = NetRefHandleManager.GetSubObjectInternalIndicesView().CountSetBits();
+
+		// Collect stats on total replicated objects
+		OutNetMetrics.EmplaceMetric(MetricNameTotalRootObjects, FNetMetric(TotalNetObjects-TotalSubObjects));
+		OutNetMetrics.EmplaceMetric(MetricNameTotalSubObjects, FNetMetric(TotalSubObjects));
+	}
 }; // end class FReplicationSystemImpl
 
 } // end namespace UE::Net::Private
@@ -870,6 +891,11 @@ void UReplicationSystem::PostSendUpdate()
 
 			FNetTypeStats& TypeStats = Impl->ReplicationSystemInternal.GetNetTypeStats();
 			TypeStats.ReportCSVStats();
+
+			if (Impl->ReplicationSystemInternal.IsDirtyNetObjectTrackerInitialized())
+			{
+				Impl->ReplicationSystemInternal.GetDirtyNetObjectTracker().ReportCSVStats();
+			}
 		}
 #endif
 
@@ -1730,6 +1756,11 @@ void UReplicationSystem::ReportErrorWithNetRefHandle(UE::Net::ENetRefHandleError
 	Impl->ReplicationSystemInternal.GetReplicationBridge()->OnErrorWithNetRefHandleReported(ErrorType, NetRefHandle, ConnectionId);
 }
 
+void UReplicationSystem::CollectNetMetrics(UE::Net::FNetMetrics& OutNetMetrics) const
+{
+	Impl->CollectNetMetrics(OutNetMetrics);
+}
+
 
 #pragma region ReplicationSystemFactory
 
diff --git a/Engine/Source/Runtime/Experimental/Iris/Core/Private/Iris/Stats/NetStats.h b/Engine/Source/Runtime/Experimental/Iris/Core/Private/Iris/Stats/NetStats.h
--- a/Engine/Source/Runtime/Experimental/Iris/Core/Private/Iris/Stats/NetStats.h
+++ b/Engine/Source/Runtime/Experimental/Iris/Core/Private/Iris/Stats/NetStats.h
@@ -60,9 +60,6 @@ public:
 	/** Add time in seconds waiting to be able to continue sending huge object. */
 	void AddHugeObjectStallTime(double Seconds);
 
-	/** Add time spent on trying to replicate an object in the packet but that ended up not fitting. */
-	void AddReplicationWasteTime(double Seconds);
-
 	/** Add stats from another instance. */
 	IRISCORE_API void Accumulate(const FNetSendStats& Stats);
 
@@ -128,6 +125,11 @@ inline void FNetSendStats::AddNumberOfReplicatedDestructionInfos(uint32 Count)
 	Stats.ReplicatedDestructionInfoCount += Count;
 }
 
+inline uint32 FNetSendStats::GetNumberOfReplicatedRootObjects() const
+{
+	return Stats.ReplicatedRootObjectCount;
+}
+
 inline uint32 FNetSendStats::GetNumberOfReplicatedObjects() const
 {
 	return Stats.ReplicatedObjectCount;
diff --git a/Engine/Source/Runtime/Experimental/Iris/Core/Public/Iris/Metrics/NetMetrics.h b/Engine/Source/Runtime/Experimental/Iris/Core/Public/Iris/Metrics/NetMetrics.h
new file mode 100644
--- /dev/null
+++ b/Engine/Source/Runtime/Experimental/Iris/Core/Public/Iris/Metrics/NetMetrics.h
@@ -0,0 +1,120 @@
+// Copyright Epic Games, Inc. All Rights Reserved.
+
+#pragma once
+
+#include "HAL/Platform.h"
+#include "Containers/Map.h"
+#include "UObject/NameTypes.h"
+#include "Misc/AssertionMacros.h"
+
+#include <type_traits>
+
+
+namespace UE::Net
+{
+
+/**
+ * Class used to store a single analytics value.
+ * Only supports integers or floating points for now.
+ */
+struct FNetMetric
+{
+public:
+	
+	enum class EDataType
+	{
+		None,
+		Unsigned,
+		Signed,
+		Double,
+	};
+
+public:
+
+	FNetMetric()
+		: Double(0.0)
+		, DataType(EDataType::None)
+	{}
+
+	template<typename T>
+	FNetMetric(T InValue)
+	{
+		static_assert(std::is_integral_v<T> || std::is_floating_point_v<T>, "Only integers and floats are supported");
+
+		if constexpr (std::is_integral_v<T> && std::is_signed_v<T>)
+		{
+			DataType = EDataType::Signed;
+			Signed = InValue;
+		}
+		else if constexpr (std::is_integral_v<T> && std::is_unsigned_v<T>)
+		{
+			DataType = EDataType::Unsigned;
+			Unsigned = InValue;
+		}
+		else if constexpr (std::is_floating_point_v<T>)
+		{
+			DataType = EDataType::Double;
+			Double = InValue;
+		}
+	}
+
+	template<typename T>
+	void Set(T InValue)
+	{
+		static_assert(std::is_integral_v<T> || std::is_floating_point_v<T>, "Only integers and floats are supported");
+		*this = FNetMetric(InValue);
+	}
+
+	EDataType GetDataType() const 
+	{ 
+		return DataType; 
+	}
+
+	int32 GetSigned() const
+	{
+		check(DataType == EDataType::Signed);
+		return Signed;
+	}
+
+	uint32 GetUnsigned() const
+	{
+		check(DataType == EDataType::Unsigned);
+		return Unsigned;
+	}
+
+	double GetDouble() const
+	{
+		check(DataType == EDataType::Double);
+		return Double;
+	}
+
+private:
+
+	union
+	{
+		uint32 Unsigned;
+		int32 Signed;
+		double Double;
+	};
+
+	EDataType DataType;
+};
+
+/**
+* Collects network metrics and keeps track of their name.
+*/
+struct FNetMetrics
+{
+public:
+
+	void EmplaceMetric(FName InName, FNetMetric&& InMetric)		{ Metrics.Emplace(InName, MoveTemp(InMetric)); }
+	void AddMetric(FName InName, const FNetMetric& InMetric)	{ Metrics.Add(InName, InMetric); }
+
+	const TMap<FName, FNetMetric>& GetMetrics() const { return Metrics; }
+
+private:
+
+	TMap<FName, FNetMetric> Metrics;
+};
+
+} // end namespace UE::Net
\ No newline at end of file
diff --git a/Engine/Source/Runtime/Experimental/Iris/Core/Public/Iris/ReplicationSystem/ReplicationSystem.h b/Engine/Source/Runtime/Experimental/Iris/Core/Public/Iris/ReplicationSystem/ReplicationSystem.h
--- a/Engine/Source/Runtime/Experimental/Iris/Core/Public/Iris/ReplicationSystem/ReplicationSystem.h
+++ b/Engine/Source/Runtime/Experimental/Iris/Core/Public/Iris/ReplicationSystem/ReplicationSystem.h
@@ -24,23 +24,28 @@ class UReplicationBridge;
 class UReplicationSystem;
 namespace UE::Net
 {
-	class FNetBitArray;
-	class FNetCullDistanceOverrides;
-	enum class ENetFilterStatus : uint32;
-	enum class ENetRefHandleError : uint32;
-	class FNetObjectAttachment;
-	enum class ENetObjectDeltaCompressionStatus : unsigned;
 	typedef uint32 FNetObjectFilterHandle;
 	typedef uint32 FNetObjectPrioritizerHandle;
+
+	class FNetBitArray;
+	class FNetCullDistanceOverrides;
+	class FNetObjectAttachment;	
 	class FNetObjectReference;
-	enum class EReplicationCondition : uint32;
-	struct FReplicationProtocol;
 	class FReplicationSystemFactory;
-	struct FReplicationSystemUtil;
-	struct FReplicationView;
 	class FStringTokenStore;
 	class FWorldLocations;
+
+	struct FReplicationSystemUtil;
+	struct FReplicationView;
 	struct FNetDebugName;
+	struct FNetMetrics;
+	struct FReplicationProtocol;
+
+	enum class ENetFilterStatus : uint32;
+	enum class ENetObjectDeltaCompressionStatus : unsigned;
+	enum class ENetRefHandleError : uint32;
+	enum class EReplicationCondition : uint32;
+
 	namespace Private
 	{
 		class FReplicationSystemImpl;
@@ -657,6 +662,9 @@ public:
 	/** Called when a connection reports a critical error with a netrefhandle object */
 	IRISCORE_API void ReportErrorWithNetRefHandle(UE::Net::ENetRefHandleError ErrorType, uint64 NetRefHandleId, uint32 ConnectionId);
 
+	/** Collect relevant metrics that could be sent to analytics */
+	IRISCORE_API void CollectNetMetrics(UE::Net::FNetMetrics& OutNetMetrics) const;
+
 public:
 	// For internal use and not exported.
 	UE::Net::Private::FReplicationSystemInternal* GetReplicationSystemInternal();
