Fix  sweep test with big triangle mesh
- Add a double precision GJK version.
- Change the Simplex size from a VectorInt to a int32, previously it was done to save some branching operations, but compare to the load and store operations it was not worth it.
- This solution is only for triangle mesh. Height field never triggered this issue because it is usually composed of small triangles. However in the editor the plane shape generates a big triangle mesh which cause this issue when user add it for the entire world.

diff --git a/Engine/Source/Programs/HeadlessChaos/Private/HeadlessChaosTestGJK.cpp b/Engine/Source/Programs/HeadlessChaos/Private/HeadlessChaosTestGJK.cpp
--- a/Engine/Source/Programs/HeadlessChaos/Private/HeadlessChaosTestGJK.cpp
+++ b/Engine/Source/Programs/HeadlessChaos/Private/HeadlessChaosTestGJK.cpp
@@ -11,6 +11,7 @@
 #include "Chaos/ImplicitObjectScaled.h"
 #include "Chaos/Collision/PBDCollisionConstraint.h"
 #include "Chaos/Triangle.h"
+#include "Chaos/TriangleMeshImplicitObject.h"
 #include "Chaos/TriangleRegister.h"
 
 namespace ChaosTest
@@ -2102,4 +2103,231 @@ namespace ChaosTest
 		EXPECT_TRUE(GJKIntersection(BigBoxScaled, SmallBox, BToATM, FReal(0), Chaos::TVector<FReal, 3>(-16000, -16000, 500)));		
 
 	}
+
+	// Test known large triangle failure case
+	GTEST_TEST(GJKTests, GJK_LargeTriangleCase)
+	{
+		VectorRegister4Float TriMeshScaleVector = MakeVectorRegister(2000.0f, 2000.0f, 1.0f, 0.0f);
+
+		VectorRegister4Float A = MakeVectorRegister(100.0f, 0.0f, 100.0f, 0.0f);
+		VectorRegister4Float B = MakeVectorRegister(0.0f, 100.0f, 100.0f, 0.0f);
+		VectorRegister4Float C = MakeVectorRegister(0.0f, 0.0f, 100.0f, 0.0f);
+
+		A = VectorMultiply(A, TriMeshScaleVector);
+		B = VectorMultiply(B, TriMeshScaleVector);
+		C = VectorMultiply(C, TriMeshScaleVector);
+
+		FTriangleRegister Tri(A, B, C);
+		const VectorRegister4Float TriNormal = VectorCross(VectorSubtract(B, A), VectorSubtract(C, A));
+
+		FRealSingle Time;
+		VectorRegister4Float OutPositionSimd, OutNormalSimd;
+		Chaos::TSphere<Chaos::FReal, 3> QueryGeom({ 0, 0, 0 }, 5.0f);
+		VectorRegister4Float RotationSimd = MakeVectorRegister(0.0f, 0.0f, 0.0f, 0.0f);
+
+		VectorRegister4Float TranslationSimd = MakeVectorRegister(100000.0f, 100022.820f, 213.002823f, 0.0f);
+		VectorRegister4Float RayDirSimd = MakeVectorRegister(-2.22044605e-16f, 0.0f, -1.0f, 0.0f);
+		float LengthScale = 1.0f;
+		float CurDataLength = 108.002823f;
+		const bool bComputeMTD = true;
+		bool bFoundIntersection;
+		if (Tri.IsTooBigForSinglePrecision())
+		{
+			VectorRegister4Double OutPositionDouble, OutNormalDouble;
+			bFoundIntersection = GJKRaycast2ImplSimd<VectorRegister4Double>(Tri, QueryGeom, VectorRegister4Double(RotationSimd), VectorRegister4Double(TranslationSimd), VectorRegister4Double(RayDirSimd), LengthScale * CurDataLength, Time, OutPositionDouble, OutNormalDouble, bComputeMTD, GlobalVectorConstants::Double1000);
+			OutPositionSimd = MakeVectorRegisterFloatFromDouble(OutPositionDouble);
+			OutNormalSimd = MakeVectorRegisterFloatFromDouble(OutNormalDouble);
+		}
+		else
+		{
+			bFoundIntersection = GJKRaycast2ImplSimd<VectorRegister4Float>(Tri, QueryGeom, RotationSimd, TranslationSimd, RayDirSimd, LengthScale * CurDataLength, Time, OutPositionSimd, OutNormalSimd, bComputeMTD, GlobalVectorConstants::FloatMinusOne);
+		}
+
+		EXPECT_FALSE(bFoundIntersection);
+	}
+
+	GTEST_TEST(GJKTests, GJK_SweepConvexLargeTriangleInMesh)
+	{
+		using namespace Chaos;
+		TArray<FConvex::FVec3Type> ConvexSurface(
+			{
+			{5.0, -5.0, -5.0},
+			{-5.0, 5.0, -5.0},
+			{5.0, 5.0, -5.0},
+			{5.0, 5.0, 5.0},
+			{-5.0, -5.0, -5.0},
+			{-5.0, -5.0, 5.0},
+			{5.0, -5.0, 5.0},
+			{-5.0, 5.0, 5.0},
+			});
+
+		FConvex ConvexBox(MoveTemp(ConvexSurface), 0.0f);
+
+		FTriangleMeshImplicitObject::ParticlesType TrimeshParticles(
+			{
+				{500.0000000, 500.0000000, 0.0},
+				{500.0000000, -500.0000000, 0.0},
+				{-500.0000000, 500.0000000, 0.0},
+				{-500.0000000, -500.0000000, 0.0}
+			});
+
+		TArray<TVec3<int32>> Indices;
+		Indices.Emplace(1, 0, 2);
+		Indices.Emplace(1, 2, 3);
+
+		TArray<uint16> Materials;
+		Materials.Emplace(0);
+		Materials.Emplace(0);
+		FTriangleMeshImplicitObjectPtr TriangleMesh(new FTriangleMeshImplicitObject(MoveTemp(TrimeshParticles), MoveTemp(Indices), MoveTemp(Materials)));
+		TImplicitObjectScaled<FTriangleMeshImplicitObject> ScaledTriangleMesh = TImplicitObjectScaled<FTriangleMeshImplicitObject>(TriangleMesh, FVec3(10000.0, 10000.0, 1.0));
+
+		FQuat Rotation(0.00488796039, 0.00569311855, -0.000786740216, 0.999971569);
+		FVec3 Translation(10500.365723, -14500.4132690, 8.2289352);
+
+		TRigidTransform<FReal, 3> Transform(Translation, Rotation);
+
+
+		FVec3 Dir(-0.00339674903, 5.76980747e-05, -0.999994159);
+		FReal Length = 10.0;
+
+		FReal OutTime = -1;
+		FVec3 Normal(0.0f);
+		FVec3 Position(0.0f);
+		int32 FaceIndex = -1;
+		FVec3 FaceNormal(0.0);
+		bool bResult = ScaledTriangleMesh.LowLevelSweepGeom(ConvexBox, Transform, Dir, Length, OutTime, Position, Normal, FaceIndex, FaceNormal, 0.0f, true);
+		EXPECT_TRUE(bResult);
+		EXPECT_EQ(Position.Z, 0.0);
+		EXPECT_NEAR(Position.X, 10505.365723, 1);
+		EXPECT_NEAR(Position.Y, -14505.4132690, 1);
+	}
+
+	GTEST_TEST(GJKTests, GJK_SweepBoxLargeTriangleInMesh)
+	{
+		Chaos::TBox<Chaos::FReal, 3> QueryGeom({ -5, -5, -5 }, { 5, 5, 5 });
+
+		FTriangleMeshImplicitObject::ParticlesType TrimeshParticles(
+			{
+				{500.0000000, 500.0000000, 0.0},
+				{500.0000000, -500.0000000, 0.0},
+				{-500.0000000, 500.0000000, 0.0},
+				{-500.0000000, -500.0000000, 0.0}
+			});
+
+		TArray<TVec3<int32>> Indices;
+		Indices.Emplace(1, 0, 2);
+		Indices.Emplace(1, 2, 3);
+
+		TArray<uint16> Materials;
+		Materials.Emplace(0);
+		Materials.Emplace(0);
+		FTriangleMeshImplicitObjectPtr TriangleMesh(new FTriangleMeshImplicitObject(MoveTemp(TrimeshParticles), MoveTemp(Indices), MoveTemp(Materials)));
+		TImplicitObjectScaled<FTriangleMeshImplicitObject> ScaledTriangleMesh = TImplicitObjectScaled<FTriangleMeshImplicitObject>(TriangleMesh, FVec3(10000.0, 10000.0, 1.0));
+
+		FQuat Rotation(0.00488796039, 0.00569311855, -0.000786740216, 0.999971569);
+		FVec3 Translation(10500.365723, -14500.4132690, 8.2289352);
+
+		TRigidTransform<FReal, 3> Transform(Translation, Rotation);
+
+
+		FVec3 Dir(-0.00339674903, 5.76980747e-05, -0.999994159);
+		FReal Length = 10.0;
+
+		FReal OutTime = -1;
+		FVec3 Normal(0.0f);
+		FVec3 Position(0.0f);
+		int32 FaceIndex = -1;
+		FVec3 FaceNormal(0.0);
+		bool bResult = ScaledTriangleMesh.LowLevelSweepGeom(QueryGeom, Transform, Dir, Length, OutTime, Position, Normal, FaceIndex, FaceNormal, 0.0f, true);
+		EXPECT_TRUE(bResult);
+		EXPECT_EQ(Position.Z, 0.0);
+		EXPECT_NEAR(Position.X, 10505.365723, 1);
+		EXPECT_NEAR(Position.Y, -14505.4132690, 1);
+	}
+
+	GTEST_TEST(GJKTests, GJK_SweepSphereLargeTriangleInMesh)
+	{
+		Chaos::TSphere<Chaos::FReal, 3> QueryGeom({ 0, 0, 0 }, 5.0f);
+
+		FTriangleMeshImplicitObject::ParticlesType TrimeshParticles(
+			{
+				{500.0000000, 500.0000000, 0.0},
+				{500.0000000, -500.0000000, 0.0},
+				{-500.0000000, 500.0000000, 0.0},
+				{-500.0000000, -500.0000000, 0.0}
+			});
+
+		TArray<TVec3<int32>> Indices;
+		Indices.Emplace(1, 0, 2);
+		Indices.Emplace(1, 2, 3);
+
+		TArray<uint16> Materials;
+		Materials.Emplace(0);
+		Materials.Emplace(0);
+		FTriangleMeshImplicitObjectPtr TriangleMesh(new FTriangleMeshImplicitObject(MoveTemp(TrimeshParticles), MoveTemp(Indices), MoveTemp(Materials)));
+		TImplicitObjectScaled<FTriangleMeshImplicitObject> ScaledTriangleMesh = TImplicitObjectScaled<FTriangleMeshImplicitObject>(TriangleMesh, FVec3(10000.0, 10000.0, 1.0));
+
+		FQuat Rotation(0.00488796039, 0.00569311855, -0.000786740216, 0.999971569);
+		FVec3 Translation(10500.365723, -14500.4132690, 8.2289352);
+
+		TRigidTransform<FReal, 3> Transform(Translation, Rotation);
+
+
+		FVec3 Dir(-0.00339674903, 5.76980747e-05, -0.999994159);
+		FReal Length = 10.0;
+
+		FReal OutTime = -1;
+		FVec3 Normal(0.0f);
+		FVec3 Position(0.0f);
+		int32 FaceIndex = -1;
+		FVec3 FaceNormal(0.0);
+		bool bResult = ScaledTriangleMesh.LowLevelSweepGeom(QueryGeom, Transform, Dir, Length, OutTime, Position, Normal, FaceIndex, FaceNormal, 0.0f, true);
+		EXPECT_TRUE(bResult);
+		EXPECT_EQ(Position.Z, 0.0);
+		EXPECT_NEAR(Position.X, 10500.365723, 1e-1);
+		EXPECT_NEAR(Position.Y, -14500.4132690, 1e-1);
+	}
+
+	GTEST_TEST(GJKTests, GJK_SweepCapsuleLargeTriangleInMesh)
+	{
+		Chaos::FCapsule QueryGeom({ 0, 0, 1.0 }, { 0, 0, -1.0 }, 5.0f);
+
+		FTriangleMeshImplicitObject::ParticlesType TrimeshParticles(
+			{
+				{500.0000000, 500.0000000, 0.0},
+				{500.0000000, -500.0000000, 0.0},
+				{-500.0000000, 500.0000000, 0.0},
+				{-500.0000000, -500.0000000, 0.0}
+			});
+
+		TArray<TVec3<int32>> Indices;
+		Indices.Emplace(1, 0, 2);
+		Indices.Emplace(1, 2, 3);
+
+		TArray<uint16> Materials;
+		Materials.Emplace(0);
+		Materials.Emplace(0);
+		FTriangleMeshImplicitObjectPtr TriangleMesh(new FTriangleMeshImplicitObject(MoveTemp(TrimeshParticles), MoveTemp(Indices), MoveTemp(Materials)));
+		TImplicitObjectScaled<FTriangleMeshImplicitObject> ScaledTriangleMesh = TImplicitObjectScaled<FTriangleMeshImplicitObject>(TriangleMesh, FVec3(10000.0, 10000.0, 1.0));
+
+		FQuat Rotation(0.00488796039, 0.00569311855, -0.000786740216, 0.999971569);
+		FVec3 Translation(10500.365723, -14500.4132690, 8.2289352);
+
+		TRigidTransform<FReal, 3> Transform(Translation, Rotation);
+
+
+		FVec3 Dir(-0.00339674903, 5.76980747e-05, -0.999994159);
+		FReal Length = 10.0;
+
+		FReal OutTime = -1;
+		FVec3 Normal(0.0f);
+		FVec3 Position(0.0f);
+		int32 FaceIndex = -1;
+		FVec3 FaceNormal(0.0);
+		bool bResult = ScaledTriangleMesh.LowLevelSweepGeom(QueryGeom, Transform, Dir, Length, OutTime, Position, Normal, FaceIndex, FaceNormal, 0.0f, true);
+		EXPECT_TRUE(bResult);
+		EXPECT_EQ(Position.Z, 0.0);
+		EXPECT_NEAR(Position.X, 10500.365723, 1e-1);
+		EXPECT_NEAR(Position.Y, -14500.4132690, 1e-1);
+	}
 }
\ No newline at end of file
diff --git a/Engine/Source/Runtime/Experimental/Chaos/Private/Chaos/TriangleMeshImplicitObject.cpp b/Engine/Source/Runtime/Experimental/Chaos/Private/Chaos/TriangleMeshImplicitObject.cpp
--- a/Engine/Source/Runtime/Experimental/Chaos/Private/Chaos/TriangleMeshImplicitObject.cpp
+++ b/Engine/Source/Runtime/Experimental/Chaos/Private/Chaos/TriangleMeshImplicitObject.cpp
@@ -1310,7 +1310,19 @@ struct FTriangleMeshSweepVisitor
 			}
 		}
 		FRealSingle Time;
-		if(GJKRaycast2ImplSimd(Tri, QueryGeom, RotationSimd, TranslationSimd, RayDirSimd, LengthScale * CurDataLength, Time, OutPositionSimd, OutNormalSimd, bComputeMTD, GlobalVectorConstants::Float1000))
+		bool bIsHitting = false;
+		if (UNLIKELY(Tri.IsTooBigForSinglePrecision()))
+		{
+			VectorRegister4Double OutPositionDouble, OutNormalDouble;
+			bIsHitting = GJKRaycast2ImplSimd<VectorRegister4Double>(Tri, QueryGeom, VectorRegister4Double(RotationSimd), VectorRegister4Double(TranslationSimd), VectorRegister4Double(RayDirSimd), LengthScale * CurDataLength, Time, OutPositionDouble, OutNormalDouble, bComputeMTD, GlobalVectorConstants::Double1000);
+			OutPositionSimd = MakeVectorRegisterFloatFromDouble(OutPositionDouble);
+			OutNormalSimd = MakeVectorRegisterFloatFromDouble(OutNormalDouble);
+		}
+		else
+		{
+			bIsHitting = GJKRaycast2ImplSimd(Tri, QueryGeom, RotationSimd, TranslationSimd, RayDirSimd, LengthScale * CurDataLength, Time, OutPositionSimd, OutNormalSimd, bComputeMTD, GlobalVectorConstants::Float1000);
+		}
+		if(bIsHitting)
 		{
 			// Don't return back faces if they are not initially overlapping
 			if (Time > 0 && VectorMaskBits(IsBackFace))
diff --git a/Engine/Source/Runtime/Experimental/Chaos/Public/Chaos/EPAVectorized.h b/Engine/Source/Runtime/Experimental/Chaos/Public/Chaos/EPAVectorized.h
--- a/Engine/Source/Runtime/Experimental/Chaos/Public/Chaos/EPAVectorized.h
+++ b/Engine/Source/Runtime/Experimental/Chaos/Public/Chaos/EPAVectorized.h
@@ -360,7 +360,6 @@ FORCEINLINE_DEBUGGABLE void VectorEPAComputeVisibilityBorder(TEPAWorkingArray<Ve
 	}
 }
 
-
 FORCEINLINE_DEBUGGABLE void VectorComputeEPAResults(const VectorRegister4Float* VertsA, const VectorRegister4Float* VertsB, const VectorTEPAEntry& Entry, VectorRegister4Float& OutPenetration, VectorRegister4Float& OutDir, VectorRegister4Float& OutA, VectorRegister4Float& OutB, EEPAResult& ResultStatus)
 {
 	//NOTE: We use this function as fallback when robustness breaks. So - do not assume adjacency is valid as these may be new uninitialized traingles that failed
@@ -369,8 +368,7 @@ FORCEINLINE_DEBUGGABLE void VectorComputeEPAResults(const VectorRegister4Float*
 	VectorRegister4Float Bs[4] = { VertsB[Entry.IdxBuffer[0]], VertsB[Entry.IdxBuffer[1]], VertsB[Entry.IdxBuffer[2]] };
 	VectorRegister4Float Simplex[4] = { VectorSubtract(As[0], Bs[0]), VectorSubtract(As[1], Bs[1]), VectorSubtract(As[2], Bs[2]) };
 	VectorRegister4Float Barycentric;
-	constexpr VectorRegister4Int three = MakeVectorRegisterIntConstant(3, 3, 3, 3);
-	VectorRegister4Int NumVerts = three;
+	int32 NumVerts = 3;
 	OutDir = VectorSimplexFindClosestToOrigin(Simplex, NumVerts, Barycentric, As, Bs);
 
 	const VectorRegister4Float DotDir = VectorDot4(OutDir, OutDir);
@@ -416,12 +414,7 @@ FORCEINLINE_DEBUGGABLE void VectorComputeEPAResults(const VectorRegister4Float*
 	Barycentrics[2] = VectorSwizzle(Barycentric, 2, 2, 2, 2);
 	Barycentrics[3] = VectorSwizzle(Barycentric, 3, 3, 3, 3);
 
-
-	alignas(16) int32 NumVertsInts[4];
-	VectorIntStoreAligned(NumVerts, NumVertsInts);
-	const int32 NumVertsInt = NumVertsInts[0];
-
-	for (int i = 0; i < NumVertsInt; ++i)
+	for (int i = 0; i < NumVerts; ++i)
 	{
 		OutA = VectorMultiplyAdd(As[i], Barycentrics[i], OutA);
 		OutB = VectorMultiplyAdd(Bs[i], Barycentrics[i], OutB);
diff --git a/Engine/Source/Runtime/Experimental/Chaos/Public/Chaos/GJK.h b/Engine/Source/Runtime/Experimental/Chaos/Public/Chaos/GJK.h
--- a/Engine/Source/Runtime/Experimental/Chaos/Public/Chaos/GJK.h
+++ b/Engine/Source/Runtime/Experimental/Chaos/Public/Chaos/GJK.h
@@ -169,7 +169,7 @@ namespace Chaos
 
 		VectorRegister4Float SimplexSimd[4] = { VectorZeroFloat(), VectorZeroFloat(), VectorZeroFloat(), VectorZeroFloat() };
 		VectorRegister4Float BarycentricSimd;
-		VectorRegister4Int NumVerts = GlobalVectorConstants::IntZero;
+		int32 NumVerts = 0;
 
 		FReal ThicknessA;
 		FReal ThicknessB;
@@ -201,17 +201,9 @@ namespace Chaos
 				return false;
 			}
 
-			{
-				// Convert simdInt to int
-				alignas(16) int32 NumVertsInts[4];
-				VectorIntStoreAligned(NumVerts, NumVertsInts);
-				const int32 NumVertsInt = NumVertsInts[0];
-				SimplexSimd[NumVertsInt] = WSimd;
-			}
+			SimplexSimd[NumVerts++] = WSimd;
 
-			NumVerts = VectorIntAdd(NumVerts, GlobalVectorConstants::IntOne);
-
-			V = VectorSimplexFindClosestToOrigin<false>(SimplexSimd, NumVerts, BarycentricSimd, nullptr, nullptr);
+			V = VectorSimplexFindClosestToOrigin<VectorRegister4Float, false>(SimplexSimd, NumVerts, BarycentricSimd, nullptr, nullptr);
 
 			const VectorRegister4Float NewDist2Simd = VectorDot3(V, V);///
 			bNearZero = VectorMaskBits(VectorCompareLT(NewDist2Simd, Inflation2)) != 0;
@@ -253,8 +245,7 @@ namespace Chaos
 
 		FSimplex SimplexIDs;
 		VectorRegister4Float Simplex[4] = { VectorZero(), VectorZero(), VectorZero(), VectorZero() };
-		VectorRegister4Int NumVerts = GlobalVectorConstants::IntZero;
-		alignas(16) int32 NumVertsInts[4];
+		int32 NumVerts = 0;
 		VectorRegister4Float Barycentric = VectorZero();
 
 		bool bTerminate;
@@ -290,12 +281,9 @@ namespace Chaos
 				return false;
 			}
 
-			VectorIntStoreAligned(NumVerts, NumVertsInts);
-			const int32 NumVertsInt = NumVertsInts[0];
-			Simplex[NumVertsInt] = W;
-			NumVerts = VectorIntAdd(NumVerts, GlobalVectorConstants::IntOne);
+			Simplex[NumVerts++] = W;
 
-			V = VectorSimplexFindClosestToOrigin<false>(Simplex, NumVerts, Barycentric, nullptr, nullptr);
+			V = VectorSimplexFindClosestToOrigin<VectorRegister4Float, false>(Simplex, NumVerts, Barycentric, nullptr, nullptr);
 			const VectorRegister4Float NewDist2 = VectorDot3(V, V);
 			
 			bNearZero = static_cast<bool>(VectorMaskBits(VectorCompareGT(Inflation2Simd, NewDist2)));
@@ -1519,11 +1507,12 @@ namespace Chaos
 	}
 	
 
+
 	/** Sweeps one geometry against the other
 	 @A The first geometry
 	 @B The second geometry
 	 @BToARotation B's starting rotation in A's local space
-         @StartPoint B's starting position in A's local space
+	 @StartPoint B's starting position in A's local space
 	 @RayDir The ray's direction (normalized)
 	 @RayLength The ray's length
 	 @OutTime The time along the ray when the objects first overlap
@@ -1533,9 +1522,9 @@ namespace Chaos
 	 @return True if the geometries overlap during the sweep, False otherwise
 	 @note If A overlaps B at the start of the ray ("initial overlap" condition) then this function returns true, and sets OutTime = 0, but does not set any other output variables.
 	 */
-	template <typename T>
-	bool GJKRaycast2ImplSimd(const FGeomGJKHelperSIMD& A, const FGeomGJKHelperSIMD& B, const VectorRegister4Float& BToARotation, const VectorRegister4Float& StartPoint, const VectorRegister4Float& RayDir,
-		T RayLength, T& OutTime, VectorRegister4Float& OutPosition, VectorRegister4Float& OutNormal, bool bComputeMTD, const VectorRegister4Float& InitialDir)
+	template<typename T>
+	bool GJKRaycast2ImplSimd(const FGeomGJKHelperSIMD& RESTRICT A, const FGeomGJKHelperSIMD& RESTRICT B, const T& BToARotation, const T& StartPoint, const T& RayDir,
+		FRealSingle RayLength, FRealSingle& OutTime, T& OutPosition, T& OutNormal, bool bComputeMTD, const T& InitialDir)
 	{
 		ensure(RayLength > 0);
 
@@ -1544,58 +1533,58 @@ namespace Chaos
 		// Spheres and Capsules: always use the core shape and full "margin" because it represents the radius
 		// Sphere/Capsule versus OtherShape: no margin on other
 		// OtherShape versus OtherShape: use margin of the smaller shape, zero margin on the other
-		const T RadiusA = static_cast<T>(A.GetRadius());
-		const T RadiusB = static_cast<T>(B.GetRadius());
+		const FRealSingle RadiusA = static_cast<FRealSingle>(A.GetRadius());
+		const FRealSingle RadiusB = static_cast<FRealSingle>(B.GetRadius());
 		const bool bHasRadiusA = RadiusA > 0;
 		const bool bHasRadiusB = RadiusB > 0;
 
 		// The sweep margins if required. Only one can be non-zero (we keep the smaller one)
-		const T SweepMarginScale = 0.05f;
+		const FRealSingle SweepMarginScale = 0.05f;
 		const bool bAIsSmallest = A.GetMargin() < B.GetMargin();
-		const T SweepMarginA = (bHasRadiusA || bHasRadiusB) ? 0.0f : (bAIsSmallest ? SweepMarginScale * static_cast<T>(A.GetMargin()) : 0.0f);
-		const T SweepMarginB = (bHasRadiusA || bHasRadiusB) ? 0.0f : (bAIsSmallest ? 0.0f : SweepMarginScale * static_cast<T>(B.GetMargin()));
+		const FRealSingle SweepMarginA = (bHasRadiusA || bHasRadiusB) ? 0.0f : (bAIsSmallest ? SweepMarginScale * static_cast<FRealSingle>(A.GetMargin()) : 0.0f);
+		const FRealSingle SweepMarginB = (bHasRadiusA || bHasRadiusB) ? 0.0f : (bAIsSmallest ? 0.0f : SweepMarginScale * static_cast<FRealSingle>(B.GetMargin()));
 
 		// Net margin (note: both SweepMargins are zero if either Radius is non-zero, and only one SweepMargin can be non-zero)
 		A.Margin = RadiusA + SweepMarginA;
 		B.Margin = RadiusB + SweepMarginB;
 
-		const VectorRegister4Float MarginASimd = VectorLoadFloat1(&A.Margin);
-		const VectorRegister4Float MarginBSimd = VectorLoadFloat1(&B.Margin);
+		const T MarginASimd = T(VectorLoadFloat1(&A.Margin));
+		const T MarginBSimd = T(VectorLoadFloat1(&B.Margin));
 
-		VectorRegister4Float Simplex[4] = { VectorZeroFloat(), VectorZeroFloat(), VectorZeroFloat(), VectorZeroFloat() };
-		VectorRegister4Float As[4] = { VectorZeroFloat(), VectorZeroFloat(), VectorZeroFloat(), VectorZeroFloat() };
-		VectorRegister4Float Bs[4] = { VectorZeroFloat(), VectorZeroFloat(), VectorZeroFloat(), VectorZeroFloat() };
+		T Simplex[4] = { TVectorZero<T>(), TVectorZero<T>(), TVectorZero<T>(), TVectorZero<T>() };
+		T As[4] = { TVectorZero<T>(), TVectorZero<T>(), TVectorZero<T>(), TVectorZero<T>() };
+		T Bs[4] = { TVectorZero<T>(), TVectorZero<T>(), TVectorZero<T>(), TVectorZero<T>() };
 
-		VectorRegister4Float Barycentric = VectorZeroFloat();
+		T Barycentric = TVectorZero<T>();
 
-		VectorRegister4Float Inflation = VectorAdd(MarginASimd, MarginBSimd);
-		constexpr VectorRegister4Float Eps2Simd = MakeVectorRegisterFloatConstant(1e-6f, 1e-6f, 1e-6f, 1e-6f);
-		const VectorRegister4Float Inflation2Simd = VectorMultiplyAdd(Inflation, Inflation, Eps2Simd);
+		T Inflation = VectorAdd(MarginASimd, MarginBSimd);
+		const T Eps2Simd = TMakeVectorRegister<T>(1e-6f, 1e-6f, 1e-6f, 1e-6f);
+		const T Inflation2Simd = VectorMultiplyAdd(Inflation, Inflation, Eps2Simd);
 
-		const VectorRegister4Float RayLengthSimd = MakeVectorRegisterFloat(RayLength, RayLength, RayLength, RayLength);
+		const T RayLengthSimd = TMakeVectorRegister<T>(RayLength, RayLength, RayLength, RayLength);
 
-		VectorRegister4Int NumVerts = GlobalVectorConstants::IntZero;
+		int32 NumVerts = 0;
 
-		VectorRegister4Float AToBRotation = VectorQuaternionInverse(BToARotation);
+		T AToBRotation = VectorQuaternionInverse(BToARotation);
 
-		VectorRegister4Float SupportA = A.SupportFunction(InitialDir);
+		T SupportA = A.SupportFunction(TMakeVectorRegisterFloatFromDouble<T>(InitialDir));
 		As[0] = SupportA;
 
-		VectorRegister4Float SupportB = B.SupportFunction(AToBRotation, BToARotation, VectorNegate(InitialDir));
+		T SupportB = B.SupportFunction(TMakeVectorRegisterFloatFromDouble<T>(AToBRotation), TMakeVectorRegisterFloatFromDouble<T>(BToARotation), TMakeVectorRegisterFloatFromDouble<T>(VectorNegate(InitialDir)));
 		Bs[0] = SupportB;
 
-		VectorRegister4Float Lambda = VectorZeroFloat();
-		VectorRegister4Float X = StartPoint;
-		VectorRegister4Float V = VectorSubtract(X, VectorSubtract(SupportA, SupportB));
-		VectorRegister4Float Normal = MakeVectorRegisterFloat(0.f, 0.f, 1.f, 0.f);
+		T Lambda = TVectorZero<T>();
+		T X = StartPoint;
+		T V = VectorSubtract(X, VectorSubtract(SupportA, SupportB));
+		T Normal = MakeVectorRegisterFloat(0.f, 0.f, 1.f, 0.f);
 
-		const VectorRegister4Float InitialPreDist2Simd = VectorDot3(V, V);
+		const T InitialPreDist2Simd = VectorDot3(V, V);
 
 		FRealSingle InitialPreDist2;
 		VectorStoreFloat1(InitialPreDist2Simd, &InitialPreDist2);
 
 		FRealSingle Inflation2;
-		VectorStoreFloat1(Inflation2Simd, &Inflation2);
+		VectorStoreFloat1(Inflation2Simd, &Inflation2); 
 
 		constexpr FRealSingle Eps2 = 1e-6f;
 
@@ -1605,8 +1594,8 @@ namespace Chaos
 		bool bTerminate = bCloseEnough;
 		bool bInflatedCloseEnough = bCloseEnough;
 		int NumIterations = 0;
-		constexpr VectorRegister4Float LimitMax = MakeVectorRegisterFloatConstant(TNumericLimits<T>::Max(), TNumericLimits<T>::Max(), TNumericLimits<T>::Max(), TNumericLimits<T>::Max());
-		VectorRegister4Float GJKPreDist2 = LimitMax;
+		constexpr T LimitMax = TMakeVectorRegisterConstant<T>(TNumericLimits<FRealSingle>::Max(), TNumericLimits<FRealSingle>::Max(), TNumericLimits<FRealSingle>::Max(), TNumericLimits<FRealSingle>::Max());
+		T GJKPreDist2 = LimitMax;
 
 		while (!bTerminate)
 		{
@@ -1617,56 +1606,51 @@ namespace Chaos
 
 			V = VectorNormalizeAccurate(V);
 
-			SupportA = A.SupportFunction(V);
-			SupportB = B.SupportFunction(AToBRotation, BToARotation, VectorNegate(V));
-			const VectorRegister4Float P = VectorSubtract(SupportA, SupportB);
-			const VectorRegister4Float W = VectorSubtract(X, P);
+			SupportA = A.SupportFunction(TMakeVectorRegisterFloatFromDouble<T>(V));
+			SupportB = B.SupportFunction(TMakeVectorRegisterFloatFromDouble<T>(AToBRotation), TMakeVectorRegisterFloatFromDouble<T>(BToARotation), TMakeVectorRegisterFloatFromDouble<T>(VectorNegate(V)));
+			T P = VectorSubtract(SupportA, SupportB);
+			T W = VectorSubtract(X, P);
 
-			// NumVerts store here at the beginning of the loop, it should be safe to reuse it further in the loop 
-			alignas(16) int32 NumVertsInts[4];
-			VectorIntStoreAligned(NumVerts, NumVertsInts);
-			const int32 NumVertsInt = NumVertsInts[0];
 
-			As[NumVertsInt] = SupportA;
-			Bs[NumVertsInt] = SupportB;
+			As[NumVerts] = SupportA;
+			Bs[NumVerts] = SupportB;
 
-			const VectorRegister4Float VDotW = VectorDot3(V, W);
+			const T VDotW = VectorDot3(V, W);
 
-			VectorRegister4Float VDotWGTInflationSimd = VectorCompareGT(VDotW, Inflation);
+			T VDotWGTInflationSimd = VectorCompareGT(VDotW, Inflation);
 
 			if (VectorMaskBits(VDotWGTInflationSimd))
 			{
-				const VectorRegister4Float VDotRayDir = VectorDot3(V, RayDir);
-				VectorRegister4Float VDotRayDirGEZero = VectorCompareGE(VDotRayDir, VectorZeroFloat());
+				const T VDotRayDir = VectorDot3(V, RayDir);
+				T VDotRayDirGEZero = VectorCompareGE(VDotRayDir, TVectorZero<T>());
 
 				if (VectorMaskBits(VDotRayDirGEZero))
 				{
 					return false;
 				}
 
-				const VectorRegister4Float PreLambda = Lambda;	//use to check for no progress
+				const T PreLambda = Lambda;	//use to check for no progress
 				// @todo(ccaulfield): this can still overflow - the comparisons against zero above should be changed (though not sure to what yet)
 				Lambda = VectorSubtract(Lambda, VectorDivide(VectorSubtract(VDotW, Inflation), VDotRayDir));
-				VectorRegister4Float LambdaGTPreLambda = VectorCompareGT(Lambda, PreLambda);
+				T LambdaGTPreLambda = VectorCompareGT(Lambda, PreLambda);
 				if (VectorMaskBits(LambdaGTPreLambda))
 				{
-					VectorRegister4Float LambdaGTRayLength = VectorCompareGT(Lambda, RayLengthSimd);
+					T LambdaGTRayLength = VectorCompareGT(Lambda, RayLengthSimd);
 					if (VectorMaskBits(LambdaGTRayLength))
 					{
 						return false;
 					}
 
-					const VectorRegister4Float OldX = X;
+					const T OldX = X;
 					X = VectorMultiplyAdd(Lambda, RayDir, StartPoint);
 					Normal = V;
 
 					//Update simplex from (OldX - P) to (X - P)
-					VectorRegister4Float XMinusOldX = VectorSubtract(X, OldX);
+					T XMinusOldX = VectorSubtract(X, OldX);
 					Simplex[0] = VectorAdd(Simplex[0], XMinusOldX);
 					Simplex[1] = VectorAdd(Simplex[1], XMinusOldX);
 					Simplex[2] = VectorAdd(Simplex[2], XMinusOldX);
-					Simplex[NumVertsInt] = VectorSubtract(X, P);
-					NumVerts = VectorIntAdd(NumVerts, GlobalVectorConstants::IntOne);
+					Simplex[NumVerts++] = VectorSubtract(X, P);
 
 					GJKPreDist2 = LimitMax; //translated origin so restart gjk search
 					bInflatedCloseEnough = false;
@@ -1674,35 +1658,35 @@ namespace Chaos
 			}
 			else
 			{
-				Simplex[NumVertsInt] = W;	//this is really X - P which is what we need for simplex computation
-				NumVerts = VectorIntAdd(NumVerts, GlobalVectorConstants::IntOne);
+				Simplex[NumVerts++] = W;	//this is really X - P which is what we need for simplex computation
 			}
 
-			if (bInflatedCloseEnough && VectorMaskBits(VectorCompareGE(VDotW, VectorZeroFloat())))
+			if (bInflatedCloseEnough && VectorMaskBits(VectorCompareGE(VDotW, TVectorZero<T>())))
 			{
 				//Inflated shapes are close enough, but we want MTD so we need to find closest point on core shape
-				const VectorRegister4Float VDotW2 = VectorDot3(VDotW, VDotW);
+				const T VDotW2 = VectorDot3(VDotW, VDotW);
 				bCloseEnough = static_cast<bool>(VectorMaskBits(VectorCompareGE(VectorAdd(Eps2Simd, VDotW), GJKPreDist2)));
 			}
 
 			if (!bCloseEnough)
 			{
-				V = VectorSimplexFindClosestToOrigin(Simplex, NumVerts, Barycentric, As, Bs);
 
-				VectorRegister4Float NewDist2 = VectorDot3(V, V);	//todo: relative error
+				V = VectorSimplexFindClosestToOrigin<T>(Simplex, NumVerts, Barycentric, As, Bs);
+
+				T NewDist2 = VectorDot3(V, V);	//todo: relative error
 				bCloseEnough = static_cast<bool>(VectorMaskBits(VectorCompareGT(Inflation2Simd, NewDist2)));
 				bDegenerate = static_cast<bool>(VectorMaskBits(VectorCompareGE(NewDist2, GJKPreDist2)));
 				GJKPreDist2 = NewDist2;
 
 				if (bComputeMTD && bCloseEnough)
 				{
-					const VectorRegister4Float LambdaEqZero = VectorCompareEQ(Lambda, VectorZeroFloat());
-					const VectorRegister4Float InGJKPreDist2GTEps2 = VectorCompareGT(GJKPreDist2, Eps2Simd);
-					const VectorRegister4Float Inflation22GTEps2 = VectorCompareGT(Inflation2Simd, Eps2Simd);
-					constexpr VectorRegister4Int fourInt = MakeVectorRegisterIntConstant(4, 4, 4, 4);
-					const VectorRegister4Int Is4GTNumVerts = VectorIntCompareGT(fourInt, NumVerts);
+					const T LambdaEqZero = VectorCompareEQ(Lambda, TVectorZero<T>());
+					const T InGJKPreDist2GTEps2 = VectorCompareGT(GJKPreDist2, Eps2Simd);
+					const T Inflation22GTEps2 = VectorCompareGT(Inflation2Simd, Eps2Simd);
+
+					const bool Is4GTNumVerts = 4 > NumVerts;
 
-					const VectorRegister4Float IsInflatCloseEnough = VectorBitwiseAnd(LambdaEqZero, VectorBitwiseAnd(InGJKPreDist2GTEps2, VectorBitwiseAnd(Inflation22GTEps2, VectorCast4IntTo4Float(Is4GTNumVerts))));
+					const T IsInflatCloseEnough = VectorBitwiseAnd(LambdaEqZero, VectorBitwiseAnd(InGJKPreDist2GTEps2, Inflation22GTEps2));
 
 					// Leaving the original code, to explain the logic there
 					//if (bComputeMTD && bCloseEnough && Lambda == 0 && GJKPreDist2 > 1e-6 && Inflation2 > 1e-6 && NumVerts < 4)
@@ -1710,14 +1694,14 @@ namespace Chaos
 					//	bCloseEnough = false;
 					//	bInflatedCloseEnough = true;
 					//}
-					bInflatedCloseEnough = static_cast<bool>(VectorMaskBits(IsInflatCloseEnough));
+					bInflatedCloseEnough = static_cast<bool>(VectorMaskBits(IsInflatCloseEnough)) && Is4GTNumVerts;
 					bCloseEnough = !bInflatedCloseEnough;
 				}
 			}
 			else
 			{
 				//It must be that we want MTD and we can terminate. However, we must make one final call to fixup the simplex
-				V = VectorSimplexFindClosestToOrigin(Simplex, NumVerts, Barycentric, As, Bs);
+				V = VectorSimplexFindClosestToOrigin<T>(Simplex, NumVerts, Barycentric, As, Bs);
 			}
 			bTerminate = bCloseEnough || bDegenerate;
 		}
@@ -1726,34 +1710,26 @@ namespace Chaos
 		if (OutTime > 0)
 		{
 			OutNormal = Normal;
-			VectorRegister4Float ClosestB = VectorZeroFloat();
+			T ClosestB = TVectorZero<T>();
 
-			VectorRegister4Float Barycentrics[4];
+			T Barycentrics[4];
 			Barycentrics[0] = VectorSwizzle(Barycentric, 0, 0, 0, 0);
 			Barycentrics[1] = VectorSwizzle(Barycentric, 1, 1, 1, 1);
 			Barycentrics[2] = VectorSwizzle(Barycentric, 2, 2, 2, 2);
 			Barycentrics[3] = VectorSwizzle(Barycentric, 3, 3, 3, 3);
 
 
-			const VectorRegister4Float  ClosestB1 = VectorMultiplyAdd(Bs[0], Barycentrics[0], ClosestB);
-			const VectorRegister4Float  ClosestB2 = VectorMultiplyAdd(Bs[1], Barycentrics[1], ClosestB1);
-			const VectorRegister4Float  ClosestB3 = VectorMultiplyAdd(Bs[2], Barycentrics[2], ClosestB2);
-			const VectorRegister4Float  ClosestB4 = VectorMultiplyAdd(Bs[3], Barycentrics[3], ClosestB3);
+			const T  ClosestB1 = VectorMultiplyAdd(Bs[0], Barycentrics[0], ClosestB);
+			const T  ClosestB2 = VectorMultiplyAdd(Bs[1], Barycentrics[1], ClosestB1);
+			const T  ClosestB3 = VectorMultiplyAdd(Bs[2], Barycentrics[2], ClosestB2);
+			const T  ClosestB4 = VectorMultiplyAdd(Bs[3], Barycentrics[3], ClosestB3);
 
-			constexpr VectorRegister4Int TwoInt = MakeVectorRegisterIntConstant(2, 2, 2, 2);
-			constexpr VectorRegister4Int ThreeInt = MakeVectorRegisterIntConstant(3, 3, 3, 3);
+			ClosestB = NumVerts == 0 ? ClosestB : ClosestB4;
+			ClosestB = NumVerts == 1 ? ClosestB1 : ClosestB;
+			ClosestB = NumVerts == 2 ? ClosestB2 : ClosestB;
+			ClosestB = NumVerts == 3 ? ClosestB3 : ClosestB;
 
-			const VectorRegister4Float IsB0 = VectorCast4IntTo4Float(VectorIntCompareEQ(NumVerts, GlobalVectorConstants::IntZero));
-			const VectorRegister4Float IsB1 = VectorCast4IntTo4Float(VectorIntCompareEQ(NumVerts, GlobalVectorConstants::IntOne));
-			const VectorRegister4Float IsB2 = VectorCast4IntTo4Float(VectorIntCompareEQ(NumVerts, TwoInt));
-			const VectorRegister4Float IsB3 = VectorCast4IntTo4Float(VectorIntCompareEQ(NumVerts, ThreeInt));
-
-			ClosestB = VectorSelect(IsB0, ClosestB, ClosestB4);
-			ClosestB = VectorSelect(IsB1, ClosestB1, ClosestB);
-			ClosestB = VectorSelect(IsB2, ClosestB2, ClosestB);
-			ClosestB = VectorSelect(IsB3, ClosestB3, ClosestB);
-
-			const VectorRegister4Float ClosestLocal = VectorNegateMultiplyAdd(OutNormal, MarginBSimd, ClosestB);
+			const T ClosestLocal = VectorNegateMultiplyAdd(Normal, MarginBSimd, ClosestB);
 
 			OutPosition = VectorAdd(VectorMultiplyAdd(RayDir, Lambda, StartPoint), ClosestLocal);
 
@@ -1763,26 +1739,23 @@ namespace Chaos
 			// If Inflation == 0 we would expect GJKPreDist2 to be 0
 			// However, due to precision we can still end up with GJK failing.
 			// When that happens fall back on EPA
-			VectorRegister4Float InflationGTZero = VectorCompareGT(Inflation, VectorZeroFloat());
-			VectorRegister4Float InGJKPreDist2GTEps2 = VectorCompareGT(GJKPreDist2, Eps2Simd);
-			VectorRegister4Float LimitMaxGTInGJKPreDist2 = VectorCompareGT(LimitMax, GJKPreDist2);
-			VectorRegister4Float IsDone = VectorBitwiseAnd(InflationGTZero, VectorBitwiseAnd(InGJKPreDist2GTEps2, LimitMaxGTInGJKPreDist2));
+			T InflationGTZero = VectorCompareGT(Inflation, TVectorZero<T>());
+			T InGJKPreDist2GTEps2 = VectorCompareGT(GJKPreDist2, Eps2Simd);
+			T LimitMaxGTInGJKPreDist2 = VectorCompareGT(LimitMax, GJKPreDist2);
+			T IsDone = VectorBitwiseAnd(InflationGTZero, VectorBitwiseAnd(InGJKPreDist2GTEps2, LimitMaxGTInGJKPreDist2));
 
-			VectorRegister4Float GJKClosestA = VectorZeroFloat();
-			VectorRegister4Float GJKClosestB = VectorZeroFloat();
+			T GJKClosestA = TVectorZero<T>();
+			T GJKClosestB = TVectorZero<T>();
 
 			if (NumIterations)
 			{
-				VectorRegister4Float Barycentrics[4];
+				T Barycentrics[4];
 				Barycentrics[0] = VectorSwizzle(Barycentric, 0, 0, 0, 0);
 				Barycentrics[1] = VectorSwizzle(Barycentric, 1, 1, 1, 1);
 				Barycentrics[2] = VectorSwizzle(Barycentric, 2, 2, 2, 2);
 				Barycentrics[3] = VectorSwizzle(Barycentric, 3, 3, 3, 3);
 
-				alignas(16) int32 NumVertsInts[4];
-				VectorIntStoreAligned(NumVerts, NumVertsInts);
-				const int NumVertsInt = NumVertsInts[0];
-				for (int i = 0; i < NumVertsInt; ++i)
+				for (int i = 0; i < NumVerts; ++i)
 				{
 					GJKClosestA = VectorMultiplyAdd(As[i], Barycentrics[i], GJKClosestA);
 					GJKClosestB = VectorMultiplyAdd(Bs[i], Barycentrics[i], GJKClosestB);
@@ -1797,17 +1770,15 @@ namespace Chaos
 
 			if (VectorMaskBits(IsDone))
 			{
-				OutNormal = Normal;
-
-				const VectorRegister4Float ClosestBInA = VectorAdd(StartPoint, GJKClosestB);
-				const VectorRegister4Float InGJKPreDist = VectorSqrt(GJKPreDist2);
+				const T ClosestBInA = VectorAdd(StartPoint, GJKClosestB);
+				const T InGJKPreDist = VectorSqrt(GJKPreDist2);
 				OutNormal = VectorNormalizeAccurate(V);
 
-				VectorRegister4Float Penetration = VectorSubtract(VectorAdd(MarginASimd, MarginBSimd), InGJKPreDist);
+				T Penetration = VectorSubtract(VectorAdd(MarginASimd, MarginBSimd), InGJKPreDist);
 				Penetration = VectorMin(Penetration, LimitMax);
-				Penetration = VectorMax(Penetration, VectorZeroFloat());
+				Penetration = VectorMax(Penetration, TVectorZero<T>());
 
-				const VectorRegister4Float ClosestLocal = VectorNegateMultiplyAdd(OutNormal, MarginBSimd, GJKClosestB);
+				const T ClosestLocal = VectorNegateMultiplyAdd(OutNormal, MarginBSimd, GJKClosestB);
 
 				OutPosition = VectorAdd(VectorMultiplyAdd(OutNormal, Penetration, StartPoint), ClosestLocal);
 				Penetration = VectorNegate(Penetration);
@@ -1823,27 +1794,23 @@ namespace Chaos
 					VertsA.Reserve(8);
 					VertsB.Reserve(8);
 
-					alignas(16) int32 NumVertsInts[4];
-					VectorIntStoreAligned(NumVerts, NumVertsInts);
-					const int32 NumVertsInt = NumVertsInts[0];
-
-					for (int i = 0; i < NumVertsInt; ++i)
+					for (int i = 0; i < NumVerts; ++i)
 					{
-						VertsA.Add(As[i]);
-						VertsB.Add(VectorAdd(Bs[i], X));
+						VertsA.Add(TMakeVectorRegisterFloatFromDouble<T>(As[i]));
+						VertsB.Add(TMakeVectorRegisterFloatFromDouble<T>(VectorAdd(Bs[i], X)));
 					}
 
 					struct SupportBAtOriginHelper
 					{
 						const FGeomGJKHelperSIMD& B;
-						const VectorRegister4Float& AToBRotation;
-						const VectorRegister4Float& BToARotation;
-						const VectorRegister4Float& StartPoint;
+						const VectorRegister4Float AToBRotation;
+						const VectorRegister4Float BToARotation;
+						const VectorRegister4Float StartPoint;
 
-						VectorRegister4Float operator()(VectorRegister4Float V) const { return VectorAdd(B.SupportFunction(AToBRotation, BToARotation, V), StartPoint); }
+						VectorRegister4Float operator()(const VectorRegister4Float& V) const { return VectorAdd(B.SupportFunction(AToBRotation, BToARotation, V), StartPoint); }
 					};
 
-					SupportBAtOriginHelper SupportBAtOrigin = { B, AToBRotation, BToARotation, StartPoint };
+					SupportBAtOriginHelper SupportBAtOrigin = { B, TMakeVectorRegisterFloatFromDouble<T>(AToBRotation), TMakeVectorRegisterFloatFromDouble<T>(BToARotation), TMakeVectorRegisterFloatFromDouble<T>(StartPoint) };
 					VectorRegister4Float Penetration;
 					VectorRegister4Float MTD, ClosestA, ClosestBInA;
 					const EEPAResult EPAResult = VectorEPA(VertsA, VertsB, A, SupportBAtOrigin, Penetration, MTD, ClosestA, ClosestBInA);
@@ -1882,7 +1849,7 @@ namespace Chaos
 				{
 					//didn't even go into gjk loop, touching hit
 					OutTime = -(A.Margin + B.Margin);
-					OutNormal = MakeVectorRegisterFloat(0.0f, 0.0f, 1.0f, 0.0f);
+					OutNormal = TMakeVectorRegister<T>(0.0f, 0.0f, 1.0f, 0.0f);
 					OutPosition = VectorMultiplyAdd(OutNormal, MarginASimd, As[0]);
 				}
 			}
@@ -1890,21 +1857,20 @@ namespace Chaos
 		else
 		{
 			// Initial overlap without MTD. These properties are not valid, but assigning them anyway so they don't contain NaNs and cause issues in invoking code.
-			OutNormal = MakeVectorRegisterFloat(0.0f, 0.0f, 1.0f, 0.0f);
-			OutPosition = MakeVectorRegisterFloat(0.0f, 0.0f, 0.0f, 0.0f);
+			OutNormal = TMakeVectorRegister<T>(0.0f, 0.0f, 1.0f, 0.0f);
+			OutPosition = TMakeVectorRegister<T>(0.0f, 0.0f, 0.0f, 0.0f);
 		}
 
 		return true;
 	}
 
-	template <typename T = FReal>
-	bool GJKRaycast2Impl(const FGeomGJKHelperSIMD& A, const FGeomGJKHelperSIMD& B, const TRigidTransform<T, 3>& StartTM, const TVector<T, 3>& RayDir, const T RayLength, 
-		T& OutTime, TVector<T, 3>& OutPosition, TVector<T, 3>& OutNormal, const T GivenThicknessA, bool bComputeMTD, const TVector<T, 3>& InitialDir, const T GivenThicknessB)
+	inline bool GJKRaycast2Impl(const FGeomGJKHelperSIMD& A, const FGeomGJKHelperSIMD& B, const TRigidTransform<FReal, 3>& StartTM, const TVector<FReal, 3>& RayDir, const FReal RayLength, 
+		FReal& OutTime, TVector<FReal, 3>& OutPosition, TVector<FReal, 3>& OutNormal, const FReal GivenThicknessA, bool bComputeMTD, const TVector<FReal, 3>& InitialDir, const FReal GivenThicknessB)
 	{
-		const UE::Math::TQuat<T>& RotationDouble = StartTM.GetRotation();
+		const UE::Math::TQuat<FReal>& RotationDouble = StartTM.GetRotation();
 		VectorRegister4Float Rotation = MakeVectorRegisterFloatFromDouble(MakeVectorRegister(RotationDouble.X, RotationDouble.Y, RotationDouble.Z, RotationDouble.W));
 
-		const UE::Math::TVector<T>& TranslationDouble = StartTM.GetTranslation();
+		const UE::Math::TVector<FReal>& TranslationDouble = StartTM.GetTranslation();
 		const VectorRegister4Float Translation = MakeVectorRegisterFloatFromDouble(MakeVectorRegister(TranslationDouble.X, TranslationDouble.Y, TranslationDouble.Z, 0.0));
 
 		// Normalize rotation
@@ -1915,7 +1881,7 @@ namespace Chaos
 
 		FRealSingle OutTimeFloat = 0.0f;
 		VectorRegister4Float OutPositionSimd, OutNormalSimd;
-		const bool Result = GJKRaycast2ImplSimd(A, B, Rotation, Translation, RayDirSimd, static_cast<FRealSingle>(RayLength), OutTimeFloat, OutPositionSimd, OutNormalSimd, bComputeMTD, InitialDirSimd);
+		const bool Result = GJKRaycast2ImplSimd<VectorRegister4Float>(A, B, Rotation, Translation, RayDirSimd, static_cast<FRealSingle>(RayLength), OutTimeFloat, OutPositionSimd, OutNormalSimd, bComputeMTD, InitialDirSimd);
 
 		OutTime = static_cast<double>(OutTimeFloat);
 
diff --git a/Engine/Source/Runtime/Experimental/Chaos/Public/Chaos/SimplexVectorized.h b/Engine/Source/Runtime/Experimental/Chaos/Public/Chaos/SimplexVectorized.h
--- a/Engine/Source/Runtime/Experimental/Chaos/Public/Chaos/SimplexVectorized.h
+++ b/Engine/Source/Runtime/Experimental/Chaos/Public/Chaos/SimplexVectorized.h
@@ -7,27 +7,28 @@
 
 namespace Chaos
 {
-	template <bool CalculatExtraInformation>
-	FORCEINLINE VectorRegister4Float VectorLineSimplexFindOrigin(VectorRegister4Float* RESTRICT Simplex, VectorRegister4Int& RESTRICT NumVerts, VectorRegister4Float& RESTRICT OutBarycentric, VectorRegister4Float* RESTRICT A, VectorRegister4Float* RESTRICT B)
+
+	template <typename T, bool CalculatExtraInformation>
+	FORCEINLINE_DEBUGGABLE T VectorLineSimplexFindOrigin(T* RESTRICT Simplex, int32& RESTRICT NumVerts, T& RESTRICT OutBarycentric, T* RESTRICT A, T* RESTRICT B)
 	{
-		const VectorRegister4Float& X0 = Simplex[0];
-		const VectorRegister4Float& X1 = Simplex[1];
-		const VectorRegister4Float X0ToX1 = VectorSubtract(X1, X0);
+		const T& X0 = Simplex[0];
+		const T& X1 = Simplex[1];
+		const T X0ToX1 = VectorSubtract(X1, X0);
 
 		//Closest Point = (-X0 dot X1-X0) / ||(X1-X0)||^2 * (X1-X0)
 
-		const VectorRegister4Float X0ToOrigin = VectorNegate(X0);
-		const VectorRegister4Float Dot = VectorDot3(X0ToOrigin, X0ToX1);
+		const T X0ToOrigin = VectorNegate(X0);
+		const T Dot = VectorDot3(X0ToOrigin, X0ToX1);
 
-		const VectorRegister4Float IsX0 = VectorCompareGE(VectorZeroFloat(), Dot);
+		const T IsX0 = VectorCompareGE(TVectorZero<T>(), Dot);
 
-		constexpr VectorRegister4Float OutBarycentricIfX0OrX1 = GlobalVectorConstants::Float1000;
-		const VectorRegister4Float X0ToX1Squared = VectorDot3(X0ToX1, X0ToX1);
-		const VectorRegister4Float DotBigger = VectorCompareGE(Dot, X0ToX1Squared);
+		constexpr T OutBarycentricIfX0OrX1 = TMakeVectorRegisterConstant<T>(1, 0, 0, 0);
+		const T X0ToX1Squared = VectorDot3(X0ToX1, X0ToX1);
+		const T DotBigger = VectorCompareGE(Dot, X0ToX1Squared);
 
-		const VectorRegister4Float MinLimt = MakeVectorRegisterFloat(std::numeric_limits<FRealSingle>::min() , std::numeric_limits<FRealSingle>::min(), std::numeric_limits<FRealSingle>::min(), std::numeric_limits<FRealSingle>::min());
-		const VectorRegister4Float X0ToX1SquaredSmall = VectorCompareGE(MinLimt, X0ToX1Squared);
-		const VectorRegister4Float IsX1 = VectorBitwiseOr(DotBigger, X0ToX1SquaredSmall);
+		const T MinLimt = TMakeVectorRegisterConstant<T>(std::numeric_limits<FRealSingle>::min() , std::numeric_limits<FRealSingle>::min(), std::numeric_limits<FRealSingle>::min(), std::numeric_limits<FRealSingle>::min());
+		const T X0ToX1SquaredSmall = VectorCompareGE(MinLimt, X0ToX1Squared);
+		const T IsX1 = VectorBitwiseOr(DotBigger, X0ToX1SquaredSmall);
 
 		Simplex[0] = VectorSelect(IsX1, Simplex[1], Simplex[0]);
 
@@ -37,21 +38,20 @@ namespace Chaos
 			B[0] = VectorSelect(IsX1, B[1], B[0]);
 		}
 
-		VectorRegister4Float Ratio = VectorDivide(Dot, X0ToX1Squared);
+		T Ratio = VectorDivide(Dot, X0ToX1Squared);
 
-		Ratio = VectorMax(Ratio, VectorZeroFloat());
-		Ratio = VectorMin(Ratio, GlobalVectorConstants::FloatOne);
+		Ratio = VectorMax(Ratio, TVectorZero<T>());
+		Ratio = VectorMin(Ratio, TMakeVectorRegisterConstant<T>(1, 1, 1, 1));
 
-		VectorRegister4Float Closest = VectorMultiplyAdd(Ratio, X0ToX1, X0);
+		T Closest = VectorMultiplyAdd(Ratio, X0ToX1, X0);
 
-		const VectorRegister4Float OneMinusRatio = VectorSubtract(GlobalVectorConstants::FloatOne, Ratio);
-		const VectorRegister4Float OutBarycentricOtherwise = VectorUnpackLo(OneMinusRatio, Ratio);
+		const T OneMinusRatio = VectorSubtract(TMakeVectorRegisterConstant<T>(1, 1, 1, 1), Ratio);
+		const T OutBarycentricOtherwise = VectorUnpackLo(OneMinusRatio, Ratio);
 
 		Closest = VectorSelect(IsX0, X0, VectorSelect(IsX1, X1, Closest));
 
-		const VectorRegister4Float IsX0OrX1 = VectorBitwiseOr(IsX0, IsX1);
-		const VectorRegister4Int IsX0OrX1Int = VectorCast4FloatTo4Int(IsX0OrX1);
-		NumVerts = VectorIntSelect(IsX0OrX1Int, GlobalVectorConstants::IntOne, NumVerts);
+		const T IsX0OrX1 = VectorBitwiseOr(IsX0, IsX1);
+		NumVerts =  VectorMaskBits(IsX0OrX1) ? 1 : NumVerts;
 
 		if (CalculatExtraInformation)
 		{
@@ -63,68 +63,67 @@ namespace Chaos
 
 	// Based on an algorithm in Real Time Collision Detection - Ericson (very close to that)
 	// Using the same variable name conventions for easy reference
-	template <bool CalculatExtraInformation>
-	FORCEINLINE_DEBUGGABLE VectorRegister4Float  TriangleSimplexFindOriginFast(VectorRegister4Float* RESTRICT Simplex, VectorRegister4Int& RESTRICT NumVerts, VectorRegister4Float& RESTRICT OutBarycentric, VectorRegister4Float* RESTRICT As, VectorRegister4Float* RESTRICT Bs)
+	template <typename T, bool CalculatExtraInformation>
+	FORCEINLINE_DEBUGGABLE T TriangleSimplexFindOriginFast(T* RESTRICT Simplex, int32& RESTRICT NumVerts, T& RESTRICT OutBarycentric, T* RESTRICT As, T* RESTRICT Bs)
 	{
-		const VectorRegister4Float& A = Simplex[0];
-		const VectorRegister4Float& B = Simplex[1];
-		const VectorRegister4Float& C = Simplex[2];
+		const T& A = Simplex[0];
+		const T& B = Simplex[1];
+		const T& C = Simplex[2];
 
-		const VectorRegister4Float AB = VectorSubtract(B, A);
-		const VectorRegister4Float AC = VectorSubtract(C, A);
+		const T AB = VectorSubtract(B, A);
+		const T AC = VectorSubtract(C, A);
 
+		const T TriNormal = Private::VectorCrossNoFMA(AB, AC);
+		const T TriNormal2 = VectorDot3(TriNormal, TriNormal);
+		const T MinDouble = TMakeVectorRegisterConstant<T>(std::numeric_limits<float>::min(), std::numeric_limits<float>::min(), std::numeric_limits<float>::min(), std::numeric_limits<float>::min());
+		const T AMin = VectorMultiply(A, MinDouble);
+		const T Eps2 = VectorDot3(AMin, AMin);
+		const T Eps2GENormal2 = VectorCompareGE(Eps2, TriNormal2);
+	
 
-		// Handle degenerate triangle
-		const VectorRegister4Float TriNormal = Private::VectorCrossNoFMA(AB, AC);
-		const VectorRegister4Float TriNormal2 = VectorDot3(TriNormal, TriNormal);
-		const VectorRegister4Float MinFloat = MakeVectorRegisterFloatConstant(std::numeric_limits<float>::min(), std::numeric_limits<float>::min(), std::numeric_limits<float>::min(), std::numeric_limits<float>::min());
-		const VectorRegister4Float AMin = VectorMultiply(A, MinFloat);
-		const VectorRegister4Float Eps2 = VectorDot3(AMin, AMin);
-		const VectorRegister4Float Eps2GENormal2 = VectorCompareGE(Eps2, TriNormal2);
 		if (VectorMaskBits(Eps2GENormal2))
 		{
-			constexpr VectorRegister4Int Two = MakeVectorRegisterIntConstant(2, 2, 2, 2);
-			NumVerts = Two;
-			return VectorLineSimplexFindOrigin<CalculatExtraInformation>(Simplex, NumVerts, OutBarycentric, As, Bs);
+			NumVerts = 2;
+			return VectorLineSimplexFindOrigin<T, CalculatExtraInformation>(Simplex, NumVerts, OutBarycentric, As, Bs);
 		}
 
 		// Vertex region A
-		const VectorRegister4Float AO = VectorNegate(A);
+		const T AO = VectorNegate(A);
 
-		const VectorRegister4Float d1 = VectorDot3(AB, AO);
-		const VectorRegister4Float d2 = VectorDot3(AC, AO);
+		const T d1 = VectorDot3(AB, AO);
+		const T d2 = VectorDot3(AC, AO);
 
-		const VectorRegister4Float IsD1SEZero = VectorCompareGE(VectorZeroFloat(), d1);
-		const VectorRegister4Float IsD2SEZero = VectorCompareGE(VectorZeroFloat(), d2);
-		const VectorRegister4Float IsA = VectorBitwiseAnd(IsD1SEZero, IsD2SEZero);
+		const T IsD1SEZero = VectorCompareGE(TVectorZero<T>(), d1);
+		const T IsD2SEZero = VectorCompareGE(TVectorZero<T>(), d2);
+		const T IsA = VectorBitwiseAnd(IsD1SEZero, IsD2SEZero);
 
 		if (VectorMaskBits(IsA))
 		{
-			NumVerts = GlobalVectorConstants::IntOne;
+			NumVerts = 1;
 			if (CalculatExtraInformation)
 			{
-				OutBarycentric = GlobalVectorConstants::Float1000;
+				OutBarycentric = TMakeVectorRegisterConstant<T>(1, 0, 0, 0);
 			}
 			return A;
 		}
 
 		//Vertex region B
-		const VectorRegister4Float BO = VectorNegate(B);
-		const VectorRegister4Float d3 = VectorDot3(AB, BO);
-		const VectorRegister4Float d4 = VectorDot3(AC, BO);
+		const T BO = VectorNegate(B);
+		const T d3 = VectorDot3(AB, BO);
+		const T d4 = VectorDot3(AC, BO);
 
-		const VectorRegister4Float IsD3GEZero = VectorCompareGE(d3, VectorZeroFloat());
-		const VectorRegister4Float IsD3GED4 = VectorCompareGE(d3, d4);
-		const VectorRegister4Float IsB = VectorBitwiseAnd(IsD3GEZero, IsD3GED4);
+		const T IsD3GEZero = VectorCompareGE(d3, TVectorZero<T>());
+		const T IsD3GED4 = VectorCompareGE(d3, d4);
+		const T IsB = VectorBitwiseAnd(IsD3GEZero, IsD3GED4);
 
 		if (VectorMaskBits(IsB))
 		{
-			NumVerts = GlobalVectorConstants::IntOne;
+			NumVerts = 1;
 			if (CalculatExtraInformation)
 			{
-				OutBarycentric = GlobalVectorConstants::Float1000;
+				OutBarycentric = TMakeVectorRegisterConstant<T>(1, 0, 0, 0); 
 			}
-			Simplex[0] = B;
+			Simplex[0] = Simplex[1];
 			if (CalculatExtraInformation)
 			{
 				As[0] = As[1];
@@ -134,23 +133,22 @@ namespace Chaos
 		}
 
 		// Edge AB
-		const VectorRegister4Float d1d4 = VectorMultiply(d1, d4);
-		const VectorRegister4Float vc = VectorNegateMultiplyAdd(d3, d2, d1d4);
-		const VectorRegister4Float NormalizationDenominatorAB = VectorSubtract(d1, d3);
+		const T d1d4 = VectorMultiply(d1, d4);
+		const T vc = VectorNegateMultiplyAdd(d3, d2, d1d4);
+		const T NormalizationDenominatorAB = VectorSubtract(d1, d3);
 
-		const VectorRegister4Float IsZeroGEvc = VectorCompareGE(VectorZeroFloat(), vc);
-		const VectorRegister4Float IsD1GEZero = VectorCompareGE(d1, VectorZeroFloat());
-		const VectorRegister4Float IsZeroGED3= VectorCompareGE(VectorZeroFloat(), d3);
-		const VectorRegister4Float IsNDABGTZero = VectorCompareGT(NormalizationDenominatorAB, VectorZeroFloat());
-		const VectorRegister4Float IsAB = VectorBitwiseAnd(VectorBitwiseAnd(IsZeroGEvc, IsD1GEZero), VectorBitwiseAnd(IsZeroGED3, IsNDABGTZero));
+		const T IsZeroGEvc = VectorCompareGE(TVectorZero<T>(), vc);
+		const T IsD1GEZero = VectorCompareGE(d1, TVectorZero<T>());
+		const T IsZeroGED3= VectorCompareGE(TVectorZero<T>(), d3);
+		const T IsNDABGTZero = VectorCompareGT(NormalizationDenominatorAB, TVectorZero<T>());
+		const T IsAB = VectorBitwiseAnd(VectorBitwiseAnd(IsZeroGEvc, IsD1GEZero), VectorBitwiseAnd(IsZeroGED3, IsNDABGTZero));
 
 		if (VectorMaskBits(IsAB))
 		{
-			constexpr VectorRegister4Int two = MakeVectorRegisterIntConstant(2, 2, 2, 2);
-			NumVerts = two;
+			NumVerts = 2;
 
-			const VectorRegister4Float v = VectorDivide(d1, NormalizationDenominatorAB);
-			const VectorRegister4Float OneMinusV = VectorSubtract(GlobalVectorConstants::FloatOne, v);
+			const T v = VectorDivide(d1, NormalizationDenominatorAB);
+			const T OneMinusV = VectorSubtract(TMakeVectorRegisterConstant<T>(1, 1, 1, 1), v);
 			// b0	a1	a2	a3
 			if (CalculatExtraInformation)
 			{
@@ -160,22 +158,22 @@ namespace Chaos
 		}
 
 		// Vertex C
-		const VectorRegister4Float CO = VectorNegate(C);
-		const VectorRegister4Float d5 = VectorDot3(AB, CO);
-		const VectorRegister4Float d6 = VectorDot3(AC, CO);
-		const VectorRegister4Float IsD6GEZero = VectorCompareGE(d6, VectorZeroFloat());
-		const VectorRegister4Float IsD6GED5 = VectorCompareGE(d6, d5);
-		const VectorRegister4Float IsC = VectorBitwiseAnd(IsD6GEZero, IsD6GED5);
+		const T CO = VectorNegate(C);
+		const T d5 = VectorDot3(AB, CO);
+		const T d6 = VectorDot3(AC, CO);
+		const T IsD6GEZero = VectorCompareGE(d6, TVectorZero<T>());
+		const T IsD6GED5 = VectorCompareGE(d6, d5);
+		const T IsC = VectorBitwiseAnd(IsD6GEZero, IsD6GED5);
 
 		if (VectorMaskBits(IsC))
 		{
-			NumVerts = GlobalVectorConstants::IntOne;
+			NumVerts = 1;
 			if (CalculatExtraInformation)
 			{
-				OutBarycentric = GlobalVectorConstants::Float1000;
+				OutBarycentric = TMakeVectorRegisterConstant<T>(1, 0, 0, 0);
 			}
 
-			Simplex[0] = C;
+			Simplex[0] = Simplex[2];
 			if (CalculatExtraInformation)
 			{
 				As[0] = As[2];
@@ -185,28 +183,27 @@ namespace Chaos
 		}
 
 		// Edge AC
-		const VectorRegister4Float d5d2 = VectorMultiply(d5, d2);
-		const VectorRegister4Float vb = VectorNegateMultiplyAdd(d1, d6, d5d2);
-		const VectorRegister4Float NormalizationDenominatorAC = VectorSubtract(d2, d6);
+		const T d5d2 = VectorMultiply(d5, d2);
+		const T vb = VectorNegateMultiplyAdd(d1, d6, d5d2);
+		const T NormalizationDenominatorAC = VectorSubtract(d2, d6);
 
-		const VectorRegister4Float IsZeroGEvb = VectorCompareGE(VectorZeroFloat(), vb);
-		const VectorRegister4Float IsD2GEZero = VectorCompareGE(d2, VectorZeroFloat());
-		const VectorRegister4Float IsZeroGED6 = VectorCompareGE(VectorZeroFloat(), d6);
-		const VectorRegister4Float IsNDACGTZero = VectorCompareGT(NormalizationDenominatorAC, VectorZeroFloat());
-		const VectorRegister4Float IsAC = VectorBitwiseAnd(VectorBitwiseAnd(IsZeroGEvb, IsD2GEZero), VectorBitwiseAnd(IsZeroGED6, IsNDACGTZero));
+		const T IsZeroGEvb = VectorCompareGE(TVectorZero<T>(), vb);
+		const T IsD2GEZero = VectorCompareGE(d2, TVectorZero<T>());
+		const T IsZeroGED6 = VectorCompareGE(TVectorZero<T>(), d6);
+		const T IsNDACGTZero = VectorCompareGT(NormalizationDenominatorAC, TVectorZero<T>());
+		const T IsAC = VectorBitwiseAnd(VectorBitwiseAnd(IsZeroGEvb, IsD2GEZero), VectorBitwiseAnd(IsZeroGED6, IsNDACGTZero));
 
 		if (VectorMaskBits(IsAC))
 		{
-			const VectorRegister4Float w = VectorDivide(d2, NormalizationDenominatorAC);
-			constexpr VectorRegister4Int two = MakeVectorRegisterIntConstant(2, 2, 2, 2);
-			NumVerts = two;
-			const VectorRegister4Float OneMinusW = VectorSubtract(GlobalVectorConstants::FloatOne, w);
+			const T w = VectorDivide(d2, NormalizationDenominatorAC);
+			NumVerts = 2;
+			const T OneMinusW = VectorSubtract(TMakeVectorRegisterConstant<T>(1, 1, 1, 1), w);
 			// b0	a1	a2	a3
 			if (CalculatExtraInformation)
 			{
 				OutBarycentric = VectorUnpackLo(OneMinusW, w);
 			}
-			Simplex[1] = C;
+			Simplex[1] = Simplex[2];
 			if (CalculatExtraInformation)
 			{
 				As[1] = As[2];
@@ -216,33 +213,32 @@ namespace Chaos
 		}
 
 		// Edge BC
-		const VectorRegister4Float d3d6 = VectorMultiply(d3, d6);
-		const VectorRegister4Float va = VectorNegateMultiplyAdd(d5, d4, d3d6);
-		const VectorRegister4Float d4MinusD3 = VectorSubtract(d4, d3);
-		const VectorRegister4Float d5MinusD6 = VectorSubtract(d5, d6);
-		const VectorRegister4Float NormalizationDenominatorBC = VectorAdd(d4MinusD3, d5MinusD6);
-
-		const VectorRegister4Float IsZeroGEva = VectorCompareGE(VectorZeroFloat(), va);
-		const VectorRegister4Float IsD4MinusD3GEZero = VectorCompareGE(d4MinusD3, VectorZeroFloat());
-		const VectorRegister4Float IsD5MinusD6GEZero = VectorCompareGE(d5MinusD6, VectorZeroFloat());
-		const VectorRegister4Float IsNDBCGTZero = VectorCompareGT(NormalizationDenominatorBC, VectorZeroFloat());
-		const VectorRegister4Float IsBC = VectorBitwiseAnd(VectorBitwiseAnd(IsZeroGEva, IsD4MinusD3GEZero), VectorBitwiseAnd(IsD5MinusD6GEZero, IsNDBCGTZero));
+		const T d3d6 = VectorMultiply(d3, d6);
+		const T va = VectorNegateMultiplyAdd(d5, d4, d3d6);
+		const T d4MinusD3 = VectorSubtract(d4, d3);
+		const T d5MinusD6 = VectorSubtract(d5, d6);
+		const T NormalizationDenominatorBC = VectorAdd(d4MinusD3, d5MinusD6);
+
+		const T IsZeroGEva = VectorCompareGE(TVectorZero<T>(), va);
+		const T IsD4MinusD3GEZero = VectorCompareGE(d4MinusD3, TVectorZero<T>());
+		const T IsD5MinusD6GEZero = VectorCompareGE(d5MinusD6, TVectorZero<T>());
+		const T IsNDBCGTZero = VectorCompareGT(NormalizationDenominatorBC, TVectorZero<T>());
+		const T IsBC = VectorBitwiseAnd(VectorBitwiseAnd(IsZeroGEva, IsD4MinusD3GEZero), VectorBitwiseAnd(IsD5MinusD6GEZero, IsNDBCGTZero));
 
 		if (VectorMaskBits(IsBC))
 		{
-			constexpr VectorRegister4Int two = MakeVectorRegisterIntConstant(2, 2, 2, 2);
-			NumVerts = two;
-			const VectorRegister4Float w = VectorDivide(d4MinusD3, NormalizationDenominatorBC);
+			NumVerts = 2;
+			const T w = VectorDivide(d4MinusD3, NormalizationDenominatorBC);
 			if (CalculatExtraInformation)
 			{
-				const VectorRegister4Float OneMinusW = VectorSubtract(GlobalVectorConstants::FloatOne, w);
+				const T OneMinusW = VectorSubtract(TMakeVectorRegisterConstant<T>(1, 1, 1, 1), w);
 				// b0	a1	a2	a3
 				OutBarycentric = VectorUnpackLo(OneMinusW, w);
 			}
-			const VectorRegister4Float CMinusB = VectorSubtract(C, B);
-			const VectorRegister4Float Result = VectorMultiplyAdd(w, CMinusB, B);
-			Simplex[0] = B;
-			Simplex[1] = C;
+			const T CMinusB = VectorSubtract(C, B);
+			const T Result = VectorMultiplyAdd(w, CMinusB, B);
+			Simplex[0] = Simplex[1];
+			Simplex[1] = Simplex[2];
 			if (CalculatExtraInformation)
 			{
 				As[0] = As[1];
@@ -254,17 +250,16 @@ namespace Chaos
 		}
 
 		// Inside triangle
-		const VectorRegister4Float denom = VectorDivide(VectorOne(), VectorAdd(va, VectorAdd(vb, vc)));
-		VectorRegister4Float v = VectorMultiply(vb, denom);
-		VectorRegister4Float w = VectorMultiply(vc, denom);
-		constexpr VectorRegister4Int three = MakeVectorRegisterIntConstant(3, 3, 3, 3);
-		NumVerts = three;
+		const T denom = VectorDivide(VectorOne(), VectorAdd(va, VectorAdd(vb, vc)));
+		T v = VectorMultiply(vb, denom);
+		T w = VectorMultiply(vc, denom);
+		NumVerts = 3;
 
 		if (CalculatExtraInformation)
 		{
-			const VectorRegister4Float OneMinusVMinusW = VectorSubtract(VectorSubtract(GlobalVectorConstants::FloatOne, v), w);
+			const T OneMinusVMinusW = VectorSubtract(VectorSubtract(TMakeVectorRegisterConstant<T>(1, 1, 1, 1), v), w);
 			// b0	a1	a2	a3
-			const VectorRegister4Float OneMinusVMinusW_W = VectorUnpackLo(OneMinusVMinusW, w);
+			const T OneMinusVMinusW_W = VectorUnpackLo(OneMinusVMinusW, w);
 			// a0	b0	a1	b1
 			OutBarycentric = VectorUnpackLo(OneMinusVMinusW_W, v);
 		}
@@ -275,15 +270,16 @@ namespace Chaos
 		// (E.g., this caused jittering when walking on box with dimensions of 100000cm or more).
 		// This fix the unit test TestSmallCapsuleLargeBoxGJKRaycast_Vertical
 		// Previously was return VectorMultiplyAdd(AC, w, VectorMultiplyAdd(AB, v, A));
-		const VectorRegister4Float TriNormalOverSize2 = VectorDivide(TriNormal, TriNormal2);
-		const VectorRegister4Float SignedDistance = VectorDot3(A, TriNormalOverSize2);
+		const T TriNormalOverSize2 = VectorDivide(TriNormal, TriNormal2);
+		const T SignedDistance = VectorDot3(T(A), TriNormalOverSize2);
 		return VectorMultiply(TriNormal, SignedDistance);
 	}
 
-	FORCEINLINE bool VectorSignMatch(VectorRegister4Float A, VectorRegister4Float B)
+	template<typename T>
+	FORCEINLINE bool VectorSignMatch(T A, T B)
 	{
-		VectorRegister4Float OneIsZero = VectorMultiply(A, B);
-		OneIsZero = VectorCompareEQ(OneIsZero, VectorZeroFloat());
+		T OneIsZero = VectorMultiply(A, B);
+		OneIsZero = VectorCompareEQ(OneIsZero, TVectorZero<T>());
 
 		const bool IsZero = static_cast<bool>(VectorMaskBits(OneIsZero));
 		const int32 MaskA = VectorMaskBits(A);
@@ -291,13 +287,13 @@ namespace Chaos
 		return (MaskA == MaskB) && !IsZero;
 	}
 
-	template <bool CalculatExtraInformation>
-	FORCEINLINE_DEBUGGABLE VectorRegister4Float VectorTetrahedronSimplexFindOrigin(VectorRegister4Float* RESTRICT Simplex, VectorRegister4Int& RESTRICT NumVerts, VectorRegister4Float& RESTRICT OutBarycentric, VectorRegister4Float* RESTRICT A, VectorRegister4Float* RESTRICT B)
+	template <typename T, bool CalculatExtraInformation>
+	FORCEINLINE_DEBUGGABLE T VectorTetrahedronSimplexFindOrigin(T* RESTRICT Simplex, int32& RESTRICT NumVerts, T& RESTRICT OutBarycentric, T* RESTRICT A, T* RESTRICT B)
 	{
-		const VectorRegister4Float& X0 = Simplex[0];
-		const VectorRegister4Float& X1 = Simplex[1];
-		const VectorRegister4Float& X2 = Simplex[2];
-		const VectorRegister4Float& X3 = Simplex[3];
+		const T& X0 = Simplex[0];
+		const T& X1 = Simplex[1];
+		const T& X2 = Simplex[2];
+		const T& X3 = Simplex[3];
 
 		//Use signed volumes to determine if origin is inside or outside
 		/*
@@ -307,20 +303,19 @@ namespace Chaos
 				 1   1   1   1]
 		*/
 
-		VectorRegister4Float Cofactors[4];
+		T Cofactors[4];
 		Cofactors[0] = VectorNegate(VectorDot3(X1, Private::VectorCrossNoFMA(X2, X3)));
 		Cofactors[1] = VectorDot3(X0, Private::VectorCrossNoFMA(X2, X3));
 		Cofactors[2] = VectorNegate(VectorDot3(X0, Private::VectorCrossNoFMA(X1, X3)));
 		Cofactors[3] = VectorDot3(X0, Private::VectorCrossNoFMA(X1, X2));
-		VectorRegister4Float DetM = VectorAdd(VectorAdd(Cofactors[0], Cofactors[1]), VectorAdd(Cofactors[2], Cofactors[3]));
+		T DetM = VectorAdd(VectorAdd(Cofactors[0], Cofactors[1]), VectorAdd(Cofactors[2], Cofactors[3]));
 
 		bool bSignMatch[4];
 		// constexpr int32 SubIdxs[4][3] = { {1,2,3}, {0,2,3}, {0,1,3}, {0,1,2} };
-		constexpr VectorRegister4Int ThreeInt = MakeVectorRegisterIntConstant(3, 3, 3, 3);
-		VectorRegister4Int SubNumVerts[4] = { ThreeInt, ThreeInt, ThreeInt, ThreeInt };
-		VectorRegister4Float SubSimplices[4][3] = { {Simplex[1], Simplex[2], Simplex[3]}, {Simplex[0], Simplex[2], Simplex[3]}, {Simplex[0], Simplex[1], Simplex[3]}, {Simplex[0], Simplex[1], Simplex[2]} };
-		VectorRegister4Float SubAs[4][3];
-		VectorRegister4Float SubBs[4][3];
+		int32 SubNumVerts[4] = { 3, 3, 3, 3};
+		T SubSimplices[4][3] = { {Simplex[1], Simplex[2], Simplex[3]}, {Simplex[0], Simplex[2], Simplex[3]}, {Simplex[0], Simplex[1], Simplex[3]}, {Simplex[0], Simplex[1], Simplex[2]} };
+		T SubAs[4][3];
+		T SubBs[4][3];
 		if (CalculatExtraInformation)
 		{
 			//SubAs = { {A[1], A[2], A[3]}, {A[0], A[2], A[3]}, {A[0], A[1], A[3]}, {A[0], A[1], A[2]} };
@@ -335,40 +330,33 @@ namespace Chaos
 			SubBs[2][0] = B[0]; SubBs[2][1] = B[1]; SubBs[2][2] = B[3];
 			SubBs[3][0] = B[0]; SubBs[3][1] = B[1]; SubBs[3][2] = B[2];
 		}
-		VectorRegister4Float ClosestPointSub[4];
-		VectorRegister4Float SubBarycentric[4];
-		constexpr VectorRegister4Int IndexNone = MakeVectorRegisterIntConstant(INDEX_NONE, INDEX_NONE, INDEX_NONE, INDEX_NONE);
-		VectorRegister4Int ClosestTriangleIdx = IndexNone;
-		VectorRegister4Float MinTriangleDist2 = VectorZeroFloat();
+		T ClosestPointSub[4];
+		T SubBarycentric[4];
+		int32 ClosestTriangleIdx = INDEX_NONE;
+		T MinTriangleDist2 = TVectorZero<T>();
 
 
-		constexpr VectorRegister4Int IdxSimd[4] = { GlobalVectorConstants::IntZero,
-													GlobalVectorConstants::IntOne,
-													MakeVectorRegisterIntConstant(2, 2, 2, 2),
-													MakeVectorRegisterIntConstant(3, 3, 3, 3)
-		};
+		constexpr int32 IdxSimd[4] = { 0, 1, 2, 3 };
 
-		VectorRegister4Float Eps = VectorMultiply(GlobalVectorConstants::KindaSmallNumber, VectorDivide(DetM, VectorSet1(4.0f)));
+		T Eps = VectorMultiply(GlobalVectorConstants::KindaSmallNumber, VectorDivide(DetM, VectorSet1(4.0f)));
 
 		bool bInside = true;
 		for (int Idx = 0; Idx < 4; ++Idx)
 		{
-			bSignMatch[Idx] = VectorSignMatch(DetM, Cofactors[Idx]);
+			bSignMatch[Idx] = VectorSignMatch<T>(DetM, Cofactors[Idx]);
 			if (!bSignMatch[Idx])
 			{
 				bInside = false;
-				ClosestPointSub[Idx] = TriangleSimplexFindOriginFast<CalculatExtraInformation>(SubSimplices[Idx], SubNumVerts[Idx], SubBarycentric[Idx], SubAs[Idx], SubBs[Idx]);
+				ClosestPointSub[Idx] = TriangleSimplexFindOriginFast<T, CalculatExtraInformation>(SubSimplices[Idx], SubNumVerts[Idx], SubBarycentric[Idx], SubAs[Idx], SubBs[Idx]);
 
-				const VectorRegister4Float Dist2 = VectorDot3(ClosestPointSub[Idx], ClosestPointSub[Idx]);
+				const T Dist2 = VectorDot3(ClosestPointSub[Idx], ClosestPointSub[Idx]);
 
-				const VectorRegister4Float MinGTDist = VectorCompareGT(MinTriangleDist2, Dist2);
-				const VectorRegister4Int ClsoestEq = VectorIntCompareEQ(ClosestTriangleIdx, IndexNone);
+				const T MinGTDist = VectorCompareGT(MinTriangleDist2, Dist2);
 
-				const VectorRegister4Float FindClosest = VectorBitwiseOr(MinGTDist, VectorCast4IntTo4Float(ClsoestEq));
-				const VectorRegister4Int FindClosestInt = VectorCast4FloatTo4Int(FindClosest);
+				const bool bFindClosest = static_cast<bool>(VectorMaskBits(MinGTDist)) || (ClosestTriangleIdx == INDEX_NONE);
 
-				MinTriangleDist2 = VectorSelect(FindClosest, Dist2, MinTriangleDist2);
-				ClosestTriangleIdx = VectorIntSelect(FindClosestInt, IdxSimd[Idx], ClosestTriangleIdx);
+				MinTriangleDist2 = bFindClosest ? Dist2 :  MinTriangleDist2;
+				ClosestTriangleIdx = bFindClosest ? IdxSimd[Idx] :  ClosestTriangleIdx;
 			}
 		}
 
@@ -376,81 +364,70 @@ namespace Chaos
 		{
 			if (CalculatExtraInformation)
 			{
-				VectorRegister4Float OutBarycentricVectors[4];
-				const VectorRegister4Float  InvDetM = VectorDivide(GlobalVectorConstants::FloatOne, DetM);
+				T OutBarycentricVectors[4];
+				const T  InvDetM = VectorDivide(TMakeVectorRegisterConstant<T>(1, 1, 1, 1), DetM);
 				OutBarycentricVectors[0] = VectorMultiply(Cofactors[0], InvDetM);
 				OutBarycentricVectors[1] = VectorMultiply(Cofactors[1], InvDetM);
 				OutBarycentricVectors[2] = VectorMultiply(Cofactors[2], InvDetM);
 				OutBarycentricVectors[3] = VectorMultiply(Cofactors[3], InvDetM);
 				// a0	b0	a1	b1
-				const VectorRegister4Float OutBarycentric0101 = VectorUnpackLo(OutBarycentricVectors[0], OutBarycentricVectors[1]);
-				const VectorRegister4Float OutBarycentric2323 = VectorUnpackLo(OutBarycentricVectors[2], OutBarycentricVectors[3]);
+				const T OutBarycentric0101 = VectorUnpackLo(OutBarycentricVectors[0], OutBarycentricVectors[1]);
+				const T OutBarycentric2323 = VectorUnpackLo(OutBarycentricVectors[2], OutBarycentricVectors[3]);
 				// a0	a1	b0	b1
-				OutBarycentric = VectorMoveLh(OutBarycentric0101, OutBarycentric2323);
+				OutBarycentric = VectorCombineLow(OutBarycentric0101, OutBarycentric2323);
 			}
 
-			return VectorZeroFloat();
+			return TVectorZero<T>();
 		}
 
-		alignas(16) int32 ClosestTriangleIdxInts[4];
-		VectorIntStoreAligned(ClosestTriangleIdx, ClosestTriangleIdxInts);
-		int32 ClosestTriangleIdxInt = ClosestTriangleIdxInts[0];
-		NumVerts = SubNumVerts[ClosestTriangleIdxInt];
+		NumVerts = SubNumVerts[ClosestTriangleIdx];
 		if (CalculatExtraInformation)
 		{
-			OutBarycentric = SubBarycentric[ClosestTriangleIdxInt];
+			OutBarycentric = SubBarycentric[ClosestTriangleIdx];
 		}
 
 		for (int i = 0; i < 3; i++)
 		{
-			Simplex[i] = SubSimplices[ClosestTriangleIdxInt][i];
+			Simplex[i] = SubSimplices[ClosestTriangleIdx][i];
 			if (CalculatExtraInformation)
 			{
-				A[i] = SubAs[ClosestTriangleIdxInt][i];
-				B[i] = SubBs[ClosestTriangleIdxInt][i];
+				A[i] = SubAs[ClosestTriangleIdx][i];
+				B[i] = SubBs[ClosestTriangleIdx][i];
 			}
 		}
 
-		return ClosestPointSub[ClosestTriangleIdxInt];
+		return ClosestPointSub[ClosestTriangleIdx];
 	}
 
 
 	// CalculatExtraHitInformation : Should we calculate the BaryCentric coordinates, As and Bs?
-	template <bool CalculatExtraInformation = true>
-	FORCEINLINE_DEBUGGABLE VectorRegister4Float VectorSimplexFindClosestToOrigin(VectorRegister4Float* RESTRICT Simplex, VectorRegister4Int& RESTRICT NumVerts, VectorRegister4Float& RESTRICT OutBarycentric, VectorRegister4Float* RESTRICT A, VectorRegister4Float* RESTRICT B)
+	template <typename T, bool CalculatExtraInformation = true>
+	FORCEINLINE_DEBUGGABLE T VectorSimplexFindClosestToOrigin(T* RESTRICT Simplex, int32& RESTRICT NumVerts, T& RESTRICT OutBarycentric, T* RESTRICT A, T* RESTRICT B)
 	{
-		VectorRegister4Float ClosestPoint;
-		alignas(16) int32 NumVertsInt[4];
-		VectorIntStoreAligned(NumVerts, NumVertsInt);
-		switch (NumVertsInt[0])
+		switch (NumVerts)
 		{
 		case 1:
 			if (CalculatExtraInformation)
 			{
-				OutBarycentric = GlobalVectorConstants::Float1000;
+				OutBarycentric = TMakeVectorRegisterConstant<T>(1, 0, 0 , 0);
 			}
-			ClosestPoint = Simplex[0]; 
-			break;
+			return Simplex[0]; 
 		case 2:
 		{
-			ClosestPoint = VectorLineSimplexFindOrigin<CalculatExtraInformation>(Simplex, NumVerts, OutBarycentric, A, B);
-			break;
+			return VectorLineSimplexFindOrigin<T, CalculatExtraInformation>(Simplex, NumVerts, OutBarycentric, A, B);
 		}
 		case 3:
 		{
-			ClosestPoint = TriangleSimplexFindOriginFast<CalculatExtraInformation>(Simplex, NumVerts, OutBarycentric, A, B);
-			break;
+			return TriangleSimplexFindOriginFast<T, CalculatExtraInformation>(Simplex, NumVerts, OutBarycentric, A, B);
 		}
 		case 4:
 		{
-			ClosestPoint = VectorTetrahedronSimplexFindOrigin<CalculatExtraInformation>(Simplex, NumVerts, OutBarycentric, A, B);
+			return VectorTetrahedronSimplexFindOrigin<T, CalculatExtraInformation>(Simplex, NumVerts, OutBarycentric, A, B);
 			break;
 		}
 		default:
 			ensure(false);
-			ClosestPoint = VectorZeroFloat();
+			return TVectorZero<T>();
 		}
-
-		return ClosestPoint;
 	}
 }
diff --git a/Engine/Source/Runtime/Experimental/Chaos/Public/Chaos/TriangleRegister.h b/Engine/Source/Runtime/Experimental/Chaos/Public/Chaos/TriangleRegister.h
--- a/Engine/Source/Runtime/Experimental/Chaos/Public/Chaos/TriangleRegister.h
+++ b/Engine/Source/Runtime/Experimental/Chaos/Public/Chaos/TriangleRegister.h
@@ -61,6 +61,16 @@ namespace Chaos
 		}
 
 
+		FORCEINLINE bool IsTooBigForSinglePrecision() const
+		{
+			constexpr FRealSingle SqrSideThreshold = 100000.0f*100000.0f;
+			// If two sides are quite small, the third one cannot be huge
+			const VectorRegister4Float BigSide = VectorMax(VectorAbs(VectorSubtract(A, B)), VectorAbs(VectorSubtract(B, C)));
+			// Warning this could 
+			return VectorDot3Scalar(BigSide, BigSide) > SqrSideThreshold;
+		}
+
+
 	private:
 
 		VectorRegister4Float A;
diff --git a/Engine/Source/Runtime/Experimental/ChaosCore/Public/Chaos/VectorUtility.h b/Engine/Source/Runtime/Experimental/ChaosCore/Public/Chaos/VectorUtility.h
--- a/Engine/Source/Runtime/Experimental/ChaosCore/Public/Chaos/VectorUtility.h
+++ b/Engine/Source/Runtime/Experimental/ChaosCore/Public/Chaos/VectorUtility.h
@@ -4,13 +4,74 @@
 #include "Math/VectorRegister.h"
 #include "Chaos/Core.h"
 
+
+template<typename T>
+T TVectorZero();
+
+template<>
+inline VectorRegister4Float TVectorZero<VectorRegister4Float>()
+{
+	return VectorZeroFloat();
+}
+
+template<>
+inline VectorRegister4Double TVectorZero<VectorRegister4Double>()
+{
+	return VectorZeroDouble();
+}
+
+template<typename T>
+T TMakeVectorRegister(float X, float Y, float Z, float W);
+
+template<>
+inline VectorRegister4Float TMakeVectorRegister<VectorRegister4Float>(float X, float Y, float Z, float W)
+{
+	return MakeVectorRegisterFloat(X, Y, Z, W);
+}
+
+template<>
+inline VectorRegister4Double TMakeVectorRegister<VectorRegister4Double>(float X, float Y, float Z, float W)
+{
+	return MakeVectorRegisterDouble(X, Y, Z, W);
+}
+
+template<typename T>
+constexpr T TMakeVectorRegisterConstant(float X, float Y, float Z, float W);
+
+template<>
+constexpr VectorRegister4Float TMakeVectorRegisterConstant<VectorRegister4Float>(float X, float Y, float Z, float W)
+{
+	return MakeVectorRegisterFloatConstant(X, Y, Z, W);
+}
+
+template<>
+constexpr VectorRegister4Double TMakeVectorRegisterConstant<VectorRegister4Double>(float X, float Y, float Z, float W)
+{
+	return MakeVectorRegisterDoubleConstant(X, Y, Z, W);
+}
+
+template<typename T>
+VectorRegister4Float TMakeVectorRegisterFloatFromDouble(const T& V);
+
+template<>
+inline VectorRegister4Float TMakeVectorRegisterFloatFromDouble<VectorRegister4Double>(const VectorRegister4Double& V)
+{
+	return MakeVectorRegisterFloatFromDouble(V);
+}
+
+// Should generate no op
+template<>
+constexpr VectorRegister4Float TMakeVectorRegisterFloatFromDouble<VectorRegister4Float>(const VectorRegister4Float& V)
+{
+	return V;
+}
+
 /**
  * Cast VectorRegister4Int in VectorRegister4Float
  *
  * @param V	vector
  * @return		VectorRegister4Float( B.x, A.y, A.z, A.w)
  */
-
 FORCEINLINE VectorRegister4Float VectorCast4IntTo4Float(const VectorRegister4Int& V)
 {
 #if (!defined(_MSC_VER) || PLATFORM_ENABLE_VECTORINTRINSICS_NEON) && PLATFORM_ENABLE_VECTORINTRINSICS
@@ -59,6 +120,35 @@ FORCEINLINE VectorRegister4Float VectorUnpackLo(const VectorRegister4Float& A, c
 #endif
 }
 
+/**
+ * Selects and interleaves the lower two DP FP values from A and B.
+ *
+ * @param A	1st vector
+ * @param B	2nd vector
+ * @return		VectorRegister4Float( A.x, B.x, A.y, B.y)
+ */
+FORCEINLINE VectorRegister4Double VectorUnpackLo(const VectorRegister4Double& A, const VectorRegister4Double& B)
+{
+#if PLATFORM_ENABLE_VECTORINTRINSICS_NEON
+	VectorRegister4Double Result;
+	Result.XY = vzip1q_f64(A.XY, B.XY); 
+	Result.ZW = vzip2q_f64(A.XY, B.XY);
+	return Result;
+#elif PLATFORM_ENABLE_VECTORINTRINSICS
+	#if UE_PLATFORM_MATH_USE_AVX
+		return _mm256_permute2f128_pd(_mm256_unpackhi_pd(A, B), _mm256_unpacklo_pd(A, B), 0x02);
+	#else
+	VectorRegister4Double Result;
+	Result.XY = _mm_unpacklo_pd(A.XY, B.XY);
+	Result.ZW = _mm_unpackhi_pd(A.XY, B.XY);
+	return Result;
+	#endif
+#else
+	return MakeVectorRegisterFloat(A.V[0], B.V[0], A.V[1], B.V[1]);
+#endif
+}
+
+
 /**
  * Selects and interleaves the higher two SP FP values from A and B.
  *
@@ -132,8 +222,8 @@ namespace Chaos::Private
 
 	/**
 	 * Calculates the cross product of two vectors (XYZ components). W of the input should be 0, and will remain 0.
-	 * This function is not using FMA for stability reason, rounding with FMA could cause numerical instability. 
-	 * 
+	 * This function is not using FMA for stability reason, rounding with FMA could cause numerical instability.
+	 *
 	 * @param Vec1	1st vector
 	 * @param Vec2	2nd vector
 	 * @return		cross(Vec1.xyz, Vec2.xyz). W of the input should be 0, and will remain 0.
@@ -157,6 +247,32 @@ namespace Chaos::Private
 #endif
 	}
 
+	/**
+		 * Calculates the cross product of two vectors (XYZ components). W of the input should be 0, and will remain 0.
+		 * This function is not using FMA for stability reason, rounding with FMA could cause numerical instability.
+		 *
+		 * @param Vec1	1st vector
+		 * @param Vec2	2nd vector
+		 * @return		cross(Vec1.xyz, Vec2.xyz). W of the input should be 0, and will remain 0.
+		 */
+	FORCEINLINE VectorRegister4Double VectorCrossNoFMA(const VectorRegister4Double& Vec1, const VectorRegister4Double& Vec2)
+	{
+#if PLATFORM_ENABLE_VECTORINTRINSICS_NEON
+		return VectorCross(Vec1, Vec2);
+#elif PLATFORM_ENABLE_VECTORINTRINSICS
+		// YZX
+		VectorRegister4Double A = VectorSwizzle(Vec2, 1, 2, 0, 3);
+		VectorRegister4Double B = VectorSwizzle(Vec1, 1, 2, 0, 3);
+		// XY, YZ, ZX
+		A = VectorMultiply(A, Vec1);
+		// XY-YX, YZ-ZY, ZX-XZ
+		A = VectorSubtract(A, VectorMultiply(B, Vec2));
+		// YZ-ZY, ZX-XZ, XY-YX
+		return VectorSwizzle(A, 1, 2, 0, 3);
+#else
+		return VectorCross(Vec1, Vec2);
+#endif
+	}
 }
 
 /**
