Fix Unity build
Fall out from CL 34206018

diff --git a/Engine/Source/Runtime/Experimental/Chaos/Public/Chaos/TriangleRegister.h b/Engine/Source/Runtime/Experimental/Chaos/Public/Chaos/TriangleRegister.h
--- a/Engine/Source/Runtime/Experimental/Chaos/Public/Chaos/TriangleRegister.h
+++ b/Engine/Source/Runtime/Experimental/Chaos/Public/Chaos/TriangleRegister.h
@@ -6,6 +6,7 @@
 
 #include "Chaos/Core.h"
 #include "Math/VectorRegister.h"
+#include "Chaos/AABB.h"
 
 namespace Chaos
 {
