ChaosFlesh: Add support for addition, removal and stiffness adjustment for dynamic springs. Add Gauss Seidel dynamic weak constraints. Rewrite Gauss Seidel weak constraints using new data structure.

diff --git a/Engine/Plugins/Experimental/ChaosFlesh/Source/ChaosFleshEngine/Private/ChaosFlesh/ChaosDeformableSolverComponent.cpp b/Engine/Plugins/Experimental/ChaosFlesh/Source/ChaosFleshEngine/Private/ChaosFlesh/ChaosDeformableSolverComponent.cpp
--- a/Engine/Plugins/Experimental/ChaosFlesh/Source/ChaosFleshEngine/Private/ChaosFlesh/ChaosDeformableSolverComponent.cpp
+++ b/Engine/Plugins/Experimental/ChaosFlesh/Source/ChaosFleshEngine/Private/ChaosFlesh/ChaosDeformableSolverComponent.cpp
@@ -227,6 +227,7 @@ void UDeformableSolverComponent::Reset()
 			, SolverConstraints.GaussSeidelConstraints.CollisionSpring.bAllowSliding
 			, SolverMuscleActivation.bDoMuscleActivation
 			, SolverConstraints.GaussSeidelConstraints.CollisionSpring.bCollideWithFullmesh
+			, SolverConstraints.GaussSeidelConstraints.bEnableDynamicSprings
 		}));
 
 		for (TObjectPtr<UDeformablePhysicsComponent>& DeformableComponent : ConnectedObjects.DeformableComponents)
diff --git a/Engine/Plugins/Experimental/ChaosFlesh/Source/ChaosFleshEngine/Public/ChaosFlesh/ChaosDeformableSolverComponent.h b/Engine/Plugins/Experimental/ChaosFlesh/Source/ChaosFleshEngine/Public/ChaosFlesh/ChaosDeformableSolverComponent.h
--- a/Engine/Plugins/Experimental/ChaosFlesh/Source/ChaosFleshEngine/Public/ChaosFlesh/ChaosDeformableSolverComponent.h
+++ b/Engine/Plugins/Experimental/ChaosFlesh/Source/ChaosFleshEngine/Public/ChaosFlesh/ChaosDeformableSolverComponent.h
@@ -164,6 +164,12 @@ struct FSolverGaussSeidelConstraintsGroup
 	UPROPERTY(EditAnywhere, Category = "GaussSeidelConstraints")
 	float OmegaSOR = 1.6f;
 
+	/**
+	* Enable dynamic springs controlled by constraint manager. 
+	*/
+	UPROPERTY(EditAnywhere, Category = "GaussSeidelConstraints")
+	bool bEnableDynamicSprings = true;
+
 	/**
 	* Collsion detection radius and stiffness
 	*/
diff --git a/Engine/Source/Runtime/Experimental/Chaos/Private/Chaos/Deformable/ChaosDeformableSolver.cpp b/Engine/Source/Runtime/Experimental/Chaos/Private/Chaos/Deformable/ChaosDeformableSolver.cpp
--- a/Engine/Source/Runtime/Experimental/Chaos/Private/Chaos/Deformable/ChaosDeformableSolver.cpp
+++ b/Engine/Source/Runtime/Experimental/Chaos/Private/Chaos/Deformable/ChaosDeformableSolver.cpp
@@ -176,6 +176,7 @@ namespace Chaos::Softs
 			AllCorotatedCodEMeshArray.Reset(new TArray<FSolverReal>());
 			AllSkinEMeshArray.Reset(new TArray<FSolverReal>());
 			GSWeakConstraints.Reset(new FGaussSeidelWeakConstraints<FSolverReal, FSolverParticles>({}, {}, {}, {}, {}, GDeformableXPBDWeakConstraintParams));
+			GSDynamicWeakConstraints.Reset(new FGaussSeidelDynamicWeakConstraints<FSolverReal, FSolverParticles>(GDeformableXPBDWeakConstraintParams));
 			MuscleActivationConstraints.Reset(new FMuscleActivationConstraints<FSolverReal, FSolverParticles>());
 		}
 		AllUnconstrainedSurfaceElementsCorotatedCod.Reset(new TArray<Chaos::TVec3<int32>>());
@@ -792,7 +793,14 @@ namespace Chaos::Softs
 								if (Geom.Intersection(Samples, TetVertices, Intersections))
 								{
 									UE_LOG(LogChaosDeformableSolver, Log, TEXT("... Intersections : %d"), Intersections.Num());
-									//GSWeakConstraints->AddExtraConstraints(PositionTargetIndices, PositionTargetWeights, PositionTargetStiffness, PositionTargetSecondIndices, PositionTargetSecondWeights);
+									if (GSDynamicWeakConstraints && Property.bEnableDynamicSprings)
+									{
+										float Stiffness = AddConstraints.Parameters.Stiffness;
+										TArray<const FGaussSeidelWeakConstraints<Softs::FSolverReal, Softs::FSolverParticles>::FGaussSeidelConstraintHandle*> ConstraintHandles = GSDynamicWeakConstraints -> AddParticleTetrahedraConstraints(Geom, Evolution->Particles(), Intersections, SampleRange, TargetRange, Stiffness);
+										Chaos::Softs::FConstraintObjectParticleHandel& HandleValue = Proxy.Constraints.FindOrAdd(AddConstraints);
+										HandleValue.Handles = ConstraintHandles;
+										bDynamicConstraintIsUpdated = true;
+									}
 								}
 							}
 						}
@@ -1601,7 +1609,7 @@ namespace Chaos::Softs
 			if (Property.bEnablePositionTargets)
 			{
 				TArray<TArray<int32>> ParticlesPerColor;
-				GSWeakConstraints->ComputeInitialWCData(Evolution->Particles(), GSNeohookeanConstraints->GetMeshArray(), GSNeohookeanConstraints->GetIncidentElements(), GSNeohookeanConstraints->GetIncidentElementsLocal(), ParticlesPerColor);
+				GSWeakConstraints->ComputeInitialWCData(Evolution->Particles());
 
 				TArray<TArray<int32>> StaticIncidentElements, StaticIncidentElementsLocal;
 				TArray<TArray<int32>> StaticConstraints = GSWeakConstraints->GetStaticConstraintArrays(StaticIncidentElements, StaticIncidentElementsLocal);
@@ -1654,7 +1662,7 @@ namespace Chaos::Softs
 			if (Property.bEnablePositionTargets)
 			{
 				TArray<TArray<int32>> ParticlesPerColor;
-				GSWeakConstraints->ComputeInitialWCData(Evolution->Particles(), GSCorotatedConstraints->GetMeshArray(), GSCorotatedConstraints->GetIncidentElements(), GSCorotatedConstraints->GetIncidentElementsLocal(), ParticlesPerColor);
+				GSWeakConstraints->ComputeInitialWCData(Evolution->Particles());
 
 				TArray<TArray<int32>> StaticIncidentElements, StaticIncidentElementsLocal;
 				TArray<TArray<int32>> StaticConstraints = GSWeakConstraints->GetStaticConstraintArrays(StaticIncidentElements, StaticIncidentElementsLocal);
@@ -1687,10 +1695,46 @@ namespace Chaos::Softs
 				};
 		}
 
-		if (Property.bDoSelfCollision)
+		if (Property.bEnableDynamicSprings)
 		{
 			int32 DynamicIndex = GSMainConstraint->AddDynamicConstraintResidualAndHessianRange(1);
 			GSMainConstraint->DynamicConstraintResidualAndHessian()[DynamicIndex] = [this](const FSolverParticles& Particles, const int32 ConstraintIndex, const int32 ConstraintIndexLocal, const FSolverReal Dt, TVec3<FSolverReal>& ParticleResidual, Chaos::PMatrix<FSolverReal, 3, 3>& ParticleHessian)
+			{
+				this->GSDynamicWeakConstraints->AddWCResidualAndHessian(Particles, ConstraintIndex, ConstraintIndexLocal, Dt, ParticleResidual, ParticleHessian);
+			};
+
+			int32 PerNodeIndex = GSMainConstraint->AddPerNodeHessianRange(1);
+			GSMainConstraint->PerNodeHessian()[PerNodeIndex] = [this](const int32 p, const FSolverReal Dt, Chaos::PMatrix<FSolverReal, 3, 3>& ParticleHessian)
+			{
+				this->GSDynamicWeakConstraints->AddWCHessian(p, Dt, ParticleHessian);
+			};
+
+			int32 InitIndex1 = Evolution->AddConstraintInitRange(1, true);
+
+			GSDynamicWeakConstraints->ComputeInitialWCData(Evolution->Particles());
+
+			Evolution->ConstraintInits()[InitIndex1] =
+				[this](FSolverParticles& InParticles, const FSolverReal Dt)
+				{
+					this->GSDynamicWeakConstraints->Init(InParticles, Dt);
+					if (this->bDynamicConstraintIsUpdated)
+					{
+						TArray<TArray<int32>> WCDynamicIncidentElements,WCDynamicIncidentElementsLocal;
+						TArray<TArray<int32>> DynamicConstraints = GSDynamicWeakConstraints->GetStaticConstraintArrays(WCDynamicIncidentElements, WCDynamicIncidentElementsLocal);
+						GSMainConstraint->ResetDynamicConstraints();
+						GSMainConstraint->AddDynamicConstraints(DynamicConstraints, WCDynamicIncidentElements, WCDynamicIncidentElementsLocal, true);
+						this->GSMainConstraint->InitDynamicColor(InParticles);
+					}
+				};
+		}
+
+
+		
+
+		if (Property.bDoSelfCollision)
+		{
+			int32 TransientIndex = GSMainConstraint->AddTransientConstraintResidualAndHessianRange(1);
+			GSMainConstraint->TransientConstraintResidualAndHessian()[TransientIndex] = [this](const FSolverParticles& Particles, const int32 ConstraintIndex, const int32 ConstraintIndexLocal, const FSolverReal Dt, TVec3<FSolverReal>& ParticleResidual, Chaos::PMatrix<FSolverReal, 3, 3>& ParticleHessian)
 				{
 					this->GSWeakConstraints->AddWCResidualAndHessian(Particles, ConstraintIndex + this->GSWeakConstraints->InitialWCSize, ConstraintIndexLocal, Dt, ParticleResidual, ParticleHessian);
 				};
@@ -1728,11 +1772,13 @@ namespace Chaos::Softs
 				{
 					TArray<TArray<int32>> WCCollisionConstraints, WCCollisionIncidentElements, WCCollisionIncidentElementsLocal;
 					this->GSWeakConstraints->ComputeCollisionWCDataSimplified(WCCollisionConstraints, WCCollisionIncidentElements, WCCollisionIncidentElementsLocal);
-					this->GSMainConstraint->AddDynamicConstraints(WCCollisionConstraints, WCCollisionIncidentElements, WCCollisionIncidentElementsLocal);
-					this->GSMainConstraint->InitDynamicColor(InParticles);
+					this->GSMainConstraint->AddTransientConstraints(WCCollisionConstraints, WCCollisionIncidentElements, WCCollisionIncidentElementsLocal);
+					this->GSMainConstraint->InitTransientColor(InParticles);
 				};
 
 		}
+
+
 	}
 
 	void FDeformableSolver::InitializeMuscleActivationVariables()
@@ -1999,6 +2045,8 @@ namespace Chaos::Softs
 	{
 		PERF_SCOPE(STAT_ChaosDeformableSolver_Update);
 
+		bDynamicConstraintIsUpdated = false;
+
 		if (!Proxies.Num()) return;
 
 		UpdateSimulationObjects(DeltaTime);
diff --git a/Engine/Source/Runtime/Experimental/Chaos/Private/Chaos/Deformable/GaussSeidelMainConstraint.cpp b/Engine/Source/Runtime/Experimental/Chaos/Private/Chaos/Deformable/GaussSeidelMainConstraint.cpp
--- a/Engine/Source/Runtime/Experimental/Chaos/Private/Chaos/Deformable/GaussSeidelMainConstraint.cpp
+++ b/Engine/Source/Runtime/Experimental/Chaos/Private/Chaos/Deformable/GaussSeidelMainConstraint.cpp
@@ -43,6 +43,41 @@ namespace Chaos::Softs
 		StaticIncidentElementsOffsets.Add(StaticConstraints.Num());
 	}
 
+	template <typename T, typename ParticleType>
+	void FGaussSeidelMainConstraint<T, ParticleType>::AddTransientConstraints(const TArray<TArray<int32>>& ExtraConstraints, TArray<TArray<int32>>& ExtraIncidentElements, TArray<TArray<int32>>& ExtraIncidentElementsLocal, bool CheckIncidentElements)
+	{
+		if (CheckIncidentElements)
+		{
+			if (!IsClean(ExtraConstraints, ExtraIncidentElements, ExtraIncidentElementsLocal))
+			{
+				ExtraIncidentElements = Chaos::Utilities::ComputeIncidentElements(ExtraConstraints, &ExtraIncidentElementsLocal);
+			}
+		}
+
+		int32 Offset = TransientConstraints.Num();
+		TransientConstraints += ExtraConstraints;
+		for (int32 i = 0; i < ExtraIncidentElements.Num(); i++)
+		{
+			if (ExtraIncidentElements[i].Num() > 0)
+			{
+				TArray<int32> ExtraIncidentElementsWithOffset = ExtraIncidentElements[i];
+				for (int32 j = 0; j < ExtraIncidentElementsWithOffset.Num(); j++)
+				{
+					ExtraIncidentElementsWithOffset[j] += Offset;
+				}
+				TransientIncidentElements[i] += ExtraIncidentElementsWithOffset;
+				TransientIncidentElementsLocal[i] += ExtraIncidentElementsLocal[i];
+			}
+		}
+		if (TransientIncidentElementsOffsets.Num() > 0)
+		{
+			TransientIncidentElementsOffsets.RemoveAt(TransientIncidentElementsOffsets.Num() - 1);
+		}
+		TransientIncidentElementsOffsets.Add(Offset);
+		TransientIncidentElementsOffsets.Add(TransientConstraints.Num());
+	
+	}
+
 	template <typename T, typename ParticleType>
 	void FGaussSeidelMainConstraint<T, ParticleType>::AddDynamicConstraints(const TArray<TArray<int32>>& ExtraConstraints, TArray<TArray<int32>>& ExtraIncidentElements, TArray<TArray<int32>>& ExtraIncidentElementsLocal, bool CheckIncidentElements)
 	{
@@ -71,11 +106,10 @@ namespace Chaos::Softs
 		}
 		if (DynamicIncidentElementsOffsets.Num() > 0)
 		{
-			DynamicIncidentElementsOffsets.RemoveAt(StaticIncidentElementsOffsets.Num() - 1);
+			DynamicIncidentElementsOffsets.RemoveAt(DynamicIncidentElementsOffsets.Num() - 1);
 		}
 		DynamicIncidentElementsOffsets.Add(Offset);
 		DynamicIncidentElementsOffsets.Add(DynamicConstraints.Num());
-	
 	}
 
 	template <typename T, typename ParticleType>
@@ -116,14 +150,14 @@ namespace Chaos::Softs
 
 						ConstraintIndex = 0;
 
-						for (int32 i = 0; i < DynamicIncidentElements[p].Num(); i++)
+						for (int32 i = 0; i < TransientIncidentElements[p].Num(); i++)
 						{
-							while (DynamicIncidentElements[p][i] >= DynamicIncidentElementsOffsets[ConstraintIndex + 1] && ConstraintIndex < DynamicIncidentElementsOffsets.Num() - 1)
+							while (TransientIncidentElements[p][i] >= TransientIncidentElementsOffsets[ConstraintIndex + 1] && ConstraintIndex < TransientIncidentElementsOffsets.Num() - 1)
 							{
 								ConstraintIndex += 1;
 							}
 
-							this->AddDynamicConstraintResidualAndHessian[ConstraintIndex](Particles, DynamicIncidentElements[p][i] - DynamicIncidentElementsOffsets[ConstraintIndex], DynamicIncidentElementsLocal[p][i], Dt, NewtonResidual[p], ParticleHessian);
+							this->AddTransientConstraintResidualAndHessian[ConstraintIndex](Particles, TransientIncidentElements[p][i] - TransientIncidentElementsOffsets[ConstraintIndex], TransientIncidentElementsLocal[p][i], Dt, NewtonResidual[p], ParticleHessian);
 						}
 
 						if (AllParticleHessian)
@@ -226,4 +260,6 @@ template class Chaos::Softs::FGaussSeidelMainConstraint<Chaos::Softs::FSolverRea
 
 template CHAOS_API void Chaos::Softs::FGaussSeidelMainConstraint<Chaos::FRealDouble, Chaos::TDynamicParticles<Chaos::FRealDouble, 3>>::AddStaticConstraints(const TArray<TArray<int32>>& ExtraConstraints, TArray<TArray<int32>>& ExtraIncidentElements, TArray<TArray<int32>>& ExtraIncidentElementsLocal);
 
+template CHAOS_API void Chaos::Softs::FGaussSeidelMainConstraint<Chaos::FRealDouble, Chaos::TDynamicParticles<Chaos::FRealDouble, 3>>::AddTransientConstraints(const TArray<TArray<int32>>& ExtraConstraints, TArray<TArray<int32>>& ExtraIncidentElements, TArray<TArray<int32>>& ExtraIncidentElementsLocal, bool CheckIncidentElements);
+
 template CHAOS_API void Chaos::Softs::FGaussSeidelMainConstraint<Chaos::FRealDouble, Chaos::TDynamicParticles<Chaos::FRealDouble, 3>>::AddDynamicConstraints(const TArray<TArray<int32>>& ExtraConstraints, TArray<TArray<int32>>& ExtraIncidentElements, TArray<TArray<int32>>& ExtraIncidentElementsLocal, bool CheckIncidentElements);
diff --git a/Engine/Source/Runtime/Experimental/Chaos/Private/Chaos/GraphColoring.cpp b/Engine/Source/Runtime/Experimental/Chaos/Private/Chaos/GraphColoring.cpp
--- a/Engine/Source/Runtime/Experimental/Chaos/Private/Chaos/GraphColoring.cpp
+++ b/Engine/Source/Runtime/Experimental/Chaos/Private/Chaos/GraphColoring.cpp
@@ -463,6 +463,83 @@ static bool VerifyExtraNodalColoring(const Chaos::TDynamicParticles<T, 3>& InPar
 	return true;
 }
 
+template <typename T>
+static bool VerifyExtraNodalColoring(const Chaos::TDynamicParticles<T, 3>& InParticles, const TArray<TArray<int32>>& StaticGraph, const TArray<TArray<int32>>& DynamicGraph, const TArray<TArray<int32>>& ExtraGraph, const TArray<TArray<int32>>& StaticIncidentElements, const TArray<TArray<int32>>& DynamicIncidentElements, const TArray<TArray<int32>>& ExtraIncidentElements,const TArray<int32>& ParticleColors,const TArray<TArray<int32>>& ParticlesPerColor)
+{
+	TBitArray ParticleIsIncluded(false, InParticles.Size());
+	for (int32 i = 0; i < ParticlesPerColor.Num(); i++)
+	{
+		for (int32 j = 0; j < ParticlesPerColor[i].Num(); j++)
+		{
+			ParticleIsIncluded[ParticlesPerColor[i][j]] = true;
+		}
+	}
+
+	for (int32 ParticleIndex = 0; ParticleIndex < (int32)InParticles.Size(); ParticleIndex++)
+	{
+		if (InParticles.InvM(ParticleIndex) != (T)0. && (StaticIncidentElements[ParticleIndex].Num() > 0 || ExtraIncidentElements[ParticleIndex].Num() > 0 || DynamicIncidentElements[ParticleIndex].Num() > 0))
+		{
+			if (!ParticleIsIncluded[ParticleIndex])
+			{
+				return false;
+			}
+		}
+	}
+
+
+	for (int32 i = 0; i < ParticlesPerColor.Num(); i++)
+	{
+		TSet<int32> IncidentParticles;
+		IncidentParticles.Reserve(ParticlesPerColor[i].Num());
+		for (int32 j = 0; j < ParticlesPerColor[i].Num(); j++)
+		{
+			int32 ParticleIndex = ParticlesPerColor[i][j];
+			TSet<int32> LocalIncidentParticles;
+			for (int32 k = 0; k < StaticIncidentElements[ParticleIndex].Num(); k++)
+			{
+				int32 ElementIndex = StaticIncidentElements[ParticleIndex][k];
+				for (int32 ie = 0; ie < StaticGraph[ElementIndex].Num(); ie++)
+				{
+					LocalIncidentParticles.Add(StaticGraph[ElementIndex][ie]);
+				}
+			}
+			for (int32 k = 0; k < DynamicIncidentElements[ParticleIndex].Num(); k++)
+			{
+				int32 ElementIndex = DynamicIncidentElements[ParticleIndex][k];
+				for (int32 ie = 0; ie < DynamicGraph[ElementIndex].Num(); ie++)
+				{
+					LocalIncidentParticles.Add(DynamicGraph[ElementIndex][ie]);
+				}
+			}
+			if (ParticleIndex < ExtraIncidentElements.Num())
+			{
+				for (int32 k = 0; k < ExtraIncidentElements[ParticleIndex].Num(); k++)
+				{
+					int32 ElementIndex = ExtraIncidentElements[ParticleIndex][k];
+					for (int32 ie = 0; ie < ExtraGraph[ElementIndex].Num(); ie++)
+					{
+						LocalIncidentParticles.Add(ExtraGraph[ElementIndex][ie]);
+					}
+				}
+			}
+			if (IncidentParticles.Contains(ParticleIndex))
+			{
+				return false;
+			}
+			else
+			{
+				for (int32 LocalParticle : LocalIncidentParticles)
+				{
+					IncidentParticles.Add(LocalParticle);
+				}
+			}
+		}
+	}
+
+	return true;
+}
+
+
 template<typename DynamicParticlesType>
 TArray<TArray<int32>> Chaos::FGraphColoring::ComputeGraphColoringParticlesOrRange(const TArray<Chaos::TVec2<int32>>& Graph, const DynamicParticlesType& InParticles, const int32 GraphParticlesStart, const int32 GraphParticlesEnd)
 {
@@ -1836,6 +1913,102 @@ void Chaos::ComputeExtraNodalColoring(const TArray<TArray<int32>>& Graph, const
 	checkSlow(VerifyExtraNodalColoring<T>(InParticles, Graph, ExtraGraph, IncidentElements, ExtraIncidentElements, ParticleColors, ParticlesPerColor));
 }
 
+template<typename T>
+void Chaos::ComputeExtraNodalColoring(const TArray<TArray<int32>>& StaticGraph, const TArray<TArray<int32>>& DynamicGraph, const TArray<TArray<int32>>& ExtraGraph, const Chaos::TDynamicParticles<T, 3>& InParticles, const TArray<TArray<int32>>& StaticIncidentElements, const TArray<TArray<int32>>& DynamicIncidentElements, const TArray<TArray<int32>>& ExtraIncidentElements, TArray<int32>& ParticleColors, TArray<TArray<int32>>& ParticlesPerColor)
+{
+	TBitArray ParticleIsAffected(false, InParticles.Size());
+	for (int32 i = 0; i < ExtraIncidentElements.Num(); i++)
+	{
+		if (ExtraIncidentElements[i].Num() > 0)
+		{
+			ParticleIsAffected[i] = true;
+		}
+	}
+
+	TSet<int32> UsedColors;
+	for (int32 i = 0; i < ExtraIncidentElements.Num(); i++)
+	{
+		if (ParticleIsAffected[i])
+		{
+			int32 OriginalColor = ParticleColors[i];
+			ParticleColors[i] = INDEX_NONE;
+			if (OriginalColor != INDEX_NONE)
+			{
+				UsedColors.Reset();
+				UsedColors.Reserve(ExtraIncidentElements[i].Num() + StaticIncidentElements[i].Num() + DynamicIncidentElements[i].Num());
+				for (int32 j = 0; j < StaticIncidentElements[i].Num(); j++)
+				{
+					int32 ConstraintIndex = StaticIncidentElements[i][j];
+					for (int32 ie = 0; ie < StaticGraph[ConstraintIndex].Num(); ie++)
+					{
+						UsedColors.Add(ParticleColors[StaticGraph[ConstraintIndex][ie]]);
+					}
+				}
+				for (int32 j = 0; j < DynamicIncidentElements[i].Num(); j++)
+				{
+					int32 ConstraintIndex = DynamicIncidentElements[i][j];
+					for (int32 ie = 0; ie < DynamicGraph[ConstraintIndex].Num(); ie++)
+					{
+						UsedColors.Add(ParticleColors[DynamicGraph[ConstraintIndex][ie]]);
+					}
+				}
+				for (int32 j = 0; j < ExtraIncidentElements[i].Num(); j++)
+				{
+					int32 ConstraintIndex = ExtraIncidentElements[i][j];
+					for (int32 ie = 0; ie < ExtraGraph[ConstraintIndex].Num(); ie++)
+					{
+						UsedColors.Add(ParticleColors[ExtraGraph[ConstraintIndex][ie]]);
+					}
+				}
+				if (UsedColors.Contains(OriginalColor))
+				{
+					OriginalColor = 0;
+					while (UsedColors.Contains(OriginalColor))
+					{
+						OriginalColor++;
+					}
+					ParticleColors[i] = OriginalColor;
+				}
+				else
+				{
+					ParticleColors[i] = OriginalColor;
+					ParticleIsAffected[i] = false;
+				}
+			}
+		}
+	}
+
+
+	PhysicsParallelFor(ParticlesPerColor.Num(), [&](const int32 i)
+		{
+			int32 CurrentIndex = 0;
+			for (int32 j = 0; j < ParticlesPerColor[i].Num(); j++)
+			{
+				if (!ParticleIsAffected[ParticlesPerColor[i][j]])
+				{
+					if (CurrentIndex != j)
+					{
+						ParticlesPerColor[i][CurrentIndex] = ParticlesPerColor[i][j];
+					}
+					CurrentIndex += 1;
+				}
+			}
+			ParticlesPerColor[i].SetNum(CurrentIndex);
+		}, ParticlesPerColor.Num() > 0 && ParticlesPerColor[0].Num() < 1000);
+
+	ParticlesPerColor.SetNum(FMath::Max<int32>(ParticleColors) + 1);
+	for (int32 i = 0; i < ExtraIncidentElements.Num(); i++)
+	{
+		if (ParticleIsAffected[i] && ParticleColors[i] != INDEX_NONE)
+		{
+			ParticlesPerColor[ParticleColors[i]].Emplace(i);
+		}
+	}
+
+	checkSlow(VerifyExtraNodalColoring<T>(InParticles, StaticGraph, DynamicGraph, ExtraGraph, StaticIncidentElements, DynamicIncidentElements, ExtraIncidentElements, ParticleColors, ParticlesPerColor));
+}
+
+
 
 
 
@@ -1864,4 +2037,6 @@ template CHAOS_API TArray<TArray<int32>> Chaos::ComputeNodalColoring<Chaos::FRea
 template CHAOS_API TArray<TArray<int32>> Chaos::ComputeNodalColoring<Chaos::FRealSingle>(const TArray<TArray<int32>>& Graph, const Chaos::TDynamicParticles<Chaos::FRealSingle, 3>& InParticles, const int32 GraphParticlesStart, const int32 GraphParticlesEnd, const TArray<TArray<int32>>& IncidentElements, const TArray<TArray<int32>>& IncidentElementsLocalIndex, const TPBDActiveView<Chaos::Softs::FSolverParticles>* InParticleActiveView, TArray<int32>* ParticleColorsOut);
 template CHAOS_API TArray<TArray<int32>> Chaos::ComputeNodalColoring<Chaos::FRealDouble>(const TArray<TArray<int32>>& Graph, const Chaos::TDynamicParticles<Chaos::FRealDouble, 3>& InParticles, const int32 GraphParticlesStart, const int32 GraphParticlesEnd, const TArray<TArray<int32>>& IncidentElements, const TArray<TArray<int32>>& IncidentElementsLocalIndex, const TPBDActiveView<Chaos::TDynamicParticles<Chaos::FRealDouble, 3>>* InParticleActiveView, TArray<int32>* ParticleColorsOut);
 template CHAOS_API void Chaos::ComputeExtraNodalColoring<Chaos::FRealSingle>(const TArray<TArray<int32>>& Graph, const TArray<TArray<int32>>& ExtraGraph, const Chaos::TDynamicParticles<Chaos::FRealSingle, 3>& InParticles, const TArray<TArray<int32>>& IncidentElements, const TArray<TArray<int32>>& ExtraIncidentElements, TArray<int32>& ParticleColors, TArray<TArray<int32>>& ParticlesPerColor);
-template CHAOS_API void Chaos::ComputeExtraNodalColoring<Chaos::FRealDouble>(const TArray<TArray<int32>>& Graph, const TArray<TArray<int32>>& ExtraGraph, const Chaos::TDynamicParticles<Chaos::FRealDouble, 3>& InParticles, const TArray<TArray<int32>>& IncidentElements, const TArray<TArray<int32>>& ExtraIncidentElements, TArray<int32>& ParticleColors, TArray<TArray<int32>>& ParticlesPerColor);
\ No newline at end of file
+template CHAOS_API void Chaos::ComputeExtraNodalColoring<Chaos::FRealDouble>(const TArray<TArray<int32>>& Graph, const TArray<TArray<int32>>& ExtraGraph, const Chaos::TDynamicParticles<Chaos::FRealDouble, 3>& InParticles, const TArray<TArray<int32>>& IncidentElements, const TArray<TArray<int32>>& ExtraIncidentElements, TArray<int32>& ParticleColors, TArray<TArray<int32>>& ParticlesPerColor);
+template CHAOS_API void Chaos::ComputeExtraNodalColoring<Chaos::FRealSingle>(const TArray<TArray<int32>>& StaticGraph, const TArray<TArray<int32>>& DynamicGraph, const TArray<TArray<int32>>& ExtraGraph, const Chaos::TDynamicParticles<Chaos::FRealSingle, 3>& InParticles, const TArray<TArray<int32>>& StaticIncidentElements, const TArray<TArray<int32>>& DynamicIncidentElements, const TArray<TArray<int32>>& ExtraIncidentElements, TArray<int32>& ParticleColors, TArray<TArray<int32>>& ParticlesPerColor);
+template CHAOS_API void Chaos::ComputeExtraNodalColoring<Chaos::FRealDouble>(const TArray<TArray<int32>>& StaticGraph, const TArray<TArray<int32>>& DynamicGraph, const TArray<TArray<int32>>& ExtraGraph, const Chaos::TDynamicParticles<Chaos::FRealDouble, 3>& InParticles, const TArray<TArray<int32>>& StaticIncidentElements, const TArray<TArray<int32>>& DynamicIncidentElements, const TArray<TArray<int32>>& ExtraIncidentElements, TArray<int32>& ParticleColors, TArray<TArray<int32>>& ParticlesPerColor);
\ No newline at end of file
diff --git a/Engine/Source/Runtime/Experimental/Chaos/Private/GeometryCollection/Facades/CollectionTetrahedralFacade.cpp b/Engine/Source/Runtime/Experimental/Chaos/Private/GeometryCollection/Facades/CollectionTetrahedralFacade.cpp
--- a/Engine/Source/Runtime/Experimental/Chaos/Private/GeometryCollection/Facades/CollectionTetrahedralFacade.cpp
+++ b/Engine/Source/Runtime/Experimental/Chaos/Private/GeometryCollection/Facades/CollectionTetrahedralFacade.cpp
@@ -76,7 +76,9 @@ namespace GeometryCollection::Facades
 			for (int32 Tdx = 0; Tdx < TetCount; Tdx++)
 			{
 				const int32 Idx = TetStart + Tdx;
-				FIntVector4 Tet = Tetrahedron[Idx] - FIntVector4(VrtStart);
+				//Test shows offset VrtStart is not needed:
+				//FIntVector4 Tet = Tetrahedron[Idx] - FIntVector4(VrtStart);
+				FIntVector4 Tet = Tetrahedron[Idx];
 				Tets[Tdx] = FTet(TetarhedronPositions[Tet[0]], TetarhedronPositions[Tet[1]], TetarhedronPositions[Tet[2]], TetarhedronPositions[Tet[3]]);
 				TetPtrs[Tdx] = &Tets[Tdx];
 			}
diff --git a/Engine/Source/Runtime/Experimental/Chaos/Public/Chaos/Deformable/ChaosDeformableConstraintsProxy.h b/Engine/Source/Runtime/Experimental/Chaos/Public/Chaos/Deformable/ChaosDeformableConstraintsProxy.h
--- a/Engine/Source/Runtime/Experimental/Chaos/Public/Chaos/Deformable/ChaosDeformableConstraintsProxy.h
+++ b/Engine/Source/Runtime/Experimental/Chaos/Public/Chaos/Deformable/ChaosDeformableConstraintsProxy.h
@@ -3,6 +3,7 @@
 #pragma once
 #include "Chaos/Deformable/ChaosDeformableSolverProxy.h"
 #include "CoreMinimal.h"
+#include "Chaos/Deformable/GaussSeidelWeakConstraints.h"
 #include "UObject/ObjectMacros.h"
 
 namespace Chaos::Softs
@@ -65,6 +66,8 @@ namespace Chaos::Softs
 
 		int32 SourceParticleIndex = INDEX_NONE;
 		int32 TargetParticleIndex = INDEX_NONE;
+
+		TArray<const Chaos::Softs::FGaussSeidelWeakConstraints<Softs::FSolverReal, Softs::FSolverParticles>::FGaussSeidelConstraintHandle*> Handles;
 	};
 
 
@@ -102,6 +105,8 @@ namespace Chaos::Softs
 			TArray<FConstraintObjectUpdated> Updated;
 		};
 
+
+		
 		TArray<FConstraintObjectAdded> ConstraintObjectsToAdd;
 		TArray< FConstraintObjectRemoved> ConstraintObjectsToRemove;
 		TMap< FConstraintObjectKey, FConstraintObjectParticleHandel > Constraints;
diff --git a/Engine/Source/Runtime/Experimental/Chaos/Public/Chaos/Deformable/ChaosDeformableSolver.h b/Engine/Source/Runtime/Experimental/Chaos/Public/Chaos/Deformable/ChaosDeformableSolver.h
--- a/Engine/Source/Runtime/Experimental/Chaos/Public/Chaos/Deformable/ChaosDeformableSolver.h
+++ b/Engine/Source/Runtime/Experimental/Chaos/Public/Chaos/Deformable/ChaosDeformableSolver.h
@@ -14,6 +14,7 @@
 #include "Chaos/Deformable/GaussSeidelCorotatedConstraints.h"
 #include "Chaos/Deformable/GaussSeidelNeohookeanConstraints.h"
 #include "Chaos/Deformable/GaussSeidelWeakConstraints.h"
+#include "Chaos/Deformable/GaussSeidelDynamicWeakConstraints.h"
 #include "Chaos/Deformable/GaussSeidelMainConstraint.h"
 #include "Chaos/Deformable/MuscleActivationConstraints.h"
 #include "Chaos/Deformable/GaussSeidelCorotatedCodimensionalConstraints.h"
@@ -191,6 +192,7 @@ namespace Chaos::Softs
 		TUniquePtr<Softs::FGaussSeidelCorotatedCodimensionalConstraints<Softs::FSolverReal, Softs::FSolverParticles>> GSCorotatedCodConstraints;
 		TUniquePtr<Softs::FGaussSeidelLinearCodimensionalConstraints<Softs::FSolverReal, Softs::FSolverParticles>> GSLinearCodConstraints;
 		TUniquePtr<Softs::FGaussSeidelWeakConstraints<Softs::FSolverReal, Softs::FSolverParticles>> GSWeakConstraints;
+		TUniquePtr<Softs::FGaussSeidelDynamicWeakConstraints<Softs::FSolverReal, Softs::FSolverParticles>> GSDynamicWeakConstraints;
 		TArray<TUniquePtr<Softs::FXPBDWeakConstraints<Softs::FSolverReal, Softs::FSolverParticles>>> WeakConstraints;
 		TArray<TUniquePtr<Softs::FBlendedXPBDCorotatedConstraints<Softs::FSolverReal, Softs::FSolverParticles>>> BlendedCorotatedConstraints;
 		TUniquePtr<Softs::FXPBDGridBasedCorotatedConstraints<Softs::FSolverReal, Softs::FSolverParticles>> GridBasedCorotatedConstraint;
@@ -232,6 +234,7 @@ namespace Chaos::Softs
 		bool bSimulationInitialized = false;
 		int32 GroupOffset = 1;
 		TArray<TVector<int32, 2>, TInlineAllocator<8>> PrevEvolutionActiveRange;
+		bool bDynamicConstraintIsUpdated = false;
 	};
 
 
diff --git a/Engine/Source/Runtime/Experimental/Chaos/Public/Chaos/Deformable/ChaosDeformableSolverTypes.h b/Engine/Source/Runtime/Experimental/Chaos/Public/Chaos/Deformable/ChaosDeformableSolverTypes.h
--- a/Engine/Source/Runtime/Experimental/Chaos/Public/Chaos/Deformable/ChaosDeformableSolverTypes.h
+++ b/Engine/Source/Runtime/Experimental/Chaos/Public/Chaos/Deformable/ChaosDeformableSolverTypes.h
@@ -44,7 +44,8 @@ namespace Chaos::Softs
 			FSolverReal InCollisionSpringStiffness = (FSolverReal)500.,
 			bool InbAllowSliding = true,
 			bool InbDoMuscleActivation = false,
-			bool InbCollideWithFullMesh = false)
+			bool InbCollideWithFullMesh = false, 
+			bool InbEnableDynamicSprings = true)
 			: NumSolverSubSteps(InNumSolverSubSteps)
 			, NumSolverIterations(InNumSolverIterations)
 			, FixTimeStep(InFixTimeStep)
@@ -74,6 +75,7 @@ namespace Chaos::Softs
 			, bAllowSliding(InbAllowSliding)
 			, bDoMuscleActivation(InbDoMuscleActivation)
 			, bCollideWithFullMesh(InbCollideWithFullMesh)
+			, bEnableDynamicSprings(InbEnableDynamicSprings)
 		{}
 
 		int32 NumSolverSubSteps = 5;
@@ -105,6 +107,7 @@ namespace Chaos::Softs
 		bool bAllowSliding = true;
 		bool bDoMuscleActivation = false;
 		bool bCollideWithFullMesh = false;
+		bool bEnableDynamicSprings = true; 
 	};
 
 
diff --git a/Engine/Source/Runtime/Experimental/Chaos/Public/Chaos/Deformable/GaussSeidelDynamicWeakConstraints.h b/Engine/Source/Runtime/Experimental/Chaos/Public/Chaos/Deformable/GaussSeidelDynamicWeakConstraints.h
new file mode 100644
--- /dev/null
+++ b/Engine/Source/Runtime/Experimental/Chaos/Public/Chaos/Deformable/GaussSeidelDynamicWeakConstraints.h
@@ -0,0 +1,215 @@
+// Copyright Epic Games, Inc. All Rights Reserved.
+#pragma once
+
+#include "Chaos/BoundingVolumeHierarchy.h"
+#include "Chaos/PBDSpringConstraintsBase.h"
+#include "Chaos/XPBDCorotatedConstraints.h"
+#include "ChaosStats.h"
+#include "Chaos/ImplicitQRSVD.h"
+#include "Chaos/GraphColoring.h"
+#include "Chaos/NewtonCorotatedCache.h"
+#include "Chaos/Framework/Parallel.h"
+#include "Chaos/Utilities.h"
+#include "Chaos/DebugDrawQueue.h"
+#include "Chaos/XPBDWeakConstraints.h"
+#include "Chaos/Triangle.h"
+#include "Chaos/TriangleCollisionPoint.h"
+#include "Chaos/TriangleMesh.h"
+#include "Chaos/Deformable/GaussSeidelWeakConstraints.h"
+#include "GeometryCollection/Facades/CollectionTetrahedralFacade.h"
+namespace Chaos::Softs
+{
+	using Chaos::TVec3;
+	template <typename T, typename ParticleType>
+	struct FGaussSeidelDynamicWeakConstraints: public FGaussSeidelWeakConstraints <T, ParticleType>
+	{
+
+		typedef FGaussSeidelWeakConstraints<T, ParticleType> Base;
+		using Base::ConstraintsData;
+		using Base::NodalWeights;
+		using Base::WCIncidentElements;
+		using Base::WCIncidentElementsLocal;
+		typedef FGaussSeidelWeakConstraints<T, ParticleType>::FGaussSeidelConstraintHandle GSConstraintHandle;
+
+		FGaussSeidelDynamicWeakConstraints(const FDeformableXPBDWeakConstraintParams& InParams)
+		: Base({}, {}, {}, {}, {}, InParams)
+		{
+			ConstraintsData.AddArray(&Handles);
+		}
+
+		virtual ~FGaussSeidelDynamicWeakConstraints()
+		{
+			for (GSConstraintHandle* Handle:Handles)
+			{
+				delete Handle;
+			}
+		}
+		
+		const GSConstraintHandle* AddSingleParticleTetrahedraConstraint(const GeometryCollection::Facades::FTetrahedralFacade& InTargetGeom, const ParticleType& AllParticles, const GeometryCollection::Facades::TetrahedralParticleEmbedding& InIntersection, const int32 InConstraintIndex, const FIntVector2& SourceRange, const FIntVector2& TargetRange, const T PositionTargetStiffness)
+		{
+			GSConstraintHandle* NewConstraintPtr = new GSConstraintHandle();
+			NewConstraintPtr -> ConstraintIndex = InConstraintIndex;
+			Handles[InConstraintIndex] = NewConstraintPtr;
+
+			const FIntVector4 TargetTet = InTargetGeom.Tetrahedron[InTargetGeom.TetrahedronStart[InIntersection.GeometryIndex] + InIntersection.TetrahedronIndex] + FIntVector4(TargetRange[0]);
+
+			FGaussSeidelWeakConstraintSingleData<T> SingleConstraintData;
+			SingleConstraintData.SingleIndices = { InIntersection.ParticleIndex + SourceRange[0] };
+			SingleConstraintData.SingleSecondIndices =  {TargetTet[0], TargetTet[1], TargetTet[2], TargetTet[3]};
+			SingleConstraintData.SingleWeights = {(T)1.};
+			SingleConstraintData.SingleSecondWeights = {InIntersection.BarycentricWeights[0], InIntersection.BarycentricWeights[1], InIntersection.BarycentricWeights[2], InIntersection.BarycentricWeights[3]};
+			SingleConstraintData.bIsAnisotropic = false;
+			SingleConstraintData.SingleNormal = TVec3<T>((T)0.);
+			float ConstraintStiffness = 0.f;
+			for (int32 k = 0; k < 4; k++)
+			{
+				ConstraintStiffness += SingleConstraintData.SingleSecondWeights[k] * PositionTargetStiffness * AllParticles.M(TargetTet[k]);
+			}
+			ConstraintStiffness += PositionTargetStiffness * AllParticles.M(SingleConstraintData.SingleIndices[0]);
+			SingleConstraintData.SingleStiffness = (T)ConstraintStiffness;
+
+			ConstraintsData.SetSingleConstraint(SingleConstraintData, InConstraintIndex);
+
+			return NewConstraintPtr;
+			
+		}
+
+		void RemoveSingleConstraint(const GSConstraintHandle* DeletedSingleConstraint)
+		{
+			const int32 OldConstraintIndex = DeletedSingleConstraint->ConstraintIndex;
+
+			const FGaussSeidelWeakConstraintSingleData<T>& SingleData = ConstraintsData.GetSingleConstraintData(OldConstraintIndex);
+
+			delete DeletedSingleConstraint;
+
+			ConstraintsData.RemoveConstraint(OldConstraintIndex);
+			Handles[OldConstraintIndex].ConstraintIndex = OldConstraintIndex;
+			
+			//Update incident element information:
+			const int32 IndicesOffset = SingleData.SingleIndices.Num();
+			for (int32 i = 0; i < SingleData.SingleIndices.Num(); i++)
+			{
+				const int32 ParticleIndex = SingleData.SingleIndices[i];
+				for (int32 j = 0; j < WCIncidentElements[ParticleIndex].Num(); j++)
+				{
+					if (WCIncidentElements[ParticleIndex][j] == i && WCIncidentElementsLocal[ParticleIndex][j] == i)
+					{
+						WCIncidentElements[ParticleIndex].Remove(j);
+						WCIncidentElementsLocal[ParticleIndex].Remove(j);
+						break;
+					}
+				}
+			}
+			for (int32 i = 0; i < SingleData.SingleSecondIndices.Num(); i++)
+			{
+				const int32 ParticleIndex = SingleData.SingleSecondIndices[i];
+				for (int32 j = 0; j < WCIncidentElements[ParticleIndex].Num(); j++)
+				{
+					if (WCIncidentElements[ParticleIndex][j] == i && WCIncidentElementsLocal[ParticleIndex][j] == i + IndicesOffset)
+					{
+						WCIncidentElements[ParticleIndex].Remove(j);
+						WCIncidentElementsLocal[ParticleIndex].Remove(j);
+						break;
+					}
+				}
+			}
+			
+		}
+
+		TArray<const GSConstraintHandle*> AddParticleTetrahedraConstraints(const GeometryCollection::Facades::FTetrahedralFacade& InTargetGeom, const ParticleType& AllParticles, const TArray<GeometryCollection::Facades::TetrahedralParticleEmbedding>& InIntersections, const FIntVector2& SourceRange, const FIntVector2& TargetRange, const T ConstraintStiffness)
+		{
+			TArray<const GSConstraintHandle*> ConstraintHandles;
+			ConstraintHandles.SetNum(InIntersections.Num());
+		
+
+			const int32 NumConstraintsOffset = ConstraintsData.Size();
+			ConstraintsData.AddConstraints(InIntersections.Num());
+			for (int32 i = 0; i < InIntersections.Num(); i++)
+			{
+				ConstraintHandles[i] = AddSingleParticleTetrahedraConstraint(InTargetGeom, AllParticles, InIntersections[i], NumConstraintsOffset + i, SourceRange, TargetRange, ConstraintStiffness);
+			}
+			Base::ComputeInitialWCData(AllParticles);
+			return ConstraintHandles;
+		}
+
+		void RemoveConstraints(const TArray<const GSConstraintHandle*>& DeletedPtrs)
+		{
+			for (const GSConstraintHandle* DeletedPtr: DeletedPtrs )
+			{
+				RemoveSingleConstraint(DeletedPtr);
+			}
+		}
+
+		void ReComputeNodalWeights(const TSet<int32>& DirtyVerts)
+		{
+			for (const int32 Vert : DirtyVerts)
+			{
+				check(Vert < NodalWeights.Num());
+				NodalWeights[Vert].Init(T(0), 6);
+				for (int32 j = 0; j < WCIncidentElements[Vert].Num(); j++)
+				{
+					const int32 ConstraintIndex = WCIncidentElements[Vert][j];
+					const int32 LocalIndex = WCIncidentElementsLocal[Vert][j];
+
+					T Weight = T(0);
+					if (LocalIndex >= ConstraintsData.GetIndices(ConstraintIndex).Num())
+					{
+						Weight = ConstraintsData.GetSecondWeights(ConstraintIndex)[LocalIndex - ConstraintsData.GetIndices(ConstraintIndex).Num()];
+					}
+					else
+					{
+						Weight = ConstraintsData.GetWeights(ConstraintIndex)[LocalIndex];
+					}
+
+					if (ConstraintsData.GetIsAnisotropic(ConstraintIndex))
+					{
+						for (int32 alpha = 0; alpha < 3; alpha++) 
+						{
+							NodalWeights[Vert][alpha] += ConstraintsData.GetNormal(ConstraintIndex)[alpha] * ConstraintsData.GetNormal(ConstraintIndex)[alpha] * Weight * Weight * ConstraintsData.GetStiffness(ConstraintIndex);
+						}
+
+						NodalWeights[Vert][3] += ConstraintsData.GetNormal(ConstraintIndex)[0] * ConstraintsData.GetNormal(ConstraintIndex)[1] * Weight * Weight * ConstraintsData.GetStiffness(ConstraintIndex);
+						NodalWeights[Vert][4] += ConstraintsData.GetNormal(ConstraintIndex)[0] * ConstraintsData.GetNormal(ConstraintIndex)[2] * Weight * Weight * ConstraintsData.GetStiffness(ConstraintIndex);
+						NodalWeights[Vert][5] += ConstraintsData.GetNormal(ConstraintIndex)[1] * ConstraintsData.GetNormal(ConstraintIndex)[2] * Weight * Weight * ConstraintsData.GetStiffness(ConstraintIndex);
+					}
+					else
+					{
+						for (int32 alpha = 0; alpha < 3; alpha++)
+						{
+							NodalWeights[Vert][alpha] += Weight * Weight * ConstraintsData.GetStiffness(ConstraintIndex);
+						}
+					}
+				}
+			}
+		}
+
+		void AdjustStiffness(const TArray<const GSConstraintHandle*> ConstraintHandles, const TArray<T>& StiffnessPerConstraint, const ParticleType& Particles)
+		{
+			TSet<int32> DirtyVertices;
+			for (int32 i = 0; i < ConstraintHandles.Num(); i++)
+			{
+				const GSConstraintHandle* SingleHandle = ConstraintHandles[i];
+				const T NewStiffness = StiffnessPerConstraint[i];
+				const int32 ConstraintIndex = SingleHandle->ConstraintIndex;
+				const FGaussSeidelWeakConstraintSingleData<T>& SingleConstraintData = ConstraintsData.GetSingleConstraintData(ConstraintIndex);
+				T NewConstraintStiffness = (T)0.;
+				for (int32 k = 0; k < SingleConstraintData.SingleWeights.Num(); k++)
+				{
+					NewConstraintStiffness += SingleConstraintData.SingleWeights[k] * NewStiffness * Particles.M(SingleConstraintData.SingleIndices[k]);
+				}
+				for (int32 k = 0; k < SingleConstraintData.SingleSecondWeights.Num(); k++)
+				{
+					NewConstraintStiffness += SingleConstraintData.SingleSecondWeights[k] * NewStiffness * Particles.M(SingleConstraintData.SingleSecondIndices[k]);
+				}
+				ConstraintsData.SetStiffness(ConstraintIndex, NewConstraintStiffness);
+				DirtyVertices.Append(SingleConstraintData.SingleIndices);
+				DirtyVertices.Append(SingleConstraintData.SingleSecondIndices);
+			}
+			ReComputeNodalWeights(DirtyVertices);
+		}
+
+		TArrayCollectionArray<GSConstraintHandle*> Handles; //same size as number of constraints 
+
+	};
+
+}
diff --git a/Engine/Source/Runtime/Experimental/Chaos/Public/Chaos/Deformable/GaussSeidelMainConstraint.h b/Engine/Source/Runtime/Experimental/Chaos/Public/Chaos/Deformable/GaussSeidelMainConstraint.h
--- a/Engine/Source/Runtime/Experimental/Chaos/Public/Chaos/Deformable/GaussSeidelMainConstraint.h
+++ b/Engine/Source/Runtime/Experimental/Chaos/Public/Chaos/Deformable/GaussSeidelMainConstraint.h
@@ -16,6 +16,7 @@
 DECLARE_CYCLE_STAT(TEXT("Chaos.Deformable.GSMainConstraint.Apply"), STAT_ChaosGSMainConstraint_Apply, STATGROUP_Chaos);
 DECLARE_CYCLE_STAT(TEXT("Chaos.Deformable.GSMainConstraint.Acceleration"), STAT_ChaosGSMainConstraint_Acceleration, STATGROUP_Chaos);
 DECLARE_CYCLE_STAT(TEXT("Chaos.Deformable.GSMainConstraint.Init"), STAT_ChaosGSMainConstraint_Init, STATGROUP_Chaos);
+DECLARE_CYCLE_STAT(TEXT("Chaos.Deformable.GSMainConstraint.InitTransientColor"), STAT_ChaosGSMainConstraint_InitTransientColor, STATGROUP_Chaos);
 DECLARE_CYCLE_STAT(TEXT("Chaos.Deformable.GSMainConstraint.InitDynamicColor"), STAT_ChaosGSMainConstraint_InitDynamicColor, STATGROUP_Chaos);
 
 DEFINE_LOG_CATEGORY_STATIC(LogDeformableGaussSeidelMainConstraint, Log, All);
@@ -72,12 +73,16 @@ namespace Chaos::Softs
 			StaticIncidentElementsLocal.SetNum(NewSize);
 			DynamicIncidentElements.SetNum(NewSize);
 			DynamicIncidentElementsLocal.SetNum(NewSize);
+			TransientIncidentElements.SetNum(NewSize);
+			TransientIncidentElementsLocal.SetNum(NewSize);
 			X_k_1.Init(TVector<T, 3>(T(0.)), NewSize);
 			X_k.Init(TVector<T, 3>(T(0.)), NewSize);
 		}
 
 		const TArray<TFunction<void(const ParticleType&, const int32, const int32, const T, TVec3<T>&, Chaos::PMatrix<T, 3, 3>&)>>& StaticConstraintResidualAndHessian() const { return AddStaticConstraintResidualAndHessian; }
 		TArray<TFunction<void(const ParticleType&, const int32, const int32, const T, TVec3<T>&, Chaos::PMatrix<T, 3, 3>&)>>& StaticConstraintResidualAndHessian() { return AddStaticConstraintResidualAndHessian; }
+		const TArray<TFunction<void(const ParticleType&, const int32, const int32, const T, TVec3<T>&, Chaos::PMatrix<T, 3, 3>&)>>& TransientConstraintResidualAndHessian() const { return AddTransientConstraintResidualAndHessian; }
+		TArray<TFunction<void(const ParticleType&, const int32, const int32, const T, TVec3<T>&, Chaos::PMatrix<T, 3, 3>&)>>& TransientConstraintResidualAndHessian() { return AddTransientConstraintResidualAndHessian; }
 		const TArray<TFunction<void(const ParticleType&, const int32, const int32, const T, TVec3<T>&, Chaos::PMatrix<T, 3, 3>&)>>& DynamicConstraintResidualAndHessian() const { return AddDynamicConstraintResidualAndHessian; }
 		TArray<TFunction<void(const ParticleType&, const int32, const int32, const T, TVec3<T>&, Chaos::PMatrix<T, 3, 3>&)>>& DynamicConstraintResidualAndHessian() { return AddDynamicConstraintResidualAndHessian; }
 		const TArray<TFunction<void(const int32, const T, Chaos::PMatrix<T, 3, 3>&)>>& PerNodeHessian() const { return AddPerNodeHessian; }
@@ -92,6 +97,13 @@ namespace Chaos::Softs
 			return CurrentSize;
 		}
 
+		int32 AddTransientConstraintResidualAndHessianRange(int32 NumConstraints)
+		{
+			int32 CurrentSize = AddTransientConstraintResidualAndHessian.Num();
+			AddTransientConstraintResidualAndHessian.AddDefaulted(NumConstraints);
+			return CurrentSize;
+		}
+
 		int32 AddDynamicConstraintResidualAndHessianRange(int32 NumConstraints)
 		{
 			int32 CurrentSize = AddDynamicConstraintResidualAndHessian.Num();
@@ -115,8 +127,22 @@ namespace Chaos::Softs
 
 		CHAOS_API void AddStaticConstraints(const TArray<TArray<int32>>& ExtraConstraints, TArray<TArray<int32>>& ExtraIncidentElements, TArray<TArray<int32>>& ExtraIncidentElementsLocal);
 
+		CHAOS_API void AddTransientConstraints(const TArray<TArray<int32>>& ExtraConstraints, TArray<TArray<int32>>& ExtraIncidentElements, TArray<TArray<int32>>& ExtraIncidentElementsLocal, bool CheckIncidentElements = false);
+
 		CHAOS_API void AddDynamicConstraints(const TArray<TArray<int32>>& ExtraConstraints, TArray<TArray<int32>>& ExtraIncidentElements, TArray<TArray<int32>>& ExtraIncidentElementsLocal, bool CheckIncidentElements = false);
 
+		inline void ResetDynamicConstraints()
+		{
+			DynamicConstraints = {};
+			for (int32 p = 0; p < DynamicIncidentElements.Num(); p++)
+			{
+				DynamicIncidentElements[p].SetNum(0);
+				DynamicIncidentElementsLocal[p].SetNum(0);
+			}
+			DynamicIncidentElementsOffsets = {};	
+		}
+
+
 		void Apply(ParticleType& Particles, const T Dt, const int32 MaxWriteIters = 10, const bool Write2File = false, const TPBDActiveView<FSolverParticles>* InParticleActiveView = nullptr)
 		{
  			PERF_SCOPE(STAT_ChaosGSMainConstraint_Apply);
@@ -179,12 +205,20 @@ namespace Chaos::Softs
 			ParticlesPerColor = StaticParticlesPerColor;
 		}
 
+		void InitTransientColor(const ParticleType& Particles)
+		{
+			PERF_SCOPE(STAT_ChaosGSMainConstraint_InitTransientColor);
+			ParticleColors = StaticParticleColors;
+			ParticlesPerColor = StaticParticlesPerColor;
+			Chaos::ComputeExtraNodalColoring(StaticConstraints, DynamicConstraints, TransientConstraints, Particles, StaticIncidentElements, DynamicIncidentElements, TransientIncidentElements, ParticleColors, ParticlesPerColor);
+		}
+
 		void InitDynamicColor(const ParticleType& Particles)
 		{
 			PERF_SCOPE(STAT_ChaosGSMainConstraint_InitDynamicColor);
 			ParticleColors = StaticParticleColors;
 			ParticlesPerColor = StaticParticlesPerColor;
-			Chaos::ComputeExtraNodalColoring(StaticConstraints, DynamicConstraints, Particles, StaticIncidentElements, DynamicIncidentElements, ParticleColors, ParticlesPerColor);
+			Chaos::ComputeExtraNodalColoring(StaticConstraints, DynamicConstraints, Particles, StaticIncidentElements, TransientIncidentElements, ParticleColors, ParticlesPerColor);
 		}
 
 		void Init(const T Dt, const ParticleType& Particles)
@@ -192,13 +226,13 @@ namespace Chaos::Softs
 			Resize((int32)Particles.Size());
 
 			PERF_SCOPE(STAT_ChaosGSMainConstraint_Init);
-			DynamicConstraints.SetNum(0);
-			for (int32 p = 0; p < DynamicIncidentElements.Num(); p++)
+			TransientConstraints.SetNum(0);
+			for (int32 p = 0; p < TransientIncidentElements.Num(); p++)
 			{
-				DynamicIncidentElements[p].SetNum(0);
-				DynamicIncidentElementsLocal[p].SetNum(0);
+				TransientIncidentElements[p].SetNum(0);
+				TransientIncidentElementsLocal[p].SetNum(0);
 			}
-			DynamicIncidentElementsOffsets.SetNum(0);
+			TransientIncidentElementsOffsets.SetNum(0);
 			if (!bDoQuasistatics)
 			{
 				for (int32 i = 0; i < xtilde.Num(); i++)
@@ -274,6 +308,18 @@ namespace Chaos::Softs
 				AddDynamicConstraintResidualAndHessian[ConstraintIndex](Particles, DynamicIncidentElements[p][i] - DynamicIncidentElementsOffsets[ConstraintIndex], DynamicIncidentElementsLocal[p][i], Dt, ParticleResidual, ParticleHessian);
 			}
 
+			ConstraintIndex = 0;
+
+			for (int32 i = 0; i < TransientIncidentElements[p].Num(); i++)
+			{
+				while (TransientIncidentElements[p][i] >= TransientIncidentElementsOffsets[ConstraintIndex + 1] && ConstraintIndex < TransientIncidentElementsOffsets.Num() - 1)
+				{
+					ConstraintIndex ++;
+				}
+
+				AddTransientConstraintResidualAndHessian[ConstraintIndex](Particles, TransientIncidentElements[p][i] - TransientIncidentElementsOffsets[ConstraintIndex], TransientIncidentElementsLocal[p][i], Dt, ParticleResidual, ParticleHessian);
+			}
+
 			for (int32 i = 0; i < AddPerNodeHessian.Num(); i++)
 			{
 				AddPerNodeHessian[i](p, Dt, ParticleHessian);
@@ -337,6 +383,9 @@ namespace Chaos::Softs
 		TArray<TArray<int32>> StaticConstraints = {};
 		TArray<TArray<int32>> StaticIncidentElements;
 		TArray<TArray<int32>> StaticIncidentElementsLocal;
+		TArray<TArray<int32>> TransientConstraints = {};
+		TArray<TArray<int32>> TransientIncidentElements;
+		TArray<TArray<int32>> TransientIncidentElementsLocal;
 		TArray<TArray<int32>> DynamicConstraints = {};
 		TArray<TArray<int32>> DynamicIncidentElements;
 		TArray<TArray<int32>> DynamicIncidentElementsLocal;
@@ -345,6 +394,7 @@ namespace Chaos::Softs
 		TFunction<void(const ParticleType&, const int32, const T, TVec3<T>&, Chaos::PMatrix<T, 3, 3>&)> ComputeInitialResidualAndHessian;
 		TArray<TFunction<void(const ParticleType&, const int32, const int32, const T, TVec3<T>&, Chaos::PMatrix<T, 3, 3>&)>> AddStaticConstraintResidualAndHessian;
 		TArray<TFunction<void(const ParticleType&, const int32, const int32, const T, TVec3<T>&, Chaos::PMatrix<T, 3, 3>&)>> AddDynamicConstraintResidualAndHessian;
+		TArray<TFunction<void(const ParticleType&, const int32, const int32, const T, TVec3<T>&, Chaos::PMatrix<T, 3, 3>&)>> AddTransientConstraintResidualAndHessian;
 		TArray<TFunction<void(const int32, const T, Chaos::PMatrix<T, 3, 3>&)>> AddPerNodeHessian;
 
 		//Coloring information:
@@ -355,6 +405,7 @@ namespace Chaos::Softs
 		TArray<TArray<int32>> ParticlesPerColor;
 
 		TArray<int32> StaticIncidentElementsOffsets;
+		TArray<int32> TransientIncidentElementsOffsets;
 		TArray<int32> DynamicIncidentElementsOffsets;
 
 		bool bDoQuasistatics = false;
diff --git a/Engine/Source/Runtime/Experimental/Chaos/Public/Chaos/Deformable/GaussSeidelWeakConstraints.h b/Engine/Source/Runtime/Experimental/Chaos/Public/Chaos/Deformable/GaussSeidelWeakConstraints.h
--- a/Engine/Source/Runtime/Experimental/Chaos/Public/Chaos/Deformable/GaussSeidelWeakConstraints.h
+++ b/Engine/Source/Runtime/Experimental/Chaos/Public/Chaos/Deformable/GaussSeidelWeakConstraints.h
@@ -18,8 +18,243 @@
 #include <unordered_map>
 namespace Chaos::Softs
 {
-
 	using Chaos::TVec3;
+
+	template <typename T>
+	struct FGaussSeidelWeakConstraintSingleData
+	{
+		TArray<int32> SingleIndices  = {};
+		TArray<int32> SingleSecondIndices = {};
+		T SingleStiffness = (T)0.;
+		TArray<T> SingleWeights = {};
+		TArray<T> SingleSecondWeights = {};
+		bool bIsAnisotropic = false;
+		TVec3<T> SingleNormal = TVec3<T>((T)0.);
+	};
+
+	template<class T>
+	class TGaussSeidelWeakConstraintData : public TArrayCollection
+	{
+	public:
+		TGaussSeidelWeakConstraintData()
+		{
+			AddArray(&MIndices);
+			AddArray(&MSecondIndices);
+			AddArray(&MWeights);
+			AddArray(&MSecondWeights);
+			AddArray(&MStiffness);
+			AddArray(&MIsAnisotropic);
+			AddArray(&MNormals);
+		}
+		TGaussSeidelWeakConstraintData(const TGaussSeidelWeakConstraintData<T>& Other) = delete;
+		TGaussSeidelWeakConstraintData(TGaussSeidelWeakConstraintData<T>&& Other)
+		    : TArrayCollection(), MIndices(MoveTemp(Other.MIndices))
+			, MSecondIndices(MoveTemp(Other.MSecondIndices))
+			, MWeights(MoveTemp(Other.MWeights))
+			, MSecondWeights(MoveTemp(Other.MSecondWeights))
+			, MStiffness(MoveTemp(Other.MStiffness))
+			, MIsAnisotropic(MoveTemp(Other.MIsAnisotropic))
+		{
+			AddParticles(Other.Size());
+			AddArray(&MIndices);
+			AddArray(&MSecondIndices);
+			AddArray(&MWeights);
+			AddArray(&MSecondWeights);
+			AddArray(&MStiffness);
+			AddArray(&MIsAnisotropic);
+			AddArray(&MNormals);
+			Other.MSize = 0;
+		}
+
+		virtual ~TGaussSeidelWeakConstraintData()
+		{}
+
+		void AddConstraints(const int32 Num)
+		{
+			AddElementsHelper(Num);
+		}
+
+		void RemoveConstraint(const int32 Idx)
+		{
+			RemoveAtSwapHelper(Idx);
+		}
+
+		void SetSingleConstraint(const FGaussSeidelWeakConstraintSingleData<T>& SingleData, const int32 ConstraintIndex)
+		{
+			MIndices[ConstraintIndex] = SingleData.SingleIndices;
+			MSecondIndices[ConstraintIndex] = SingleData.SingleSecondIndices;
+			MStiffness[ConstraintIndex] = SingleData.SingleStiffness;
+			MWeights[ConstraintIndex] = SingleData.SingleWeights;
+			MSecondWeights[ConstraintIndex] = SingleData.SingleSecondWeights;
+			MNormals[ConstraintIndex] = SingleData.SingleNormal;
+			MIsAnisotropic[ConstraintIndex] = SingleData.bIsAnisotropic;
+		}
+
+		void AddSingleConstraint(const FGaussSeidelWeakConstraintSingleData<T>& SingleData)
+		{
+			AddConstraints(1);
+			SetSingleConstraint(SingleData, MSize - 1);
+		}
+
+		int32 Size() const 
+		{
+			return static_cast<int32>(MSize);
+		}
+
+		void Resize(const int32 Num)
+		{
+			ResizeHelper(Num);
+		}
+
+		TGaussSeidelWeakConstraintData& operator=(TGaussSeidelWeakConstraintData<T>&& Other)
+		{
+			MIndices = MoveTemp(Other.MIndices);
+			MIndices = MoveTemp(Other.MSecondIndices);
+			MIndices = MoveTemp(Other.MWeights);
+			MIndices = MoveTemp(Other.MSecondWeights);
+			MIndices = MoveTemp(Other.MStiffness);
+			MIndices = MoveTemp(Other.MIsAnisotropic);
+			MIndices = MoveTemp(Other.MNormals);
+			ResizeHelper(Other.Size());
+			Other.MSize = 0;
+			return *this;
+		}
+
+		inline const TArrayCollectionArray<TArray<int32>>& Indices() const
+		{
+			return MIndices;
+		}
+
+		const TArray<int32>& GetIndices(const int32 Index) const
+		{
+			return MIndices[Index];
+		}
+
+		void SetIndices(const int32 Index, const TArray<int32>& InIndices)
+		{
+			MIndices[Index] = InIndices;
+		}
+
+		inline const TArrayCollectionArray<TArray<int32>>& SecondIndices() const
+		{
+			return MSecondIndices;
+		}
+
+		const TArray<int32>& GetSecondIndices(const int32 Index) const
+		{
+			return MSecondIndices[Index];
+		}
+
+		void SetSecondIndices(const int32 Index, const TArray<int32>& InIndices)
+		{
+			MSecondIndices[Index] = InIndices;
+		}
+
+		inline const TArrayCollectionArray<TArray<T>>& Weights() const
+		{
+			return MWeights;
+		}
+
+		const TArray<T>& GetWeights(const int32 Index) const
+		{
+			return MWeights[Index];
+		}
+
+		void SetWeights(const int32 Index, const TArray<int32>& InWeights)
+		{
+			MWeights[Index] = InWeights;
+		}
+
+		inline const TArrayCollectionArray<TArray<T>>& SecondWeights() const
+		{
+			return MSecondWeights;
+		}
+
+		const TArray<T>& GetSecondWeights(const int32 Index) const
+		{
+			return MSecondWeights[Index];
+		}
+
+		void SetSecondWeights(const int32 Index, const TArray<int32>& InWeights)
+		{
+			MSecondWeights[Index] = InWeights;
+		}
+
+		inline const TArrayCollectionArray<bool>& IsAnisotropic() const
+		{
+			return MIsAnisotropic;
+		}
+
+		const bool GetIsAnisotropic(const int32 Index) const
+		{
+			return MIsAnisotropic[Index];
+		}
+
+		void SetIsAnisotropic(const int32 Index, const bool InIsAnisotropic)
+		{
+			MIsAnisotropic[Index] = InIsAnisotropic;
+		}
+
+		inline const TArrayCollectionArray<TVec3<T>>& Normals() const
+		{
+			return MNormals;
+		}
+
+		const TVec3<T>& GetNormal(const int32 Index) const
+		{
+			return MNormals[Index];
+		}
+
+		void SetNormal(const int32 Index, const TVec3<T>& InNormal)
+		{
+			MNormals[Index] = InNormal;
+		}
+
+		inline const TArrayCollectionArray<T>& Stiffness() const
+		{
+			return MStiffness;
+		}
+
+		T GetStiffness(const int32 Index) const
+		{
+			return MStiffness[Index];
+		}
+
+		void SetStiffness(const int32 Index, const T InStiffness)
+		{
+			MStiffness[Index] = InStiffness;
+		}
+
+		const FGaussSeidelWeakConstraintSingleData<T> GetSingleConstraintData(const int32 ConstraintIndex) const 
+		{
+			FGaussSeidelWeakConstraintSingleData<T> SingleConstraintData;
+			check(static_cast<uint32>(ConstraintIndex) < MSize);
+			if (ConstraintIndex > INDEX_NONE && static_cast<uint32>(ConstraintIndex) < MSize)
+			{
+				SingleConstraintData.SingleIndices = MIndices[ConstraintIndex];
+				SingleConstraintData.SingleSecondIndices = MSecondIndices[ConstraintIndex];
+				SingleConstraintData.SingleStiffness = MStiffness[ConstraintIndex];
+				SingleConstraintData.SingleWeights = MWeights[ConstraintIndex];
+				SingleConstraintData.SingleSecondWeights = MSecondWeights[ConstraintIndex];
+				SingleConstraintData.bIsAnisotropic = MIsAnisotropic[ConstraintIndex];
+				SingleConstraintData.SingleNormal = MNormals[ConstraintIndex];
+
+			}
+			return SingleConstraintData;
+		}
+
+	private:
+		TArrayCollectionArray<TArray<int32>> MIndices;
+		TArrayCollectionArray<TArray<int32>> MSecondIndices;
+		TArrayCollectionArray<TArray<T>> MWeights;
+		TArrayCollectionArray<TArray<T>> MSecondWeights;
+		TArrayCollectionArray<T> MStiffness;
+		TArrayCollectionArray<bool> MIsAnisotropic;
+		TArrayCollectionArray<TVector<T, 3>> MNormals;
+	};
+
+
+
 	template <typename T, typename ParticleType>
 	struct FGaussSeidelWeakConstraints 
 	{
@@ -31,39 +266,57 @@ namespace Chaos::Softs
 			const TArray<TArray<int32>>& InSecondIndices,
 			const TArray<TArray<T>>& InSecondWeights,
 			const FDeformableXPBDWeakConstraintParams& InParams
-		)
-			: Indices(InIndices), Weights(InWeights), SecondIndices(InSecondIndices), SecondWeights(InSecondWeights), Stiffness(InStiffness), DebugDrawParams(InParams)
+		): DebugDrawParams(InParams)
 		{
-			ensureMsgf(Indices.Num() == SecondIndices.Num(), TEXT("Input Double Bindings have wrong size"));
+			ensureMsgf(InIndices.Num() == InSecondIndices.Num(), TEXT("Input Double Bindings have wrong size"));
 
-			for (int32 i = 0; i < Indices.Num(); i++)
+			ConstraintsData.Resize(0);
+
+			ConstraintsData.AddConstraints(InIndices.Num());
+
+			for (int32 i = 0; i < InIndices.Num(); i++)
 			{
-				TSet<int32> IndicesSet = TSet<int32>(Indices[i]);
-				for (int32 j = 0; j < SecondIndices[i].Num(); j++)
+				FGaussSeidelWeakConstraintSingleData<T> SingleConstraintData;
+				SingleConstraintData.SingleIndices = InIndices[i];
+				SingleConstraintData.SingleSecondIndices = InSecondIndices[i];
+				SingleConstraintData.SingleWeights = InWeights[i];
+				SingleConstraintData.SingleSecondWeights = InSecondWeights[i];
+				SingleConstraintData.SingleStiffness = InStiffness[i];
+				ConstraintsData.SetSingleConstraint(SingleConstraintData, i);
+			}
+
+			for (int32 i = 0; i < ConstraintsData.Size(); i++)
+			{
+				const TArray<int32>& SingleIndices = ConstraintsData.GetIndices(i);
+				const TArray<int32>& SingleSecondIndices = ConstraintsData.GetSecondIndices(i);
+				for (int32 j = 0; j < SingleSecondIndices.Num(); j++)
 				{
-					ensureMsgf(!IndicesSet.Contains(SecondIndices[i][j]), TEXT("Indices and Second Indices overlaps. Currently not supported"));
+					ensureMsgf(!SingleIndices.Contains(SingleSecondIndices[j]), TEXT("Indices and Second Indices overlaps. Currently not supported"));
 				}
 			}
-
-			IsAnisotropic.Init(false, Indices.Num());
 		}
 
+		struct FGaussSeidelConstraintHandle
+		{
+			int32 ConstraintIndex;
+		};
+
 		virtual ~FGaussSeidelWeakConstraints() {}
 
-		void ComputeInitialWCData(const ParticleType& InParticles, const TArray<TArray<int32>>& MeshConstraints, const TArray <TArray<int32>>& MeshIncidentElements, const TArray <TArray<int32>>& MeshIncidentElementsLocal, TArray<TArray<int32>>& ParticlesPerColor)
+		void ComputeInitialWCData(const ParticleType& InParticles)
 		{
 			TArray<TArray<int32>> ExtraConstraints;
-			ExtraConstraints.Init(TArray<int32>(), Indices.Num());
-			for (int32 i = 0; i < Indices.Num(); i++)
+			ExtraConstraints.Init(TArray<int32>(), ConstraintsData.Size());
+			for (int32 i = 0; i < ExtraConstraints.Num(); i++)
 			{
-				ExtraConstraints[i].SetNum(Indices[i].Num() + SecondIndices[i].Num());
-				for (int32 j = 0; j < Indices[i].Num(); j++)
+				ExtraConstraints[i].SetNum(ConstraintsData.GetIndices(i).Num() + ConstraintsData.GetSecondIndices(i).Num());
+				for (int32 j = 0; j < ConstraintsData.GetIndices(i).Num(); j++)
 				{
-					ExtraConstraints[i][j] = Indices[i][j];
+					ExtraConstraints[i][j] = ConstraintsData.GetIndices(i)[j];
 				}
-				for (int32 j = 0; j < SecondIndices[i].Num(); j++)
+				for (int32 j = 0; j < ConstraintsData.GetSecondIndices(i).Num(); j++)
 				{
-					ExtraConstraints[i][j+Indices[i].Num()] = SecondIndices[i][j];
+					ExtraConstraints[i][j+ConstraintsData.GetIndices(i).Num()] = ConstraintsData.GetSecondIndices(i)[j];
 				}
 			}
 			WCIncidentElements = Chaos::Utilities::ComputeIncidentElements(ExtraConstraints, &WCIncidentElementsLocal);
@@ -82,31 +335,31 @@ namespace Chaos::Softs
 						int32 LocalIndex = WCIncidentElementsLocal[i][j];
 
 						T Weight = T(0);
-						if (LocalIndex >= Indices[ConstraintIndex].Num())
+						if (LocalIndex >= ConstraintsData.GetIndices(ConstraintIndex).Num())
 						{
-							Weight = SecondWeights[ConstraintIndex][LocalIndex - Indices[ConstraintIndex].Num()];
+							Weight = ConstraintsData.GetSecondWeights(ConstraintIndex)[LocalIndex - ConstraintsData.GetIndices(ConstraintIndex).Num()];
 						}
 						else
 						{
-							Weight = Weights[ConstraintIndex][LocalIndex];
+							Weight = ConstraintsData.GetWeights(ConstraintIndex)[LocalIndex];
 						}
 
-						if (IsAnisotropic[ConstraintIndex])
+						if (ConstraintsData.GetIsAnisotropic(ConstraintIndex))
 						{
 							for (int32 alpha = 0; alpha < 3; alpha++) 
 							{
-								NodalWeights[p][alpha] += Normals[ConstraintIndex][alpha] * Normals[ConstraintIndex][alpha] * Weight * Weight * Stiffness[ConstraintIndex];
+								NodalWeights[p][alpha] += ConstraintsData.GetNormal(ConstraintIndex)[alpha] * ConstraintsData.GetNormal(ConstraintIndex)[alpha] * Weight * Weight * ConstraintsData.GetStiffness(ConstraintIndex);
 							}
 
-							NodalWeights[p][3] += Normals[ConstraintIndex][0] * Normals[ConstraintIndex][1] * Weight * Weight * Stiffness[ConstraintIndex];
-							NodalWeights[p][4] += Normals[ConstraintIndex][0] * Normals[ConstraintIndex][2] * Weight * Weight * Stiffness[ConstraintIndex];
-							NodalWeights[p][5] += Normals[ConstraintIndex][1] * Normals[ConstraintIndex][2] * Weight * Weight * Stiffness[ConstraintIndex];
+							NodalWeights[p][3] += ConstraintsData.GetNormal(ConstraintIndex)[0] * ConstraintsData.GetNormal(ConstraintIndex)[1] * Weight * Weight * ConstraintsData.GetStiffness(ConstraintIndex);
+							NodalWeights[p][4] += ConstraintsData.GetNormal(ConstraintIndex)[0] * ConstraintsData.GetNormal(ConstraintIndex)[2] * Weight * Weight * ConstraintsData.GetStiffness(ConstraintIndex);
+							NodalWeights[p][5] += ConstraintsData.GetNormal(ConstraintIndex)[1] * ConstraintsData.GetNormal(ConstraintIndex)[2] * Weight * Weight * ConstraintsData.GetStiffness(ConstraintIndex);
 						}
 						else
 						{
 							for (int32 alpha = 0; alpha < 3; alpha++)
 							{
-								NodalWeights[p][alpha] += Weight * Weight * Stiffness[ConstraintIndex];
+								NodalWeights[p][alpha] += Weight * Weight * ConstraintsData.GetStiffness(ConstraintIndex);
 							}
 						}
 					}
@@ -114,21 +367,10 @@ namespace Chaos::Softs
 			}
 			NoCollisionNodalWeights = NodalWeights;
 			NoCollisionConstraints = ExtraConstraints;
-			InitialWCSize = Indices.Num();
+			InitialWCSize = ConstraintsData.Size();
 
 			NoCollisionWCIncidentElements = WCIncidentElements;
 			NoCollisionWCIncidentElementsLocal = WCIncidentElementsLocal;
-
-			StaticConstraints = MeshConstraints;
-			StaticConstraints += NoCollisionConstraints;
-			StaticIncidentElements = MeshIncidentElements;
-			for (int32 i = 0; i < NoCollisionWCIncidentElements.Num(); i++)
-			{
-				if (NoCollisionWCIncidentElements[i].Num() > 0)
-				{
-					StaticIncidentElements[i] += NoCollisionWCIncidentElements[i];
-				}
-			}
 		}
 
 		void AddWCResidual(const ParticleType& InParticles, const int32 p, const T Dt, TVec3<T>& res)
@@ -141,43 +383,47 @@ namespace Chaos::Softs
 					int32 ConstraintIndex = WCIncidentElements[WCIndex][j];
 					int32 LocalIndex = WCIncidentElementsLocal[WCIndex][j];
 					TVec3<T> SpringEdge((T)0.);
-					for (int32 l = 0; l < Weights[ConstraintIndex].Num(); l++)
+					const TArray<T>& Weight = ConstraintsData.GetWeights(ConstraintIndex);
+					const TArray<T>& SecondWeight = ConstraintsData.GetSecondWeights(ConstraintIndex);
+					const TArray<int32>& LocalIndices = ConstraintsData.GetIndices(ConstraintIndex);
+					const TArray<int32>& LocalSecondIndices = ConstraintsData.GetSecondIndices(ConstraintIndex);
+					for (int32 l = 0; l < Weight.Num(); l++)
 					{
 						for (int32 beta = 0; beta < 3; beta++)
 						{
-							SpringEdge[beta] += Weights[ConstraintIndex][l] * InParticles.P(Indices[ConstraintIndex][l])[beta];
+							SpringEdge[beta] += Weight[l] * InParticles.P(LocalIndices[l])[beta];
 						}
 					}
-					for (int32 l = 0; l < SecondWeights[ConstraintIndex].Num(); l++)
+					for (int32 l = 0; l < LocalSecondIndices.Num(); l++)
 					{
 						for (int32 beta = 0; beta < 3; beta++)
 						{
-							SpringEdge[beta] -= SecondWeights[ConstraintIndex][l] * InParticles.P(SecondIndices[ConstraintIndex][l])[beta];
+							SpringEdge[beta] -= SecondWeight[l] * InParticles.P(LocalSecondIndices[l])[beta];
 						}
 					}
 					T weight = T(0);
-					if (LocalIndex >= Indices[ConstraintIndex].Num())
+					if (LocalIndex >= LocalIndices.Num())
 					{
-						weight = -SecondWeights[ConstraintIndex][LocalIndex - Indices[ConstraintIndex].Num()];
+						weight = -SecondWeight[LocalIndex - LocalIndices.Num()];
 					}
 					else
 					{
-						weight = Weights[ConstraintIndex][LocalIndex];
+						weight = Weight[LocalIndex];
 					}
-					if (IsAnisotropic[ConstraintIndex])
+					if (ConstraintsData.GetIsAnisotropic(ConstraintIndex))
 					{
-						T comp = TVec3<T>::DotProduct(SpringEdge, Normals[ConstraintIndex]);
-						TVec3<T> proj = Normals[ConstraintIndex] * comp;
+						T comp = TVec3<T>::DotProduct(SpringEdge, ConstraintsData.GetNormal(ConstraintIndex));
+						TVec3<T> proj = ConstraintsData.GetNormal(ConstraintIndex) * comp;
 						for (int32 alpha = 0; alpha < 3; alpha++)
 						{
-							res[alpha] += Dt * Dt * Stiffness[ConstraintIndex] * proj[alpha] * weight;
+							res[alpha] += Dt * Dt * ConstraintsData.GetStiffness(ConstraintIndex) * proj[alpha] * weight;
 						}
 					}
 					else
 					{
 						for (int32 alpha = 0; alpha < 3; alpha++)
 						{
-							res[alpha] += Dt * Dt * Stiffness[ConstraintIndex] * SpringEdge[alpha] * weight;
+							res[alpha] += Dt * Dt * ConstraintsData.GetStiffness(ConstraintIndex) * SpringEdge[alpha] * weight;
 						}
 					}
 				}
@@ -214,70 +460,44 @@ namespace Chaos::Softs
 								const TArray<TArray<int32>>& InSecondIndices,
 								const TArray<TArray<T>>& InSecondWeights)
 		{
-			int32 Offset = Indices.Num();
-			Indices.SetNum(Offset + InIndices.Num());
-			SecondIndices.SetNum(Offset + InSecondIndices.Num());
-			Stiffness.SetNum(Offset + InStiffness.Num());
-			Weights.SetNum(Offset + InIndices.Num());
-			SecondWeights.SetNum(Offset + InIndices.Num());
+			const int32 Offset = ConstraintsData.Size();
+
+			ConstraintsData.AddConstraints(InIndices.Num());
 
 			for (int32 i = 0; i < InIndices.Num(); i++)
 			{
-				Indices[i + Offset] = InIndices[i];
-				SecondIndices[i + Offset] = InSecondIndices[i];
-				Stiffness[i + Offset] = InStiffness[i];
-				Weights[i + Offset] = InWeights[i];
-				SecondWeights[i + Offset] = InSecondWeights[i];
+				FGaussSeidelWeakConstraintSingleData<T> SingleConstraintData;
+				SingleConstraintData.SingleIndices = InIndices[i];
+				SingleConstraintData.SingleSecondIndices = InSecondIndices[i];
+				SingleConstraintData.SingleWeights = InWeights[i];
+				SingleConstraintData.SingleSecondWeights = InSecondWeights[i];
+				SingleConstraintData.SingleStiffness = InStiffness[i];
+				ConstraintsData.SetSingleConstraint(SingleConstraintData, i + Offset);
 			}
-
-			IsAnisotropic.Init(false, Indices.Num());
-			Normals.SetNum(Indices.Num());
 		}
 
 		void Resize(int32 Size)
 		{
-			Indices.SetNum(Size);
-			SecondIndices.SetNum(Size);
-			Stiffness.SetNum(Size);
-			Weights.SetNum(Size);
-			SecondWeights.SetNum(Size);
-			IsAnisotropic.SetNum(Size);
-			Normals.SetNum(Size);
+			ConstraintsData.Resize(Size);
 		}
 
 		void UpdatePointTriangleCollisionWCData(const FSolverParticles& Particles)
 		{	
-			TArray<TArray<int32>> IndicesTemp = Indices;
-			TArray<TArray<int32>> SecondIndicesTemp = SecondIndices;
-			TArray<T> StiffnessTemp = Stiffness;
-			TArray<TArray<T>> WeightsTemp = Weights;
-			TArray<TArray<T>> SecondWeightsTemp = SecondWeights;
-			TArray<bool> IsAnisotropicTemp = IsAnisotropic;
-			TArray<TVector<T, 3>> NormalsTemp = Normals;
-
-			Indices.SetNum(InitialWCSize);
-			SecondIndices.SetNum(InitialWCSize);
-			Stiffness.SetNum(InitialWCSize);
-			Weights.SetNum(InitialWCSize);
-			SecondWeights.SetNum(InitialWCSize);
-			IsAnisotropic.SetNum(InitialWCSize);
-			Normals.SetNum(InitialWCSize);
-
-			for (int32 i = InitialWCSize; i < IndicesTemp.Num(); i++)
+			TGaussSeidelWeakConstraintData<T> OriginalConstraintsData = ConstraintsData;
+
+			ConstraintsData.Resize(InitialWCSize);
+
+			for (int32 i = InitialWCSize; i < OriginalConstraintsData.Size(); i++)
 			{
-				ensureMsgf(IndicesTemp[i].Num() == 3, TEXT("Collision format is not point-triangle"));
-				ensureMsgf(SecondIndicesTemp[i].Num() == 1, TEXT("Collision format is not point-triangle"));
-				Chaos::TVector<float, 3> TriPos0(Particles.P(IndicesTemp[i][0])), TriPos1(Particles.P(IndicesTemp[i][1])), TriPos2(Particles.P(IndicesTemp[i][2])), ParticlePos(Particles.P(SecondIndicesTemp[i][0]));
+				const TArray<int32>& IndicesTemp = OriginalConstraintsData.GetIndices(i);
+				const TArray<int32>& SecondIndicesTemp = OriginalConstraintsData.GetSecondIndices(i);
+				ensureMsgf(OriginalConstraintsData.GetIndices(i).Num() == 3, TEXT("Collision format is not point-triangle"));
+				ensureMsgf(OriginalConstraintsData.GetSecondIndices(i).Num() == 1, TEXT("Collision format is not point-triangle"));
+				Chaos::TVector<float, 3> TriPos0(Particles.P(IndicesTemp[0])), TriPos1(Particles.P(IndicesTemp[1])), TriPos2(Particles.P(IndicesTemp[2])), ParticlePos(Particles.P(SecondIndicesTemp[0]));
 				Chaos::TVector<T, 3> Normal = FVector3f::CrossProduct(TriPos1 - TriPos0, TriPos2 - TriPos0);
 				if (FVector3f::DotProduct(ParticlePos - TriPos0, Normal) < 0.f) //not resolved, keep the spring
 				{
-					Indices.Add(IndicesTemp[i]);
-					SecondIndices.Add(SecondIndicesTemp[i]);
-					Stiffness.Add(StiffnessTemp[i]);
-					Weights.Add(WeightsTemp[i]);
-					SecondWeights.Add(SecondWeightsTemp[i]);
-					IsAnisotropic.Add(IsAnisotropicTemp[i]);
-					Normals.Add(NormalsTemp[i]);
+					ConstraintsData.AddConstraints(OriginalConstraintsData.GetSingleConstraintData(i));
 				}
 			}
 		}
@@ -286,39 +506,40 @@ namespace Chaos::Softs
 		{
 #if WITH_EDITOR
 			auto DoubleVert = [](Chaos::TVec3<T> V) { return FVector3d(V.X, V.Y, V.Z); };
-			for (int32 i = 0; i < Indices.Num(); i++)
+			for (int32 i = 0; i < ConstraintsData.Size(); i++)
 			{
+				const FGaussSeidelWeakConstraintSingleData<T>& SingleConstraintData = ConstraintsData.GetSingleConstraintData(i);
 				Chaos::TVec3<T> SourcePos((T)0.), TargetPos((T)0.);
-				for (int32 j = 0; j < Indices[i].Num(); j++)
+				for (int32 j = 0; j < SingleConstraintData.SingleIndices.Num(); j++)
 				{
-					SourcePos += Weights[i][j] * InParticles.P(Indices[i][j]);
+					SourcePos += SingleConstraintData.SingleWeights[j] * InParticles.P(SingleConstraintData.SingleIndices[j]);
 				}
-				for (int32 j = 0; j < SecondIndices[i].Num(); j++)
+				for (int32 j = 0; j < SingleConstraintData.SingleSecondIndices.Num(); j++)
 				{
-					TargetPos += SecondWeights[i][j] * InParticles.P(SecondIndices[i][j]);
+					TargetPos += SingleConstraintData.SingleSecondWeights[j] * InParticles.P(SingleConstraintData.SingleSecondIndices[j]);
 				}
 
 				float ParticleThickness = DebugDrawParams.DebugParticleWidth;
 				float LineThickness = DebugDrawParams.DebugLineWidth;
 
-				if (Indices[i].Num() == 1)
+				if (SingleConstraintData.SingleIndices.Num() == 1)
 				{
 					Chaos::FDebugDrawQueue::GetInstance().DrawDebugPoint(DoubleVert(SourcePos), FColor::Red, false, Dt, 0, ParticleThickness);
-					for (int32 j = 0; j < SecondIndices[i].Num(); j++)
+					for (int32 j = 0; j < SingleConstraintData.SingleSecondIndices.Num(); j++)
 					{
-						Chaos::FDebugDrawQueue::GetInstance().DrawDebugPoint(DoubleVert(InParticles.P(SecondIndices[i][j])), FColor::Green, false, Dt, 0, ParticleThickness);
-						Chaos::FDebugDrawQueue::GetInstance().DrawDebugLine(DoubleVert(InParticles.P(SecondIndices[i][j])), DoubleVert(InParticles.P(SecondIndices[i][(j + 1) % SecondIndices[i].Num()])), FColor::Green, false, Dt, 0, LineThickness);
+						Chaos::FDebugDrawQueue::GetInstance().DrawDebugPoint(DoubleVert(InParticles.P(SingleConstraintData.SingleSecondIndices[j])), FColor::Green, false, Dt, 0, ParticleThickness);
+						Chaos::FDebugDrawQueue::GetInstance().DrawDebugLine(DoubleVert(InParticles.P(SingleConstraintData.SingleSecondIndices[j])), DoubleVert(InParticles.P(SingleConstraintData.SingleSecondIndices[(j + 1) % SingleConstraintData.SingleSecondIndices.Num()])), FColor::Green, false, Dt, 0, LineThickness);
 					}
 
 				}
 
-				if (SecondIndices[i].Num() == 1)
+				if (SingleConstraintData.SingleSecondIndices.Num() == 1)
 				{
 					Chaos::FDebugDrawQueue::GetInstance().DrawDebugPoint(DoubleVert(TargetPos), FColor::Red, false, Dt, 0, ParticleThickness);
-					for (int32 j = 0; j < Indices[i].Num(); j++)
+					for (int32 j = 0; j < SingleConstraintData.SingleIndices.Num(); j++)
 					{
-						Chaos::FDebugDrawQueue::GetInstance().DrawDebugPoint(DoubleVert(InParticles.P(Indices[i][j])), FColor::Green, false, Dt, 0, ParticleThickness);
-						Chaos::FDebugDrawQueue::GetInstance().DrawDebugLine(DoubleVert(InParticles.P(Indices[i][j])), DoubleVert(InParticles.P(Indices[i][(j + 1) % SecondIndices[i].Num()])), FColor::Green, false, Dt, 0, LineThickness);
+						Chaos::FDebugDrawQueue::GetInstance().DrawDebugPoint(DoubleVert(InParticles.P(SingleConstraintData.SingleIndices[j])), FColor::Green, false, Dt, 0, ParticleThickness);
+						Chaos::FDebugDrawQueue::GetInstance().DrawDebugLine(DoubleVert(InParticles.P(SingleConstraintData.SingleIndices[j])), DoubleVert(InParticles.P(SingleConstraintData.SingleIndices[(j + 1) % SingleConstraintData.SingleIndices.Num()])), FColor::Green, false, Dt, 0, LineThickness);
 					}
 				}
 
@@ -453,28 +674,22 @@ namespace Chaos::Softs
 					&& MinIndex != SurfaceElements[i][1]
 					&& MinIndex != SurfaceElements[i][2])
 				{
-					Indices.SetNum(Indices.Num() + 1);
-					SecondIndices.SetNum(SecondIndices.Num() + 1);
-					Weights.SetNum(Weights.Num() + 1);
-					SecondWeights.SetNum(SecondWeights.Num() + 1);
-					Indices[Indices.Num() - 1].Add(SurfaceElements[i][0]);
-					Indices[Indices.Num() - 1].Add(SurfaceElements[i][1]);
-					Indices[Indices.Num() - 1].Add(SurfaceElements[i][2]);
-					SecondIndices[SecondIndices.Num() - 1].Add(MinIndex);
-					Weights[Weights.Num() - 1].Add(ClosestBary[0]);
-					Weights[Weights.Num() - 1].Add(ClosestBary[1]);
-					Weights[Weights.Num() - 1].Add(ClosestBary[2]);
-					SecondWeights[SecondWeights.Num() - 1].Add(1.f);
-
+					FGaussSeidelWeakConstraintSingleData<T> SingleConstraintData;
+					SingleConstraintData.SingleIndices = {SurfaceElements[i][0], SurfaceElements[i][1], SurfaceElements[i][2]};
+					SingleConstraintData.SingleSecondIndices = {MinIndex};
+					SingleConstraintData.Weights = {ClosestBary[0], ClosestBary[1], ClosestBary[2]};
+					SingleConstraintData.SecondWeights = {T(1.f)};
+					SingleConstraintData.bIsAnisotropic = UseAnisotropicSpring;
+					SingleConstraintData.SingleNormal = FaceNormal;
+					
 					float SpringStiffness = 0.f;
 					for (int32 k = 0; k < 3; k++)
 					{
 						SpringStiffness += ClosestBary[k] * PositionTargetStiffness * Particles.M(SurfaceElements[i][k]);
 					}
 					SpringStiffness += PositionTargetStiffness * Particles.M(MinIndex);
-					Stiffness.Add(SpringStiffness);
-					IsAnisotropic.Add(UseAnisotropicSpring);
-					Normals.Add(FaceNormal);
+					SingleConstraintData.SingleStiffness = (T)SpringStiffness;
+					ConstraintsData.AddSingleConstraint(SingleConstraintData);
 				}
 			}
 		}
@@ -505,19 +720,23 @@ namespace Chaos::Softs
 							{
 								const TVector<int32, 3>& Elem = Elements[CollisionPoint.Indices[1]];
 								const int32 IndexToWrite = ConstraintIndex.fetch_add(1);
-								Indices[IndexToWrite] = { Elem[0], Elem[1] ,Elem[2] };
-								SecondIndices[IndexToWrite] = { Index };
-								Weights[IndexToWrite] = { CollisionPoint.Bary[1], CollisionPoint.Bary[2], CollisionPoint.Bary[3] };
-								SecondWeights[IndexToWrite] = { 1.f };
+
+								FGaussSeidelWeakConstraintSingleData<T> SingleConstraintData;
+								SingleConstraintData.SingleIndices = { Elem[0], Elem[1] ,Elem[2] };
+								SingleConstraintData.SingleSecondIndices =  { Index };
+								SingleConstraintData.SingleWeights = { CollisionPoint.Bary[1], CollisionPoint.Bary[2], CollisionPoint.Bary[3] };
+								SingleConstraintData.SingleSecondWeights = {T(1.f)};
+								SingleConstraintData.bIsAnisotropic = UseAnisotropicSpring;
+								SingleConstraintData.SingleNormal = CollisionPoint.Normal;
+
 								float SpringStiffness = 0.f;
 								for (int32 k = 0; k < 3; k++)
 								{
-									SpringStiffness += Weights[IndexToWrite][k] * PositionTargetStiffness * Particles.M(Elem[k]);
+									SpringStiffness += SingleConstraintData.SingleWeights[k] * PositionTargetStiffness * Particles.M(Elem[k]);
 								}
 								SpringStiffness += PositionTargetStiffness * Particles.M(Index);
-								Stiffness[IndexToWrite] = SpringStiffness;
-								IsAnisotropic[IndexToWrite] = UseAnisotropicSpring;
-								Normals[IndexToWrite] = CollisionPoint.Normal;
+								SingleConstraintData.SingleStiffness = (T)SpringStiffness;
+								ConstraintsData.SetSingleConstraint(SingleConstraintData, IndexToWrite);
 							}
 						}
 					}
@@ -555,19 +774,23 @@ namespace Chaos::Softs
 							{
 								const TVector<int32, 3>& Elem = Elements[CollisionPoint.Indices[1]];
 								const int32 IndexToWrite = ConstraintIndex.fetch_add(1);
-								Indices[IndexToWrite] = { Elem[0], Elem[1] ,Elem[2] };
-								SecondIndices[IndexToWrite] = { Index };
-								Weights[IndexToWrite] = { CollisionPoint.Bary[1], CollisionPoint.Bary[2], CollisionPoint.Bary[3] };
-								SecondWeights[IndexToWrite] = { 1.f };
+
+								FGaussSeidelWeakConstraintSingleData<T> SingleConstraintData;
+								SingleConstraintData.SingleIndices = { Elem[0], Elem[1] ,Elem[2] };
+								SingleConstraintData.SingleSecondIndices =  { Index };
+								SingleConstraintData.SingleWeights = { CollisionPoint.Bary[1], CollisionPoint.Bary[2], CollisionPoint.Bary[3] };
+								SingleConstraintData.SingleSecondWeights = {T(1.f)};
+								SingleConstraintData.bIsAnisotropic = UseAnisotropicSpring;
+								SingleConstraintData.SingleNormal = CollisionPoint.Normal;
+
 								float SpringStiffness = 0.f;
 								for (int32 k = 0; k < 3; k++)
 								{
-									SpringStiffness += Weights[IndexToWrite][k] * PositionTargetStiffness * Particles.M(Elem[k]);
+									SpringStiffness += SingleConstraintData.SingleWeights[k] * PositionTargetStiffness * Particles.M(Elem[k]);
 								}
 								SpringStiffness += PositionTargetStiffness * Particles.M(Index);
-								Stiffness[IndexToWrite] = SpringStiffness;
-								IsAnisotropic[IndexToWrite] = UseAnisotropicSpring;
-								Normals[IndexToWrite] = CollisionPoint.Normal;
+								SingleConstraintData.SingleStiffness = (T)SpringStiffness;
+								ConstraintsData.SetSingleConstraint(SingleConstraintData, IndexToWrite);
 							}
 						}
 					}
@@ -579,73 +802,28 @@ namespace Chaos::Softs
 			Resize(ConstraintNum);
 		}
 
-		void ComputeCollisionWCData(const ParticleType& InParticles, const TArray<TArray<int32>>& MeshConstraints, const TArray <TArray<int32>>& MeshIncidentElements, const TArray <TArray<int32>>& MeshIncidentElementsLocal, TArray<TArray<int32>>& ParticlesPerColor)
+		void ComputeCollisionWCDataSimplified(TArray<TArray<int32>>& ExtraConstraints, TArray<TArray<int32>>& ExtraWCIncidentElements, TArray<TArray<int32>>& ExtraWCIncidentElementsLocal)
 		{
-			ensureMsgf(Indices.Num() >= InitialWCSize, TEXT("The size of Indices is smaller than InitialWCSize"));
-			TArray<TArray<int32>> ExtraConstraints;
-			ExtraConstraints.Init(TArray<int32>(), Indices.Num());
-			for (int32 i = InitialWCSize; i < Indices.Num(); i++)
-			{
-				ExtraConstraints[i - InitialWCSize].SetNum(Indices[i].Num() + SecondIndices[i].Num());
-				for (int32 j = 0; j < Indices[i].Num(); j++)
-				{
-					ExtraConstraints[i - InitialWCSize][j] = Indices[i][j];
-				}
-				for (int32 j = 0; j < SecondIndices[i].Num(); j++)
-				{
-					ExtraConstraints[i - InitialWCSize][j + Indices[i].Num()] = SecondIndices[i][j];
-				}
-			}
-
-			TArray<TArray<int32>> ExtraWCIncidentElements, ExtraWCIncidentElementsLocal;
-			ExtraWCIncidentElements = Chaos::Utilities::ComputeIncidentElements(ExtraConstraints, &ExtraWCIncidentElementsLocal);
+			ensureMsgf(ConstraintsData.Size() >= InitialWCSize, TEXT("The size of Indices is smaller than InitialWCSize"));
 
-			//TArray<TVector<int32, 4>> TotalConstraints = NoCollisionConstraints;
-			//TotalConstraints += ExtraConstraints;
-			//TODO (Yizhou): Make the following more efficient by computing only extra incident elements in the future. 
-			//WCIncidentElements = Chaos::Utilities::ComputeIncidentElements(TotalConstraints, &WCIncidentElementsLocal);
-
-			//TArray<TVector<int32, 4>> TotalConstraints = NoCollisionConstraints;
-			//TotalConstraints += ExtraConstraints;
-			//TODO (Yizhou): Make the following more efficient by computing only extra incident elements in the future. 
-			//TArray<TArray<int32>> WCIncidentElementsLocalTemp;
-			////WCIncidentElements = Chaos::Utilities::ComputeIncidentElements(TotalConstraints, &WCIncidentElementsLocalTemp);
-			//WCIncidentElementsLocal = WCIncidentElementsLocalTemp;
-			WCIncidentElements = NoCollisionWCIncidentElements;
-			WCIncidentElementsLocal = NoCollisionWCIncidentElementsLocal;
-
-			if (NoCollisionWCIncidentElements.Num() < ExtraWCIncidentElements.Num())
+			ExtraConstraints.Init(TArray<int32>(), ConstraintsData.Size() - InitialWCSize);
+			for (int32 i = InitialWCSize; i < static_cast<int32>(ConstraintsData.Size()); i++)
 			{
-				WCIncidentElements.SetNum(ExtraWCIncidentElements.Num());
-				WCIncidentElementsLocal.SetNum(ExtraWCIncidentElementsLocal.Num());
-			}
+				const TArray<int32>& LocalIndices = ConstraintsData.GetIndices(i);
+				const TArray<int32>& LocalSecondIndices = ConstraintsData.GetSecondIndices(i);
 
-			for (int32 i = 0; i < ExtraWCIncidentElements.Num(); i++)
-			{
-				if (ExtraWCIncidentElements[i].Num() > 0)
+				ExtraConstraints[i - InitialWCSize].SetNum(LocalIndices.Num() + LocalSecondIndices.Num());
+				for (int32 j = 0; j < LocalIndices.Num(); j++)
 				{
-					TArray<int32> ExtraWCIncidentElementsTemp = ExtraWCIncidentElements[i];
-					for (int32 j = 0; j < ExtraWCIncidentElements[i].Num(); j++)
-					{
-						ExtraWCIncidentElementsTemp[j] += InitialWCSize;
-					}
-					WCIncidentElements[i] += ExtraWCIncidentElementsTemp;
-					WCIncidentElementsLocal[i] += ExtraWCIncidentElementsLocal[i];
+					ExtraConstraints[i - InitialWCSize][j] = LocalIndices[j];
 				}
-			}
-
-			//TODO(Yizhou): Check if the following variable is really necessary:
-			Particle2WCIndices.Init(INDEX_NONE, InParticles.Size());
-			for (int32 i = 0; i < WCIncidentElements.Num(); i++)
-			{
-				if (WCIncidentElements[i].Num() > 0)
+				for (int32 j = 0; j < LocalSecondIndices.Num(); j++)
 				{
-					int32 p = ExtraConstraints[WCIncidentElements[i][0]][WCIncidentElementsLocal[i][0]];
-					Particle2WCIndices[p] = i;
+					ExtraConstraints[i - InitialWCSize][j + LocalIndices.Num()] = LocalSecondIndices[j];
 				}
 			}
 
-			Chaos::ComputeExtraNodalColoring(StaticConstraints, ExtraConstraints, InParticles, StaticIncidentElements, ExtraWCIncidentElements, ParticleColors, ParticlesPerColor);
+			ExtraWCIncidentElements = Chaos::Utilities::ComputeIncidentElements(ExtraConstraints, &ExtraWCIncidentElementsLocal);
 
 			NodalWeights = NoCollisionNodalWeights;
 			for (int32 i = 0; i < ExtraWCIncidentElements.Num(); i++)
@@ -661,98 +839,34 @@ namespace Chaos::Softs
 					{
 						int32 LocalIndex = ExtraWCIncidentElementsLocal[i][j];
 						int32 ConstraintIndex = ExtraWCIncidentElements[i][j] + InitialWCSize;
-						T weight = T(0);
-						if (LocalIndex >= Indices[ConstraintIndex].Num())
-						{
-							weight = SecondWeights[ConstraintIndex][LocalIndex - Indices[ConstraintIndex].Num()];
-						}
-						else
-						{
-							weight = Weights[ConstraintIndex][LocalIndex];
-						}
-						if (IsAnisotropic[ConstraintIndex])
-						{
-							for (int32 alpha = 0; alpha < 3; alpha++)
-							{
-								NodalWeights[p][alpha] += Normals[ConstraintIndex][alpha] * Normals[ConstraintIndex][alpha] * weight * weight * Stiffness[ConstraintIndex];
-							}
-
-							NodalWeights[p][3] += Normals[ConstraintIndex][0] * Normals[ConstraintIndex][1] * weight * weight * Stiffness[ConstraintIndex];
-							NodalWeights[p][4] += Normals[ConstraintIndex][0] * Normals[ConstraintIndex][2] * weight * weight * Stiffness[ConstraintIndex];
-							NodalWeights[p][5] += Normals[ConstraintIndex][1] * Normals[ConstraintIndex][2] * weight * weight * Stiffness[ConstraintIndex];
-						}
-						else 
-						{
-							for (int32 alpha = 0; alpha < 3; alpha++)
-							{
-								NodalWeights[p][alpha] += weight * weight * Stiffness[ConstraintIndex];
-							}
-						}
-					}
-				}
-			}
-		}
 
+						const FGaussSeidelWeakConstraintSingleData<T>& SingleData = ConstraintsData.GetSingleConstraintData(ConstraintIndex);
 
-		void ComputeCollisionWCDataSimplified(TArray<TArray<int32>>& ExtraConstraints, TArray<TArray<int32>>& ExtraWCIncidentElements, TArray<TArray<int32>>& ExtraWCIncidentElementsLocal)
-		{
-			ensureMsgf(Indices.Num() >= InitialWCSize, TEXT("The size of Indices is smaller than InitialWCSize"));
-
-			ExtraConstraints.Init(TArray<int32>(), Indices.Num() - InitialWCSize);
-			for (int32 i = InitialWCSize; i < Indices.Num(); i++)
-			{
-				ExtraConstraints[i - InitialWCSize].SetNum(Indices[i].Num() + SecondIndices[i].Num());
-				for (int32 j = 0; j < Indices[i].Num(); j++)
-				{
-					ExtraConstraints[i - InitialWCSize][j] = Indices[i][j];
-				}
-				for (int32 j = 0; j < SecondIndices[i].Num(); j++)
-				{
-					ExtraConstraints[i - InitialWCSize][j + Indices[i].Num()] = SecondIndices[i][j];
-				}
-			}
-
-			ExtraWCIncidentElements = Chaos::Utilities::ComputeIncidentElements(ExtraConstraints, &ExtraWCIncidentElementsLocal);
-
-			NodalWeights = NoCollisionNodalWeights;
-			for (int32 i = 0; i < ExtraWCIncidentElements.Num(); i++)
-			{
-				if (ExtraWCIncidentElements[i].Num() > 0)
-				{
-					int32 p = ExtraConstraints[ExtraWCIncidentElements[i][0]][ExtraWCIncidentElementsLocal[i][0]];
-					if (NodalWeights[p].Num() == 0)
-					{
-						NodalWeights[p].Init(T(0), 6);
-					}
-					for (int32 j = 0; j < ExtraWCIncidentElements[i].Num(); j++)
-					{
-						int32 LocalIndex = ExtraWCIncidentElementsLocal[i][j];
-						int32 ConstraintIndex = ExtraWCIncidentElements[i][j] + InitialWCSize;
 						T weight = T(0);
-						if (LocalIndex >= Indices[ConstraintIndex].Num())
+						if (LocalIndex >= SingleData.SingleIndices.Num())
 						{
-							weight = SecondWeights[ConstraintIndex][LocalIndex - Indices[ConstraintIndex].Num()];
+							weight = SingleData.SingleWeights[LocalIndex - SingleData.SingleIndices.Num()];
 						}
 						else
 						{
-							weight = Weights[ConstraintIndex][LocalIndex];
+							weight = SingleData.SingleWeights[LocalIndex];
 						}
-						if (IsAnisotropic[ConstraintIndex])
+						if (SingleData.bIsAnisotropic)
 						{
 							for (int32 alpha = 0; alpha < 3; alpha++)
 							{
-								NodalWeights[p][alpha] += Normals[ConstraintIndex][alpha] * Normals[ConstraintIndex][alpha] * weight * weight * Stiffness[ConstraintIndex];
+								NodalWeights[p][alpha] += SingleData.SingleNormal[alpha] * SingleData.SingleNormal[alpha] * weight * weight * SingleData.SingleStiffness;
 							}
 
-							NodalWeights[p][3] += Normals[ConstraintIndex][0] * Normals[ConstraintIndex][1] * weight * weight * Stiffness[ConstraintIndex];
-							NodalWeights[p][4] += Normals[ConstraintIndex][0] * Normals[ConstraintIndex][2] * weight * weight * Stiffness[ConstraintIndex];
-							NodalWeights[p][5] += Normals[ConstraintIndex][1] * Normals[ConstraintIndex][2] * weight * weight * Stiffness[ConstraintIndex];
+							NodalWeights[p][3] += SingleData.SingleNormal[0] * SingleData.SingleNormal[1] * weight * weight * SingleData.SingleStiffness;
+							NodalWeights[p][4] += SingleData.SingleNormal[0] * SingleData.SingleNormal[2] * weight * weight * SingleData.SingleStiffness;
+							NodalWeights[p][5] += SingleData.SingleNormal[1] * SingleData.SingleNormal[2] * weight * weight * SingleData.SingleStiffness;
 						}
 						else
 						{
 							for (int32 alpha = 0; alpha < 3; alpha++)
 							{
-								NodalWeights[p][alpha] += weight * weight * Stiffness[ConstraintIndex];
+								NodalWeights[p][alpha] += weight * weight * SingleData.SingleStiffness;
 							}
 						}
 					}
@@ -761,27 +875,29 @@ namespace Chaos::Softs
 		}
 
 
-		TArray<TArray<int32>> GetStaticConstraintArrays(TArray<TArray<int32>>& IncidentElements, TArray<TArray<int32>> IncidentElementsLocal)
+		TArray<TArray<int32>> GetStaticConstraintArrays(TArray<TArray<int32>>& IncidentElements, TArray<TArray<int32>>& IncidentElementsLocal)
 		{
 			IncidentElements = NoCollisionWCIncidentElements;
 			IncidentElementsLocal = NoCollisionWCIncidentElementsLocal;
 			return NoCollisionConstraints;
 		}
 
-		TArray<TArray<int32>> GetDynamicConstraintArrays(TArray<TArray<int32>>& IncidentElements, TArray<TArray<int32>> IncidentElementsLocal)
+		TArray<TArray<int32>> GetDynamicConstraintArrays(TArray<TArray<int32>>& IncidentElements, TArray<TArray<int32>>& IncidentElementsLocal)
 		{
 			TArray<TArray<int32>> ExtraConstraints;
-			ExtraConstraints.Init(TArray<int32>(), Indices.Num());
-			for (int32 i = InitialWCSize; i < Indices.Num(); i++)
+			ExtraConstraints.Init(TArray<int32>(), ConstraintsData.Size());
+			for (int32 i = InitialWCSize; i < ConstraintsData.Size(); i++)
 			{
-				ExtraConstraints[i - InitialWCSize].SetNum(Indices[i].Num() + SecondIndices[i].Num());
-				for (int32 j = 0; j < Indices[i].Num(); j++)
+				const TArray<int32>& LocalIndices = ConstraintsData.GetIndices(i);
+				const TArray<int32>& LocalSecondIndices = ConstraintsData.GetSecondIndices(i);
+				ExtraConstraints[i - InitialWCSize].SetNum(LocalIndices.Num() + LocalSecondIndices.Num());
+				for (int32 j = 0; j < LocalIndices.Num(); j++)
 				{
-					ExtraConstraints[i - InitialWCSize][j] = Indices[i][j];
+					ExtraConstraints[i - InitialWCSize][j] = LocalIndices[j];
 				}
-				for (int32 j = 0; j < SecondIndices[i].Num(); j++)
+				for (int32 j = 0; j < LocalSecondIndices.Num(); j++)
 				{
-					ExtraConstraints[i - InitialWCSize][j + Indices[i].Num()] = SecondIndices[i][j];
+					ExtraConstraints[i - InitialWCSize][j + LocalIndices.Num()] = LocalSecondIndices[j];
 				}
 			}
 
@@ -792,57 +908,53 @@ namespace Chaos::Softs
 
 		void AddWCResidualAndHessian(const ParticleType& InParticles, const int32 ConstraintIndex, const int32 LocalIndex, const T Dt, TVec3<T>& ParticleResidual, Chaos::PMatrix<T, 3, 3>& ParticleHessian)
 		{
+			const FGaussSeidelWeakConstraintSingleData<T>& SingleData = ConstraintsData.GetSingleConstraintData(ConstraintIndex);
+
 			TVec3<T> SpringEdge((T)0.);
-			for (int32 l = 0; l < Weights[ConstraintIndex].Num(); l++)
+			for (int32 l = 0; l < SingleData.SingleWeights.Num(); l++)
 			{
 				for (int32 beta = 0; beta < 3; beta++)
 				{
-					SpringEdge[beta] += Weights[ConstraintIndex][l] * InParticles.P(Indices[ConstraintIndex][l])[beta];
+					SpringEdge[beta] += SingleData.SingleWeights[l] * InParticles.P(SingleData.SingleIndices[l])[beta];
 				}
 			}
-			for (int32 l = 0; l < SecondWeights[ConstraintIndex].Num(); l++)
+			for (int32 l = 0; l < SingleData.SingleSecondWeights.Num(); l++)
 			{
 				for (int32 beta = 0; beta < 3; beta++)
 				{
-					SpringEdge[beta] -= SecondWeights[ConstraintIndex][l] * InParticles.P(SecondIndices[ConstraintIndex][l])[beta];
+					SpringEdge[beta] -= SingleData.SingleSecondWeights[l] * InParticles.P(SingleData.SingleSecondIndices[l])[beta];
 				}
 			}
 			T weight = T(0);
-			if (LocalIndex >= Indices[ConstraintIndex].Num())
+			if (LocalIndex >= SingleData.SingleIndices.Num())
 			{
-				weight = -SecondWeights[ConstraintIndex][LocalIndex - Indices[ConstraintIndex].Num()];
+				weight = -SingleData.SingleSecondWeights[LocalIndex - SingleData.SingleIndices.Num()];
 			}
 			else
 			{
-				weight = Weights[ConstraintIndex][LocalIndex];
+				weight = SingleData.SingleWeights[LocalIndex];
 			}
-			if (IsAnisotropic[ConstraintIndex])
+			if (SingleData.bIsAnisotropic)
 			{
-				T comp = TVec3<T>::DotProduct(SpringEdge, Normals[ConstraintIndex]);
-				TVec3<T> proj = Normals[ConstraintIndex] * comp;
+				T comp = TVec3<T>::DotProduct(SpringEdge, SingleData.SingleNormal);
+				TVec3<T> proj = SingleData.SingleNormal * comp;
 				for (int32 alpha = 0; alpha < 3; alpha++)
 				{
-					ParticleResidual[alpha] += Dt * Dt * Stiffness[ConstraintIndex] * proj[alpha] * weight;
+					ParticleResidual[alpha] += Dt * Dt * SingleData.SingleStiffness * proj[alpha] * weight;
 				}
 			}
 			else
 			{
 				for (int32 alpha = 0; alpha < 3; alpha++)
 				{
-					ParticleResidual[alpha] += Dt * Dt * Stiffness[ConstraintIndex] * SpringEdge[alpha] * weight;
+					ParticleResidual[alpha] += Dt * Dt * SingleData.SingleStiffness * SpringEdge[alpha] * weight;
 				}
 			}
 
 		}
 
-		TArray<TArray<int32>> Indices;
-		TArray<TArray<T>> Weights;
-		TArray<TVector<T, 3>> Constraints;
-		TArray<TArray<int32>> SecondIndices;
-		TArray<TArray<T>> SecondWeights;
-		TArray<T> Stiffness;
-		TArray<bool> IsAnisotropic;
-		TArray<TVector<T, 3>> Normals;
+		TGaussSeidelWeakConstraintData<T> ConstraintsData;
+
 		TArray<TArray<T>> NodalWeights;
 		TArray<int32> Particle2WCIndices;
 
@@ -853,18 +965,9 @@ namespace Chaos::Softs
 
 		int32 InitialWCSize;
 		TArray<TArray<T>> NoCollisionNodalWeights;
-		//TArray<TVector<int32, 4>> NoCollisionConstraints;
-		//For debugging
-		bool Detected = false;
-
 		TArray<TArray<int32>> NoCollisionConstraints;
 		TArray<TArray<int32>> NoCollisionWCIncidentElements;
 		TArray<TArray<int32>> NoCollisionWCIncidentElementsLocal;
-
-		TArray<TArray<int32>> StaticConstraints;
-		TArray<TArray<int32>> StaticIncidentElements;
-
-		TArray<int32> ParticleColors; 
 	};
 
 
diff --git a/Engine/Source/Runtime/Experimental/Chaos/Public/Chaos/GraphColoring.h b/Engine/Source/Runtime/Experimental/Chaos/Public/Chaos/GraphColoring.h
--- a/Engine/Source/Runtime/Experimental/Chaos/Public/Chaos/GraphColoring.h
+++ b/Engine/Source/Runtime/Experimental/Chaos/Public/Chaos/GraphColoring.h
@@ -129,4 +129,8 @@ CHAOS_API TArray<TArray<int32>> ComputeNodalColoring(const TArray<TArray<int32>>
 template<typename T>
 CHAOS_API void ComputeExtraNodalColoring(const TArray<TArray<int32>>& Graph, const TArray<TArray<int32>>& ExtraGraph, const Chaos::TDynamicParticles<T, 3>& InParticles, const TArray<TArray<int32>>& IncidentElements, const TArray<TArray<int32>>& ExtraIncidentElements, TArray<int32>& ParticleColors, TArray<TArray<int32>>& ParticlesPerColor);
 
+template<typename T>
+CHAOS_API void ComputeExtraNodalColoring(const TArray<TArray<int32>>& StaticGraph, const TArray<TArray<int32>>& DynamicGraph, const TArray<TArray<int32>>& ExtraGraph, const Chaos::TDynamicParticles<T, 3>& InParticles, const TArray<TArray<int32>>& StaticIncidentElements, const TArray<TArray<int32>>& DynamicIncidentElements, const TArray<TArray<int32>>& ExtraIncidentElements, TArray<int32>& ParticleColors, TArray<TArray<int32>>& ParticlesPerColor);
+
+
 }
