Chaos - Fix Capsule-Triangle collision

diff --git a/Engine/Source/Runtime/Experimental/Chaos/Private/Chaos/Collision/CapsuleTriangleContactPoint.cpp b/Engine/Source/Runtime/Experimental/Chaos/Private/Chaos/Collision/CapsuleTriangleContactPoint.cpp
--- a/Engine/Source/Runtime/Experimental/Chaos/Private/Chaos/Collision/CapsuleTriangleContactPoint.cpp
+++ b/Engine/Source/Runtime/Experimental/Chaos/Private/Chaos/Collision/CapsuleTriangleContactPoint.cpp
@@ -160,21 +160,21 @@ namespace Chaos
 			FVec3 SegmentP, EdgeP;
 			Utilities::NearestPointsOnLineSegments(P0, P1, EdgeP0, EdgeP1, SegmentT, EdgeT, SegmentP, EdgeP);
 
-			// Calculate the separation vector, correct for sign
+			// Calculate the separation vector (from triangle to capsule)
 			FVec3 SegmentEdgeN = SegmentP - EdgeP;
 			FReal SegmentEdgeDistSign = FReal(1);
 			const FReal SegmentEdgeDistSq = SegmentEdgeN.SizeSquared();
 
-			// Separating axis always points away from the triangle
+			// If the near point on the capsule axis is inside the triangle, fix the normal
 			const FReal DotEdge = FVec3::DotProduct(SegmentEdgeN, EdgeNs[EdgeIndex]);
 			if (DotEdge < FReal(-NormalTolerance))
 			{
 				SegmentEdgeN = -SegmentEdgeN;
 				SegmentEdgeDistSign = FReal(-1);
 			}
-
 			const FReal DotFace = FVec3::DotProduct(SegmentEdgeN, FaceN);
 
+			// If the near point on the capsule axis is outside the triangle check for cull distance
 			if (SegmentEdgeDistSign > FReal(0))
 			{
 				// We generate contacts when separation is within cull distance
@@ -403,15 +403,17 @@ namespace Chaos
 					}
 
 					// For Vertex contacts, check that the normal is in the valid range
+					// it must point away from the edge vectors that share the vertex
 					if (EdgeT == FReal(0))
 					{
 						const int32 PrevEdgeVertexIndex0 = (EdgeIndex >= 2) ? (EdgeIndex - 2) : (EdgeIndex - 2 + 3);
 						const FVec3& PrevEdgeP0 = Triangle.GetVertex(PrevEdgeVertexIndex0);
 						const FVec3& PrevEdgeP1 = EdgeP0;
 
+						// NOTE: both edge vetors here point towards the shared vertex (at EdgeP0)
 						const FReal PrevEdgeDotNormal = FVec3::DotProduct(PrevEdgeP1 - PrevEdgeP0, SegmentEdgeN);
-						const FReal EdgeDotNormal = FVec3::DotProduct(EdgeP1 - EdgeP0, SegmentEdgeN);
-						if ((PrevEdgeDotNormal < -NormalTolerance) || (EdgeDotNormal <- NormalTolerance))
+						const FReal EdgeDotNormal = FVec3::DotProduct(EdgeP0 - EdgeP1, SegmentEdgeN);
+						if ((PrevEdgeDotNormal < -NormalTolerance) || (EdgeDotNormal < -NormalTolerance))
 						{
 							continue;
 						}
@@ -422,9 +424,10 @@ namespace Chaos
 						const FVec3& NextEdgeP0 = EdgeP1;
 						const FVec3& NextEdgeP1 = Triangle.GetVertex(NextEdgeVertexIndex1);
 
+						// NOTE: both edge vetors here point towards the shared vertex (at EgdeP1)
 						const FReal EdgeDotNormal = FVec3::DotProduct(EdgeP1 - EdgeP0, SegmentEdgeN);
 						const FReal NextEdgeDotNormal = FVec3::DotProduct(NextEdgeP0 - NextEdgeP1, SegmentEdgeN);
-						if ((EdgeDotNormal < -NormalTolerance) || (NextEdgeDotNormal <- NormalTolerance))
+						if ((EdgeDotNormal < -NormalTolerance) || (NextEdgeDotNormal < -NormalTolerance))
 						{
 							continue;
 						}
