Ensure that duplicate entries from unresolved reference lookup tables are removed when an object is no longer replicated to the client.

diff --git a/Engine/Source/Runtime/Experimental/Iris/Core/Private/Iris/ReplicationSystem/ReplicationReader.cpp b/Engine/Source/Runtime/Experimental/Iris/Core/Private/Iris/ReplicationSystem/ReplicationReader.cpp
--- a/Engine/Source/Runtime/Experimental/Iris/Core/Private/Iris/ReplicationSystem/ReplicationReader.cpp
+++ b/Engine/Source/Runtime/Experimental/Iris/Core/Private/Iris/ReplicationSystem/ReplicationReader.cpp
@@ -1359,6 +1359,11 @@ void FReplicationReader::UpdateObjectReferenceTracking_Fast(FReplicatedObjectInf
 
 					// Add to tracking
 					const uint32 OwnerInternalIndex = ReplicationInfo->InternalIndex;
+
+					// Calling TMultiMap::Add() as this is O(1) performance as oppossed to TMultiMap::AddUnique() which is O(n). This does 
+					// mean that UnresolvedHandleToDependents can contain duplicate (RefHandle, OwnerInternalIndex) pairs but is handled
+					// gracefully by the rest of the code. It does mean that TMultiMap::Remove() must be called instead of TMultiMap::RemoveSingle()
+					// to remove all pairs.
 					UnresolvedHandleToDependents.Add(RefHandle, OwnerInternalIndex);
 					UE_LOG(LogIris, Verbose, TEXT("FReplicationReader::UpdateObjectReferenceTracking Adding unresolved reference %s for %s (OwnerInternalIndex=%d)"), ToCStr(RefHandle.ToString()), ToCStr(NetRefHandleManager->GetNetRefHandleFromInternalIndex(OwnerInternalIndex).ToString()), OwnerInternalIndex);
 				}
@@ -1431,6 +1436,11 @@ void FReplicationReader::UpdateObjectReferenceTracking_Fast(FReplicatedObjectInf
 
 					// Add to tracking
 					const uint32 OwnerInternalIndex = ReplicationInfo->InternalIndex;
+
+					// Calling TMultiMap::Add() as this is O(1) performance as oppossed to TMultiMap::AddUnique() which is O(n). This does 
+					// mean that UnresolvedHandleToDependents can contain duplicate (RefHandle, OwnerInternalIndex) pairs but is handled
+					// gracefully by the rest of the code. It does mean that TMultiMap::Remove() must be called instead of TMultiMap::RemoveSingle()
+					// to remove all pairs.
 					ResolvedDynamicHandleToDependents.Add(RefHandle, OwnerInternalIndex);
 					UE_LOG(LogIris, Verbose, TEXT("FReplicationReader::UpdateObjectReferenceTracking Adding resolved dynamic reference %s for %s"), ToCStr(RefHandle.ToString()), ToCStr(NetRefHandleManager->GetNetRefHandleFromInternalIndex(OwnerInternalIndex).ToString()));
 				}
@@ -1547,7 +1557,7 @@ void FReplicationReader::CleanupReferenceTracking(FReplicatedObjectInfo* ObjectI
 	{
 		// Remove from tracking
 		FNetRefHandle Handle = Element.Value;
-		UnresolvedHandleToDependents.RemoveSingle(Handle, ObjectIndex);
+		UnresolvedHandleToDependents.Remove(Handle, ObjectIndex);
 		RemoveFromUnresolvedCache(Handle);
 		UE_LOG(LogIris, Verbose, TEXT("FReplicationReader::CleanupReferenceTracking Removing unresolved reference %s for %s"), *Handle.ToString(), *(NetRefHandleManager->GetNetRefHandleFromInternalIndex(ObjectIndex).ToString()));
 	}
@@ -1559,7 +1569,7 @@ void FReplicationReader::CleanupReferenceTracking(FReplicatedObjectInfo* ObjectI
 	{
 		// Remove from tracking
 		const FNetRefHandle Handle = Element.Value;
-		ResolvedDynamicHandleToDependents.RemoveSingle(Handle, ObjectIndex);
+		ResolvedDynamicHandleToDependents.Remove(Handle, ObjectIndex);
 		UE_LOG(LogIris, Verbose, TEXT("FReplicationReader::CleanupReferenceTracking Removing resolved dynamic reference %s for %s"), *Handle.ToString(), *(NetRefHandleManager->GetNetRefHandleFromInternalIndex(ObjectIndex).ToString()));
 	}
 	ObjectInfo->ResolvedDynamicObjectReferences.Reset();
