Add scalar math node using numeric generic inputs and outputs
- those nodes are part of the dataflow core
- there's many to be added yet
- they currently co-exist with the existing float math node from the geometry collection dataflow nodes ( those will be deprecated later , possibly with upgrade path ? )
- remove the constraint of unique display name when registering nodes

diff --git a/Engine/Source/Runtime/Experimental/Dataflow/Core/Private/Dataflow/DataflowMathNodes.cpp b/Engine/Source/Runtime/Experimental/Dataflow/Core/Private/Dataflow/DataflowMathNodes.cpp
new file mode 100644
--- /dev/null
+++ b/Engine/Source/Runtime/Experimental/Dataflow/Core/Private/Dataflow/DataflowMathNodes.cpp
@@ -0,0 +1,422 @@
+// Copyright Epic Games, Inc. All Rights Reserved.
+
+#include "Dataflow/DataflowMathNodes.h"
+#include "Dataflow/DataflowNode.h"
+#include "Dataflow/DataflowNodeFactory.h"
+#include "Dataflow/DataflowNodeColorsRegistry.h"
+
+namespace Dataflow
+{
+	void RegisterDataflowMathNodes()
+	{
+		DATAFLOW_NODE_REGISTER_CREATION_FACTORY(FDataflowMathAbsNode);
+		DATAFLOW_NODE_REGISTER_CREATION_FACTORY(FDataflowMathAddNode);
+		DATAFLOW_NODE_REGISTER_CREATION_FACTORY(FDataflowMathCeilNode);
+		DATAFLOW_NODE_REGISTER_CREATION_FACTORY(FDataflowMathCubeNode);
+		DATAFLOW_NODE_REGISTER_CREATION_FACTORY(FDataflowMathDivideNode);
+		DATAFLOW_NODE_REGISTER_CREATION_FACTORY(FDataflowMathExpNode);
+		DATAFLOW_NODE_REGISTER_CREATION_FACTORY(FDataflowMathFloorNode);
+		DATAFLOW_NODE_REGISTER_CREATION_FACTORY(FDataflowMathFracNode);
+		DATAFLOW_NODE_REGISTER_CREATION_FACTORY(FDataflowMathInverseSquareRootNode);
+		DATAFLOW_NODE_REGISTER_CREATION_FACTORY(FDataflowMathLogNode);
+		DATAFLOW_NODE_REGISTER_CREATION_FACTORY(FDataflowMathLogXNode);
+		DATAFLOW_NODE_REGISTER_CREATION_FACTORY(FDataflowMathMaximumNode);
+		DATAFLOW_NODE_REGISTER_CREATION_FACTORY(FDataflowMathMinimumNode);
+		DATAFLOW_NODE_REGISTER_CREATION_FACTORY(FDataflowMathMultiplyNode);
+		DATAFLOW_NODE_REGISTER_CREATION_FACTORY(FDataflowMathNegateNode);
+		DATAFLOW_NODE_REGISTER_CREATION_FACTORY(FDataflowMathPowNode);
+		DATAFLOW_NODE_REGISTER_CREATION_FACTORY(FDataflowMathReciprocalNode);
+		DATAFLOW_NODE_REGISTER_CREATION_FACTORY(FDataflowMathRoundNode);
+		DATAFLOW_NODE_REGISTER_CREATION_FACTORY(FDataflowMathSignNode);
+		DATAFLOW_NODE_REGISTER_CREATION_FACTORY(FDataflowMathSquareNode);
+		DATAFLOW_NODE_REGISTER_CREATION_FACTORY(FDataflowMathSquareRootNode);
+		DATAFLOW_NODE_REGISTER_CREATION_FACTORY(FDataflowMathSubtractNode);
+		DATAFLOW_NODE_REGISTER_CREATION_FACTORY(FDataflowMathTruncNode);
+
+		// Math
+		static constexpr FLinearColor CDefaultMathNodeBodyTintColor(0.f, 0.f, 0.f, 0.5f);
+		DATAFLOW_NODE_REGISTER_CREATION_FACTORY_NODE_COLORS_BY_CATEGORY("Math", FLinearColor(0.f, 0.4f, 0.8f), CDefaultMathNodeBodyTintColor);
+	}
+}
+
+//-----------------------------------------------------------------------------------------------
+
+FDataflowMathOneInputOperatorNode::FDataflowMathOneInputOperatorNode(const Dataflow::FNodeParameters& InParam, FGuid InGuid)
+	:FDataflowNode(InParam, InGuid)
+{
+}
+
+void FDataflowMathOneInputOperatorNode::RegisterInputsAndOutputs()
+{
+	RegisterInputConnection(&A);
+	RegisterOutputConnection(&Result);
+
+	// set the output to double for now so that it is strongly type and easy to connect to the next node
+	// Cnce we can change the outpout type from the UI , this could be removed 
+	SetOutputConcreteType(&Result, TDataflowSingleTypePolicy<double>::TypeName);
+}
+
+void FDataflowMathOneInputOperatorNode::Evaluate(Dataflow::FContext& Context, const FDataflowOutput* Out) const
+{
+	if (Out->IsA(&Result))
+	{
+		const double InA = GetValue(Context, &A);
+		const double OutResult = ComputeResult(Context, InA);
+		SetValue(Context, OutResult, &Result);
+	}
+}
+
+//-----------------------------------------------------------------------------------------------
+
+FDataflowMathTwoInputsOperatorNode::FDataflowMathTwoInputsOperatorNode(const Dataflow::FNodeParameters& InParam, FGuid InGuid)
+	:FDataflowNode(InParam, InGuid)
+{
+}
+
+void FDataflowMathTwoInputsOperatorNode::RegisterInputsAndOutputs()
+{
+	RegisterInputConnection(&A);
+	RegisterInputConnection(&B);
+	RegisterOutputConnection(&Result);
+
+	// set the output to double for now so that it is strongly type and easy to connect to the next node
+	// Cnce we can change the outpout type from the UI , this could be removed 
+	SetOutputConcreteType(&Result, TDataflowSingleTypePolicy<double>::TypeName);
+}
+
+void FDataflowMathTwoInputsOperatorNode::Evaluate(Dataflow::FContext& Context, const FDataflowOutput* Out) const
+{
+	if (Out->IsA(&Result))
+	{
+		const double InA = GetValue(Context, &A);
+		const double InB = GetValue(Context, &B);
+		const double OutResult = ComputeResult(Context, InA, InB);
+		SetValue(Context, OutResult, &Result);
+	}
+}
+
+//-----------------------------------------------------------------------------------------------
+
+FDataflowMathAddNode::FDataflowMathAddNode(const Dataflow::FNodeParameters& InParam, FGuid InGuid)
+	: FDataflowMathTwoInputsOperatorNode(InParam, InGuid)
+{
+	RegisterInputsAndOutputs();
+}
+
+double FDataflowMathAddNode::ComputeResult(Dataflow::FContext& Context, double InA, double InB) const
+{
+	return (InA + InB);
+}
+
+//-----------------------------------------------------------------------------------------------
+
+FDataflowMathSubtractNode::FDataflowMathSubtractNode(const Dataflow::FNodeParameters& InParam, FGuid InGuid)
+	: FDataflowMathTwoInputsOperatorNode(InParam, InGuid)
+{
+	RegisterInputsAndOutputs();
+}
+
+double FDataflowMathSubtractNode::ComputeResult(Dataflow::FContext& Context, double InA, double InB) const
+{
+	return (InA - InB);
+}
+
+//-----------------------------------------------------------------------------------------------
+
+FDataflowMathMultiplyNode::FDataflowMathMultiplyNode(const Dataflow::FNodeParameters& InParam, FGuid InGuid)
+	: FDataflowMathTwoInputsOperatorNode(InParam, InGuid)
+{
+	RegisterInputsAndOutputs();
+}
+
+double FDataflowMathMultiplyNode::ComputeResult(Dataflow::FContext& Context, double InA, double InB) const
+{
+	return (InA * InB);
+}
+
+//-----------------------------------------------------------------------------------------------
+
+FDataflowMathDivideNode::FDataflowMathDivideNode(const Dataflow::FNodeParameters& InParam, FGuid InGuid)
+	: FDataflowMathTwoInputsOperatorNode(InParam, InGuid)
+{
+	RegisterInputsAndOutputs();
+	RegisterInputConnection(&Fallback);
+}
+
+double FDataflowMathDivideNode::ComputeResult(Dataflow::FContext& Context, double InA, double InB) const
+{
+	if (InB == 0)
+	{
+		return GetValue(Context, &Fallback);
+	}
+	return (InA / InB);
+}
+
+//-----------------------------------------------------------------------------------------------
+
+FDataflowMathMinimumNode::FDataflowMathMinimumNode(const Dataflow::FNodeParameters& InParam, FGuid InGuid)
+	: FDataflowMathTwoInputsOperatorNode(InParam, InGuid)
+{
+	RegisterInputsAndOutputs();
+}
+
+double FDataflowMathMinimumNode::ComputeResult(Dataflow::FContext& Context, double InA, double InB) const
+{
+	return FMath::Min(InA, InB);
+}
+
+//-----------------------------------------------------------------------------------------------
+
+FDataflowMathMaximumNode::FDataflowMathMaximumNode(const Dataflow::FNodeParameters& InParam, FGuid InGuid)
+	: FDataflowMathTwoInputsOperatorNode(InParam, InGuid)
+{
+	RegisterInputsAndOutputs();
+}
+
+double FDataflowMathMaximumNode::ComputeResult(Dataflow::FContext& Context, double InA, double InB) const
+{
+	return FMath::Max(InA, InB);
+}
+
+//-----------------------------------------------------------------------------------------------
+
+FDataflowMathReciprocalNode::FDataflowMathReciprocalNode(const Dataflow::FNodeParameters& InParam, FGuid InGuid)
+	: FDataflowMathOneInputOperatorNode(InParam, InGuid)
+{
+	RegisterInputsAndOutputs();
+	RegisterInputConnection(&Fallback);
+}
+
+double FDataflowMathReciprocalNode::ComputeResult(Dataflow::FContext& Context, double InA) const
+{
+	if (InA == 0)
+	{
+		return GetValue(Context, &Fallback);
+	}
+	return (1.0 / InA);
+}
+
+//-----------------------------------------------------------------------------------------------
+
+FDataflowMathSquareNode::FDataflowMathSquareNode(const Dataflow::FNodeParameters& InParam, FGuid InGuid)
+	: FDataflowMathOneInputOperatorNode(InParam, InGuid)
+{
+	RegisterInputsAndOutputs();
+}
+
+double FDataflowMathSquareNode::ComputeResult(Dataflow::FContext& Context, double InA) const
+{
+	return (InA * InA);
+}
+
+//-----------------------------------------------------------------------------------------------
+
+FDataflowMathCubeNode::FDataflowMathCubeNode(const Dataflow::FNodeParameters& InParam, FGuid InGuid)
+	: FDataflowMathOneInputOperatorNode(InParam, InGuid)
+{
+	RegisterInputsAndOutputs();
+}
+
+double FDataflowMathCubeNode::ComputeResult(Dataflow::FContext& Context, double InA) const
+{
+	return (InA * InA * InA);
+}
+
+//-----------------------------------------------------------------------------------------------
+
+FDataflowMathSquareRootNode::FDataflowMathSquareRootNode(const Dataflow::FNodeParameters& InParam, FGuid InGuid)
+	: FDataflowMathOneInputOperatorNode(InParam, InGuid)
+{
+	RegisterInputsAndOutputs();
+}
+
+double FDataflowMathSquareRootNode::ComputeResult(Dataflow::FContext& Context, double InA) const
+{
+	if (InA < 0)
+	{
+		// Context.Error()
+		return 0.0;
+	}
+	return FMath::Sqrt(InA * InA);
+}
+
+//-----------------------------------------------------------------------------------------------
+
+FDataflowMathInverseSquareRootNode::FDataflowMathInverseSquareRootNode(const Dataflow::FNodeParameters& InParam, FGuid InGuid)
+	: FDataflowMathOneInputOperatorNode(InParam, InGuid)
+{
+	RegisterInputsAndOutputs();
+	RegisterInputConnection(&Fallback);
+}
+
+double FDataflowMathInverseSquareRootNode::ComputeResult(Dataflow::FContext& Context, double InA) const
+{
+	if (InA == 0)
+	{
+		return GetValue(Context, &Fallback);
+	}
+	return FMath::InvSqrt(InA);
+}
+
+//-----------------------------------------------------------------------------------------------
+
+FDataflowMathNegateNode::FDataflowMathNegateNode(const Dataflow::FNodeParameters& InParam, FGuid InGuid)
+	: FDataflowMathOneInputOperatorNode(InParam, InGuid)
+{
+	RegisterInputsAndOutputs();
+}
+
+double FDataflowMathNegateNode::ComputeResult(Dataflow::FContext& Context, double InA) const
+{
+	return -InA;
+}
+
+//-----------------------------------------------------------------------------------------------
+
+FDataflowMathAbsNode::FDataflowMathAbsNode(const Dataflow::FNodeParameters& InParam, FGuid InGuid)
+	: FDataflowMathOneInputOperatorNode(InParam, InGuid)
+{
+	RegisterInputsAndOutputs();
+}
+
+double FDataflowMathAbsNode::ComputeResult(Dataflow::FContext& Context, double InA) const
+{
+	return FMath::Abs(InA);
+}
+
+//-----------------------------------------------------------------------------------------------
+
+FDataflowMathFloorNode::FDataflowMathFloorNode(const Dataflow::FNodeParameters& InParam, FGuid InGuid)
+	: FDataflowMathOneInputOperatorNode(InParam, InGuid)
+{
+	RegisterInputsAndOutputs();
+}
+
+double FDataflowMathFloorNode::ComputeResult(Dataflow::FContext& Context, double InA) const
+{
+	return FMath::FloorToDouble(InA);
+}
+
+//-----------------------------------------------------------------------------------------------
+
+FDataflowMathCeilNode::FDataflowMathCeilNode(const Dataflow::FNodeParameters& InParam, FGuid InGuid)
+	: FDataflowMathOneInputOperatorNode(InParam, InGuid)
+{
+	RegisterInputsAndOutputs();
+}
+
+double FDataflowMathCeilNode::ComputeResult(Dataflow::FContext& Context, double InA) const
+{
+	return FMath::CeilToDouble(InA);
+}
+
+//-----------------------------------------------------------------------------------------------
+
+FDataflowMathRoundNode::FDataflowMathRoundNode(const Dataflow::FNodeParameters& InParam, FGuid InGuid)
+	: FDataflowMathOneInputOperatorNode(InParam, InGuid)
+{
+	RegisterInputsAndOutputs();
+}
+
+double FDataflowMathRoundNode::ComputeResult(Dataflow::FContext& Context, double InA) const
+{
+	return FMath::RoundToDouble(InA);
+}
+
+//-----------------------------------------------------------------------------------------------
+
+FDataflowMathTruncNode::FDataflowMathTruncNode(const Dataflow::FNodeParameters& InParam, FGuid InGuid)
+	: FDataflowMathOneInputOperatorNode(InParam, InGuid)
+{
+	RegisterInputsAndOutputs();
+}
+
+double FDataflowMathTruncNode::ComputeResult(Dataflow::FContext& Context, double InA) const
+{
+	return FMath::TruncToDouble(InA);
+}
+
+//-----------------------------------------------------------------------------------------------
+
+FDataflowMathFracNode::FDataflowMathFracNode(const Dataflow::FNodeParameters& InParam, FGuid InGuid)
+	: FDataflowMathOneInputOperatorNode(InParam, InGuid)
+{
+	RegisterInputsAndOutputs();
+}
+
+double FDataflowMathFracNode::ComputeResult(Dataflow::FContext& Context, double InA) const
+{
+	return FMath::Frac(InA);
+}
+
+//-----------------------------------------------------------------------------------------------
+
+FDataflowMathPowNode::FDataflowMathPowNode(const Dataflow::FNodeParameters& InParam, FGuid InGuid)
+	: FDataflowMathTwoInputsOperatorNode(InParam, InGuid)
+{
+	RegisterInputsAndOutputs();
+}
+
+double FDataflowMathPowNode::ComputeResult(Dataflow::FContext& Context, double InA, double InB) const
+{
+	return FMath::Pow(InA, InB);
+}
+
+//-----------------------------------------------------------------------------------------------
+
+FDataflowMathLogXNode::FDataflowMathLogXNode(const Dataflow::FNodeParameters& InParam, FGuid InGuid)
+	: FDataflowMathOneInputOperatorNode(InParam, InGuid)
+{
+	Base.Value = 10.0; // default is base 10
+	RegisterInputsAndOutputs();
+	RegisterInputConnection(&Base);
+}
+
+double FDataflowMathLogXNode::ComputeResult(Dataflow::FContext& Context, double InA) const
+{
+	const double InBase = GetValue(Context, &Base);
+	if (InBase <= 0.f)
+	{
+		return 0.0;
+	}
+	return FMath::LogX(InBase, InA);
+}
+
+//-----------------------------------------------------------------------------------------------
+
+FDataflowMathLogNode::FDataflowMathLogNode(const Dataflow::FNodeParameters& InParam, FGuid InGuid)
+	: FDataflowMathOneInputOperatorNode(InParam, InGuid)
+{
+	RegisterInputsAndOutputs();
+}
+
+double FDataflowMathLogNode::ComputeResult(Dataflow::FContext& Context, double InA) const
+{
+	return FMath::Loge(InA);
+}
+
+//-----------------------------------------------------------------------------------------------
+
+FDataflowMathExpNode::FDataflowMathExpNode(const Dataflow::FNodeParameters& InParam, FGuid InGuid)
+	: FDataflowMathOneInputOperatorNode(InParam, InGuid)
+{
+	RegisterInputsAndOutputs();
+}
+
+double FDataflowMathExpNode::ComputeResult(Dataflow::FContext& Context, double InA) const
+{
+	return FMath::Exp(InA);
+}
+
+//-----------------------------------------------------------------------------------------------
+
+FDataflowMathSignNode::FDataflowMathSignNode(const Dataflow::FNodeParameters& InParam, FGuid InGuid)
+	: FDataflowMathOneInputOperatorNode(InParam, InGuid)
+{
+	RegisterInputsAndOutputs();
+}
+
+double FDataflowMathSignNode::ComputeResult(Dataflow::FContext& Context, double InA) const
+{
+	return FMath::Sign(InA);
+}
\ No newline at end of file
diff --git a/Engine/Source/Runtime/Experimental/Dataflow/Core/Private/Dataflow/DataflowNodeFactory.cpp b/Engine/Source/Runtime/Experimental/Dataflow/Core/Private/Dataflow/DataflowNodeFactory.cpp
--- a/Engine/Source/Runtime/Experimental/Dataflow/Core/Private/Dataflow/DataflowNodeFactory.cpp
+++ b/Engine/Source/Runtime/Experimental/Dataflow/Core/Private/Dataflow/DataflowNodeFactory.cpp
@@ -4,6 +4,7 @@
 
 #include "Dataflow/DataflowNode.h"
 #include "Dataflow/DataflowCoreNodes.h"
+#include "Dataflow/DataflowMathNodes.h"
 #include "Misc/MessageDialog.h"
 #include "Logging/LogMacros.h"
 
@@ -16,6 +17,7 @@ namespace Dataflow
 	void FNodeFactory::RegisterDefaultNodes()
 	{
 		Dataflow::RegisterCoreNodes();
+		Dataflow::RegisterDataflowMathNodes();
 	}
 
 	TSharedPtr<FDataflowNode> FNodeFactory::NewNodeFromRegisteredType(FGraph& Graph, const FNewNodeParameters& Param)
@@ -52,7 +54,7 @@ namespace Dataflow
 
 		FFactoryParameters NewParameters(Parameters.TypeName, FName(*NewDisplayNameString), Parameters.Category, Parameters.Tags, Parameters.ToolTip, Parameters.bIsDeprecated, Parameters.bIsExperimental, Parameters.NodeVersion);
 
-		if (ClassMap.Contains(NewParameters.TypeName) || DisplayMap.Contains(NewParameters.DisplayName))
+		if (ClassMap.Contains(NewParameters.TypeName))
 		{
 			if (ParametersMap[NewParameters.TypeName].DisplayName.IsEqual(NewParameters.DisplayName))
 			{
@@ -85,7 +87,6 @@ namespace Dataflow
 		{
 			ClassMap.Add(NewParameters.TypeName, NewFunction);
 			ParametersMap.Add(NewParameters.TypeName, NewParameters);
-			DisplayMap.Add(NewParameters.DisplayName, NewParameters.TypeName);
 
 			const FName TypeNameNoVersion = GetTypeNameNoVersion(NewParameters.TypeName);
 			VersionMap.FindOrAdd(TypeNameNoVersion).AddUnique(NewParameters.TypeName);
@@ -162,15 +163,6 @@ namespace Dataflow
 #endif // WITH_EDITOR
 	}
 
-	FName FNodeFactory::TypeNameFromDisplayName(const FName& DisplayName)
-	{
-		if (DisplayMap.Contains(DisplayName))
-		{
-			return DisplayMap[DisplayName];
-		}
-		return "";
-	}
-
 	const FFactoryParameters& FNodeFactory::GetParameters(FName InTypeName) const
 	{
 		static FFactoryParameters EmptyFactoryParameters;
diff --git a/Engine/Source/Runtime/Experimental/Dataflow/Core/Public/Dataflow/DataflowMathNodes.h b/Engine/Source/Runtime/Experimental/Dataflow/Core/Public/Dataflow/DataflowMathNodes.h
new file mode 100644
--- /dev/null
+++ b/Engine/Source/Runtime/Experimental/Dataflow/Core/Public/Dataflow/DataflowMathNodes.h
@@ -0,0 +1,415 @@
+// Copyright Epic Games, Inc. All Rights Reserved.
+
+#pragma once
+
+#include "ChaosLog.h"
+#include "CoreMinimal.h"
+#include "Dataflow/DataflowConnection.h"
+#include "Dataflow/DataflowNodeParameters.h"
+#include "Dataflow/DataflowNode.h"
+
+#include "DataflowMathNodes.generated.h"
+
+#define DATAFLOW_MATH_NODES_CATEGORY "Math|Scalar"
+
+/** One input operators base class */
+USTRUCT()
+struct FDataflowMathOneInputOperatorNode : public FDataflowNode
+{
+	GENERATED_USTRUCT_BODY()
+
+	UPROPERTY(EditAnywhere, Category = Operands, meta = (DataflowInput));
+	FDataflowNumericTypes A;
+
+	UPROPERTY(meta = (DataflowOutput))
+	FDataflowNumericTypes Result;
+
+public:
+	FDataflowMathOneInputOperatorNode() {};
+	FDataflowMathOneInputOperatorNode(const Dataflow::FNodeParameters& InParam, FGuid InGuid = FGuid::NewGuid());
+	virtual void Evaluate(Dataflow::FContext& Context, const FDataflowOutput* Out) const override;
+
+protected:
+	void RegisterInputsAndOutputs();
+	virtual double ComputeResult(Dataflow::FContext& Context, double InA) const { ensure(false); return 0.0; };
+};
+
+/** Two inputs operators base class */
+USTRUCT()
+struct FDataflowMathTwoInputsOperatorNode : public FDataflowNode
+{
+	GENERATED_USTRUCT_BODY()
+
+	UPROPERTY(EditAnywhere, Category = Operands, meta = (DataflowInput));
+	FDataflowNumericTypes A;
+
+	UPROPERTY(EditAnywhere, Category = Operands, meta = (DataflowInput));
+	FDataflowNumericTypes B;
+
+	UPROPERTY(meta = (DataflowOutput))
+	FDataflowNumericTypes Result;
+
+public:
+	FDataflowMathTwoInputsOperatorNode() {};
+	FDataflowMathTwoInputsOperatorNode(const Dataflow::FNodeParameters& InParam, FGuid InGuid = FGuid::NewGuid());
+	virtual void Evaluate(Dataflow::FContext& Context, const FDataflowOutput* Out) const override;
+
+protected:
+	void RegisterInputsAndOutputs();
+	virtual double ComputeResult(Dataflow::FContext& Context, double InA, double InB) const { ensure(false); return 0.0; };
+};
+
+/** Addition (A + B) */
+USTRUCT()
+struct FDataflowMathAddNode : public FDataflowMathTwoInputsOperatorNode
+{
+	GENERATED_USTRUCT_BODY()
+	DATAFLOW_NODE_DEFINE_INTERNAL(FDataflowMathAddNode, "Add", DATAFLOW_MATH_NODES_CATEGORY, "")
+
+public:
+	FDataflowMathAddNode(const Dataflow::FNodeParameters& InParam, FGuid InGuid = FGuid::NewGuid());
+	virtual double ComputeResult(Dataflow::FContext& Context, double InA, double InB) const override;
+};
+
+/** Subtraction (A - B) */
+USTRUCT()
+struct FDataflowMathSubtractNode : public FDataflowMathTwoInputsOperatorNode
+{
+	GENERATED_USTRUCT_BODY()
+	DATAFLOW_NODE_DEFINE_INTERNAL(FDataflowMathSubtractNode, "Subtract", DATAFLOW_MATH_NODES_CATEGORY, "")
+
+public:
+	FDataflowMathSubtractNode(const Dataflow::FNodeParameters& InParam, FGuid InGuid = FGuid::NewGuid());
+	virtual double ComputeResult(Dataflow::FContext& Context, double InA, double InB) const override;
+};
+
+/** Multiplication (A * B) */
+USTRUCT()
+struct FDataflowMathMultiplyNode : public FDataflowMathTwoInputsOperatorNode
+{
+	GENERATED_USTRUCT_BODY()
+	DATAFLOW_NODE_DEFINE_INTERNAL(FDataflowMathMultiplyNode, "Multiply", DATAFLOW_MATH_NODES_CATEGORY, "")
+
+public:
+	FDataflowMathMultiplyNode(const Dataflow::FNodeParameters& InParam, FGuid InGuid = FGuid::NewGuid());
+	virtual double ComputeResult(Dataflow::FContext& Context, double InA, double InB) const override;
+};
+
+/**
+* Division (A / B)
+* if B is equal to 0, 0 is returned Fallback value
+*/
+USTRUCT()
+struct FDataflowMathDivideNode : public FDataflowMathTwoInputsOperatorNode
+{
+	GENERATED_USTRUCT_BODY()
+	DATAFLOW_NODE_DEFINE_INTERNAL(FDataflowMathDivideNode, "Divide", DATAFLOW_MATH_NODES_CATEGORY, "")
+
+	UPROPERTY(EditAnywhere, Category = SafeDivide, meta = (DataflowInput));
+	FDataflowNumericTypes Fallback;
+
+public:
+	FDataflowMathDivideNode(const Dataflow::FNodeParameters& InParam, FGuid InGuid = FGuid::NewGuid());
+	virtual double ComputeResult(Dataflow::FContext& Context, double InA, double InB) const override;
+};
+
+/** Minimum ( Min(A, B) ) */
+USTRUCT()
+struct FDataflowMathMinimumNode : public FDataflowMathTwoInputsOperatorNode
+{
+	GENERATED_USTRUCT_BODY()
+	DATAFLOW_NODE_DEFINE_INTERNAL(FDataflowMathMinimumNode, "Minimum", DATAFLOW_MATH_NODES_CATEGORY, "")
+
+public:
+	FDataflowMathMinimumNode(const Dataflow::FNodeParameters& InParam, FGuid InGuid = FGuid::NewGuid());
+	virtual double ComputeResult(Dataflow::FContext& Context, double InA, double InB) const override;
+};
+
+/** Maximum ( Max(A, B) ) */
+USTRUCT()
+struct FDataflowMathMaximumNode : public FDataflowMathTwoInputsOperatorNode
+{
+	GENERATED_USTRUCT_BODY()
+	DATAFLOW_NODE_DEFINE_INTERNAL(FDataflowMathMaximumNode, "Maximum", DATAFLOW_MATH_NODES_CATEGORY, "")
+
+public:
+	FDataflowMathMaximumNode(const Dataflow::FNodeParameters& InParam, FGuid InGuid = FGuid::NewGuid());
+	virtual double ComputeResult(Dataflow::FContext& Context, double InA, double InB) const override;
+};
+
+/** 
+* Reciprocal( 1 / A )
+* if A is equal to 0, returns Fallback
+*/
+USTRUCT()
+struct FDataflowMathReciprocalNode : public FDataflowMathOneInputOperatorNode
+{
+	GENERATED_USTRUCT_BODY()
+	DATAFLOW_NODE_DEFINE_INTERNAL(FDataflowMathReciprocalNode, "Reciprocal", DATAFLOW_MATH_NODES_CATEGORY, "")
+
+	UPROPERTY(EditAnywhere, Category = SafeDivide, meta = (DataflowInput));
+	FDataflowNumericTypes Fallback;
+
+public:
+	FDataflowMathReciprocalNode(const Dataflow::FNodeParameters& InParam, FGuid InGuid = FGuid::NewGuid());
+	virtual double ComputeResult(Dataflow::FContext& Context, double InA) const override;
+};
+
+/** Square ( A * A ) */
+USTRUCT()
+struct FDataflowMathSquareNode : public FDataflowMathOneInputOperatorNode
+{
+	GENERATED_USTRUCT_BODY()
+	DATAFLOW_NODE_DEFINE_INTERNAL(FDataflowMathSquareNode, "Square", DATAFLOW_MATH_NODES_CATEGORY, "")
+
+public:
+	FDataflowMathSquareNode(const Dataflow::FNodeParameters& InParam, FGuid InGuid = FGuid::NewGuid());
+	virtual double ComputeResult(Dataflow::FContext& Context, double InA) const override;
+};
+
+/** Cube ( A * A * A ) */
+USTRUCT()
+struct FDataflowMathCubeNode : public FDataflowMathOneInputOperatorNode
+{
+	GENERATED_USTRUCT_BODY()
+	DATAFLOW_NODE_DEFINE_INTERNAL(FDataflowMathCubeNode, "Cube", DATAFLOW_MATH_NODES_CATEGORY, "")
+
+public:
+	FDataflowMathCubeNode(const Dataflow::FNodeParameters& InParam, FGuid InGuid = FGuid::NewGuid());
+	virtual double ComputeResult(Dataflow::FContext& Context, double InA) const override;
+};
+
+/** Square Root ( sqrt(A) ) */
+USTRUCT()
+struct FDataflowMathSquareRootNode : public FDataflowMathOneInputOperatorNode
+{
+	GENERATED_USTRUCT_BODY()
+	DATAFLOW_NODE_DEFINE_INTERNAL(FDataflowMathSquareRootNode, "SquareRoot", DATAFLOW_MATH_NODES_CATEGORY, "")
+
+public:
+	FDataflowMathSquareRootNode(const Dataflow::FNodeParameters& InParam, FGuid InGuid = FGuid::NewGuid());
+	virtual double ComputeResult(Dataflow::FContext& Context, double InA) const override;
+};
+
+/** 
+* Inverse Square Root ( 1 / sqrt(A) ) 
+* if A is equal to 0, returns Fallback
+*/
+USTRUCT()
+struct FDataflowMathInverseSquareRootNode : public FDataflowMathOneInputOperatorNode
+{
+	GENERATED_USTRUCT_BODY()
+	DATAFLOW_NODE_DEFINE_INTERNAL(FDataflowMathInverseSquareRootNode, "InverseSquareRoot", DATAFLOW_MATH_NODES_CATEGORY, "")
+
+	UPROPERTY(EditAnywhere, Category = SafeDivide, meta = (DataflowInput));
+	FDataflowNumericTypes Fallback;
+
+public:
+	FDataflowMathInverseSquareRootNode(const Dataflow::FNodeParameters& InParam, FGuid InGuid = FGuid::NewGuid());
+	virtual double ComputeResult(Dataflow::FContext& Context, double InA) const override;
+};
+
+/** Negate ( -A ) */
+USTRUCT()
+struct FDataflowMathNegateNode : public FDataflowMathOneInputOperatorNode
+{
+	GENERATED_USTRUCT_BODY()
+	DATAFLOW_NODE_DEFINE_INTERNAL(FDataflowMathNegateNode, "Negate", DATAFLOW_MATH_NODES_CATEGORY, "")
+
+	UPROPERTY(EditAnywhere, Category = SafeDivide, meta = (DataflowInput));
+	FDataflowNumericTypes Fallback;
+
+public:
+	FDataflowMathNegateNode(const Dataflow::FNodeParameters& InParam, FGuid InGuid = FGuid::NewGuid());
+	virtual double ComputeResult(Dataflow::FContext& Context, double InA) const override;
+};
+
+/** Absolute value ( |A| ) */
+USTRUCT()
+struct FDataflowMathAbsNode : public FDataflowMathOneInputOperatorNode
+{
+	GENERATED_USTRUCT_BODY()
+	DATAFLOW_NODE_DEFINE_INTERNAL(FDataflowMathAbsNode, "Abs", DATAFLOW_MATH_NODES_CATEGORY, "")
+
+	UPROPERTY(EditAnywhere, Category = SafeDivide, meta = (DataflowInput));
+	FDataflowNumericTypes Fallback;
+
+public:
+	FDataflowMathAbsNode(const Dataflow::FNodeParameters& InParam, FGuid InGuid = FGuid::NewGuid());
+	virtual double ComputeResult(Dataflow::FContext& Context, double InA) const override;
+};
+
+/** Floor ( 1.4 => 1.0 | 1.9 => 1.0 | -5.3 => -6.0 ) */
+USTRUCT()
+struct FDataflowMathFloorNode : public FDataflowMathOneInputOperatorNode
+{
+	GENERATED_USTRUCT_BODY()
+	DATAFLOW_NODE_DEFINE_INTERNAL(FDataflowMathFloorNode, "Floor", DATAFLOW_MATH_NODES_CATEGORY, "")
+
+public:
+	FDataflowMathFloorNode(const Dataflow::FNodeParameters& InParam, FGuid InGuid = FGuid::NewGuid());
+	virtual double ComputeResult(Dataflow::FContext& Context, double InA) const override;
+};
+
+/** Ceil ( 1.4 => 2.0 | 1.9 => 2.0 | -5.3 => -5.0) */
+USTRUCT()
+struct FDataflowMathCeilNode : public FDataflowMathOneInputOperatorNode
+{
+	GENERATED_USTRUCT_BODY()
+	DATAFLOW_NODE_DEFINE_INTERNAL(FDataflowMathCeilNode, "Ceil", DATAFLOW_MATH_NODES_CATEGORY, "")
+
+public:
+	FDataflowMathCeilNode(const Dataflow::FNodeParameters& InParam, FGuid InGuid = FGuid::NewGuid());
+	virtual double ComputeResult(Dataflow::FContext& Context, double InA) const override;
+};
+
+/** Round ( 1.4 => 1.0 | 1.9 => 2.0 | -5.3 => -5.0) */
+USTRUCT()
+struct FDataflowMathRoundNode : public FDataflowMathOneInputOperatorNode
+{
+	GENERATED_USTRUCT_BODY()
+	DATAFLOW_NODE_DEFINE_INTERNAL(FDataflowMathRoundNode, "Round", DATAFLOW_MATH_NODES_CATEGORY, "")
+
+public:
+	FDataflowMathRoundNode(const Dataflow::FNodeParameters& InParam, FGuid InGuid = FGuid::NewGuid());
+	virtual double ComputeResult(Dataflow::FContext& Context, double InA) const override;
+};
+
+/** Trunc ( 1.4 => 1.0 | 1.9 => 1.0 | -5.3 => -5.0) */
+USTRUCT()
+struct FDataflowMathTruncNode : public FDataflowMathOneInputOperatorNode
+{
+	GENERATED_USTRUCT_BODY()
+	DATAFLOW_NODE_DEFINE_INTERNAL(FDataflowMathTruncNode, "Trunc", DATAFLOW_MATH_NODES_CATEGORY, "")
+
+public:
+	FDataflowMathTruncNode(const Dataflow::FNodeParameters& InParam, FGuid InGuid = FGuid::NewGuid());
+	virtual double ComputeResult(Dataflow::FContext& Context, double InA) const override;
+};
+
+/** Frac ( 1.4 => 0.4 | 1.9 => 0.9 | -5.3 => 0.3 ) */
+USTRUCT()
+struct FDataflowMathFracNode : public FDataflowMathOneInputOperatorNode
+{
+	GENERATED_USTRUCT_BODY()
+	DATAFLOW_NODE_DEFINE_INTERNAL(FDataflowMathFracNode, "Frac", DATAFLOW_MATH_NODES_CATEGORY, "")
+
+public:
+	FDataflowMathFracNode(const Dataflow::FNodeParameters& InParam, FGuid InGuid = FGuid::NewGuid());
+	virtual double ComputeResult(Dataflow::FContext& Context, double InA) const override;
+};
+
+/** power ( A ^ B) */
+USTRUCT()
+struct FDataflowMathPowNode : public FDataflowMathTwoInputsOperatorNode
+{
+	GENERATED_USTRUCT_BODY()
+	DATAFLOW_NODE_DEFINE_INTERNAL(FDataflowMathPowNode, "Pow", DATAFLOW_MATH_NODES_CATEGORY, "")
+
+public:
+	FDataflowMathPowNode(const Dataflow::FNodeParameters& InParam, FGuid InGuid = FGuid::NewGuid());
+	virtual double ComputeResult(Dataflow::FContext& Context, double InA, double InB) const override;
+};
+
+/**
+* Log for a specific base ( Log[Base](A) ) 
+* If base is negative or zero returns 0
+*/
+USTRUCT()
+struct FDataflowMathLogXNode : public FDataflowMathOneInputOperatorNode
+{
+	GENERATED_USTRUCT_BODY()
+	DATAFLOW_NODE_DEFINE_INTERNAL(FDataflowMathLogXNode, "LogX", DATAFLOW_MATH_NODES_CATEGORY, "")
+
+	UPROPERTY(EditAnywhere, Category = SafeDivide, meta = (DataflowInput));
+	FDataflowNumericTypes Base;
+
+public:
+	FDataflowMathLogXNode(const Dataflow::FNodeParameters& InParam, FGuid InGuid = FGuid::NewGuid());
+	virtual double ComputeResult(Dataflow::FContext& Context, double InA) const override;
+};
+
+/** Natural log ( Log(A) ) */
+USTRUCT()
+struct FDataflowMathLogNode : public FDataflowMathOneInputOperatorNode
+{
+	GENERATED_USTRUCT_BODY()
+	DATAFLOW_NODE_DEFINE_INTERNAL(FDataflowMathLogNode, "Log", DATAFLOW_MATH_NODES_CATEGORY, "")
+
+public:
+	FDataflowMathLogNode(const Dataflow::FNodeParameters& InParam, FGuid InGuid = FGuid::NewGuid());
+	virtual double ComputeResult(Dataflow::FContext& Context, double InA) const override;
+};
+
+/** Exponential ( Exp(A) ) */
+USTRUCT()
+struct FDataflowMathExpNode : public FDataflowMathOneInputOperatorNode
+{
+	GENERATED_USTRUCT_BODY()
+	DATAFLOW_NODE_DEFINE_INTERNAL(FDataflowMathExpNode, "Exp", DATAFLOW_MATH_NODES_CATEGORY, "")
+
+public:
+	FDataflowMathExpNode(const Dataflow::FNodeParameters& InParam, FGuid InGuid = FGuid::NewGuid());
+	virtual double ComputeResult(Dataflow::FContext& Context, double InA) const override;
+};
+
+/** return -1, 0, +1 whether the input is respectively negative, zero or positive ( Sign(A) ) */
+USTRUCT()
+struct FDataflowMathSignNode : public FDataflowMathOneInputOperatorNode
+{
+	GENERATED_USTRUCT_BODY()
+	DATAFLOW_NODE_DEFINE_INTERNAL(FDataflowMathSignNode, "Sign", DATAFLOW_MATH_NODES_CATEGORY, "")
+
+public:
+	FDataflowMathSignNode(const Dataflow::FNodeParameters& InParam, FGuid InGuid = FGuid::NewGuid());
+	virtual double ComputeResult(Dataflow::FContext& Context, double InA) const override;
+};
+
+namespace Dataflow
+{
+	void RegisterDataflowMathNodes();
+}
+
+/*
+Nodes Left to be converted :
+	FMathConstantsDataflowNode);
+	FOneMinusDataflowNode);
+	
+	FFloatMathExpressionDataflowNode);
+	FMathExpressionDataflowNode);
+	
+	FClampDataflowNode);
+	FFitDataflowNode);
+	FEFitDataflowNode);
+	
+	FLerpDataflowNode);
+	FWrapDataflowNode);
+
+	// trig
+	FSinDataflowNode);
+	FArcSinDataflowNode);
+	FCosDataflowNode);
+	FArcCosDataflowNode);
+	FTanDataflowNode);
+	FArcTanDataflowNode);
+	FArcTan2DataflowNode);
+	FRadiansToDegreesDataflowNode);
+	FDegreesToRadiansDataflowNode);
+
+	// vectors - requires Vector any type ? 
+	FNormalizeToRangeDataflowNode);
+	FScaleVectorDataflowNode);
+	FDotProductDataflowNode);
+	FCrossProductDataflowNode);
+	FNormalizeDataflowNode);
+	FLengthDataflowNode);
+	FDistanceDataflowNode);
+	FIsNearlyZeroDataflowNode);
+
+	// random 
+	FRandomFloatDataflowNode);
+	FRandomFloatInRangeDataflowNode);
+	FRandomUnitVectorDataflowNode);
+	FRandomUnitVectorInConeDataflowNode);
+
+*/
\ No newline at end of file
diff --git a/Engine/Source/Runtime/Experimental/Dataflow/Core/Public/Dataflow/DataflowNodeFactory.h b/Engine/Source/Runtime/Experimental/Dataflow/Core/Public/Dataflow/DataflowNodeFactory.h
--- a/Engine/Source/Runtime/Experimental/Dataflow/Core/Public/Dataflow/DataflowNodeFactory.h
+++ b/Engine/Source/Runtime/Experimental/Dataflow/Core/Public/Dataflow/DataflowNodeFactory.h
@@ -58,7 +58,6 @@ namespace Dataflow
 		// All Maps indexed by TypeName
 		TMap<FName, FNewNodeFunction > ClassMap;			// [TypeName] -> NewNodeFunction
 		TMap<FName, FFactoryParameters > ParametersMap;		// [TypeName] -> Parameters
-		TMap<FName, FName > DisplayMap;						// [DisplayName] -> TypeName
 		TMap<FName, TArray<FName>> VersionMap;				// [TypeNameNoVersion] -> Array of TypeName(versions)
 
 		DATAFLOWCORE_API static FNodeFactory* Instance;
@@ -131,8 +130,6 @@ namespace Dataflow
 
 		DATAFLOWCORE_API static FString GetToolTipFromStruct(UScriptStruct* InStruct, const FName& InTypeName, const FName& InDisplayName);
 
-		DATAFLOWCORE_API FName TypeNameFromDisplayName(const FName& DisplayName);
-
 		template<class T> TSharedPtr<T> NewNode(FGraph& Graph, const FNewNodeParameters& Param)
 		{
 			return Graph.AddNode(new T(Param.Name, Param.Guid));
