Chunked transfer support.

Destination should be set to a suitably sized working buffer to receive data into, content length will be -1 (because it's unknown!). Sink gets called multiple times as and when data becomes available.

#rnx

diff --git a/Engine/Source/Runtime/Experimental/IoStore/HttpClient/Private/Loop.inl b/Engine/Source/Runtime/Experimental/IoStore/HttpClient/Private/Loop.inl
--- a/Engine/Source/Runtime/Experimental/IoStore/HttpClient/Private/Loop.inl
+++ b/Engine/Source/Runtime/Experimental/IoStore/HttpClient/Private/Loop.inl
@@ -177,12 +177,13 @@ static FOutcome DoRecvMessage(FActivity* Activity, FPeerType& Peer)
 		Internal.MessageLength - Internal.Offsets.Headers - 2 // "-2" trims off '\r\n' that signals end of headers
 	);
 
-	int32 Count = 2;
+	int32 Count = 3;
+	bool bChunked = false;
 	bool IsKeepAlive = true;
 	int32 ContentLength = -1;
 	EnumerateHeaders(
 		Headers,
-		[&ContentLength, &IsKeepAlive, &Count] (FAnsiStringView Name, FAnsiStringView Value)
+		[&ContentLength, &bChunked, &IsKeepAlive, &Count] (FAnsiStringView Name, FAnsiStringView Value)
 		{
 			// todo; may need smarter value handling; ;/, separated options & key-value pairs (ex. in rfc2068)
 
@@ -195,6 +196,12 @@ static FOutcome DoRecvMessage(FActivity* Activity, FPeerType& Peer)
 				Count--;
 			}
 
+			else if (Name.Equals("Transfer-Encoding", ESearchCase::IgnoreCase))
+			{
+				bChunked = Value.Equals("chunked", ESearchCase::IgnoreCase);
+				Count--;
+			}
+
 			else if (Name.Equals("Connection", ESearchCase::IgnoreCase))
 			{
 				IsKeepAlive = !Value.Equals("close");
@@ -207,13 +214,13 @@ static FOutcome DoRecvMessage(FActivity* Activity, FPeerType& Peer)
 
 	Activity->IsKeepAlive &= IsKeepAlive;
 
-	if (ContentLength < 0)
+	// Validate that the server's told us how and how much it will transmit
+	if (bChunked)
+	{
+		ContentLength = -1;
+	}
+	else if (ContentLength < 0)
 	{
-		if (ContentLength == -1)
-		{
-			// todo; query Transfer-Encoding for chunked
-		}
-
 		Activity_SetError(Activity, "Unknown content length value");
 		return FOutcome::Error(Activity->ErrorReason);
 	}
@@ -241,14 +248,20 @@ static FOutcome DoRecvMessage(FActivity* Activity, FPeerType& Peer)
 		// The user seems to have forgotten something. Let's help them along
 		if (int32 DestSize = int32(Dest.GetSize()); DestSize == 0)
 		{
-			Dest = FIoBuffer(ContentLength);
+			static const uint32 DefaultChunkSize = 4 << 10;
+			uint32 Size = bChunked ? DefaultChunkSize : ContentLength;
+			Dest = FIoBuffer(Size);
 		}
-		else if (DestSize < ContentLength)
+		else if (!bChunked && DestSize < ContentLength)
 		{
 			// todo: support piece-wise transfer of content (a la chunked).
 			Activity_SetError(Activity, "Destination buffer too small");
 			return FOutcome::Error(Activity->ErrorReason);
 		}
+		else if (enum { MinStreamBuf = 256 }; bChunked && DestSize < MinStreamBuf)
+		{
+			Dest = FIoBuffer(MinStreamBuf);
+		}
 	}
 
 	// Perhaps we have some of the content already?
@@ -273,9 +286,7 @@ static FOutcome DoRecvMessage(FActivity* Activity, FPeerType& Peer)
 
 	check(Activity->Dest != nullptr);
 
-	const bool bStreamed = Activity->Dest->GetSize() < ContentLength;
-
-	auto NextState = bStreamed ? FActivity::EState::RecvStream : FActivity::EState::RecvContent;
+	auto NextState = bChunked ? FActivity::EState::RecvStream : FActivity::EState::RecvContent;
 	Activity_ChangeState(Activity, NextState, AlreadyReceived);
 
 	// Copy any of the content we may have already received.
@@ -328,10 +339,198 @@ static FOutcome DoRecvContent(FActivity* Activity, FPeerType& Peer, int32& MaxRe
 }
 
 ////////////////////////////////////////////////////////////////////////////////
-static FOutcome DoRecvStream(FActivity*, FPeerType&, uint32)
+static FOutcome DoRecvStream(FActivity* Activity, FPeerType& Peer, int32& MaxRecvSize)
 {
-	check(false); // not yet implemented
-	return FOutcome::Error("Not supported");
+	auto RaiseCrLfError = [Activity] ()
+	{
+		Activity_SetError(Activity, "Expected CRLF chunk terminal");
+		return FOutcome::Error(Activity->ErrorReason);
+	};
+
+	auto SinkData = [Activity] (FMemoryView View)
+	{
+		if (View.GetSize() == 0)
+		{
+			return;
+		}
+
+		// Temporarily clamp IoBuffer so if the sink does GetView/GetSize() it
+		// represents actual content and not the underlying working buffer.
+		FIoBuffer& Dest = *(Activity->Dest);
+		FIoBuffer Slice(View, Dest);
+		Swap(Dest, Slice);
+		
+		FTicketStatus& SinkArg = *(FTicketStatus*)Activity;
+		Activity->Sink(SinkArg);
+
+		Swap(Dest, Slice);
+	};
+
+	auto Done = [Activity] ()
+	{
+#if IAS_HTTP_WITH_PERF
+			Activity->Stopwatch.RecvEnd();
+#endif
+			*(Activity->Dest) = FIoBuffer();
+			Activity_ChangeState(Activity, FActivity::EState::RecvDone);
+			return FOutcome::Ok();
+	};
+
+	enum { CrLfLength = 2 };
+
+	int32 Size = int32(Activity->StateParam);
+
+	// Trailing chunk data.
+	while (Size < 0)
+	{
+		Activity->StateParam = 0;
+		Size = -Size;
+
+		int32 RefillSize = FMath::Min<int32>(Size, int32(Activity->Dest->GetSize()));
+		FOutcome Outcome = DoRecvPeer(Activity, Peer, MaxRecvSize, RefillSize);
+		if (!Outcome.IsOk())
+		{
+			Activity->StateParam = 0 - Size;
+			return Outcome;
+		}
+
+		int32 Result = Outcome.GetResult();
+		check(Result > 0);
+
+		FMemoryView View = Activity->Dest->GetView();
+		if (Size > CrLfLength)
+		{
+			int32 SinkSize = Size - CrLfLength;
+			SinkSize = FMath::Min(Result, SinkSize);
+			SinkData(View.Left(SinkSize));
+			View = View.Mid(SinkSize);
+			Size -= SinkSize;
+			Result -= SinkSize;
+		}
+
+		const char* Cursor = (char*)View.GetData();
+		int32 CrLfError = 0;
+		if (int32 n = CrLfLength; Size == n && Result >= n)
+		{
+			CrLfError |= (Cursor[0] != '\r');
+			--Size; --Result;
+			++Cursor;
+		}
+		if (int32 n = CrLfLength - 1; Size == n && Result >= n)
+		{
+			CrLfError |= (Cursor[0] != '\n');
+			--Size; --Result;
+		}
+		if (CrLfError)
+		{
+			return RaiseCrLfError();
+		}
+
+		Size = Result - Size;
+		Activity->StateParam = Size;
+		check(Size <= 0);
+
+		// Have we found the trailer-section that follows last-chunk?
+		if (Size == 0 && Activity->NoContent)
+		{
+			return Done();
+		}
+	}
+
+	// Peel off chunks
+	for (FMemoryView View = Activity->Dest->GetView(); Size > 0;)
+	{
+		const char* Cursor = (char*)(View.GetData());
+
+		// Isolate chunk size
+ 		int32 ChunkSize = -1;
+		uint32 HeaderLength = 0;
+		for (; HeaderLength < uint32(Size - 1); ++HeaderLength)
+		{
+			// Detect CRLF.
+			if (Cursor[HeaderLength + 1] != '\n')
+			{
+				continue;
+			}
+
+			++HeaderLength;
+			if (Cursor[HeaderLength - 1] != '\r')
+			{
+				continue;
+			}
+			++HeaderLength;
+
+			ChunkSize = int32(CrudeToInt<16>(Cursor));
+			if (ChunkSize < 0)
+			{
+				Activity_SetError(Activity, "Unparsable chunk size");
+				return FOutcome::Error(Activity->ErrorReason);
+			}
+
+			break;
+		}
+
+		// Maybe we were not able to find a CRLF terminator and need more data
+		if (ChunkSize < 0)
+		{
+			FMutableMemoryView WriteView = Activity->Dest->GetMutableView();
+			std::memmove(WriteView.GetData(), Cursor, Size);
+			Activity->StateParam = Size;
+			break;
+		}
+
+		check(ChunkSize >= 0);
+		Size -= HeaderLength;
+
+		// Dispatch as much data as we can.
+		uint32 SinkSize = FMath::Min<uint32>(ChunkSize, uint32(Size));
+		SinkData(View.Mid(HeaderLength, SinkSize));
+		View = View.Mid(HeaderLength + SinkSize);
+
+		Activity->StateParam = (Size -= ChunkSize);
+		Activity->NoContent = (ChunkSize == 0);
+
+		// A CRLF follows a chunk's data
+		Cursor = (char*)View.GetData();
+		int32 CrLfError = 0;
+		CrLfError |= (Size >= (CrLfLength - 1)) && Cursor[0] != '\r';
+		CrLfError |= (Size >= (CrLfLength - 0)) && Cursor[1] != '\n';
+		if (CrLfError != 0)
+		{
+			return RaiseCrLfError();
+		}
+
+		// Can we do CRLF now?
+		if (Size >= CrLfLength)
+		{
+			// Have we found the trailer-section that follows last-chunk?
+			if (Activity->NoContent)
+			{
+				return Done();
+			}
+
+			Activity->StateParam = (Size -= CrLfLength);
+			View = View.Mid(CrLfLength);
+			continue;
+		}
+
+		Activity->StateParam -= CrLfLength;
+		check(int32(Activity->StateParam) < 0);
+		break;
+	}
+
+	// Refill
+	if (int32(Activity->StateParam) >= 0)
+	{
+		uint32 RefillSize = uint32(Activity->Dest->GetSize()) - Activity->StateParam;
+		FOutcome Outcome = DoRecvPeer(Activity, Peer, MaxRecvSize, RefillSize);
+		if (!Outcome.IsOk())
+		{
+			return Outcome;
+		}
+	}
+
+	return DoRecvStream(Activity, Peer, MaxRecvSize);
 }
 
 ////////////////////////////////////////////////////////////////////////////////
diff --git a/Engine/Source/Runtime/Experimental/IoStore/OnDemand/Private/Tool/Purl.cpp b/Engine/Source/Runtime/Experimental/IoStore/OnDemand/Private/Tool/Purl.cpp
--- a/Engine/Source/Runtime/Experimental/IoStore/OnDemand/Private/Tool/Purl.cpp
+++ b/Engine/Source/Runtime/Experimental/IoStore/OnDemand/Private/Tool/Purl.cpp
@@ -76,8 +76,9 @@ static int32 PurlCommandEntry(const FContext& Context)
 	}
 	FRequest Request = Loop.Request(AnsiMethod, AnsiUrl, &RequestParams);
 
-	FIoBuffer Dest;
-	Loop.Send(MoveTemp(Request), [&Dest] (const FTicketStatus& Status)
+	bool bChunked = false;
+	uint32 ContentSize = 0;
+	Loop.Send(MoveTemp(Request), [Dest=FIoBuffer(), &ContentSize, &bChunked] (const FTicketStatus& Status) mutable
 	{
 		if (Status.GetId() == FTicketStatus::EId::Response)
 		{
@@ -99,10 +100,17 @@ static int32 PurlCommandEntry(const FContext& Context)
 				return true;
 			});
 
+			bChunked = (Response.GetContentLength() == -1);
 			Response.SetDestination(&Dest);
 			return;
 		}
 
+		if (Status.GetId() == FTicketStatus::EId::Content)
+		{
+			ContentSize += uint32(Dest.GetSize());
+			return;
+		}
+
 		if (Status.GetId() == FTicketStatus::EId::Error)
 		{
 			const char* Reason = Status.GetErrorReason();
@@ -116,7 +124,7 @@ static int32 PurlCommandEntry(const FContext& Context)
 		FPlatformProcess::SleepNoStats(0.1f);
 	}
 
-	std::printf("Data: %u bytes\n", uint32(Dest.GetSize()));
+	std::printf("Data: %u bytes\n", ContentSize);
 
 	return 0;
 }
