Delineated a few paragraphs, tidied comments, tweaked an error message.

#rnx

diff --git a/Engine/Source/Runtime/Experimental/IoStore/HttpClient/Private/Loop.inl b/Engine/Source/Runtime/Experimental/IoStore/HttpClient/Private/Loop.inl
--- a/Engine/Source/Runtime/Experimental/IoStore/HttpClient/Private/Loop.inl
+++ b/Engine/Source/Runtime/Experimental/IoStore/HttpClient/Private/Loop.inl
@@ -187,9 +187,6 @@ static FOutcome DoRecvMessage(FActivity* Activity, FPeerType& Peer)
 		{
 			// todo; may need smarter value handling; ;/, separated options & key-value pairs (ex. in rfc2068)
 
-			// "Keep-Alive"			- deprecated
-			// "Transfer-Encoding"	- may be required later
-
 			if (Name.Equals("Content-Length", ESearchCase::IgnoreCase))
 			{
 				ContentLength = int32(CrudeToInt(Value));
@@ -221,12 +218,12 @@ static FOutcome DoRecvMessage(FActivity* Activity, FPeerType& Peer)
 	}
 	else if (ContentLength < 0)
 	{
-		Activity_SetError(Activity, "Unknown content length value");
+		Activity_SetError(Activity, "Missing/invalid Content-Length header");
 		return FOutcome::Error(Activity->ErrorReason);
 	}
 
 	// Call out to the sink to get a content destination
-	FIoBuffer* PriorDest = Activity->Dest; // to retain unioned Host ptr
+	FIoBuffer* PriorDest = Activity->Dest; // to retain unioned Host ptr (redirect uses it in sink)
 	Internal.Code = -1;
 	Internal.ContentLength = ContentLength;
 	{
@@ -273,6 +270,7 @@ static FOutcome DoRecvMessage(FActivity* Activity, FPeerType& Peer)
 		return FOutcome::Error(Activity->ErrorReason);
 	}
 
+	// HEAD methods
 	if (Activity->NoContent == 1)
 	{
 		if (AlreadyReceived)
@@ -284,6 +282,7 @@ static FOutcome DoRecvMessage(FActivity* Activity, FPeerType& Peer)
 		return FOutcome::Ok();
 	}
 
+	// We're all set to go and get content
 	check(Activity->Dest != nullptr);
 
 	auto NextState = bChunked ? FActivity::EState::RecvStream : FActivity::EState::RecvContent;
