Missing header and modeline footer

#rnx

diff --git a/Engine/Source/Runtime/Experimental/IoStore/HttpClient/Private/TestServer.py b/Engine/Source/Runtime/Experimental/IoStore/HttpClient/Private/TestServer.py
--- a/Engine/Source/Runtime/Experimental/IoStore/HttpClient/Private/TestServer.py
+++ b/Engine/Source/Runtime/Experimental/IoStore/HttpClient/Private/TestServer.py
@@ -1,3 +1,5 @@
+# Copyright Epic Games, Inc. All Rights Reserved.
+
 import io
 import os
 import time
@@ -387,3 +389,5 @@ def main():
 
 if __name__ == "__main__":
     main()
+
+# vim: et
