Separate out payload generation

#rnx

diff --git a/Engine/Source/Runtime/Experimental/IoStore/HttpClient/Private/TestServer.py b/Engine/Source/Runtime/Experimental/IoStore/HttpClient/Private/TestServer.py
--- a/Engine/Source/Runtime/Experimental/IoStore/HttpClient/Private/TestServer.py
+++ b/Engine/Source/Runtime/Experimental/IoStore/HttpClient/Private/TestServer.py
@@ -176,6 +176,20 @@ def tls_proxy_loop(httpd_port, root_cert, *server_pems):
 
 payload_data = random.randbytes(2 << 20)
 
+def _make_payload(size):
+    size = int(size)
+    if size <= 0:
+        size = int(random.random() * (8 << 10)) + 16
+    size = min(size, len(payload_data))
+
+    ret = payload_data[-size:]
+
+    ret_hash = 0x493
+    for c in ret:
+        ret_hash = ((ret_hash + c) * 0x493) & 0xFFffFFff
+
+    return ret, ret_hash
+
 def http_seed(handler, value=0):
     handler.send_response(200)
     handler.send_header("Content-Length", 0)
@@ -183,6 +197,8 @@ def http_seed(handler, value=0):
     random.seed(value)
 
 def http_data(handler, payload_size=0):
+    payload, payload_hash = _make_payload(payload_size)
+
     handler.send_response(200)
 
     mega_size = int(random.random() * (4 << 10))
@@ -195,18 +211,7 @@ def http_data(handler, payload_size=0):
         if mega_size <= 0:
             break
 
-    payload_size = int(payload_size)
-    if payload_size <= 0:
-        payload_size = int(random.random() * (8 << 10)) + 16
-    payload_size = min(payload_size, len(payload_data))
-
-    payload = payload_data[-payload_size:]
-
-    payload_hash = 0x493
-    for c in payload:
-        payload_hash = ((payload_hash + c) * 0x493) & 0xffffffff
     handler.send_header("X-TestServer-Hash", payload_hash)
-
     handler.send_header("Content-Length", len(payload))
     handler.send_header("Content-Type", "application/octet-stream")
     handler.end_headers()
