Endpoint to get actual test port.

#rnx

diff --git a/Engine/Source/Runtime/Experimental/IoStore/HttpClient/Private/TestServer.py b/Engine/Source/Runtime/Experimental/IoStore/HttpClient/Private/TestServer.py
--- a/Engine/Source/Runtime/Experimental/IoStore/HttpClient/Private/TestServer.py
+++ b/Engine/Source/Runtime/Experimental/IoStore/HttpClient/Private/TestServer.py
@@ -266,6 +266,14 @@ class Handler(http.server.BaseHTTPRequestHandler):
         if query := parts[-1].split("?"):
             parts[-1] = query[0]
 
+        if parts[0] == "port":
+            payload = str(self.server.server_address[1]).encode()
+            self.send_response(200)
+            self.send_header("Content-Length", len(payload))
+            self.end_headers()
+            self.wfile.write(payload)
+            return
+
         if parts[0] == "ca":
             ca_pem = self.server.ca_pem
             self.send_response(200)
