Chunked test endpoint.

#rnx

diff --git a/Engine/Source/Runtime/Experimental/IoStore/HttpClient/Private/TestServer.py b/Engine/Source/Runtime/Experimental/IoStore/HttpClient/Private/TestServer.py
--- a/Engine/Source/Runtime/Experimental/IoStore/HttpClient/Private/TestServer.py
+++ b/Engine/Source/Runtime/Experimental/IoStore/HttpClient/Private/TestServer.py
@@ -217,6 +217,32 @@ def http_data(handler, payload_size=0):
     handler.end_headers()
     handler.wfile.write(payload)
 
+def http_chunked(handler, payload_size=0):
+    payload, payload_hash = _make_payload(payload_size)
+    handler.send_response(200)
+    handler.send_header("Transfer-Encoding", "chunked")
+    handler.send_header("X-TestServer-Hash", payload_hash)
+    handler.send_header("X-TestServer-Size", len(payload))
+    handler.send_header("Content-Type", "application/octet-stream")
+    handler.end_headers()
+
+    max_chunk_size = int(random.random() * 1024) + 1
+    while True:
+        chunk_size = int(random.random() * max_chunk_size) + 1
+        piece = payload[:chunk_size]
+
+        header = b"%x" % len(piece)
+        if piece and piece[0] & 0b0100:
+            header = header.upper()
+        header += b"\r\n"
+
+        handler.wfile.write(header)
+        handler.wfile.write(piece)
+        handler.wfile.write(b"\r\n")
+        if not piece:
+            break
+        payload = payload[chunk_size:]
+
 def http_redirect(handler, style="abs", code=302, *dest):
     loc = "/" + "/".join(dest)
     if style.startswith("abs"):
@@ -295,6 +321,9 @@ class Handler(http.server.BaseHTTPRequestHandler):
         if parts[0] == "data":
             return http_data(self, *parts[1:])
 
+        if parts[0] == "chunked":
+            return http_chunked(self, *parts[1:])
+
         if parts[0] == "seed":
             return http_seed(self, *parts[1:])
 
