More logicial order of paragraphs.

#rnx

diff --git a/Engine/Source/Runtime/Experimental/IoStore/OnDemand/Private/Tool/Purl.cpp b/Engine/Source/Runtime/Experimental/IoStore/OnDemand/Private/Tool/Purl.cpp
--- a/Engine/Source/Runtime/Experimental/IoStore/OnDemand/Private/Tool/Purl.cpp
+++ b/Engine/Source/Runtime/Experimental/IoStore/OnDemand/Private/Tool/Purl.cpp
@@ -67,18 +67,9 @@ static int32 PurlCommandEntry(const FContext& Context)
 	Method = Method.ToUpper();
 	auto AnsiMethod = StringCast<ANSICHAR>(*Method);
 
-	FEventLoop Loop;
-
-	FEventLoop::FRequestParams RequestParams;
-	if (Context.Get<bool>(TEXT("-Redirect")))
-	{
-		RequestParams.bAutoRedirect = true;
-	}
-	FRequest Request = Loop.Request(AnsiMethod, AnsiUrl, &RequestParams);
-
 	bool bChunked = false;
 	uint32 ContentSize = 0;
-	Loop.Send(MoveTemp(Request), [Dest=FIoBuffer(), &ContentSize, &bChunked] (const FTicketStatus& Status) mutable
+	auto Sink = [Dest=FIoBuffer(), &ContentSize, &bChunked] (const FTicketStatus& Status) mutable
 	{
 		if (Status.GetId() == FTicketStatus::EId::Response)
 		{
@@ -117,7 +108,16 @@ static int32 PurlCommandEntry(const FContext& Context)
 			std::printf("ERROR: %s\n", Reason);
 			return;
 		}
-	});
+	};
+
+	FEventLoop Loop;
+	FEventLoop::FRequestParams RequestParams;
+	if (Context.Get<bool>(TEXT("-Redirect")))
+	{
+		RequestParams.bAutoRedirect = true;
+	}
+	FRequest Request = Loop.Request(AnsiMethod, AnsiUrl, &RequestParams);
+	Loop.Send(MoveTemp(Request), Sink);
 	
 	while (Loop.Tick(-1))
 	{
