Experimental journaling IAD cache

- store chunks in 32 MiB blocks
- track changes in a journal file
- purge cache based on usage and last access time
#rnx

diff --git a/Engine/Source/Runtime/Experimental/IoStore/OnDemand/Private/GenericHash.cpp b/Engine/Source/Runtime/Experimental/IoStore/OnDemand/Private/GenericHash.cpp
new file mode 100644
--- /dev/null
+++ b/Engine/Source/Runtime/Experimental/IoStore/OnDemand/Private/GenericHash.cpp
@@ -0,0 +1,33 @@
+// Copyright Epic Games, Inc. All Rights Reserved.
+
+#include "GenericHash.h"
+
+#include "Containers/ArrayView.h"
+#include "IO/IoHash.h"
+#include "Misc/StringBuilder.h"
+#include "String/BytesToHex.h"
+
+namespace UE::GenericHash
+{
+
+static_assert(sizeof(FIoHash) == 20);
+
+FStringBuilderBase& ToHex(FMemoryView Memory, FStringBuilderBase& Out)
+{
+	UE::String::BytesToHexLower(
+		MakeArrayView<const uint8>(
+			reinterpret_cast<const uint8*>(Memory.GetData()),
+			IntCastChecked<int32>(Memory.GetSize())),
+		Out);
+
+	return Out;
+}
+
+FString ToHex(FMemoryView Memory)
+{
+	FStringBuilderBase Sb;
+	ToHex(Memory, Sb);
+	return Sb.ToString();
+}
+
+} // namesapce UE::GenericHash
diff --git a/Engine/Source/Runtime/Experimental/IoStore/OnDemand/Private/GenericHash.h b/Engine/Source/Runtime/Experimental/IoStore/OnDemand/Private/GenericHash.h
new file mode 100644
--- /dev/null
+++ b/Engine/Source/Runtime/Experimental/IoStore/OnDemand/Private/GenericHash.h
@@ -0,0 +1,109 @@
+// Copyright Epic Games, Inc. All Rights Reserved.
+
+#pragma once
+
+#include "CoreTypes.h"
+#include "Containers/StringFwd.h"
+#include "Containers/UnrealString.h"
+#include "HAL/UnrealMemory.h"
+#include "Math/UnrealMathUtility.h"
+#include "Memory/MemoryView.h"
+
+struct FIoHash;
+
+namespace UE
+{
+
+namespace GenericHash
+{
+
+FString				ToHex(FMemoryView Memory);
+FStringBuilderBase& ToHex(FMemoryView Memory, FStringBuilderBase& Out);
+
+} // namespace UE::GenericHash
+
+template<SIZE_T BitCount>
+struct THash
+{
+	static constexpr SIZE_T Size	= BitCount >> 3;
+	using ByteArray					= uint8[Size];
+
+						THash() = default;
+	const uint8*		GetData() const { return Hash; }
+	SIZE_T				GetSize() const { return Size; }
+	FMemoryView			GetView() const { return FMemoryView(Hash, Size); }
+
+	static THash		From(FMemoryView Memory);
+	static THash		From(const uint8* Data, SIZE_T DataSize);
+	static THash		From(const FIoHash& IoHash);
+	static const		THash Zero;
+
+	friend inline bool operator==(const THash& A, const THash& B)
+	{
+		return FMemory::Memcmp(A.Hash, B.Hash, Size) == 0;
+	}
+
+	friend inline bool operator!=(const THash& A, const THash& B)
+	{
+		return FMemory::Memcmp(A.Hash, B.Hash, Size) != 0;
+	}
+
+	friend inline bool operator<(const THash& A, const THash& B)
+	{
+		return FMemory::Memcmp(A.Hash, B.Hash, Size) < 0;
+	}
+
+	friend inline uint32 GetTypeHash(const THash& H)
+	{
+		return *reinterpret_cast<const uint32*>(H.Hash);
+	}
+
+	friend inline FString LexToString(const THash& H)
+	{
+		return GenericHash::ToHex(H.GetView());
+	}
+
+	friend inline FStringBuilderBase& LexToString(const THash& H, FStringBuilderBase& Out)
+	{
+		return GenericHash::ToHex(H.GetView(), Out);
+	}
+
+private:
+	FMutableMemoryView GetMutableView() { return FMutableMemoryView(Hash, Size); }
+
+	alignas(uint32) ByteArray Hash{};
+};
+
+template<SIZE_T BitCount>
+const THash<BitCount> THash<BitCount>::Zero;
+
+template<SIZE_T BitCount>
+THash<BitCount> THash<BitCount>::From(FMemoryView Memory)
+{
+	THash<BitCount> Hash;
+	Hash.GetMutableView().CopyFrom(Memory.Left(THash<BitCount>::Size));
+	return Hash;
+}
+
+template<SIZE_T BitCount>
+THash<BitCount> THash<BitCount>::From(const uint8* Data, SIZE_T DataSize)
+{
+	return From(MakeMemoryView(Data, DataSize));
+}
+
+template<SIZE_T BitCount>
+THash<BitCount> THash<BitCount>::From(const FIoHash& IoHash)
+{
+	THash<BitCount> Hash;
+	FMemory::Memcpy(&Hash, &IoHash, FMath::Min(THash<BitCount>::Size, SIZE_T(20)));
+	return Hash;
+}
+
+using FHash32	= THash<32>;
+using FHash64	= THash<64>;
+using FHash96	= THash<96>;
+using FHash128	= THash<128>;
+using FHash160	= THash<160>;
+using FHash256	= THash<256>;
+
+} // namespace UE
diff --git a/Engine/Source/Runtime/Experimental/IoStore/OnDemand/Private/OnDemandInstallCache.cpp b/Engine/Source/Runtime/Experimental/IoStore/OnDemand/Private/OnDemandInstallCache.cpp
--- a/Engine/Source/Runtime/Experimental/IoStore/OnDemand/Private/OnDemandInstallCache.cpp
+++ b/Engine/Source/Runtime/Experimental/IoStore/OnDemand/Private/OnDemandInstallCache.cpp
@@ -6,16 +6,794 @@
 
 #include "Async/Mutex.h"
 #include "Async/UniqueLock.h"
+#include "Async/AsyncFileHandle.h"
+#include "Containers/UnrealString.h"
+#include "GenericHash.h"
 #include "HAL/FileManager.h"
 #include "HAL/PlatformFile.h"
 #include "HAL/PlatformFileManager.h"
 #include "IO/IoContainerHeader.h"
+#include "IO/IoChunkId.h"
 #include "IO/IoChunkEncoding.h"
+#include "Misc/DateTime.h"
 #include "Misc/PathViews.h"
+#include "Misc/ScopeExit.h"
+#include "Misc/StringBuilder.h"
 #include "Serialization/MemoryReader.h"
+#include "Serialization/LargeMemoryWriter.h"
+#include "Tasks/Task.h"
 
 namespace UE::IoStore
 {
+///////////////////////////////////////////////////////////////////////////////
+double ToKiB(uint64 Value)
+{
+	return double(Value) / 1024.0;
+}
+
+///////////////////////////////////////////////////////////////////////////////
+double ToMiB(uint64 Value)
+{
+	return double(Value) / 1024.0 / 1024.0;
+}
+
+///////////////////////////////////////////////////////////////////////////////
+using FSharedAsyncFileHandle	= TSharedPtr<IAsyncReadFileHandle, ESPMode::ThreadSafe>;
+using FWeakAsyncFileHandle		= TWeakPtr<IAsyncReadFileHandle, ESPMode::ThreadSafe>;
+using FUniqueFileHandle			= TUniquePtr<IFileHandle>;
+using FCasAddr					= FHash96;
+
+///////////////////////////////////////////////////////////////////////////////
+struct FCasBlockId
+{
+	FCasBlockId() = default;
+	explicit FCasBlockId(uint32 InId)
+		: Id(InId) { }
+
+	bool IsValid() const { return Id != 0; }
+
+	friend inline bool operator==(FCasBlockId LHS, FCasBlockId RHS)
+	{
+		return LHS.Id == RHS.Id;
+	}
+
+	friend inline uint32 GetTypeHash(FCasBlockId BlockId)
+	{
+		return GetTypeHash(BlockId.Id);
+	}
+
+	static const FCasBlockId Invalid;
+
+	uint32 Id = 0;
+};
+
+const FCasBlockId FCasBlockId::Invalid = FCasBlockId();
+
+///////////////////////////////////////////////////////////////////////////////
+struct FCasLocation
+{
+	bool IsValid() const { return BlockId.IsValid() && BlockOffset != MAX_uint32; }
+
+	static const FCasLocation Invalid;
+
+	FCasBlockId	BlockId;
+	uint32		BlockOffset = MAX_uint32; 
+};
+
+const FCasLocation FCasLocation::Invalid = FCasLocation();
+
+///////////////////////////////////////////////////////////////////////////////
+struct FCasBlockInfo
+{
+	uint64	FileSize = 0;
+	int64	LastAccess = 0;
+	uint32	RefCount = 0;
+};
+
+using FCasBlockInfoMap = TMap<FCasBlockId, FCasBlockInfo>;
+
+///////////////////////////////////////////////////////////////////////////////
+struct FCas
+{
+	using FLookup			= TMap<FCasAddr, FCasLocation>;
+	using FReadHandles		= TMap<FCasBlockId, FWeakAsyncFileHandle>;
+	using FLastAccess		= TMap<FCasBlockId, int64>;
+
+	FIoStatus				Initialize(FStringView Directory);
+	FCasLocation			FindChunk(const FIoHash& Hash);
+	FCasBlockId				CreateBlock();
+	FIoStatus				DeleteBlock(FCasBlockId BlockId, TArray<FCasAddr>& OutAddrs);
+	FString					GetBlockFilename(FCasBlockId BlockId) const;
+	FSharedAsyncFileHandle	OpenAsyncRead(FCasBlockId  BlockId);
+	FUniqueFileHandle		OpenWrite(FCasBlockId BlockId);
+	void					TrackAccess(FCasBlockId BlockId, int64 UtcTicks);
+	void					TrackAccess(FCasBlockId BlockId) { TrackAccess(BlockId, FDateTime::UtcNow().GetTicks()); }
+	uint64					GetBlockInfo(FCasBlockInfoMap& OutBlockInfo);
+	void					Compact();
+	FIoStatus				Verify(TArray<FCasAddr>& OutAddrs);
+
+	FStringView			RootDirectory;
+	FLookup				Lookup;
+	TSet<FCasBlockId>	BlockIds;
+	FLastAccess			LastAccess;
+	FReadHandles		ReadHandles;	
+	uint32				MaxBlockSize = 32 << 20; //TODO: Make configurable
+	FCasBlockId			CurrentBlock;
+	UE::FMutex			Mutex;
+};
+
+///////////////////////////////////////////////////////////////////////////////
+FIoStatus FCas::Initialize(FStringView Directory)
+{
+	RootDirectory = Directory;
+
+	Lookup.Empty();
+	BlockIds.Empty();
+	LastAccess.Empty();
+	CurrentBlock = FCasBlockId::Invalid;
+
+	TStringBuilder<256> Path;
+	FPathViews::Append(Path, RootDirectory, TEXT("blocks"));
+
+	IFileManager& Ifm = IFileManager::Get();
+	if (Ifm.DirectoryExists(Path.ToString()) == false)
+	{
+		const bool bTree = true;
+		if (Ifm.MakeDirectory(Path.ToString(), bTree) == false)
+		{
+			FIoStatus Status = FIoStatusBuilder(EIoErrorCode::WriteError)
+				<< TEXT("Failed to create directory '")
+				<< Path.ToString()
+				<< TEXT("'");
+			return Status;
+		}
+	}
+
+	return EIoErrorCode::Ok;
+};
+
+FCasLocation FCas::FindChunk(const FIoHash& Hash)
+{
+	const FCasAddr* Addr	= reinterpret_cast<const FCasAddr*>(&Hash);
+	const uint32 TypeHash	= GetTypeHash(*Addr);
+	{
+		UE::TUniqueLock Lock(Mutex);
+		if (const FCasLocation* Loc = Lookup.FindByHash(TypeHash, *Addr))
+		{
+			return *Loc;
+		}
+	}
+
+	return FCasLocation{};
+}
+
+FCasBlockId FCas::CreateBlock()
+{
+	IPlatformFile& Ipf = FPlatformFileManager::Get().GetPlatformFile();
+
+	UE::TUniqueLock Lock(Mutex);
+
+	uint32 BlockIdValue = 1;
+	for (;;)
+	{
+		if (BlockIdValue == 0)
+		{
+			UE_LOG(LogIoStoreOnDemand, Error, TEXT("Failed to generate block ID"));
+			return FCasBlockId::Invalid;
+		}
+
+		const FCasBlockId BlockId(BlockIdValue++);
+		if (BlockIds.Contains(BlockId))
+		{
+			continue;
+		}
+
+		const FString Filename = GetBlockFilename(BlockId);
+		if (Ipf.FileExists(*Filename))
+		{
+			UE_LOG(LogIoStoreOnDemand, Warning, TEXT("Unused CAS block id %u already exists on disk"), BlockId.Id);
+			continue;
+		}
+
+		BlockIds.Add(BlockId);
+		LastAccess.FindOrAdd(BlockId, FDateTime::UtcNow().GetTicks());
+		return BlockId;
+	}
+
+	return FCasBlockId::Invalid;
+}
+
+FIoStatus FCas::DeleteBlock(FCasBlockId BlockId, TArray<FCasAddr>& OutAddrs)
+{
+	UE::TUniqueLock Lock(Mutex);
+
+	IPlatformFile&	Ipf = FPlatformFileManager::Get().GetPlatformFile();
+	const FString	Filename = GetBlockFilename(BlockId);
+
+	UE_LOG(LogIoStoreOnDemand, Log, TEXT("Deleting CAS block '%s'"), *Filename);
+	if (Ipf.DeleteFile(*Filename) == false)
+	{
+		return FIoStatusBuilder(EIoErrorCode::WriteError)
+			<< TEXT("Failed to delete CAS block '")
+			<< Filename
+			<< TEXT("'");
+	}
+
+	BlockIds.Remove(BlockId);
+	for (auto It = Lookup.CreateIterator(); It; ++It)
+	{
+		if (It->Value.BlockId == BlockId)
+		{
+			OutAddrs.Add(It->Key);
+			It.RemoveCurrent();
+		}
+	}
+
+	return FIoStatus::Ok;
+}
+
+FString FCas::GetBlockFilename(FCasBlockId BlockId) const
+{
+	check(BlockId.IsValid());
+	const uint32 Id = NETWORK_ORDER32(BlockId.Id);
+	FString Hex;
+	BytesToHexLower(reinterpret_cast<const uint8*>(&Id), sizeof(int32), Hex);
+	TStringBuilder<256> Path;
+	FPathViews::Append(Path, RootDirectory, TEXT("blocks"), Hex);
+	Path << TEXT(".ucas");
+
+	return FString(Path.ToView());
+}
+
+FSharedAsyncFileHandle FCas::OpenAsyncRead(FCasBlockId BlockId)
+{
+	UE::TUniqueLock Lock(Mutex);
+
+	if (FWeakAsyncFileHandle* MaybeHandle = ReadHandles.Find(BlockId))
+	{
+		if (FSharedAsyncFileHandle Handle = MaybeHandle->Pin(); Handle.IsValid())
+		{
+			return Handle;
+		}
+	}
+
+	IPlatformFile&			Ipf = FPlatformFileManager::Get().GetPlatformFile();
+	const FString			Filename = GetBlockFilename(BlockId);
+	FSharedAsyncFileHandle	NewHandle(Ipf.OpenAsyncRead(*Filename));
+
+	if (NewHandle.IsValid())
+	{
+		ReadHandles.FindOrAdd(BlockId, NewHandle);
+	}
+
+	return NewHandle;
+}
+
+FUniqueFileHandle FCas::OpenWrite(FCasBlockId BlockId)
+{
+	IPlatformFile&	Ipf = FPlatformFileManager::Get().GetPlatformFile();
+	const FString	Filename = GetBlockFilename(BlockId);
+	const bool		bAppend = true;
+
+	return FUniqueFileHandle(Ipf.OpenWrite(*Filename, bAppend));
+}
+
+void FCas::TrackAccess(FCasBlockId BlockId, int64 UtcTicks)
+{
+	check(BlockId.IsValid());
+	UE::TUniqueLock Lock(Mutex);
+	LastAccess.FindOrAdd(BlockId, UtcTicks);
+}
+
+uint64 FCas::GetBlockInfo(FCasBlockInfoMap& OutBlockInfo)
+{
+	TStringBuilder<256> Path;
+	FPathViews::Append(Path, RootDirectory, TEXT("blocks"));
+
+	struct FDirectoryVisitor final
+		: public IPlatformFile::FDirectoryVisitor
+	{
+		FDirectoryVisitor(IPlatformFile& PlatformFile, FCasBlockInfoMap& InBlockInfo, FLastAccess&& Access)
+			: Ipf(PlatformFile)
+			, BlockInfo(InBlockInfo)
+			, LastAccess(MoveTemp(Access))
+		{ }
+		
+		virtual bool Visit(const TCHAR* FilenameOrDirectory, bool bIsDirectory) override
+		{
+			if (bIsDirectory)
+			{
+				return true;
+			}
+
+			const FStringView Filename(FilenameOrDirectory);
+			if (FPathViews::GetExtension(Filename) == TEXTVIEW("ucas") == false)
+			{
+				return true;
+			}
+
+			const int64			FileSize = Ipf.FileSize(FilenameOrDirectory);
+			const FStringView	IndexHex = FPathViews::GetBaseFilename(Filename);
+			const FCasBlockId	BlockId(FParse::HexNumber(WriteToString<128>(IndexHex).ToString()));
+
+			if (BlockId.IsValid() == false || FileSize < 0)
+			{
+				UE_LOG(LogIoStoreOnDemand, Warning, TEXT("Found invalid CAS block '%s', FileSize=%lld"),
+					FilenameOrDirectory, FileSize);
+				return true;
+			}
+
+			if (BlockInfo.Contains(BlockId))
+			{
+				UE_LOG(LogIoStoreOnDemand, Warning, TEXT("Found duplicate CAS block '%s'"), FilenameOrDirectory);
+				return true;
+			}
+
+			const int64* UtcTicks = LastAccess.Find(BlockId);
+
+			BlockInfo.Add(BlockId, FCasBlockInfo
+			{
+				.FileSize = uint64(FileSize),
+				.LastAccess = UtcTicks != nullptr ? *UtcTicks : 0
+			});
+			TotalSize += uint64(FileSize);
+
+			return true;
+		}
+
+		IPlatformFile&		Ipf;
+		FCasBlockInfoMap&	BlockInfo;
+		FLastAccess			LastAccess;
+		uint64				TotalSize = 0;
+	};
+
+	FLastAccess Access;
+	{
+		TUniqueLock Lock(Mutex);
+		Access = LastAccess;
+	}
+	IPlatformFile& Ipf = FPlatformFileManager::Get().GetPlatformFile();
+	FDirectoryVisitor Visitor(Ipf, OutBlockInfo, MoveTemp(Access));
+	Ipf.IterateDirectory(Path.ToString(), Visitor);
+
+	return Visitor.TotalSize;
+}
+
+void FCas::Compact()
+{
+	UE::TUniqueLock Lock(Mutex);
+	Lookup.Compact();
+	BlockIds.Compact();
+	ReadHandles.Compact();
+	LastAccess.Compact();
+}
+
+FIoStatus FCas::Verify(TArray<FCasAddr>& OutAddrs)
+{
+	FCasBlockInfoMap	BlockInfo;
+	const uint64		TotalSize = GetBlockInfo(BlockInfo);
+	uint64				TotalVerifiedBytes = 0;
+	FIoStatus			Status = FIoStatus::Ok;
+
+	for (auto BlockIt = BlockIds.CreateIterator(); BlockIt; ++BlockIt)
+	{
+		const FCasBlockId BlockId = *BlockIt;
+		if (const FCasBlockInfo* Info = BlockInfo.Find(BlockId))
+		{
+			TotalVerifiedBytes += Info->FileSize;
+			continue;
+		}
+
+		const FString Filename = GetBlockFilename(BlockId);
+		UE_LOG(LogIoStoreOnDemand, Warning, TEXT("Missing CAS block '%s'"), *Filename);
+
+		for (auto It = Lookup.CreateIterator(); It; ++It)
+		{
+			if (It->Value.BlockId == BlockId)
+			{
+				OutAddrs.Add(It->Key);
+				It.RemoveCurrent();
+			}
+		}
+
+		LastAccess.Remove(BlockId);
+		BlockIt.RemoveCurrent();
+		Status = EIoErrorCode::NotFound;
+	}
+
+	UE_LOG(LogIoStoreOnDemand, Log, TEXT("Verified %d CAS blocks of total %.2lf MiB"),
+		BlockIds.Num(), ToMiB(TotalVerifiedBytes));
+
+	IPlatformFile& Ipf = FPlatformFileManager::Get().GetPlatformFile();
+	for (const TPair<FCasBlockId, FCasBlockInfo>& Kv : BlockInfo)
+	{
+		const FCasBlockId BlockId = Kv.Key;
+		if (BlockIds.Contains(BlockId))
+		{
+			continue;
+		}
+
+		const FString Filename = GetBlockFilename(BlockId);
+		if (Ipf.DeleteFile(*Filename))
+		{
+			UE_LOG(LogIoStoreOnDemand, Warning, TEXT("Deleted orphaned CAS block '%s'"), *Filename);
+		}
+	}
+
+	return Status; 
+}
+
+///////////////////////////////////////////////////////////////////////////////
+struct FCasJournal
+{
+	enum class EVersion : uint32
+	{
+		Invalid	= 0,
+		Initial,
+
+		LatestPlusOne,
+		Latest = LatestPlusOne - 1
+	};
+
+	struct FHeader
+	{
+		static const inline uint8 MagicSequence[16] = {'C', 'A', 'S', 'J', 'O', 'U', 'R', 'N', 'A', 'L', 'H', 'E', 'A', 'D', 'E', 'R'};
+
+		bool		IsValid() const;
+
+		uint8		Magic[16] = {0};
+		EVersion	Version = EVersion::Invalid;
+		uint8		Pad[12] = {0};
+	};
+	static_assert(sizeof(FHeader) == 32);
+
+	struct FFooter
+	{
+		static const inline uint8 MagicSequence[16] = {'C', 'A', 'S', 'J', 'O', 'U', 'R', 'N', 'A', 'L', 'F', 'O', 'O', 'T', 'E', 'R'};
+
+		bool IsValid() const;
+
+		uint8 Magic[16] = {0};
+	};
+	static_assert(sizeof(FFooter) == 16);
+
+	struct FEntry
+	{
+		enum class EType : uint8
+		{
+			None = 0,
+			ChunkLocation,
+			BlockCreated,
+			BlockDeleted,
+			BlockAccess
+		};
+
+		struct FChunkLocation
+		{
+			EType			Type = EType::ChunkLocation;
+			uint8			Pad[3]= {0};
+			FCasLocation	CasLocation;
+			FCasAddr		CasAddr;
+		};
+		static_assert(sizeof(FChunkLocation) == 24);
+
+		struct FBlockOperation
+		{
+			EType		Type = EType::None;
+			uint8		Pad[3]= {0};
+			FCasBlockId	BlockId;
+			int64		UtcTicks = 0;
+			uint8		Pad1[8]= {0};
+		};
+		static_assert(sizeof(FBlockOperation) == 24);
+
+		union
+		{
+			FChunkLocation	ChunkLocation;
+			FBlockOperation	BlockOperation;
+		};
+
+		EType Type() const { return *reinterpret_cast<const EType*>(this); }
+	};
+	static_assert(sizeof(FEntry) == 24);
+
+	struct FTransaction
+	{
+		void			ChunkLocation(const FCasLocation& Location, const FCasAddr& Addr);
+		void			BlockCreated(FCasBlockId BlockId);
+		void			BlockDeleted(FCasBlockId BlockId);
+		void			BlockAccess(FCasBlockId BlockId, int64 UtcTicks);
+
+		FString			JournalFile;
+		TArray<FEntry>	Entries;
+	};
+
+	using FEntryHandler		= TFunction<void(const FEntry&)>;
+
+	static FIoStatus		Replay(const FString& JournalFile, FEntryHandler&& Handler);
+	static FIoStatus		Create(const FString& JournalFile);
+	static FTransaction		Begin(FString&& JournalFile);
+	static FIoStatus		Commit(FTransaction&& Transaction);
+};
+
+///////////////////////////////////////////////////////////////////////////////
+bool FCasJournal::FHeader::IsValid() const
+{
+	if (FMemory::Memcmp(&Magic, &FHeader::MagicSequence, sizeof(FHeader::MagicSequence)) != 0)
+	{
+		return false;
+	}
+
+	if (static_cast<uint32>(Version) > static_cast<uint32>(EVersion::Latest))
+	{
+		return false;
+	}
+
+	return true;
+}
+
+bool FCasJournal::FFooter::IsValid() const
+{
+	return FMemory::Memcmp(Magic, FFooter::MagicSequence, sizeof(FFooter::MagicSequence)) == 0;
+}
+
+FIoStatus FCasJournal::Replay(const FString& JournalFile, FEntryHandler&& Handler)
+{
+	IPlatformFile& Ipf = FPlatformFileManager::Get().GetPlatformFile();
+
+	if (Ipf.FileExists(*JournalFile) == false)
+	{
+		return EIoErrorCode::NotFound;
+	}
+
+	TUniquePtr<IFileHandle> FileHandle(Ipf.OpenRead(*JournalFile));
+	if (FileHandle.IsValid() == false)
+	{
+		return EIoErrorCode::FileNotOpen;
+	}
+
+	FHeader Header;
+	if ((FileHandle->Read(reinterpret_cast<uint8*>(&Header), sizeof(FHeader)) == false) || (Header.IsValid() == false))
+	{
+		return FIoStatusBuilder(EIoErrorCode::ReadError)
+			<< TEXT("Failed to validate journal header '")
+			<< JournalFile
+			<< TEXT("'");
+	}
+
+	const int64 FileSize	= FileHandle->Size();
+	const int64 EntryCount	= (FileSize - sizeof(FHeader) - sizeof(FFooter)) / sizeof(FEntry);
+
+	if (EntryCount < 0)
+	{
+		return EIoErrorCode::ReadError;
+	}
+
+	if (EntryCount == 0)
+	{
+		return EIoErrorCode::Ok;
+	}
+
+	const int64 FooterPos = FileSize - sizeof(FFooter);
+	if (FooterPos < 0)
+	{
+		return FIoStatusBuilder(EIoErrorCode::CorruptToc)
+			<< TEXT("Invalid journal footer");
+	}
+
+	const int64 EntriesPos = FileHandle->Tell();
+	if (FileHandle->Seek(FooterPos) == false)
+	{
+		return EIoErrorCode::ReadError;
+	}
+
+	FFooter Footer;
+	if ((FileHandle->Read(reinterpret_cast<uint8*>(&Footer), sizeof(FFooter)) == false) || (Footer.IsValid() == false))
+	{
+		return FIoStatusBuilder(EIoErrorCode::ReadError)
+			<< TEXT("Failed to validate journal footer '")
+			<< JournalFile
+			<< TEXT("'");
+	}
+
+	if (FileHandle->Seek(EntriesPos) == false)
+	{
+		return EIoErrorCode::ReadError;
+	}
+
+	TArray<FEntry> Entries;
+	Entries.SetNumZeroed(IntCastChecked<int32>(EntryCount));
+
+	if (FileHandle->Read(reinterpret_cast<uint8*>(Entries.GetData()), sizeof(FEntry) * EntryCount) == false)
+	{
+		return EIoErrorCode::ReadError;
+	}
+
+	UE_LOG(LogIoStoreOnDemand, Log, TEXT("Replaying %d CAS journal entries of total %.2lf KiB from '%s'"),
+		EntryCount, ToKiB(sizeof(FEntry) * EntryCount), *JournalFile);
+
+	for (const FEntry& Entry : Entries)
+	{
+		Handler(Entry);
+	}
+
+	return EIoErrorCode::Ok;
+}
+
+FIoStatus FCasJournal::Create(const FString& JournalFile)
+{
+	IPlatformFile& Ipf = FPlatformFileManager::Get().GetPlatformFile();
+	Ipf.DeleteFile(*JournalFile);
+
+	TUniquePtr<IFileHandle> FileHandle(Ipf.OpenWrite(*JournalFile));
+	if (FileHandle.IsValid() == false)
+	{
+		return EIoErrorCode::FileNotOpen;
+	}
+
+	FHeader Header;
+	FMemory::Memcpy(&Header.Magic, &FHeader::MagicSequence, sizeof(FHeader::MagicSequence));
+	Header.Version = EVersion::Latest;
+
+	if (FileHandle->Write(reinterpret_cast<uint8*>(&Header), sizeof(FHeader)) == false)
+	{
+		return EIoErrorCode::WriteError;
+	}
+
+	FFooter Footer;
+	FMemory::Memcpy(&Footer.Magic, &FFooter::MagicSequence, sizeof(FFooter::MagicSequence));
+	if (FileHandle->Write(reinterpret_cast<uint8*>(&Footer), sizeof(FFooter)) == false)
+	{
+		return EIoErrorCode::WriteError;
+	}
+
+	return EIoErrorCode::Ok;
+}
+
+FCasJournal::FTransaction FCasJournal::Begin(FString&& JournalFile)
+{
+	return FTransaction
+	{
+		.JournalFile = MoveTemp(JournalFile)
+	};
+}
+
+FIoStatus FCasJournal::Commit(FTransaction&& Transaction)
+{
+	if (Transaction.Entries.IsEmpty())
+	{
+		return EIoErrorCode::Ok;
+	}
+
+	IPlatformFile&	Ipf = FPlatformFileManager::Get().GetPlatformFile();
+
+	// Validate header and footer
+	{
+		TUniquePtr<IFileHandle> FileHandle(Ipf.OpenRead(*Transaction.JournalFile));
+		const int64				FileSize = FileHandle.IsValid() ? FileHandle->Size() : -1;
+
+		if (FileSize < sizeof(FHeader))
+		{
+			return FIoStatusBuilder(EIoErrorCode::FileOpenFailed)
+				<< TEXT("Failed to validate CAS journal file '")
+				<< Transaction.JournalFile
+				<< TEXT("'");
+		}
+
+		FHeader Header;
+		if ((FileHandle->Read(reinterpret_cast<uint8*>(&Header), sizeof(FHeader)) == false) || (Header.IsValid() == false))
+		{
+			return FIoStatusBuilder(EIoErrorCode::ReadError)
+				<< TEXT("Failed to validate CAS journal header '")
+				<< Transaction.JournalFile
+				<< TEXT("'");
+		}
+
+		const int64 FooterPos = FileSize - sizeof(FFooter);
+		if (FileHandle->Seek(FooterPos) == false)
+		{
+			return FIoStatusBuilder(EIoErrorCode::ReadError)
+				<< TEXT("Failed to validate CAS journal footer '")
+				<< Transaction.JournalFile
+				<< TEXT("'");
+		}
+
+		FFooter Footer;
+		if ((FileHandle->Read(reinterpret_cast<uint8*>(&Footer), sizeof(FFooter)) == false) || (Footer.IsValid() == false))
+		{
+			return FIoStatusBuilder(EIoErrorCode::ReadError)
+				<< TEXT("Failed to validate CAS journal footer '")
+				<< Transaction.JournalFile
+				<< TEXT("'");
+		}
+	}
+
+	// Append entries
+	{
+		const bool				bAppend = true;
+		TUniquePtr<IFileHandle> FileHandle(Ipf.OpenWrite(*Transaction.JournalFile, bAppend));
+		const int64				FileSize	= FileHandle.IsValid() ? FileHandle->Size() : -1;
+		const int64				EntriesPos	= FileSize > 0 ? FileSize - sizeof(FFooter) : -1;
+
+		if ((EntriesPos < 0) || (FileHandle->Seek(EntriesPos) == false))
+		{
+			return FIoStatusBuilder(EIoErrorCode::FileOpenFailed)
+				<< TEXT("Failed to open CAS journal '")
+				<< Transaction.JournalFile
+				<< TEXT("'");
+		}
+
+		const int64 TotalEntrySize = Transaction.Entries.Num() * sizeof(FEntry);
+		if (FileHandle->Write(
+			reinterpret_cast<const uint8*>(Transaction.Entries.GetData()),
+			TotalEntrySize) == false)
+		{
+			return FIoStatusBuilder(EIoErrorCode::WriteError)
+				<< TEXT("Failed to write CAS journal entries to '")
+				<< Transaction.JournalFile
+				<< TEXT("'");
+		}
+
+		FFooter Footer;
+		FMemory::Memcpy(&Footer.Magic, &FFooter::MagicSequence, sizeof(FFooter::MagicSequence));
+		if (FileHandle->Write(reinterpret_cast<uint8*>(&Footer), sizeof(FFooter)) == false)
+		{
+			return FIoStatusBuilder(EIoErrorCode::WriteError)
+				<< TEXT("Failed to write CAS journal footer to '")
+				<< Transaction.JournalFile
+				<< TEXT("'");
+		}
+
+		if (FileHandle->Flush() == false)
+		{
+			return EIoErrorCode::WriteError;
+		}
+
+		UE_LOG(LogIoStoreOnDemand, Log, TEXT("Committed %d CAS journal entries of total %.2lf KiB to '%s'"),
+			Transaction.Entries.Num(), ToKiB(TotalEntrySize), *Transaction.JournalFile);
+
+		return EIoErrorCode::Ok;
+	}
+}
+
+void FCasJournal::FTransaction::ChunkLocation(const FCasLocation& Location, const FCasAddr& Addr)
+{
+	Entries.AddZeroed_GetRef().ChunkLocation = FEntry::FChunkLocation
+	{
+		.CasLocation	= Location,
+		.CasAddr		= Addr
+	};
+}
+
+void FCasJournal::FTransaction::BlockCreated(FCasBlockId BlockId)
+{
+	Entries.AddZeroed_GetRef().BlockOperation = FEntry::FBlockOperation
+	{
+		.Type		= FEntry::EType::BlockCreated,
+		.BlockId	= BlockId,
+		.UtcTicks	= FDateTime::UtcNow().GetTicks()
+	};
+}
+
+void FCasJournal::FTransaction::BlockDeleted(FCasBlockId BlockId)
+{
+	Entries.AddZeroed_GetRef().BlockOperation = FEntry::FBlockOperation
+	{
+		.Type		= FEntry::EType::BlockDeleted,
+		.BlockId	= BlockId,
+		.UtcTicks	= FDateTime::UtcNow().GetTicks()
+	};
+}
+
+void FCasJournal::FTransaction::BlockAccess(FCasBlockId BlockId, int64 UtcTicks)
+{
+	Entries.AddZeroed_GetRef().BlockOperation = FEntry::FBlockOperation
+	{
+		.Type		= FEntry::EType::BlockAccess,
+		.BlockId	= BlockId,
+		.UtcTicks	= UtcTicks
+	};
+}
 
 ///////////////////////////////////////////////////////////////////////////////
 class FOnDemandInstallCache final
@@ -47,11 +825,53 @@ class FOnDemandInstallCache final
 
 		FIoRequestImpl*		DispatcherRequest;
 		FOnDemandChunkInfo	ChunkInfo;
-		FIoOffsetAndLength 	ChunkRange;
+		FIoOffsetAndLength	ChunkRange;
 		FIoBuffer			EncodedChunk;
 		uint64				RawSize;
 	};
 
+	struct FPendingChunks
+	{
+		static constexpr uint64 MaxPendingBytes = 4ull << 20;
+
+		bool IsEmpty() const
+		{
+			check(Chunks.Num() == ChunkHashes.Num());
+			return TotalSize == 0 && Chunks.IsEmpty() && ChunkHashes.IsEmpty();
+		}
+
+		void Append(FIoBuffer&& Chunk, const FIoHash& ChunkHash)
+		{
+			check(Chunks.Num() == ChunkHashes.Num());
+			TotalSize += Chunk.GetSize();
+			ChunkHashes.Add(ChunkHash);
+			Chunks.Add(MoveTemp(Chunk));
+		}
+
+		FIoBuffer Pop(FIoHash& OutChunkHash)
+		{
+			check(Chunks.Num() == ChunkHashes.Num());
+			check(Chunks.IsEmpty() == false);
+			FIoBuffer Chunk = Chunks.Pop(EAllowShrinking::No);
+			TotalSize		= TotalSize - Chunk.GetSize();
+			OutChunkHash	= ChunkHashes.Pop(EAllowShrinking::No);
+			return Chunk;
+		}
+
+		void Reset()
+		{
+			Chunks.Reset();
+			ChunkHashes.Reset();
+			TotalSize = 0;
+		}
+
+		TArray<FIoBuffer>	Chunks;
+		TArray<FIoHash>		ChunkHashes;
+		uint64				TotalSize = 0;
+	};
+
+	using FUniquePendingChunks = TUniquePtr<FPendingChunks>;
+
 public:
 	FOnDemandInstallCache(const FOnDemandInstallCacheConfig& Config, FOnDemandIoStore& IoStore);
 	virtual ~FOnDemandInstallCache();
@@ -68,27 +888,131 @@ public:
 	virtual TIoStatusOr<FIoMappedRegion> OpenMapped(const FIoChunkId& ChunkId, const FIoReadOptions& Options) override;
 
 	// IOnDemandInstallCache
-	virtual bool ContainsChunk(const FIoHash& Hash) override;
-	virtual FIoStatus PutChunk(FIoBuffer&& Chunk, const FIoHash& Hash) override;
+	virtual bool				IsChunkCached(const FIoHash& ChunkHash) override;
+	virtual FIoStatus			PutChunk(FIoBuffer&& Chunk, const FIoHash& ChunkHash) override;
+	virtual FIoStatus			Purge(TMap<FIoHash, uint64>&& ChunksToIntall) override;
+	virtual FIoStatus			Flush() override;
 
 private:
 	bool						Resolve(FIoRequestImpl* Request);
-	void						CompleteRequest(TUniquePtr<FChunkRequest>&& ChunkRequest);
-	FIoStatus					WriteChunk(FIoBuffer Chunk, const FIoHash& ExpectedHash);
-	void						GetChunkFilename(const FIoHash& Hash, FStringBuilderBase& OutPath);
+	void						CompleteRequest(FChunkRequest& ChunkRequest);
+	FIoStatus					FlushPendingChunks(FPendingChunks& Block);
+	FString						GetJournalFilename() const { return CacheDirectory / TEXT("cas.jrn"); }
 
 	FOnDemandIoStore&		IoStore;
+	FString					CacheDirectory;
+	FCas					Cas;
+	FUniquePendingChunks	PendingChunks;
 	FSharedBackendContext	BackendContext;
 	FIoRequestList			CompletedRequests;
 	UE::FMutex				Mutex;
-	FString					CacheDirectory;
+	uint64					MaxCacheSize;
 };
 
 ///////////////////////////////////////////////////////////////////////////////
 FOnDemandInstallCache::FOnDemandInstallCache(const FOnDemandInstallCacheConfig& Config, FOnDemandIoStore& InIoStore)
 	: IoStore(InIoStore)
 	, CacheDirectory(Config.RootDirectory)
+	, MaxCacheSize(Config.DiskQuota)
 {
+	UE_LOG(LogIoStoreOnDemand, Log, TEXT("Initializing install cache, MaxCacheSize=%.2lf MiB"),
+		ToMiB(MaxCacheSize));
+
+	FIoStatus Status = Cas.Initialize(CacheDirectory);
+	if (Status.IsOk() == false)
+	{
+		UE_LOG(LogIoStoreOnDemand, Error, TEXT("Failed to initialize install cache"));
+		return;
+	}
+
+	// Replay the journal to get the current state 
+	// TODO: Purge the journal or create snapshot when the journal gets too big
+	const FString JournalFile = GetJournalFilename();
+	Status = FCasJournal::Replay(JournalFile, [this](const FCasJournal::FEntry& JournalEntry)
+	{
+		switch(JournalEntry.Type())
+		{
+		case FCasJournal::FEntry::EType::ChunkLocation:
+		{
+			const FCasJournal::FEntry::FChunkLocation& ChunkLocation = JournalEntry.ChunkLocation;
+			if (ChunkLocation.CasLocation.IsValid())
+			{
+				Cas.Lookup.FindOrAdd(ChunkLocation.CasAddr, ChunkLocation.CasLocation);
+			}
+			else
+			{
+				Cas.Lookup.Remove(ChunkLocation.CasAddr);
+			}
+			break;
+		}
+		case FCasJournal::FEntry::EType::BlockCreated:
+		{
+			const FCasJournal::FEntry::FBlockOperation& Op = JournalEntry.BlockOperation;
+			Cas.CurrentBlock = Op.BlockId;
+			Cas.BlockIds.Add(Op.BlockId);
+			break;
+		}
+		case FCasJournal::FEntry::EType::BlockDeleted:
+		{
+			const FCasJournal::FEntry::FBlockOperation& Op = JournalEntry.BlockOperation;
+			Cas.BlockIds.Remove(Op.BlockId);
+			if (Cas.CurrentBlock == Op.BlockId)
+			{
+				Cas.CurrentBlock = FCasBlockId::Invalid;
+			}
+			break;
+		}
+		case FCasJournal::FEntry::EType::BlockAccess:
+		{
+			const FCasJournal::FEntry::FBlockOperation& Op = JournalEntry.BlockOperation;
+			Cas.TrackAccess(Op.BlockId, Op.UtcTicks);
+			break;
+		}
+		};
+	});
+
+	// Verify the current state with the cached content on disk 
+	// TODO: Add checksums etc
+	TArray<FCasAddr> RemovedChunks;
+	if (FIoStatus Verify = Cas.Verify(RemovedChunks); !Verify.IsOk())
+	{
+		// Try to recover if the CAS blocks on disk doesn't match
+		FCasJournal::FTransaction Transaction = FCasJournal::Begin(GetJournalFilename());
+		for (const FCasAddr& Addr : RemovedChunks)
+		{
+			Transaction.ChunkLocation(FCasLocation::Invalid, Addr);
+		}
+		Status = FCasJournal::Commit(MoveTemp(Transaction));
+	}
+
+	Cas.Compact();
+
+	if (Status.IsOk() == false)
+	{
+		if (Status.GetErrorCode() != EIoErrorCode::NotFound)
+		{
+			UE_LOG(LogIoStoreOnDemand, Warning, TEXT("Failed to replay install cache journal file '%s, reason '%s'"),
+				*Status.ToString(), *JournalFile);
+
+			UE_LOG(LogIoStoreOnDemand, Log, TEXT("Deleting installed content and reinitializing cache"));
+			IFileManager::Get().DeleteDirectory(*CacheDirectory);
+			if (Status = Cas.Initialize(CacheDirectory); Status.IsOk() == false)
+			{
+				UE_LOG(LogIoStoreOnDemand, Error, TEXT("Failed to initialize install cache, reason '%s'"),
+					*Status.ToString());
+				return;
+			}
+		}
+
+		if (Status = FCasJournal::Create(JournalFile); Status.IsOk())
+		{
+			UE_LOG(LogIoStoreOnDemand, Log, TEXT("Created CAS journal '%s'"), *JournalFile);
+		}
+		else
+		{
+			UE_LOG(LogIoStoreOnDemand, Error, TEXT("Failed to create CAS journal '%s'"), *JournalFile);
+		}
+	}
 }
 
 FOnDemandInstallCache::~FOnDemandInstallCache()
@@ -102,6 +1026,18 @@ void FOnDemandInstallCache::Initialize(FSharedBackendContextRef Context)
 
 void FOnDemandInstallCache::Shutdown()
 {
+	FCas::FLastAccess LastAccess;
+	{
+		TUniqueLock Lock(Cas.Mutex);
+		LastAccess = MoveTemp(Cas.LastAccess);
+	}
+
+	FCasJournal::FTransaction Transaction = FCasJournal::Begin(GetJournalFilename());
+	for (const TPair<FCasBlockId, int64>& Kv : LastAccess)
+	{
+		Transaction.BlockAccess(Kv.Key, Kv.Value);
+	}
+	FCasJournal::Commit(MoveTemp(Transaction));
 }
 
 void FOnDemandInstallCache::ResolveIoRequests(FIoRequestList Requests, FIoRequestList& OutUnresolved)
@@ -164,6 +1100,12 @@ bool FOnDemandInstallCache::Resolve(FIoRequestImpl* Request)
 		return false;
 	}
 
+	const FCasLocation CasLoc = Cas.FindChunk(ChunkInfo.Hash());
+	if (CasLoc.IsValid() == false)
+	{
+		return false;
+	}
+
 	const uint64 RequestSize = FMath::Min<uint64>(
 		Request->Options.GetSize(),
 		ChunkInfo.RawSize() - Request->Options.GetOffset());
@@ -181,132 +1123,291 @@ bool FOnDemandInstallCache::Resolve(FIoRequestImpl* Request)
 		return false;
 	}
 
-	TStringBuilder<256> Filename;
-	GetChunkFilename(ChunkInfo.Hash(), Filename);
-
-	IPlatformFile& Ipf = FPlatformFileManager::Get().GetPlatformFile();
-	TUniquePtr<IFileHandle> FileHandle(Ipf.OpenRead(Filename.ToString()));
+	FSharedAsyncFileHandle FileHandle = Cas.OpenAsyncRead(CasLoc.BlockId);
 	if (FileHandle.IsValid() == false)
 	{
-		return false;
-	}
-
-	const FIoOffsetAndLength Range = ChunkRange.ConsumeValueOrDie();
-	if (!FileHandle->Seek(IntCastChecked<int64>(Range.GetOffset())))
-	{
-		UE_LOG(LogIoStoreOnDemand, Error, TEXT("Seek file failed"));
-		return false;
+		UE_LOG(LogIoStoreOnDemand, Error, TEXT("Failed to CAS block file for reading"));
 	}
 
-	TUniquePtr<FChunkRequest> ChunkRequest = MakeUnique<FChunkRequest>(
+	TSharedPtr<FChunkRequest> ChunkRequest = MakeShared<FChunkRequest>(
 		Request,
 		MoveTemp(ChunkInfo),
-		Range,
+		ChunkRange.ConsumeValueOrDie(),
 		RequestSize);
 
-	FIoBuffer& EncodedChunk = ChunkRequest->EncodedChunk;
-	if (!FileHandle->Read(reinterpret_cast<uint8*>(EncodedChunk.GetData()), IntCastChecked<int64>(EncodedChunk.GetSize())))
+	FIoBuffer EncodedChunk		= ChunkRequest->EncodedChunk;
+	FAsyncFileCallBack Callback = [this, ChunkRequest, FileHandle](bool bWasCanceled, IAsyncReadRequest* ReadRequest) 
 	{
-		UE_LOG(LogIoStoreOnDemand, Error, TEXT("Failed to read '%s'"), Filename.ToString());
-		return false;
-	}
+		UE::Tasks::Launch(
+			UE_SOURCE_LOCATION,
+			[this, ChunkRequest, FileHandle, ReadRequest = TUniquePtr<IAsyncReadRequest>(ReadRequest), bWasCanceled]
+			{
+				if (bWasCanceled)
+				{
+					ChunkRequest->EncodedChunk = FIoBuffer();
+				}
+				CompleteRequest(*ChunkRequest);
+			});
+	};
 
-	UE::Tasks::Launch(UE_SOURCE_LOCATION, [this, ChunkRequest = MoveTemp(ChunkRequest)]() mutable
-	{
-		CompleteRequest(MoveTemp(ChunkRequest));
-	});
+	Cas.TrackAccess(CasLoc.BlockId);
+	IAsyncReadRequest* ReadRequest = FileHandle->ReadRequest(
+		CasLoc.BlockOffset + ChunkRequest->ChunkRange.GetOffset(),
+		ChunkRequest->ChunkRange.GetLength(),
+		EAsyncIOPriorityAndFlags::AIOP_BelowNormal,
+		&Callback,
+		EncodedChunk.GetData());
 
-	return true;
+	return ReadRequest != nullptr;
 }
 
-bool FOnDemandInstallCache::ContainsChunk(const FIoHash& Hash)
+bool FOnDemandInstallCache::IsChunkCached(const FIoHash& ChunkHash)
 {
-	TStringBuilder<256> Filename;
-	GetChunkFilename(Hash, Filename);
-	return IFileManager::Get().FileSize(*Filename) > 0;
+	const FCasLocation Loc = Cas.FindChunk(ChunkHash);
+	return Loc.IsValid();
 }
 
-FIoStatus FOnDemandInstallCache::PutChunk(FIoBuffer&& Chunk, const FIoHash& Hash)
+FIoStatus FOnDemandInstallCache::PutChunk(FIoBuffer&& Chunk, const FIoHash& ChunkHash)
 {
-	return WriteChunk(Chunk, Hash);
+	if (PendingChunks.IsValid() == false)
+	{
+		PendingChunks = MakeUnique<FPendingChunks>();
+	}
+
+	if (PendingChunks->TotalSize > FPendingChunks::MaxPendingBytes)
+	{
+		if (FIoStatus Status = FlushPendingChunks(*PendingChunks); Status.IsOk() == false)
+		{
+			return Status;
+		}
+		check(PendingChunks->IsEmpty());
+	}
+
+	PendingChunks->Append(MoveTemp(Chunk), ChunkHash);
+	return FIoStatus::Ok;;
 }
 
-FIoStatus FOnDemandInstallCache::WriteChunk(FIoBuffer Chunk, const FIoHash& ExpectedHash)
+FIoStatus FOnDemandInstallCache::Purge(TMap<FIoHash, uint64>&& ChunksToIntall) 
 {
-	const FIoHash ChunkHash = FIoHash::HashBuffer(Chunk.GetView());
+	FCasBlockInfoMap	BlockInfo;
+	const uint64		TotalCachedBytes = Cas.GetBlockInfo(BlockInfo);
+	const int64			AvailableBytes = MaxCacheSize - TotalCachedBytes;
+	uint64				TotalUncachedBytes = 0;
 
-	if (ChunkHash != ExpectedHash)
+	for (const TPair<FIoHash, uint64>& Kv : ChunksToIntall)
 	{
-		return FIoStatus(EIoErrorCode::ReadError, TEXTVIEW("Hash mismatch"));
+		if (FCasLocation Loc = Cas.FindChunk(Kv.Key); Loc.IsValid())
+		{
+			BlockInfo.FindOrAdd(Loc.BlockId).RefCount++;
+		}
+		else
+		{
+			TotalUncachedBytes += Kv.Value;
+		}
 	}
 
-	TStringBuilder<256> Filename;
-	GetChunkFilename(ChunkHash, Filename);
-	const FString Directory = FString(FPathViews::GetPath(Filename));
+	if (AvailableBytes >= IntCastChecked<int64>(TotalUncachedBytes))
+	{
+		UE_LOG(LogIoStoreOnDemand, Log, TEXT("Skipping cache purge, MaxCacheSize=%.2lf MiB, CacheSize=%.2lf MiB, UncachedSize=%.2lf MiB"),
+			ToMiB(MaxCacheSize), ToMiB(TotalCachedBytes), ToMiB(TotalUncachedBytes));
+		return FIoStatus::Ok;
+	}
 
-	const bool bTree = true;
-	if (IFileManager& Ifm = IFileManager::Get(); !Ifm.MakeDirectory(*Directory, bTree))
+	UE_LOG(LogIoStoreOnDemand, Log, TEXT("Purging install cache, MaxCacheSize=%.2lf MiB, CacheSize=%.2lf MiB, UncachedSize=%.2lf MiB"),
+		ToMiB(MaxCacheSize), ToMiB(TotalCachedBytes), ToMiB(TotalUncachedBytes));
+
+	TSet<FCasBlockId>					ReferencedBlocks;
+	TArray<FSharedOnDemandContainer>	MountedContainers = IoStore.GetMountedContainers();
+
+	//TODO: Compute fragmentation metric and redownload chunks when this number gets too high
+	for (const FSharedOnDemandContainer& Container : MountedContainers)
 	{
-		return FIoStatusBuilder(EIoErrorCode::WriteError)
-			<< TEXT("Failed to create directory '")
-			<< Directory
-			<< TEXT("'");
+		for (const TPair<FIoChunkId, FOnDemandChunkEntry>& Kv : Container->ChunkEntries)
+		{
+			if (FCasLocation Loc = Cas.FindChunk(Kv.Value.Hash); Loc.IsValid())
+			{
+				BlockInfo.FindOrAdd(Loc.BlockId).RefCount++;
+			}
+		}
 	}
 
-	IPlatformFile& Ipf = FPlatformFileManager::Get().GetPlatformFile();
-	TUniquePtr<IFileHandle> FileHandle(Ipf.OpenWrite(Filename.ToString()));
+	BlockInfo.ValueSort([](const FCasBlockInfo& LHS, const FCasBlockInfo& RHS)
+	{
+		return LHS.LastAccess < RHS.LastAccess;
+	});
 
-	if (FileHandle.IsValid() == false)
+	uint64 TotalPurgedBytes = 0;
+	for (const TPair<FCasBlockId, FCasBlockInfo>& Kv : BlockInfo)
 	{
-		return EIoErrorCode::FileOpenFailed;
+		const FCasBlockInfo& Info = Kv.Value;
+		if (Info.RefCount > 0)
+		{
+			continue;
+		}
+
+		FCasJournal::FTransaction	Transaction = FCasJournal::Begin(GetJournalFilename());
+		TArray<FCasAddr>			RemovedChunks;
+
+		if (FIoStatus Status = Cas.DeleteBlock(Kv.Key, RemovedChunks); !Status.IsOk())
+		{
+			return Status;
+		}
+
+		for (const FCasAddr& Addr : RemovedChunks)
+		{
+			Transaction.ChunkLocation(FCasLocation::Invalid, Addr);
+		}
+		Transaction.BlockDeleted(Kv.Key);
+
+		if (FIoStatus Status = FCasJournal::Commit(MoveTemp(Transaction)); !Status.IsOk())
+		{
+			return Status;
+		}
+
+		TotalPurgedBytes += Info.FileSize;
+		if (TotalPurgedBytes >= TotalUncachedBytes)
+		{
+			break;
+		}
 	}
 
-	UE_LOG(LogIoStoreOnDemand, Log, TEXT("Writing '%s' %.2lf KiB"), Filename.ToString(), double(Chunk.GetSize()) / 1024.0);
-	if (!FileHandle->Write(reinterpret_cast<const uint8*>(Chunk.GetData()), Chunk.GetSize()))
+	const uint64 NewCachedBytes = TotalCachedBytes - TotalPurgedBytes;
+	UE_LOG(LogIoStoreOnDemand, Log, TEXT("Purged %.2lf MiB from install cache"), ToMiB(TotalPurgedBytes));
+	UE_CLOG(NewCachedBytes > MaxCacheSize,
+		LogIoStoreOnDemand, Warning, TEXT("Max install cache size exceeded by %.2lf MiB"),
+			ToMiB(NewCachedBytes - MaxCacheSize));
+
+	return FIoStatus::Ok;
+}
+
+FIoStatus FOnDemandInstallCache::Flush()
+{
+	if (PendingChunks.IsValid())
 	{
-		return FIoStatusBuilder(EIoErrorCode::WriteError)
-			<< TEXT("Failed to write file '")
-			<< Filename.ToString()
-			<< TEXT("'");
+		FUniquePendingChunks Chunks = MoveTemp(PendingChunks);
+		return FlushPendingChunks(*Chunks);
 	}
 
-	return EIoErrorCode::Ok;
+	Cas.Compact();
+	return FIoStatus::Ok;;
 }
 
-void FOnDemandInstallCache::GetChunkFilename(const FIoHash& Hash, FStringBuilderBase& OutPath)
+FIoStatus FOnDemandInstallCache::FlushPendingChunks(FPendingChunks& Chunks)
 {
-	const FString HashString = LexToString(Hash);
+	ON_SCOPE_EXIT { Chunks.Reset(); };
+
+	while (Chunks.IsEmpty() == false)
+	{
+		FCasJournal::FTransaction Transaction = FCasJournal::Begin(GetJournalFilename());
+
+		if (Cas.CurrentBlock.IsValid() == false)
+		{
+			Cas.CurrentBlock = Cas.CreateBlock();
+			ensure(Cas.CurrentBlock.IsValid());
+			Transaction.BlockCreated(Cas.CurrentBlock);
+		}
+
+		//TODO: Allow reading and writing to the same block? 
+		TUniquePtr<IFileHandle>	CasFileHandle = Cas.OpenWrite(Cas.CurrentBlock);
+		if (CasFileHandle.IsValid() == false)
+		{
+			return FIoStatusBuilder(EIoErrorCode::FileOpenFailed)
+				<< TEXT("Failed to open cache block file '")
+				<< Cas.GetBlockFilename(Cas.CurrentBlock)
+				<< TEXT("'");
+		}
+
+		const int64 CasBlockOffset = CasFileHandle->Tell();
+
+		FLargeMemoryWriter	Ar(Chunks.TotalSize);
+		TArray<FIoHash>		ChunkHashes;
+		TArray<int64>		Offsets;
+
+		while (Chunks.IsEmpty() == false)
+		{
+			if (CasBlockOffset > 0 && CasBlockOffset + Ar.Tell() + Chunks.Chunks[0].GetSize() > Cas.MaxBlockSize)
+			{
+				break;
+			}
+			FIoBuffer Chunk = Chunks.Pop(ChunkHashes.AddDefaulted_GetRef());
+			Offsets.Add(CasBlockOffset + Ar.Tell());
+			Ar.Serialize(Chunk.GetData(), Chunk.GetSize());
+		}
+
+		if (Ar.Tell() > 0)
+		{
+			UE_LOG(LogIoStoreOnDemand, Log, TEXT("Writing %.2lf MiB to CAS block %u"),
+				ToMiB(Ar.Tell()), Cas.CurrentBlock.Id);
+
+			if (CasFileHandle->Write(Ar.GetData(), Ar.Tell()) == false)
+			{
+				return FIoStatusBuilder(EIoErrorCode::WriteError)
+					<< TEXT("Failed to serialize chunks to cache block");
+			}
+			Cas.TrackAccess(Cas.CurrentBlock);
+
+			if (CasFileHandle->Flush() == false)
+			{
+				return FIoStatusBuilder(EIoErrorCode::WriteError)
+					<< TEXT("Failed to flush cache block to disk");
+			}
+
+			check(ChunkHashes.Num() == Offsets.Num());
+			check(Cas.CurrentBlock.IsValid());
+			for (int32 Idx = 0, Count = Offsets.Num(); Idx < Count; ++Idx)
+			{
+				const FCasAddr	CasAddr = FCasAddr::From(ChunkHashes[Idx]);
+				const uint32	ChunkOffset = IntCastChecked<uint32>(Offsets[Idx]);
 
-	FPathViews::Append(OutPath, CacheDirectory);
-	FPathViews::Append(OutPath, TEXTVIEW("chunks"));
-	FPathViews::Append(OutPath, HashString.Left(2));
-	FPathViews::Append(OutPath, HashString);
-	OutPath << TEXT(".iochunk");
+				FCasLocation& Loc = Cas.Lookup.FindOrAdd(CasAddr);
+				Loc.BlockId	= Cas.CurrentBlock;
+				Loc.BlockOffset	= ChunkOffset;
+				Transaction.ChunkLocation(Loc, CasAddr);
+			}
+		}
+
+		if (FIoStatus Status = FCasJournal::Commit(MoveTemp(Transaction)); Status.IsOk() == false)
+		{
+			return Status;
+		}
+
+		if (Chunks.IsEmpty() == false)
+		{
+			Cas.CurrentBlock = FCasBlockId::Invalid;
+		}
+	}
+
+	return FIoStatus::Ok;;
 }
 
-void FOnDemandInstallCache::CompleteRequest(TUniquePtr<FChunkRequest>&& ChunkRequest)
+void FOnDemandInstallCache::CompleteRequest(FChunkRequest& ChunkRequest)
 {
-	FIoRequestImpl* Request				= ChunkRequest->DispatcherRequest;
-	const FOnDemandChunkInfo& ChunkInfo = ChunkRequest->ChunkInfo;
+	FIoRequestImpl* Request				= ChunkRequest.DispatcherRequest;
+	const FOnDemandChunkInfo& ChunkInfo = ChunkRequest.ChunkInfo;
+	FMemoryView EncodedChunk			= ChunkRequest.EncodedChunk.GetView();
+	bool bSucceeded						= EncodedChunk.IsEmpty() == false;
 
-	FIoChunkDecodingParams Params;
-	Params.CompressionFormat	= ChunkInfo.CompressionFormat();
-	Params.EncryptionKey		= ChunkInfo.EncryptionKey();
-	Params.BlockSize			= ChunkInfo.BlockSize();
-	Params.TotalRawSize			= ChunkInfo.RawSize();
-	Params.RawOffset			= Request->Options.GetOffset();
-	Params.EncodedOffset		= ChunkRequest->ChunkRange.GetOffset();
-	Params.EncodedBlockSize		= ChunkInfo.Blocks();
-	Params.BlockHash			= ChunkInfo.BlockHashes();
+	if (bSucceeded)
+	{
+		FIoChunkDecodingParams Params;
+		Params.CompressionFormat	= ChunkInfo.CompressionFormat();
+		Params.EncryptionKey		= ChunkInfo.EncryptionKey();
+		Params.BlockSize			= ChunkInfo.BlockSize();
+		Params.TotalRawSize			= ChunkInfo.RawSize();
+		Params.RawOffset			= Request->Options.GetOffset();
+		Params.EncodedOffset		= ChunkRequest.ChunkRange.GetOffset();
+		Params.EncodedBlockSize		= ChunkInfo.Blocks();
+		Params.BlockHash			= ChunkInfo.BlockHashes();
 
-	Request->CreateBuffer(ChunkRequest->RawSize);
+		Request->CreateBuffer(ChunkRequest.RawSize);
+		FMutableMemoryView RawChunk = Request->GetBuffer().GetMutableView();
 
-	FMutableMemoryView RawChunk = Request->GetBuffer().GetMutableView();
-	FMemoryView EncodedChunk	= ChunkRequest->EncodedChunk.GetView();
+		bSucceeded = FIoChunkEncoding::Decode(Params, EncodedChunk, RawChunk);
+		UE_CLOG(!bSucceeded, LogIoStoreOnDemand, Error, TEXT("Failed to decode chunk, ChunkId='%s'"), *LexToString(Request->ChunkId));
+	}
 
-	if (FIoChunkEncoding::Decode(Params, EncodedChunk, RawChunk) == false)
+	if (bSucceeded == false)
 	{
-		UE_LOG(LogIoStoreOnDemand, Error, TEXT("Failed to read chunk, ChunkId='%s'"), *LexToString(Request->ChunkId));
 		Request->SetResult(FIoBuffer());
 		Request->SetFailed();
 	}
diff --git a/Engine/Source/Runtime/Experimental/IoStore/OnDemand/Private/OnDemandInstallCache.h b/Engine/Source/Runtime/Experimental/IoStore/OnDemand/Private/OnDemandInstallCache.h
--- a/Engine/Source/Runtime/Experimental/IoStore/OnDemand/Private/OnDemandInstallCache.h
+++ b/Engine/Source/Runtime/Experimental/IoStore/OnDemand/Private/OnDemandInstallCache.h
@@ -18,15 +18,18 @@ class IOnDemandInstallCache
 	: public IIoDispatcherBackend
 {
 public:
-	virtual ~IOnDemandInstallCache() = default;
-	virtual bool ContainsChunk(const FIoHash& Hash) = 0;
-	virtual FIoStatus PutChunk(FIoBuffer&& Chunk, const FIoHash& Hash) = 0;
+	virtual				~IOnDemandInstallCache() = default;
+	virtual bool		IsChunkCached(const FIoHash& ChunkHash) = 0;
+	virtual FIoStatus	PutChunk(FIoBuffer&& Chunk, const FIoHash& ChunkHash) = 0;
+	virtual FIoStatus	Purge(TMap<FIoHash, uint64>&& ChunksToInstall) = 0;
+	virtual FIoStatus	Flush() = 0;
 };
 
 struct FOnDemandInstallCacheConfig
 {
 	FString RootDirectory;
-	bool bDropCache = false;
+	uint64	DiskQuota = 1ull << 30;
+	bool	bDropCache = false;
 };
 
 TSharedPtr<IOnDemandInstallCache> MakeOnDemandInstallCache(
diff --git a/Engine/Source/Runtime/Experimental/IoStore/OnDemand/Private/OnDemandIoStore.cpp b/Engine/Source/Runtime/Experimental/IoStore/OnDemand/Private/OnDemandIoStore.cpp
--- a/Engine/Source/Runtime/Experimental/IoStore/OnDemand/Private/OnDemandIoStore.cpp
+++ b/Engine/Source/Runtime/Experimental/IoStore/OnDemand/Private/OnDemandIoStore.cpp
@@ -974,9 +974,30 @@ FIoStatus FOnDemandIoStore::TickInstallRequest(FMountRequest& MountRequest)
 		}
 	}
 
+	// Purge
+	{
+		TMap<FIoHash, uint64> ChunksToInstall;
+		for (const auto& Kv : InstallData)
+		{
+			FOnDemandContainer&						Container = *Kv.Key;
+			const Private::FContainerInstallData&	Data = Kv.Value;
+
+			for (const FIoChunkId& ChunkId : Data.ResolvedChunks)
+			{
+				const FOnDemandChunkEntry& Entry = Container.ChunkEntries.FindRef(ChunkId);
+				ChunksToInstall.Add(Entry.Hash, Entry.EncodedSize);
+			}
+		}
+
+		if (Status = InstallCache->Purge(MoveTemp(ChunksToInstall)); Status.IsOk() == false)
+		{
+			return Status;
+		}
+	}
+
 	// Download all chunks
-	const int32 MaxConcurrentRequests = 32;
-	int32 ConcurrentRequests = 0;
+	const int32 MaxConcurrentRequests = 16;
+	int32		ConcurrentRequests = 0;
 
 	FHttpClientConfig HttpConfig;
 	HttpConfig.MaxConnectionCount = 8;
@@ -1001,19 +1022,19 @@ FIoStatus FOnDemandIoStore::TickInstallRequest(FMountRequest& MountRequest)
 
 		for (const FIoChunkId& ChunkId : Data.ResolvedChunks)
 		{
-			const FOnDemandChunkEntry& Entry = Container.ChunkEntries.FindRef(ChunkId);
+			FOnDemandChunkEntry* ChunkEntry = Container.ChunkEntries.Find(ChunkId);
+
 			++TotalChunkCount;
-			TotalBytes += Entry.EncodedSize;
+			TotalBytes += ChunkEntry->EncodedSize;
 
-			if (InstallCache->ContainsChunk(Entry.Hash))
+			if (InstallCache->IsChunkCached(ChunkEntry->Hash))
 			{
 				continue;
 			}
 
 			++DownloadedChunkCount;
-			DownloadedBytes += Entry.EncodedSize;
+			DownloadedBytes += ChunkEntry->EncodedSize;
 
-			FOnDemandChunkEntry* ChunkEntry = Container.ChunkEntries.Find(ChunkId);
 			ConcurrentRequests++;
 			HttpClient->Get(
 				GetChunkUrl(FStringView(), Container, *ChunkEntry, ChunkUrl).ToView(),
@@ -1027,7 +1048,14 @@ FIoStatus FOnDemandIoStore::TickInstallRequest(FMountRequest& MountRequest)
 						return;
 					}
 
-					Status = InstallCache->PutChunk(ChunkStatus.ConsumeValueOrDie(), ChunkEntry->Hash);
+					FIoBuffer Chunk			= ChunkStatus.ConsumeValueOrDie();
+					const FIoHash ChunkHash = FIoHash::HashBuffer(Chunk.GetView());
+					if (ChunkHash != ChunkEntry->Hash)
+					{
+						Status = FIoStatus(EIoErrorCode::ReadError, TEXTVIEW("Hash mismatch"));
+						return;
+					}
+					Status = InstallCache->PutChunk(MoveTemp(Chunk), ChunkHash);
 				});
 
 			while (ConcurrentRequests >= MaxConcurrentRequests)
@@ -1044,6 +1072,11 @@ FIoStatus FOnDemandIoStore::TickInstallRequest(FMountRequest& MountRequest)
 
 	while (HttpClient->Tick())
 		;
+	
+	if (Status = InstallCache->Flush(); Status.IsOk() == false)
+	{
+		return FIoStatusBuilder(EIoErrorCode::WriteError) << TEXT("Failed to flush downloded content to cache");
+	}
 
 	// TODO: Only mount what has been installed
 	for (const auto& Kv : InstallData)
