Couple of extra check() tests

#rnx

diff --git a/Engine/Source/Runtime/Experimental/IoStore/HttpClient/Private/Test.inl b/Engine/Source/Runtime/Experimental/IoStore/HttpClient/Private/Test.inl
--- a/Engine/Source/Runtime/Experimental/IoStore/HttpClient/Private/Test.inl
+++ b/Engine/Source/Runtime/Experimental/IoStore/HttpClient/Private/Test.inl
@@ -82,6 +82,8 @@ static void MiscTest()
 	check(CrudeToInt<16>("f") == 0x0f); check(CrudeToInt<16>("F") == 0x0f);
 	check(CrudeToInt<16>("g") < 0);
 	check(CrudeToInt<16>("49e") == 0x49e);
+	check(CrudeToInt<16>("aBcD") == 0xabcd);
+	check(CrudeToInt<16>("eEeE") == 0xeeee);
 
 	FUrlOffsets UrlOut;
 
