Moved seeding of test server out to easier set a consistent environment

#rnx

diff --git a/Engine/Source/Runtime/Experimental/IoStore/HttpClient/Private/Test.inl b/Engine/Source/Runtime/Experimental/IoStore/HttpClient/Private/Test.inl
--- a/Engine/Source/Runtime/Experimental/IoStore/HttpClient/Private/Test.inl
+++ b/Engine/Source/Runtime/Experimental/IoStore/HttpClient/Private/Test.inl
@@ -421,6 +421,17 @@ static void RedirectTest(const ANSICHAR* TestHost, FCertRootsRef VerifyCert)
 	}
 }
 
+////////////////////////////////////////////////////////////////////////////////
+static void SeedHttp(const ANSICHAR* TestHost, uint32 Seed)
+{
+	TAnsiStringBuilder<64> Url;
+	Url << "http://" << TestHost << ":9493/seed/" << Seed;
+	FEventLoop Loop;
+	FRequest Request = Loop.Request("GET", Url, nullptr);
+	Loop.Send(MoveTemp(Request), [] (const FTicketStatus&) {});
+	for (; Loop.Tick(-1); FPlatformProcess::SleepNoStats(0.02f));
+}
+
 ////////////////////////////////////////////////////////////////////////////////
 static void HttpTest(const ANSICHAR* TestHost, FCertRootsRef VerifyCert)
 {
@@ -542,10 +553,10 @@ static void HttpTest(const ANSICHAR* TestHost, FCertRootsRef VerifyCert)
 
 	// foundational
 	{
-		FRequest Request = Loop.Request("GET", BuildUrl("/seed/493"), ReqParams);
+		FRequest Request = Loop.Request("GET", BuildUrl("/data/67"), ReqParams);
 		Request.Accept(EMimeType::Json);
 
-		FTicket Ticket = Loop.Send(MoveTemp(Request), NullSink);
+		FTicket Ticket = Loop.Send(MoveTemp(Request), HashSink);
 
 		WaitForLoopIdle();
 	}
@@ -825,7 +836,7 @@ static void HttpTest(const ANSICHAR* TestHost, FCertRootsRef VerifyCert)
 }
 
 ////////////////////////////////////////////////////////////////////////////////
-IOSTOREHTTPCLIENT_API void IasHttpTest(const ANSICHAR* TestHost="localhost")
+IOSTOREHTTPCLIENT_API void IasHttpTest(const ANSICHAR* TestHost="localhost", uint32 Seed=493)
 {
 #if PLATFORM_WINDOWS
 	WSADATA WsaData;
@@ -858,14 +869,12 @@ IOSTOREHTTPCLIENT_API void IasHttpTest(const ANSICHAR* TestHost="localhost")
 
 		TestServerCaChain = FCertRoots(CertBuffer.GetView());
 	}
-
 	FCertRootsRef TestServerCertRef = FCertRoots::Explicit(TestServerCaChain);
 
+	SeedHttp(TestHost, Seed);
 	HttpTest(TestHost, FCertRoots::NoTls());
 	HttpTest(TestHost, TestServerCertRef);
-
 	RedirectTest(TestHost, TestServerCertRef);
-
 	TlsLoadRootCerts();
 	TlsTest();
 }
diff --git a/Engine/Source/Runtime/Experimental/IoStore/OnDemand/Private/Tool/Test.cpp b/Engine/Source/Runtime/Experimental/IoStore/OnDemand/Private/Tool/Test.cpp
--- a/Engine/Source/Runtime/Experimental/IoStore/OnDemand/Private/Tool/Test.cpp
+++ b/Engine/Source/Runtime/Experimental/IoStore/OnDemand/Private/Tool/Test.cpp
@@ -12,8 +12,8 @@
 namespace UE::IoStore {
 
 ////////////////////////////////////////////////////////////////////////////////
-namespace IasJournaledFileCacheTest { void Tests(const TCHAR*);				}
-namespace HTTP						{ void IasHttpTest(const ANSICHAR*);	}
+namespace IasJournaledFileCacheTest { void Tests(const TCHAR*);					}
+namespace HTTP						{ void IasHttpTest(const ANSICHAR*, uint32);}
 
 namespace Tool {
 
@@ -26,7 +26,8 @@ static void HttpTests(const FContext& Context)
 	auto TestHost = Context.Get<FStringView>(TEXT("-Host"), TEXT("localhost"));
 
 	auto TestHostAnsi = StringCast<ANSICHAR>(TestHost.GetData());
-	HTTP::IasHttpTest(TestHostAnsi.Get());
+	uint32 Seed = Context.Get<uint32>(TEXT("-HttpSeed"), 493);
+	HTTP::IasHttpTest(TestHostAnsi.Get(), Seed);
 }
 
 ////////////////////////////////////////////////////////////////////////////////
@@ -70,6 +71,7 @@ static FCommand TestCommand(
 		TArgument<FStringView>(TEXT("-Host"), TEXT("Host of the HTTP test server")),
 		TArgument<FStringView>(TEXT("-Dir"), TEXT("Primary directory to use for cache tests")),
 		TArgument<FStringView>(TEXT("-Only"), TEXT("Only run a particular test (http|cache)")),
+		TArgument<uint32>(TEXT("-HttpSeed"), TEXT("Integer value to seed test HTTP server")),
 	}
 );
 
