Fix ensure caused by missing Super::EndPlay in UNetTestUnmappedComponent and ChaosSolverActor

diff --git a/Engine/Source/Runtime/Experimental/ChaosSolverEngine/Private/Chaos/ChaosSolverActor.cpp b/Engine/Source/Runtime/Experimental/ChaosSolverEngine/Private/Chaos/ChaosSolverActor.cpp
--- a/Engine/Source/Runtime/Experimental/ChaosSolverEngine/Private/Chaos/ChaosSolverActor.cpp
+++ b/Engine/Source/Runtime/Experimental/ChaosSolverEngine/Private/Chaos/ChaosSolverActor.cpp
@@ -305,6 +305,8 @@ void AChaosSolverActor::BeginPlay()
 
 void AChaosSolverActor::EndPlay(const EEndPlayReason::Type ReasonEnd)
 {
+	Super::EndPlay(ReasonEnd);
+
 	if(!Solver)
 	{
 		return;
