Iris
* Refactored NetRefHandleManager::AddDependentActor
* Downgraded ensure to warning when we detect that the dependency link is already present.
* Added ensure to make sure a change to the SchedulingHint was not missed.

#rnx
#-> 66671ecf630317ecf5c85007

diff --git a/Engine/Source/Runtime/Experimental/Iris/Core/Private/Iris/ReplicationSystem/NetRefHandleManager.cpp b/Engine/Source/Runtime/Experimental/Iris/Core/Private/Iris/ReplicationSystem/NetRefHandleManager.cpp
--- a/Engine/Source/Runtime/Experimental/Iris/Core/Private/Iris/ReplicationSystem/NetRefHandleManager.cpp
+++ b/Engine/Source/Runtime/Experimental/Iris/Core/Private/Iris/ReplicationSystem/NetRefHandleManager.cpp
@@ -361,7 +361,7 @@ FNetRefHandle FNetRefHandleManager::CreateNetObject(FNetRefHandle WantedHandle,
 // Create NetRefHandle not owned by us
 FNetRefHandle FNetRefHandleManager::CreateNetObjectFromRemote(FNetRefHandle WantedHandle, const FReplicationProtocol* ReplicationProtocol)
 {
-	if (!ensureAlwaysMsgf(WantedHandle.IsValid() && !WantedHandle.IsCompleteHandle(), TEXT("FNetRefHandleManager::CreateNetObjectFromRemote Expected WantedHandle %s to be valid and incomplete"), *WantedHandle.ToString()))
+	if (!ensureMsgf(WantedHandle.IsValid() && !WantedHandle.IsCompleteHandle(), TEXT("FNetRefHandleManager::CreateNetObjectFromRemote Expected WantedHandle %s to be valid and incomplete"), *WantedHandle.ToString()))
 	{
 		return FNetRefHandle();
 	}
@@ -584,7 +584,7 @@ bool FNetRefHandleManager::InternalAddSubObject(FInternalNetRefIndex OwnerIntern
 	using namespace UE::Net::Private;
 
 	FReplicatedObjectData& SubObjectData = GetReplicatedObjectDataNoCheck(SubObjectInternalIndex);
-	if (!ensureAlwaysMsgf(SubObjectData.SubObjectRootIndex == InvalidInternalIndex, TEXT("FNetRefHandleManager::AddSubObject %s is already marked as a subobject"), ToCStr(SubObjectData.RefHandle.ToString())))
+	if (!ensureMsgf(SubObjectData.SubObjectRootIndex == InvalidInternalIndex, TEXT("FNetRefHandleManager::AddSubObject %s is already marked as a subobject"), ToCStr(SubObjectData.RefHandle.ToString())))
 	{
 		return false;
 	}
@@ -596,7 +596,7 @@ bool FNetRefHandleManager::InternalAddSubObject(FInternalNetRefIndex OwnerIntern
 	// Mark the object as a subobject
 	SetIsSubObject(SubObjectInternalIndex, true);
 
-	if (RelativeOtherSubObjectInternalIndex != InvalidInternalIndex && ensureAlwaysMsgf(SubObjectArray.Find(RelativeOtherSubObjectInternalIndex) != INDEX_NONE, TEXT("RelativeOtherSubObjectHandle %s Must be a Subobject of %s"), ToCStr(GetNetRefHandleFromInternalIndex(RelativeOtherSubObjectInternalIndex).ToString()), ToCStr(GetNetRefHandleFromInternalIndex(OwnerInternalIndex).ToString())))
+	if (RelativeOtherSubObjectInternalIndex != InvalidInternalIndex && ensureMsgf(SubObjectArray.Find(RelativeOtherSubObjectInternalIndex) != INDEX_NONE, TEXT("RelativeOtherSubObjectHandle %s Must be a Subobject of %s"), ToCStr(GetNetRefHandleFromInternalIndex(RelativeOtherSubObjectInternalIndex).ToString()), ToCStr(GetNetRefHandleFromInternalIndex(OwnerInternalIndex).ToString())))
 	{
 		// Add to child array of RelativeOtherSubObjectRelativeIndex for hierarchical replication order
 		FNetDependencyData::FSubObjectConditionalsArray* SubObjectConditionalsArray = nullptr;
@@ -730,7 +730,7 @@ FInternalNetRefIndex FNetRefHandleManager::GetRootObjectInternalIndexOfSubObject
 	return SubObjectIndex != InvalidInternalIndex ? ReplicatedObjectData[SubObjectIndex].SubObjectRootIndex : InvalidInternalIndex;
 }
 
-bool FNetRefHandleManager::AddDependentObject(FNetRefHandle ParentRefHandle, FNetRefHandle DependentObjectRefHandle, EDependentObjectSchedulingHint SchedulingHint, EAddDependentObjectFlags Flags)
+bool FNetRefHandleManager::AddDependentObject(FNetRefHandle ParentRefHandle, FNetRefHandle DependentObjectRefHandle, EDependentObjectSchedulingHint SchedulingHint)
 {
 	check(ParentRefHandle != DependentObjectRefHandle);
 
@@ -754,25 +754,37 @@ bool FNetRefHandleManager::AddDependentObject(FNetRefHandle ParentRefHandle, FNe
 	check(!SubObjectInternalIndices.GetBit(DependentObjectInternalIndex));
 	check(!SubObjectInternalIndices.GetBit(ParentInternalIndex));
 
-	// Add dependent to parents dependent object list
 	FNetDependencyData::FDependentObjectInfoArray& ParentDependentObjectsArray = SubObjects.GetOrCreateDependentObjectInfoArray(ParentInternalIndex);
-	const bool bDependentIsAlreadyDependant = ParentDependentObjectsArray.FindByPredicate([DependentObjectInternalIndex](const FDependentObjectInfo& Entry) { return Entry.NetRefIndex == DependentObjectInternalIndex;}) != nullptr;
-	if (!bDependentIsAlreadyDependant)
+	FNetDependencyData::FInternalNetRefIndexArray& DependentParentObjectArray = SubObjects.GetOrCreateInternalIndexArray<FNetDependencyData::DependentParentObjects>(DependentObjectInternalIndex);
+	
+	// Make sure parent didn't set the child as a dependent already
 	{
-		FDependentObjectInfo DependentObjectInfo;
+		const FDependentObjectInfo* DependentInfo = ParentDependentObjectsArray.FindByPredicate([DependentObjectInternalIndex](const FDependentObjectInfo& Entry) { return Entry.NetRefIndex == DependentObjectInternalIndex;});
+		if (DependentInfo)
+		{
+			// Make sure the children is also dependent to the Parent
+			checkf(DependentParentObjectArray.Find(ParentInternalIndex) != INDEX_NONE, TEXT("FNetRefHandleManager::AddDependentObject: Parent: %s already has child: %s as dependent but not the inverse."), 
+				*PrintObjectFromNetRefHandle(ParentRefHandle), *PrintObjectFromNetRefHandle(DependentObjectRefHandle));
+
+			// If they were already dependent there is no side-effect, unless the scheduler hint would have been changed by the new call.
+			UE_LOG(LogIris, Warning, TEXT("FNetRefHandleManager::AddDependentObject: Parent: %s already has child: %s as a dependent"), *PrintObjectFromNetRefHandle(ParentRefHandle), *PrintObjectFromNetRefHandle(DependentObjectRefHandle));
+			ensureMsgf(DependentInfo->SchedulingHint == SchedulingHint, TEXT("FNetRefHandleManager::AddDependentObject: Conflicting scheduling hint between Child: %s and Parent: %s. Requested %s but was already set to %s"),
+				*PrintObjectFromNetRefHandle(DependentObjectRefHandle), *PrintObjectFromNetRefHandle(ParentRefHandle), LexToString(SchedulingHint), LexToString(DependentInfo->SchedulingHint));
+			return false;
+		}
+	}
 
-		DependentObjectInfo.NetRefIndex = DependentObjectInternalIndex;
-		DependentObjectInfo.SchedulingHint = SchedulingHint;
+	// If child was already set as dependent on the parent there is a logic error somewhere.
+	checkf(DependentParentObjectArray.Find(ParentInternalIndex) == INDEX_NONE, TEXT("FNetRefHandleManager::AddDependentObject: Child: %s already dependent of Parent: %s but not the inverse."), *PrintObjectFromNetRefHandle(DependentObjectRefHandle), *PrintObjectFromNetRefHandle(ParentRefHandle));
 
-		ParentDependentObjectsArray.Add(DependentObjectInfo);
-	}
+	// Add dependent to parent's dependent object list
+	FDependentObjectInfo DependentObjectInfo;
+	DependentObjectInfo.NetRefIndex = DependentObjectInternalIndex;
+	DependentObjectInfo.SchedulingHint = SchedulingHint;
+	ParentDependentObjectsArray.Add(DependentObjectInfo);
 
-	FNetDependencyData::FInternalNetRefIndexArray& DependentParentObjectArray = SubObjects.GetOrCreateInternalIndexArray<FNetDependencyData::DependentParentObjects>(DependentObjectInternalIndex);
-	const bool bDependentHadParentAlready = DependentParentObjectArray.Find(ParentInternalIndex) != INDEX_NONE;
-	if (!bDependentHadParentAlready)
-	{
-		DependentParentObjectArray.Add(ParentInternalIndex);
-	}
+	// Add parent to dependent's list
+	DependentParentObjectArray.Add(ParentInternalIndex);
 
 	// Update cached info to avoid to do map lookups to find out if we are a dependent object or have dependent objects
 	DependentObjectData.bIsDependentObject = true;
@@ -780,14 +792,6 @@ bool FNetRefHandleManager::AddDependentObject(FNetRefHandle ParentRefHandle, FNe
 	ObjectsWithDependentObjectsInternalIndices.SetBit(ParentInternalIndex);
 	DependentObjectInternalIndices.SetBit(DependentObjectInternalIndex);
 
-	if (EnumHasAnyFlags(Flags, EAddDependentObjectFlags::WarnIfAlreadyDependentObject) && (bDependentHadParentAlready || bDependentIsAlreadyDependant))
-	{
-		// If this gets out of sync something is messed up
-		check(bDependentHadParentAlready == bDependentIsAlreadyDependant);
-		ensureMsgf(false, TEXT("FNetRefHandleManager::AddDependentObject: %s already is a child of parent: %s. Cannot be added to: %s"), 
-			*PrintObjectFromNetRefHandle(DependentObjectRefHandle), *PrintObjectFromIndex(DependentObjectData.SubObjectRootIndex), *PrintObjectFromNetRefHandle(ParentRefHandle));
-	}
-
 	return true;
 }
 
diff --git a/Engine/Source/Runtime/Experimental/Iris/Core/Private/Iris/ReplicationSystem/NetRefHandleManager.h b/Engine/Source/Runtime/Experimental/Iris/Core/Private/Iris/ReplicationSystem/NetRefHandleManager.h
--- a/Engine/Source/Runtime/Experimental/Iris/Core/Private/Iris/ReplicationSystem/NetRefHandleManager.h
+++ b/Engine/Source/Runtime/Experimental/Iris/Core/Private/Iris/ReplicationSystem/NetRefHandleManager.h
@@ -44,13 +44,6 @@ enum class EAddSubObjectFlags : uint32
 };
 ENUM_CLASS_FLAGS(EAddSubObjectFlags);
 
-enum class EAddDependentObjectFlags : uint32
-{
-	None = 0U,
-	WarnIfAlreadyDependentObject = 1U,
-};
-ENUM_CLASS_FLAGS(EAddDependentObjectFlags);
-
 enum class ERemoveDependentObjectFlags : uint32
 {
 	None = 0U,
@@ -250,7 +243,7 @@ public:
 	// DependentObjects
 	const FNetBitArrayView GetDependentObjectInternalIndices() const { return MakeNetBitArrayView(DependentObjectInternalIndices); }
 	const FNetBitArrayView GetObjectsWithDependentObjectsInternalIndices() const { return MakeNetBitArrayView(ObjectsWithDependentObjectsInternalIndices); }
-	bool AddDependentObject(FNetRefHandle ParentHandle, FNetRefHandle DependentHandle, EDependentObjectSchedulingHint SchedulingHint, EAddDependentObjectFlags Flags = EAddDependentObjectFlags::WarnIfAlreadyDependentObject);
+	bool AddDependentObject(FNetRefHandle ParentHandle, FNetRefHandle DependentHandle, EDependentObjectSchedulingHint SchedulingHint);
 	void RemoveDependentObject(FNetRefHandle ParentHandle, FNetRefHandle DependentHandle);
 
 	// Remove DependentHandles from all dependent object tracking
diff --git a/Engine/Source/Runtime/Experimental/Iris/Core/Public/Iris/ReplicationSystem/ReplicationSystemTypes.h b/Engine/Source/Runtime/Experimental/Iris/Core/Public/Iris/ReplicationSystem/ReplicationSystemTypes.h
--- a/Engine/Source/Runtime/Experimental/Iris/Core/Public/Iris/ReplicationSystem/ReplicationSystemTypes.h
+++ b/Engine/Source/Runtime/Experimental/Iris/Core/Public/Iris/ReplicationSystem/ReplicationSystemTypes.h
@@ -78,4 +78,15 @@ inline const TCHAR* LexToString(ENetObjectAttachmentSendPolicyFlags SendFlags)
 	}
 }
 
+inline const TCHAR* LexToString(EDependentObjectSchedulingHint SchedulingHint)
+{
+	switch (SchedulingHint)
+	{
+	case EDependentObjectSchedulingHint::Default: return TEXT("Default");
+	case EDependentObjectSchedulingHint::ScheduleBeforeParent: return TEXT("BeforeParent");
+	case EDependentObjectSchedulingHint::ScheduleBeforeParentIfInitialState: return TEXT("BeforeParentIfInitialState");
+	default: ensure(false); return TEXT("Missing");
+	}
+}
+
 }
