Chaos Vehicles (new eSimTypes);

Adding a jump and booster sim type.

-- Code --
SimulationModuleBase:
eSimType::Jump - Applies burst upward force.
eSimType::Booster - Applies continuous forward force to the center of mass.

#tests PIE, sandbox, make sure everything still works as expected.

diff --git a/Engine/Source/Runtime/Experimental/ChaosVehicles/ChaosVehiclesCore/Public/SimModule/SimulationModuleBase.h b/Engine/Source/Runtime/Experimental/ChaosVehicles/ChaosVehiclesCore/Public/SimModule/SimulationModuleBase.h
--- a/Engine/Source/Runtime/Experimental/ChaosVehicles/ChaosVehiclesCore/Public/SimModule/SimulationModuleBase.h
+++ b/Engine/Source/Runtime/Experimental/ChaosVehicles/ChaosVehiclesCore/Public/SimModule/SimulationModuleBase.h
@@ -117,7 +117,9 @@ namespace Chaos
 		Rudder,			// controls aircraft yaw
 		Elevator,		// controls aircraft pitch
 		Propeller,		// generates thrust when connected to a motor/engine
-		TorqueSim
+		TorqueSim,
+		Booster,		// applies continuous forward force to the center of mass
+		Jump,			// applies burst upward force
 	};
 
 	enum EWheelAxis
