Chaos - fix ShapeIndex for shapes in the game thread particle after welding

diff --git a/Engine/Source/Runtime/Experimental/Chaos/Public/Chaos/ParticleHandle.h b/Engine/Source/Runtime/Experimental/Chaos/Public/Chaos/ParticleHandle.h
--- a/Engine/Source/Runtime/Experimental/Chaos/Public/Chaos/ParticleHandle.h
+++ b/Engine/Source/Runtime/Experimental/Chaos/Public/Chaos/ParticleHandle.h
@@ -2867,6 +2867,7 @@ public:
 		for (FShapeInstanceProxyPtr& Shape : InShapes)
 		{
 			ensure(Idx < MShapesArray.Num());
+			Shape->ModifyShapeIndex(Idx);
 			MShapesArray[Idx++] = MoveTemp(Shape);
 		}
 	}
