Chaos Debug Draw - add render budget
- add full set of primitives to match legacy system

diff --git a/Engine/Source/Runtime/Experimental/Chaos/Private/Chaos/DebugDraw/DebugDrawPrimitives.cpp b/Engine/Source/Runtime/Experimental/Chaos/Private/Chaos/DebugDraw/DebugDrawPrimitives.cpp
--- a/Engine/Source/Runtime/Experimental/Chaos/Private/Chaos/DebugDraw/DebugDrawPrimitives.cpp
+++ b/Engine/Source/Runtime/Experimental/Chaos/Private/Chaos/DebugDraw/DebugDrawPrimitives.cpp
@@ -8,6 +8,18 @@
 
 namespace Chaos
 {
+	void FChaosDDPrimitives::DrawPoint(const FVector3d& Position, const FColor& Color, float PointSize, float Duration)
+	{
+		constexpr int32 Cost = 1;
+		const FBox3d Bounds = FBox3d(Position, Position);
+
+		ChaosDD::Private::FChaosDDContext::GetWriter().TryEnqueueCommand(Cost, Bounds,
+			[=](ChaosDD::Private::IChaosDDRenderer& Renderer)
+			{
+				Renderer.RenderPoint(Position, Color, PointSize, Duration);
+			});
+	}
+
 	void FChaosDDPrimitives::DrawLine(const FVector3d& A, const FVector3d& B, const FColor& Color, float LineThickness, float Duration)
 	{
 		constexpr int32 Cost = 1;
@@ -20,7 +32,30 @@ namespace Chaos
 			});
 	}
 
+	void FChaosDDPrimitives::DrawArrow(const FVector3d& A, const FVector3d& B, float ArrowSize, const FColor& Color, float LineThickness, float Duration)
+	{
+		constexpr int32 Cost = 1;
+		const FBox3d Bounds = FBox3d(FVector3d::Min(A, B), FVector3d::Max(A, B));
+
+		ChaosDD::Private::FChaosDDContext::GetWriter().TryEnqueueCommand(Cost, Bounds,
+			[=](ChaosDD::Private::IChaosDDRenderer& Renderer)
+			{
+				Renderer.RenderArrow(A, B, ArrowSize, Color, LineThickness, Duration);
+			});
+	}
 	
+	void FChaosDDPrimitives::DrawCircle(const FVector3d& Center, const FMatrix& Axes, float Radius, const FColor& Color, float LineThickness, float Duration)
+	{
+		constexpr int32 Cost = 1;
+		const FBox3d Bounds = FBox3d(Center - FVector3d(Radius), Center + FVector3d(Radius));
+
+		ChaosDD::Private::FChaosDDContext::GetWriter().TryEnqueueCommand(Cost, Bounds,
+			[=](ChaosDD::Private::IChaosDDRenderer& Renderer)
+			{
+				Renderer.RenderCircle(Center, Axes, Radius, Color, LineThickness, Duration);
+			});
+	}
+
 	void FChaosDDPrimitives::DrawSphere(const FVector3d& Center, const float Radius, const FColor& Color, float LineThickness, float Duration)
 	{
 		constexpr int32 Cost = 64;
@@ -71,6 +106,18 @@ namespace Chaos
 				Renderer.RenderTriangle(A, B, C, Color, LineThickness, Duration);
 			});
 	}
+
+	void FChaosDDPrimitives::DrawString(const FVector3d& TextLocation, const FString& Text, const FColor& Color, float FontScale, bool bDrawShadow, float Duration)
+	{
+		const int32 Cost = 10;
+		const FBox3d Bounds = FBox3d(TextLocation, TextLocation);
+
+		ChaosDD::Private::FChaosDDContext::GetWriter().TryEnqueueCommand(Cost, Bounds,
+			[=](ChaosDD::Private::IChaosDDRenderer& Renderer)
+			{
+				Renderer.RenderString(TextLocation, Text, Color, FontScale, bDrawShadow, Duration);
+			});
+	}
 }
 
 #endif
\ No newline at end of file
diff --git a/Engine/Source/Runtime/Experimental/Chaos/Private/ChaosDebugDraw/ChaosDDContext.cpp b/Engine/Source/Runtime/Experimental/Chaos/Private/ChaosDebugDraw/ChaosDDContext.cpp
--- a/Engine/Source/Runtime/Experimental/Chaos/Private/ChaosDebugDraw/ChaosDDContext.cpp
+++ b/Engine/Source/Runtime/Experimental/Chaos/Private/ChaosDebugDraw/ChaosDDContext.cpp
@@ -16,6 +16,8 @@ namespace ChaosDD::Private
 	bool bChaosDebugDraw_EnableGlobalQueue = true;
 	FAutoConsoleVariableRef CVarChaos_DebugDraw_EnableGlobalQueue(TEXT("p.Chaos.DebugDraw.EnableGlobalQueue"), bChaosDebugDraw_EnableGlobalQueue, TEXT(""));
 
+	bool FChaosDDContext::bDebugDrawEnabled = false;
+
 	FCriticalSection FChaosDDContext::GlobalFrameCS;
 	FChaosDDFramePtr FChaosDDContext::GlobalFrame;
 	int32 FChaosDDContext::GlobalCommandBudget = 20000;
@@ -28,32 +30,44 @@ namespace ChaosDD::Private
 
 	void FChaosDDTimelineContext::BeginFrame(const FChaosDDTimelinePtr& InTimeline, double InTime, double InDt)
 	{
-		FChaosDDContext& Context = FChaosDDContext::Get();
-		PreviousFrame = Context.Frame;
+		check(!bInContext);
 
-		if (InTimeline.IsValid())
-		{
-			Timeline = InTimeline;
-			Timeline->BeginFrame(InTime, InDt);
-			Context.Frame = Timeline->GetActiveFrame();
-		}
-		else
+		if (FChaosDDContext::IsDebugDrawEnabled())
 		{
-			Context.Frame.Reset();
+			FChaosDDContext& Context = FChaosDDContext::Get();
+			PreviousFrame = Context.Frame;
+
+			if (InTimeline.IsValid())
+			{
+				Timeline = InTimeline;
+				Timeline->BeginFrame(InTime, InDt);
+				Context.Frame = Timeline->GetActiveFrame();
+			}
+			else
+			{
+				Context.Frame.Reset();
+			}
+
+			bInContext = true;
 		}
 	}
 
 	void FChaosDDTimelineContext::EndFrame()
 	{
-		if (Timeline.IsValid())
+		if (bInContext)
 		{
-			Timeline->EndFrame();
-			Timeline.Reset();
-		}
+			if (Timeline.IsValid())
+			{
+				Timeline->EndFrame();
+				Timeline.Reset();
+			}
+
+			FChaosDDContext& Context = FChaosDDContext::Get();
+			Context.Frame = PreviousFrame;
+			PreviousFrame.Reset();
 
-		FChaosDDContext& Context = FChaosDDContext::Get();
-		Context.Frame = PreviousFrame;
-		PreviousFrame.Reset();
+			bInContext = false;
+		}
 	}
 
 	FChaosDDScopeTimelineContext::FChaosDDScopeTimelineContext(const FChaosDDTimelinePtr& InTimeline, double InTime, double InDt)
@@ -79,21 +93,30 @@ namespace ChaosDD::Private
 
 	void FChaosDDTaskContext::BeginThread(const FChaosDDTaskParentContext& InParentDDContext)
 	{
+		check(!bInContext);
+
 		// NOTE: (UE-216178) We used to pass a reference to the parent FChaosDDContext directly to the
 		// child thread and pulled the FramePointer from it in BeginThread. That is not safe because 
 		// the parent thread may also be helping with tasks and so the Frame on that context will be 
 		// getting set/unset. Instead we copy the Frame pointer on the parent thread and pass it in.
-
-		FChaosDDContext& Context = FChaosDDContext::Get();
-		PreviousFrame = Context.Frame;
-		Context.Frame = InParentDDContext.Frame;
+		if (FChaosDDContext::IsDebugDrawEnabled())
+		{
+			FChaosDDContext& Context = FChaosDDContext::Get();
+			PreviousFrame = Context.Frame;
+			Context.Frame = InParentDDContext.Frame;
+			bInContext = true;
+		}
 	}
 
 	void FChaosDDTaskContext::EndThread()
 	{
-		FChaosDDContext& Context = FChaosDDContext::Get();
-		Context.Frame = PreviousFrame;
-		PreviousFrame.Reset();
+		if (bInContext)
+		{
+			FChaosDDContext& Context = FChaosDDContext::Get();
+			Context.Frame = PreviousFrame;
+			PreviousFrame.Reset();
+			bInContext = false;
+		}
 	}
 
 	FChaosDDScopeTaskContext::FChaosDDScopeTaskContext(const FChaosDDTaskParentContext& InParentDDContext)
diff --git a/Engine/Source/Runtime/Experimental/Chaos/Public/Chaos/DebugDraw/DebugDrawPrimitives.h b/Engine/Source/Runtime/Experimental/Chaos/Public/Chaos/DebugDraw/DebugDrawPrimitives.h
--- a/Engine/Source/Runtime/Experimental/Chaos/Public/Chaos/DebugDraw/DebugDrawPrimitives.h
+++ b/Engine/Source/Runtime/Experimental/Chaos/Public/Chaos/DebugDraw/DebugDrawPrimitives.h
@@ -28,11 +28,15 @@ namespace Chaos
 	//
 	struct CHAOS_API FChaosDDPrimitives
 	{
+		static void DrawPoint(const FVector3d& Position, const FColor& Color, float PointSize, float Duration);
 		static void DrawLine(const FVector3d& A, const FVector3d& B, const FColor& Color, float LineThickness, float Duration);
+		static void DrawArrow(const FVector3d& A, const FVector3d& B, float ArrowSize, const FColor& Color, float LineThickness, float Duration);
+		static void DrawCircle(const FVector3d& Center, const FMatrix& Axes, float Radius, const FColor& Color, float LineThickness, float Duration);
 		static void DrawSphere(const FVector3d& Center, const float Radius, const FColor& Color, float LineThickness, float Duration);
 		static void DrawCapsule(const FVector3d& Center, const FQuat4d& Rotation, float HalfHeight, float Radius, const FColor& Color, float LineThickness, float Duration);
 		static void DrawBox(const FVector3d& Center, const FQuat4d& Rotation, const FVector3d& Size, const FColor& Color, float LineThickness, float Duration);
 		static void DrawTriangle(const FVector3d& A, const FVector3d& B, const FVector3d& C, const FColor& Color, float LineThickness, float Duration);
+		static void DrawString(const FVector3d& TextLocation, const FString& Text, const FColor& Color, float FontScale, bool bDrawShadow, float Duration);
 	};
 }
 
diff --git a/Engine/Source/Runtime/Experimental/Chaos/Public/ChaosDebugDraw/ChaosDDContext.h b/Engine/Source/Runtime/Experimental/Chaos/Public/ChaosDebugDraw/ChaosDDContext.h
--- a/Engine/Source/Runtime/Experimental/Chaos/Public/ChaosDebugDraw/ChaosDDContext.h
+++ b/Engine/Source/Runtime/Experimental/Chaos/Public/ChaosDebugDraw/ChaosDDContext.h
@@ -20,11 +20,29 @@ namespace ChaosDD::Private
 	public:
 		FChaosDDContext();
 
+		// Whether the debug draw system is enabled
+		static bool IsDebugDrawEnabled()
+		{
+			return bDebugDrawEnabled;
+		}
+
+		// Enable/Disable the debug draw system
+		static void SetIsDebugDrawEnabled(bool bInEnabled)
+		{
+			bDebugDrawEnabled = bInEnabled;
+		}
+
+		// Write access to the debug draw frame for the current thread
 		static FChaosDDFrameWriter GetWriter()
 		{
-			return FChaosDDFrameWriter(Get().GetFrame());
+			if (bDebugDrawEnabled)
+			{
+				return FChaosDDFrameWriter(Get().GetFrame());
+			}
+			return FChaosDDFrameWriter({});
 		}
 
+		// For internal use - collect all out-of-frame debug draw commands for rendering
 		static FChaosDDFramePtr ExtractGlobalFrame();
 
 	private:
@@ -57,6 +75,9 @@ namespace ChaosDD::Private
 		// The frame to draw to on this thread (or null)
 		FChaosDDFramePtr Frame;
 
+		// Whether the system is enabled
+		static bool bDebugDrawEnabled;
+
 		// Global frame: fallback for out-of-context debug draw
 		static FCriticalSection GlobalFrameCS;
 		static FChaosDDFramePtr GlobalFrame;
@@ -79,6 +100,7 @@ namespace ChaosDD::Private
 	private:
 		FChaosDDTimelinePtr Timeline;
 		FChaosDDFramePtr PreviousFrame;
+		bool bInContext = false;
 	};
 
 	//
@@ -134,6 +156,7 @@ namespace ChaosDD::Private
 
 	private:
 		FChaosDDFramePtr PreviousFrame;
+		bool bInContext = false;
 	};
 
 	//
diff --git a/Engine/Source/Runtime/Experimental/Chaos/Public/ChaosDebugDraw/ChaosDDFrame.h b/Engine/Source/Runtime/Experimental/Chaos/Public/ChaosDebugDraw/ChaosDDFrame.h
--- a/Engine/Source/Runtime/Experimental/Chaos/Public/ChaosDebugDraw/ChaosDDFrame.h
+++ b/Engine/Source/Runtime/Experimental/Chaos/Public/ChaosDebugDraw/ChaosDDFrame.h
@@ -126,7 +126,7 @@ namespace ChaosDD::Private
 
 		bool WasCommandBudgetExceeded() const
 		{
-			return CommandCost > CommandBudget;
+			return (CommandCost > CommandBudget) && (CommandBudget > 0);
 		}
 
 		bool AddToCost(int32 InCost)
diff --git a/Engine/Source/Runtime/Experimental/Chaos/Public/ChaosDebugDraw/ChaosDDRenderer.h b/Engine/Source/Runtime/Experimental/Chaos/Public/ChaosDebugDraw/ChaosDDRenderer.h
--- a/Engine/Source/Runtime/Experimental/Chaos/Public/ChaosDebugDraw/ChaosDDRenderer.h
+++ b/Engine/Source/Runtime/Experimental/Chaos/Public/ChaosDebugDraw/ChaosDDRenderer.h
@@ -29,11 +29,15 @@ namespace ChaosDD::Private
 		virtual FSphere3d GetDrawRegion() const = 0;
 
 		// Utility functions for use by Debug Draw commands (e.g., FChaosDDLine)
+		virtual void RenderPoint(const FVector3d& Position, const FColor& Color, float PointSize, float Lifetime) = 0;
 		virtual void RenderLine(const FVector3d& A, const FVector3d& B, const FColor& Color, float LineThickness, float Lifetime) = 0;
+		virtual void RenderArrow(const FVector3d& A, const FVector3d& B, float ArrowSize, const FColor& Color, float LineThickness, float Lifetime) = 0;
+		virtual void RenderCircle(const FVector3d& Center, const FMatrix& Axes, float Radius, const FColor& Color, float LineThickness, float Lifetime) = 0;
 		virtual void RenderSphere(const FVector3d& Center, float Radius, const FColor& Color, float LineThickness, float Lifetime) = 0;
 		virtual void RenderCapsule(const FVector3d& Center, const FQuat4d& Rotation, float HalfHeight, float Radius, const FColor& Color, float LineThickness, float Lifetime) = 0;
 		virtual void RenderBox(const FVector3d& Position, const FQuat4d& Rotation, const FVector3d& Size, const FColor& Color, float LineThickness, float Lifetime) = 0;
 		virtual void RenderTriangle(const FVector3d& A, const FVector3d& B, const FVector3d& C, const FColor& Color, float LineThickness, float Lifetime) = 0;
+		virtual void RenderString(const FVector3d& TextLocation, const FString& Text, const FColor& Color, float FontScale, bool bDrawShadow, float Lifetime) = 0;
 
 		// Render legacy debug draw command (See FChaosDDScene::RenderLatestFrames)
 		virtual void RenderLatentCommand(const Chaos::FLatentDrawCommand& Command) = 0;
diff --git a/Engine/Source/Runtime/Experimental/ChaosSolverEngine/Private/Chaos/ChaosDebugDrawSubsystem.cpp b/Engine/Source/Runtime/Experimental/ChaosSolverEngine/Private/Chaos/ChaosDebugDrawSubsystem.cpp
--- a/Engine/Source/Runtime/Experimental/ChaosSolverEngine/Private/Chaos/ChaosDebugDrawSubsystem.cpp
+++ b/Engine/Source/Runtime/Experimental/ChaosSolverEngine/Private/Chaos/ChaosDebugDrawSubsystem.cpp
@@ -19,6 +19,7 @@ FDelegateHandle UChaosDebugDrawSubsystem::OnTickWorldStartDelegate;
 FDelegateHandle UChaosDebugDrawSubsystem::OnTickWorldEndDelegate;
 
 // @todo(chaos): Stuff in ChaosDebugDrawComponent that should move here when we drop that file
+extern bool bChaosDebugDraw_Enabled;
 extern float ChaosDebugDraw_Radius;
 extern int32 bChaosDebugDraw_DrawMode;
 extern int32 ChaosDebugDraw_MaxElements;
@@ -39,6 +40,7 @@ public:
 		: World(InWorld)
 		, DrawRegion(InDrawRegion)
 		, RenderBudget(InRenderBudget)
+		, RenderCost(0)
 		, SphereSegments(8)
 		, DepthPriority(10)
 	{
@@ -54,19 +56,79 @@ public:
 		return DrawRegion;
 	}
 
+	int32 GetRenderCost() const
+	{
+		return RenderCost;
+	}
+
+	int32 GetRenderBudget() const
+	{
+		return RenderBudget;
+	}
+
+	bool WasRenderBudgetExceeded() const
+	{
+		return (RenderCost > RenderBudget) && (RenderBudget > 0);
+	}
+
+	bool IsInDrawRegion(const FBox3d& Bounds) const
+	{
+		const double DistanceSq = Bounds.ComputeSquaredDistanceToPoint(DrawRegion.Center);
+		return DistanceSq <= FMath::Square(DrawRegion.W);
+	}
+
+	virtual void RenderPoint(const FVector3d& Position, const FColor& Color, float PointSize, float Lifetime) override final
+	{
+		constexpr int32 Cost = 1;
+		const FBox3d Bounds = FBox3d(Position, Position);
+
+		if (IsInDrawRegion(Bounds) && TryAddToCost(Cost))
+		{
+			DrawDebugPoint(World, Position, PointSize, Color, false, CommandLifeTime(Lifetime), DepthPriority);
+		}
+	}
+
 	virtual void RenderLine(const FVector3d& A, const FVector3d& B, const FColor& Color, float LineThickness, float Lifetime) override final
 	{
 		constexpr int32 Cost = 1;
-		if (TryAddToCost(Cost))
+		const FBox3d Bounds = FBox3d(FVector3d::Min(A, B), FVector3d::Max(A, B));
+
+		if (IsInDrawRegion(Bounds) && TryAddToCost(Cost))
 		{
 			DrawDebugLine(World, A, B, Color, false, CommandLifeTime(Lifetime), DepthPriority, LineThickness);
 		}
 	}
 
+	virtual void RenderArrow(const FVector3d& A, const FVector3d& B, float ArrowSize, const FColor& Color, float LineThickness, float Lifetime) override final
+	{
+		constexpr int32 Cost = 3;
+		const FBox3d Bounds = FBox3d(FVector3d::Min(A, B), FVector3d::Max(A, B));
+
+		if (IsInDrawRegion(Bounds) && TryAddToCost(Cost))
+		{
+			DrawDebugDirectionalArrow(World, A, B, ArrowSize, Color, false, CommandLifeTime(Lifetime), DepthPriority, LineThickness);
+		}
+	}
+
+	virtual void RenderCircle(const FVector3d& Center, const FMatrix& Axes, float Radius, const FColor& Color, float LineThickness, float Lifetime) override final
+	{
+		constexpr int32 Cost = 8;
+		const FBox3d Bounds = FBox3d(Center - FVector3d(Radius), Center + FVector3d(Radius));
+
+		if (IsInDrawRegion(Bounds) && TryAddToCost(Cost))
+		{
+			FMatrix M = Axes;
+			M.SetOrigin(Center);
+			DrawDebugCircle(World, M, Radius, SphereSegments, Color, false, CommandLifeTime(Lifetime), DepthPriority, LineThickness, false);
+		}
+	}
+
 	virtual void RenderSphere(const FVector3d& Center, float Radius, const FColor& Color, float LineThickness, float Lifetime) override final
 	{
 		constexpr int32 Cost = 16;
-		if (TryAddToCost(Cost))
+		const FBox3d Bounds = FBox3d(Center - FVector3d(Radius), Center + FVector3d(Radius));
+
+		if (IsInDrawRegion(Bounds) && TryAddToCost(Cost))
 		{
 			DrawDebugSphere(World, Center, Radius, SphereSegments, Color, false, CommandLifeTime(Lifetime), DepthPriority, LineThickness);
 		}
@@ -75,7 +137,12 @@ public:
 	virtual void RenderCapsule(const FVector3d& Center, const FQuat4d& Rotation, float HalfHeight, float Radius, const FColor& Color, float LineThickness, float Lifetime) override final
 	{
 		constexpr int32 Cost = 16;
-		if (TryAddToCost(Cost))
+		const FVector3d EndOffset = HalfHeight * (Rotation * FVector3d::UnitZ());
+		const FVector3d A = Center - EndOffset;
+		const FVector3d B = Center + EndOffset;
+		const FBox3d Bounds = FBox3d(FVector3d::Min(A, B), FVector3d::Max(A, B));
+
+		if (IsInDrawRegion(Bounds) && TryAddToCost(Cost))
 		{
 			DrawDebugCapsule(World, Center, HalfHeight, Radius, Rotation, Color, false, CommandLifeTime(Lifetime), DepthPriority, LineThickness);
 		}
@@ -84,7 +151,9 @@ public:
 	virtual void RenderBox(const FVector3d& Position, const FQuat4d& Rotation, const FVector3d& Size, const FColor& Color, float LineThickness, float Lifetime) override final
 	{
 		constexpr int32 Cost = 12;
-		if (TryAddToCost(Cost))
+		const FBox3d Bounds = FBox3d(-0.5 * Size, 0.5 * Size).TransformBy(FTransform(Rotation, Position));
+
+		if (IsInDrawRegion(Bounds) && TryAddToCost(Cost))
 		{
 			DrawDebugBox(World, Position, Size, Rotation, Color, false, CommandLifeTime(Lifetime), DepthPriority, LineThickness);
 		}
@@ -93,7 +162,9 @@ public:
 	virtual void RenderTriangle(const FVector3d& A, const FVector3d& B, const FVector3d& C, const FColor& Color, float LineThickness, float Lifetime) override final
 	{
 		constexpr int32 Cost = 3;
-		if (TryAddToCost(Cost))
+		const FBox3d Bounds = FBox3d(FVector3d::Min3(A, B, C), FVector3d::Max3(A, B, C));
+
+		if (IsInDrawRegion(Bounds) && TryAddToCost(Cost))
 		{
 			DrawDebugLine(World, A, B, Color, false, CommandLifeTime(Lifetime), DepthPriority, LineThickness);
 			DrawDebugLine(World, B, C, Color, false, CommandLifeTime(Lifetime), DepthPriority, LineThickness);
@@ -101,6 +172,17 @@ public:
 		}
 	}
 
+	virtual void RenderString(const FVector3d& TextLocation, const FString& Text, const FColor& Color, float FontScale, bool bDrawShadow, float Lifetime) override final
+	{
+		constexpr int32 Cost = 10;
+		const FBox3d Bounds = FBox3d(TextLocation, TextLocation);
+
+		if (IsInDrawRegion(Bounds) && TryAddToCost(Cost))
+		{
+			DrawDebugString(World, TextLocation, Text, nullptr, Color, CommandLifeTime(Lifetime), bDrawShadow, FontScale);
+		}
+	}
+
 	virtual void RenderLatentCommand(const Chaos::FLatentDrawCommand& Command) override final
 	{
 		// We don't check the cost of the legacy commands because the budget is applied at capture time for them
@@ -263,6 +345,9 @@ void UChaosDebugDrawSubsystem::OnWorldTickStart(ELevelTick TickType, float Dt)
 {
 	if (UWorld* World = GetWorld())
 	{
+		// Enable or disable the debug draw system
+		ChaosDD::Private::FChaosDDContext::SetIsDebugDrawEnabled(bChaosDebugDraw_Enabled);
+
 		CDDWorldTimelineContext.BeginFrame(CDDWorldTimeline, World->GetTimeSeconds(), Dt);
 	}
 }
@@ -307,7 +392,7 @@ void UChaosDebugDrawSubsystem::RenderScene()
 		{
 			if (!World->IsPaused())
 			{
-				// @todo(chaos): 
+				// @todo(chaos): command budget and render budget should be two different values
 				FChaosDDRenderer Renderer = FChaosDDRenderer(World, CDDScene->GetDrawRegion(), CDDScene->GetCommandBudget());
 
 				// Render all of the out-of frame commands
@@ -320,6 +405,13 @@ void UChaosDebugDrawSubsystem::RenderScene()
 
 				// Render the commands from this world
 				RenderScene(Renderer, CDDScene);
+
+				if (Renderer.WasRenderBudgetExceeded())
+				{
+					constexpr int32 MsgId = 86421358;
+					const FString Msg = FString::Format(TEXT("Debug Draw Render Budget Exceeded for {0} [{1} / {2}]"), { CDDScene->GetName(), Renderer.GetRenderCost(), Renderer.GetRenderBudget()});
+					GEngine->AddOnScreenDebugMessage(MsgId, 1.0f, FColor::Red, *Msg);
+				}
 			}
 		}
 	}
@@ -347,7 +439,9 @@ void UChaosDebugDrawSubsystem::RenderScene(FChaosDDRenderer& Renderer, const Cha
 
 void UChaosDebugDrawSubsystem::RenderFrame(FChaosDDRenderer& Renderer, const ChaosDD::Private::FChaosDDFramePtr& Frame)
 {
-	if (Frame.IsValid() && (Frame->GetNumCommands() + Frame->GetNumLatentCommands() > 0))
+	const bool bDebugDrawEnabled = ChaosDD::Private::FChaosDDContext::IsDebugDrawEnabled();
+
+	if (bDebugDrawEnabled && Frame.IsValid() && (Frame->GetNumCommands() + Frame->GetNumLatentCommands() > 0))
 	{
 		UE_LOG(LogChaosDD, VeryVerbose, TEXT("Render %s %d %d+%d Commands"), *Frame->GetTimeline()->GetName(), Frame->GetFrameIndex(), Frame->GetNumCommands(), Frame->GetNumLatentCommands());
 
