Decode container header before serializing dependencies

#rnx

diff --git a/Engine/Source/Runtime/Experimental/IoStore/OnDemand/Private/OnDemandIoStore.cpp b/Engine/Source/Runtime/Experimental/IoStore/OnDemand/Private/OnDemandIoStore.cpp
--- a/Engine/Source/Runtime/Experimental/IoStore/OnDemand/Private/OnDemandIoStore.cpp
+++ b/Engine/Source/Runtime/Experimental/IoStore/OnDemand/Private/OnDemandIoStore.cpp
@@ -18,6 +18,7 @@
 #include "HAL/PlatformFileManager.h"
 #include "HAL/PlatformTime.h"
 #include "HAL/PlatformProcess.h"
+#include "IO/IoChunkEncoding.h"
 #include "IO/IoContainerHeader.h"
 #include "IO/PackageStore.h"
 #include "Misc/CommandLine.h"
@@ -824,26 +825,47 @@ FIoStatus FOnDemandIoStore::TickInstallRequest(FMountRequest& MountRequest)
 	};
 
 	auto FetchContainerHeader = [&Host, &GetChunkUrl, &ChunkUrl](
-		const FOnDemandContainer& Container,
+		FSharedOnDemandContainer Container,
 		FIoContainerHeader& Header) -> FIoStatus 
 	{
-		const FIoChunkId			ChunkId = CreateContainerHeaderChunkId(Container.ContainerId);
-		const FOnDemandChunkEntry*	Entry = Container.ChunkEntries.Find(ChunkId);
+		const FIoChunkId			ChunkId = CreateContainerHeaderChunkId(Container->ContainerId);
+		const FOnDemandChunkEntry*	Entry = Container->ChunkEntries.Find(ChunkId);
 
 		if (Entry == nullptr)
 		{
 			return EIoErrorCode::Ok;
 		}
 
-		UE_LOG(LogIoStoreOnDemand, VeryVerbose, TEXT("Fetching container header, ContainerName='%s'"), *Container.Name);
-		TIoStatusOr<FIoBuffer> Response = FHttpClient::Get(GetChunkUrl(Host, Container, *Entry, ChunkUrl).ToView(), 2, EHttpRedirects::Follow);
+		UE_LOG(LogIoStoreOnDemand, Verbose, TEXT("Fetching container header, ContainerName='%s'"), *Container->Name);
+		TIoStatusOr<FIoBuffer> Response = FHttpClient::Get(GetChunkUrl(Host, *Container, *Entry, ChunkUrl).ToView(), 2, EHttpRedirects::Follow);
+
 		if (Response.IsOk() == false)
 		{
-			FIoStatus Status = FIoStatusBuilder(EIoErrorCode::InvalidCode) << TEXT("Failed to fetch container header chunk from URL");
+			FIoStatus Status = FIoStatusBuilder(EIoErrorCode::ReadError) << TEXT("Failed to fetch container header chunk");
 			return Status;
 		}
 
-		FMemoryReaderView Ar(Response.ValueOrDie().GetView());
+		FOnDemandChunkInfo ChunkInfo(Container, *Entry);
+
+		FIoChunkDecodingParams Params;
+		Params.CompressionFormat	= ChunkInfo.CompressionFormat();
+		Params.EncryptionKey		= ChunkInfo.EncryptionKey();
+		Params.BlockSize			= ChunkInfo.BlockSize();
+		Params.TotalRawSize			= ChunkInfo.RawSize();
+		Params.RawOffset			= 0; 
+		Params.EncodedOffset		= 0; 
+		Params.EncodedBlockSize		= ChunkInfo.Blocks();
+		Params.BlockHash			= ChunkInfo.BlockHashes();
+
+		FIoBuffer EncodedChunk = Response.ConsumeValueOrDie();
+		FIoBuffer RawChunk(ChunkInfo.RawSize());
+		if (FIoChunkEncoding::Decode(Params, EncodedChunk.GetView(), RawChunk.GetMutableView()) == false)
+		{
+			FIoStatus Status = FIoStatusBuilder(EIoErrorCode::ReadError)
+				<< TEXT("Failed to decode container header chunk");
+		}
+
+		FMemoryReaderView Ar(RawChunk.GetView());
 		Ar << Header;
 		Ar.Close();
 
@@ -868,7 +890,7 @@ FIoStatus FOnDemandIoStore::TickInstallRequest(FMountRequest& MountRequest)
 		if (!Container->Header.IsValid())
 		{
 			FSharedContainerHeader Header = MakeShared<FIoContainerHeader>();
-			if (FIoStatus Status = FetchContainerHeader(*Container, *Header); !Status.IsOk())
+			if (FIoStatus Status = FetchContainerHeader(Container, *Header); !Status.IsOk())
 			{
 				return Status; 
 			}
