Chunk-ext test path.

#rnx

diff --git a/Engine/Source/Runtime/Experimental/IoStore/HttpClient/Private/TestServer.py b/Engine/Source/Runtime/Experimental/IoStore/HttpClient/Private/TestServer.py
--- a/Engine/Source/Runtime/Experimental/IoStore/HttpClient/Private/TestServer.py
+++ b/Engine/Source/Runtime/Experimental/IoStore/HttpClient/Private/TestServer.py
@@ -217,7 +217,13 @@ def http_data(handler, payload_size=0):
     handler.end_headers()
     handler.wfile.write(payload)
 
-def http_chunked(handler, payload_size=0):
+def http_chunked(handler, payload_size=0, *options):
+    ext_payload = b""
+    if "ext" in options:
+        n = int(random.random() * 32)
+        ext_payload = "".join(random.choices("Trigrams; Diner", k=n))
+        ext_payload = b";" + ext_payload.encode()
+
     payload, payload_hash = _make_payload(payload_size)
     handler.send_response(200)
     handler.send_header("Transfer-Encoding", "chunked")
@@ -234,7 +240,7 @@ def http_chunked(handler, payload_size=0):
         header = b"%x" % len(piece)
         if piece and piece[0] & 0b0100:
             header = header.upper()
-        header += b"\r\n"
+        header += ext_payload + b"\r\n"
 
         handler.wfile.write(header)
         handler.wfile.write(piece)
