Add container flags to on-demand container TOC

#rnx

diff --git a/Engine/Source/Runtime/Experimental/IoStore/OnDemand/Private/IoStoreOnDemand.cpp b/Engine/Source/Runtime/Experimental/IoStore/OnDemand/Private/IoStoreOnDemand.cpp
--- a/Engine/Source/Runtime/Experimental/IoStore/OnDemand/Private/IoStoreOnDemand.cpp
+++ b/Engine/Source/Runtime/Experimental/IoStore/OnDemand/Private/IoStoreOnDemand.cpp
@@ -513,13 +513,15 @@ bool LoadFromCompactBinary(FCbFieldView Field, FOnDemandTocEntry& OutTocEntry)
 
 FArchive& operator<<(FArchive& Ar, FOnDemandTocContainerEntry& ContainerEntry)
 {
+	EOnDemandTocVersion TocVersion = EOnDemandTocVersion::Latest;
+
 	if (Ar.IsLoading())
 	{
 		const FCustomVersion* CustomVersion = Ar.GetCustomVersions().GetVersion(FOnDemandToc::VersionGuid);
 		check(CustomVersion);
-		const uint32 TocVersion = IntCastChecked<uint32>(CustomVersion->Version);
+		TocVersion = static_cast<EOnDemandTocVersion>(CustomVersion->Version);
 
-		if (TocVersion >= uint32(EOnDemandTocVersion::ContainerId))
+		if (TocVersion >= EOnDemandTocVersion::ContainerId)
 		{
 			Ar << ContainerEntry.ContainerId;
 		}
@@ -536,6 +538,11 @@ FArchive& operator<<(FArchive& Ar, FOnDemandTocContainerEntry& ContainerEntry)
 	Ar << ContainerEntry.BlockHashes;
 	Ar << ContainerEntry.UTocHash;
 
+	if (!Ar.IsLoading() || (TocVersion >= EOnDemandTocVersion::ContainerFlags))
+	{
+		Ar << ContainerEntry.ContainerFlags;
+	}
+
 	return Ar;
 }
 
diff --git a/Engine/Source/Runtime/Experimental/IoStore/OnDemand/Public/IO/IoStoreOnDemand.h b/Engine/Source/Runtime/Experimental/IoStore/OnDemand/Public/IO/IoStoreOnDemand.h
--- a/Engine/Source/Runtime/Experimental/IoStore/OnDemand/Public/IO/IoStoreOnDemand.h
+++ b/Engine/Source/Runtime/Experimental/IoStore/OnDemand/Public/IO/IoStoreOnDemand.h
@@ -60,6 +60,7 @@ enum class EOnDemandTocVersion : uint32
 	ContainerId		= 6,
 	AdditionalFiles	= 7,
 	TagSets			= 8,
+	ContainerFlags	= 9,
 
 	LatestPlusOne,
 	Latest			= (LatestPlusOne - 1)
@@ -129,6 +130,7 @@ struct FOnDemandTocContainerEntry
 
 	/** Hash of the .utoc file (on disk) used to generate this data */
 	FIoHash UTocHash;
+	uint8 ContainerFlags = 0;
 
 	UE_API friend FArchive& operator<<(FArchive& Ar, FOnDemandTocContainerEntry& ContainerEntry);
 	UE_API friend FCbWriter& operator<<(FCbWriter& Writer, const FOnDemandTocContainerEntry& ContainerEntry);
