Add container flags to on demand container TOC when uploding chunks

#rnx

diff --git a/Engine/Source/Runtime/Experimental/IoStore/OnDemand/Private/Tool/Upload.cpp b/Engine/Source/Runtime/Experimental/IoStore/OnDemand/Private/Tool/Upload.cpp
--- a/Engine/Source/Runtime/Experimental/IoStore/OnDemand/Private/Tool/Upload.cpp
+++ b/Engine/Source/Runtime/Experimental/IoStore/OnDemand/Private/Tool/Upload.cpp
@@ -533,6 +533,7 @@ static TIoStatusOr<FUploadResult> UploadContainerFiles(
 		FOnDemandTocContainerEntry& ContainerEntry = OnDemandToc.Containers.AddDefaulted_GetRef();
 		ContainerEntry.ContainerName = FPaths::GetBaseFilename(Path);
 		ContainerEntry.ContainerId = ContainerFileReader.GetContainerId();
+		ContainerEntry.ContainerFlags = static_cast<uint8>(ContainerFileReader.GetContainerFlags());
 
 		if (EnumHasAnyFlags(ContainerFileReader.GetContainerFlags(), EIoContainerFlags::Encrypted))
 		{
