Iris
- Changed RotatorNetSerializer to use same quantization functions as NetSerialize

diff --git a/Engine/Source/Runtime/Experimental/Iris/Core/Private/Iris/Serialization/RotatorNetSerializers.cpp b/Engine/Source/Runtime/Experimental/Iris/Core/Private/Iris/Serialization/RotatorNetSerializers.cpp
--- a/Engine/Source/Runtime/Experimental/Iris/Core/Private/Iris/Serialization/RotatorNetSerializers.cpp
+++ b/Engine/Source/Runtime/Experimental/Iris/Core/Private/Iris/Serialization/RotatorNetSerializers.cpp
@@ -46,10 +46,6 @@ protected:
 
 		XYZDiffersMask = 7U,
 	};
-
-	static constexpr ScalarType Scale = ScalarType(65536)/ScalarType(360);
-	static constexpr ScalarType Bias = ScalarType(0.5f);
-	static constexpr ScalarType InvScale = ScalarType(360)/ScalarType(65536);
 };
 
 template<typename RotatorType>
@@ -90,10 +86,6 @@ protected:
 
 		XYZDiffersMask = 7U,
 	};
-
-	static constexpr ScalarType Scale = ScalarType(256)/ScalarType(360);
-	static constexpr ScalarType Bias = ScalarType(0.5f);
-	static constexpr ScalarType InvScale = ScalarType(360)/ScalarType(256);
 };
 
 // FRotatorAsShortNetSerializerBase implementation
@@ -212,9 +204,11 @@ void FRotatorAsShortNetSerializerBase<T>::Quantize(FNetSerializationContext&, co
 	QuantizedType& Target = *reinterpret_cast<QuantizedType*>(Args.Target);
 
 	QuantizedType TempValue = {};
-	TempValue.X = uint16(uint32(int32(Source.Pitch*Scale + Bias)));
-	TempValue.Y = uint16(uint32(int32(Source.Yaw*Scale + Bias)));
-	TempValue.Z = uint16(uint32(int32(Source.Roll*Scale + Bias)));
+
+	TempValue.X = SourceType::CompressAxisToShort(Source.Pitch);
+	TempValue.Y = SourceType::CompressAxisToShort(Source.Yaw);
+	TempValue.Z = SourceType::CompressAxisToShort(Source.Roll);
+
 	TempValue.XYZIsNotZero |= (TempValue.X != 0) ? XDiffersMask : uint16(0);
 	TempValue.XYZIsNotZero |= (TempValue.Y != 0) ? YDiffersMask : uint16(0);
 	TempValue.XYZIsNotZero |= (TempValue.Z != 0) ? ZDiffersMask : uint16(0);
@@ -229,9 +223,10 @@ void FRotatorAsShortNetSerializerBase<T>::Dequantize(FNetSerializationContext&,
 	SourceType& Target = *reinterpret_cast<SourceType*>(Args.Target);
 
 	SourceType TempValue;
-	TempValue.Pitch = Source.X*InvScale;
-	TempValue.Yaw = Source.Y*InvScale;
-	TempValue.Roll = Source.Z*InvScale;
+
+	TempValue.Pitch = SourceType::DecompressAxisFromShort(Source.X);
+	TempValue.Yaw = SourceType::DecompressAxisFromShort(Source.Y);
+	TempValue.Roll = SourceType::DecompressAxisFromShort(Source.Z);
 
 	Target = TempValue;
 }
@@ -266,8 +261,7 @@ template<typename T>
 bool FRotatorAsShortNetSerializerBase<T>::Validate(FNetSerializationContext& Context, const FNetValidateArgs& Args)
 {
 	const SourceType& Value = *reinterpret_cast<SourceType*>(Args.Source);
-	// Make sure values are valid degree angles. This should catch NaNs as well.
-	return (Value.Pitch >= ScalarType(0) && Value.Pitch < ScalarType(360)) && (Value.Yaw >= ScalarType(0) && Value.Yaw < ScalarType(360)) & (Value.Roll >= ScalarType(0) && Value.Roll < ScalarType(360)); 
+	return !Value.ContainsNaN();
 }
 
 // FRotatorAsByteNetSerializerBase implementation
@@ -375,9 +369,11 @@ void FRotatorAsByteNetSerializerBase<T>::Quantize(FNetSerializationContext&, con
 	QuantizedType& Target = *reinterpret_cast<QuantizedType*>(Args.Target);
 
 	QuantizedType TempValue = {};
-	TempValue.X = uint8(uint32(int32(Source.Pitch*Scale + Bias)));
-	TempValue.Y = uint8(uint32(int32(Source.Yaw*Scale + Bias)));
-	TempValue.Z = uint8(uint32(int32(Source.Roll*Scale + Bias)));
+
+	TempValue.X = SourceType::CompressAxisToByte(Source.Pitch);
+	TempValue.Y = SourceType::CompressAxisToByte(Source.Yaw);
+	TempValue.Z = SourceType::CompressAxisToByte(Source.Roll);
+
 	TempValue.XYZIsNotZero |= (TempValue.X != 0) ? XDiffersMask : uint8(0);
 	TempValue.XYZIsNotZero |= (TempValue.Y != 0) ? YDiffersMask : uint8(0);
 	TempValue.XYZIsNotZero |= (TempValue.Z != 0) ? ZDiffersMask : uint8(0);
@@ -392,9 +388,9 @@ void FRotatorAsByteNetSerializerBase<T>::Dequantize(FNetSerializationContext&, c
 	SourceType& Target = *reinterpret_cast<SourceType*>(Args.Target);
 
 	SourceType TempValue;
-	TempValue.Pitch = Source.X*InvScale;
-	TempValue.Yaw = Source.Y*InvScale;
-	TempValue.Roll = Source.Z*InvScale;
+	TempValue.Pitch = SourceType::DecompressAxisFromByte(Source.X);
+	TempValue.Yaw = SourceType::DecompressAxisFromByte(Source.Y);
+	TempValue.Roll = SourceType::DecompressAxisFromByte(Source.Z);
 
 	Target = TempValue;
 }
@@ -429,8 +425,7 @@ template<typename T>
 bool FRotatorAsByteNetSerializerBase<T>::Validate(FNetSerializationContext& Context, const FNetValidateArgs& Args)
 {
 	const SourceType& Value = *reinterpret_cast<SourceType*>(Args.Source);
-	// Make sure values are valid degree angles. This should catch NaNs as well.
-	return (Value.Pitch >= ScalarType(0) && Value.Pitch < ScalarType(360)) && (Value.Yaw >= ScalarType(0) && Value.Yaw < ScalarType(360)) & (Value.Roll >= ScalarType(0) && Value.Roll < ScalarType(360)); 
+	return !Value.ContainsNaN();
 }
 
 struct FRotatorNetSerializer : public FRotatorAsShortNetSerializerBase<FRotator>
