Dataflow - Fixed crash when undoing AddOptionPin.

[RN] Dataflow - Fixed a crash when undoing after an Add Option Pin.

diff --git a/Engine/Source/Editor/GraphEditor/Private/ConnectionDrawingPolicy.cpp b/Engine/Source/Editor/GraphEditor/Private/ConnectionDrawingPolicy.cpp
--- a/Engine/Source/Editor/GraphEditor/Private/ConnectionDrawingPolicy.cpp
+++ b/Engine/Source/Editor/GraphEditor/Private/ConnectionDrawingPolicy.cpp
@@ -477,7 +477,7 @@ void FConnectionDrawingPolicy::DrawPinGeometries(TMap<TSharedRef<SWidget>, FArra
 		SGraphPin& PinWidget = static_cast<SGraphPin&>(SomePinWidget.Get());
 		UEdGraphPin* ThePin = PinWidget.GetPinObj();
 
-		if (ThePin->Direction == EGPD_Output)
+		if (ThePin && ThePin->Direction == EGPD_Output)
 		{
 			for (int32 LinkIndex=0; LinkIndex < ThePin->LinkedTo.Num(); ++LinkIndex)
 			{
diff --git a/Engine/Source/Runtime/Experimental/Dataflow/Engine/Private/Dataflow/DataflowEdNode.cpp b/Engine/Source/Runtime/Experimental/Dataflow/Engine/Private/Dataflow/DataflowEdNode.cpp
--- a/Engine/Source/Runtime/Experimental/Dataflow/Engine/Private/Dataflow/DataflowEdNode.cpp
+++ b/Engine/Source/Runtime/Experimental/Dataflow/Engine/Private/Dataflow/DataflowEdNode.cpp
@@ -629,6 +629,15 @@ bool UDataflowEdNode::ShouldDrawNodeAsControlPointOnly(int32& OutInputPinIndex,
 	return false;
 }
 
+void UDataflowEdNode::PostEditUndo()
+{
+	// Refresh the current graph, so the pins or whatever happened to this object can be reflected to the graph
+	if (UEdGraph* const ParentGraph = GetGraph())
+	{
+		ParentGraph->NotifyGraphChanged();
+	}
+}
+
 void UDataflowEdNode::HideAllInputPins()
 {
 	bool bAnyHidden = false;
diff --git a/Engine/Source/Runtime/Experimental/Dataflow/Engine/Public/Dataflow/DataflowEdNode.h b/Engine/Source/Runtime/Experimental/Dataflow/Engine/Public/Dataflow/DataflowEdNode.h
--- a/Engine/Source/Runtime/Experimental/Dataflow/Engine/Public/Dataflow/DataflowEdNode.h
+++ b/Engine/Source/Runtime/Experimental/Dataflow/Engine/Public/Dataflow/DataflowEdNode.h
@@ -50,6 +50,7 @@ public:
 	DATAFLOWENGINE_API virtual void AutowireNewNode(UEdGraphPin* FromPin) override;
 	DATAFLOWENGINE_API virtual void OnPinRemoved(UEdGraphPin* InRemovedPin) override;
 	DATAFLOWENGINE_API virtual bool ShouldDrawNodeAsControlPointOnly(int32& OutInputPinIndex, int32& OutOutputPinIndex) const override;
+	DATAFLOWENGINE_API virtual void PostEditUndo() override;
 #endif // WITH_EDITOR
 	// End of UEdGraphNode interface
 
