Retrieve/wait for encryption key(s) before installing on-demand content

#rnx

diff --git a/Engine/Source/Runtime/Experimental/IoStore/OnDemand/Private/OnDemandIoStore.cpp b/Engine/Source/Runtime/Experimental/IoStore/OnDemand/Private/OnDemandIoStore.cpp
--- a/Engine/Source/Runtime/Experimental/IoStore/OnDemand/Private/OnDemandIoStore.cpp
+++ b/Engine/Source/Runtime/Experimental/IoStore/OnDemand/Private/OnDemandIoStore.cpp
@@ -278,7 +278,8 @@ void LexToString(EOnDemandContainerFlags Flags, FStringBuilderBase& Out)
 		TEXT("PendingEncryptionKey"),
 		TEXT("Mounted"),
 		TEXT("Streaming"),
-		TEXT("Installed")
+		TEXT("Installed"),
+		TEXT("Encrypted")
 	};
 
 	if (Flags == EOnDemandContainerFlags::None)
@@ -629,21 +630,61 @@ bool FOnDemandIoStore::Tick()
 			continue;
 		}
 
+		// Retrieve encryption key(s)
+		bool bPendingEncryptionKeys = false;
+		for (const FSharedOnDemandContainer& Container : Request->Containers)
+		{
+			if (EnumHasAnyFlags(Request->MountArgs.Options, EOnDemandMountOptions::StreamOnDemand))
+			{
+				EnumAddFlags(Container->Flags, EOnDemandContainerFlags::Streaming);
+			}
+
+			if (EnumHasAnyFlags(Container->Flags, EOnDemandContainerFlags::Encrypted) &&
+				Container->EncryptionKey.IsValid() == false)
+			{
+				FGuid KeyGuid;
+				ensure(FGuid::Parse(Container->EncryptionKeyGuid, KeyGuid));
+				if (FEncryptionKeyManager::Get().TryGetKey(KeyGuid, Container->EncryptionKey))
+				{
+					EnumRemoveFlags(Container->Flags, EOnDemandContainerFlags::PendingEncryptionKey);
+				}
+				else
+				{
+					EnumAddFlags(Container->Flags, EOnDemandContainerFlags::PendingEncryptionKey);
+					bPendingEncryptionKeys = true;
+				}
+			}
+		}
+
 		// Tick install
 		if (EnumHasAnyFlags(Request->MountArgs.Options, EOnDemandMountOptions::Install))
 		{
+			if (bPendingEncryptionKeys)
+			{
+				UE_LOG(LogIoStoreOnDemand, Log, TEXT("Deferring install request until encryption key(s) becomes available"));
+				continue;
+			}
+
 			if (FIoStatus Status = TickInstallRequest(*Request); !Status.IsOk())
 			{
+				// Unknown == Pending
+				if (Status.GetErrorCode() == EIoErrorCode::Unknown)
+				{
+					UE_LOG(LogIoStoreOnDemand, Log, TEXT("%s"), *Status.ToString());
+					continue;
+				}
+
 				{
 					UE::TUniqueLock Lock(MountRequestMutex);
 					MountRequests.Remove(Request->MountArgs.MountId);
 				}
 				FOnDemandMountCompleted OnCompleted = MoveTemp(Request->OnCompleted);
 				OnCompleted(Status);
-				continue;;
+				continue;
 			}
 		}
 
+		// Complete mount/install request
 		{
 			{
 				UE::TUniqueLock Lock(MountRequestMutex);
@@ -654,11 +695,6 @@ bool FOnDemandIoStore::Tick()
 			UE::TUniqueLock Lock(ContainerMutex);
 			for (const FSharedOnDemandContainer& Container : Request->Containers)
 			{
-				if (EnumHasAnyFlags(Request->MountArgs.Options, EOnDemandMountOptions::StreamOnDemand))
-				{
-					EnumAddFlags(Container->Flags, EOnDemandContainerFlags::Streaming);
-				}
-
 				if (EnumHasAnyFlags(Container->Flags, EOnDemandContainerFlags::Mounted))
 				{
 					continue;
@@ -666,25 +702,21 @@ bool FOnDemandIoStore::Tick()
 
 				Containers.Add(Container);
 
-				if (Container->EncryptionKeyGuid.IsEmpty() == false)
+				// To conform with PAK mount behavior, requests are completed even if encryption keys are pending
+				if (EnumHasAnyFlags(Container->Flags, EOnDemandContainerFlags::PendingEncryptionKey))
 				{
-					FGuid KeyGuid;
-					ensure(FGuid::Parse(Container->EncryptionKeyGuid, KeyGuid));
-					if (FEncryptionKeyManager::Get().TryGetKey(KeyGuid, Container->EncryptionKey) == false)
-					{
-						UE_LOG(LogIoStoreOnDemand, Log, TEXT("Deferring container '%s' until encryption key '%s' becomes available"),
-							*Container->Name, *Container->EncryptionKeyGuid);
-						EnumAddFlags(Container->Flags, EOnDemandContainerFlags::PendingEncryptionKey);
-						continue;
-					}
+					UE_LOG(LogIoStoreOnDemand, Log, TEXT("Deferring container '%s' until encryption key '%s' becomes available"),
+						*Container->Name, *Container->EncryptionKeyGuid);
 				}
+				else
+				{
+					Sb.Reset();
+					LexToString(Container->Flags, Sb);
+					UE_LOG(LogIoStoreOnDemand, Log, TEXT("Mounting container '%s', Entries=%d, Flags='%s'"),
+						*Container->Name, Container->ChunkEntries.Num(), Sb.ToString());
 
-				Sb.Reset();
-				LexToString(Container->Flags, Sb);
-				UE_LOG(LogIoStoreOnDemand, Log, TEXT("Mounting container '%s', Entries=%d, Flags='%s'"),
-					*Container->Name, Container->ChunkEntries.Num(), Sb.ToString());
-
-				EnumAddFlags(Container->Flags, EOnDemandContainerFlags::Mounted);
+					EnumAddFlags(Container->Flags, EOnDemandContainerFlags::Mounted);
+				}
 			}
 
 			FOnDemandMountCompleted OnCompleted = MoveTemp(Request->OnCompleted);
@@ -887,6 +919,12 @@ FIoStatus FOnDemandIoStore::TickInstallRequest(FMountRequest& MountRequest)
 	// Fetch all container headers
 	for (const FSharedOnDemandContainer& Container : AllContainers)
 	{
+		if (EnumHasAnyFlags(Container->Flags, EOnDemandContainerFlags::PendingEncryptionKey))
+		{
+			return FIoStatusBuilder(EIoErrorCode::Unknown)
+				<< TEXT("Deferring install request until encryption key(s) becomes available");
+		}
+
 		if (!Container->Header.IsValid())
 		{
 			FSharedContainerHeader Header = MakeShared<FIoContainerHeader>();
@@ -1123,25 +1161,30 @@ FIoStatus FOnDemandIoStore::TickInstallRequest(FMountRequest& MountRequest)
 
 void FOnDemandIoStore::OnEncryptionKeyAdded(const FGuid& Id, const FAES::FAESKey& Key)
 {
-	TUniqueLock Lock(ContainerMutex);
-
-	for (FSharedOnDemandContainer& Container : Containers)
 	{
-		if (EnumHasAnyFlags(Container->Flags, EOnDemandContainerFlags::PendingEncryptionKey))
-		{
-			FGuid KeyGuid;
-			ensure(FGuid::Parse(FString(StringCast<TCHAR>(*Container->EncryptionKeyGuid)), KeyGuid));
+		TUniqueLock Lock(ContainerMutex);
 
-			if (FEncryptionKeyManager::Get().TryGetKey(KeyGuid, Container->EncryptionKey))
+		for (FSharedOnDemandContainer& Container : Containers)
+		{
+			if (EnumHasAnyFlags(Container->Flags, EOnDemandContainerFlags::PendingEncryptionKey))
 			{
-				UE_LOG(LogIoStoreOnDemand, Log, TEXT("Mounting container '%s', Entries=%d, Flags='%s'"),
-					*Container->Name, Container->ChunkEntries.Num(), *LexToString(Container->Flags));
+				FGuid KeyGuid;
+				ensure(FGuid::Parse(FString(StringCast<TCHAR>(*Container->EncryptionKeyGuid)), KeyGuid));
+
+				if (FEncryptionKeyManager::Get().TryGetKey(KeyGuid, Container->EncryptionKey))
+				{
+					UE_LOG(LogIoStoreOnDemand, Log, TEXT("Mounting container '%s', Entries=%d, Flags='%s'"),
+							*Container->Name, Container->ChunkEntries.Num(), *LexToString(Container->Flags));
 
-				EnumRemoveFlags(Container->Flags, EOnDemandContainerFlags::PendingEncryptionKey);
-				EnumAddFlags(Container->Flags, EOnDemandContainerFlags::Mounted);
+					EnumRemoveFlags(Container->Flags, EOnDemandContainerFlags::PendingEncryptionKey);
+					EnumAddFlags(Container->Flags, EOnDemandContainerFlags::Mounted);
+				}
 			}
 		}
 	}
+
+	// Tick pending mount/install request waiting on encryption key(s)
+	TryEnterTickLoop();
 }
 
 void FOnDemandIoStore::CreateContainersFromToc(
@@ -1192,6 +1235,12 @@ void FOnDemandIoStore::CreateContainersFromToc(
 		Container->ContainerId			= ContainerEntry.ContainerId;
 		Container->CompressionFormats.Add(CompressionFormat);
 
+		const EIoContainerFlags ContainerFlags = static_cast<EIoContainerFlags>(ContainerEntry.ContainerFlags);
+		if (EnumHasAnyFlags(ContainerFlags, EIoContainerFlags::Encrypted))
+		{
+			EnumAddFlags(Container->Flags, EOnDemandContainerFlags::Encrypted);
+		}
+
 		Container->ChunkEntries.Reserve(ContainerEntry.Entries.Num());
 		for (const FOnDemandTocEntry& TocEntry : ContainerEntry.Entries)
 		{
diff --git a/Engine/Source/Runtime/Experimental/IoStore/OnDemand/Private/OnDemandIoStore.h b/Engine/Source/Runtime/Experimental/IoStore/OnDemand/Private/OnDemandIoStore.h
--- a/Engine/Source/Runtime/Experimental/IoStore/OnDemand/Private/OnDemandIoStore.h
+++ b/Engine/Source/Runtime/Experimental/IoStore/OnDemand/Private/OnDemandIoStore.h
@@ -33,6 +33,7 @@ enum class EOnDemandContainerFlags : uint8
 	Mounted					= (1 << 1),
 	Streaming				= (1 << 2),
 	Installed				= (1 << 3),
+	Encrypted				= (1 << 4),
 	Count
 };
 ENUM_CLASS_FLAGS(EOnDemandContainerFlags);
