Fixed compile warning in IoStore on-demand install cache

#rnx

diff --git a/Engine/Source/Runtime/Experimental/IoStore/OnDemand/Private/OnDemandInstallCache.cpp b/Engine/Source/Runtime/Experimental/IoStore/OnDemand/Private/OnDemandInstallCache.cpp
--- a/Engine/Source/Runtime/Experimental/IoStore/OnDemand/Private/OnDemandInstallCache.cpp
+++ b/Engine/Source/Runtime/Experimental/IoStore/OnDemand/Private/OnDemandInstallCache.cpp
@@ -169,20 +169,14 @@ FCasLocation FCas::FindChunk(const FIoHash& Hash)
 
 FCasBlockId FCas::CreateBlock()
 {
-	IPlatformFile& Ipf = FPlatformFileManager::Get().GetPlatformFile();
+	IPlatformFile&	Ipf = FPlatformFileManager::Get().GetPlatformFile();
+	FCasBlockId		Out = FCasBlockId::Invalid;
 
 	UE::TUniqueLock Lock(Mutex);
 
-	uint32 BlockIdValue = 1;
-	for (;;)
+	for (uint32 Id = 1; Id < MAX_uint32 && !Out.IsValid(); Id++)
 	{
-		if (BlockIdValue == 0)
-		{
-			UE_LOG(LogIoStoreOnDemand, Error, TEXT("Failed to generate block ID"));
-			return FCasBlockId::Invalid;
-		}
-
-		const FCasBlockId BlockId(BlockIdValue++);
+		const FCasBlockId BlockId(Id);
 		if (BlockIds.Contains(BlockId))
 		{
 			continue;
@@ -197,10 +191,10 @@ FCasBlockId FCas::CreateBlock()
 
 		BlockIds.Add(BlockId);
 		LastAccess.FindOrAdd(BlockId, FDateTime::UtcNow().GetTicks());
-		return BlockId;
+		Out = BlockId;
 	}
 
-	return FCasBlockId::Invalid;
+	return Out;
 }
 
 FIoStatus FCas::DeleteBlock(FCasBlockId BlockId, TArray<FCasAddr>& OutAddrs)
