ChaosFlesh: Bug fixes and add unit tests for GS dynamic springs

diff --git a/Engine/Source/Runtime/Experimental/Chaos/Public/Chaos/Deformable/GaussSeidelDynamicWeakConstraints.h b/Engine/Source/Runtime/Experimental/Chaos/Public/Chaos/Deformable/GaussSeidelDynamicWeakConstraints.h
--- a/Engine/Source/Runtime/Experimental/Chaos/Public/Chaos/Deformable/GaussSeidelDynamicWeakConstraints.h
+++ b/Engine/Source/Runtime/Experimental/Chaos/Public/Chaos/Deformable/GaussSeidelDynamicWeakConstraints.h
@@ -83,8 +83,10 @@ namespace Chaos::Softs
 			delete DeletedSingleConstraint;
 
 			ConstraintsData.RemoveConstraint(OldConstraintIndex);
-			Handles[OldConstraintIndex].ConstraintIndex = OldConstraintIndex;
-			
+			if (ConstraintsData.Size() > 0)
+			{
+				Handles[OldConstraintIndex]->ConstraintIndex = OldConstraintIndex;
+			}
 			//Update incident element information:
 			const int32 IndicesOffset = SingleData.SingleIndices.Num();
 			for (int32 i = 0; i < SingleData.SingleIndices.Num(); i++)
