Dataflow : add support for TObjectPtr<UObject> and compatible types

diff --git a/Engine/Source/Runtime/Experimental/Dataflow/Core/Public/Dataflow/DataflowAnyType.h b/Engine/Source/Runtime/Experimental/Dataflow/Core/Public/Dataflow/DataflowAnyType.h
--- a/Engine/Source/Runtime/Experimental/Dataflow/Core/Public/Dataflow/DataflowAnyType.h
+++ b/Engine/Source/Runtime/Experimental/Dataflow/Core/Public/Dataflow/DataflowAnyType.h
@@ -61,3 +61,15 @@ struct FDataflowStringConvertibleTypes : public FDataflowAnyType
 	FString Value;
 };
 
+
+USTRUCT()
+struct FDataflowUObjectConvertibleTypes : public FDataflowAnyType
+{
+	using FPolicyType = FDataflowUObjectConvertibleTypePolicy;
+	using FStorageType = TObjectPtr<UObject>;
+
+	GENERATED_USTRUCT_BODY()
+
+	UPROPERTY(EditAnywhere, Category = Value)
+	TObjectPtr<UObject> Value;
+};
\ No newline at end of file
diff --git a/Engine/Source/Runtime/Experimental/Dataflow/Core/Public/Dataflow/DataflowTypePolicy.h b/Engine/Source/Runtime/Experimental/Dataflow/Core/Public/Dataflow/DataflowTypePolicy.h
--- a/Engine/Source/Runtime/Experimental/Dataflow/Core/Public/Dataflow/DataflowTypePolicy.h
+++ b/Engine/Source/Runtime/Experimental/Dataflow/Core/Public/Dataflow/DataflowTypePolicy.h
@@ -4,6 +4,42 @@
 
 #include "CoreMinimal.h"
 #include "Misc/TVariant.h"
+#include "UObject/Object.h"
+#include <string>
+
+class UObject;
+
+template<typename TType>
+struct TDataflowPolicyTypeName
+{
+};
+
+#define UE_DATAFLOW_POLICY_DECLARE_TYPENAME(TType) \
+template<> \
+struct TDataflowPolicyTypeName<TType> \
+{  \
+	FORCEINLINE static const TCHAR* GetName() \
+	{ \
+		return TEXT(#TType); \
+	} \
+};
+
+UE_DATAFLOW_POLICY_DECLARE_TYPENAME(bool)
+UE_DATAFLOW_POLICY_DECLARE_TYPENAME(uint8)
+UE_DATAFLOW_POLICY_DECLARE_TYPENAME(uint16)
+UE_DATAFLOW_POLICY_DECLARE_TYPENAME(uint32)
+UE_DATAFLOW_POLICY_DECLARE_TYPENAME(uint64)
+UE_DATAFLOW_POLICY_DECLARE_TYPENAME(int8)
+UE_DATAFLOW_POLICY_DECLARE_TYPENAME(int16)
+UE_DATAFLOW_POLICY_DECLARE_TYPENAME(int32)
+UE_DATAFLOW_POLICY_DECLARE_TYPENAME(int64)
+UE_DATAFLOW_POLICY_DECLARE_TYPENAME(float)
+UE_DATAFLOW_POLICY_DECLARE_TYPENAME(double)
+UE_DATAFLOW_POLICY_DECLARE_TYPENAME(FName)
+UE_DATAFLOW_POLICY_DECLARE_TYPENAME(FString)
+UE_DATAFLOW_POLICY_DECLARE_TYPENAME(TObjectPtr<UObject>)
+
+#undef UE_DATAFLOW_POLICY_DECLARE_TYPENAME
 
 struct IDataflowTypePolicy
 {
@@ -57,7 +93,7 @@ struct TDataflowSingleTypePolicy : public IDataflowTypePolicy
 		return &Instance;
 	}
 
-	inline static const FName TypeName = FName(TNameOf<T>::GetName());
+	inline static const FName TypeName = FName(TDataflowPolicyTypeName<T>::GetName());
 };
 
 template <typename... TTypes>
@@ -83,8 +119,6 @@ struct TDataflowMultiTypePolicy<>: public IDataflowTypePolicy
 	}
 };
 
-Expose_TNameOf(bool);
-
 template <typename T, typename... TTypes>
 struct TDataflowMultiTypePolicy<T, TTypes...>: public TDataflowMultiTypePolicy<TTypes...>
 {
@@ -165,6 +199,62 @@ struct FDataflowStringConvertibleTypePolicy : IDataflowTypePolicy
 	}
 };
 
+struct FDataflowUObjectConvertibleTypePolicy : IDataflowTypePolicy
+{
+	virtual bool SupportsType(FName InType) const override
+	{
+		return SupportsTypeStatic(InType);
+	}
+
+	static bool SupportsTypeStatic(FName InType)
+	{
+		FString InnerTypeStr;
+		if (GetObjectPtrInnerType(InType.ToString(), InnerTypeStr))
+		{
+			if (StaticFindFirstObject(UObject::StaticClass(), *InnerTypeStr, EFindFirstObjectOptions::NativeFirst))
+			{
+				return true;
+			}
+		}
+		// not a proper object pointer 
+		return false;
+	}
+
+	template <typename TVisitor>
+	static bool VisitPolicyByType(FName RequestedType, TVisitor Visitor)
+	{
+		if (SupportsTypeStatic(RequestedType))
+		{
+			TDataflowSingleTypePolicy<TObjectPtr<UObject>> SingleTypePolicy;
+			Visitor(SingleTypePolicy);
+			return true;
+		}
+		return false;
+	}
+
+	static IDataflowTypePolicy* GetInterface()
+	{
+		static FDataflowUObjectConvertibleTypePolicy Instance;
+		return &Instance;
+	}
+
+	// returns true if the type was a TObjectPtr and the inner type was properly extracted
+	static bool GetObjectPtrInnerType(const FString& InTypeStr, FString& InnerType)
+	{
+		static constexpr const TCHAR* ObjectPtrPrefix = TEXT("TObjectPtr<U");
+		static constexpr size_t ObjectPtrPrefixLen = std::char_traits<TCHAR>::length(ObjectPtrPrefix);
+		if (InTypeStr.StartsWith(ObjectPtrPrefix))
+		{
+			InnerType = InTypeStr
+				.RightChop(ObjectPtrPrefixLen) // remove the TObjectPtr< type
+				.LeftChop(1) // remove the last ">"
+				.TrimStartAndEnd();
+				return true;
+		}
+		return false;
+	}
+};
+
 // type Converters
 
 template <typename T>
