Tests for chunked encoding

#rnx

diff --git a/Engine/Source/Runtime/Experimental/IoStore/HttpClient/Private/Test.inl b/Engine/Source/Runtime/Experimental/IoStore/HttpClient/Private/Test.inl
--- a/Engine/Source/Runtime/Experimental/IoStore/HttpClient/Private/Test.inl
+++ b/Engine/Source/Runtime/Experimental/IoStore/HttpClient/Private/Test.inl
@@ -421,6 +421,98 @@ static void RedirectTest(const ANSICHAR* TestHost, FCertRootsRef VerifyCert)
 	}
 }
 
+////////////////////////////////////////////////////////////////////////////////
+static void ChunkedTest(const ANSICHAR* TestHost)
+{
+	FEventLoop Loop;
+
+	auto WaitForLoopIdle = [&] {
+		for (; Loop.Tick(-1); FPlatformProcess::SleepNoStats(0.02f));
+	};
+
+	TAnsiStringBuilder<64> Url;
+
+	// TestServer proxy doesn't support chunked transfer so find the actual httpd
+	int32 HttpdPort = -1;
+	Url << "http://" << TestHost << ":9493/port";
+	Loop.Send(Loop.Get(Url), [&HttpdPort, Dest=FIoBuffer()] (const FTicketStatus& Status) mutable
+	{
+		if (Status.GetId() == FTicketStatus::EId::Response)
+		{
+			Status.GetResponse().SetDestination(&Dest);
+			return;
+		}
+
+		check(Status.GetId() == FTicketStatus::EId::Content);
+		HttpdPort = int32(CrudeToInt({ (char*)Dest.GetView().GetData(), int32(Dest.GetSize()) }));
+	});
+	WaitForLoopIdle();
+	check(HttpdPort > -1);
+
+	auto BuildUrl = [&Url, TestHost, HttpdPort] (int32 PayloadSize, FAnsiStringView UrlSuffix) -> FAnsiStringView
+	{
+		Url.Reset();
+		Url << "http://" << TestHost << ":" << HttpdPort << "/chunked/" << PayloadSize << UrlSuffix;
+		return Url;
+	};
+
+	struct FTestState
+	{
+		uint32	Size = 0;
+		uint32	Hash = 0x493;
+		uint32	ExpectedHash = 0;
+		int32	ExpectedSize = -1;
+	};
+	auto ChunkedSink = [State=FTestState(), Dest=FIoBuffer()] (const FTicketStatus& Status) mutable
+	{
+		if (Status.GetId() == FTicketStatus::EId::Response)
+		{
+			FResponse& Response = Status.GetResponse();
+			check(Response.GetStatus() == EStatusCodeClass::Successful);
+			check(Response.GetStatusCode() == 200);
+
+			State.ExpectedHash = uint32(CrudeToInt(Response.GetHeader("X-TestServer-Hash")));
+			State.ExpectedSize = uint32(CrudeToInt(Response.GetHeader("X-TestServer-Size")));
+
+			uint32 DestSize = ((State.ExpectedHash & 0x3f) / 7) * 67;
+			Dest = FIoBuffer(DestSize);
+
+			Response.SetDestination(&Dest);
+			return;
+		}
+
+		check(Status.GetId() == FTicketStatus::EId::Content);
+
+		FMemoryView View = Dest.GetView(); 
+		State.Size += uint32(View.GetSize());
+		for (uint32 i = 0, n = uint32(View.GetSize()); i < n; ++i)
+		{
+			uint8 c = ((const uint8*)(View.GetData()))[i];
+			State.Hash = (State.Hash + c) * 0x493;
+		}
+
+		if (View.GetSize() == 0)
+		{
+			check(State.Hash == State.ExpectedHash);
+			check(State.Size == State.ExpectedSize);
+		}
+	};
+
+	for (FAnsiStringView UrlSuffix : { "", "/ext" })
+	{
+		for (uint32 Mixer : { 1, 2, 3, 17, 71, 4931, 0xa9e })
+		{
+			for (uint32 SizeToGet : { 4,8,32,64,1,2,3,5,7,11,13,17,19,41,43,47,59,67,71,83,89,103,109 })
+			{
+				BuildUrl(SizeToGet * Mixer, UrlSuffix);
+				FRequest Request = Loop.Get(Url);
+				Loop.Send(MoveTemp(Request), ChunkedSink);
+			}
+			WaitForLoopIdle();
+		}
+	}
+}
+
 ////////////////////////////////////////////////////////////////////////////////
 static void SeedHttp(const ANSICHAR* TestHost, uint32 Seed)
 {
@@ -821,7 +913,6 @@ static void HttpTest(const ANSICHAR* TestHost, FCertRootsRef VerifyCert)
 	// pre-generated headers
 	// request-with-body
 	// proxy
-	// chunked transfer encoding
 	// gzip / deflate
 	// loop multi-req.
 	// url auth credentials
@@ -874,6 +965,7 @@ IOSTOREHTTPCLIENT_API void IasHttpTest(const ANSICHAR* TestHost="localhost", uin
 	SeedHttp(TestHost, Seed);
 	HttpTest(TestHost, FCertRoots::NoTls());
 	HttpTest(TestHost, TestServerCertRef);
+	ChunkedTest(TestHost);
 	RedirectTest(TestHost, TestServerCertRef);
 	TlsLoadRootCerts();
 	TlsTest();
