Added trailer option to test server

#rnx

diff --git a/Engine/Source/Runtime/Experimental/IoStore/HttpClient/Private/TestServer.py b/Engine/Source/Runtime/Experimental/IoStore/HttpClient/Private/TestServer.py
--- a/Engine/Source/Runtime/Experimental/IoStore/HttpClient/Private/TestServer.py
+++ b/Engine/Source/Runtime/Experimental/IoStore/HttpClient/Private/TestServer.py
@@ -224,12 +224,17 @@ def http_chunked(handler, payload_size=0, *options):
         ext_payload = "".join(random.choices("Trigrams; Diner", k=n))
         ext_payload = b";" + ext_payload.encode()
 
+    trailer_payload = b"X-TestServer-Trailer" if "trailer" in options else b""
+
     payload, payload_hash = _make_payload(payload_size)
+
     handler.send_response(200)
     handler.send_header("Transfer-Encoding", "chunked")
     handler.send_header("X-TestServer-Hash", payload_hash)
     handler.send_header("X-TestServer-Size", len(payload))
     handler.send_header("Content-Type", "application/octet-stream")
+    if trailer_payload:
+        handler.send_header("Trailer:", trailer_payload)
     handler.end_headers()
 
     max_chunk_size = int(random.random() * 1024) + 1
@@ -244,6 +249,8 @@ def http_chunked(handler, payload_size=0, *options):
 
         handler.wfile.write(header)
         handler.wfile.write(piece)
+        if trailer_payload and not piece:
+            handler.wfile.write(trailer_payload + b": true\r\n")
         handler.wfile.write(b"\r\n")
         if not piece:
             break
