Add decode error and http error counts to IAS on screen statistics.

#rnx

- Add a new line to the IAS on screen statistics that shows the number of decode errors (when loading data from disk) and the number of http errors (when streaming over the network) as an error message.
- Only show the line if at least one of the values is non-zero. If everything is working then the line should not be seen.
- Minor optimization by moving instead of copying the messages when adding them to the FSeverityMessageMap.

diff --git a/Engine/Source/Runtime/Experimental/IoStore/OnDemand/Private/Statistics.cpp b/Engine/Source/Runtime/Experimental/IoStore/OnDemand/Private/Statistics.cpp
--- a/Engine/Source/Runtime/Experimental/IoStore/OnDemand/Private/Statistics.cpp
+++ b/Engine/Source/Runtime/Experimental/IoStore/OnDemand/Private/Statistics.cpp
@@ -419,7 +419,17 @@ FOnDemandIoBackendStats::FOnDemandIoBackendStats()
 					GHttpRetryCount.Get(),
 					GHttpPendingCount.Get()
 				);
-				Out.Add(FCoreDelegates::EOnScreenMessageSeverity::Info, Message);
+				Out.Add(FCoreDelegates::EOnScreenMessageSeverity::Info, MoveTemp(Message));
+
+				if (GIoDecodeErrorCount.Get() > 0 || GHttpErrorCount.Get() > 0)
+				{
+					FText Warning = FText::Format(LOCTEXT("IASWarning", "IAS - Decode Errors {0} | Http Errors {1}"),
+						GIoDecodeErrorCount.Get(),
+						GHttpErrorCount.Get()
+					);
+
+					Out.Add(FCoreDelegates::EOnScreenMessageSeverity::Error, MoveTemp(Warning));
+				}
 			}
 		});
 #undef LOCTEXT_NAMESPACE
