Clear indication (plus tests) when unsupported trailers are encountered.

#rnx

diff --git a/Engine/Source/Runtime/Experimental/IoStore/HttpClient/Private/Loop.inl b/Engine/Source/Runtime/Experimental/IoStore/HttpClient/Private/Loop.inl
--- a/Engine/Source/Runtime/Experimental/IoStore/HttpClient/Private/Loop.inl
+++ b/Engine/Source/Runtime/Experimental/IoStore/HttpClient/Private/Loop.inl
@@ -342,6 +342,12 @@ static FOutcome DoRecvStream(FActivity* Activity, FPeerType& Peer, int32& MaxRec
 {
 	auto RaiseCrLfError = [Activity] ()
 	{
+		if (Activity->NoContent)
+		{
+			Activity_SetError(Activity, "Trailing headers are not supported (ERRTRAIL)");
+			return FOutcome::Error(Activity->ErrorReason);
+		}
+
 		Activity_SetError(Activity, "Expected CRLF chunk terminal");
 		return FOutcome::Error(Activity->ErrorReason);
 	};
diff --git a/Engine/Source/Runtime/Experimental/IoStore/HttpClient/Private/Test.inl b/Engine/Source/Runtime/Experimental/IoStore/HttpClient/Private/Test.inl
--- a/Engine/Source/Runtime/Experimental/IoStore/HttpClient/Private/Test.inl
+++ b/Engine/Source/Runtime/Experimental/IoStore/HttpClient/Private/Test.inl
@@ -498,6 +498,7 @@ static void ChunkedTest(const ANSICHAR* TestHost)
 		}
 	};
 
+	// General soak test
 	for (FAnsiStringView UrlSuffix : { "", "/ext" })
 	{
 		for (uint32 Mixer : { 1, 2, 3, 17, 71, 4931, 0xa9e })
@@ -511,6 +512,30 @@ static void ChunkedTest(const ANSICHAR* TestHost)
 			WaitForLoopIdle();
 		}
 	}
+
+	// Rudimentary coverage for tranfers with trailing headers.
+	bool bThError = false;
+	auto ExpectError = [&bThError, Dest=FIoBuffer()] (const FTicketStatus& Status) mutable
+	{
+		if (Status.GetId() == FTicketStatus::EId::Response)
+		{
+			FResponse& Response = Status.GetResponse();
+			Response.SetDestination(&Dest);
+			return;
+		}
+
+		if (Status.GetId() != FTicketStatus::EId::Error)
+		{
+			return;
+		}
+
+		FAnsiStringView Reason = Status.GetErrorReason();
+		bThError = Reason.Contains("ERRTRAIL", ESearchCase::IgnoreCase);
+	};
+	BuildUrl(16 << 10, "/trailer");
+	Loop.Send(Loop.Get(Url), ExpectError);
+	WaitForLoopIdle();
+	check(bThError);
 }
 
 ////////////////////////////////////////////////////////////////////////////////
