Improve logging after IAS cache initialization failure.

#rnx

diff --git a/Engine/Source/Runtime/Experimental/IoStore/OnDemand/Private/IoStoreOnDemand.cpp b/Engine/Source/Runtime/Experimental/IoStore/OnDemand/Private/IoStoreOnDemand.cpp
--- a/Engine/Source/Runtime/Experimental/IoStore/OnDemand/Private/IoStoreOnDemand.cpp
+++ b/Engine/Source/Runtime/Experimental/IoStore/OnDemand/Private/IoStoreOnDemand.cpp
@@ -1127,6 +1127,7 @@ void FIoStoreOnDemandModule::InitializeInternal()
 
 	TUniquePtr<IIasCache> Cache;
 	FIasCacheConfig CacheConfig = GetIasCacheConfig(CommandLine);
+
 	CacheConfig.DropCache = DeferredAbandonCache.Get(CacheConfig.DropCache);
 	if (CacheConfig.DiskQuota > 0)
 	{
@@ -1134,12 +1135,17 @@ void FIoStoreOnDemandModule::InitializeInternal()
 		{
 			FString CacheDir = FPaths::ProjectPersistentDownloadDir();
 			Cache = MakeIasCache(*CacheDir, CacheConfig);
+
+			UE_CLOG(!Cache.IsValid(), LogIas, Warning, TEXT("File cache disabled - streaming only (init-fail)"));
+		}
+		else
+		{
+			UE_LOG(LogIas, Warning, TEXT("File cache disabled - streaming only (project has no persistent download dir enabled for this platform)"));
 		}
 	}
-	if (!Cache.IsValid())
+	else
 	{
-		UE_LOG(LogIas, Log, TEXT("File cache disabled - streaming only (%s)"),
-			(CacheConfig.DiskQuota > 0) ? TEXT("init-fail") : TEXT("zero-quota"));
+		UE_LOG(LogIas, Log, TEXT("File cache disabled - streaming only (zero-quota)"));
 	}
 
 	HttpIoDispatcherBackend = MakeOnDemandIoDispatcherBackend(EndpointConfig, *IoStore, MoveTemp(Cache));
