Use IoStore itself to fetch distribution endpoints.

#rnx

diff --git a/Engine/Source/Runtime/Experimental/IoStore/OnDemand/IoStoreOnDemand.Build.cs b/Engine/Source/Runtime/Experimental/IoStore/OnDemand/IoStoreOnDemand.Build.cs
--- a/Engine/Source/Runtime/Experimental/IoStore/OnDemand/IoStoreOnDemand.Build.cs
+++ b/Engine/Source/Runtime/Experimental/IoStore/OnDemand/IoStoreOnDemand.Build.cs
@@ -12,7 +12,6 @@ public class IoStoreOnDemand : ModuleRules
 			new string[] {
 				"Analytics",
 				"CoreUObject",
-				"HTTP",
 				"IoStoreHttpClient",
 				"Json",
 			}
diff --git a/Engine/Source/Runtime/Experimental/IoStore/OnDemand/Private/DistributionEndpoints.cpp b/Engine/Source/Runtime/Experimental/IoStore/OnDemand/Private/DistributionEndpoints.cpp
--- a/Engine/Source/Runtime/Experimental/IoStore/OnDemand/Private/DistributionEndpoints.cpp
+++ b/Engine/Source/Runtime/Experimental/IoStore/OnDemand/Private/DistributionEndpoints.cpp
@@ -2,17 +2,19 @@
 
 #include "DistributionEndpoints.h"
 
+#include "Containers/StringConv.h"
 #include "Dom/JsonValue.h"
 #include "HAL/IConsoleManager.h"
 #include "HAL/PlatformTime.h"
-#include "HttpManager.h"
-#include "HttpModule.h"
+#include "IO/Http/Client.h"
+#include "IO/IoBuffer.h"
 #include "IO/IoStoreOnDemand.h"
-#include "Interfaces/IHttpResponse.h"
 #include "Misc/ScopeRWLock.h"
 #include "Serialization/JsonReader.h"
 #include "Serialization/JsonSerializer.h"
+#include "Serialization/MemoryReader.h"
 #include "Statistics.h"
+#include "Tasks/Task.h"
 
 namespace UE::IoStore
 {
@@ -24,6 +26,8 @@ static FAutoConsoleVariableRef CVar_DistributedEndpointTimeout(
 	TEXT("How long to wait (in seconds) for a distributed endoint resolve request before timing out")
 );
 
+static FDistributionEndpoints::EResult ParseResponse(FIoBuffer, TArray<FString>&);
+
 FDistributionEndpoints::EResult FDistributionEndpoints::ResolveEndpoints(const FString& DistributionUrl, TArray<FString>& OutServiceUrls)
 {
 	FEventRef Event;
@@ -32,91 +36,104 @@ FDistributionEndpoints::EResult FDistributionEndpoints::ResolveEndpoints(const F
 
 FDistributionEndpoints::EResult FDistributionEndpoints::ResolveEndpoints(const FString& DistributionUrl, TArray<FString>& OutServiceUrls, FEvent& Event)
 {
+	using namespace HTTP;
+
 	TRACE_CPUPROFILER_EVENT_SCOPE(FDistributionEndpoints::ResolveEndpoints);
 
+	UE_LOG(LogIas, Log, TEXT("Resolving distributed endpoint '%s'"), *DistributionUrl);
+
 	std::atomic_bool bHasResponse = false;
 	EResult Result = EResult::Failure;
+	auto OnRequestStatus = [
+		&bHasResponse,
+		&OutServiceUrls,
+		&Event,
+		&Result,
+		Dest=FIoBuffer()
+		] (const FTicketStatus& Status) mutable
+	{
+		if (Status.GetId() >= FTicketStatus::EId::Cancelled)
+		{
+			return;
+		}
 
-	// The timeout is a cvar but should not be less than 10 seconds.
-	const float HttpTimeout = FMath::Max(static_cast<float>(GDistributedEndpointTimeout), 10.0f);
-
-	UE_LOG(LogIas, Log, TEXT("Resolving distributed endpoint '%s'"), *DistributionUrl);
-
-	FHttpRequestPtr HttpRequest = FHttpModule::Get().CreateRequest();
-	HttpRequest->SetTimeout(HttpTimeout);
-	HttpRequest->SetURL(DistributionUrl);
-	HttpRequest->SetVerb(TEXT("GET"));
-	HttpRequest->SetHeader(TEXT("Accept"), TEXT("application/json"));
-	HttpRequest->OnProcessRequestComplete().BindLambda(
-		[this, &bHasResponse, &Event, &Result, &OutServiceUrls](FHttpRequestPtr Request, FHttpResponsePtr Response, bool bOk)
+		if (Status.GetId() == FTicketStatus::EId::Content)
 		{
-			LLM_SCOPE_BYTAG(Ias);
-			/* Work around a code generation issue in MSVC: forced usage of the request so it is not optimized away */
-			UE_LOG(LogIas, VeryVerbose, TEXT("HTTP [%s] request completed"), *Request->GetVerb());
-			
 			bHasResponse = true;
-			Result = ParseResponse(Response, OutServiceUrls);
-
+			Result = ParseResponse(Dest, OutServiceUrls);
 			Event.Trigger();
-		});
+			return;
+		}
 
-	HttpRequest->ProcessRequest();
+		FResponse& Response = Status.GetResponse();
+		if (Response.GetStatus() != EStatusCodeClass::Successful)
+		{
+			return;
+		}
+
+		Response.SetDestination(&Dest);
+	};
 
-	const uint32 WaitTime = GDistributedEndpointTimeout >= 0 ? (static_cast<uint32>(GDistributedEndpointTimeout) * 1000) : MAX_uint32;
+	auto AnsiUrl = StringCast<ANSICHAR>(*DistributionUrl, DistributionUrl.Len());
 
-	if (!Event.Wait(FTimespan::FromSeconds(WaitTime)) || bHasResponse == false)
+	FEventLoop Loop;
+	if (GDistributedEndpointTimeout >= 0)
 	{
-		HttpRequest->CancelRequest();
-		HttpRequest->OnProcessRequestComplete().Unbind();
+		Loop.SetFailTimeout(GDistributedEndpointTimeout * 1000);
 	}
 
+	FRequest Request = Loop.Get(AnsiUrl);
+	Request.Header("Accept", "application/json");
+	Loop.Send(MoveTemp(Request), OnRequestStatus);
+
+	while (Loop.Tick(-1))
+		;
+
 	UE_CLOG(Result == EResult::Success, LogIas, Log, TEXT("Successfully resolved distributed endpoint '%s' %d urls found"), *DistributionUrl, OutServiceUrls.Num());
 	UE_CLOG(Result != EResult::Success, LogIas, Log, TEXT("Failed to resolve distributed endpoint '%s'"), *DistributionUrl);
 
 	return Result;
 }
 
-FDistributionEndpoints::EResult FDistributionEndpoints::ParseResponse(FHttpResponsePtr HttpResponse, TArray<FString>& OutUrls)
+static FDistributionEndpoints::EResult ParseResponse(FIoBuffer Data, TArray<FString>& OutUrls)
 {
 	TRACE_CPUPROFILER_EVENT_SCOPE(FDistributionEndpoints::ParseResponse);
 
 	using FJsonValuePtr = TSharedPtr<FJsonValue>;
 	using FJsonObjPtr = TSharedPtr<FJsonObject>;
-	using FJsonReader = TJsonReader<TCHAR>;
+	using FJsonReader = TJsonReader<UTF8CHAR>;
 	using FJsonReaderPtr = TSharedRef<FJsonReader>;
 
-	if (HttpResponse != nullptr && HttpResponse->GetResponseCode() == 200)
+	FMemoryView JsonView = Data.GetView();
+	FUtf8StringView Json((UTF8CHAR*)JsonView.GetData(), int32(JsonView.GetSize()));
+	
+	FJsonReaderPtr JsonReader = TJsonReaderFactory<UTF8CHAR>::CreateFromView(Json);
+
+	FJsonObjPtr JsonObj;
+	if (!FJsonSerializer::Deserialize(JsonReader, JsonObj))
 	{
-		FString Json = HttpResponse->GetContentAsString();
-		
-		FJsonReaderPtr JsonReader = TJsonReaderFactory<TCHAR>::Create(Json);
+		return FDistributionEndpoints::EResult::Failure;
+	}
+
+	if (!JsonObj->HasTypedField< EJson::Array>(TEXT("distributions")))
+	{
+		return FDistributionEndpoints::EResult::Failure;
+	}
+
+	TArray<FJsonValuePtr> JsonValues = JsonObj->GetArrayField(TEXT("distributions"));
+	OutUrls.Reserve(JsonValues.Num());
 
-		FJsonObjPtr JsonObj;
-		if (FJsonSerializer::Deserialize(JsonReader, JsonObj))
+	for (const FJsonValuePtr& JsonValue : JsonValues)
+	{
+		FString ServiceUrl = JsonValue->AsString();
+		if (ServiceUrl.EndsWith(TEXT("/")))
 		{
-			if (!JsonObj->HasTypedField< EJson::Array>(TEXT("distributions")))
-			{
-				return EResult::Failure;
-			}
-
-			TArray<FJsonValuePtr> JsonValues = JsonObj->GetArrayField(TEXT("distributions"));
-			OutUrls.Reserve(JsonValues.Num());
-
-			for (const FJsonValuePtr& JsonValue : JsonValues)
-			{
-				FString ServiceUrl = JsonValue->AsString();
-				if (ServiceUrl.EndsWith(TEXT("/")))
-				{
-					ServiceUrl.LeftInline(ServiceUrl.Len() - 1);
-				}
-				OutUrls.Add(MoveTemp(ServiceUrl));
-			}
-
-			return !OutUrls.IsEmpty() ? EResult::Success : EResult::Failure;
+			ServiceUrl.LeftInline(ServiceUrl.Len() - 1);
 		}
+		OutUrls.Add(MoveTemp(ServiceUrl));
 	}
 
-	return EResult::Failure;
+	return !OutUrls.IsEmpty() ? FDistributionEndpoints::EResult::Success : FDistributionEndpoints::EResult::Failure;
 }
 
 } // namespace UE::IoStore
diff --git a/Engine/Source/Runtime/Experimental/IoStore/OnDemand/Private/DistributionEndpoints.h b/Engine/Source/Runtime/Experimental/IoStore/OnDemand/Private/DistributionEndpoints.h
--- a/Engine/Source/Runtime/Experimental/IoStore/OnDemand/Private/DistributionEndpoints.h
+++ b/Engine/Source/Runtime/Experimental/IoStore/OnDemand/Private/DistributionEndpoints.h
@@ -6,7 +6,6 @@
 #include "Containers/Map.h"
 #include "Containers/UnrealString.h"
 #include "HAL/CriticalSection.h"
-#include "HttpFwd.h"
 
 namespace UE::IoStore
 {
@@ -39,10 +38,6 @@ public:
 	 * @param Event				An event which can be triggered to cancel the request
 	 */
 	FDistributionEndpoints::EResult ResolveEndpoints(const FString& DistributionUrl, TArray<FString>& OutServiceUrls, FEvent& Event);
-
-private:
-
-	FDistributionEndpoints::EResult ParseResponse(FHttpResponsePtr HttpResponse, TArray<FString>& OutUrls);
 };
 
 } // namespace UE::IoStore
