Iris
- Fixed issue with changes to SubObjectFilterGroups would not enforce dirtiness until subobjects was modified

diff --git a/Engine/Source/Runtime/Experimental/Iris/Core/Private/Iris/ReplicationSystem/Conditionals/ReplicationConditionals.cpp b/Engine/Source/Runtime/Experimental/Iris/Core/Private/Iris/ReplicationSystem/Conditionals/ReplicationConditionals.cpp
--- a/Engine/Source/Runtime/Experimental/Iris/Core/Private/Iris/ReplicationSystem/Conditionals/ReplicationConditionals.cpp
+++ b/Engine/Source/Runtime/Experimental/Iris/Core/Private/Iris/ReplicationSystem/Conditionals/ReplicationConditionals.cpp
@@ -78,6 +78,26 @@ void FReplicationConditionals::OnMaxInternalNetRefIndexIncreased(FInternalNetRef
 	}
 }
 
+void FReplicationConditionals::MarkLifeTimeConditionalsDirtyForObjectsInGroup(FNetObjectGroupHandle GroupHandle)
+{
+	IRIS_PROFILER_SCOPE(MarkLifeTimeConditionalsDirtyForObjectsInGroup)
+
+	const FNetObjectGroupHandle::FGroupIndexType GroupIndex = GroupHandle.GetGroupIndex();
+	if (GroupHandle.IsReservedNetObjectGroup())
+	{
+		UE_LOG(LogIris, Warning, TEXT("FReplicationConditionals::MarkLifeTimeConditionalsDirtyForObjectsInGroup - Marking reserved group dirty is not allowed. GroupIndex: %u which is not allowed."), GroupIndex);
+		return;
+	}
+
+	if (const FNetObjectGroup* Group = NetObjectGroups->GetGroup(GroupHandle))
+	{
+		for (FInternalNetRefIndex InternalObjectIndex : Group->Members)
+		{
+			ObjectsWithDirtyLifetimeConditionals.SetBit(InternalObjectIndex);
+		}
+	}
+}
+
 bool FReplicationConditionals::SetConditionConnectionFilter(FInternalNetRefIndex ObjectIndex, EReplicationCondition Condition, uint32 ConnectionId, bool bEnable)
 {
 	if (ConnectionId >= MaxConnectionCount)
diff --git a/Engine/Source/Runtime/Experimental/Iris/Core/Private/Iris/ReplicationSystem/Conditionals/ReplicationConditionals.h b/Engine/Source/Runtime/Experimental/Iris/Core/Private/Iris/ReplicationSystem/Conditionals/ReplicationConditionals.h
--- a/Engine/Source/Runtime/Experimental/Iris/Core/Private/Iris/ReplicationSystem/Conditionals/ReplicationConditionals.h
+++ b/Engine/Source/Runtime/Experimental/Iris/Core/Private/Iris/ReplicationSystem/Conditionals/ReplicationConditionals.h
@@ -5,6 +5,7 @@
 #include "CoreTypes.h"
 #include "Containers/Array.h"
 #include "Containers/Map.h"
+#include "Iris/ReplicationSystem/NetObjectGroupHandle.h"
 #include "Net/Core/NetBitArray.h"
 
 enum ELifetimeCondition : int;
@@ -18,6 +19,7 @@ namespace UE::Net
 		typedef uint32 FInternalNetRefIndex;
 		class FNetRefHandleManager;
 		class FReplicationConnections;
+		class FNetObjectGroups;
 		class FReplicationFiltering;
 		class FNetObjectGroups;
 	}
@@ -66,6 +68,8 @@ public:
 	void InitPropertyCustomConditions(FInternalNetRefIndex ObjectIndex);
 	bool SetPropertyCustomCondition(FInternalNetRefIndex ObjectIndex, const void* Owner, uint16 RepIndex, bool bIsActive);
 	bool SetPropertyDynamicCondition(FInternalNetRefIndex ObjectIndex, const void* Owner, uint16 RepIndex, ELifetimeCondition Condition);
+
+	void MarkLifeTimeConditionalsDirtyForObjectsInGroup(FNetObjectGroupHandle GroupHandle);
 	
 	/** Unconditionally marks a property as dirty, causing it to replicate with the object at the earliest convenience. */
 	void MarkPropertyDirty(FInternalNetRefIndex ObjectIndex, uint16 RepIndex);
diff --git a/Engine/Source/Runtime/Experimental/Iris/Core/Private/Iris/ReplicationSystem/ReplicationSystem.cpp b/Engine/Source/Runtime/Experimental/Iris/Core/Private/Iris/ReplicationSystem/ReplicationSystem.cpp
--- a/Engine/Source/Runtime/Experimental/Iris/Core/Private/Iris/ReplicationSystem/ReplicationSystem.cpp
+++ b/Engine/Source/Runtime/Experimental/Iris/Core/Private/Iris/ReplicationSystem/ReplicationSystem.cpp
@@ -1296,6 +1296,9 @@ void UReplicationSystem::SetSubObjectFilterStatus(FName GroupName, uint32 Connec
 	{
 		FReplicationFiltering& Filtering = Impl->ReplicationSystemInternal.GetFiltering();
 		Filtering.SetSubObjectFilterStatus(GroupHandle, ConnectionId, ReplicationStatus);
+
+		FReplicationConditionals& Conditionals = Impl->ReplicationSystemInternal.GetConditionals();
+		Conditionals.MarkLifeTimeConditionalsDirtyForObjectsInGroup(GroupHandle);
 	}
 }
 
