Added command line option to fill on-demand install cache with dummy data for testing cache purge

#rnx

diff --git a/Engine/Source/Runtime/Experimental/IoStore/OnDemand/Private/OnDemandInstallCache.cpp b/Engine/Source/Runtime/Experimental/IoStore/OnDemand/Private/OnDemandInstallCache.cpp
--- a/Engine/Source/Runtime/Experimental/IoStore/OnDemand/Private/OnDemandInstallCache.cpp
+++ b/Engine/Source/Runtime/Experimental/IoStore/OnDemand/Private/OnDemandInstallCache.cpp
@@ -1185,7 +1185,6 @@ FIoStatus FOnDemandInstallCache::Purge(TMap<FIoHash, uint64>&& ChunksToIntall)
 {
 	FCasBlockInfoMap	BlockInfo;
 	const uint64		TotalCachedBytes = Cas.GetBlockInfo(BlockInfo);
-	const int64			AvailableBytes = MaxCacheSize - TotalCachedBytes;
 	uint64				TotalUncachedBytes = 0;
 
 	for (const TPair<FIoHash, uint64>& Kv : ChunksToIntall)
@@ -1200,16 +1199,14 @@ FIoStatus FOnDemandInstallCache::Purge(TMap<FIoHash, uint64>&& ChunksToIntall)
 		}
 	}
 
-	if (AvailableBytes >= IntCastChecked<int64>(TotalUncachedBytes))
+	const uint64 TotalRequiredBytes = TotalCachedBytes + TotalUncachedBytes;
+	if (TotalRequiredBytes <= MaxCacheSize)
 	{
 		UE_LOG(LogIoStoreOnDemand, Log, TEXT("Skipping cache purge, MaxCacheSize=%.2lf MiB, CacheSize=%.2lf MiB, UncachedSize=%.2lf MiB"),
 			ToMiB(MaxCacheSize), ToMiB(TotalCachedBytes), ToMiB(TotalUncachedBytes));
 		return FIoStatus::Ok;
 	}
 
-	UE_LOG(LogIoStoreOnDemand, Log, TEXT("Purging install cache, MaxCacheSize=%.2lf MiB, CacheSize=%.2lf MiB, UncachedSize=%.2lf MiB"),
-		ToMiB(MaxCacheSize), ToMiB(TotalCachedBytes), ToMiB(TotalUncachedBytes));
-
 	TSet<FCasBlockId>					ReferencedBlocks;
 	TArray<FSharedOnDemandContainer>	MountedContainers = IoStore.GetMountedContainers();
 
@@ -1220,17 +1217,34 @@ FIoStatus FOnDemandInstallCache::Purge(TMap<FIoHash, uint64>&& ChunksToIntall)
 		{
 			if (FCasLocation Loc = Cas.FindChunk(Kv.Value.Hash); Loc.IsValid())
 			{
-				BlockInfo.FindOrAdd(Loc.BlockId).RefCount++;
+				if (FCasBlockInfo* Info = BlockInfo.Find(Loc.BlockId))
+				{
+					Info->RefCount++;
+				}
 			}
 		}
 	}
 
+	uint64 TotalReferencedBytes = 0;
+	for (const TPair<FCasBlockId, FCasBlockInfo>& Kv : BlockInfo)
+	{
+		if (Kv.Value.RefCount > 0)
+		{
+			TotalReferencedBytes += Kv.Value.FileSize; 
+		}
+	}
+
 	BlockInfo.ValueSort([](const FCasBlockInfo& LHS, const FCasBlockInfo& RHS)
 	{
 		return LHS.LastAccess < RHS.LastAccess;
 	});
 
-	uint64 TotalPurgedBytes = 0;
+	UE_LOG(LogIoStoreOnDemand, Log, TEXT("Purging install cache, MaxCacheSize=%.2lf MiB, CacheSize=%.2lf MiB, UncachedSize=%.2lf MiB, ReferencedBytes=%.2lf MiB"),
+		ToMiB(MaxCacheSize), ToMiB(TotalCachedBytes), ToMiB(TotalUncachedBytes), ToMiB(TotalReferencedBytes));
+
+	const uint64	TotalBytesToPurge	= TotalRequiredBytes - MaxCacheSize;
+	uint64			TotalPurgedBytes	= 0;
+
 	for (const TPair<FCasBlockId, FCasBlockInfo>& Kv : BlockInfo)
 	{
 		const FCasBlockInfo& Info = Kv.Value;
@@ -1259,17 +1273,19 @@ FIoStatus FOnDemandInstallCache::Purge(TMap<FIoHash, uint64>&& ChunksToIntall)
 		}
 
 		TotalPurgedBytes += Info.FileSize;
-		if (TotalPurgedBytes >= TotalUncachedBytes)
+		if (TotalPurgedBytes >= TotalBytesToPurge)
 		{
 			break;
 		}
 	}
 
+	UE_LOG(LogIoStoreOnDemand, Log, TEXT("Purged %.2lf MiB (%.2lf%%) from install cache"),
+		ToMiB(TotalPurgedBytes), 100.0 * (double(TotalPurgedBytes) / double(TotalCachedBytes)));
+
 	const uint64 NewCachedBytes = TotalCachedBytes - TotalPurgedBytes;
-	UE_LOG(LogIoStoreOnDemand, Log, TEXT("Purged %.2lf MiB from install cache"), ToMiB(TotalPurgedBytes));
 	UE_CLOG(NewCachedBytes > MaxCacheSize,
-		LogIoStoreOnDemand, Warning, TEXT("Max install cache size exceeded by %.2lf MiB"),
-			ToMiB(NewCachedBytes - MaxCacheSize));
+		LogIoStoreOnDemand, Warning, TEXT("Max install cache size exceeded by %.2lf MiB (%.2lf%%)"),
+			ToMiB(NewCachedBytes - MaxCacheSize), 100.0 * (double(NewCachedBytes - MaxCacheSize) / double(MaxCacheSize)));
 
 	return FIoStatus::Ok;
 }
diff --git a/Engine/Source/Runtime/Experimental/IoStore/OnDemand/Private/OnDemandIoStore.cpp b/Engine/Source/Runtime/Experimental/IoStore/OnDemand/Private/OnDemandIoStore.cpp
--- a/Engine/Source/Runtime/Experimental/IoStore/OnDemand/Private/OnDemandIoStore.cpp
+++ b/Engine/Source/Runtime/Experimental/IoStore/OnDemand/Private/OnDemandIoStore.cpp
@@ -29,6 +29,9 @@
 #include "Misc/Paths.h"
 #include "Misc/PathViews.h"
 #include "Serialization/MemoryReader.h"
+#if !(UE_BUILD_SHIPPING|UE_BUILD_TEST)
+#include "String/LexFromString.h"
+#endif
 
 ///////////////////////////////////////////////////////////////////////////////
 bool GIoStoreOnDemandInstallCacheEnabled = true;
@@ -456,6 +459,56 @@ FIoStatus FOnDemandIoStore::Initialize()
 	}
 #endif
 
+#if !(UE_BUILD_SHIPPING|UE_BUILD_TEST)
+	FString ParamValue;
+	if (FParse::Value(FCommandLine::Get(), TEXT("Iad.Fill="), ParamValue))
+	{
+		ParamValue.TrimStartAndEndInline();
+		int64 FillSize = -1;
+		LexFromString(FillSize, ParamValue);
+
+		if (FillSize > 0)
+		{
+			if (ParamValue.EndsWith(TEXT("GB")))
+			{
+				FillSize = FillSize << 30;
+			}
+			if (ParamValue.EndsWith(TEXT("MB")))
+			{
+				FillSize = FillSize << 20;
+			}
+
+			UE_LOG(LogIoStoreOnDemand, Log, TEXT("Filling install cache with %.2lf MiB of dummy data"), double(FillSize) / 1024.0 / 1024.0);
+
+			FIoStatus	Status = FIoStatus::Ok;
+			uint64		Seed = 1;
+			while (FillSize >= 0 && Status.IsOk())
+			{
+				const uint64		ChunkSize = 256 << 10;
+				FIoBuffer			Chunk(ChunkSize);
+				TArrayView<uint64>	Values(reinterpret_cast<uint64*>(Chunk.GetData()), ChunkSize / sizeof(uint64));
+
+				for (uint64& Value : Values)
+				{
+					Value = Seed;
+				}
+
+				const FIoHash ChunkHash = FIoHash::HashBuffer(Chunk.GetView());
+				Status = InstallCache->PutChunk(MoveTemp(Chunk), ChunkHash);
+				Seed++;
+				FillSize -= ChunkSize;
+			}
+
+			if (Status.IsOk())
+			{
+				Status = InstallCache->Flush();
+			}
+
+			UE_CLOG(!Status.IsOk(), LogIoStoreOnDemand, Warning, TEXT("Failed to fill install cache with dummy data"));
+		}
+	}
+#endif
+
 	return EIoErrorCode::Ok;
 }
 
