Dataflow:
- Fixed duplication of the graph nodes whenever an Undo/Redo is called on the graph object.
- Fixed the Undo/Redo on the Add/RemoveOptionPin that wasn't serializing the graph nodes.

diff --git a/Engine/Plugins/Experimental/Dataflow/Source/DataflowEditor/Private/Dataflow/DataflowGraphEditor.cpp b/Engine/Plugins/Experimental/Dataflow/Source/DataflowEditor/Private/Dataflow/DataflowGraphEditor.cpp
--- a/Engine/Plugins/Experimental/Dataflow/Source/DataflowEditor/Private/Dataflow/DataflowGraphEditor.cpp
+++ b/Engine/Plugins/Experimental/Dataflow/Source/DataflowEditor/Private/Dataflow/DataflowGraphEditor.cpp
@@ -457,6 +457,7 @@ void SDataflowGraphEditor::OnAddOptionPin()
 				if (Node->CanAddPin())
 				{
 					const FScopedTransaction Transaction(LOCTEXT("AddOptionPin", "Add Option Pin"));
+					Graph->Modify();
 					EdNode->Modify();
 
 					EdNode->AddOptionPin();
@@ -521,6 +522,7 @@ void SDataflowGraphEditor::OnRemoveOptionPin()
 				if (Node->CanRemovePin())
 				{
 					const FScopedTransaction Transaction(LOCTEXT("RemoveOptionPin", "Remove Option Pin"));
+					Graph->Modify();
 					EdNode->Modify();
 
 					EdNode->RemoveOptionPin();
diff --git a/Engine/Source/Runtime/Experimental/Dataflow/Core/Private/Dataflow/DataflowGraph.cpp b/Engine/Source/Runtime/Experimental/Dataflow/Core/Private/Dataflow/DataflowGraph.cpp
--- a/Engine/Source/Runtime/Experimental/Dataflow/Core/Private/Dataflow/DataflowGraph.cpp
+++ b/Engine/Source/Runtime/Experimental/Dataflow/Core/Private/Dataflow/DataflowGraph.cpp
@@ -21,6 +21,13 @@ namespace Dataflow
 	{
 	}
 
+	void FGraph::Reset()
+	{
+		Nodes.Reset();
+		FilteredNodes.Reset();
+		Connections.Reset();
+		DisabledNodes.Reset();
+	}
 
 	void FGraph::RemoveNode(TSharedPtr<FDataflowNode> Node)
 	{
@@ -210,6 +217,8 @@ namespace Dataflow
 
 	void FGraph::SerializeForLoading(FArchive& Ar, FGraph* InGraph, UObject* OwningObject)
 	{
+		InGraph->Reset();
+
 		const bool bDataflowSeparateInputOutputSerialization = (Ar.CustomVer(FUE5MainStreamObjectVersion::GUID) >= FUE5MainStreamObjectVersion::DataflowSeparateInputOutputSerialization);
 		const bool bDataflowAnyTypeSupport = (Ar.CustomVer(FFortniteMainBranchObjectVersion::GUID) >= FFortniteMainBranchObjectVersion::DataflowAnyTypeSupport);
 		const bool bDataflowTemplateTypeFix = (Ar.CustomVer(FFortniteMainBranchObjectVersion::GUID) >= FFortniteMainBranchObjectVersion::DataflowTemplatedTypeFix);
diff --git a/Engine/Source/Runtime/Experimental/Dataflow/Core/Public/Dataflow/DataflowGraph.h b/Engine/Source/Runtime/Experimental/Dataflow/Core/Public/Dataflow/DataflowGraph.h
--- a/Engine/Source/Runtime/Experimental/Dataflow/Core/Public/Dataflow/DataflowGraph.h
+++ b/Engine/Source/Runtime/Experimental/Dataflow/Core/Public/Dataflow/DataflowGraph.h
@@ -181,6 +181,9 @@ namespace Dataflow
 
 		DATAFLOWCORE_API static void SerializeForSaving(FArchive& Ar, FGraph* InGraph, TArray<TSharedPtr<FDataflowNode>>& InNodes, TArray<FLink>& InConnections);
 		DATAFLOWCORE_API static void SerializeForLoading(FArchive& Ar, FGraph* InGraph, UObject* OwningObject);
+
+	private:
+		void Reset();
 	};
 
 	DATAFLOWCORE_API void RegisterNodeFilter(const FName& NodeFilter);
