Add tampering to chunked endpoint

#rnx

diff --git a/Engine/Source/Runtime/Experimental/IoStore/HttpClient/Private/TestServer.py b/Engine/Source/Runtime/Experimental/IoStore/HttpClient/Private/TestServer.py
--- a/Engine/Source/Runtime/Experimental/IoStore/HttpClient/Private/TestServer.py
+++ b/Engine/Source/Runtime/Experimental/IoStore/HttpClient/Private/TestServer.py
@@ -226,6 +226,10 @@ def http_chunked(handler, payload_size=0, *options):
 
     trailer_payload = b"X-TestServer-Trailer" if "trailer" in options else b""
 
+    get_crlf = lambda: b"\r\n"
+    if "tamper" in options:
+        get_crlf = lambda: "".join(random.choices("XX\r\r\n", k=2)).encode()
+
     payload, payload_hash = _make_payload(payload_size)
 
     handler.send_response(200)
@@ -245,13 +249,13 @@ def http_chunked(handler, payload_size=0, *options):
         header = b"%x" % len(piece)
         if piece and piece[0] & 0b0100:
             header = header.upper()
-        header += ext_payload + b"\r\n"
+        header += ext_payload + get_crlf()
 
         handler.wfile.write(header)
         handler.wfile.write(piece)
         if trailer_payload and not piece:
             handler.wfile.write(trailer_payload + b": true\r\n")
-        handler.wfile.write(b"\r\n")
+        handler.wfile.write(get_crlf())
         if not piece:
             break
         payload = payload[chunk_size:]
