Get test server properties from environment to make life a little easier when testing.

#rnx

diff --git a/Engine/Source/Runtime/Experimental/IoStore/HttpClient/Private/TestServer.py b/Engine/Source/Runtime/Experimental/IoStore/HttpClient/Private/TestServer.py
--- a/Engine/Source/Runtime/Experimental/IoStore/HttpClient/Private/TestServer.py
+++ b/Engine/Source/Runtime/Experimental/IoStore/HttpClient/Private/TestServer.py
@@ -431,7 +431,8 @@ def main():
         openssl_bin = str(x)
     root_cert, server_cert, server_key = gen_test_certs(openssl_bin)
 
-    httpd_port = int(random.random() * 0x8000) + 0x4000
+    httpd_port = int(os.getenv("IasTestServerPort", 0))
+    httpd_port = httpd_port or (int(random.random() * 0x8000) + 0x4000)
 
     def start_svc(target, *args):
         threading.Thread(target=target, args=args, daemon=True).start()
