Iris
- Fixed bitstream error that could occur if total size of a batch exceeded 2^20 bits.

diff --git a/Engine/Source/Runtime/Experimental/Iris/Core/Private/Iris/ReplicationSystem/ReplicationReader.cpp b/Engine/Source/Runtime/Experimental/Iris/Core/Private/Iris/ReplicationSystem/ReplicationReader.cpp
--- a/Engine/Source/Runtime/Experimental/Iris/Core/Private/Iris/ReplicationSystem/ReplicationReader.cpp
+++ b/Engine/Source/Runtime/Experimental/Iris/Core/Private/Iris/ReplicationSystem/ReplicationReader.cpp
@@ -661,7 +661,7 @@ uint32 FReplicationReader::ReadObjectsInBatch(FNetSerializationContext& Context,
 	return ReadObjectCount;
 }
 
-uint32 FReplicationReader::ReadObjectBatch(FNetSerializationContext& Context)
+uint32 FReplicationReader::ReadObjectBatch(FNetSerializationContext& Context, uint32 ReadObjectFlags)
 {
 	FNetBitStreamReader& Reader = *Context.GetBitStreamReader();
 
@@ -705,7 +705,7 @@ uint32 FReplicationReader::ReadObjectBatch(FNetSerializationContext& Context)
 	uint32 BatchSize = 0U;
 	// Read Batch size
 	{
-		const uint32 NumBitsUsedForBatchSize = Parameters.NumBitsUsedForBatchSize;
+		const uint32 NumBitsUsedForBatchSize = (ReadObjectFlags & EReadObjectFlag::ReadObjectFlag_IsReadingHugeObjectBatch) == 0U ? Parameters.NumBitsUsedForBatchSize : Parameters.NumBitsUsedForHugeObjectBatchSize;
 
 		UE_NET_TRACE_SCOPE(BatchSize, Reader, Context.GetTraceCollector(), ENetTraceVerbosity::Trace);
 		BatchSize = Reader.ReadBits(NumBitsUsedForBatchSize);
@@ -2176,7 +2176,7 @@ void FReplicationReader::DispatchEndReplication(FNetSerializationContext& Contex
 }
 
 
-void FReplicationReader::ReadObjects(FNetSerializationContext& Context, uint32 ObjectBatchCountToRead)
+void FReplicationReader::ReadObjects(FNetSerializationContext& Context, uint32 ObjectBatchCountToRead, uint32 ReadObjectFlags)
 {
 	IRIS_PROFILER_SCOPE(ReplicationReader_ReadObjects);
 
@@ -2184,7 +2184,7 @@ void FReplicationReader::ReadObjects(FNetSerializationContext& Context, uint32 O
 	
 	while (ObjectBatchCountToRead && !Context.HasErrorOrOverflow())
 	{
-		ReadObjectBatch(Context);
+		ReadObjectBatch(Context, ReadObjectFlags);
 		--ObjectBatchCountToRead;
 	}
 
@@ -2232,7 +2232,8 @@ void FReplicationReader::ProcessHugeObjectAttachment(FNetSerializationContext& C
 	// Reserve space for more dispatch infos as needed, we allocate some extra to account for subobjects etc
 	ObjectsToDispatchArray->Grow(HugeObjectHeader.ObjectCount + ObjectsToDispatchSlackCount, TempLinearAllocator);
 
-	ReadObjects(HugeObjectSerializationContext, HugeObjectHeader.ObjectCount);
+	const uint32 ReadObjectFlags = EReadObjectFlag::ReadObjectFlag_IsReadingHugeObjectBatch;
+	ReadObjects(HugeObjectSerializationContext, HugeObjectHeader.ObjectCount, ReadObjectFlags);
 	if (HugeObjectSerializationContext.HasErrorOrOverflow())
 	{
 		Context.SetError(GNetError_BitStreamError);
@@ -2515,7 +2516,8 @@ void FReplicationReader::Read(FNetSerializationContext& Context)
 		return;
 	}
 
-	ReadObjects(Context, ObjectBatchCountToRead);
+	const uint32 ReadObjectFlags = 0U;
+	ReadObjects(Context, ObjectBatchCountToRead, ReadObjectFlags);
 	if (Context.HasErrorOrOverflow())
 	{
 		return;
diff --git a/Engine/Source/Runtime/Experimental/Iris/Core/Private/Iris/ReplicationSystem/ReplicationReader.h b/Engine/Source/Runtime/Experimental/Iris/Core/Private/Iris/ReplicationSystem/ReplicationReader.h
--- a/Engine/Source/Runtime/Experimental/Iris/Core/Private/Iris/ReplicationSystem/ReplicationReader.h
+++ b/Engine/Source/Runtime/Experimental/Iris/Core/Private/Iris/ReplicationSystem/ReplicationReader.h
@@ -140,8 +140,13 @@ private:
 	// Read index part of handle
 	FNetRefHandle ReadNetRefHandleId(FNetSerializationContext& Context, FNetBitStreamReader& Reader) const;
 
+	enum EReadObjectFlag : unsigned
+	{
+		ReadObjectFlag_IsReadingHugeObjectBatch = 1U,
+	};
+
 	// Read a new or updated object
-	uint32 ReadObjectBatch(FNetSerializationContext& Context);
+	uint32 ReadObjectBatch(FNetSerializationContext& Context, uint32 ReadObjectFlags);
 
 	// Read object or subobject
 	void ReadObjectInBatch(FNetSerializationContext& Context, FNetRefHandle BatchHandle, bool bIsSubObject);
@@ -152,7 +157,7 @@ private:
 	uint32 ReadObjectsPendingDestroy(FNetSerializationContext& Context);
 
 	// Read state data for all incoming objects
-	void ReadObjects(FNetSerializationContext& Context, uint32 ObjectCountToRead);
+	void ReadObjects(FNetSerializationContext& Context, uint32 ObjectCountToRead, uint32 ReadObjectFlags);
 
 	// Process a single huge object attachment
 	void ProcessHugeObjectAttachment(FNetSerializationContext& Context, const TRefCountPtr<FNetBlob>& Attachment);
diff --git a/Engine/Source/Runtime/Experimental/Iris/Core/Private/Iris/ReplicationSystem/ReplicationTypes.h b/Engine/Source/Runtime/Experimental/Iris/Core/Private/Iris/ReplicationSystem/ReplicationTypes.h
--- a/Engine/Source/Runtime/Experimental/Iris/Core/Private/Iris/ReplicationSystem/ReplicationTypes.h
+++ b/Engine/Source/Runtime/Experimental/Iris/Core/Private/Iris/ReplicationSystem/ReplicationTypes.h
@@ -25,7 +25,8 @@ struct FReplicationParameters
 	bool bAllowDelayingAttachmentsWithUnresolvedReferences = false;
 	uint32 SmallObjectBitThreshold = 160U; // Number of bits remaining in a packet for us to consider trying to serialize a replicated object
 	uint32 MaxFailedSmallObjectCount = 10U;	// Number of objects that we try to serialize after an initial stream overflow to fill up a packet, this can improve bandwidth usage but comes at a cpu cost
-	uint32 NumBitsUsedForBatchSize = 20U;
+	uint32 NumBitsUsedForBatchSize = 16U;
+	uint32 NumBitsUsedForHugeObjectBatchSize = 32U;
 };
 
 enum EReplicatedDestroyHeaderFlags : uint32
diff --git a/Engine/Source/Runtime/Experimental/Iris/Core/Private/Iris/ReplicationSystem/ReplicationWriter.cpp b/Engine/Source/Runtime/Experimental/Iris/Core/Private/Iris/ReplicationSystem/ReplicationWriter.cpp
--- a/Engine/Source/Runtime/Experimental/Iris/Core/Private/Iris/ReplicationSystem/ReplicationWriter.cpp
+++ b/Engine/Source/Runtime/Experimental/Iris/Core/Private/Iris/ReplicationSystem/ReplicationWriter.cpp
@@ -2007,7 +2007,7 @@ FReplicationWriter::EWriteObjectStatus FReplicationWriter::WriteObjectAndSubObje
 	// We only need to write batch info for root objects
 	const bool bWriteBatchInfo = !Info.IsSubObject;
 	uint32 InitialStateHeaderPos = 0U;
-	const uint32 NumBitsUsedForBatchSize = Parameters.NumBitsUsedForBatchSize;
+	const uint32 NumBitsUsedForBatchSize = (WriteObjectFlags & EWriteObjectFlag::WriteObjectFlag_IsWritingHugeObjectBatch) == 0U ? Parameters.NumBitsUsedForBatchSize : Parameters.NumBitsUsedForHugeObjectBatchSize;
 
 	// This is the beginning of what we treat as a batch on the receiving end
 	if (bWriteBatchInfo)
@@ -2439,12 +2439,24 @@ FReplicationWriter::EWriteObjectStatus FReplicationWriter::WriteObjectAndSubObje
 	if (OutBatchInfo.ParentInternalIndex == InternalIndex)
 	{
 		FBatchObjectInfo& ParentBatchEntry = OutBatchInfo.ObjectInfos[ParentBatchEntryIndex];
-
-		const uint32 WrittenBitsInBatch = (Writer.GetPosBits() - InitialStateHeaderPos) - NumBitsUsedForBatchSize;
 		
 		const bool bWroteData = (ParentBatchEntry.bSentState || ParentBatchEntry.bSentAttachments || bSentTearOff || Info.SubObjectPendingDestroy);
 		if (bWroteData || (SubObjectsWrittenBits != 0U))
 		{
+			const uint32 MaxBatchSize = NumBitsUsedForBatchSize == 32U ? ~0U : ((1U << NumBitsUsedForBatchSize) - 1U);
+			const uint32 WrittenBitsInBatch = (Writer.GetPosBits() - InitialStateHeaderPos) - NumBitsUsedForBatchSize;
+
+			// Validate size written (excluding exports)
+			if (WrittenBitsInBatch >= MaxBatchSize)
+			{
+				UE_LOG(LogIris, Error, TEXT("FReplicationWriter::WriteObjectAndSubObjects batch too large Conn: %u, WrittenBitsInBatch: %u >= MaxBatchSize:%u when writing object %s ( InternalIndex: %u )"), Parameters.ConnectionId, WrittenBitsInBatch, MaxBatchSize, *NetRefHandle.ToString(), InternalIndex);
+				ensureMsgf(WrittenBitsInBatch >= MaxBatchSize, TEXT("FReplicationWriter::WriteObjectAndSubObjects batch too large WrittenBitsInBatch: %u >= MaxBatchSize:%u when writing object %s ( InternalIndex: %u )"), WrittenBitsInBatch, MaxBatchSize, *NetRefHandle.ToString(), InternalIndex);
+				Context.SetError(NetError_ObjectStateTooLarge);
+				Writer.DoOverflow();
+
+				return EWriteObjectStatus::BitStreamOverflow;
+			}
+
 			const FObjectReferenceCache::EWriteExportsResult WriteExportResult = ObjectReferenceCache->WritePendingExports(Context, InternalIndex);
 
 			if (WriteExportResult == FObjectReferenceCache::EWriteExportsResult::BitStreamOverflow)
@@ -2576,7 +2588,7 @@ int FReplicationWriter::PrepareAndSendHugeObjectPayload(FNetSerializationContext
 	FBatchInfo BatchInfo;
 	BatchInfo.Type = EBatchInfoType::Internal;
 	BatchInfo.ParentInternalIndex = InternalIndex;
-	uint32 WriteObjectFlags = WriteObjectFlag_State;
+	uint32 WriteObjectFlags = WriteObjectFlag_State | WriteObjectFlag_IsWritingHugeObjectBatch;
 	// Get the creation going as quickly as possible.
 	if (!Context.IsInitState())
 	{
diff --git a/Engine/Source/Runtime/Experimental/Iris/Core/Private/Iris/ReplicationSystem/ReplicationWriter.h b/Engine/Source/Runtime/Experimental/Iris/Core/Private/Iris/ReplicationSystem/ReplicationWriter.h
--- a/Engine/Source/Runtime/Experimental/Iris/Core/Private/Iris/ReplicationSystem/ReplicationWriter.h
+++ b/Engine/Source/Runtime/Experimental/Iris/Core/Private/Iris/ReplicationSystem/ReplicationWriter.h
@@ -373,6 +373,7 @@ private:
 		WriteObjectFlag_State = 1U,
 		WriteObjectFlag_Attachments = WriteObjectFlag_State << 1U,
 		WriteObjectFlag_HugeObject = WriteObjectFlag_Attachments << 1U,
+		WriteObjectFlag_IsWritingHugeObjectBatch = WriteObjectFlag_HugeObject << 1U,
 	};
 
 	enum class EWriteObjectRetryMode : unsigned
