Dirty kinematic particles that move on the physics thread. This is required to make sure we get up to date acceleration structures for kinematics that are substepped - or kinematics that move on intermediate async tick updates.

diff --git a/Engine/Source/Runtime/Experimental/Chaos/Private/Chaos/PBDRigidsEvolutionGBF.cpp b/Engine/Source/Runtime/Experimental/Chaos/Private/Chaos/PBDRigidsEvolutionGBF.cpp
--- a/Engine/Source/Runtime/Experimental/Chaos/Private/Chaos/PBDRigidsEvolutionGBF.cpp
+++ b/Engine/Source/Runtime/Experimental/Chaos/Private/Chaos/PBDRigidsEvolutionGBF.cpp
@@ -1029,6 +1029,14 @@ void FPBDRigidsEvolutionGBF::ApplyKinematicTargets(const FReal Dt, const FReal S
 	// Apply kinematic targets in parallel
 	Particles.GetActiveMovingKinematicParticlesView().ParallelFor(ApplyParticleKinematicTarget);
 
+	for(FTransientPBDRigidParticleHandle& Particle : Particles.GetActiveMovingKinematicParticlesView())
+	{
+		// Moving kinematics need to have SQ updates queued to apply substepped movement to the
+		// physics thread acceleration structure in order to get correct kinematic collisions
+		// and pass a correct structure back to the external thread
+		DirtyParticle(Particle);
+	}
+
 	// done with update, let's clear the tracking structures
 	if (IsLastStep)
 	{
