Iris
- FReplicationWriter::UpdateDirtyGlobalLifetimeConditionals now marks subobject owner dirty

diff --git a/Engine/Source/Runtime/Experimental/Iris/Core/Private/Iris/ReplicationSystem/ReplicationWriter.cpp b/Engine/Source/Runtime/Experimental/Iris/Core/Private/Iris/ReplicationSystem/ReplicationWriter.cpp
--- a/Engine/Source/Runtime/Experimental/Iris/Core/Private/Iris/ReplicationSystem/ReplicationWriter.cpp
+++ b/Engine/Source/Runtime/Experimental/Iris/Core/Private/Iris/ReplicationSystem/ReplicationWriter.cpp
@@ -852,10 +852,30 @@ void FReplicationWriter::UpdateDirtyGlobalLifetimeConditionals(TArrayView<FInter
 	for (FInternalNetRefIndex InternalObjectIndex : ObjectsWithDirtyConditionals)
 	{
 		FReplicationInfo& Info = ReplicatedObjects[InternalObjectIndex];
-		if (Info.GetState() != EReplicatedObjectState::Invalid)
+		if (Info.GetState() != EReplicatedObjectState::Invalid && Info.GetState() < EReplicatedObjectState::PendingDestroy)
 		{
-			// We just set a flag that we will process if the object is going to be replicated
-			Info.HasDirtyConditionals = 1;
+			if (Info.IsSubObject)
+			{
+				const FNetRefHandleManager::FReplicatedObjectData& ObjectData = NetRefHandleManager->GetReplicatedObjectDataNoCheck(InternalObjectIndex);
+				if (ObjectData.IsSubObject())
+				{
+					FReplicationInfo& OwnerInfo = ReplicatedObjects[ObjectData.SubObjectRootIndex];
+					if (OwnerInfo.GetState() != EReplicatedObjectState::Invalid && OwnerInfo.GetState() < EReplicatedObjectState::PendingDestroy)
+					{
+						UE_LOG_REPLICATIONWRITER_CONN(TEXT("UpdateDirtyGlobalLifetimeConditionals for - %s"), ToCStr(NetRefHandleManager->PrintObjectFromIndex(ObjectData.SubObjectRootIndex)));
+						
+						MarkObjectDirty(ObjectData.SubObjectRootIndex, "UpdateDirtyGlobalLifetimeConditionals");
+						OwnerInfo.HasDirtyConditionals = 1U;
+					}
+				}
+			}
+			else
+			{
+				UE_LOG_REPLICATIONWRITER_CONN(TEXT("UpdateDirtyGlobalLifetimeConditionals for - %s"), ToCStr(NetRefHandleManager->PrintObjectFromIndex(InternalObjectIndex)));
+	
+				MarkObjectDirty(InternalObjectIndex, "UpdateDirtyGlobalLifetimeConditionals2");
+				Info.HasDirtyConditionals = 1U;
+			}
 		}
 	}
 }
