GeometryCollection : Blueprint access for ChaosSolverActor

diff --git a/Engine/Source/Runtime/Experimental/GeometryCollectionEngine/Public/GeometryCollection/GeometryCollectionComponent.h b/Engine/Source/Runtime/Experimental/GeometryCollectionEngine/Public/GeometryCollection/GeometryCollectionComponent.h
--- a/Engine/Source/Runtime/Experimental/GeometryCollectionEngine/Public/GeometryCollection/GeometryCollectionComponent.h
+++ b/Engine/Source/Runtime/Experimental/GeometryCollectionEngine/Public/GeometryCollection/GeometryCollectionComponent.h
@@ -683,7 +683,7 @@ public:
 	//~ End UMeshComponent Interface.
 
 	/** Chaos RBD Solver override. Will use the world's default solver actor if null. */
-	UPROPERTY(EditAnywhere, Category = "ChaosPhysics", meta = (DisplayName = "Chaos Solver"))
+	UPROPERTY(EditAnywhere, BlueprintReadWrite, Category = "ChaosPhysics", meta = (DisplayName = "Chaos Solver"))
 	TObjectPtr<AChaosSolverActor> ChaosSolverActor;
 
 	/**
