Option to disable receving chunked content

#rnx

diff --git a/Engine/Source/Runtime/Experimental/IoStore/HttpClient/Internal/IO/Http/Client.h b/Engine/Source/Runtime/Experimental/IoStore/HttpClient/Internal/IO/Http/Client.h
--- a/Engine/Source/Runtime/Experimental/IoStore/HttpClient/Internal/IO/Http/Client.h
+++ b/Engine/Source/Runtime/Experimental/IoStore/HttpClient/Internal/IO/Http/Client.h
@@ -215,6 +215,7 @@ public:
 		FCertRootsRef		VerifyCert		= {};
 		uint32				BufferSize		= 256;
 		bool				bAutoRedirect	= false;
+		bool				bAllowChunked	= true;
 	};
 
 	template <typename... T> [[nodiscard]] FRequest Get(T&&... t)  { return Request("GET",  Forward<T&&>(t)...); }
diff --git a/Engine/Source/Runtime/Experimental/IoStore/HttpClient/Private/Activity.inl b/Engine/Source/Runtime/Experimental/IoStore/HttpClient/Private/Activity.inl
--- a/Engine/Source/Runtime/Experimental/IoStore/HttpClient/Private/Activity.inl
+++ b/Engine/Source/Runtime/Experimental/IoStore/HttpClient/Private/Activity.inl
@@ -80,7 +80,8 @@ struct alignas(16) FActivity
 	uint8				IsKeepAlive : 1;
 	uint8				NoContent : 1;
 	uint8				bFollow30x : 1;
-	uint8				_Unused0 : 5;
+	uint8				bAllowChunked : 1;
+	uint8				_Unused0 : 4;
 	uint32				StateParam = 0;
 #if IAS_HTTP_WITH_PERF
 	FStopwatch			Stopwatch;
diff --git a/Engine/Source/Runtime/Experimental/IoStore/HttpClient/Private/Loop.inl b/Engine/Source/Runtime/Experimental/IoStore/HttpClient/Private/Loop.inl
--- a/Engine/Source/Runtime/Experimental/IoStore/HttpClient/Private/Loop.inl
+++ b/Engine/Source/Runtime/Experimental/IoStore/HttpClient/Private/Loop.inl
@@ -214,6 +214,12 @@ static FOutcome DoRecvMessage(FActivity* Activity, FPeerType& Peer)
 	// Validate that the server's told us how and how much it will transmit
 	if (bChunked)
 	{
+		if (Activity->bAllowChunked == 0)
+		{
+			Activity_SetError(Activity, "Chunked transfer encoding disabled (ERRNOCHUNK)");
+			return FOutcome::Error(Activity->ErrorReason);
+		}
+
 		ContentLength = -1;
 	}
 	else if (ContentLength < 0)
@@ -1526,6 +1532,7 @@ FRequest FEventLoop::Request(
 	FHost* Host = Activity->Host = Buffer.Alloc<FHost>();
 	Activity->IsKeepAlive = 0;
 	Activity->bFollow30x = (Params->bAutoRedirect == true);
+	Activity->bAllowChunked = (Params->bAllowChunked == true);
 
 	uint32 HostNameLength = HostName.Len();
 	char* HostNamePtr = Buffer.Alloc<char>(HostNameLength + 1);
