Tests for disabling reception of chunks

#rnx

diff --git a/Engine/Source/Runtime/Experimental/IoStore/HttpClient/Private/Test.inl b/Engine/Source/Runtime/Experimental/IoStore/HttpClient/Private/Test.inl
--- a/Engine/Source/Runtime/Experimental/IoStore/HttpClient/Private/Test.inl
+++ b/Engine/Source/Runtime/Experimental/IoStore/HttpClient/Private/Test.inl
@@ -513,9 +513,9 @@ static void ChunkedTest(const ANSICHAR* TestHost)
 		}
 	}
 
-	// Rudimentary coverage for tranfers with trailing headers.
-	bool bThError = false;
-	auto ExpectError = [&bThError, Dest=FIoBuffer()] (const FTicketStatus& Status) mutable
+	// Rudimentary coverage for transfers with trailing headers.
+	uint32 ErrorMarks;
+	auto ExpectError = [&ErrorMarks, Dest=FIoBuffer()] (const FTicketStatus& Status) mutable
 	{
 		if (Status.GetId() == FTicketStatus::EId::Response)
 		{
@@ -530,12 +530,24 @@ static void ChunkedTest(const ANSICHAR* TestHost)
 		}
 
 		FAnsiStringView Reason = Status.GetErrorReason();
-		bThError = Reason.Contains("ERRTRAIL", ESearchCase::IgnoreCase);
+		ErrorMarks |= Reason.Contains("ERRTRAIL",   ESearchCase::CaseSensitive) ? 1 : 0;
+		ErrorMarks |= Reason.Contains("ERRNOCHUNK", ESearchCase::CaseSensitive) ? 2 : 0;
 	};
+	ErrorMarks = 0;
 	BuildUrl(16 << 10, "/trailer");
 	Loop.Send(Loop.Get(Url), ExpectError);
 	WaitForLoopIdle();
-	check(bThError);
+	check(ErrorMarks == 1);
+
+	// Disabling of chunked transfers 
+	{
+		ErrorMarks = 0;
+		FEventLoop::FRequestParams RequestParams = { .bAllowChunked = false };
+		BuildUrl(16 << 10, "");
+		Loop.Send(Loop.Get(Url, &RequestParams), ExpectError);
+		WaitForLoopIdle();
+		check(ErrorMarks == 2);
+	}
 }
 
 ////////////////////////////////////////////////////////////////////////////////
