Disable IAS' chunk reception so things remain as they were.

#rnx

diff --git a/Engine/Source/Runtime/Experimental/IoStore/OnDemand/Private/DistributionEndpoints.cpp b/Engine/Source/Runtime/Experimental/IoStore/OnDemand/Private/DistributionEndpoints.cpp
--- a/Engine/Source/Runtime/Experimental/IoStore/OnDemand/Private/DistributionEndpoints.cpp
+++ b/Engine/Source/Runtime/Experimental/IoStore/OnDemand/Private/DistributionEndpoints.cpp
@@ -82,7 +82,8 @@ FDistributionEndpoints::EResult FDistributionEndpoints::ResolveEndpoints(const F
 		Loop.SetFailTimeout(GDistributedEndpointTimeout * 1000);
 	}
 
-	FRequest Request = Loop.Get(AnsiUrl);
+	FEventLoop::FRequestParams RequestParams = { .bAllowChunked = false };
+	FRequest Request = Loop.Get(AnsiUrl, &RequestParams);
 	Request.Header("Accept", "application/json");
 	Loop.Send(MoveTemp(Request), OnRequestStatus);
 
diff --git a/Engine/Source/Runtime/Experimental/IoStore/OnDemand/Private/OnDemandHttpClient.cpp b/Engine/Source/Runtime/Experimental/IoStore/OnDemand/Private/OnDemandHttpClient.cpp
--- a/Engine/Source/Runtime/Experimental/IoStore/OnDemand/Private/OnDemandHttpClient.cpp
+++ b/Engine/Source/Runtime/Experimental/IoStore/OnDemand/Private/OnDemandHttpClient.cpp
@@ -18,6 +18,14 @@ static FAutoConsoleVariableRef CVar_IasHttpFailTimeOutMs(
 	TEXT("Fail infinite network waits that take longer than this (in ms, 0=disabled)")
 );
 
+
+static bool GIasHttpAllowChunkedXfer = false;
+static FAutoConsoleVariableRef CVar_IasHttpAllowChunkedXfer(
+	TEXT("ias.HttpAllowChunkedXfer"),
+	GIasHttpAllowChunkedXfer,
+	TEXT("Enable/disable IAS' support for chunked transfer encoding")
+);
+
 static void LogHttpResult(const TCHAR* Host, const TCHAR* Url, uint32 StatusCode, uint64 DurationMs, uint64 Size, uint64 Offset, const char* Memo = "ok")
 {
 	Size >>= 10;
@@ -205,7 +213,8 @@ void FHttpClient::IssueRequest(FRequestParams&& Params)
 	TUniquePtr<FConnectionPool>& Connection = Connections[Params.Endpoint];
 	check(Connection.IsValid());
 
-	FRequest Request = EventLoop.Get(Url, *Connection);
+	FEventLoop::FRequestParams RequestParams = { .bAllowChunked = GIasHttpAllowChunkedXfer };
+	FRequest Request = EventLoop.Get(Url, *Connection, &RequestParams);
 
 	if (Params.Range.GetOffset() > 0 || Params.Range.GetLength() > 0)
 	{
