If the first read from the IAS cache fails we now log an error instead of a veryverbose message.

#rnx

diff --git a/Engine/Source/Runtime/Experimental/IoStore/OnDemand/Private/JournaledCache.cpp b/Engine/Source/Runtime/Experimental/IoStore/OnDemand/Private/JournaledCache.cpp
--- a/Engine/Source/Runtime/Experimental/IoStore/OnDemand/Private/JournaledCache.cpp
+++ b/Engine/Source/Runtime/Experimental/IoStore/OnDemand/Private/JournaledCache.cpp
@@ -923,7 +923,7 @@ static int32 LoadCache(FDiskCache& DiskCache)
 		Handle->Seek(0);
 		if (!Handle->Read(Data.Get(), DataSize))
 		{
-			UE_LOG(LogIas, VeryVerbose, TEXT("JournaledCache: failed reading journal"));
+			UE_LOG(LogIas, Error, TEXT("JournaledCache: failed reading journal"));
 		}
 	}
 
