Better feedback to user from proxy; show port, added summary header and sign of life.

#rnx

diff --git a/Engine/Source/Runtime/Experimental/IoStore/OnDemand/Private/Tool/Socks.cpp b/Engine/Source/Runtime/Experimental/IoStore/OnDemand/Private/Tool/Socks.cpp
--- a/Engine/Source/Runtime/Experimental/IoStore/OnDemand/Private/Tool/Socks.cpp
+++ b/Engine/Source/Runtime/Experimental/IoStore/OnDemand/Private/Tool/Socks.cpp
@@ -164,6 +164,9 @@ static void Socks5Server(uint32 Port)
 
 	TArray<FSocks5Stats*> AllStats;
 	int32 Counter = 0;
+	uint64 UpDown[2] = {};
+
+	std::printf("Listening on 0.0.0.0:%u\n\n", Port);
 
 	SOCKET Sock = socket(AF_INET, SOCK_STREAM, IPPROTO_TCP);
 	check(Sock != INVALID_SOCKET);
@@ -174,7 +177,7 @@ static void Socks5Server(uint32 Port)
 	Addr.sin_port = htons(uint16(Port));
 	check(bind(Sock, &(sockaddr&)Addr, sizeof(Addr)) == 0);
 	check(listen(Sock, 32) == 0);
-	for (;;)
+	for (uint32 Tick = 0; !IsEngineExitRequested();)
 	{
 		pollfd Poll = { Sock, POLLIN };
 		int32 Result = WSAPoll(&Poll, 1, 654);
@@ -209,25 +212,31 @@ static void Socks5Server(uint32 Port)
 				continue;
 			}
 
+			uint64 DownKiB = Stats->Counts[0].load(std::memory_order_relaxed);
+			uint64 UpKiB = Stats->Counts[1].load(std::memory_order_relaxed);
+
 			std::printf("%04d: ", Stats->Id.load(std::memory_order_relaxed));
 			PrintIp(Stats->Source);
 			std::printf("  ->  ");
 			PrintIp(Stats->Dest.load(std::memory_order_relaxed));
 			std::printf(
-				" : d/u %9llu/%-9llu KiB\n",
-				Stats->Counts[0].load(std::memory_order_relaxed),
-				Stats->Counts[1].load(std::memory_order_relaxed)
+				" : d:%9llu u:%9llu KiB\n",
+				DownKiB,
+				UpKiB
 			);
+
+			UpDown[0] += DownKiB;
+			UpDown[1] += UpKiB;
 		}
 
 		for (uint32 i = 0; i < LineClears; ++i)
 		{
 			std::printf("\x1b[2K\n");
 		}
-		if (uint32 Count = uint32(AllStats.Num()) + LineClears; Count)
-		{
-			std::printf("\x1b[%dF", Count);
-		}
+		std::printf("\x1b[%dF", uint32(AllStats.Num()) + LineClears + 1);
+
+		const char* SignOfLife = ".oOo" + (Tick++ & 0x3);
+		std::printf("[%.1s] n:%d d:%llu u:%llu\n", SignOfLife, Counter, UpDown[0], UpDown[1]);
 	}
 
 	closesocket(Sock);
