Fixed parsing of command line argument returning junk.

#rnx

diff --git a/Engine/Source/Runtime/Experimental/IoStore/OnDemand/Private/Tool/Socks.cpp b/Engine/Source/Runtime/Experimental/IoStore/OnDemand/Private/Tool/Socks.cpp
--- a/Engine/Source/Runtime/Experimental/IoStore/OnDemand/Private/Tool/Socks.cpp
+++ b/Engine/Source/Runtime/Experimental/IoStore/OnDemand/Private/Tool/Socks.cpp
@@ -254,7 +254,7 @@ static int32 SocksCommandEntry(const FContext& Context)
 	ON_SCOPE_EXIT { WSACleanup(); };
 #endif
 
-	uint32 Port = Context.Get(TEXT("-Port"), 24930u);
+	uint32 Port = Context.Get<uint32>(TEXT("-Port"), 24930u);
 	std::thread(Socks5Server, Port).join();
 
 	return 0;
@@ -266,7 +266,7 @@ static FCommand SocksCommand(
 	TEXT("Socks"),
 	TEXT("Rudimentary SOCKS5 proxy to aid in testing IAS traffic"),
 	{
-		TArgument<FStringView>(TEXT("-Port"), TEXT("Port to listen on (default=24930")),
+		TArgument<uint32>(TEXT("-Port"), TEXT("Port to listen on (default=24930)")),
 	}
 );
 
