Fixed summary up/down running away with itself.

#rnx

diff --git a/Engine/Source/Runtime/Experimental/IoStore/OnDemand/Private/Tool/Socks.cpp b/Engine/Source/Runtime/Experimental/IoStore/OnDemand/Private/Tool/Socks.cpp
--- a/Engine/Source/Runtime/Experimental/IoStore/OnDemand/Private/Tool/Socks.cpp
+++ b/Engine/Source/Runtime/Experimental/IoStore/OnDemand/Private/Tool/Socks.cpp
@@ -164,7 +164,7 @@ static void Socks5Server(uint32 Port)
 
 	TArray<FSocks5Stats*> AllStats;
 	int32 Counter = 0;
-	uint64 UpDown[2] = {};
+	uint64 SoFar[2] = {};
 
 	std::printf("Listening on 0.0.0.0:%u\n\n", Port);
 
@@ -199,12 +199,19 @@ static void Socks5Server(uint32 Port)
 			continue;
 		}
 
+		uint64 UpDown[2] = {};
 		uint32 LineClears = 0;
 		for (int32 i = 0, n = AllStats.Num(); i < n; ++i)
 		{
 			FSocks5Stats* Stats = AllStats[i];
+
+			uint64 DownKiB = Stats->Counts[0].load(std::memory_order_relaxed);
+			uint64 UpKiB = Stats->Counts[1].load(std::memory_order_relaxed);
+
 			if (Stats->Id.load(std::memory_order_acquire) == -1)
 			{
+				SoFar[0] += UpKiB;
+				SoFar[1] += DownKiB;
 				delete Stats;
 				AllStats[i] = AllStats.Last();
 				AllStats.Pop();
@@ -212,9 +219,6 @@ static void Socks5Server(uint32 Port)
 				continue;
 			}
 
-			uint64 DownKiB = Stats->Counts[0].load(std::memory_order_relaxed);
-			uint64 UpKiB = Stats->Counts[1].load(std::memory_order_relaxed);
-
 			std::printf("%04d: ", Stats->Id.load(std::memory_order_relaxed));
 			PrintIp(Stats->Source);
 			std::printf("  ->  ");
@@ -236,7 +240,13 @@ static void Socks5Server(uint32 Port)
 		std::printf("\x1b[%dF", uint32(AllStats.Num()) + LineClears + 1);
 
 		const char* SignOfLife = ".oOo" + (Tick++ & 0x3);
-		std::printf("[%.1s] n:%d d:%llu u:%llu\n", SignOfLife, Counter, UpDown[0], UpDown[1]);
+		std::printf(
+			"[%.1s] n:%d d:%llu u:%llu\n",
+			SignOfLife,
+			Counter,
+			SoFar[0] + UpDown[0],
+			SoFar[1] + UpDown[1]
+		);
 	}
 
 	closesocket(Sock);
