Iris
- Moved ensure in protocolmanager to AFTER error logging to get proper logs in CrashReporter

diff --git a/Engine/Source/Runtime/Experimental/Iris/Core/Private/Iris/ReplicationSystem/ReplicationProtocolManager.cpp b/Engine/Source/Runtime/Experimental/Iris/Core/Private/Iris/ReplicationSystem/ReplicationProtocolManager.cpp
--- a/Engine/Source/Runtime/Experimental/Iris/Core/Private/Iris/ReplicationSystem/ReplicationProtocolManager.cpp
+++ b/Engine/Source/Runtime/Experimental/Iris/Core/Private/Iris/ReplicationSystem/ReplicationProtocolManager.cpp
@@ -223,7 +223,7 @@ const FReplicationProtocol* FReplicationProtocolManager::CreateReplicationProtoc
 	if (Params.bValidateProtocolId)
 	{
 		const FReplicationProtocolIdentifier NewProtocolId = CalculateProtocolIdentifier(Fragments);
-		if (!ensureMsgf(NewProtocolId == ProtocolId, TEXT("FReplicationProtocolManager::CreateReplicationProtocol Id mismatch when creating protocol named %s with in ProtocolId:0x%" UINT64_x_FMT "Calculated ProtocolId:0x$%" UINT64_x_FMT), DebugName, ProtocolId, NewProtocolId))
+		if (!NewProtocolId == ProtocolId)
 		{
 			UE_LOG(LogIris, Warning, TEXT("FReplicationProtocolManager::CreateReplicationProtocol Id mismatch when creating protocol named %s with ProtocolId:0x%" UINT64_x_FMT), DebugName, ProtocolId);
  #if UE_NET_ENABLE_PROTOCOLMANAGER_LOG
@@ -235,6 +235,7 @@ const FReplicationProtocol* FReplicationProtocolManager::CreateReplicationProtoc
  				UE_LOG(LogIris, Warning, TEXT("%s"), StringBuilder.ToString());
  			}
  #endif
+			ensureMsgf(NewProtocolId == ProtocolId, TEXT("FReplicationProtocolManager::CreateReplicationProtocol Id mismatch when creating protocol named %s with in ProtocolId:0x%" UINT64_x_FMT "Calculated ProtocolId:0x$%" UINT64_x_FMT), DebugName, ProtocolId, NewProtocolId);
 			return nullptr;
 		}
 	}
