Iris - Fixed bug with object scope hysteresis kicking in for newly added objects. Fixed bug with NetRefHandleIndex growing caused some filtering information to be reset.

diff --git a/Engine/Plugins/Runtime/ReplicationSystemTestPlugin/Source/Private/Tests/ReplicationSystem/Filtering/TestFiltering.cpp b/Engine/Plugins/Runtime/ReplicationSystemTestPlugin/Source/Private/Tests/ReplicationSystem/Filtering/TestFiltering.cpp
--- a/Engine/Plugins/Runtime/ReplicationSystemTestPlugin/Source/Private/Tests/ReplicationSystem/Filtering/TestFiltering.cpp
+++ b/Engine/Plugins/Runtime/ReplicationSystemTestPlugin/Source/Private/Tests/ReplicationSystem/Filtering/TestFiltering.cpp
@@ -2564,7 +2564,7 @@ UE_NET_TEST_FIXTURE(FTestFilteringFixture, HysteresisDoesNotKickInForNewlyCreate
 	Server->UpdateAndSend({ Client });
 
 	// The object was filtered out from the start so hysteresis should not cause it to start replicating.
-	UE_NET_ASSERT_NE(Client->GetReplicationBridge()->GetReplicatedObject(ServerObject->NetRefHandle), nullptr);
+	UE_NET_ASSERT_EQ(Client->GetReplicationBridge()->GetReplicatedObject(ServerObject->NetRefHandle), nullptr);
 }
 
 // Destroyed objects are expected to be destroyed as quickly as possible.
diff --git a/Engine/Source/Runtime/Experimental/Iris/Core/Private/Iris/ReplicationSystem/Filtering/ObjectScopeHysteresisUpdater.cpp b/Engine/Source/Runtime/Experimental/Iris/Core/Private/Iris/ReplicationSystem/Filtering/ObjectScopeHysteresisUpdater.cpp
--- a/Engine/Source/Runtime/Experimental/Iris/Core/Private/Iris/ReplicationSystem/Filtering/ObjectScopeHysteresisUpdater.cpp
+++ b/Engine/Source/Runtime/Experimental/Iris/Core/Private/Iris/ReplicationSystem/Filtering/ObjectScopeHysteresisUpdater.cpp
@@ -21,6 +21,11 @@ void FObjectScopeHysteresisUpdater::Deinit()
 	ObjectsToUpdate.Empty();
 }
 
+void FObjectScopeHysteresisUpdater::OnMaxInternalNetRefIndexIncreased(FInternalNetRefIndex NewMaxInternalIndex)
+{
+	ObjectsToUpdate.SetNumBits(NewMaxInternalIndex);
+}
+
 void FObjectScopeHysteresisUpdater::SetHysteresisFrameCount(FInternalNetRefIndex NetRefIndex, uint16 HysteresisFrameCount)
 {
 	const FLocalIndex LocalIndex = GetOrCreateLocalIndex(NetRefIndex);
diff --git a/Engine/Source/Runtime/Experimental/Iris/Core/Private/Iris/ReplicationSystem/Filtering/ObjectScopeHysteresisUpdater.h b/Engine/Source/Runtime/Experimental/Iris/Core/Private/Iris/ReplicationSystem/Filtering/ObjectScopeHysteresisUpdater.h
--- a/Engine/Source/Runtime/Experimental/Iris/Core/Private/Iris/ReplicationSystem/Filtering/ObjectScopeHysteresisUpdater.h
+++ b/Engine/Source/Runtime/Experimental/Iris/Core/Private/Iris/ReplicationSystem/Filtering/ObjectScopeHysteresisUpdater.h
@@ -21,6 +21,9 @@ public:
 	void Init(uint32 MaxObjectCount);
 	void Deinit();
 
+	/** Called when the maximum InternalNetRefIndex increased and we need to realloc our lists */
+	void OnMaxInternalNetRefIndexIncreased(FInternalNetRefIndex NewMaxInternalIndex);
+
 	/** Sets an hysteresis frame count such that an object will be kept in scope until such many frames has passed. */
 	void SetHysteresisFrameCount(FInternalNetRefIndex NetRefIndex, uint16 HysteresisFrameCount);
 
diff --git a/Engine/Source/Runtime/Experimental/Iris/Core/Private/Iris/ReplicationSystem/Filtering/ReplicationFiltering.cpp b/Engine/Source/Runtime/Experimental/Iris/Core/Private/Iris/ReplicationSystem/Filtering/ReplicationFiltering.cpp
--- a/Engine/Source/Runtime/Experimental/Iris/Core/Private/Iris/ReplicationSystem/Filtering/ReplicationFiltering.cpp
+++ b/Engine/Source/Runtime/Experimental/Iris/Core/Private/Iris/ReplicationSystem/Filtering/ReplicationFiltering.cpp
@@ -282,6 +282,7 @@ void FReplicationFiltering::SetNetObjectListsSize(FInternalNetRefIndex MaxIntern
 
 	// Always allocated and maintained regardless of whether the feature is enabled or not.
 	HysteresisState.ObjectsToClear.SetNumBits(MaxInternalNetRefIndex);
+	HysteresisState.ObjectsExemptFromHysteresis.SetNumBits(MaxInternalNetRefIndex);
 }
 
 void FReplicationFiltering::OnMaxInternalNetRefIndexIncreased(FInternalNetRefIndex NewMaxInternalIndex)
@@ -872,8 +873,8 @@ void FReplicationFiltering::SetPerConnectionListsSize(FPerConnectionInfo& Connec
 	{
 		ConnectionInfo.DynamicFilteredOutObjects.SetNumBits(NewMaxInternalIndex);
 		ConnectionInfo.InProgressDynamicFilteredOutObjects.SetNumBits(NewMaxInternalIndex);
-		ConnectionInfo.DynamicFilteredOutObjectsHysteresisAdjusted.Init(NewMaxInternalIndex);
-		ConnectionInfo.HysteresisUpdater.Init(NewMaxInternalIndex);
+		ConnectionInfo.DynamicFilteredOutObjectsHysteresisAdjusted.SetNumBits(NewMaxInternalIndex);
+		ConnectionInfo.HysteresisUpdater.OnMaxInternalNetRefIndexIncreased(NewMaxInternalIndex);
 	}
 }
 
@@ -926,8 +927,10 @@ void FReplicationFiltering::UpdateObjectsInScope()
 		return;
 	}
 
-	// Clear info for deleted objects
+	// Clear info for deleted objects and dirty filter information for added objects.
 	{
+		uint32* ObjectsWithOwnerFilterStorage = ObjectsWithOwnerFilter.GetDataChecked(WordCountForObjectBitArrays);
+		uint32* ObjectsExemptFromHysteresisStorage = HysteresisState.ObjectsExemptFromHysteresis.GetDataChecked(WordCountForObjectBitArrays);
 		const FNetBitArrayView SubObjectInternalIndices = NetRefHandleManager->GetSubObjectInternalIndicesView();
 		uint32 PrevParentIndex = FNetRefHandleManager::InvalidInternalIndex;
 		for (uint32 WordIt = 0, WordEndIt = ModifiedWordIndex; WordIt != WordEndIt; ++WordIt)
@@ -938,7 +941,7 @@ void FReplicationFiltering::UpdateObjectsInScope()
 
 			// Deleted objects can't be dirty and can't have filtering.
 			ObjectsWithDirtyConnectionFilterStorage[WordIndex] &= ExistingObjects;
-			ObjectsWithOwnerFilter.GetWord(WordIndex) &= ExistingObjects;
+			ObjectsWithOwnerFilterStorage[WordIndex] &= ExistingObjects;
 			ObjectsWithDirtyOwnerStorage[WordIndex] &= ExistingObjects;
 
 			// Clear dynamic filters and owner info from deleted objects.
@@ -963,9 +966,13 @@ void FReplicationFiltering::UpdateObjectsInScope()
 				}
 			}
 
+			const uint32 AddedObjects = ExistingObjects & ~PrevExistingObjects;
+
+			// Prevent hysteresis from kicking in on just added objects.
+			ObjectsExemptFromHysteresisStorage[WordIndex] |= AddedObjects;
+
 			// Make sure subobjects that are added after the parent gets properly updated.
 			// Dirtying the parent will cause the subobjects to be updated too.
-			uint32 AddedObjects = ExistingObjects & ~PrevExistingObjects;
 			if (uint32 AddedSubObjects = AddedObjects & SubObjectInternalIndices.GetWord(WordIndex))
 			{
 				for ( ; AddedSubObjects; )
@@ -1497,7 +1504,7 @@ void FReplicationFiltering::UpdateDynamicFiltering()
 			for (const uint32 ObjectIndex : FilteredOutObjects)
 			{
 				const uint32 HysteresisFrameCount = ObjectScopeHysteresisFrameCounts[ObjectIndex];
-				if (HysteresisFrameCount && DynamicFilterEnabledObjects.GetBit(ObjectIndex))
+				if (HysteresisFrameCount && DynamicFilterEnabledObjects.GetBit(ObjectIndex) && !HysteresisState.ObjectsExemptFromHysteresis.GetBit(ObjectIndex))
 				{
 					// We need to adjust the hysteresis frame count to account for when it will be updated. The -1 stems from the fact that updating won't happen until next frame at the earliest.
 					const uint16 TotalHysteresisFrameCount = static_cast<uint16>(HysteresisFrameCount - 1 + AdjustHysteresisForUpdateThrottling);
@@ -2797,6 +2804,7 @@ void FReplicationFiltering::PostUpdateObjectScopeHysteresis()
 {
 	HysteresisState.ObjectsToClearCount = 0;
 	HysteresisState.ObjectsToClear.ClearAllBits();
+	HysteresisState.ObjectsExemptFromHysteresis.ClearAllBits();
 }
 
 void FReplicationFiltering::ClearObjectsFromHysteresis()
diff --git a/Engine/Source/Runtime/Experimental/Iris/Core/Private/Iris/ReplicationSystem/Filtering/ReplicationFiltering.h b/Engine/Source/Runtime/Experimental/Iris/Core/Private/Iris/ReplicationSystem/Filtering/ReplicationFiltering.h
--- a/Engine/Source/Runtime/Experimental/Iris/Core/Private/Iris/ReplicationSystem/Filtering/ReplicationFiltering.h
+++ b/Engine/Source/Runtime/Experimental/Iris/Core/Private/Iris/ReplicationSystem/Filtering/ReplicationFiltering.h
@@ -299,6 +299,9 @@ private:
 
 		// Objects to clear from hysteresis due to being destroyed or removed from dynamic filtering.
 		FNetBitArray ObjectsToClear;
+
+		// Objects that should not be added to hysteresis this frame. Example use case is newly added objects that become filtered out on the first frame.
+		FNetBitArray ObjectsExemptFromHysteresis;
 	};
 
 	// Used for ObjectIndexToDynamicFilterIndex lookup
