Dataflow : ThreadingContext update
- Prevent the ThreadedContext from setting store entries twice.

diff --git a/Engine/Source/Runtime/Experimental/Dataflow/Core/Public/Dataflow/DataflowNodeParameters.h b/Engine/Source/Runtime/Experimental/Dataflow/Core/Public/Dataflow/DataflowNodeParameters.h
--- a/Engine/Source/Runtime/Experimental/Dataflow/Core/Public/Dataflow/DataflowNodeParameters.h
+++ b/Engine/Source/Runtime/Experimental/Dataflow/Core/Public/Dataflow/DataflowNodeParameters.h
@@ -353,7 +353,14 @@ namespace Dataflow
 		{
 			CacheLock->Lock(); ON_SCOPE_EXIT { CacheLock->Unlock(); };
 			
-			DataStore.Emplace(Key, MoveTemp(DataStoreEntry));
+
+			// Threaded evaluation can only set an output once per context evaluation. Otherwise
+			// downstream nodes that are extracting the data will get currupted store entries. 
+			TUniquePtr<FContextCacheElementBase>* CurrentData = DataStore.Find(Key);
+			if (!CurrentData || !(*CurrentData) || (*CurrentData)->GetTimestamp() < GetTimestamp())
+			{
+				DataStore.Emplace(Key, MoveTemp(DataStoreEntry));
+			}
 		}
 
 		virtual TUniquePtr<FContextCacheElementBase>* GetDataImpl(FContextCacheKey Key) override
