3 - Chaos flesh simulation support

#virtualized

diff --git a/Engine/Plugins/Animation/MLDeformer/ChaosFleshGenerator/Source/ChaosFleshGenerator/ChaosFleshGenerator.build.cs b/Engine/Plugins/Animation/MLDeformer/ChaosFleshGenerator/Source/ChaosFleshGenerator/ChaosFleshGenerator.build.cs
--- a/Engine/Plugins/Animation/MLDeformer/ChaosFleshGenerator/Source/ChaosFleshGenerator/ChaosFleshGenerator.build.cs
+++ b/Engine/Plugins/Animation/MLDeformer/ChaosFleshGenerator/Source/ChaosFleshGenerator/ChaosFleshGenerator.build.cs
@@ -40,6 +40,9 @@ public class ChaosFleshGenerator : ModuleRules
 				"ChaosFlesh",
 				"ChaosFleshEngine",
 				"ChaosFleshNodes",
+				"DataflowCore",
+				"DataflowEngine",
+				"DataflowSimulation",
 				"Engine",
 				"GeometryCache",
 				"MLDeformerFramework",
diff --git a/Engine/Plugins/Animation/MLDeformer/ChaosFleshGenerator/Source/ChaosFleshGenerator/Private/ChaosFleshGeneratorSimulation.cpp b/Engine/Plugins/Animation/MLDeformer/ChaosFleshGenerator/Source/ChaosFleshGenerator/Private/ChaosFleshGeneratorSimulation.cpp
--- a/Engine/Plugins/Animation/MLDeformer/ChaosFleshGenerator/Source/ChaosFleshGenerator/Private/ChaosFleshGeneratorSimulation.cpp
+++ b/Engine/Plugins/Animation/MLDeformer/ChaosFleshGenerator/Source/ChaosFleshGenerator/Private/ChaosFleshGeneratorSimulation.cpp
@@ -61,9 +61,9 @@ namespace UE::Chaos::FleshGenerator
 	
 		const TArray<FTransform> Transforms = GetBoneTransforms( SkeletalComponent, AnimFrame);
 		FleshComponent.Pose(SkeletalComponent, Transforms);
-		SolverComponent.UpdateFromGameThread(DeltaTime);
+		SolverComponent.WriteToSimulation(DeltaTime);
 		SolverComponent.Simulate(DeltaTime);
-		SolverComponent.UpdateFromSimulation(DeltaTime);
+		SolverComponent.ReadFromSimulation(DeltaTime);
 		SimResource.SimulatedPositions[CacheFrame] = GetRenderPositions(SimResource);
 
 		SimResource.FinishFrame();
diff --git a/Engine/Plugins/Animation/MLDeformer/ChaosFleshGenerator/Source/ChaosFleshGenerator/Private/ChaosFleshGeneratorThreading.cpp b/Engine/Plugins/Animation/MLDeformer/ChaosFleshGenerator/Source/ChaosFleshGenerator/Private/ChaosFleshGeneratorThreading.cpp
--- a/Engine/Plugins/Animation/MLDeformer/ChaosFleshGenerator/Source/ChaosFleshGenerator/Private/ChaosFleshGeneratorThreading.cpp
+++ b/Engine/Plugins/Animation/MLDeformer/ChaosFleshGenerator/Source/ChaosFleshGenerator/Private/ChaosFleshGeneratorThreading.cpp
@@ -51,7 +51,7 @@ namespace UE::Chaos::FleshGenerator
 			DeformableSolver->SolverConstraints = Properties->SolverConstraints;
 			DeformableSolver->SolverForces = Properties->SolverForces;
 			DeformableSolver->SolverDebugging = Properties->SolverDebugging;
-			DeformableSolver->Reset();
+			DeformableSolver->BuildSimulationProxy();
 
 
 			constexpr int32 LODIndex = 0;
diff --git a/Engine/Plugins/Experimental/ChaosFlesh/Source/ChaosFleshEditor/ChaosFleshEditor.Build.cs b/Engine/Plugins/Experimental/ChaosFlesh/Source/ChaosFleshEditor/ChaosFleshEditor.Build.cs
--- a/Engine/Plugins/Experimental/ChaosFlesh/Source/ChaosFleshEditor/ChaosFleshEditor.Build.cs
+++ b/Engine/Plugins/Experimental/ChaosFlesh/Source/ChaosFleshEditor/ChaosFleshEditor.Build.cs
@@ -40,6 +40,7 @@ namespace UnrealBuildTool.Rules
 					"DataflowEngine",
 					"DataflowEnginePlugin",
 					"DataflowEditor",
+					"DataflowSimulation",
 					"EditorStyle",
 					"Engine",
 					"EditorFramework",
diff --git a/Engine/Plugins/Experimental/ChaosFlesh/Source/ChaosFleshEditor/Private/ChaosFlesh/Asset/AssetDefinition_FleshAsset.cpp b/Engine/Plugins/Experimental/ChaosFlesh/Source/ChaosFleshEditor/Private/ChaosFlesh/Asset/AssetDefinition_FleshAsset.cpp
--- a/Engine/Plugins/Experimental/ChaosFlesh/Source/ChaosFleshEditor/Private/ChaosFlesh/Asset/AssetDefinition_FleshAsset.cpp
+++ b/Engine/Plugins/Experimental/ChaosFlesh/Source/ChaosFleshEditor/Private/ChaosFlesh/Asset/AssetDefinition_FleshAsset.cpp
@@ -73,7 +73,11 @@ EAssetCommandResult UAssetDefinition_FleshAsset::OpenAssets(const FAssetOpenArgs
 			{
 				UAssetEditorSubsystem* const AssetEditorSubsystem = GEditor->GetEditorSubsystem<UAssetEditorSubsystem>();
 				UDataflowEditor* const AssetEditor = NewObject<UDataflowEditor>(AssetEditorSubsystem, NAME_None, RF_Transient);
-				AssetEditor->Initialize({ FleshAsset });
+
+				UClass* ActorClass = StaticLoadClass(AActor::StaticClass(), nullptr,
+					TEXT("/ChaosFlesh/BP_FleshPreview.BP_FleshPreview_C"), nullptr, LOAD_None, nullptr);
+				AssetEditor->Initialize({ FleshAsset }, ActorClass);
+				
 				return EAssetCommandResult::Handled;
 			}
 
diff --git a/Engine/Plugins/Experimental/ChaosFlesh/Source/ChaosFleshEngine/ChaosFleshEngine.Build.cs b/Engine/Plugins/Experimental/ChaosFlesh/Source/ChaosFleshEngine/ChaosFleshEngine.Build.cs
--- a/Engine/Plugins/Experimental/ChaosFlesh/Source/ChaosFleshEngine/ChaosFleshEngine.Build.cs
+++ b/Engine/Plugins/Experimental/ChaosFlesh/Source/ChaosFleshEngine/ChaosFleshEngine.Build.cs
@@ -27,6 +27,9 @@ namespace UnrealBuildTool.Rules
 				{
 					"Core",
 					"ProceduralMeshComponent",
+					"DataflowCore",
+					"DataflowEngine",
+					"DataflowSimulation",
 					// ... add other public dependencies that you statically link with here ...
 				}
 				);
@@ -41,6 +44,7 @@ namespace UnrealBuildTool.Rules
 					"ChaosFlesh",
 					"DataflowCore",
 					"DataflowEngine",
+					"DataflowSimulation",
 					"Engine",
 					"FieldSystemEngine",
 					"NetCore",
diff --git a/Engine/Plugins/Experimental/ChaosFlesh/Source/ChaosFleshEngine/Private/ChaosCache/FleshComponentCacheAdapter.cpp b/Engine/Plugins/Experimental/ChaosFlesh/Source/ChaosFleshEngine/Private/ChaosCache/FleshComponentCacheAdapter.cpp
--- a/Engine/Plugins/Experimental/ChaosFlesh/Source/ChaosFleshEngine/Private/ChaosCache/FleshComponentCacheAdapter.cpp
+++ b/Engine/Plugins/Experimental/ChaosFlesh/Source/ChaosFleshEngine/Private/ChaosCache/FleshComponentCacheAdapter.cpp
@@ -382,7 +382,7 @@ namespace Chaos
 			{
 				if(FleshComp->GetDeformableSolver())
 				{
-					return FleshComp->GetDeformableSolver()->Solver.Get();
+					return FleshComp->GetDeformableSolver()->FleshSolverProxy.Solver.Get();
 				}
 			}
 		}
diff --git a/Engine/Plugins/Experimental/ChaosFlesh/Source/ChaosFleshEngine/Private/ChaosFlesh/ChaosDeformableSolverComponent.cpp b/Engine/Plugins/Experimental/ChaosFlesh/Source/ChaosFleshEngine/Private/ChaosFlesh/ChaosDeformableSolverComponent.cpp
--- a/Engine/Plugins/Experimental/ChaosFlesh/Source/ChaosFleshEngine/Private/ChaosFlesh/ChaosDeformableSolverComponent.cpp
+++ b/Engine/Plugins/Experimental/ChaosFlesh/Source/ChaosFleshEngine/Private/ChaosFlesh/ChaosDeformableSolverComponent.cpp
@@ -29,7 +29,7 @@ FAutoConsoleVariableRef CVarChaosEngineDeformableSolverbEnabled(TEXT("p.Chaos.De
 
 UDeformableSolverComponent::UDeformableSolverComponent(const FObjectInitializer& ObjectInitializer)
 	: Super(ObjectInitializer)
-	, Solver()
+	, FleshSolverProxy()
 {
 	PrimaryComponentTick.bCanEverTick = true;
 	bTickInEditor = false;
@@ -91,9 +91,14 @@ void UDeformableSolverComponent::UpdateTickGroup()
 UDeformableSolverComponent::FDeformableSolver::FGameThreadAccess
 UDeformableSolverComponent::GameThreadAccess()
 {
-	return FDeformableSolver::FGameThreadAccess(Solver.Get(), Chaos::Softs::FGameThreadAccessor());
+	return FDeformableSolver::FGameThreadAccess(FleshSolverProxy.Solver.Get(), Chaos::Softs::FGameThreadAccessor());
 }
 
+UDeformableSolverComponent::FDeformableSolver::FPhysicsThreadAccess
+UDeformableSolverComponent::PhysicsThreadAccess()
+{
+	return FDeformableSolver::FPhysicsThreadAccess(FleshSolverProxy.Solver.Get(), Chaos::Softs::FPhysicsThreadAccessor());
+}
 
 bool UDeformableSolverComponent::IsSimulatable() const
 {
@@ -150,7 +155,6 @@ void UDeformableSolverComponent::UpdateDeformableEndTickState(bool bRegister)
 void UDeformableSolverComponent::BeginPlay()
 {
 	Super::BeginPlay();
-	Reset();
 }
 
 void UDeformableSolverComponent::TickComponent(float DeltaTime, enum ELevelTick TickType, FActorComponentTickFunction* ThisTickFunction)
@@ -163,23 +167,28 @@ void UDeformableSolverComponent::TickComponent(float DeltaTime, enum ELevelTick
 		UpdateTickGroup();
 
 		UpdateDeformableEndTickState(IsSimulatable());
-		if(bSimulationTicking)
-		{
-			UpdateFromGameThread(DeltaTime);
 
-			if (SolverTiming.bDoThreadedAdvance)
-			{
-				// see FParallelClothCompletionTask
-				FGraphEventArray Prerequisites;
-				Prerequisites.Add(ParallelDeformableTask);
-				FGraphEventRef DeformableCompletionEvent = TGraphTask<FParallelDeformableTask>::CreateTask(&Prerequisites, ENamedThreads::GameThread).ConstructAndDispatchWhenReady(this, DeltaTime);
-				ThisTickFunction->GetCompletionHandle()->DontCompleteUntil(DeformableCompletionEvent);
-			}
-			else
+		// We only run the simulation if no dataflow solver has been defined
+		if(!SimulationAsset.DataflowAsset)
+		{
+			if(bSimulationTicking)
 			{
-				Simulate(DeltaTime);
+				WriteToSimulation(DeltaTime);
 
-				UpdateFromSimulation(DeltaTime);
+				if (SolverTiming.bDoThreadedAdvance)
+				{
+					// see FParallelClothCompletionTask
+					FGraphEventArray Prerequisites;
+					Prerequisites.Add(ParallelDeformableTask);
+					FGraphEventRef DeformableCompletionEvent = TGraphTask<FParallelDeformableTask>::CreateTask(&Prerequisites, ENamedThreads::GameThread).ConstructAndDispatchWhenReady(this, DeltaTime);
+					ThisTickFunction->GetCompletionHandle()->DontCompleteUntil(DeformableCompletionEvent);
+				}
+				else
+				{
+					Simulate(DeltaTime);
+
+					ReadFromSimulation(DeltaTime);
+				}
 			}
 		}
 	}
@@ -190,14 +199,14 @@ void UDeformableSolverComponent::EndPlay(const EEndPlayReason::Type EndPlayReaso
 	Super::EndPlay(EndPlayReason);
 }
 
-void UDeformableSolverComponent::Reset()
+void UDeformableSolverComponent::BuildSimulationProxy()
 {
 	SCOPE_CYCLE_COUNTER(STAT_ChaosDeformable_UDeformableSolverComponent_Reset);
 	TRACE_CPUPROFILER_EVENT_SCOPE(ChaosDeformable_UDeformableSolverComponent_Reset);
 
 	if (GChaosEngineDeformableCVarParams.bEnableDeformableSolver)
 	{
-		Solver.Reset(new FDeformableSolver({
+		FleshSolverProxy.Solver = MakeUnique<Chaos::Softs::FDeformableSolver>(Chaos::Softs::FDeformableSolverProperties(
 			SolverTiming.NumSubSteps
 			, SolverTiming.NumSolverIterations
 			, SolverTiming.FixTimeStep
@@ -228,7 +237,7 @@ void UDeformableSolverComponent::Reset()
 			, SolverMuscleActivation.bDoMuscleActivation
 			, SolverConstraints.GaussSeidelConstraints.CollisionSpring.bCollideWithFullmesh
 			, SolverConstraints.GaussSeidelConstraints.bEnableDynamicSprings
-		}));
+		));
 
 		for (TObjectPtr<UDeformablePhysicsComponent>& DeformableComponent : ConnectedObjects.DeformableComponents)
 		{
@@ -243,14 +252,29 @@ void UDeformableSolverComponent::Reset()
 	}
 }
 
+void UDeformableSolverComponent::ResetSimulationProxy()
+{
+	for (TObjectPtr<UDeformablePhysicsComponent>& DeformableComponent : ConnectedObjects.DeformableComponents)
+	{
+		if( DeformableComponent )
+		{
+			if (IsSimulating(DeformableComponent))
+			{
+				RemoveDeformableProxy(DeformableComponent);
+			}
+		}
+	}
+	FleshSolverProxy.Solver.Reset();
+}
+
 void UDeformableSolverComponent::RemoveDeformableProxy(UDeformablePhysicsComponent* InComponent)
 {
 	SCOPE_CYCLE_COUNTER(STAT_ChaosDeformable_UDeformableSolverComponent_RemoveDeformableProxy);
 	TRACE_CPUPROFILER_EVENT_SCOPE(ChaosDeformable_UDeformableSolverComponent_RemoveDeformableProxy);
 
-	if (Solver && IsSimulating(InComponent))
+	if (FleshSolverProxy.Solver && IsSimulating(InComponent))
 	{
-		FDeformableSolver::FGameThreadAccess GameThreadSolver(Solver.Get(), Chaos::Softs::FGameThreadAccessor());
+		FDeformableSolver::FGameThreadAccess GameThreadSolver = GameThreadAccess();
 		if (!GameThreadSolver.HasObject(InComponent))
 		{
 			InComponent->RemoveProxy(GameThreadSolver);
@@ -263,9 +287,9 @@ void UDeformableSolverComponent::AddDeformableProxy(UDeformablePhysicsComponent*
 	SCOPE_CYCLE_COUNTER(STAT_ChaosDeformable_UDeformableSolverComponent_AddDeformableProxy);
 	TRACE_CPUPROFILER_EVENT_SCOPE(ChaosDeformable_UDeformableSolverComponent_AddDeformableProxy);
 
-	if (Solver && IsSimulating(InComponent))
+	if (FleshSolverProxy.Solver && IsSimulating(InComponent))
 	{
-		FDeformableSolver::FGameThreadAccess GameThreadSolver(Solver.Get(), Chaos::Softs::FGameThreadAccessor());
+		FDeformableSolver::FGameThreadAccess GameThreadSolver = GameThreadAccess();
 		if (!GameThreadSolver.HasObject(InComponent))
 		{
 			InComponent->AddProxy(GameThreadSolver);
@@ -278,22 +302,25 @@ void UDeformableSolverComponent::Simulate(float DeltaTime)
 	SCOPE_CYCLE_COUNTER(STAT_ChaosDeformable_UDeformableSolverComponent_Simulate);
 	TRACE_CPUPROFILER_EVENT_SCOPE(ChaosDeformable_UDeformableSolverComponent_Simulate);
 
-	if (Solver)
+	if (FleshSolverProxy.Solver)
 	{
 		// @todo(accessor) : Should be coming from the threading class. 
-		FDeformableSolver::FPhysicsThreadAccess PhysicsThreadSolver(Solver.Get(), Chaos::Softs::FPhysicsThreadAccessor());
-		PhysicsThreadSolver.Simulate(DeltaTime);
+		FDeformableSolver::FPhysicsThreadAccess PhysicsThreadSolver = PhysicsThreadAccess();
+		if(!SimulationAsset.DataflowAsset)
+		{
+			PhysicsThreadSolver.Simulate(DeltaTime);
+		}
 	}
 }
 
-void UDeformableSolverComponent::UpdateFromGameThread(float DeltaTime)
+void UDeformableSolverComponent::WriteToSimulation(float DeltaTime)
 {
 	SCOPE_CYCLE_COUNTER(STAT_ChaosDeformable_UDeformableSolverComponent_UpdateFromGameThread);
 	TRACE_CPUPROFILER_EVENT_SCOPE(ChaosDeformable_UDeformableSolverComponent_UpdateFromGameThread);
 
-	if (Solver)
+	if (FleshSolverProxy.Solver)
 	{
-		FDeformableSolver::FGameThreadAccess GameThreadSolver(Solver.Get(), Chaos::Softs::FGameThreadAccessor());
+		FDeformableSolver::FGameThreadAccess GameThreadSolver = GameThreadAccess();
 
 		Chaos::Softs::FDeformableDataMap DataMap;
 		for (TObjectPtr<UDeformablePhysicsComponent>& DeformableComponent : ConnectedObjects.DeformableComponents)
@@ -315,14 +342,14 @@ void UDeformableSolverComponent::UpdateFromGameThread(float DeltaTime)
 	}
 }
 
-void UDeformableSolverComponent::UpdateFromSimulation(float DeltaTime)
+void UDeformableSolverComponent::ReadFromSimulation(float DeltaTime)
 {
 	SCOPE_CYCLE_COUNTER(STAT_ChaosDeformable_UDeformableSolverComponent_UpdateFromSimulation);
 	TRACE_CPUPROFILER_EVENT_SCOPE(ChaosDeformable_UDeformableSolverComponent_UpdateFromSimulation);
 
-	if (Solver)
+	if (FleshSolverProxy.Solver)
 	{
-		FDeformableSolver::FGameThreadAccess GameThreadSolver(Solver.Get(), Chaos::Softs::FGameThreadAccessor());
+		FDeformableSolver::FGameThreadAccess GameThreadSolver = GameThreadAccess();
 		TUniquePtr<FDeformablePackage> Output(nullptr);
 		while (TUniquePtr<FDeformablePackage> SolverOutput = GameThreadSolver.PullOutputPackage())
 		{
@@ -354,3 +381,5 @@ void UDeformableSolverComponent::UpdateFromSimulation(float DeltaTime)
 
 
 
+
+
diff --git a/Engine/Plugins/Experimental/ChaosFlesh/Source/ChaosFleshEngine/Private/ChaosFlesh/ChaosDeformableSolverThreading.cpp b/Engine/Plugins/Experimental/ChaosFlesh/Source/ChaosFleshEngine/Private/ChaosFlesh/ChaosDeformableSolverThreading.cpp
--- a/Engine/Plugins/Experimental/ChaosFlesh/Source/ChaosFleshEngine/Private/ChaosFlesh/ChaosDeformableSolverThreading.cpp
+++ b/Engine/Plugins/Experimental/ChaosFlesh/Source/ChaosFleshEngine/Private/ChaosFlesh/ChaosDeformableSolverThreading.cpp
@@ -45,7 +45,7 @@ void FDeformableEndTickFunction::ExecuteTick(float DeltaTime, enum ELevelTick Ti
 
 	if (DeformableSolverComponent)
 	{
-		DeformableSolverComponent->UpdateFromSimulation(DeltaTime);
+		DeformableSolverComponent->ReadFromSimulation(DeltaTime);
 	}
 }
 
diff --git a/Engine/Plugins/Experimental/ChaosFlesh/Source/ChaosFleshEngine/Public/ChaosFlesh/ChaosDeformableSolverComponent.h b/Engine/Plugins/Experimental/ChaosFlesh/Source/ChaosFleshEngine/Public/ChaosFlesh/ChaosDeformableSolverComponent.h
--- a/Engine/Plugins/Experimental/ChaosFlesh/Source/ChaosFleshEngine/Public/ChaosFlesh/ChaosDeformableSolverComponent.h
+++ b/Engine/Plugins/Experimental/ChaosFlesh/Source/ChaosFleshEngine/Public/ChaosFlesh/ChaosDeformableSolverComponent.h
@@ -1,14 +1,15 @@
 // Copyright Epic Games, Inc. All Rights Reserved.
 #pragma once
 
+#include "ChaosDeformablePhysicsComponent.h"
 #include "Chaos/Deformable/ChaosDeformableSolver.h"
 #include "Chaos/Deformable/ChaosDeformableSolverProxy.h"
 #include "Chaos/Deformable/ChaosDeformableSolverTypes.h"
 #include "ChaosFlesh/ChaosDeformableSolverThreading.h"
 #include "ChaosFlesh/ChaosDeformableTypes.h"
-#include "ChaosFlesh/FleshComponent.h"
 #include "DeformableInterface.h"
 #include "Components/SceneComponent.h"
+#include "Interfaces/DataflowPhysicsSolver.h"
 #include "UObject/ObjectMacros.h"
 
 #include "ChaosDeformableSolverComponent.generated.h"
@@ -249,11 +250,46 @@ struct FSolverMuscleActivationGroup
 	bool bDoMuscleActivation = false;
 };
 
+USTRUCT(BlueprintType)
+struct FDataflowFleshSolverProxy : public FDataflowPhysicsSolverProxy
+{
+	GENERATED_USTRUCT_BODY()
+
+	FDataflowFleshSolverProxy(Chaos::Softs::FDeformableSolverProperties InProp = Chaos::Softs::FDeformableSolverProperties()) :
+		FDataflowPhysicsSolverProxy()
+	{
+		Solver = MakeUnique<Chaos::Softs::FDeformableSolver>(InProp);
+	}
+	virtual ~FDataflowFleshSolverProxy() override = default;
+
+	// Begin FPhysicsSolverInterface overrides
+	virtual void AdvanceSolverDatas(const float DeltaTime) override
+	{
+		Chaos::Softs::FDeformableSolver::FPhysicsThreadAccess PhysicsThreadAccess(Solver.Get(), Chaos::Softs::FPhysicsThreadAccessor());
+		PhysicsThreadAccess.Simulate(DeltaTime);
+	}
+	virtual float GetTimeStep() override
+	{
+		const Chaos::Softs::FDeformableSolver::FPhysicsThreadAccess PhysicsThreadAccess(Solver.Get(), Chaos::Softs::FPhysicsThreadAccessor());
+		return PhysicsThreadAccess.GetProperties().TimeStepSize;
+	}
+	// End FPhysicsSolverInterface overrides
+
+	/** Chaos deformable solver that will be used in the component */
+	TUniquePtr<Chaos::Softs::FDeformableSolver> Solver;
+};
+
+template <>
+struct TStructOpsTypeTraits<FDataflowFleshSolverProxy> : public TStructOpsTypeTraitsBase2<FDataflowFleshSolverProxy>
+{
+	enum { WithCopy = false };
+};
+
 /**
 *	UDeformableSolverComponent
 */
 UCLASS(meta = (BlueprintSpawnableComponent))
-class CHAOSFLESHENGINE_API UDeformableSolverComponent : public USceneComponent, public IDeformableInterface
+class CHAOSFLESHENGINE_API UDeformableSolverComponent : public USceneComponent, public IDeformableInterface, public IDataflowPhysicsSolverInterface
 {
 	GENERATED_UCLASS_BODY()
 
@@ -267,25 +303,44 @@ public:
 	~UDeformableSolverComponent();
 	void UpdateTickGroup();
 
-	/* Solver API */
+	// Begin IDataflowPhysicsSolverInterface overrides
+	virtual FString GetSimulationName() const override {return GetName();};
+	virtual FDataflowSimulationAsset& GetSimulationAsset() override {return SimulationAsset;};
+	virtual const FDataflowSimulationAsset& GetSimulationAsset() const override {return SimulationAsset;};
+	virtual FDataflowSimulationProxy* GetSimulationProxy() override {return &FleshSolverProxy;}
+	virtual const FDataflowSimulationProxy* GetSimulationProxy() const  override {return &FleshSolverProxy;}
+	virtual void BuildSimulationProxy() override;
+	virtual void ResetSimulationProxy() override;
+	virtual void WriteToSimulation(const float DeltaTime) override;
+	virtual void ReadFromSimulation(const float DeltaTime) override;
+	// End IDataflowPhysicsSolverInterface overrides
+
+	// Begin UActorComponent overrides
+	virtual bool ShouldCreatePhysicsState() const override {return true;}
+	virtual void BeginPlay() override;
+	virtual void EndPlay(const EEndPlayReason::Type EndPlayReason) override;
+	virtual void TickComponent(float DeltaTime, enum ELevelTick TickType, FActorComponentTickFunction* ThisTickFunction) override;
+	// End UActorComponent overrides
+
+	/* Game thread access to the solver proxy */
 	FDeformableSolver::FGameThreadAccess GameThreadAccess();
 
+	/* Physics thread access to the solver proxy */
+	FDeformableSolver::FPhysicsThreadAccess PhysicsThreadAccess();
+
 	bool IsSimulating(UDeformablePhysicsComponent*) const;
 	bool IsSimulatable() const;
-	void Reset();
 	void AddDeformableProxy(UDeformablePhysicsComponent* InComponent);
 	void RemoveDeformableProxy(UDeformablePhysicsComponent* InComponent);
 	void Simulate(float DeltaTime);
-	void UpdateFromGameThread(float DeltaTime);
-	void UpdateFromSimulation(float DeltaTime);
 	void SetSimulationTicking(const bool InSimulationTicking) {bSimulationTicking = InSimulationTicking;}
 
-	/* Component Thread Management */
-	virtual void BeginPlay() override;
-	virtual void EndPlay(const EEndPlayReason::Type EndPlayReason) override;
-	virtual void TickComponent(float DeltaTime, enum ELevelTick TickType, FActorComponentTickFunction* ThisTickFunction) override;
-	//bool ShouldWaitForDeformableInTickFunction() const;
+	/* Callback to trigger the deformable update after the simulation */
 	void UpdateDeformableEndTickState(bool bRegister);
+	
+	/* Solver dataflow asset used to advance in time */
+	UPROPERTY(EditAnywhere, BlueprintReadWrite, Category = "Physics")
+	FDataflowSimulationAsset SimulationAsset;
 
 	/* Properties : Do NOT place ungrouped properties in this class */
 	UPROPERTY(EditAnywhere, BlueprintReadOnly, Category = "Physics", meta = (EditCondition = "false"))
@@ -303,8 +358,6 @@ public:
 	UPROPERTY(EditAnywhere, BlueprintReadOnly, Category = "Physics")
 	FSolverConstraintsGroup SolverConstraints;
 
-	
-
 	UPROPERTY(EditAnywhere, BlueprintReadOnly, Category = "Physics")
 	FSolverForcesGroup SolverForces;
 
@@ -315,7 +368,7 @@ public:
 	FSolverMuscleActivationGroup SolverMuscleActivation;
 
 	// Simulation Variables
-	TUniquePtr<FDeformableSolver> Solver;
+	FDataflowFleshSolverProxy FleshSolverProxy;
 
 #if WITH_EDITOR
 	virtual void CustomizeDetails(IDetailLayoutBuilder& DetailBuilder) override;
diff --git a/Engine/Plugins/Experimental/ChaosFlesh/Source/ChaosFleshEngine/Public/ChaosFlesh/ChaosDeformableTetrahedralComponent.h b/Engine/Plugins/Experimental/ChaosFlesh/Source/ChaosFleshEngine/Public/ChaosFlesh/ChaosDeformableTetrahedralComponent.h
--- a/Engine/Plugins/Experimental/ChaosFlesh/Source/ChaosFleshEngine/Public/ChaosFlesh/ChaosDeformableTetrahedralComponent.h
+++ b/Engine/Plugins/Experimental/ChaosFlesh/Source/ChaosFleshEngine/Public/ChaosFlesh/ChaosDeformableTetrahedralComponent.h
@@ -114,7 +114,10 @@ public:
 	virtual void UpdateFromSimulation(const FDataMapValue* SimualtionBuffer) override;
 
 	/** RestCollection */
+	UFUNCTION(BlueprintCallable, Category = "Physics")
 	void SetRestCollection(const UFleshAsset * InRestCollection);
+	
+	UFUNCTION(BlueprintCallable, Category = "Physics")
 	const UFleshAsset* GetRestCollection() const { return RestCollection; }
 
 	/** DynamicCollection */
diff --git a/Engine/Plugins/Experimental/ChaosFlesh/Source/ChaosFleshEngine/Public/ChaosFlesh/FleshAsset.h b/Engine/Plugins/Experimental/ChaosFlesh/Source/ChaosFleshEngine/Public/ChaosFlesh/FleshAsset.h
--- a/Engine/Plugins/Experimental/ChaosFlesh/Source/ChaosFleshEngine/Public/ChaosFlesh/FleshAsset.h
+++ b/Engine/Plugins/Experimental/ChaosFlesh/Source/ChaosFleshEngine/Public/ChaosFlesh/FleshAsset.h
@@ -50,7 +50,7 @@ private:
 * UObject wrapper for the FFleshAsset
 *
 */
-UCLASS(customconstructor)
+UCLASS(BlueprintType, customconstructor)
 class CHAOSFLESHENGINE_API UFleshAsset : public UObject, public IDataflowContentOwner
 {
 	GENERATED_UCLASS_BODY()
@@ -121,10 +121,10 @@ public:
 	//
 	// SkeletalMesh
 	//
-	UPROPERTY(EditAnywhere, Category = "Animation")
+	UPROPERTY(EditAnywhere, BlueprintReadWrite, Category = "Animation")
 	TObjectPtr<USkeletalMesh> SkeletalMesh;
 
-	UPROPERTY(EditAnywhere, Category = "Animation")
+	UPROPERTY(EditAnywhere, BlueprintReadWrite, Category = "Animation")
 	TObjectPtr<USkeleton> Skeleton;
 
 	/**
diff --git a/Engine/Plugins/Experimental/ChaosFlesh/Source/ChaosFleshNodes/ChaosFleshNodes.Build.cs b/Engine/Plugins/Experimental/ChaosFlesh/Source/ChaosFleshNodes/ChaosFleshNodes.Build.cs
--- a/Engine/Plugins/Experimental/ChaosFlesh/Source/ChaosFleshNodes/ChaosFleshNodes.Build.cs
+++ b/Engine/Plugins/Experimental/ChaosFlesh/Source/ChaosFleshNodes/ChaosFleshNodes.Build.cs
@@ -16,6 +16,7 @@ namespace UnrealBuildTool.Rules
 					"Chaos",
 					"DataflowCore",
 					"DataflowEngine",
+					"DataflowSimulation",
 					"ChaosFlesh",
 					"ChaosFleshEngine",
 					"Engine",
diff --git a/Engine/Source/Runtime/Experimental/Chaos/Private/Chaos/Deformable/ChaosDeformableSolver.cpp b/Engine/Source/Runtime/Experimental/Chaos/Private/Chaos/Deformable/ChaosDeformableSolver.cpp
--- a/Engine/Source/Runtime/Experimental/Chaos/Private/Chaos/Deformable/ChaosDeformableSolver.cpp
+++ b/Engine/Source/Runtime/Experimental/Chaos/Private/Chaos/Deformable/ChaosDeformableSolver.cpp
@@ -221,72 +221,75 @@ namespace Chaos::Softs
 			const FThreadingProxy::FKey& Owner = ProxyIt.Key();
 			if (const FFleshThreadingProxy* Proxy = ProxyIt.Value()->As<FFleshThreadingProxy>())
 			{
-				if (FFleshThreadingProxy::FFleshInputBuffer* FleshInputBuffer =
-					this->CurrentInputPackage->ObjectMap.Contains(Owner) ?
-					this->CurrentInputPackage->ObjectMap[Owner]->As<FFleshThreadingProxy::FFleshInputBuffer>() :
-					nullptr)
+				if(this->CurrentInputPackage)
 				{
-					GeometryCollection::Facades::FConstraintOverrideTargetFacade CnstrTargets(FleshInputBuffer->SimulationCollection);
-					if (CnstrTargets.IsValid() && CnstrTargets.Num())
+					if (FFleshThreadingProxy::FFleshInputBuffer* FleshInputBuffer =
+						this->CurrentInputPackage->ObjectMap.Contains(Owner) ?
+						this->CurrentInputPackage->ObjectMap[Owner]->As<FFleshThreadingProxy::FFleshInputBuffer>() :
+						nullptr)
 					{
-						const FIntVector2& Range = Proxy->GetSolverParticleRange();
-						const FSolverReal CurrentRatio = FSolverReal(this->Iteration) / FSolverReal(this->Property.NumSolverSubSteps);
-						const FTransform WorldToSim = Proxy->GetCurrentPointsTransform();
-
-						if (this->Iteration == 1)
+						GeometryCollection::Facades::FConstraintOverrideTargetFacade CnstrTargets(FleshInputBuffer->SimulationCollection);
+						if (CnstrTargets.IsValid() && CnstrTargets.Num())
 						{
-							TransientConstraintBuffer.Reserve(TransientConstraintBuffer.Num() + CnstrTargets.Num());
-							for (int32 i = 0; i < CnstrTargets.Num(); i++)
+							const FIntVector2& Range = Proxy->GetSolverParticleRange();
+							const FSolverReal CurrentRatio = FSolverReal(this->Iteration) / FSolverReal(this->Property.NumSolverSubSteps);
+							const FTransform WorldToSim = Proxy->GetCurrentPointsTransform();
+
+							if (this->Iteration == 1)
 							{
-								int32 LocalIndex = CnstrTargets.GetIndex(i);
-								int32 ParticleIndex = Range[0] + LocalIndex;
-								// Set particle kinematic state to kinematic, saving prior state.
-								TransientConstraintBuffer.Add(
-									TPair<int32, TTuple<float, float, FVector3f>>(
-										ParticleIndex,
-										TTuple<float, float, FVector3f>(
-											Evolution->Particles().InvM(ParticleIndex),
-											Evolution->Particles().PAndInvM(ParticleIndex).InvM,
-											Evolution->Particles().GetX(ParticleIndex))));
-
-								Evolution->Particles().InvM(ParticleIndex) = 0.f;
-								Evolution->Particles().PAndInvM(ParticleIndex).InvM = 0.f;
+								TransientConstraintBuffer.Reserve(TransientConstraintBuffer.Num() + CnstrTargets.Num());
+								for (int32 i = 0; i < CnstrTargets.Num(); i++)
+								{
+									int32 LocalIndex = CnstrTargets.GetIndex(i);
+									int32 ParticleIndex = Range[0] + LocalIndex;
+									// Set particle kinematic state to kinematic, saving prior state.
+									TransientConstraintBuffer.Add(
+										TPair<int32, TTuple<float, float, FVector3f>>(
+											ParticleIndex,
+											TTuple<float, float, FVector3f>(
+												Evolution->Particles().InvM(ParticleIndex),
+												Evolution->Particles().PAndInvM(ParticleIndex).InvM,
+												Evolution->Particles().GetX(ParticleIndex))));
+
+									Evolution->Particles().InvM(ParticleIndex) = 0.f;
+									Evolution->Particles().PAndInvM(ParticleIndex).InvM = 0.f;
+								}
 							}
-						}
 
-						auto ToDouble = [](FVector3f V) { return FVector(V[0], V[1], V[2]); };
-						auto ToSingle = [](FVector V) { return FVector3f(static_cast<float>(V[0]), static_cast<float>(V[1]), static_cast<float>(V[2])); };
+							auto ToDouble = [](FVector3f V) { return FVector(V[0], V[1], V[2]); };
+							auto ToSingle = [](FVector V) { return FVector3f(static_cast<float>(V[0]), static_cast<float>(V[1]), static_cast<float>(V[2])); };
 
-						for (int32 i = 0; i < CnstrTargets.Num(); i++)
-						{
-							int32 LocalIndex = CnstrTargets.GetIndex(i);
-							int32 ParticleIndex = Range[0] + LocalIndex;
-
-							const FVector3f& WorldSpaceTarget = CnstrTargets.GetPosition(i);
-							FVector3f SimSpaceTarget = ToSingle(WorldToSim.TransformPosition(ToDouble(WorldSpaceTarget)));
-							const FVector3f& SimSpaceSource = TransientConstraintBuffer[ParticleIndex].Get<2>();
-
-							// Lerp from previous particle position to the target over the solver iterations.
-							Evolution->Particles().SetX(ParticleIndex,
-								SimSpaceTarget * CurrentRatio +
-								SimSpaceSource * (static_cast<FSolverReal>(1.) - CurrentRatio));
-							Evolution->Particles().PAndInvM(ParticleIndex).P =
-								Evolution->Particles().GetX(ParticleIndex);
-						}
-#if WITH_EDITOR
-						if (GDeformableDebugParams.IsDebugDrawingEnabled() && GDeformableDebugParams.bDoDrawTransientKinematicParticles)
-						{
-							auto DoubleVert = [](FVector3f V) { return FVector3d(V.X, V.Y, V.Z); };
 							for (int32 i = 0; i < CnstrTargets.Num(); i++)
 							{
 								int32 LocalIndex = CnstrTargets.GetIndex(i);
 								int32 ParticleIndex = Range[0] + LocalIndex;
-								Chaos::FDebugDrawQueue::GetInstance().DrawDebugPoint(ToDouble(Evolution->Particles().GetX(ParticleIndex)), FColor::Orange, false, -1.0f, 0, 5);
+
+								const FVector3f& WorldSpaceTarget = CnstrTargets.GetPosition(i);
+								FVector3f SimSpaceTarget = ToSingle(WorldToSim.TransformPosition(ToDouble(WorldSpaceTarget)));
+								const FVector3f& SimSpaceSource = TransientConstraintBuffer[ParticleIndex].Get<2>();
+
+								// Lerp from previous particle position to the target over the solver iterations.
+								Evolution->Particles().SetX(ParticleIndex,
+									SimSpaceTarget * CurrentRatio +
+									SimSpaceSource * (static_cast<FSolverReal>(1.) - CurrentRatio));
+								Evolution->Particles().PAndInvM(ParticleIndex).P =
+									Evolution->Particles().GetX(ParticleIndex);
+							}
+#if WITH_EDITOR
+							if (GDeformableDebugParams.IsDebugDrawingEnabled() && GDeformableDebugParams.bDoDrawTransientKinematicParticles)
+							{
+								auto DoubleVert = [](FVector3f V) { return FVector3d(V.X, V.Y, V.Z); };
+								for (int32 i = 0; i < CnstrTargets.Num(); i++)
+								{
+									int32 LocalIndex = CnstrTargets.GetIndex(i);
+									int32 ParticleIndex = Range[0] + LocalIndex;
+									Chaos::FDebugDrawQueue::GetInstance().DrawDebugPoint(ToDouble(Evolution->Particles().GetX(ParticleIndex)), FColor::Orange, false, -1.0f, 0, 5);
+								}
 							}
-						}
 #endif
-					} // if has constraint overrides
-				} // if flesh input buffer
+						} // if has constraint overrides
+					} // if flesh input buffer
+				}
 			}
 		} // for all proxies
 	}
@@ -314,7 +317,7 @@ namespace Chaos::Softs
 			{
 				if (FFleshThreadingProxy* Proxy = Proxies[Owner]->As<FFleshThreadingProxy>())
 				{
-					if (this->CurrentInputPackage->ObjectMap.Contains(Owner))
+					if (this->CurrentInputPackage && this->CurrentInputPackage->ObjectMap.Contains(Owner))
 					{
 						FFleshThreadingProxy::FFleshInputBuffer* FleshInputBuffer =
 							this->CurrentInputPackage->ObjectMap[Owner]->As<FFleshThreadingProxy::FFleshInputBuffer>();
@@ -2135,23 +2138,26 @@ namespace Chaos::Softs
 
 				TManagedArray<FVector3f>& Position = Proxy->GetDynamicCollection().ModifyAttribute<FVector3f>("Vertex", FGeometryCollection::VerticesGroup);
 
-				// The final transform gets us from whatever the simulation space is,
-				// to component space.
-				const FTransform FinalXf = Proxy->GetFinalTransform();
-				if (!FinalXf.Equals(FTransform::Identity))
+				if((Position.Num() + Range[0]) <= (int32)(Evolution->Particles().Size()))
 				{
-					for (int32 vdx = 0; vdx < Position.Num(); vdx++)
+					// The final transform gets us from whatever the simulation space is,
+					// to component space.
+					const FTransform FinalXf = Proxy->GetFinalTransform();
+					if (!FinalXf.Equals(FTransform::Identity))
 					{
-						const Chaos::FVec3f& Pos = Evolution->Particles().GetX(vdx + Range[0]);
-						FVector PosD = UEVertd(Pos);
-						Position[vdx] = UEVertf(FinalXf.TransformPosition(PosD));
+						for (int32 vdx = 0; vdx < Position.Num(); vdx++)
+						{
+							const Chaos::FVec3f& Pos = Evolution->Particles().GetX(vdx + Range[0]);
+							FVector PosD = UEVertd(Pos);
+							Position[vdx] = UEVertf(FinalXf.TransformPosition(PosD));
+						}
 					}
-				}
-				else
-				{
-					for (int32 vdx = 0; vdx < Position.Num(); vdx++)
+					else
 					{
-						Position[vdx] = UEVertf(UEVertd(Evolution->Particles().GetX(vdx + Range[0])));
+						for (int32 vdx = 0; vdx < Position.Num(); vdx++)
+						{
+							Position[vdx] = UEVertf(UEVertd(Evolution->Particles().GetX(vdx + Range[0])));
+						}
 					}
 				}
 			}
