[NetPhysics] CVar for Z axis offset on Render Interpolation DrawDebug

diff --git a/Engine/Source/Runtime/Experimental/Chaos/Private/Framework/PhysicsProxyBase.cpp b/Engine/Source/Runtime/Experimental/Chaos/Private/Framework/PhysicsProxyBase.cpp
--- a/Engine/Source/Runtime/Experimental/Chaos/Private/Framework/PhysicsProxyBase.cpp
+++ b/Engine/Source/Runtime/Experimental/Chaos/Private/Framework/PhysicsProxyBase.cpp
@@ -17,11 +17,14 @@ namespace RenderInterpolationCVars
 	float RenderInterpErrorDirectionalDecayMultiplier = 0.0f;
 	static FAutoConsoleVariableRef CVarRenderInterpErrorDirectionalDecayMultiplier(TEXT("p.RenderInterp.DirectionalDecayMultiplier"), RenderInterpErrorDirectionalDecayMultiplier, TEXT("Decay error offset in the direction that the physics object is moving, value is multiplier of projected offset direction, 0.25 means a 25% decay of the magnitude in the direction of physics travel. Deactivate by setting to 0."));
 
-	bool bRenderInterpDebugDraw = false;
-	static FAutoConsoleVariableRef CVarRenderInterpDebugDraw(TEXT("p.RenderInterp.DebugDraw"), bRenderInterpDebugDraw, TEXT("Draw debug lines for physics render interpolation, also needs p.Chaos.DebugDraw.Enabled set"));
-
 	bool bRenderInterpErrorVelocityCorrection = false;
 	static FAutoConsoleVariableRef CVarRenderInterpErrorVelocityCorrection(TEXT("p.RenderInterp.ErrorVelocityCorrection"), bRenderInterpErrorVelocityCorrection, TEXT("EXPERIMENTAL - Take incoming velocity into consideration when performing render interpolation, the correction will be more organic but might result in clipping and it's heavier for memory and CPU."));
+
+	bool bRenderInterpDebugDraw = false;
+	static FAutoConsoleVariableRef CVarRenderInterpDebugDraw(TEXT("p.RenderInterp.DebugDraw"), bRenderInterpDebugDraw, TEXT("Draw debug lines for physics render interpolation, also needs p.Chaos.DebugDraw.Enabled set"));
+	
+	float RenderInterpDebugDrawZOffset = 0.0f;
+	static FAutoConsoleVariableRef CVarRenderInterpDebugDrawZOffset(TEXT("p.RenderInterp.DebugDrawZOffset"), RenderInterpDebugDrawZOffset, TEXT("Add Z axis offset to DebugDraw calls for Render Interpolation."));
 }
 
 IPhysicsProxyBase::~IPhysicsProxyBase()
diff --git a/Engine/Source/Runtime/Experimental/Chaos/Private/PhysicsProxy/ClusterUnionPhysicsProxy.cpp b/Engine/Source/Runtime/Experimental/Chaos/Private/PhysicsProxy/ClusterUnionPhysicsProxy.cpp
--- a/Engine/Source/Runtime/Experimental/Chaos/Private/PhysicsProxy/ClusterUnionPhysicsProxy.cpp
+++ b/Engine/Source/Runtime/Experimental/Chaos/Private/PhysicsProxy/ClusterUnionPhysicsProxy.cpp
@@ -610,19 +610,20 @@ namespace Chaos
 #if CHAOS_DEBUG_DRAW
 			if (RenderInterpolationCVars::bRenderInterpDebugDraw)
 			{
-				Chaos::FDebugDrawQueue::GetInstance().DrawDebugBox(NextPullData->X, FVector(2, 1, 1), NextPullData->R, FColor::Yellow, false, 5.f, 0, 0.5f);
-				Chaos::FDebugDrawQueue::GetInstance().DrawDebugDirectionalArrow(PullData.X, NextPullData->X, 0.5f, FColor::Yellow, false, 5.0f, 0, 0.5f);
-				Chaos::FDebugDrawQueue::GetInstance().DrawDebugBox(Particle_External->GetX(), FVector(2, 1, 1), Particle_External->R(), DirectionalDecayPerformed ? FColor::Cyan : FColor::Green, false, 5.f, 0, 0.5f);
+				const FVector ZOffset = FVector(0, 0, RenderInterpolationCVars::RenderInterpDebugDrawZOffset);
+				Chaos::FDebugDrawQueue::GetInstance().DrawDebugBox(ZOffset + NextPullData->X, FVector(2, 1, 1), NextPullData->R, FColor::Yellow, false, 5.f, 0, 0.5f);
+				Chaos::FDebugDrawQueue::GetInstance().DrawDebugDirectionalArrow(ZOffset + PullData.X, ZOffset + NextPullData->X, 0.5f, FColor::Yellow, false, 5.0f, 0, 0.5f);
+				Chaos::FDebugDrawQueue::GetInstance().DrawDebugBox(ZOffset + Particle_External->GetX(), FVector(2, 1, 1), Particle_External->R(), DirectionalDecayPerformed ? FColor::Cyan : FColor::Green, false, 5.f, 0, 0.5f);
 
 				if (bIsReplicationErrorSmoothing)
 				{
 					if (Error)
 					{
-						Chaos::FDebugDrawQueue::GetInstance().DrawDebugBox(PullData.X, FVector(4, 2, 2), PullData.R, FColor::Red, false, 5.f, 0, 0.5f);
-						Chaos::FDebugDrawQueue::GetInstance().DrawDebugDirectionalArrow(PullData.X, (PullData.X + InterpData->GetErrorX(0)), 1, FColor::Red, false, 5.0f, 0, 0.5f);
+						Chaos::FDebugDrawQueue::GetInstance().DrawDebugBox(ZOffset + PullData.X, FVector(4, 2, 2), PullData.R, FColor::Red, false, 5.f, 0, 0.5f);
+						Chaos::FDebugDrawQueue::GetInstance().DrawDebugDirectionalArrow(ZOffset + PullData.X, ZOffset + (PullData.X + InterpData->GetErrorX(0)), 1, FColor::Red, false, 5.0f, 0, 0.5f);
 					}
 
-					Chaos::FDebugDrawQueue::GetInstance().DrawDebugDirectionalArrow((Particle_External->GetX() - InterpData->GetErrorX(*Alpha)), Particle_External->GetX(), 1, FColor::Blue, false, 5.0f, 0, 0.5f);
+					Chaos::FDebugDrawQueue::GetInstance().DrawDebugDirectionalArrow(ZOffset + (Particle_External->GetX() - InterpData->GetErrorX(*Alpha)), ZOffset + Particle_External->GetX(), 1, FColor::Blue, false, 5.0f, 0, 0.5f);
 				}
 			}
 #endif // CHAOS_DEBUG_DRAW
diff --git a/Engine/Source/Runtime/Experimental/Chaos/Private/PhysicsProxy/SingleParticlePhysicsProxy.cpp b/Engine/Source/Runtime/Experimental/Chaos/Private/PhysicsProxy/SingleParticlePhysicsProxy.cpp
--- a/Engine/Source/Runtime/Experimental/Chaos/Private/PhysicsProxy/SingleParticlePhysicsProxy.cpp
+++ b/Engine/Source/Runtime/Experimental/Chaos/Private/PhysicsProxy/SingleParticlePhysicsProxy.cpp
@@ -352,9 +352,10 @@ bool FSingleParticlePhysicsProxy::PullFromPhysicsState(const Chaos::FDirtyRigidP
 #if CHAOS_DEBUG_DRAW
 							if (RenderInterpolationCVars::bRenderInterpDebugDraw)
 							{
-								Chaos::FDebugDrawQueue::GetInstance().DrawDebugDirectionalArrow((Target - InterpData->GetErrorX(*Alpha)), Target, 1, FColor::Blue, false, 5.0f, 0, 0.5f);
-								Chaos::FDebugDrawQueue::GetInstance().DrawDebugDirectionalArrow(InterpData->GetErrorVelocitySmoothingX(*Alpha), Target, 1, FColor::Blue, false, 5.0f, 0, 0.5f);
-								Chaos::FDebugDrawQueue::GetInstance().DrawDebugBox(Target, FVector(2, 1, 1), Rigid->R(), FColor::Cyan, false, 5.f, 0, 0.25f);
+								const FVector ZOffset = FVector(0,0, RenderInterpolationCVars::RenderInterpDebugDrawZOffset);
+								Chaos::FDebugDrawQueue::GetInstance().DrawDebugDirectionalArrow(ZOffset + (Target - InterpData->GetErrorX(*Alpha)), ZOffset + Target, 1, FColor::Blue, false, 5.0f, 0, 0.5f);
+								Chaos::FDebugDrawQueue::GetInstance().DrawDebugDirectionalArrow(ZOffset + InterpData->GetErrorVelocitySmoothingX(*Alpha), ZOffset + Target, 1, FColor::Blue, false, 5.0f, 0, 0.5f);
+								Chaos::FDebugDrawQueue::GetInstance().DrawDebugBox(ZOffset + Target, FVector(2, 1, 1), Rigid->R(), FColor::Cyan, false, 5.f, 0, 0.25f);
 							}
 #endif // CHAOS_DEBUG_DRAW
 
@@ -379,25 +380,26 @@ bool FSingleParticlePhysicsProxy::PullFromPhysicsState(const Chaos::FDirtyRigidP
 #if CHAOS_DEBUG_DRAW
 				if (RenderInterpolationCVars::bRenderInterpDebugDraw)
 				{
-					Chaos::FDebugDrawQueue::GetInstance().DrawDebugBox(NextPullData->X, FVector(2, 1, 1), NextPullData->R, FColor::Yellow, false, 5.f, 0, 0.5f);
-					Chaos::FDebugDrawQueue::GetInstance().DrawDebugDirectionalArrow(PullData.X, NextPullData->X, 0.5f, FColor::Yellow, false, 5.0f, 0, 0.5f);
-					Chaos::FDebugDrawQueue::GetInstance().DrawDebugBox(Rigid->X(), FVector(2, 1, 1), Rigid->R(), DirectionalDecayPerformed ? FColor::Cyan : FColor::Green, false, 5.f, 0, 0.5f);
+					const FVector ZOffset = FVector(0, 0, RenderInterpolationCVars::RenderInterpDebugDrawZOffset);
+					Chaos::FDebugDrawQueue::GetInstance().DrawDebugBox(ZOffset + NextPullData->X, FVector(2, 1, 1), NextPullData->R, FColor::Yellow, false, 5.f, 0, 0.5f);
+					Chaos::FDebugDrawQueue::GetInstance().DrawDebugDirectionalArrow(ZOffset + PullData.X, ZOffset + NextPullData->X, 0.5f, FColor::Yellow, false, 5.0f, 0, 0.5f);
+					Chaos::FDebugDrawQueue::GetInstance().DrawDebugBox(ZOffset + Rigid->X(), FVector(2, 1, 1), Rigid->R(), DirectionalDecayPerformed ? FColor::Cyan : FColor::Green, false, 5.f, 0, 0.5f);
 
 					if (bIsReplicationErrorSmoothing)
 					{
 						if (Error)
 						{
-							Chaos::FDebugDrawQueue::GetInstance().DrawDebugBox(PullData.X, FVector(4, 2, 2), PullData.R, FColor::Red, false, 5.f, 0, 0.5f);
-							Chaos::FDebugDrawQueue::GetInstance().DrawDebugDirectionalArrow(PullData.X, (PullData.X + InterpData->GetErrorX(0)), 1, FColor::Red, false, 5.0f, 0, 0.5f);
+							Chaos::FDebugDrawQueue::GetInstance().DrawDebugBox(ZOffset + PullData.X, FVector(4, 2, 2), PullData.R, FColor::Red, false, 5.f, 0, 0.5f);
+							Chaos::FDebugDrawQueue::GetInstance().DrawDebugDirectionalArrow(ZOffset + PullData.X, ZOffset + (PullData.X + InterpData->GetErrorX(0)), 1, FColor::Red, false, 5.0f, 0, 0.5f);
 						}
 
 						if (InterpData->IsErrorVelocitySmoothing())
 						{
-							Chaos::FDebugDrawQueue::GetInstance().DrawDebugBox(InterpData->GetErrorVelocitySmoothingX(*Alpha), FVector(2, 2, 2), Rigid->R(), FColor::Purple, false, 5.f, 0, 0.5f);
+							Chaos::FDebugDrawQueue::GetInstance().DrawDebugBox(ZOffset + InterpData->GetErrorVelocitySmoothingX(*Alpha), FVector(2, 2, 2), Rigid->R(), FColor::Purple, false, 5.f, 0, 0.5f);
 						}
 						else
 						{
-							Chaos::FDebugDrawQueue::GetInstance().DrawDebugDirectionalArrow((Rigid->X() - InterpData->GetErrorX(*Alpha)), Rigid->X(), 1, FColor::Blue, false, 5.0f, 0, 0.5f);
+							Chaos::FDebugDrawQueue::GetInstance().DrawDebugDirectionalArrow(ZOffset + (Rigid->X() - InterpData->GetErrorX(*Alpha)), ZOffset + Rigid->X(), 1, FColor::Blue, false, 5.0f, 0, 0.5f);
 						}
 					}
 				}
diff --git a/Engine/Source/Runtime/Experimental/Chaos/Public/Chaos/Framework/PhysicsProxyBase.h b/Engine/Source/Runtime/Experimental/Chaos/Public/Chaos/Framework/PhysicsProxyBase.h
--- a/Engine/Source/Runtime/Experimental/Chaos/Public/Chaos/Framework/PhysicsProxyBase.h
+++ b/Engine/Source/Runtime/Experimental/Chaos/Public/Chaos/Framework/PhysicsProxyBase.h
@@ -33,8 +33,9 @@ namespace RenderInterpolationCVars
 	extern float RenderInterpMaximumErrorCorrectionBeforeSnapping;
 	extern float RenderInterpErrorVelocitySmoothingDuration;
 	extern float RenderInterpErrorDirectionalDecayMultiplier;
-	extern bool bRenderInterpDebugDraw;
 	extern bool bRenderInterpErrorVelocityCorrection;
+	extern bool bRenderInterpDebugDraw;
+	extern float RenderInterpDebugDrawZOffset;
 }
 
 enum class EProxyInterpolationType : uint32
