Iris
- FReplicationWriter::UpdateDirtyGlobalLifetimeConditionals now marks subobject owner dirty - now fixed to not cause bad state transitons
- Added test for previouse bad case.
- Changed checks to ensures to avoid taking down server

diff --git a/Engine/Plugins/Runtime/ReplicationSystemTestPlugin/Source/Private/Tests/ReplicationSystem/Conditionals/TestConditionals.cpp b/Engine/Plugins/Runtime/ReplicationSystemTestPlugin/Source/Private/Tests/ReplicationSystem/Conditionals/TestConditionals.cpp
--- a/Engine/Plugins/Runtime/ReplicationSystemTestPlugin/Source/Private/Tests/ReplicationSystem/Conditionals/TestConditionals.cpp
+++ b/Engine/Plugins/Runtime/ReplicationSystemTestPlugin/Source/Private/Tests/ReplicationSystem/Conditionals/TestConditionals.cpp
@@ -817,6 +817,70 @@ UE_NET_TEST_FIXTURE(FTestConditionalsFixture, WithSkipOwnerSubObjectReplicatedTo
 	}
 }
 
+// Test case where we ended up trying to send subobjects waiting for creation confirmation causing server check/crash
+UE_NET_TEST_FIXTURE(FTestConditionalsFixture, WithSkipOwnerSubObjectOwnerCanBeChangedWithDataInFlight)
+{
+	// Add client
+	FReplicationSystemTestClient* Client = CreateClient();
+
+	// Spawn object + subobject on server
+	UTestReplicatedIrisObject* ServerObject = Server->CreateObject(0, 0);
+	UTestReplicatedIrisObject* ServerSubObject = Server->CreateSubObject(ServerObject->NetRefHandle, 0, 0);
+
+	Server->ReplicationBridge->SetSubObjectNetCondition(ServerSubObject->NetRefHandle, ELifetimeCondition::COND_SkipOwner);
+
+	// Set owner
+	Server->ReplicationSystem->SetOwningNetConnection(ServerObject->NetRefHandle, Client->ConnectionIdOnServer);
+
+	// Set some values
+	ServerSubObject->IntA = 13; 
+
+	// Send and deliver packet
+	Server->PreSendUpdate();
+	Server->SendAndDeliverTo(Client, DeliverPacket);
+	Server->PostSendUpdate();
+	
+	// Object should have been created on the client
+	const UTestReplicatedIrisObject* ClientObject = Cast<UTestReplicatedIrisObject>(Client->GetReplicationBridge()->GetReplicatedObject(ServerObject->NetRefHandle));
+	UE_NET_ASSERT_NE(ClientObject, nullptr);
+
+	// SubObject should not have been created on client
+	{
+		const UTestReplicatedIrisObject* ClientSubObject = Cast<UTestReplicatedIrisObject>(Client->GetReplicationBridge()->GetReplicatedObject(ServerSubObject->NetRefHandle));
+		UE_NET_ASSERT_EQ(ClientSubObject, nullptr);
+	}
+
+	// Change owner
+	Server->ReplicationSystem->SetOwningNetConnection(ServerObject->NetRefHandle, 0);
+
+	// Send and delay packet
+	Server->PreSendUpdate();
+	Server->SendTo(Client);
+	Server->PostSendUpdate();
+
+	// Change owner again, just to dirty things and cause a conditional change with subobject pending create in-flight
+	// before a fix in ReplicationWriter this would cause an invalid state assert/check
+	Server->ReplicationSystem->SetOwningNetConnection(ServerObject->NetRefHandle, Client->ConnectionIdOnServer);
+	Server->ReplicationSystem->SetOwningNetConnection(ServerObject->NetRefHandle, 0);
+
+	// Send and delay packet
+	Server->PreSendUpdate();
+	Server->SendTo(Client);
+	Server->PostSendUpdate();
+
+	Server->DeliverTo(Client, true);
+	Server->DeliverTo(Client, true);
+
+	// SubObject should now have been created on client
+	{
+		const UTestReplicatedIrisObject* ClientSubObject = Cast<UTestReplicatedIrisObject>(Client->GetReplicationBridge()->GetReplicatedObject(ServerSubObject->NetRefHandle));
+		UE_NET_ASSERT_NE(ClientSubObject, nullptr);
+
+		UE_NET_ASSERT_EQ(ServerSubObject->IntA, ClientSubObject->IntA);
+	}
+}
+
+
 UE_NET_TEST_FIXTURE(FTestConditionalsFixture, HierarchicalSubObjectIsNotReplicatedToOwnerUnlessParentIs)
 {
 	// Add client
diff --git a/Engine/Source/Runtime/Experimental/Iris/Core/Private/Iris/ReplicationSystem/ReplicationWriter.cpp b/Engine/Source/Runtime/Experimental/Iris/Core/Private/Iris/ReplicationSystem/ReplicationWriter.cpp
--- a/Engine/Source/Runtime/Experimental/Iris/Core/Private/Iris/ReplicationSystem/ReplicationWriter.cpp
+++ b/Engine/Source/Runtime/Experimental/Iris/Core/Private/Iris/ReplicationSystem/ReplicationWriter.cpp
@@ -100,67 +100,67 @@ void FReplicationWriter::FReplicationInfo::SetState(EReplicatedObjectState NewSt
 	{
 		case EReplicatedObjectState::PendingCreate:
 		{
-			checkf(CurrentState == EReplicatedObjectState::Invalid || CurrentState == EReplicatedObjectState::WaitOnCreateConfirmation, TEXT("Trying to set state %s when state is %s. IsDestructionInfo: %u IsSubObject: %u"), LexToString(NewState), LexToString(CurrentState), IsDestructionInfo, IsSubObject);
+			ensureMsgf(CurrentState == EReplicatedObjectState::Invalid || CurrentState == EReplicatedObjectState::WaitOnCreateConfirmation, TEXT("Trying to set state %s when state is %s. IsDestructionInfo: %u IsSubObject: %u"), LexToString(NewState), LexToString(CurrentState), IsDestructionInfo, IsSubObject);
 		}
 		break;
 		case EReplicatedObjectState::WaitOnCreateConfirmation:
 		{
-			checkf(CurrentState == EReplicatedObjectState::PendingCreate || CurrentState == EReplicatedObjectState::CancelPendingDestroy, TEXT("Trying to set state %s when state is %s. IsDestructionInfo: %u IsSubObject: %u"), LexToString(NewState), LexToString(CurrentState), IsDestructionInfo, IsSubObject);
+			ensureMsgf(CurrentState == EReplicatedObjectState::PendingCreate || CurrentState == EReplicatedObjectState::CancelPendingDestroy, TEXT("Trying to set state %s when state is %s. IsDestructionInfo: %u IsSubObject: %u"), LexToString(NewState), LexToString(CurrentState), IsDestructionInfo, IsSubObject);
 		}
 		break;
 		case EReplicatedObjectState::Created:
 		{
-			checkf(CurrentState == EReplicatedObjectState::PendingCreate || CurrentState == EReplicatedObjectState::WaitOnCreateConfirmation || CurrentState == EReplicatedObjectState::CancelPendingDestroy || CurrentState == EReplicatedObjectState::WaitOnFlush, TEXT("Trying to set state %s when state is %s. IsDestructionInfo: %u IsSubObject: %u"), LexToString(NewState), LexToString(CurrentState), IsDestructionInfo, IsSubObject);
+			ensureMsgf(CurrentState == EReplicatedObjectState::PendingCreate || CurrentState == EReplicatedObjectState::WaitOnCreateConfirmation || CurrentState == EReplicatedObjectState::CancelPendingDestroy || CurrentState == EReplicatedObjectState::WaitOnFlush, TEXT("Trying to set state %s when state is %s. IsDestructionInfo: %u IsSubObject: %u"), LexToString(NewState), LexToString(CurrentState), IsDestructionInfo, IsSubObject);
 		}
 		break;
 		case EReplicatedObjectState::PendingTearOff:
 		{
-			checkf(CurrentState == EReplicatedObjectState::PendingTearOff || CurrentState == EReplicatedObjectState::WaitOnFlush || CurrentState == EReplicatedObjectState::WaitOnCreateConfirmation || CurrentState == EReplicatedObjectState::Created || CurrentState == EReplicatedObjectState::WaitOnDestroyConfirmation, TEXT("Trying to set state %s when state is %s. IsDestructionInfo: %u IsSubObject: %u"), LexToString(NewState), LexToString(CurrentState), IsDestructionInfo, IsSubObject);
+			ensureMsgf(CurrentState == EReplicatedObjectState::PendingTearOff || CurrentState == EReplicatedObjectState::WaitOnFlush || CurrentState == EReplicatedObjectState::WaitOnCreateConfirmation || CurrentState == EReplicatedObjectState::Created || CurrentState == EReplicatedObjectState::WaitOnDestroyConfirmation, TEXT("Trying to set state %s when state is %s. IsDestructionInfo: %u IsSubObject: %u"), LexToString(NewState), LexToString(CurrentState), IsDestructionInfo, IsSubObject);
 		}
 		break;
 		case EReplicatedObjectState::SubObjectPendingDestroy:
 		{
-			checkf(CurrentState == EReplicatedObjectState::PendingDestroy || CurrentState == EReplicatedObjectState::SubObjectPendingDestroy || CurrentState == EReplicatedObjectState::WaitOnCreateConfirmation || CurrentState == EReplicatedObjectState::Created || CurrentState == EReplicatedObjectState::WaitOnFlush || CurrentState == EReplicatedObjectState::WaitOnDestroyConfirmation, TEXT("Trying to set state %s when state is %s. IsDestructionInfo: %u IsSubObject: %u"), LexToString(NewState), LexToString(CurrentState), IsDestructionInfo, IsSubObject);
+			ensureMsgf(CurrentState == EReplicatedObjectState::PendingDestroy || CurrentState == EReplicatedObjectState::SubObjectPendingDestroy || CurrentState == EReplicatedObjectState::WaitOnCreateConfirmation || CurrentState == EReplicatedObjectState::Created || CurrentState == EReplicatedObjectState::WaitOnFlush || CurrentState == EReplicatedObjectState::WaitOnDestroyConfirmation, TEXT("Trying to set state %s when state is %s. IsDestructionInfo: %u IsSubObject: %u"), LexToString(NewState), LexToString(CurrentState), IsDestructionInfo, IsSubObject);
 		}
 		break;
 		case EReplicatedObjectState::WaitOnFlush:
 		{
-			checkf(CurrentState != EReplicatedObjectState::Invalid, TEXT("Trying to set state %s when state is %s. IsDestructionInfo: %u IsSubObject: %u"), LexToString(NewState), LexToString(CurrentState), IsDestructionInfo, IsSubObject);
+			ensureMsgf(CurrentState != EReplicatedObjectState::Invalid, TEXT("Trying to set state %s when state is %s. IsDestructionInfo: %u IsSubObject: %u"), LexToString(NewState), LexToString(CurrentState), IsDestructionInfo, IsSubObject);
 		}
 		break;
 		case EReplicatedObjectState::PendingDestroy:
 		{
-			checkf(CurrentState != EReplicatedObjectState::Invalid, TEXT("Trying to set state %s when state is %s. IsDestructionInfo: %u IsSubObject: %u"), LexToString(NewState), LexToString(CurrentState), IsDestructionInfo, IsSubObject);
+			ensureMsgf(CurrentState != EReplicatedObjectState::Invalid, TEXT("Trying to set state %s when state is %s. IsDestructionInfo: %u IsSubObject: %u"), LexToString(NewState), LexToString(CurrentState), IsDestructionInfo, IsSubObject);
 		}
 		break;
 		case EReplicatedObjectState::WaitOnDestroyConfirmation:
 		{
-			checkf(CurrentState >= EReplicatedObjectState::PendingTearOff, TEXT("Trying to set state %s when state is %s. IsDestructionInfo: %u IsSubObject: %u"), LexToString(NewState), LexToString(CurrentState), IsDestructionInfo, IsSubObject);
+			ensureMsgf(CurrentState >= EReplicatedObjectState::PendingTearOff, TEXT("Trying to set state %s when state is %s. IsDestructionInfo: %u IsSubObject: %u"), LexToString(NewState), LexToString(CurrentState), IsDestructionInfo, IsSubObject);
 		}
 		break;
 		case EReplicatedObjectState::CancelPendingDestroy:
 		{
-			checkf(CurrentState == EReplicatedObjectState::WaitOnDestroyConfirmation || CurrentState == EReplicatedObjectState::CancelPendingDestroy, TEXT("Trying to set state %s when state is %s. IsDestructionInfo: %u IsSubObject: %u"), LexToString(NewState), LexToString(CurrentState), IsDestructionInfo, IsSubObject);
+			ensureMsgf(CurrentState == EReplicatedObjectState::WaitOnDestroyConfirmation || CurrentState == EReplicatedObjectState::CancelPendingDestroy, TEXT("Trying to set state %s when state is %s. IsDestructionInfo: %u IsSubObject: %u"), LexToString(NewState), LexToString(CurrentState), IsDestructionInfo, IsSubObject);
 		}
 		break;
 		case EReplicatedObjectState::Destroyed:
 		{
-			checkf(CurrentState == EReplicatedObjectState::WaitOnDestroyConfirmation || CurrentState == EReplicatedObjectState::PendingTearOff || CurrentState == EReplicatedObjectState::CancelPendingDestroy, TEXT("Trying to set state %s when state is %s. IsDestructionInfo: %u IsSubObject: %u"), LexToString(NewState), LexToString(CurrentState), IsDestructionInfo, IsSubObject);
+			ensureMsgf(CurrentState == EReplicatedObjectState::WaitOnDestroyConfirmation || CurrentState == EReplicatedObjectState::PendingTearOff || CurrentState == EReplicatedObjectState::CancelPendingDestroy, TEXT("Trying to set state %s when state is %s. IsDestructionInfo: %u IsSubObject: %u"), LexToString(NewState), LexToString(CurrentState), IsDestructionInfo, IsSubObject);
 		}
 		break;
 		case EReplicatedObjectState::PermanentlyDestroyed:
 		{
-			checkf(CurrentState == EReplicatedObjectState::Invalid || CurrentState == EReplicatedObjectState::WaitOnCreateConfirmation, TEXT("Trying to set state %s when state is %s. IsDestructionInfo: %u IsSubObject: %u"), LexToString(NewState), LexToString(CurrentState), IsDestructionInfo, IsSubObject);
+			ensureMsgf(CurrentState == EReplicatedObjectState::Invalid || CurrentState == EReplicatedObjectState::WaitOnCreateConfirmation, TEXT("Trying to set state %s when state is %s. IsDestructionInfo: %u IsSubObject: %u"), LexToString(NewState), LexToString(CurrentState), IsDestructionInfo, IsSubObject);
 		}
 		break;
 		case EReplicatedObjectState::Invalid:
 		{
-			checkf(CurrentState == EReplicatedObjectState::PermanentlyDestroyed || CurrentState == EReplicatedObjectState::Destroyed || CurrentState == EReplicatedObjectState::PendingCreate, TEXT("Trying to set state %s when state is %s. IsDestructionInfo: %u IsSubObject: %u"), LexToString(NewState), LexToString(CurrentState), IsDestructionInfo, IsSubObject);
+			ensureMsgf(CurrentState == EReplicatedObjectState::PermanentlyDestroyed || CurrentState == EReplicatedObjectState::Destroyed || CurrentState == EReplicatedObjectState::PendingCreate, TEXT("Trying to set state %s when state is %s. IsDestructionInfo: %u IsSubObject: %u"), LexToString(NewState), LexToString(CurrentState), IsDestructionInfo, IsSubObject);
 		}
 		break;
 
 		default:
-			checkf(false, TEXT("Trying to set state %s when state is %s. IsDestructionInfo: %u IsSubObject: %u"), ToCStr(FString::FromInt(int(NewState))), LexToString(CurrentState), IsDestructionInfo, IsSubObject);
+			ensureMsgf(false, TEXT("Trying to set state %s when state is %s. IsDestructionInfo: %u IsSubObject: %u"), ToCStr(FString::FromInt(int(NewState))), LexToString(CurrentState), IsDestructionInfo, IsSubObject);
 			break;
 	};
 
@@ -858,10 +858,32 @@ void FReplicationWriter::UpdateDirtyGlobalLifetimeConditionals(TArrayView<FInter
 	for (FInternalNetRefIndex InternalObjectIndex : ObjectsWithDirtyConditionals)
 	{
 		FReplicationInfo& Info = ReplicatedObjects[InternalObjectIndex];
-		if (Info.GetState() != EReplicatedObjectState::Invalid)
+		if (Info.GetState() != EReplicatedObjectState::Invalid && Info.GetState() < EReplicatedObjectState::PendingDestroy)
 		{
-			// We just set a flag that we will process if the object is going to be replicated
-			Info.HasDirtyConditionals = 1;
+			if (Info.IsSubObject)
+			{
+				const FNetRefHandleManager::FReplicatedObjectData& ObjectData = NetRefHandleManager->GetReplicatedObjectDataNoCheck(InternalObjectIndex);
+				if (ObjectData.IsSubObject())
+				{
+					FReplicationInfo& OwnerInfo = ReplicatedObjects[ObjectData.SubObjectRootIndex];
+					if (OwnerInfo.GetState() != EReplicatedObjectState::Invalid && OwnerInfo.GetState() < EReplicatedObjectState::PendingDestroy)
+					{
+						UE_LOG_REPLICATIONWRITER_CONN(TEXT("UpdateDirtyGlobalLifetimeConditionals for - %s"), ToCStr(NetRefHandleManager->PrintObjectFromIndex(ObjectData.SubObjectRootIndex)));
+						
+						MarkObjectDirty(ObjectData.SubObjectRootIndex, "UpdateDirtyGlobalLifetimeConditionals");
+						OwnerInfo.HasDirtyConditionals = 1U;
+						OwnerInfo.HasDirtySubObjects = 1U;
+					}
+				}
+			}
+			else
+			{
+				UE_LOG_REPLICATIONWRITER_CONN(TEXT("UpdateDirtyGlobalLifetimeConditionals for - %s"), ToCStr(NetRefHandleManager->PrintObjectFromIndex(InternalObjectIndex)));
+	
+				MarkObjectDirty(InternalObjectIndex, "UpdateDirtyGlobalLifetimeConditionals2");
+				Info.HasDirtyConditionals = 1U;
+				Info.HasDirtySubObjects = 1U;
+			}
 		}
 	}
 }
@@ -2091,14 +2113,16 @@ FReplicationWriter::EWriteObjectStatus FReplicationWriter::WriteObjectAndSubObje
 	{
 		for (uint32 SubObjectIndex : NetRefHandleManager->GetSubObjects(InternalIndex))
 		{
-			FReplicationInfo& SubObjectInfo = GetReplicationInfo(SubObjectIndex);
-			const EReplicatedObjectState SubObjectState = SubObjectInfo.GetState();
-			if (SubObjectInfo.HasChangemaskFilter && (SubObjectState > EReplicatedObjectState::Invalid && SubObjectState < EReplicatedObjectState::PendingDestroy))
+			// Need to be a bit careful what we explicitly dirty here, as we might have subobjects waiting for creation confirmation.
+			if (CanSendObject(SubObjectIndex))
 			{
 				MarkObjectDirty(SubObjectIndex, "UpdateGlobalConditional");
+				FReplicationInfo& SubObjectInfo = GetReplicationInfo(SubObjectIndex);
 				SubObjectInfo.HasDirtyChangeMask = 1;
-				Info.HasDirtySubObjects = 1;
 			}
+
+			// Always process subobjects when we have updated conditionals to ensure that nested conditionals are processed.
+			Info.HasDirtySubObjects = 1;
 		}
 	}
 	Info.HasDirtyConditionals = 0;
