Iris
* NetGroupHandles are now unique and cannot reference a future unrelated group once the group it referenced has been destroyed.

diff --git a/Engine/Plugins/Runtime/ReplicationSystemTestPlugin/Source/Private/Tests/ReplicationSystem/Filtering/TestGroups.cpp b/Engine/Plugins/Runtime/ReplicationSystemTestPlugin/Source/Private/Tests/ReplicationSystem/Filtering/TestGroups.cpp
--- a/Engine/Plugins/Runtime/ReplicationSystemTestPlugin/Source/Private/Tests/ReplicationSystem/Filtering/TestGroups.cpp
+++ b/Engine/Plugins/Runtime/ReplicationSystemTestPlugin/Source/Private/Tests/ReplicationSystem/Filtering/TestGroups.cpp
@@ -122,18 +122,18 @@ UE_NET_TEST_FIXTURE(FTestGroupsFixture, PublicGroupAPI)
 		UE_NET_ASSERT_FALSE(ReplicationSystem->IsValidGroup(FoundGroupHandle));
 	}
 
-	// Destroy/Create: Does not pass at the moment
-	/*{
+	// Destroy/Create
+	{
 		// Create a group
-		FNetObjectGroupHandle ReusedGroupHandle = ReplicationSystem->CreateGroup(NAME_None);
-		UE_NET_ASSERT_TRUE(ReusedGroupHandle.IsValid());
-		UE_NET_ASSERT_TRUE(ReplicationSystem->IsValidGroup(ReusedGroupHandle));
+		FNetObjectGroupHandle DestroyedGroupHandle = ReplicationSystem->CreateGroup(NAME_None);
+		UE_NET_ASSERT_TRUE(DestroyedGroupHandle.IsValid());
+		UE_NET_ASSERT_TRUE(ReplicationSystem->IsValidGroup(DestroyedGroupHandle));
 
 		// Destroy the group
-		ReplicationSystem->DestroyGroup(ReusedGroupHandle);
+		ReplicationSystem->DestroyGroup(DestroyedGroupHandle);
 
 		// Should be invalid
-		UE_NET_ASSERT_FALSE(ReplicationSystem->IsValidGroup(ReusedGroupHandle));
+		UE_NET_ASSERT_FALSE(ReplicationSystem->IsValidGroup(DestroyedGroupHandle));
 
 		{
 			// Create a new random group
@@ -141,10 +141,9 @@ UE_NET_TEST_FIXTURE(FTestGroupsFixture, PublicGroupAPI)
 			UE_NET_ASSERT_TRUE(ReplicationSystem->IsValidGroup(RandomGroupHandle));
 		}
 
-		// Problem: The destroyed group should still be invalid but the old handle index is now assigned to the new group.
-		// This means the old handle now points to the new group and is considered valid.
-		UE_NET_ASSERT_FALSE(ReplicationSystem->IsValidGroup(ReusedGroupHandle));
-	}*/
+		// The destroyed group should still be invalid
+		UE_NET_ASSERT_FALSE(ReplicationSystem->IsValidGroup(DestroyedGroupHandle));
+	}
 
 	// Add/Remove
 	{
diff --git a/Engine/Source/Runtime/Experimental/Iris/Core/Private/Iris/ReplicationSystem/Conditionals/ReplicationConditionals.cpp b/Engine/Source/Runtime/Experimental/Iris/Core/Private/Iris/ReplicationSystem/Conditionals/ReplicationConditionals.cpp
--- a/Engine/Source/Runtime/Experimental/Iris/Core/Private/Iris/ReplicationSystem/Conditionals/ReplicationConditionals.cpp
+++ b/Engine/Source/Runtime/Experimental/Iris/Core/Private/Iris/ReplicationSystem/Conditionals/ReplicationConditionals.cpp
@@ -555,41 +555,39 @@ void FReplicationConditionals::GetChildSubObjectsToReplicate(uint32 ReplicatingC
 				const ELifetimeCondition LifeTimeCondition = (ELifetimeCondition)SubObjectsInfo.SubObjectLifeTimeConditions[ArrayIndex];
 				if (LifeTimeCondition == COND_NetGroup)
 				{
-					uint32 GroupCount = 0U;
-					if (const FNetObjectGroupHandle* GroupMemberships = NetObjectGroups->GetGroupMemberships(SubObjectIndex, GroupCount))
+					bool bShouldReplicateSubObject = false;
+					//TArray<FNetObjectGroupHandle> GroupsMemberOf;
+					//NetObjectGroups->GetGroupHandlesOfNetObject(SubObjectIndex, GroupsMemberOf);
+					// 
+					 const TArrayView<const FNetObjectGroupHandle::FGroupIndexType> GroupIndexes = NetObjectGroups->GetGroupIndexesOfNetObject(SubObjectIndex);
+					for (const FNetObjectGroupHandle::FGroupIndexType GroupIndex : GroupIndexes)
 					{
-						bool bShouldReplicateSubObject = false;
-						for (uint32 GroupIt = 0U; GroupIt < GroupCount; ++GroupIt)
+						const FNetObjectGroupHandle NetGroup = NetObjectGroups->GetHandleFromIndex(GroupIndex);
+
+						if (NetGroup.IsNetGroupOwnerNetObjectGroup())
 						{
-							const FNetObjectGroupHandle NetGroup = GroupMemberships[GroupIt];
-							if (NetGroup.IsNetGroupOwnerNetObjectGroup())
-							{
-								bShouldReplicateSubObject = LifetimeConditionals.IsConditionEnabled(COND_OwnerOnly);
-							}
-							else if (NetGroup.IsNetGroupReplayNetObjectGroup())
-							{
-								bShouldReplicateSubObject = LifetimeConditionals.IsConditionEnabled(COND_ReplayOnly);
-							}
-							else
-							{
-								ENetFilterStatus ReplicationStatus = ENetFilterStatus::Disallow;
-								ensureAlwaysMsgf(ReplicationFiltering->GetSubObjectFilterStatus(NetGroup, ReplicatingConnectionId, ReplicationStatus), TEXT("FReplicationConditionals::GetChildSubObjectsToReplicat Trying to filter with group %u that is not a SubObjectFilterGroup"), NetGroup.GetGroupIndex());
-								bShouldReplicateSubObject = ReplicationStatus != ENetFilterStatus::Disallow;
-							}
-						
-							if (bShouldReplicateSubObject)
-							{
-								GetChildSubObjectsToReplicate(ReplicatingConnectionId, LifetimeConditionals, SubObjectIndex, OutSubObjectsToReplicate);
-								OutSubObjectsToReplicate.Add(SubObjectIndex);
-								break;
-							}
+							bShouldReplicateSubObject = LifetimeConditionals.IsConditionEnabled(COND_OwnerOnly);
 						}
-
-						if (!bShouldReplicateSubObject)
+						else if (NetGroup.IsNetGroupReplayNetObjectGroup())
+						{
+							bShouldReplicateSubObject = LifetimeConditionals.IsConditionEnabled(COND_ReplayOnly);
+						}
+						else
 						{
-							UE_LOG(LogIrisConditionals, VeryVerbose, TEXT("%s Filtered out by COND_NetGroup"), *NetRefHandleManager->PrintObjectFromIndex(SubObjectIndex));
+							ENetFilterStatus ReplicationStatus = ENetFilterStatus::Disallow;
+							ensureMsgf(ReplicationFiltering->GetSubObjectFilterStatus(NetGroup, ReplicatingConnectionId, ReplicationStatus), TEXT("FReplicationConditionals::GetChildSubObjectsToReplicat Trying to filter with group %u that is not a SubObjectFilterGroup"), NetGroup.GetGroupIndex());
+							bShouldReplicateSubObject = ReplicationStatus != ENetFilterStatus::Disallow;
+						}
+						
+						if (bShouldReplicateSubObject)
+						{
+							GetChildSubObjectsToReplicate(ReplicatingConnectionId, LifetimeConditionals, SubObjectIndex, OutSubObjectsToReplicate);
+							OutSubObjectsToReplicate.Add(SubObjectIndex);
+							break;
 						}
 					}
+
+					UE_CLOG(!bShouldReplicateSubObject, LogIrisConditionals, VeryVerbose, TEXT("%s Filtered out by COND_NetGroup"), *NetRefHandleManager->PrintObjectFromIndex(SubObjectIndex));
 				}
 				else if (LifetimeConditionals.IsConditionEnabled(LifeTimeCondition))
 				{
diff --git a/Engine/Source/Runtime/Experimental/Iris/Core/Private/Iris/ReplicationSystem/Filtering/NetObjectGroups.cpp b/Engine/Source/Runtime/Experimental/Iris/Core/Private/Iris/ReplicationSystem/Filtering/NetObjectGroups.cpp
--- a/Engine/Source/Runtime/Experimental/Iris/Core/Private/Iris/ReplicationSystem/Filtering/NetObjectGroups.cpp
+++ b/Engine/Source/Runtime/Experimental/Iris/Core/Private/Iris/ReplicationSystem/Filtering/NetObjectGroups.cpp
@@ -13,11 +13,16 @@ DEFINE_LOG_CATEGORY_STATIC(LogIrisGroup, Log, All)
 namespace UE::Net::Private
 {
 
+	namespace NetObjectGroupsInternal
+	{
+		static FNetObjectGroupHandle::FGroupIndexType NextEpoch = 1U;
+	}
+
 static int32 CVarEnsureIfNumGroupMembershipsExceedsNum = 128;
 static FAutoConsoleVariableRef CVarEnsureIfNumGroupMembershipsExceeds(TEXT("net.Iris.EnsureIfNumGroupMembershipsExceeds"), CVarEnsureIfNumGroupMembershipsExceedsNum, TEXT("If set to a positive number we will warn and ensure if an object is added to a high number of groups."), ECVF_Default );
 
 FNetObjectGroups::FNetObjectGroups()
-: CurrentEpoch(++NextEpoch)
+: CurrentEpoch(++NetObjectGroupsInternal::NextEpoch)
 {
 }
 
@@ -27,14 +32,14 @@ FNetObjectGroups::~FNetObjectGroups()
 
 bool FNetObjectGroups::IsMemberOf(const FNetObjectGroupMembership& Target, FNetObjectGroupHandle Group)
 {
-	return Target.Groups.Contains(Group);
+	return Target.ContainsMembership(Group);
 }
 
 bool FNetObjectGroups::AddGroupMembership(FNetObjectGroupMembership& Target, FNetObjectGroupHandle Group)
 {
-	if (!Target.Groups.Contains(Group))
+	if (!Target.ContainsMembership(Group))
 	{
-		Target.Groups.Add(Group);
+		Target.AddMembership(Group);
 
 		return true;
 	}
@@ -42,14 +47,9 @@ bool FNetObjectGroups::AddGroupMembership(FNetObjectGroupMembership& Target, FNe
 	return false;
 }
 
-void FNetObjectGroups::RemoveGroupMembership(FNetObjectGroupMembership& Target, FNetObjectGroupHandle Group)
-{
-	Target.Groups.RemoveSingleSwap(Group);
-}
-
 void FNetObjectGroups::ResetGroupMembership(FNetObjectGroupMembership& Target)
 {
-	Target.Groups.Empty();
+	Target.ResetMemberships();
 }
 
 void FNetObjectGroups::Init(const FNetObjectGroupInitParams& Params)
@@ -100,14 +100,20 @@ FNetObjectGroupHandle FNetObjectGroups::CreateGroup(FName InGroupName)
 		}
 		
 	}
+
+	const uint32 NewGroupId = NextGroupUniqueId++;
 	
-	const uint32 Index = static_cast<uint32>(Groups.Add(FNetObjectGroup{.GroupName = InGroupName}));
+	const uint32 Index = (uint32)Groups.Emplace(FNetObjectGroup{.GroupName = InGroupName, .GroupId = NewGroupId});
 
-	UE_LOG(LogIrisGroup, Log, TEXT("Created GroupHandle: %u named: %s"), Index, *InGroupName.ToString());
+	UE_LOG(LogIrisGroup, Log, TEXT("Created GroupHandle named: %s Index: %u Id: %u"), *InGroupName.ToString(), Index, NewGroupId);
 
 	FNetObjectGroupHandle GroupHandle;
 	GroupHandle.Index = static_cast<FNetObjectGroupHandle::FGroupIndexType>(Index);
 	GroupHandle.Epoch = CurrentEpoch;
+	GroupHandle.UniqueId = NewGroupId;
+
+	ensureMsgf(NextGroupUniqueId != 0, TEXT("We created more than MAX_uint32 unique NetGroups. An Id overflow occured and it may cause conflicts with existing handles."));
+
 	return GroupHandle;
 }
 
@@ -115,9 +121,11 @@ void FNetObjectGroups::DestroyGroup(FNetObjectGroupHandle GroupHandle)
 {
 	if (IsValidGroup(GroupHandle))
 	{
-		ClearGroup(GroupHandle);
-
 		FNetObjectGroup& Group = Groups[GroupHandle.GetGroupIndex()];
+
+		UE_LOG(LogIrisGroup, Log, TEXT("Destroyed GroupHandle named: %s Index: %u Id: %u"), *Group.GroupName.ToString(), GroupHandle.Index, Group.GroupId);
+
+		ClearGroup(GroupHandle);		
 		Groups.RemoveAt(GroupHandle.GetGroupIndex());
 	}
 }
@@ -131,7 +139,7 @@ void FNetObjectGroups::ClearGroup(FNetObjectGroupHandle GroupHandle)
 		for (FInternalNetRefIndex InternalIndex : Group.Members)
 		{
 			checkSlow(IsMemberOf(GroupMemberships[InternalIndex], GroupHandle));
-			RemoveGroupMembership(GroupMemberships[InternalIndex], GroupHandle);
+			GroupMemberships[InternalIndex].RemoveMembership(GroupHandle);
 		}
 
 		Group.Members.Empty();
@@ -148,16 +156,22 @@ FNetObjectGroup* FNetObjectGroups::GetGroup(FNetObjectGroupHandle GroupHandle)
 	return IsValidGroup(GroupHandle) ? &Groups[GroupHandle.GetGroupIndex()] : nullptr;
 }
 
-const FNetObjectGroup* FNetObjectGroups::GetGroupByIndex(FNetObjectGroupHandle::FGroupIndexType GroupIndex) const
+const FNetObjectGroup* FNetObjectGroups::GetGroupFromIndex(FNetObjectGroupHandle::FGroupIndexType GroupIndex) const
 {
 	return (GroupIndex != FNetObjectGroupHandle::InvalidNetObjectGroupIndex && Groups.IsValidIndex(GroupIndex)) ? &Groups[GroupIndex] : nullptr;
 }
 
-FNetObjectGroup* FNetObjectGroups::GetGroupByIndex(FNetObjectGroupHandle::FGroupIndexType GroupIndex)
+FNetObjectGroup* FNetObjectGroups::GetGroupFromIndex(FNetObjectGroupHandle::FGroupIndexType GroupIndex)
 {
 	return (GroupIndex != FNetObjectGroupHandle::InvalidNetObjectGroupIndex && Groups.IsValidIndex(GroupIndex)) ? &Groups[GroupIndex] : nullptr;
 }
 
+FNetObjectGroupHandle FNetObjectGroups::GetHandleFromGroup(const FNetObjectGroup* InGroup) const
+{
+	check(InGroup);
+	return FNetObjectGroupHandle(GetIndexFromGroup(InGroup), CurrentEpoch, InGroup->GroupId);
+}
+
 bool FNetObjectGroups::Contains(FNetObjectGroupHandle GroupHandle, FInternalNetRefIndex InternalIndex) const
 {
 	// If the group does not exist we cannot be in it..
@@ -174,51 +188,67 @@ bool FNetObjectGroups::Contains(FNetObjectGroupHandle GroupHandle, FInternalNetR
 
 void FNetObjectGroups::AddToGroup(FNetObjectGroupHandle GroupHandle, FInternalNetRefIndex InternalIndex)
 {
+	if (InternalIndex == FNetRefHandleManager::InvalidInternalIndex) 
+	{
+		return;
+	}
+	
 	FNetObjectGroup* Group = GetGroup(GroupHandle);
-	if (InternalIndex != FNetRefHandleManager::InvalidInternalIndex && Group)
+	if (!Group)
 	{
-		UE_LOG(LogIrisGroup, Verbose, TEXT("FNetObjectGroups::AddToGroup Adding %s to Group %s (GroupIndex: %u)"),
-			*NetRefHandleManager->PrintObjectFromIndex(InternalIndex), *GetGroupName(GroupHandle).ToString(), GroupHandle.GetGroupIndex());
+		UE_LOG(LogIrisGroup, Warning, TEXT("FNetObjectGroups::AddToGroup received invalid group (Index: %u Id: %u). Cannot add %s to it."), GroupHandle.Index, GroupHandle.UniqueId, *NetRefHandleManager->PrintObjectFromIndex(InternalIndex));
+		return;
+	}
 
-		FNetObjectGroupMembership& Membership = GroupMemberships[InternalIndex];
-		if (AddGroupMembership(Membership, GroupHandle))
-		{
-			Group->Members.AddUnique(InternalIndex);
+	UE_LOG(LogIrisGroup, Verbose, TEXT("FNetObjectGroups::AddToGroup Adding %s to Group %s (Index: %u Id: %u)"),
+		*NetRefHandleManager->PrintObjectFromIndex(InternalIndex), *GetGroupName(GroupHandle).ToString(), GroupHandle.Index, GroupHandle.UniqueId);
 
-			if (IsFilterGroup(*Group))
-			{
-				GroupFilteredOutObjects.SetBit(InternalIndex);
-			}
+	FNetObjectGroupMembership& Membership = GroupMemberships[InternalIndex];
+	if (AddGroupMembership(Membership, GroupHandle))
+	{
+		Group->Members.AddUnique(InternalIndex);
 
-			if (CVarEnsureIfNumGroupMembershipsExceedsNum > 0 && Membership.Groups.Num() > CVarEnsureIfNumGroupMembershipsExceedsNum)
-			{
-				UE_LOG(LogIrisGroup, Error, TEXT("FNetObjectGroups::AddGroupMembership Unexpected high num groupmemberships for group %s (GroupIndex: %u) NetObject %s is member of %d groups."),
-					*GetGroupName(GroupHandle).ToString(), GroupHandle.GetGroupIndex(), *NetRefHandleManager->PrintObjectFromIndex(InternalIndex), Membership.Groups.Num());
-				ensure(false);			
-			}
+		if (IsFilterGroup(*Group))
+		{
+			GroupFilteredOutObjects.SetBit(InternalIndex);
+		}
+
+		if (CVarEnsureIfNumGroupMembershipsExceedsNum > 0 && Membership.NumMemberships() > CVarEnsureIfNumGroupMembershipsExceedsNum)
+		{
+			UE_LOG(LogIrisGroup, Error, TEXT("FNetObjectGroups::AddGroupMembership Unexpected high num groupmemberships for group %s (Index: %u Id: %u) NetObject %s is member of %d groups."),
+				*GetGroupName(GroupHandle).ToString(), GroupHandle.Index, GroupHandle.UniqueId, *NetRefHandleManager->PrintObjectFromIndex(InternalIndex), Membership.NumMemberships());
+			ensure(false);
 		}
 	}
 }
 
 void FNetObjectGroups::RemoveFromGroup(FNetObjectGroupHandle GroupHandle, FInternalNetRefIndex InternalIndex)
 {
+	if (InternalIndex == FNetRefHandleManager::InvalidInternalIndex)
+	{
+		return;
+	}
+
 	FNetObjectGroup* Group = GetGroup(GroupHandle);
-	if (InternalIndex != FNetRefHandleManager::InvalidInternalIndex && Group)
+	if (!Group)
 	{
-		UE_LOG(LogIrisGroup, Verbose, TEXT("FNetObjectGroups::RemoveFromGroup Removing %s from Group %s (GroupIndex: %u)"),
-			*NetRefHandleManager->PrintObjectFromIndex(InternalIndex), *GetGroupName(GroupHandle).ToString(), GroupHandle.GetGroupIndex());
+		UE_LOG(LogIrisGroup, Warning, TEXT("FNetObjectGroups::RemoveFromGroup received invalid group (Index: %u Id: %u). Cannot remove %s from it."), GroupHandle.Index, GroupHandle.UniqueId, *NetRefHandleManager->PrintObjectFromIndex(InternalIndex));
+		return;
+	}
 
-		FNetObjectGroupMembership& GroupMembership = GroupMemberships[InternalIndex];
-		checkSlow(IsMemberOf(GroupMembership, GroupHandle));
+	UE_LOG(LogIrisGroup, Verbose, TEXT("FNetObjectGroups::RemoveFromGroup Removing %s from Group %s (Index: %u Id: %u)"),
+		*NetRefHandleManager->PrintObjectFromIndex(InternalIndex), *GetGroupName(GroupHandle).ToString(), GroupHandle.Index, GroupHandle.UniqueId);
 
-		RemoveGroupMembership(GroupMembership, GroupHandle);
-		Group->Members.RemoveSingle(InternalIndex);
+	FNetObjectGroupMembership& GroupMembership = GroupMemberships[InternalIndex];
+	checkSlow(IsMemberOf(GroupMembership, GroupHandle));
 
-		// Check to see if the object is still part of a filter group
-		if (!IsInAnyFilterGroup(GroupMembership))
-		{
-			GroupFilteredOutObjects.ClearBit(InternalIndex);
-		}
+	GroupMembership.RemoveMembership(GroupHandle);
+	Group->Members.RemoveSingle(InternalIndex);
+
+	// Check to see if the object is still part of a filter group
+	if (!IsInAnyFilterGroup(GroupMembership))
+	{
+		GroupFilteredOutObjects.ClearBit(InternalIndex);
 	}
 }
 
@@ -304,32 +334,42 @@ bool FNetObjectGroups::IsFilterGroup(FNetObjectGroupHandle GroupHandle) const
 
 bool FNetObjectGroups::IsInAnyFilterGroup(const FNetObjectGroupMembership& GroupMembership) const
 {
-	for (FNetObjectGroupHandle AssignedGroup : GroupMembership.Groups)
+	for (FNetObjectGroupHandle::FGroupIndexType AssignedGroupIndex : GroupMembership.GetGroupIndexes())
 	{
-		if (!AssignedGroup.IsValid())
-		{
-			// Note: An invalid group means we found the end of the array
-			return false;
-		}
-		else if (IsFilterGroup(Groups[AssignedGroup.GetGroupIndex()]))
+		if (const FNetObjectGroup* Group = GetGroupFromIndex(AssignedGroupIndex))
 		{
-			return true;
+			if (IsFilterGroup(*Group))
+			{
+				return true;
+			}
 		}
 	}
 
 	return false;
 }
 
-const FNetObjectGroupHandle* FNetObjectGroups::GetGroupMemberships(FInternalNetRefIndex InternalIndex, uint32& GroupCount) const
+const TArrayView<const FNetObjectGroupHandle::FGroupIndexType> FNetObjectGroups::GetGroupIndexesOfNetObject(FInternalNetRefIndex InternalIndex) const
 {
-	if (InternalIndex >= (uint32)GroupMemberships.Num())
+	if (ensure(GroupMemberships.IsValidIndex(InternalIndex)))
 	{
-		return nullptr;
+		return GroupMemberships[InternalIndex].GetGroupIndexes();
 	}
 
-	const FNetObjectGroupMembership& MemberShip = GroupMemberships[InternalIndex];
-	GroupCount = MemberShip.Groups.Num();
-	return GroupCount ? MemberShip.Groups.GetData() : nullptr;
+	return TArrayView<const FNetObjectGroupHandle::FGroupIndexType>();
 }
 
+void FNetObjectGroups::GetGroupHandlesOfNetObject(FInternalNetRefIndex InternalIndex, TArray<FNetObjectGroupHandle>& OutHandles) const
+{
+	if (ensure(GroupMemberships.IsValidIndex(InternalIndex)))
+	{
+		const TArrayView<const FNetObjectGroupHandle::FGroupIndexType> GroupIndexes = GroupMemberships[InternalIndex].GetGroupIndexes();
+		OutHandles.Reserve(GroupIndexes.Num());
+
+		for (FNetObjectGroupHandle::FGroupIndexType AssignedGroupIndex : GroupIndexes)
+		{
+			OutHandles.Add(GetHandleFromIndex(AssignedGroupIndex));
+		}
+	}
 }
+
+} // end namespace UE::Net::Private
diff --git a/Engine/Source/Runtime/Experimental/Iris/Core/Private/Iris/ReplicationSystem/Filtering/NetObjectGroups.h b/Engine/Source/Runtime/Experimental/Iris/Core/Private/Iris/ReplicationSystem/Filtering/NetObjectGroups.h
--- a/Engine/Source/Runtime/Experimental/Iris/Core/Private/Iris/ReplicationSystem/Filtering/NetObjectGroups.h
+++ b/Engine/Source/Runtime/Experimental/Iris/Core/Private/Iris/ReplicationSystem/Filtering/NetObjectGroups.h
@@ -4,6 +4,7 @@
 
 #include "CoreTypes.h"
 #include "Containers/Array.h"
+#include "Containers/ArrayView.h"
 #include "Containers/Map.h"
 #include "Containers/SparseArray.h"
 #include "UObject/NameTypes.h"
@@ -36,14 +37,15 @@ struct FNetObjectGroup
 	// Group members can only be replicated objects that have internal indices
 	TArray<FInternalNetRefIndex> Members;
 	FName GroupName;
+	uint32 GroupId = 0U;
 	ENetObjectGroupTraits Traits = ENetObjectGroupTraits::None;
 };
 
 struct FNetObjectGroupInitParams
 {
 	FNetRefHandleManager* NetRefHandleManager = nullptr;
-	uint32 MaxInternalNetRefIndex = 0;
-	uint32 MaxGroupCount = 0;
+	uint32 MaxInternalNetRefIndex = 0U;
+	uint32 MaxGroupCount = 0U;
 };
 
 class FNetObjectGroups
@@ -66,9 +68,11 @@ public:
 	const FNetObjectGroup* GetGroup(FNetObjectGroupHandle GroupHandle) const;
 	FNetObjectGroup* GetGroup(FNetObjectGroupHandle GroupHandle);
 	
-	FNetObjectGroupHandle MakeNetObjectGroupHandle(FNetObjectGroupHandle::FGroupIndexType GroupIndex) const { return FNetObjectGroupHandle(GroupIndex, CurrentEpoch); }
-	const FNetObjectGroup* GetGroupByIndex(FNetObjectGroupHandle::FGroupIndexType GroupIndex) const;
-	FNetObjectGroup* GetGroupByIndex(FNetObjectGroupHandle::FGroupIndexType GroupIndex);
+	const FNetObjectGroup* GetGroupFromIndex(FNetObjectGroupHandle::FGroupIndexType GroupIndex) const;
+	FNetObjectGroup* GetGroupFromIndex(FNetObjectGroupHandle::FGroupIndexType GroupIndex);
+
+	FNetObjectGroupHandle GetHandleFromGroup(const FNetObjectGroup* InGroup) const;
+	FNetObjectGroupHandle GetHandleFromIndex(FNetObjectGroupHandle::FGroupIndexType GroupIndex) const;
 
 	inline FName GetGroupName(FNetObjectGroupHandle GroupHandle) const;
 	inline FString GetGroupNameString(FNetObjectGroupHandle GroupHandle) const;
@@ -100,14 +104,14 @@ public:
 	/** Does the group have the inclusion filter trait */
 	bool IsInclusionFilterGroup(FNetObjectGroupHandle GroupHandle) const;
 
-	/** Get the array of all groups that the NetObject is a member of */
-	const FNetObjectGroupHandle* GetGroupMemberships(FInternalNetRefIndex InternalIndex, uint32& GroupCount) const;
+	/** Get a reference to the indexes of all groups that the NetObject is a member of */
+	const TArrayView<const FNetObjectGroupHandle::FGroupIndexType> GetGroupIndexesOfNetObject(FInternalNetRefIndex InternalIndex) const;
+
+	/** Get a list of all group handles the NetObject is a member of */
+	void GetGroupHandlesOfNetObject(FInternalNetRefIndex InternalIndex, TArray<FNetObjectGroupHandle>& OutHandles) const;
 
 	/** Returns a list of all objects currently part of a group with the filter trait */
-	const FNetBitArrayView GetGroupFilteredOutObjects() const
-	{
-		return MakeNetBitArrayView(GroupFilteredOutObjects);
-	}
+	const FNetBitArrayView GetGroupFilteredOutObjects() const { return MakeNetBitArrayView(GroupFilteredOutObjects); }
 
 	/** Called when the maximum InternalNetRefIndex increased and we need to realloc our lists */
 	void OnMaxInternalNetRefIndexIncreased(FInternalNetRefIndex NewMaxInternalIndex);
@@ -115,12 +119,23 @@ public:
 private:
 	struct FNetObjectGroupMembership
 	{
+	private:
 		enum { NumInlinedGroupHandles = 2 };
-		TArray<FNetObjectGroupHandle, TInlineAllocator<NumInlinedGroupHandles>> Groups;
+		/** The indexes of the groups the netobject is a member of. */
+		TArray<FNetObjectGroupHandle::FGroupIndexType, TInlineAllocator<NumInlinedGroupHandles>> GroupIndexes;
+
+	public:
+
+		bool ContainsMembership(FNetObjectGroupHandle InGroupHandle) const	{ return GroupIndexes.Contains(InGroupHandle.GetGroupIndex()); }
+		void AddMembership(FNetObjectGroupHandle InGroupHandle)		{ GroupIndexes.Add(InGroupHandle.GetGroupIndex()); }
+		void RemoveMembership(FNetObjectGroupHandle InGroupHandle)	{ GroupIndexes.RemoveSingleSwap(InGroupHandle.GetGroupIndex()); }
+		void ResetMemberships()										{ GroupIndexes.Reset(); }
+		int32 NumMemberships() const								{ return GroupIndexes.Num(); }
+
+		const TArrayView<const FNetObjectGroupHandle::FGroupIndexType> GetGroupIndexes() const { return MakeArrayView(GroupIndexes.GetData(), GroupIndexes.Num()); }
 	};
 
 	static bool AddGroupMembership(FNetObjectGroupMembership& Target, FNetObjectGroupHandle Group);
-	static void RemoveGroupMembership(FNetObjectGroupMembership& Target, FNetObjectGroupHandle Group);
 	static void ResetGroupMembership(FNetObjectGroupMembership& Target);
 	static bool IsMemberOf(const FNetObjectGroupMembership& Target, FNetObjectGroupHandle Group);
 
@@ -130,6 +145,8 @@ private:
 
 	bool IsInAnyFilterGroup(const FNetObjectGroupMembership& GroupMembership) const;
 
+	const FNetObjectGroupHandle::FGroupIndexType GetIndexFromGroup(const FNetObjectGroup* InGroup) const;
+
 private:
 
 	FNetRefHandleManager* NetRefHandleManager = nullptr;
@@ -143,20 +160,23 @@ private:
 	// Maximum number of groups that can be exist at once
 	uint32 MaxGroupCount = 0U;
 
-	// Index to use for groups with autogenerated names
+	// Index to use for groups with auto-generated names
 	int32 AutogeneratedGroupNameId = 0;
 
 	// List of objects that are members of a group with a filter trait
 	FNetBitArray GroupFilteredOutObjects;
 
+	// Identifies the ReplicationSystem the group handles were created by
 	FNetObjectGroupHandle::FGroupIndexType CurrentEpoch = 0U;
 
-	inline static FNetObjectGroupHandle::FGroupIndexType NextEpoch = 1U;
+	// Unique Id assigned to each group handle
+	uint32 NextGroupUniqueId = 1U;
 };
 
 inline bool FNetObjectGroups::IsValidGroup(FNetObjectGroupHandle GroupHandle) const
 {
-	return GroupHandle.IsValid() && GroupHandle.Epoch == CurrentEpoch && Groups.IsValidIndex(GroupHandle.GetGroupIndex());
+	const bool bGroupIndexExists = GroupHandle.IsValid() && GroupHandle.Epoch == CurrentEpoch && Groups.IsValidIndex(GroupHandle.Index);
+	return bGroupIndexExists && Groups[GroupHandle.Index].GroupId == GroupHandle.UniqueId;
 }
 
 inline bool FNetObjectGroups::IsFilterGroup(const FNetObjectGroup& Group) const
@@ -180,14 +200,26 @@ inline FNetObjectGroupHandle FNetObjectGroups::FindGroupHandle(FName InGroupName
 	{
 		if (Group.GroupName == InGroupName)
 		{
-			const int32 Index = Groups.PointerToIndex(&Group);
-			return FNetObjectGroupHandle((FNetObjectGroupHandle::FGroupIndexType)Index, CurrentEpoch);
+			const int32 Index = GetIndexFromGroup(&Group);
+			return FNetObjectGroupHandle((FNetObjectGroupHandle::FGroupIndexType)Index, CurrentEpoch, Group.GroupId);
 		}
 	}
 
 	return FNetObjectGroupHandle();
 }
 
+inline FNetObjectGroupHandle FNetObjectGroups::GetHandleFromIndex(FNetObjectGroupHandle::FGroupIndexType GroupIndex) const
+{
+	const FNetObjectGroup* Group = GetGroupFromIndex(GroupIndex);
+	return Group ? FNetObjectGroupHandle(GroupIndex, CurrentEpoch, Group->GroupId) : FNetObjectGroupHandle();
+}
+
+inline const FNetObjectGroupHandle::FGroupIndexType FNetObjectGroups::GetIndexFromGroup(const FNetObjectGroup* InGroup) const
+{
+	check(InGroup);
+	return (FNetObjectGroupHandle::FGroupIndexType)Groups.PointerToIndex(InGroup);
+}
+
 inline FName FNetObjectGroups::GetGroupName(FNetObjectGroupHandle GroupHandle) const
 {
 	if (const FNetObjectGroup* Group = GetGroup(GroupHandle))
diff --git a/Engine/Source/Runtime/Experimental/Iris/Core/Private/Iris/ReplicationSystem/Filtering/ReplicationFiltering.cpp b/Engine/Source/Runtime/Experimental/Iris/Core/Private/Iris/ReplicationSystem/Filtering/ReplicationFiltering.cpp
--- a/Engine/Source/Runtime/Experimental/Iris/Core/Private/Iris/ReplicationSystem/Filtering/ReplicationFiltering.cpp
+++ b/Engine/Source/Runtime/Experimental/Iris/Core/Private/Iris/ReplicationSystem/Filtering/ReplicationFiltering.cpp
@@ -733,7 +733,7 @@ void FReplicationFiltering::InitNewConnections()
 				if (GetConnectionFilterStatus(*ConnectionState, ConnectionId) == ENetFilterStatus::Disallow)
 				{
 					// Apply filter
-					const FNetObjectGroup* Group = this->Groups->GetGroupByIndex(GroupIndex);
+					const FNetObjectGroup* Group = this->Groups->GetGroupFromIndex(GroupIndex);
 					FNetBitArray& GroupExcludedObjects = ConnectionInfo.GroupExcludedObjects;
 					for (const FInternalNetRefIndex ObjectIndex : Group->Members)
 					{
@@ -765,7 +765,7 @@ void FReplicationFiltering::InitNewConnections()
 					const FNetBitArrayView SubObjectInternalIndices = NetRefHandleManager->GetSubObjectInternalIndicesView();
 
 					// Apply filter
-					const FNetObjectGroup* Group = this->Groups->GetGroupByIndex(GroupIndex);
+					const FNetObjectGroup* Group = this->Groups->GetGroupFromIndex(GroupIndex);
 					FNetBitArray& GroupIncludedObjects = ConnectionInfo.GroupIncludedObjects;
 					for (const FInternalNetRefIndex ObjectIndex : Group->Members)
 					{
@@ -994,11 +994,9 @@ void FReplicationFiltering::UpdateObjectsInScope()
 					PrevParentIndex = ParentIndex;
 
 					// If parent is a member of a group filter we need to refresh group filtering to include subobject
-					uint32 GroupMembershipCount = 0U;
-					const FNetObjectGroupHandle* GroupHandles = Groups->GetGroupMemberships(ParentIndex, GroupMembershipCount);
-					for (const FNetObjectGroupHandle GroupHandle : MakeArrayView(GroupHandles, GroupMembershipCount))
+					const TArrayView<const FNetObjectGroupHandle::FGroupIndexType> GroupIndexes = Groups->GetGroupIndexesOfNetObject(ParentIndex);
+					for (const FNetObjectGroupHandle::FGroupIndexType GroupIndex : GroupIndexes)
 					{
-						const FNetObjectGroupHandle::FGroupIndexType GroupIndex = GroupHandle.GetGroupIndex();
 						if (ExclusionFilterGroups.GetBit(GroupIndex))
 						{
 							DirtyExclusionFilterGroups.SetBit(GroupIndex);
@@ -1207,7 +1205,7 @@ void FReplicationFiltering::UpdateGroupExclusionFiltering()
 		IRIS_PROFILER_SCOPE(FReplicationFiltering_UpdateGroupExclusionFiltering);
 
 		const FPerObjectInfo* ConnectionStateInfo = GetPerObjectInfo(this->GroupInfos[GroupIndex].ConnectionStateIndex);
-		const FNetObjectGroup* Group = this->Groups->GetGroupByIndex(FNetObjectGroupHandle::FGroupIndexType(GroupIndex));
+		const FNetObjectGroup* Group = this->Groups->GetGroupFromIndex(FNetObjectGroupHandle::FGroupIndexType(GroupIndex));
 
 		const FNetBitArrayView CurrentFrameScopableObjects = NetRefHandleManager->GetCurrentFrameScopableInternalIndices();
 
@@ -1256,7 +1254,7 @@ void FReplicationFiltering::UpdateGroupInclusionFiltering()
 		IRIS_PROFILER_SCOPE(FReplicationFiltering_UpdateGroupInclusionFiltering);
 
 		const FPerObjectInfo* ConnectionStateInfo = GetPerObjectInfo(this->GroupInfos[GroupIndex].ConnectionStateIndex);
-		const FNetObjectGroup* Group = this->Groups->GetGroupByIndex(FNetObjectGroupHandle::FGroupIndexType(GroupIndex));
+		const FNetObjectGroup* Group = this->Groups->GetGroupFromIndex(FNetObjectGroupHandle::FGroupIndexType(GroupIndex));
 
 		const FNetBitArrayView CurrentFrameScopableObjects = NetRefHandleManager->GetCurrentFrameScopableInternalIndices();
 		const FNetBitArrayView SubObjectInternalIndices = NetRefHandleManager->GetSubObjectInternalIndicesView();
@@ -1849,15 +1847,16 @@ const FReplicationFiltering::FPerObjectInfo* FReplicationFiltering::GetPerObject
 
 void FReplicationFiltering::AddSubObjectFilter(FNetObjectGroupHandle GroupHandle)
 {
-	const FNetObjectGroupHandle::FGroupIndexType GroupIndex = GroupHandle.GetGroupIndex();
-	const bool bIsValidGroup = ensureMsgf(Groups->IsValidGroup(GroupHandle), TEXT("AddGroupFilter received an invalid group handle GroupIndex: %u"), GroupIndex);
+	const bool bIsValidGroup = ensureMsgf(Groups->IsValidGroup(GroupHandle), TEXT("AddSubObjectFilter received invalid group Index: %u Id: %u"), GroupHandle.GetGroupIndex(), GroupHandle.GetUniqueId());
 	if (!bIsValidGroup)
 	{
 		return;
 	}
 
+	const FNetObjectGroupHandle::FGroupIndexType GroupIndex = GroupHandle.GetGroupIndex();
+
 	const bool bIsFiltering = Groups->IsFilterGroup(GroupHandle) || SubObjectFilterGroups.GetBit(GroupIndex);
-	ensureMsgf(!bIsFiltering, TEXT("NetObjectGroup Name '%s' GroupIndex %u was asked to start subobject filtering but it was already used for filtering."), *Groups->GetGroupNameString(GroupHandle), GroupIndex);
+	ensureMsgf(!bIsFiltering, TEXT("NetObjectGroup Name: %s Index: %u Id: %u was asked to start subobject filtering but it was already used for filtering."), *Groups->GetGroupNameString(GroupHandle), GroupIndex, GroupHandle.GetUniqueId());
 	if (bIsFiltering)
 	{
 		return;
@@ -1894,12 +1893,12 @@ void FReplicationFiltering::UpdateSubObjectFilters()
 	// We want to remove all groups that have no members and no enabled connections
 	auto UpdateSubObjectFilterGroup = [this](uint32 GroupIndex)
 	{
-		FNetObjectGroupHandle GroupHandle = Groups->MakeNetObjectGroupHandle(static_cast<FNetObjectGroupHandle::FGroupIndexType>(GroupIndex));
-		if (const FNetObjectGroup* Group = Groups->GetGroup(GroupHandle))
+		if (const FNetObjectGroup* Group = Groups->GetGroupFromIndex((FNetObjectGroupHandle::FGroupIndexType)GroupIndex))
 		{
 			if (Group->Members.IsEmpty() && !IsAnyConnectionFilterStatusAllowed(*GetPerObjectInfo(GroupInfos[GroupIndex].ConnectionStateIndex)))
 			{
-				ReplicationSystem->DestroyGroup(GroupHandle);
+				UE_LOG(LogIrisFiltering, Verbose, TEXT("UpdateSubObjectFilters is destroying group %s since its empty"), *Group->GroupName.ToString());
+				ReplicationSystem->DestroyGroup(Groups->GetHandleFromGroup(Group));
 			}
 		}
 	};
@@ -2008,15 +2007,16 @@ bool FReplicationFiltering::GetSubObjectFilterStatus(FNetObjectGroupHandle Group
 
 bool FReplicationFiltering::AddExclusionFilterGroup(FNetObjectGroupHandle GroupHandle)
 {
-	const FNetObjectGroupHandle::FGroupIndexType GroupIndex = GroupHandle.GetGroupIndex();
-	const bool bIsValidGroup = ensureMsgf(Groups->IsValidGroup(GroupHandle), TEXT("AddExclusionFilterGroup received an invalid group handle GroupIndex: %u"), GroupIndex);
+	const bool bIsValidGroup = ensureMsgf(Groups->IsValidGroup(GroupHandle), TEXT("AddExclusionFilterGroup received an invalid group: Index: %u Id: %u"), GroupHandle.GetGroupIndex(), GroupHandle.GetUniqueId());
 	if (!bIsValidGroup)
 	{
 		return false;
 	}
 
+	const FNetObjectGroupHandle::FGroupIndexType GroupIndex = GroupHandle.GetGroupIndex();
+
 	const bool bIsFiltering = Groups->IsFilterGroup(GroupHandle) || SubObjectFilterGroups.GetBit(GroupIndex);
-	ensureMsgf(!bIsFiltering, TEXT("NetObjectGroup Name '%s' GroupIndex %u was asked to start exclusion filtering but it was already used for filtering."), *Groups->GetGroupNameString(GroupHandle), GroupIndex);
+	ensureMsgf(!bIsFiltering, TEXT("NetObjectGroup Name: %s Index: %u Id: %u was asked to start exclusion filtering but it was already used for filtering."), *Groups->GetGroupNameString(GroupHandle), GroupIndex, GroupHandle.GetUniqueId());
 	if (bIsFiltering)
 	{
 		return false;
@@ -2038,15 +2038,16 @@ bool FReplicationFiltering::AddExclusionFilterGroup(FNetObjectGroupHandle GroupH
 
 bool FReplicationFiltering::AddInclusionFilterGroup(FNetObjectGroupHandle GroupHandle)
 {
-	const FNetObjectGroupHandle::FGroupIndexType GroupIndex = GroupHandle.GetGroupIndex();
-	const bool bIsValidGroup = ensureMsgf(Groups->IsValidGroup(GroupHandle), TEXT("AddInclusionFilterGroup received an invalid group handle GroupIndex: %u"), GroupIndex);
+	const bool bIsValidGroup = ensureMsgf(Groups->IsValidGroup(GroupHandle), TEXT("AddInclusionFilterGroup received an invalid group: Index: %u Id: %u"), GroupHandle.GetGroupIndex(), GroupHandle.GetUniqueId());
 	if (!bIsValidGroup)
 	{
 		return false;
 	}
 
+	const FNetObjectGroupHandle::FGroupIndexType GroupIndex = GroupHandle.GetGroupIndex();
+
 	const bool bIsFiltering = Groups->IsFilterGroup(GroupHandle) || SubObjectFilterGroups.GetBit(GroupIndex);
-	ensureMsgf(!bIsFiltering, TEXT("NetObjectGroup Name '%s' GroupIndex %u was asked to start exclusion filtering but it was already used for filtering."), *Groups->GetGroupNameString(GroupHandle), GroupIndex);
+	ensureMsgf(!bIsFiltering, TEXT("NetObjectGroup Name: %s Index: %u Id: %u was asked to start exclusion filtering but it was already used for filtering."), *Groups->GetGroupNameString(GroupHandle), GroupIndex, GroupHandle.GetUniqueId());
 	if (bIsFiltering)
 	{
 		return false;
@@ -2224,17 +2225,9 @@ void FReplicationFiltering::SetGroupFilterStatus(FNetObjectGroupHandle GroupHand
 
 bool FReplicationFiltering::IsExcludedByAnyGroup(uint32 ObjectInternalIndex, uint32 ConnectionId) const
 {
-	uint32 GroupMembershipCount = 0U;
-	const FNetObjectGroupHandle* GroupHandles = Groups->GetGroupMemberships(ObjectInternalIndex, GroupMembershipCount);
-
-	if (!GroupHandles)
+	const TArrayView<const FNetObjectGroupHandle::FGroupIndexType> GroupIndexes = Groups->GetGroupIndexesOfNetObject(ObjectInternalIndex);
+	for (const FNetObjectGroupHandle::FGroupIndexType GroupIndex : GroupIndexes)
 	{
-		return false;
-	}
-
-	for (const FNetObjectGroupHandle GroupHandle : MakeArrayView(GroupHandles, GroupMembershipCount))
-	{
-		const FNetObjectGroupHandle::FGroupIndexType GroupIndex = GroupHandle.GetGroupIndex();
 		if (ExclusionFilterGroups.GetBit(GroupIndex))
 		{
 			const FPerObjectInfo* ConnectionState = GetPerObjectInfo(GroupInfos[GroupIndex].ConnectionStateIndex);
@@ -2250,17 +2243,9 @@ bool FReplicationFiltering::IsExcludedByAnyGroup(uint32 ObjectInternalIndex, uin
 
 bool FReplicationFiltering::IsIncludedByAnyGroup(uint32 ObjectInternalIndex, uint32 ConnectionId) const
 {
-	uint32 GroupMembershipCount = 0U;
-	const FNetObjectGroupHandle* GroupHandles = Groups->GetGroupMemberships(ObjectInternalIndex, GroupMembershipCount);
-
-	if (!GroupHandles)
-	{
-		return false;
-	}
-
-	for (const FNetObjectGroupHandle GroupHandle : MakeArrayView(GroupHandles, GroupMembershipCount))
+	const TArrayView<const FNetObjectGroupHandle::FGroupIndexType> GroupIndexes = Groups->GetGroupIndexesOfNetObject(ObjectInternalIndex);
+	for (const FNetObjectGroupHandle::FGroupIndexType GroupIndex : GroupIndexes)
 	{
-		const FNetObjectGroupHandle::FGroupIndexType GroupIndex = GroupHandle.GetGroupIndex();
 		if (InclusionFilterGroups.GetBit(GroupIndex))
 		{
 			const FPerObjectInfo* ConnectionState = GetPerObjectInfo(GroupInfos[GroupIndex].ConnectionStateIndex);
diff --git a/Engine/Source/Runtime/Experimental/Iris/Core/Private/Iris/ReplicationSystem/ReplicationSystem.cpp b/Engine/Source/Runtime/Experimental/Iris/Core/Private/Iris/ReplicationSystem/ReplicationSystem.cpp
--- a/Engine/Source/Runtime/Experimental/Iris/Core/Private/Iris/ReplicationSystem/ReplicationSystem.cpp
+++ b/Engine/Source/Runtime/Experimental/Iris/Core/Private/Iris/ReplicationSystem/ReplicationSystem.cpp
@@ -1426,15 +1426,14 @@ void UReplicationSystem::RemoveFromAllGroups(FNetRefHandle Handle)
 		return;
 	}
 
-	if (const FNetObjectGroupHandle* GroupHandles = Groups.GetGroupMemberships(ObjectInternalIndex, NumGroupMemberShips))
+	// We copy the membership array as it is modified during removal
+	TArray<FNetObjectGroupHandle> CopiedGroupHandles;
+	Groups.GetGroupHandlesOfNetObject(ObjectInternalIndex, CopiedGroupHandles);
+
+	for (FNetObjectGroupHandle GroupHandle : CopiedGroupHandles)
 	{
-		// We copy the membership array as it is modified during removal
-		TArray<FNetObjectGroupHandle> CopiedGroupHandles(MakeArrayView(GroupHandles, NumGroupMemberShips));
-		for (FNetObjectGroupHandle GroupHandle : MakeArrayView(CopiedGroupHandles.GetData(), CopiedGroupHandles.Num()))
-		{
-			Groups.RemoveFromGroup(GroupHandle, ObjectInternalIndex);
-			Filtering.NotifyObjectRemovedFromGroup(GroupHandle, ObjectInternalIndex);
-		}
+		Groups.RemoveFromGroup(GroupHandle, ObjectInternalIndex);
+		Filtering.NotifyObjectRemovedFromGroup(GroupHandle, ObjectInternalIndex);
 	}	
 }
 
diff --git a/Engine/Source/Runtime/Experimental/Iris/Core/Public/Iris/ReplicationSystem/NetObjectGroupHandle.h b/Engine/Source/Runtime/Experimental/Iris/Core/Public/Iris/ReplicationSystem/NetObjectGroupHandle.h
--- a/Engine/Source/Runtime/Experimental/Iris/Core/Public/Iris/ReplicationSystem/NetObjectGroupHandle.h
+++ b/Engine/Source/Runtime/Experimental/Iris/Core/Public/Iris/ReplicationSystem/NetObjectGroupHandle.h
@@ -36,12 +36,16 @@ public:
 
 	FNetObjectGroupHandle() : Value(0u) {}
 
+	// $IRIS TODO: IsValid could be considered a shortcut to Group->IsValidGroup but it's abolutely not the same thing. Rename this to IsInitialized() to remove the confusion.
 	/** Returns true if the handle is valid, note this does not mean that the group is valid */
 	inline bool IsValid() const { return Value != 0u; }
 
-	/** Returns the GroupIndex of the group assosciated with the handle */
+	/** Returns the GroupIndex of the group associated with the handle */
 	inline FGroupIndexType GetGroupIndex() const { return FGroupIndexType(Index); }
 
+	/** Returns the unique id for this group */
+	uint32 GetUniqueId() const { return UniqueId; }
+
 	/** Returns true if the provided GroupIndex is a reserved NetObjectsGroupIndex */
 	static bool IsReservedNetObjectGroupIndex(FGroupIndexType GroupIndex) { return GroupIndex >= NotReplicatedNetObjectGroupIndex && GroupIndex <= NetGroupReplayNetObjectGroupIndex; }
 
@@ -57,7 +61,7 @@ public:
 	/** Special group, SubObjects assigned to this group will replicate if replay netconditions is met  */
 	bool IsNetGroupReplayNetObjectGroup() const { return Index == NetGroupReplayNetObjectGroupIndex; }
 
-	uint32 GetRawValue() const
+	uint64 GetRawValue() const
 	{
 		return Value;
 	}
@@ -65,7 +69,7 @@ public:
 private:
 	friend UE::Net::Private::FNetObjectGroups;
 
-	FNetObjectGroupHandle(FGroupIndexType IndexIn, FGroupIndexType EpochIn)
+	FNetObjectGroupHandle(FGroupIndexType IndexIn, FGroupIndexType EpochIn, uint32 InUniqueId)
 	{ 
 		if (IndexIn == InvalidNetObjectGroupIndex)
 		{
@@ -75,16 +79,18 @@ private:
 		{
 			Index = IndexIn;
 			Epoch = EpochIn;
+			UniqueId = InUniqueId;
 		}
 	}
 
 	union 
 	{
-		uint32 Value;
+		uint64 Value;
 		struct
 		{
 			FGroupIndexType Index;
 			FGroupIndexType Epoch;
+			uint32 UniqueId;
 		};
 	};
 
@@ -92,10 +98,10 @@ private:
 	friend inline bool operator!=(const FNetObjectGroupHandle& Lhs, const FNetObjectGroupHandle& Rhs) { return Lhs.Value != Rhs.Value; }
 };
 
-static_assert(sizeof(FNetObjectGroupHandle) == sizeof(uint32), "Bad packing");
+static_assert(sizeof(FNetObjectGroupHandle) == sizeof(uint64), "FNetObjectGroupHandle must be of size 64bits.");
 
 FORCEINLINE uint32 GetTypeHash(const FNetObjectGroupHandle Handle)
 {
-	return ::GetTypeHash(Handle.GetGroupIndex());
+	return ::GetTypeHash(Handle.GetRawValue());
 }
 }
