Fix FOnDemandIoStore trying to fetch headers for all containers instead of only containers tagged with Install

diff --git a/Engine/Source/Runtime/Experimental/IoStore/OnDemand/Private/OnDemandIoStore.cpp b/Engine/Source/Runtime/Experimental/IoStore/OnDemand/Private/OnDemandIoStore.cpp
--- a/Engine/Source/Runtime/Experimental/IoStore/OnDemand/Private/OnDemandIoStore.cpp
+++ b/Engine/Source/Runtime/Experimental/IoStore/OnDemand/Private/OnDemandIoStore.cpp
@@ -966,7 +966,8 @@ FIoStatus FOnDemandIoStore::TickInstallRequest(FMountRequest& MountRequest)
 	TArray<FSharedOnDemandContainer> MountedContainers = GetMountedContainers();
 	TSet<FSharedOnDemandContainer> AllContainers;
 
-	Algo::Copy(MountedContainers, AllContainers);
+	Algo::CopyIf(MountedContainers, AllContainers, 
+		[](const FSharedOnDemandContainer& Container){ return EnumHasAnyFlags(Container->Flags, EOnDemandContainerFlags::Installed); });
 	Algo::Copy(MountRequest.Containers, AllContainers);
 
 	// Fetch all container headers
