Export CanEditchange so it can be found by the linker if it does not get inlined

diff --git a/Engine/Source/Runtime/Experimental/Dataflow/Engine/Public/Dataflow/DataflowObject.h b/Engine/Source/Runtime/Experimental/Dataflow/Engine/Public/Dataflow/DataflowObject.h
--- a/Engine/Source/Runtime/Experimental/Dataflow/Engine/Public/Dataflow/DataflowObject.h
+++ b/Engine/Source/Runtime/Experimental/Dataflow/Engine/Public/Dataflow/DataflowObject.h
@@ -108,7 +108,7 @@ public:
 	static DATAFLOWENGINE_API void AddReferencedObjects(UObject* InThis, FReferenceCollector& Collector);
 
 #if WITH_EDITOR
-	virtual bool CanEditChange(const FProperty* InProperty) const override;
+	DATAFLOWENGINE_API virtual bool CanEditChange(const FProperty* InProperty) const override;
 	DATAFLOWENGINE_API virtual void PostEditChangeProperty(struct FPropertyChangedEvent& PropertyChangedEvent) override;
 #endif
 	DATAFLOWENGINE_API virtual void PostLoad() override;
