ChaosFlesh : Deformable sovler range access issues.

diff --git a/Engine/Source/Runtime/Experimental/Chaos/Private/Chaos/Deformable/ChaosDeformableSolver.cpp b/Engine/Source/Runtime/Experimental/Chaos/Private/Chaos/Deformable/ChaosDeformableSolver.cpp
--- a/Engine/Source/Runtime/Experimental/Chaos/Private/Chaos/Deformable/ChaosDeformableSolver.cpp
+++ b/Engine/Source/Runtime/Experimental/Chaos/Private/Chaos/Deformable/ChaosDeformableSolver.cpp
@@ -1262,7 +1262,7 @@ namespace Chaos::Softs
 								auto SolverParticleToObjectVertexIndex = [&](int32 SolverParticleIndex) {return SolverParticleIndex - Range[0]; };
 
 								FFleshThreadingProxy::FFleshInputBuffer* FleshInputBuffer = nullptr;
-								if (this->CurrentInputPackage->ObjectMap.Contains(Owner))
+								if (this->CurrentInputPackage && this->CurrentInputPackage->ObjectMap.Contains(Owner))
 								{
 									FleshInputBuffer = this->CurrentInputPackage->ObjectMap[Owner]->As<FFleshThreadingProxy::FFleshInputBuffer>();
 								}
@@ -1323,7 +1323,7 @@ namespace Chaos::Softs
 										}
 									}
 								}
-								if (!bParticleTouched)
+								if (!bParticleTouched && ensure(Vertex.IsValidIndex(Index - Range[0])))
 								{
 									MParticles.SetX(Index, GlobalTransform.TransformPosition(ChaosVert(Vertex[Index - Range[0]])));
 									MParticles.PAndInvM(Index).P = MParticles.GetX(Index);
