Chaos Vehicles (sim type as name);

Adding a way to define sim types by FName instead of enum. This will allow new sim modules to define their type by name and be able to compare to a gameplay tag without having to update an engine enum.

-- Code --
ModuleFactoryRegister:
Removed the CoreMinimal.h include and included what we need.
RegisterFactory() - Overloaded to take a FName as a parameter.
ContainsFactory() - Overloaded to take a FName as a parameter.
GenerateNetData() - Overloaded to take a FName as a parameter.
RegisteredFactoriesByName - TMap<FName, TWeakPtr<IFactoryModule>> property. Map of our factories, keyed by a FName.

SimulationModuleBase:
ISimulationModuleBase::GetSimTypeName() - Return a FName instead of enum.
FModuleNetData::GetTypeName() - Return a FName instead of enum.
FSimOutputData::GetTypeName() - Return a FName instead of enum.

#tests PIE, created new modules using tags instead of the enum to figure out what type they are.

diff --git a/Engine/Source/Runtime/Experimental/ChaosVehicles/ChaosVehiclesCore/Private/SimModule/ModuleFactoryRegister.cpp b/Engine/Source/Runtime/Experimental/ChaosVehicles/ChaosVehiclesCore/Private/SimModule/ModuleFactoryRegister.cpp
--- a/Engine/Source/Runtime/Experimental/ChaosVehicles/ChaosVehiclesCore/Private/SimModule/ModuleFactoryRegister.cpp
+++ b/Engine/Source/Runtime/Experimental/ChaosVehicles/ChaosVehiclesCore/Private/SimModule/ModuleFactoryRegister.cpp
@@ -13,6 +13,14 @@ namespace Chaos
 		}
 	}
 
+	void FModuleFactoryRegister::RegisterFactory(const FName TypeName, TWeakPtr<IFactoryModule> InFactory)
+	{
+		if (!ContainsFactory(TypeName))
+		{
+			RegisteredFactoriesByName.Add(TypeName, InFactory);
+		}
+	}
+
 	void FModuleFactoryRegister::RemoveFactory(TWeakPtr<IFactoryModule> InFactory)
 	{
 		for (auto& It : RegisteredFactories)
@@ -23,11 +31,21 @@ namespace Chaos
 				return;
 			}
 		}
+
+		for (TPair<FName, TWeakPtr<IFactoryModule>> Pair : RegisteredFactoriesByName)
+		{
+			if (Pair.Value == InFactory)
+			{
+				RegisteredFactoriesByName.Remove(Pair.Key);
+				return;
+			}
+		}
 	}
 
 	void FModuleFactoryRegister::Reset()
 	{
 		RegisteredFactories.Reset();
+		RegisteredFactoriesByName.Reset();
 	}
 
 	bool FModuleFactoryRegister::ContainsFactory(int32 TypeID)
@@ -35,6 +53,11 @@ namespace Chaos
 		return RegisteredFactories.Contains(TypeID);
 	}
 
+	bool FModuleFactoryRegister::ContainsFactory(const FName TypeName)
+	{
+		return RegisteredFactoriesByName.Contains(TypeName);
+	}
+
 	TSharedPtr<Chaos::FModuleNetData> FModuleFactoryRegister::GenerateNetData(int32 TypeID, int32 SimArrayIndex)
 	{
 		using namespace Chaos;
@@ -56,4 +79,24 @@ namespace Chaos
 		return nullptr;
 	}
 
+	TSharedPtr<Chaos::FModuleNetData> FModuleFactoryRegister::GenerateNetData(const FName TypeName, const int32 SimArrayIndex)
+	{
+		using namespace Chaos;
+
+		if (RegisteredFactoriesByName.Contains(TypeName))
+		{
+			TSharedPtr<IFactoryModule> PinnedFactory = RegisteredFactoriesByName[TypeName].Pin();
+
+			if (PinnedFactory.IsValid())
+			{
+				return PinnedFactory->GenerateNetData(SimArrayIndex);
+			}
+		}
+		else
+		{
+			UE_LOG(LogTemp, Error, TEXT("No factory registered for name '%s'"), *TypeName.ToString());
+		}
+
+		return nullptr;
+	}
 }
\ No newline at end of file
diff --git a/Engine/Source/Runtime/Experimental/ChaosVehicles/ChaosVehiclesCore/Public/SimModule/ModuleFactoryRegister.h b/Engine/Source/Runtime/Experimental/ChaosVehicles/ChaosVehiclesCore/Public/SimModule/ModuleFactoryRegister.h
--- a/Engine/Source/Runtime/Experimental/ChaosVehicles/ChaosVehiclesCore/Public/SimModule/ModuleFactoryRegister.h
+++ b/Engine/Source/Runtime/Experimental/ChaosVehicles/ChaosVehiclesCore/Public/SimModule/ModuleFactoryRegister.h
@@ -2,34 +2,38 @@
 
 #pragma once
 
-#include "CoreMinimal.h"
-
+#include "CoreTypes.h"
+#include "Templates/SharedPointer.h"
 
 namespace Chaos
 {
 
-struct FModuleNetData;
-class IFactoryModule;
+	struct FModuleNetData;
+	class IFactoryModule;
 
-class CHAOSVEHICLESCORE_API FModuleFactoryRegister
-{
-public:
-	static FModuleFactoryRegister& Get()
+	class CHAOSVEHICLESCORE_API FModuleFactoryRegister
 	{
-		static FModuleFactoryRegister Instance;
-		return Instance;
-	}
-
-	void RegisterFactory(int32 TypeID, TWeakPtr<IFactoryModule> InFactory);
-	void RemoveFactory(TWeakPtr<IFactoryModule> InFactory);
-	void Reset();
-	bool ContainsFactory(int32 TypeID);
-	TSharedPtr<Chaos::FModuleNetData> GenerateNetData(int32 TypeID, int32 SimArrayIndex);
-
-protected:
-
-	FModuleFactoryRegister() = default;
-	TMap<int32, TWeakPtr<IFactoryModule>> RegisteredFactories;
-};
+	public:
+		static FModuleFactoryRegister& Get()
+		{
+			static FModuleFactoryRegister Instance;
+			return Instance;
+		}
+
+		void RegisterFactory(int32 TypeID, TWeakPtr<IFactoryModule> InFactory);
+		void RegisterFactory(const FName TypeName, TWeakPtr<IFactoryModule> InFactory);
+		void RemoveFactory(TWeakPtr<IFactoryModule> InFactory);
+		void Reset();
+		bool ContainsFactory(int32 TypeID);
+		bool ContainsFactory(const FName TypeName);
+		TSharedPtr<Chaos::FModuleNetData> GenerateNetData(int32 TypeID, int32 SimArrayIndex);
+		TSharedPtr<Chaos::FModuleNetData> GenerateNetData(const FName TypeName, const int32 SimArrayIndex);
+
+	protected:
+
+		FModuleFactoryRegister() = default;
+		TMap<int32, TWeakPtr<IFactoryModule>> RegisteredFactories;
+		TMap<FName, TWeakPtr<IFactoryModule>> RegisteredFactoriesByName;
+	};
 
 } // namespace Chaos
diff --git a/Engine/Source/Runtime/Experimental/ChaosVehicles/ChaosVehiclesCore/Public/SimModule/SimulationModuleBase.h b/Engine/Source/Runtime/Experimental/ChaosVehicles/ChaosVehiclesCore/Public/SimModule/SimulationModuleBase.h
--- a/Engine/Source/Runtime/Experimental/ChaosVehicles/ChaosVehiclesCore/Public/SimModule/SimulationModuleBase.h
+++ b/Engine/Source/Runtime/Experimental/ChaosVehicles/ChaosVehiclesCore/Public/SimModule/SimulationModuleBase.h
@@ -2,10 +2,9 @@
 
 #pragma once
 
-#include "CoreMinimal.h"
-#include "Logging/LogMacros.h"
-#include "Chaos/ParticleHandleFwd.h"
 #include "Chaos/GeometryParticlesfwd.h"
+#include "Chaos/ParticleHandleFwd.h"
+#include "Logging/LogMacros.h"
 #include "SimModule/ModuleFactoryRegister.h"
 #include "SimModule/ModuleInput.h"
 
@@ -175,7 +174,12 @@ namespace Chaos
 		virtual eSimType GetSimType() const = 0;
 
 		/**
-		* Is Module active and simulating
+		 * The specific simulation type
+		 */
+		virtual FName GetSimTypeName() const { return NAME_None; }
+
+		/**
+		 * Is Module active and simulating
 		 */
 		virtual bool IsEnabled() const { return (StateFlags == eSimModuleState::Enabled); }
 
@@ -376,6 +380,7 @@ namespace Chaos
 		virtual ~FModuleNetData() {}
 
 		virtual eSimType GetType() = 0;
+		virtual FName GetTypeName() const { return NAME_None; };
 		virtual void Serialize(FArchive& Ar) = 0;
 		virtual void FillNetState(const ISimulationModuleBase* SimModule) = 0;
 		virtual void FillSimState(ISimulationModuleBase* SimModule) = 0;
@@ -404,6 +409,7 @@ namespace Chaos
 		virtual ~FSimOutputData() {}
 
 		virtual eSimType GetType() = 0;
+		virtual FName GetTypeName() const { return NAME_None; };
 		virtual bool IsEnabled() { return bEnabled; }
 		virtual FSimOutputData* MakeNewData() = 0;
 		virtual void FillOutputState(const ISimulationModuleBase* SimModule);
