Updating UVehicleInputProducerBase to inherit from UObject instead of UClass so that NewObject works with this class

diff --git a/Engine/Source/Runtime/Experimental/ChaosVehicles/ChaosVehiclesCore/Public/SimModule/ModuleInput.h b/Engine/Source/Runtime/Experimental/ChaosVehicles/ChaosVehiclesCore/Public/SimModule/ModuleInput.h
--- a/Engine/Source/Runtime/Experimental/ChaosVehicles/ChaosVehiclesCore/Public/SimModule/ModuleInput.h
+++ b/Engine/Source/Runtime/Experimental/ChaosVehicles/ChaosVehiclesCore/Public/SimModule/ModuleInput.h
@@ -423,7 +423,7 @@ public:
 };
 
 UCLASS(Abstract, BlueprintType, Blueprintable, EditInlineNew)
-class CHAOSVEHICLESCORE_API UVehicleInputProducerBase : public UClass
+class CHAOSVEHICLESCORE_API UVehicleInputProducerBase : public UObject
 {
 	GENERATED_BODY()
 	
