Dataflow: Add Select (switch/branch) node.

diff --git a/Engine/Source/Runtime/Experimental/Dataflow/Core/Private/Dataflow/DataflowCoreNodes.cpp b/Engine/Source/Runtime/Experimental/Dataflow/Core/Private/Dataflow/DataflowCoreNodes.cpp
--- a/Engine/Source/Runtime/Experimental/Dataflow/Core/Private/Dataflow/DataflowCoreNodes.cpp
+++ b/Engine/Source/Runtime/Experimental/Dataflow/Core/Private/Dataflow/DataflowCoreNodes.cpp
@@ -10,6 +10,7 @@ namespace Dataflow
 	{
 		DATAFLOW_NODE_REGISTER_CREATION_FACTORY(FDataflowReRouteNode);
 		DATAFLOW_NODE_REGISTER_CREATION_FACTORY(FDataflowBranchNode);
+		DATAFLOW_NODE_REGISTER_CREATION_FACTORY(FDataflowSelectNode);
 		DATAFLOW_NODE_REGISTER_CREATION_FACTORY(FDataflowPrintNode);
 	}
 }
@@ -85,6 +86,117 @@ bool FDataflowBranchNode::OnOutputTypeChanged(const FDataflowOutput* Input)
 		);
 }
 
+FDataflowSelectNode::FDataflowSelectNode(const Dataflow::FNodeParameters& Param, FGuid InGuid)
+	: Super(Param, InGuid)
+{
+	// Add two sets of pins to start.
+	AddPins();
+	AddPins();
+	RegisterInputConnection(&SelectedIndex);
+	RegisterOutputConnection(&Result)
+		.SetPassthroughInput(GetConnectionReference(0));
+}
+
+void FDataflowSelectNode::Evaluate(Dataflow::FContext& Context, const FDataflowOutput* Out) const
+{
+	if (Out->IsA(&Result))
+	{
+		const int32 InSelectedIndex = GetValue<int32>(Context, &SelectedIndex);
+		if (Inputs.IsValidIndex(InSelectedIndex))
+		{
+			const Dataflow::TConnectionReference<FDataflowAnyType> SelectedInputReference = GetConnectionReference(InSelectedIndex);
+			if (IsConnected(SelectedInputReference))
+			{
+				ForwardInput(Context, SelectedInputReference, &Result);
+			}
+			else
+			{
+				// TODO : throw an invalid type error when context error is available 
+				// Context.Error(TEXT("Both True and False Inputs must be connected"));
+			}
+		}
+	}
+}
+
+bool FDataflowSelectNode::OnInputTypeChanged(const FDataflowInput* Input)
+{
+	bool bResult = SetOutputConcreteType(&Result, Input->GetType());
+	for (int32 Index = 0; Index < Inputs.Num(); ++Index)
+	{
+		bResult |= SetInputConcreteType(GetConnectionReference(Index), Input->GetType());
+	}
+	return bResult;
+}
+
+bool FDataflowSelectNode::OnOutputTypeChanged(const FDataflowOutput* Input)
+{
+	bool bResult = false;
+	for (int32 Index = 0; Index < Inputs.Num(); ++Index)
+	{
+		bResult |= SetInputConcreteType(GetConnectionReference(Index), Input->GetType());
+	}
+	return bResult;
+}
+
+TArray<Dataflow::FPin> FDataflowSelectNode::AddPins()
+{
+	const int32 Index = Inputs.AddDefaulted();
+	const FDataflowInput& Input = RegisterInputArrayConnection(GetConnectionReference(Index));
+	if (Index > 0)
+	{
+		// Set concrete type the same as Index0.
+		const FDataflowInput* const Input0 = FindInput(GetConnectionReference(0));
+		check(Input0);
+		SetInputConcreteType(GetConnectionReference(Index), Input0->GetType());
+	}
+	return { { Dataflow::FPin::EDirection::INPUT, Input.GetType(), Input.GetName() } };
+}
+
+TArray<Dataflow::FPin> FDataflowSelectNode::GetPinsToRemove() const
+{
+	const int32 Index = Inputs.Num() - 1;
+	check(Inputs.IsValidIndex(Index));
+	if (const FDataflowInput* const Input = FindInput(GetConnectionReference(Index)))
+	{
+		return { { Dataflow::FPin::EDirection::INPUT, Input->GetType(), Input->GetName() } };
+	}
+	return Super::GetPinsToRemove();
+}
+
+void FDataflowSelectNode::OnPinRemoved(const Dataflow::FPin& Pin)
+{
+	const int32 Index = Inputs.Num() - 1;
+	check(Inputs.IsValidIndex(Index));
+#if DO_CHECK
+	const FDataflowInput* const Input = FindInput(GetConnectionReference(Index));
+	check(Input);
+	check(Input->GetName() == Pin.Name);
+	check(Input->GetType() == Pin.Type);
+#endif
+	Inputs.SetNum(Index);
+
+	return Super::OnPinRemoved(Pin);
+}
+
+void FDataflowSelectNode::Serialize(FArchive& Ar)
+{
+	if (Ar.IsLoading())
+	{
+		check(Inputs.Num() > 1);
+		check(FindInput(GetConnectionReference(0)));
+		check(FindInput(GetConnectionReference(1)));
+
+		for (int32 Index = 2; Index < Inputs.Num(); ++Index)
+		{
+			RegisterInputArrayConnection(GetConnectionReference(Index));
+		}
+	}
+}
+
+Dataflow::TConnectionReference<FDataflowAnyType> FDataflowSelectNode::GetConnectionReference(int32 Index) const
+{
+	return { &Inputs[Index], Index, &Inputs };
+}
 
 FDataflowPrintNode::FDataflowPrintNode(const Dataflow::FNodeParameters& InParam, FGuid InGuid)
 	: FDataflowNode(InParam, InGuid)
diff --git a/Engine/Source/Runtime/Experimental/Dataflow/Core/Public/Dataflow/DataflowCoreNodes.h b/Engine/Source/Runtime/Experimental/Dataflow/Core/Public/Dataflow/DataflowCoreNodes.h
--- a/Engine/Source/Runtime/Experimental/Dataflow/Core/Public/Dataflow/DataflowCoreNodes.h
+++ b/Engine/Source/Runtime/Experimental/Dataflow/Core/Public/Dataflow/DataflowCoreNodes.h
@@ -63,6 +63,41 @@ private:
 	virtual bool OnOutputTypeChanged(const FDataflowOutput* Input) override;
 };
 
+
+USTRUCT(meta = (Icon = "GraphEditor.Branch_16x"))
+struct FDataflowSelectNode : public FDataflowNode
+{
+	GENERATED_USTRUCT_BODY()
+
+	DATAFLOW_NODE_DEFINE_INTERNAL(FDataflowSelectNode, "Select", "FlowControl", "")
+
+public:
+	FDataflowSelectNode(const Dataflow::FNodeParameters& Param, FGuid InGuid = FGuid::NewGuid());
+	virtual void Evaluate(Dataflow::FContext& Context, const FDataflowOutput* Out) const;
+
+public:
+	UPROPERTY()
+	TArray<FDataflowAnyType> Inputs;
+
+	UPROPERTY(EditAnywhere, Category = "Condition", meta = (DataflowInput))
+	int32 SelectedIndex = 0;
+
+	UPROPERTY(meta = (DataflowOutput, DataflowPassthrough = "Inputs[0]"))
+	FDataflowAnyType Result;
+
+private:
+	virtual bool OnInputTypeChanged(const FDataflowInput* Input) override;
+	virtual bool OnOutputTypeChanged(const FDataflowOutput* Input) override;
+	virtual TArray<Dataflow::FPin> AddPins() override;
+	virtual bool CanAddPin() const override { return true; }
+	virtual bool CanRemovePin() const override { return Inputs.Num() > 2; }
+	virtual TArray<Dataflow::FPin> GetPinsToRemove() const override;
+	virtual void OnPinRemoved(const Dataflow::FPin& Pin) override;
+	virtual void Serialize(FArchive& Ar) override;
+
+	Dataflow::TConnectionReference<FDataflowAnyType> GetConnectionReference(int32 Index) const;
+};
+
 /** 
 * Print value in the log
 * Supports any type comnvertible to a string 
