Chaos - debug draw mass and transforms now obey ShowStatic/Kinematic/QueryOnly options

diff --git a/Engine/Source/Runtime/Experimental/Chaos/Private/Chaos/ChaosDebugDraw.cpp b/Engine/Source/Runtime/Experimental/Chaos/Private/Chaos/ChaosDebugDraw.cpp
--- a/Engine/Source/Runtime/Experimental/Chaos/Private/Chaos/ChaosDebugDraw.cpp
+++ b/Engine/Source/Runtime/Experimental/Chaos/Private/Chaos/ChaosDebugDraw.cpp
@@ -38,6 +38,9 @@ namespace Chaos
 	namespace CVars
 	{
 		extern int32 ChaosOneWayInteractionPairCollisionMode;
+		extern int32 ChaosSolverDrawShapesShowStatic;
+		extern int32 ChaosSolverDrawShapesShowKinematic;
+		extern int32 ChaosSolverDrawShapesShowDynamic;
 	}
 
 	namespace DebugDraw
@@ -445,6 +448,55 @@ namespace Chaos
 			return true;
 		}
 
+		bool ShouldIgnoreParticle(const FGeometryParticleHandle* InParticle)
+		{
+			if (!CVars::ChaosSolverDrawShapesShowStatic)
+			{
+				if (FConstGenericParticleHandle(InParticle)->IsStatic())
+				{
+					return true;
+				}
+			}
+			if (!CVars::ChaosSolverDrawShapesShowKinematic)
+			{
+				if (FConstGenericParticleHandle(InParticle)->IsKinematic())
+				{
+					return true;
+				}
+			}
+			if (!CVars::ChaosSolverDrawShapesShowDynamic)
+			{
+				if (FConstGenericParticleHandle(InParticle)->IsDynamic())
+				{
+					return true;
+				}
+			}
+
+			if (!bChaosDebugDebugDrawShowQueryOnlyShapes)
+			{
+				if (IsParticleQueryOnly(InParticle))
+				{
+					return true;
+				}
+			}
+			if (!bChaosDebugDebugDrawShowSimOnlyShapes)
+			{
+				if (IsParticleSimOnly(InParticle))
+				{
+					return true;
+				}
+			}
+			if (!bChaosDebugDebugDrawShowProbeOnlyShapes)
+			{
+				if (IsParticleProbeOnly(InParticle))
+				{
+					return true;
+				}
+			}
+
+			return false;
+		}
+
 		//
 		//
 		//
@@ -990,6 +1042,11 @@ namespace Chaos
 
 		void DrawBVHImpl(const FGeometryParticleHandle* Particle, const FRigidTransform3& ShapeTransform, const Private::FImplicitBVH* BVH, const FColor& UnusedColor, const FRealSingle Duration, const FChaosDebugDrawSettings& Settings)
 		{
+			if (ShouldIgnoreParticle(Particle))
+			{
+				return;
+			}
+
 			const FReal BoundsDrawShrink = 1.0f;	// 1cm
 			int32 LeafIndex = 0;
 
@@ -1070,6 +1127,11 @@ namespace Chaos
 
 		void DrawParticleBoundsImpl(const FRigidTransform3& SpaceTransform, const FGeometryParticleHandle* InParticle, const FReal Dt, const FChaosDebugDrawSettings& Settings)
 		{
+			if (ShouldIgnoreParticle(InParticle))
+			{
+				return;
+			}
+
 			FConstGenericParticleHandle Particle = InParticle;
 
 			const FAABB3 Box = InParticle->WorldSpaceInflatedBounds();
@@ -1105,26 +1167,9 @@ namespace Chaos
 
 		void DrawParticleTransformImpl(const FRigidTransform3& SpaceTransform, const FGeometryParticleHandle* InParticle, int32 Index, FRealSingle ColorScale, const FChaosDebugDrawSettings& Settings)
 		{
-			if (!bChaosDebugDebugDrawShowQueryOnlyShapes)
-			{
-				if (IsParticleQueryOnly(InParticle))
-				{
-					return;
-				}
-			}
-			if (!bChaosDebugDebugDrawShowSimOnlyShapes)
+			if (ShouldIgnoreParticle(InParticle))
 			{
-				if (IsParticleSimOnly(InParticle))
-				{
-					return;
-				}
-			}
-			if (!bChaosDebugDebugDrawShowProbeOnlyShapes)
-			{
-				if (IsParticleProbeOnly(InParticle))
-				{
-					return;
-				}
+				return;
 			}
 
 			const TPBDRigidParticleHandle<FReal, 3>* Rigid = InParticle->CastToRigidParticle();
@@ -1528,16 +1573,13 @@ namespace Chaos
 
 		void DrawParticleMassImpl(const FRigidTransform3& SpaceTransform, const FGeometryParticleHandle* InParticle, const FChaosDebugDrawSettings& Settings)
 		{
-			if (const auto RigidParticle = InParticle->CastToRigidParticle())
+			if (ShouldIgnoreParticle(InParticle))
 			{
-				if (!bChaosDebugDebugDrawShowQueryOnlyShapes)
-				{
-					if (IsParticleQueryOnly(InParticle))
-					{
-						return;
-					}
-				}
+				return;
+			}
 
+			if (const auto RigidParticle = InParticle->CastToRigidParticle())
+			{
 				const FColor Color = RigidParticle->IsDynamic()? FColor::Yellow: FColor::Cyan;
 
 				const FVec3 PCOM = SpaceTransform.TransformPosition(FParticleUtilities::GetCoMWorldPosition(RigidParticle));
