Chaos SoftsEvolution: add ability to add and remove soft bodies.

I don't have an easy test case for cloth, but this is following the same pattern as for collision bodies, which has been tested, so barring a typo, it should work.

diff --git a/Engine/Source/Runtime/Experimental/Chaos/Private/Chaos/SoftsEvolution.cpp b/Engine/Source/Runtime/Experimental/Chaos/Private/Chaos/SoftsEvolution.cpp
--- a/Engine/Source/Runtime/Experimental/Chaos/Private/Chaos/SoftsEvolution.cpp
+++ b/Engine/Source/Runtime/Experimental/Chaos/Private/Chaos/SoftsEvolution.cpp
@@ -185,13 +185,6 @@ void FEvolution::Reset()
 
 	Groups.Reset();
 
-	ParallelInitRules.Reset();
-	ConstraintRules.Reset();
-	PBDConstraintRules.Reset();
-	PBDCollisionConstraintRules.Reset();
-	UpdateLinearSystemRules.Reset();
-	UpdateLinearSystemCollisionsRules.Reset();
-
 	KinematicUpdate.Reset();
 	CollisionKinematicUpdate.Reset();
 }
@@ -215,21 +208,59 @@ int32 FEvolution::AddSoftBody(uint32 GroupId, int32 NumParticles, bool bEnable)
 	{
 		Groups.AddGroupsToSize(GroupId + 1);
 	}
-	const int32 SoftBodyId = SoftBodies.AddSoftBody();
-	Groups.SoftBodies[GroupId].Emplace(SoftBodyId);
 
-	SoftBodies.ParticleRanges[SoftBodyId] = FSolverParticlesRange::AddParticleRange(Particles, NumParticles);
+	int32 SoftBodyId = INDEX_NONE;
+	TArray<int32>* const FreeList = SoftBodyFreeList.Find(NumParticles);
+	if (FreeList && FreeList->Num())
+	{
+		SoftBodyId = FreeList->Pop();
+		check(SoftBodies.Status[SoftBodyId] == FSoftBodies::EStatus::Free);
+	}
+	else
+	{
+		SoftBodyId = SoftBodies.AddSoftBody();
+		SoftBodies.ParticleRanges[SoftBodyId] = FSolverParticlesRange::AddParticleRange(Particles, NumParticles);
+	}
+
+	Groups.SoftBodies[GroupId].Emplace(SoftBodyId);
 
 	// Set properties
-	SoftBodies.Active[SoftBodyId] = false;
+	SoftBodies.Status[SoftBodyId] = FSoftBodies::EStatus::Inactive;
 	SoftBodies.GroupId[SoftBodyId] = GroupId;
 	SoftBodies.GlobalDampings[SoftBodyId] = Private::EvolutionSoftBodyDefault::GlobalDamping;
 	SoftBodies.LocalDampings[SoftBodyId] = Private::EvolutionSoftBodyDefault::LocalDamping;
 	SoftBodies.UsePerParticleDamping[SoftBodyId] = Private::EvolutionSoftBodyDefault::bUsePerParticleDamping;
+	SoftBodies.PreSubstepParallelInits[SoftBodyId].Reset();
+	SoftBodies.PBDExternalForceRules[SoftBodyId].Reset();
+	SoftBodies.PostInitialGuessParallelInits[SoftBodyId].Reset();
+	SoftBodies.PreSubstepConstraintRules[SoftBodyId].Reset();
+	SoftBodies.PerIterationPBDConstraintRules[SoftBodyId].Reset();
+	SoftBodies.PerIterationCollisionPBDConstraintRules[SoftBodyId].Reset();
+	SoftBodies.PerIterationPostCollisionsPBDConstraintRules[SoftBodyId].Reset();
+	SoftBodies.UpdateLinearSystemRules[SoftBodyId].Reset();
+	SoftBodies.UpdateLinearSystemCollisionsRules[SoftBodyId].Reset();
+	SoftBodies.PostSubstepConstraintRules[SoftBodyId].Reset();
 	ActivateSoftBody(SoftBodyId, bEnable);
 	return SoftBodyId;
 }
 
+void FEvolution::RemoveSoftBody(int32 SoftBodyId)
+{
+	check(SoftBodies.Status[SoftBodyId] == FSoftBodies::EStatus::Active || SoftBodies.Status[SoftBodyId] == FSoftBodies::EStatus::Inactive);
+
+	// Remove from active group if it was active.
+	Groups.ActiveSoftBodies[SoftBodies.GroupId[SoftBodyId]].Remove(SoftBodyId);
+
+	// Mark free
+	SoftBodies.Status[SoftBodyId] = FSoftBodies::EStatus::Free;
+
+	// Add to free list
+	const int32 NumParticles = SoftBodies.ParticleRanges[SoftBodyId].GetRangeSize();
+	TArray<int32>& FreeList = SoftBodyFreeList.FindOrAdd(NumParticles);
+	FreeList.Add(SoftBodyId);
+}
+
+
 void FEvolution::SetSoftBodyProperties(int32 SoftBodyId, const FCollectionPropertyConstFacade& PropertyCollection,
 	const TMap<FString, TConstArrayView<FRealSingle>>& WeightMaps)
 {
@@ -257,7 +288,8 @@ void FEvolution::SetSoftBodyProperties(int32 SoftBodyId, const FCollectionProper
 
 void FEvolution::ActivateSoftBody(int32 SoftBodyId, bool bActivate)
 {
-	SoftBodies.Active[SoftBodyId] = bActivate;
+	check(SoftBodies.Status[SoftBodyId] != FSoftBodies::EStatus::Free);
+	SoftBodies.Status[SoftBodyId] = bActivate ? FSoftBodies::EStatus::Active : FSoftBodies::EStatus::Inactive;
 	if (bActivate)
 	{
 		Groups.ActiveSoftBodies[SoftBodies.GroupId[SoftBodyId]].Add(SoftBodyId);
@@ -386,9 +418,9 @@ void FEvolution::AdvanceOneTimeStepInternal(const FSolverReal Dt, const int32 Ti
 	// Right now, typically groups only have one active softbody in them, so don't parallelize softbodies
 	for (const int32 SoftBodyId : Groups.ActiveSoftBodies[GroupId])
 	{
-		check(SoftBodies.Active[SoftBodyId]);
+		check(SoftBodies.Status[SoftBodyId] == FSoftBodies::EStatus::Active);
 
-		const bool bDoNewtonUpdate = bEnableForceBasedSolver && NumNewtonIterations > 0 && LinearSystemParameters.MaxNumCGIterations > 0 && SoftBodies.UpdateLinearSystemRules[SoftBodyId].GetRangeSize() > 0;
+		const bool bDoNewtonUpdate = bEnableForceBasedSolver && NumNewtonIterations > 0 && LinearSystemParameters.MaxNumCGIterations > 0 && !SoftBodies.UpdateLinearSystemRules[SoftBodyId].IsEmpty();
 
 		// Do at least one PBD iteration if not doing a newton update.
 		const int32 NumPBDIterations = bDoNewtonUpdate ? TimeDependentNumIterations : FMath::Max(1, TimeDependentNumIterations);
@@ -404,7 +436,7 @@ void FEvolution::AdvanceOneTimeStepInternal(const FSolverReal Dt, const int32 Ti
 		{
 			TRACE_CPUPROFILER_EVENT_SCOPE(FSoftsEvolution_AdvanceOneTimeStepInternal_PreSubstepParallelInits);
 
-			TConstArrayView<ParallelInitFunc> PreSubstepParallelInits = SoftBodies.PreSubstepParallelInits[SoftBodyId].GetConstArrayView();
+			TConstArrayView<ParallelInitFunc> PreSubstepParallelInits(SoftBodies.PreSubstepParallelInits[SoftBodyId]);
 			const int32 NumPreSubstepInits = PreSubstepParallelInits.Num() + (bDoLocalDamping ? 1 : 0);
 			PhysicsParallelFor(NumPreSubstepInits,
 				[this, &PreSubstepParallelInits, &DampVelocityRule, Dt, SoftBodyId, SolverMode](int32 FuncIdx)
@@ -451,7 +483,7 @@ void FEvolution::AdvanceOneTimeStepInternal(const FSolverReal Dt, const int32 Ti
 				else
 				{
 					// ExternalForceRules
-					TConstArrayView<PBDConstraintRuleFunc> ExternalForceRules = SoftBodies.PBDExternalForceRules[SoftBodyId].GetConstArrayView();
+					TConstArrayView<PBDConstraintRuleFunc> ExternalForceRules(SoftBodies.PBDExternalForceRules[SoftBodyId]);
 					for (const PBDConstraintRuleFunc& ExternalForceRule : ExternalForceRules)
 					{
 						ExternalForceRule(SoftBodies.ParticleRanges[SoftBodyId], Dt);
@@ -479,7 +511,7 @@ void FEvolution::AdvanceOneTimeStepInternal(const FSolverReal Dt, const int32 Ti
 		{
 			// Parallel Inits
 			TRACE_CPUPROFILER_EVENT_SCOPE(FSoftsEvolution_AdvanceOneTimeStepInternal_PostInitialGuessParallelInits);
-			TConstArrayView<ParallelInitFunc> ConstraintParallelInits = SoftBodies.PostInitialGuessParallelInits[SoftBodyId].GetConstArrayView();
+			TConstArrayView<ParallelInitFunc> ConstraintParallelInits(SoftBodies.PostInitialGuessParallelInits[SoftBodyId]);
 			PhysicsParallelFor(ConstraintParallelInits.Num(),
 				[this, &ConstraintParallelInits, Dt, SoftBodyId, SolverMode](int32 Index)
 			{
@@ -489,7 +521,7 @@ void FEvolution::AdvanceOneTimeStepInternal(const FSolverReal Dt, const int32 Ti
 		{
 			// PreSubstep rules
 			TRACE_CPUPROFILER_EVENT_SCOPE(FSoftsEvolution_AdvanceOneTimeStepInternal_PreSubstepConstraintRules);
-			TConstArrayView<ConstraintRuleFunc> PreSubstepConstraintRules = SoftBodies.PreSubstepConstraintRules[SoftBodyId].GetConstArrayView();
+			TConstArrayView<ConstraintRuleFunc> PreSubstepConstraintRules(SoftBodies.PreSubstepConstraintRules[SoftBodyId]);
 			for (const ConstraintRuleFunc& ConstraintRule : PreSubstepConstraintRules)
 			{
 				ConstraintRule(SoftBodies.ParticleRanges[SoftBodyId], Dt, SolverMode);
@@ -502,19 +534,19 @@ void FEvolution::AdvanceOneTimeStepInternal(const FSolverReal Dt, const int32 Ti
 				TRACE_CPUPROFILER_EVENT_SCOPE(FSoftsEvolution_AdvanceOneTimeStepInternal_IterationLoop);
 				for (int32 i = 0; i < NumPBDIterations; ++i)
 				{
-					TConstArrayView<PBDConstraintRuleFunc> PerIterationConstraintRules = SoftBodies.PerIterationPBDConstraintRules[SoftBodyId].GetConstArrayView();
+					TConstArrayView<PBDConstraintRuleFunc> PerIterationConstraintRules(SoftBodies.PerIterationPBDConstraintRules[SoftBodyId]);
 					for (const PBDConstraintRuleFunc& ConstraintRule : PerIterationConstraintRules)
 					{
 						ConstraintRule(SoftBodies.ParticleRanges[SoftBodyId], Dt);
 					}
 
-					TConstArrayView<PBDCollisionConstraintRuleFunc> PerIterationCollisionConstraintRules = SoftBodies.PerIterationCollisionPBDConstraintRules[SoftBodyId].GetConstArrayView();
+					TConstArrayView<PBDCollisionConstraintRuleFunc> PerIterationCollisionConstraintRules(SoftBodies.PerIterationCollisionPBDConstraintRules[SoftBodyId]);
 					for (const PBDCollisionConstraintRuleFunc& ConstraintRule : PerIterationCollisionConstraintRules)
 					{
 						ConstraintRule(SoftBodies.ParticleRanges[SoftBodyId], Dt, ActiveCollisionRanges);
 					}
 
-					TConstArrayView<PBDConstraintRuleFunc> PerIterationPostCollisionsConstraintRules = SoftBodies.PerIterationPostCollisionsPBDConstraintRules[SoftBodyId].GetConstArrayView();
+					TConstArrayView<PBDConstraintRuleFunc> PerIterationPostCollisionsConstraintRules(SoftBodies.PerIterationPostCollisionsPBDConstraintRules[SoftBodyId]);
 					for (const PBDConstraintRuleFunc& ConstraintRule : PerIterationPostCollisionsConstraintRules)
 					{
 						ConstraintRule(SoftBodies.ParticleRanges[SoftBodyId], Dt);
@@ -543,9 +575,9 @@ void FEvolution::AdvanceOneTimeStepInternal(const FSolverReal Dt, const int32 Ti
 		// If we did, the linear system would need to live at the group level.
 		if (bDoNewtonUpdate)
 		{
-			TConstArrayView<UpdateLinearSystemFunc> LinearSystemRules = SoftBodies.UpdateLinearSystemRules[SoftBodyId].GetConstArrayView();
+			TConstArrayView<UpdateLinearSystemFunc> LinearSystemRules(SoftBodies.UpdateLinearSystemRules[SoftBodyId]);
 			check(LinearSystemRules.Num());
-			TConstArrayView<UpdateLinearSystemCollisionsFunc> LinearSystemCollisionRules = SoftBodies.UpdateLinearSystemCollisionsRules[SoftBodyId].GetConstArrayView();
+			TConstArrayView<UpdateLinearSystemCollisionsFunc> LinearSystemCollisionRules(SoftBodies.UpdateLinearSystemCollisionsRules[SoftBodyId]);
 			check(NumNewtonIterations > 0);
 			for (int32 NewtonIteration = 0; NewtonIteration < NumNewtonIterations; ++NewtonIteration)
 			{
@@ -574,7 +606,7 @@ void FEvolution::AdvanceOneTimeStepInternal(const FSolverReal Dt, const int32 Ti
 		{
 			// Post Substep Rules
 			TRACE_CPUPROFILER_EVENT_SCOPE(FSoftsEvolution_AdvanceOneTimeStepInternal_PostSubstep);
-			TConstArrayView<ConstraintRuleFunc> PostSubstepConstraintRules = SoftBodies.PostSubstepConstraintRules[SoftBodyId].GetConstArrayView();
+			TConstArrayView<ConstraintRuleFunc> PostSubstepConstraintRules(SoftBodies.PostSubstepConstraintRules[SoftBodyId]);
 			for (const ConstraintRuleFunc& ConstraintRule : PostSubstepConstraintRules)
 			{
 				ConstraintRule(SoftBodies.ParticleRanges[SoftBodyId], Dt, SolverMode);
diff --git a/Engine/Source/Runtime/Experimental/Chaos/Public/Chaos/SoftsEvolution.h b/Engine/Source/Runtime/Experimental/Chaos/Public/Chaos/SoftsEvolution.h
--- a/Engine/Source/Runtime/Experimental/Chaos/Public/Chaos/SoftsEvolution.h
+++ b/Engine/Source/Runtime/Experimental/Chaos/Public/Chaos/SoftsEvolution.h
@@ -63,12 +63,16 @@ public:
 	 * @return SoftBodyId
 	 */
 	CHAOS_API int32 AddSoftBody(uint32 GroupId, int32 NumParticles, bool bEnable);
+	// TODO: add garbage collection. Currently soft bodies and their particles are recycled only if another softbody with the exact same number of particles is requested. 
+	// This is the common use case with cloth collision ranges, but we will likely need something more sophisticated.
+	// (Or we could have each SoftBody own its own SolverParticles--we don't take advantage of a single Particle list anywhere)
+	CHAOS_API void RemoveSoftBody(int32 SoftBodyId);
 	int32 GetSoftBodyParticleNum(int32 SoftBodyId) const { return SoftBodies.ParticleRanges[SoftBodyId].GetRangeSize(); }
 	int32 GetSoftBodyGroupId(int32 SoftBodyId) const { return SoftBodies.GroupId[SoftBodyId]; }
 	CHAOS_API void SetSoftBodyProperties(int32 SoftBodyId, const FCollectionPropertyConstFacade& PropertyCollection,
 		const TMap<FString, TConstArrayView<FRealSingle>>& WeightMaps);
 	CHAOS_API void ActivateSoftBody(int32 SoftBodyId, bool bActivate);
-	bool IsSoftBodyActive(int32 SoftBodyId) const { return SoftBodies.Active[SoftBodyId]; }
+	bool IsSoftBodyActive(int32 SoftBodyId) const { return SoftBodies.Status[SoftBodyId] == FSoftBodies::EStatus::Active; }
 	FSolverParticlesRange& GetSoftBodyParticles(int32 SoftBodyId) { return SoftBodies.ParticleRanges[SoftBodyId]; }
 	const FSolverParticlesRange& GetSoftBodyParticles(int32 SoftBodyId) const { return SoftBodies.ParticleRanges[SoftBodyId]; }
 	const TArray<int32>& GetGroupSoftBodies(uint32 GroupId) const { return Groups.SoftBodies[GroupId]; }
@@ -106,99 +110,97 @@ public:
 	typedef TFunction<void(const FSolverParticlesRange&, const FSolverReal Dt, const TArray<FSolverCollisionParticlesRange>&, FEvolutionLinearSystem&)> UpdateLinearSystemCollisionsFunc;
 
 	/* Add Ranges to allocate space for your rules. */
-	// Warning: Rules are allocated into shared buffers. The ArrayViews may go stale if you allocate ANY new rules.
-	// Allocate your batch of rules, then get your array views.	
 
 	// Presubstep init methods (always run at beginning of substep)
 	void AllocatePreSubstepParallelInitRange(int32 SoftBodyId, int32 NumRules)
 	{
-		return AllocateRulesRange(SoftBodyId, NumRules, ParallelInitRules, SoftBodies.PreSubstepParallelInits);
+		return AllocateRules(SoftBodyId, NumRules, SoftBodies.PreSubstepParallelInits);
 	}
 	// PBD Rules that apply external forces (only run if doing PBD)
 	void AllocatePBDExternalForceRulesRange(int32 SoftBodyId, int32 NumRules)
 	{
-		return AllocateRulesRange(SoftBodyId, NumRules, PBDConstraintRules, SoftBodies.PBDExternalForceRules);
+		return AllocateRules(SoftBodyId, NumRules, SoftBodies.PBDExternalForceRules);
 	}
 	// Post initial guess init methods (always run after kinematic and initial guess update, before any solving.) 
 	void AllocatePostInitialGuessParallelInitRange(int32 SoftBodyId, int32 NumRules)
 	{
-		return AllocateRulesRange(SoftBodyId, NumRules, ParallelInitRules, SoftBodies.PostInitialGuessParallelInits);
+		return AllocateRules(SoftBodyId, NumRules, SoftBodies.PostInitialGuessParallelInits);
 	}
 	// Rules that run once per substep after all initial guess and initialization is done.
 	void AllocatePreSubstepConstraintRulesRange(int32 SoftBodyId, int32 NumRules)
 	{
-		return AllocateRulesRange(SoftBodyId, NumRules, ConstraintRules, SoftBodies.PreSubstepConstraintRules);
+		return AllocateRules(SoftBodyId, NumRules, SoftBodies.PreSubstepConstraintRules);
 	}
 	// Normal per-iteration PBD rules (only run if doing PBD)
 	void AllocatePerIterationPBDConstraintRulesRange(int32 SoftBodyId, int32 NumRules)
 	{
-		return AllocateRulesRange(SoftBodyId, NumRules, PBDConstraintRules, SoftBodies.PerIterationPBDConstraintRules);
+		return AllocateRules(SoftBodyId, NumRules, SoftBodies.PerIterationPBDConstraintRules);
 	}
 	// Collision per-iteration PBD rules (only run if doing PBD)
 	void AllocatePerIterationCollisionPBDConstraintRulesRange(int32 SoftBodyId, int32 NumRules)
 	{
-		return AllocateRulesRange(SoftBodyId, NumRules, PBDCollisionConstraintRules, SoftBodies.PerIterationCollisionPBDConstraintRules);
+		return AllocateRules(SoftBodyId, NumRules, SoftBodies.PerIterationCollisionPBDConstraintRules);
 	}
 	// Normal per-iteration PBD rules that run after collisions (only run if doing PBD)
 	void AllocatePerIterationPostCollisionsPBDConstraintRulesRange(int32 SoftBodyId, int32 NumRules)
 	{
-		return AllocateRulesRange(SoftBodyId, NumRules, PBDConstraintRules, SoftBodies.PerIterationPostCollisionsPBDConstraintRules);
+		return AllocateRules(SoftBodyId, NumRules, SoftBodies.PerIterationPostCollisionsPBDConstraintRules);
 	}
 	// Linear system rules (only run if doing ForceBased)
 	void AllocateUpdateLinearSystemRulesRange(int32 SoftBodyId, int32 NumRules)
 	{
-		return AllocateRulesRange(SoftBodyId, NumRules, UpdateLinearSystemRules, SoftBodies.UpdateLinearSystemRules);
+		return AllocateRules(SoftBodyId, NumRules, SoftBodies.UpdateLinearSystemRules);
 	}
 	// Linear system collision rules (only run if doing ForceBased)
 	void AllocateUpdateLinearSystemCollisionsRulesRange(int32 SoftBodyId, int32 NumRules)
 	{
-		return AllocateRulesRange(SoftBodyId, NumRules, UpdateLinearSystemCollisionsRules, SoftBodies.UpdateLinearSystemCollisionsRules);
+		return AllocateRules(SoftBodyId, NumRules, SoftBodies.UpdateLinearSystemCollisionsRules);
 	}
 	// Post substep rules (always run at end of substep)
 	void AllocatePostSubstepConstraintRulesRange(int32 SoftBodyId, int32 NumRules)
 	{
-		return AllocateRulesRange(SoftBodyId, NumRules, ConstraintRules, SoftBodies.PostSubstepConstraintRules);
+		return AllocateRules(SoftBodyId, NumRules, SoftBodies.PostSubstepConstraintRules);
 	}
 
 	TArrayView<ParallelInitFunc> GetPreSubstepParallelInitRange(int32 SoftBodyId)
 	{
-		return GetRulesRange(SoftBodyId, SoftBodies.PreSubstepParallelInits);
+		return GetRules(SoftBodyId, SoftBodies.PreSubstepParallelInits);
 	}
 	TArrayView<PBDConstraintRuleFunc> GetPBDExternalForceRulesRange(int32 SoftBodyId)
 	{
-		return GetRulesRange(SoftBodyId, SoftBodies.PBDExternalForceRules);
+		return GetRules(SoftBodyId, SoftBodies.PBDExternalForceRules);
 	}
 	TArrayView<ParallelInitFunc> GetPostInitialGuessParallelInitRange(int32 SoftBodyId)
 	{
-		return GetRulesRange(SoftBodyId, SoftBodies.PostInitialGuessParallelInits);
+		return GetRules(SoftBodyId, SoftBodies.PostInitialGuessParallelInits);
 	}
 	TArrayView<ConstraintRuleFunc> GetPreSubstepConstraintRulesRange(int32 SoftBodyId)
 	{
-		return GetRulesRange(SoftBodyId, SoftBodies.PreSubstepConstraintRules);
+		return GetRules(SoftBodyId, SoftBodies.PreSubstepConstraintRules);
 	}
 	TArrayView<PBDConstraintRuleFunc> GetPerIterationPBDConstraintRulesRange(int32 SoftBodyId)
 	{
-		return GetRulesRange(SoftBodyId, SoftBodies.PerIterationPBDConstraintRules);
+		return GetRules(SoftBodyId, SoftBodies.PerIterationPBDConstraintRules);
 	}
 	TArrayView<PBDCollisionConstraintRuleFunc> GetPerIterationCollisionPBDConstraintRulesRange(int32 SoftBodyId)
 	{
-		return GetRulesRange(SoftBodyId, SoftBodies.PerIterationCollisionPBDConstraintRules);
+		return GetRules(SoftBodyId, SoftBodies.PerIterationCollisionPBDConstraintRules);
 	}
 	TArrayView<PBDConstraintRuleFunc> GetPerIterationPostCollisionsPBDConstraintRulesRange(int32 SoftBodyId)
 	{
-		return GetRulesRange(SoftBodyId, SoftBodies.PerIterationPostCollisionsPBDConstraintRules);
+		return GetRules(SoftBodyId, SoftBodies.PerIterationPostCollisionsPBDConstraintRules);
 	}
 	TArrayView<UpdateLinearSystemFunc> GetUpdateLinearSystemRulesRange(int32 SoftBodyId)
 	{
-		return GetRulesRange(SoftBodyId, SoftBodies.UpdateLinearSystemRules);
+		return GetRules(SoftBodyId, SoftBodies.UpdateLinearSystemRules);
 	}
 	TArrayView<UpdateLinearSystemCollisionsFunc> GetUpdateLinearSystemCollisionsRulesRange(int32 SoftBodyId)
 	{
-		return GetRulesRange(SoftBodyId, SoftBodies.UpdateLinearSystemCollisionsRules);
+		return GetRules(SoftBodyId, SoftBodies.UpdateLinearSystemCollisionsRules);
 	}
 	TArrayView<ConstraintRuleFunc> GetPostSubstepConstraintRulesRange(int32 SoftBodyId)
 	{
-		return GetRulesRange(SoftBodyId, SoftBodies.PostSubstepConstraintRules);
+		return GetRules(SoftBodyId, SoftBodies.PostSubstepConstraintRules);
 	}
 
 	/** Solver settings */
@@ -256,7 +258,7 @@ private:
 	{
 		FSoftBodies()
 		{
-			TArrayCollection::AddArray(&Active);
+			TArrayCollection::AddArray(&Status);
 			TArrayCollection::AddArray(&GroupId);
 			TArrayCollection::AddArray(&ParticleRanges);
 			TArrayCollection::AddArray(&GlobalDampings);
@@ -288,7 +290,14 @@ private:
 			return Offset;
 		}
 
-		TArrayCollectionArray<bool> Active;
+		enum struct EStatus : uint8
+		{
+			Invalid = 0,
+			Active = 1,
+			Inactive = 2,
+			Free = 3 // Available for recycling
+		};
+		TArrayCollectionArray<EStatus> Status;
 		TArrayCollectionArray<uint32> GroupId;
 		TArrayCollectionArray<FSolverParticlesRange> ParticleRanges;
 		TArrayCollectionArray<FSolverReal> GlobalDampings;
@@ -296,16 +305,16 @@ private:
 		TArrayCollectionArray<bool> UsePerParticleDamping;
 		TArrayCollectionArray<FEvolutionLinearSystem> LinearSystems;
 
-		TArrayCollectionArray<TArrayRange<ParallelInitFunc>> PreSubstepParallelInits;
-		TArrayCollectionArray<TArrayRange<PBDConstraintRuleFunc>> PBDExternalForceRules;
-		TArrayCollectionArray<TArrayRange<ParallelInitFunc>> PostInitialGuessParallelInits;
-		TArrayCollectionArray<TArrayRange<ConstraintRuleFunc>> PreSubstepConstraintRules;
-		TArrayCollectionArray<TArrayRange<PBDConstraintRuleFunc>> PerIterationPBDConstraintRules;
-		TArrayCollectionArray<TArrayRange<PBDCollisionConstraintRuleFunc>> PerIterationCollisionPBDConstraintRules;
-		TArrayCollectionArray<TArrayRange<PBDConstraintRuleFunc>> PerIterationPostCollisionsPBDConstraintRules;
-		TArrayCollectionArray<TArrayRange<UpdateLinearSystemFunc>> UpdateLinearSystemRules;
-		TArrayCollectionArray<TArrayRange<UpdateLinearSystemCollisionsFunc>> UpdateLinearSystemCollisionsRules;
-		TArrayCollectionArray<TArrayRange<ConstraintRuleFunc>> PostSubstepConstraintRules;
+		TArrayCollectionArray<TArray<ParallelInitFunc>> PreSubstepParallelInits;
+		TArrayCollectionArray<TArray<PBDConstraintRuleFunc>> PBDExternalForceRules;
+		TArrayCollectionArray<TArray<ParallelInitFunc>> PostInitialGuessParallelInits;
+		TArrayCollectionArray<TArray<ConstraintRuleFunc>> PreSubstepConstraintRules;
+		TArrayCollectionArray<TArray<PBDConstraintRuleFunc>> PerIterationPBDConstraintRules;
+		TArrayCollectionArray<TArray<PBDCollisionConstraintRuleFunc>> PerIterationCollisionPBDConstraintRules;
+		TArrayCollectionArray<TArray<PBDConstraintRuleFunc>> PerIterationPostCollisionsPBDConstraintRules;
+		TArrayCollectionArray<TArray<UpdateLinearSystemFunc>> UpdateLinearSystemRules;
+		TArrayCollectionArray<TArray<UpdateLinearSystemCollisionsFunc>> UpdateLinearSystemCollisionsRules;
+		TArrayCollectionArray<TArray<ConstraintRuleFunc>> PostSubstepConstraintRules;
 	};
 
 	// CollisionBodyRange SOA
@@ -402,16 +411,16 @@ private:
 	};
 
 	template<typename RuleFunc>
-	void AllocateRulesRange(int32 SoftBodyId, int32 NumRules, TArray<RuleFunc>& RuleArray, TArrayCollectionArray<TArrayRange<RuleFunc>>& RangeArray)
+	void AllocateRules(int32 SoftBodyId, int32 NumRules, TArrayCollectionArray<TArray<RuleFunc>>& RuleArray)
 	{
-		check(RangeArray[SoftBodyId].IsEmpty());
-		RangeArray[SoftBodyId] = TArrayRange<RuleFunc>::AddRange(RuleArray, NumRules);
+		check(RuleArray[SoftBodyId].IsEmpty());
+		RuleArray[SoftBodyId].SetNum(NumRules);
 	}
 
 	template<typename RuleFunc>
-	TArrayView<RuleFunc> GetRulesRange(int32 SoftBodyId, TArrayCollectionArray<TArrayRange<RuleFunc>>& RangeArray)
+	TArrayView<RuleFunc> GetRules(int32 SoftBodyId, TArrayCollectionArray<TArray<RuleFunc>>& RuleArray)
 	{
-		return RangeArray[SoftBodyId].GetArrayView();
+		return TArrayView<RuleFunc>(RuleArray[SoftBodyId]);
 	}
 
 	void AdvanceOneTimeStepInternal(const FSolverReal Dt, const int32 TimeDependentNumIterations, uint32 GroupId);
@@ -438,6 +447,9 @@ private:
 	// Per-SoftBody data
 	FSoftBodies SoftBodies;
 
+	// SoftBody free-list
+	TMap<int32, TArray<int32>> SoftBodyFreeList; // Key = NumParticles, Value = SoftBodyId(s)
+
 	// Per-CollisionBodyRange data
 	FCollisionBodyRanges CollisionRanges;
 
@@ -451,16 +463,6 @@ private:
 	KinematicUpdateFunc KinematicUpdate;
 	CollisionKinematicUpdateFunc CollisionKinematicUpdate;
 
-	// Rules that run on ParticleRanges (SoftBodies will have views into these)
-	TArray<ParallelInitFunc> ParallelInitRules;
-	TArray<ConstraintRuleFunc> ConstraintRules;
-	TArray<PBDConstraintRuleFunc> PBDConstraintRules;
-	TArray<PBDCollisionConstraintRuleFunc> PBDCollisionConstraintRules;
-	TArray<UpdateLinearSystemFunc> UpdateLinearSystemRules;
-	TArray<UpdateLinearSystemCollisionsFunc> UpdateLinearSystemCollisionsRules;
-
-	// Collision Rules
-
 	UE_CHAOS_DECLARE_INDEXLESS_PROPERTYCOLLECTION_NAME(DampingCoefficient, float);
 	UE_CHAOS_DECLARE_INDEXLESS_PROPERTYCOLLECTION_NAME(LocalDampingCoefficient, float);
 	UE_CHAOS_DECLARE_INDEXLESS_PROPERTYCOLLECTION_NAME(MaxNumIterations, int32);
