Save of the dataflow scene description

diff --git a/Engine/Plugins/Experimental/ChaosFlesh/Source/ChaosFleshEngine/Private/ChaosFlesh/FleshAsset.cpp b/Engine/Plugins/Experimental/ChaosFlesh/Source/ChaosFleshEngine/Private/ChaosFlesh/FleshAsset.cpp
--- a/Engine/Plugins/Experimental/ChaosFlesh/Source/ChaosFleshEngine/Private/ChaosFlesh/FleshAsset.cpp
+++ b/Engine/Plugins/Experimental/ChaosFlesh/Source/ChaosFleshEngine/Private/ChaosFlesh/FleshAsset.cpp
@@ -80,20 +80,34 @@ TObjectPtr<UDataflowBaseContent> UFleshAsset::CreateDataflowContent()
 	SkeletalContent->SetDataflowOwner(this);
 	SkeletalContent->SetTerminalAsset(this);
 
-	UpdateDataflowContent(SkeletalContent);
+	WriteDataflowContent(SkeletalContent);
 	
 	return SkeletalContent;
 }
 
-void UFleshAsset::UpdateDataflowContent(const TObjectPtr<UDataflowBaseContent>& DataflowContent) const
+void UFleshAsset::WriteDataflowContent(const TObjectPtr<UDataflowBaseContent>& DataflowContent) const
 {
 	if(const TObjectPtr<UDataflowSkeletalContent> SkeletalContent = Cast<UDataflowSkeletalContent>(DataflowContent))
 	{
 		SkeletalContent->SetDataflowAsset(DataflowAsset);
 		SkeletalContent->SetDataflowTerminal(DataflowTerminal);
 		
-		SkeletalContent->SetSkeletalMesh(SkeletalMesh);
+		SkeletalContent->SetSkeletalMesh(SkeletalMesh, true);
 		SkeletalContent->SetSkeleton(Skeleton);
+
+#if WITH_EDITORONLY_DATA
+		SkeletalContent->SetAnimationAsset(PreviewAnimationAsset.Get());
+#endif
+	}
+}
+
+void UFleshAsset::ReadDataflowContent(const TObjectPtr<UDataflowBaseContent>& DataflowContent)
+{
+	if(const TObjectPtr<UDataflowSkeletalContent> SkeletalContent = Cast<UDataflowSkeletalContent>(DataflowContent))
+	{
+#if WITH_EDITORONLY_DATA
+		PreviewAnimationAsset = SkeletalContent->GetAnimationAsset();
+#endif
 	}
 }
 
diff --git a/Engine/Plugins/Experimental/ChaosFlesh/Source/ChaosFleshEngine/Public/ChaosFlesh/FleshAsset.h b/Engine/Plugins/Experimental/ChaosFlesh/Source/ChaosFleshEngine/Public/ChaosFlesh/FleshAsset.h
--- a/Engine/Plugins/Experimental/ChaosFlesh/Source/ChaosFleshEngine/Public/ChaosFlesh/FleshAsset.h
+++ b/Engine/Plugins/Experimental/ChaosFlesh/Source/ChaosFleshEngine/Public/ChaosFlesh/FleshAsset.h
@@ -102,10 +102,12 @@ public:
 
 	void Serialize(FArchive& Ar);
 
-	/** IDataflowContentOwner interface */
+	//~ Begin IDataflowContentOwner interface
 	virtual TObjectPtr<UDataflowBaseContent> CreateDataflowContent() override;
-	virtual void UpdateDataflowContent(const TObjectPtr<UDataflowBaseContent>& DataflowContent) const override;
-
+	virtual void WriteDataflowContent(const TObjectPtr<UDataflowBaseContent>& DataflowContent) const override;
+	virtual void ReadDataflowContent(const TObjectPtr<UDataflowBaseContent>& DataflowContent) override;
+	//~ End IDataflowContentOwner interface
+	
 	//
 	// Dataflow
 	//
@@ -147,6 +149,19 @@ public:
 	/** Information for thumbnail rendering */
 	UPROPERTY()
 	TObjectPtr<class UThumbnailInfo> ThumbnailInfo;
+
+	/*
+	* The following PreviewScene properties are modeled after PreviewSkeletalMesh in USkeleton
+	*	- they are inside WITH_EDITORONLY_DATA because they are not used at game runtime
+	*	- TSoftObjectPtrs since that will make it possible to avoid loading these assets until the PreviewScene asks for them
+	*	- DuplicateTransient so that if you copy a ClothAsset it won't copy these preview properties
+	*	- AssetRegistrySearchable makes it so that if the user searches the name of a PreviewScene asset in the Asset Browser
+	*/
+
+	/** Animation asset used in this asset */
+	UPROPERTY(DuplicateTransient, AssetRegistrySearchable)
+	TSoftObjectPtr<UAnimationAsset> PreviewAnimationAsset = nullptr;
+	
 #endif // WITH_EDITORONLY_DATA
 
 };
diff --git a/Engine/Plugins/Experimental/Dataflow/Source/DataflowEditor/Private/Dataflow/DataflowEditor.cpp b/Engine/Plugins/Experimental/Dataflow/Source/DataflowEditor/Private/Dataflow/DataflowEditor.cpp
--- a/Engine/Plugins/Experimental/Dataflow/Source/DataflowEditor/Private/Dataflow/DataflowEditor.cpp
+++ b/Engine/Plugins/Experimental/Dataflow/Source/DataflowEditor/Private/Dataflow/DataflowEditor.cpp
@@ -82,16 +82,12 @@ void UDataflowEditor::UpdateEditorContent()
 	if(const TObjectPtr<UDataflowSkeletalContent> SkeletalContent = Cast<UDataflowSkeletalContent>(EditorContent))
 	{
 		// Add a skeleton mesh in case we don't have one
-		if(!SkeletalContent->GetSkeleton() && !SkeletalContent->GetSkeletalMesh())
+		if(!SkeletalContent->GetSkeleton())
 		{
-			const FName SkeletalMeshName = MakeUniqueObjectName(SkeletalContent->GetDataflowAsset(), UDataflow::StaticClass(), FName("USkeletalMesh"));
-			USkeletalMesh* SkeletalMesh = NewObject<USkeletalMesh>(SkeletalContent->GetDataflowAsset(), SkeletalMeshName);
-
 			const FName SkeletonName = MakeUniqueObjectName(SkeletalContent->GetDataflowAsset(), UDataflow::StaticClass(), FName("USkeleton"));
 			USkeleton* Skeleton = NewObject<USkeleton>(SkeletalContent->GetDataflowAsset(), SkeletonName);
-			
-			SkeletalMesh->SetSkeleton(Skeleton);
-			SkeletalContent->SetSkeletalMesh(SkeletalMesh);
+
+			SkeletalContent->SetSkeleton(Skeleton);
 		}
 	}
 	if(EditorContent && EditorContent->GetDataflowAsset())
diff --git a/Engine/Plugins/Experimental/Dataflow/Source/DataflowEditor/Private/Dataflow/DataflowEditorToolkit.cpp b/Engine/Plugins/Experimental/Dataflow/Source/DataflowEditor/Private/Dataflow/DataflowEditorToolkit.cpp
--- a/Engine/Plugins/Experimental/Dataflow/Source/DataflowEditor/Private/Dataflow/DataflowEditorToolkit.cpp
+++ b/Engine/Plugins/Experimental/Dataflow/Source/DataflowEditor/Private/Dataflow/DataflowEditorToolkit.cpp
@@ -547,11 +547,14 @@ public:
 FDataflowPreviewSceneDescriptionCustomization::FDataflowPreviewSceneDescriptionCustomization(const TArray<UDataflowBaseContent*>& DataflowContents) :
 	IDetailCustomization(), ContentTypesObjects()
 {
-	static const FString CategoryName = TEXT("Terminals");
+	static const FString PreviewCategory = TEXT("Preview");
+	TArray<UObject*>& PreviewObjects = ContentTypesObjects.FindOrAdd(PreviewCategory);
 	for(UDataflowBaseContent* DataflowContent : DataflowContents)
 	{
-		TArray<UObject*>& ContentTypeObjects = ContentTypesObjects.FindOrAdd(CategoryName);
-		ContentTypeObjects.Add(DataflowContent); 
+		if(DataflowContent)
+		{
+			PreviewObjects.Add(DataflowContent); 
+		}
 	}
 }
 
diff --git a/Engine/Plugins/Experimental/Dataflow/Source/DataflowEditor/Private/Dataflow/DataflowSimulationGenerator.cpp b/Engine/Plugins/Experimental/Dataflow/Source/DataflowEditor/Private/Dataflow/DataflowSimulationGenerator.cpp
--- a/Engine/Plugins/Experimental/Dataflow/Source/DataflowEditor/Private/Dataflow/DataflowSimulationGenerator.cpp
+++ b/Engine/Plugins/Experimental/Dataflow/Source/DataflowEditor/Private/Dataflow/DataflowSimulationGenerator.cpp
@@ -158,11 +158,11 @@ void FDataflowSimulationGenerator::StartGenerateSimulation()
 
 	TaskManager->SimulationTask = MakeUnique<FAsyncTask<FDataflowSimulationTask>>();
 	TaskManager->SimulationTask->GetTask().TaskManager = TaskManager;
-	TaskManager->SimulationTask->GetTask().bBackgroundTask = bBackgroundTask;
+	TaskManager->SimulationTask->GetTask().bBackgroundTask = CacheParams.bBackgroundTask;
 	
-	TaskManager->AllocateSimulationResource(TimeRange, FrameRate, CacheAsset, ActorClass, DataflowContent);
+	TaskManager->AllocateSimulationResource(CacheParams.TimeRange, CacheParams.FrameRate, CacheAsset, BlueprintClass, DataflowContent);
 
-	if(bBackgroundTask)
+	if(CacheParams.bBackgroundTask)
 	{
 		TaskManager->SimulationTask->StartBackgroundTask();
 	}
@@ -222,15 +222,10 @@ void FDataflowSimulationGenerator::TickGenerateSimulation()
 		PendingAction = EDataflowGeneratorActions::NoAction;
 	}
 }
-
-void FDataflowSimulationGenerator::SetFrameRate(const int32 InFrameRate) 
-{
-	FrameRate = InFrameRate;
-}
-
-void FDataflowSimulationGenerator::SetTimeRange(const FVector2f& InTimeRange) 
+	
+void FDataflowSimulationGenerator::SetCacheParams(const FDataflowPreviewCacheParams& InCacheParams)
 {
-	TimeRange = InTimeRange;
+	CacheParams = InCacheParams;
 }
 
 void FDataflowSimulationGenerator::SetCacheAsset(const TObjectPtr<UChaosCacheCollection>& InCacheAsset) 
@@ -238,14 +233,9 @@ void FDataflowSimulationGenerator::SetCacheAsset(const TObjectPtr<UChaosCacheCol
 	CacheAsset = InCacheAsset;
 }
 
-void FDataflowSimulationGenerator::SetActorClass(const TSubclassOf<AActor>& InActorClass)
-{
-	ActorClass = InActorClass;
-}
-
-void FDataflowSimulationGenerator::SetBackgroundTask(const bool bInBackgroundTask)
+void FDataflowSimulationGenerator::SetBlueprintClass(const TSubclassOf<AActor>& InBlueprintClass)
 {
-	bBackgroundTask = bInBackgroundTask;
+	BlueprintClass = InBlueprintClass;
 }
 
 void FDataflowSimulationGenerator::SetDataflowContent(const TObjectPtr<UDataflowBaseContent>& InDataflowContent)
diff --git a/Engine/Plugins/Experimental/Dataflow/Source/DataflowEditor/Private/Dataflow/DataflowSimulationScene.cpp b/Engine/Plugins/Experimental/Dataflow/Source/DataflowEditor/Private/Dataflow/DataflowSimulationScene.cpp
--- a/Engine/Plugins/Experimental/Dataflow/Source/DataflowEditor/Private/Dataflow/DataflowSimulationScene.cpp
+++ b/Engine/Plugins/Experimental/Dataflow/Source/DataflowEditor/Private/Dataflow/DataflowSimulationScene.cpp
@@ -23,7 +23,20 @@ FDataflowSimulationScene::FDataflowSimulationScene(FPreviewScene::ConstructionVa
 	SimulationGenerator = MakeShared<Dataflow::FDataflowSimulationGenerator>();
 	RootSceneActor = GetWorld()->SpawnActor<AChaosCacheManager>();
 
-	SceneDescription->ActorClass = GetEditorContent()->GetPreviewClass();
+	if(GetEditorContent())
+	{
+#if WITH_EDITORONLY_DATA
+		if(SceneDescription->BlueprintClass == nullptr)
+		{
+			SceneDescription->BlueprintClass = GetEditorContent()->GetPreviewClass();
+		}
+		if(const UDataflow* DataflowAsset = GetEditorContent()->GetDataflowAsset())
+		{
+			SceneDescription->CacheParams = DataflowAsset->PreviewCacheParams;
+			SceneDescription->CacheAsset = Cast<UChaosCacheCollection>(DataflowAsset->PreviewCacheAsset.Get());
+		}
+#endif
+	}
 
 	CreateSimulationScene();
 }
@@ -126,23 +139,21 @@ void FDataflowSimulationScene::CreateSimulationScene()
 {
 	if(SimulationGenerator && SceneDescription && GetWorld())
 	{
-		SimulationGenerator->SetFrameRate(SceneDescription->FrameRate);
+		SimulationGenerator->SetCacheParams(SceneDescription->CacheParams);
 		SimulationGenerator->SetCacheAsset(SceneDescription->CacheAsset);
-		SimulationGenerator->SetActorClass(SceneDescription->ActorClass);
-		SimulationGenerator->SetTimeRange(SceneDescription->TimeRange);
-		SimulationGenerator->SetBackgroundTask(SceneDescription->bBackgroundTask);
+		SimulationGenerator->SetBlueprintClass(SceneDescription->BlueprintClass);
 		SimulationGenerator->SetDataflowContent(GetEditorContent());
 
-		TimeRange = SceneDescription->TimeRange;
-		NumFrames = (TimeRange[1] > TimeRange[0]) ? FMath::Floor((TimeRange[1] - TimeRange[0]) * SceneDescription->FrameRate) : 0;
+		TimeRange = SceneDescription->CacheParams.TimeRange;
+		NumFrames = (TimeRange[1] > TimeRange[0]) ? FMath::Floor((TimeRange[1] - TimeRange[0]) * SceneDescription->CacheParams.FrameRate) : 0;
 		
-		PreviewActor = Dataflow::SpawnSimulatedActor(SceneDescription->ActorClass, Cast<AChaosCacheManager>(RootSceneActor),
+		PreviewActor = Dataflow::SpawnSimulatedActor(SceneDescription->BlueprintClass, Cast<AChaosCacheManager>(RootSceneActor),
 			SceneDescription->CacheAsset, false, GetEditorContent());
 
 		// Setup all the skelmesh animations
 		Dataflow::SetupSkeletonAnimation(PreviewActor);
 		
-		GetWorld()->GetSubsystem<UDataflowSimulationManager>()->SetSimulationEnabled(SceneDescription->CacheAsset == nullptr);
+		GetWorld()->GetSubsystem<UDataflowSimulationManager>()->SetSimulationEnabled(false);
 	}
 
 	// update the selection binding since we are constantly editing the graph
@@ -197,18 +208,11 @@ void FDataflowSimulationScene::AddReferencedObjects(FReferenceCollector& Collect
 
 void FDataflowSimulationScene::SceneDescriptionPropertyChanged(const FName& PropertyName)
 {
-	if (PropertyName == GET_MEMBER_NAME_CHECKED(UDataflowSimulationSceneDescription, FrameRate))
+	if (PropertyName == GET_MEMBER_NAME_CHECKED(UDataflowSimulationSceneDescription, CacheParams))
 	{
 		if(SimulationGenerator)
 		{
-			SimulationGenerator->SetFrameRate(SceneDescription->FrameRate);
-		}
-	}
-	else if (PropertyName == GET_MEMBER_NAME_CHECKED(UDataflowSimulationSceneDescription, TimeRange))
-	{
-		if(SimulationGenerator)
-		{
-			SimulationGenerator->SetTimeRange(SceneDescription->TimeRange);
+			SimulationGenerator->SetCacheParams(SceneDescription->CacheParams);
 		}
 	}
 	else if(PropertyName == GET_MEMBER_NAME_CHECKED(UDataflowSimulationSceneDescription, CacheAsset))
@@ -218,20 +222,25 @@ void FDataflowSimulationScene::SceneDescriptionPropertyChanged(const FName& Prop
 			SimulationGenerator->SetCacheAsset(SceneDescription->CacheAsset);
 		}
 	}
-	else if(PropertyName == GET_MEMBER_NAME_CHECKED(UDataflowSimulationSceneDescription, ActorClass))
+	else if(PropertyName == GET_MEMBER_NAME_CHECKED(UDataflowSimulationSceneDescription, BlueprintClass))
 	{
 		if(SimulationGenerator)
 		{
-			SimulationGenerator->SetActorClass(SceneDescription->ActorClass);
+			SimulationGenerator->SetBlueprintClass(SceneDescription->BlueprintClass);
 		}
 	}
-	else if(PropertyName == GET_MEMBER_NAME_CHECKED(UDataflowSimulationSceneDescription, bBackgroundTask))
+	if(GetEditorContent())
 	{
-		if(SimulationGenerator)
+		if(UDataflow* DataflowAsset = GetEditorContent()->GetDataflowAsset())
 		{
-			SimulationGenerator->SetBackgroundTask(SceneDescription->bBackgroundTask);
+#if WITH_EDITORONLY_DATA
+			DataflowAsset->PreviewCacheParams = SceneDescription->CacheParams;
+			DataflowAsset->PreviewCacheAsset = SceneDescription->CacheAsset;
+			DataflowAsset->PreviewBlueprintClass = SceneDescription->BlueprintClass;
+#endif
 		}
 	}
+	
 	// Unregister components, cache manager, selection...
 	ResetSimulationScene();
 
diff --git a/Engine/Plugins/Experimental/Dataflow/Source/DataflowEditor/Public/Dataflow/DataflowSimulationGenerator.h b/Engine/Plugins/Experimental/Dataflow/Source/DataflowEditor/Public/Dataflow/DataflowSimulationGenerator.h
--- a/Engine/Plugins/Experimental/Dataflow/Source/DataflowEditor/Public/Dataflow/DataflowSimulationGenerator.h
+++ b/Engine/Plugins/Experimental/Dataflow/Source/DataflowEditor/Public/Dataflow/DataflowSimulationGenerator.h
@@ -8,6 +8,7 @@
 #include "Logging/LogMacros.h"
 #include "Misc/AsyncTaskNotification.h"
 #include "Dataflow/DataflowContent.h"
+#include "Dataflow/DataflowPreview.h"
 #include "Dataflow/DataflowSimulationUtils.h"
 
 DECLARE_LOG_CATEGORY_EXTERN(LogDataflowSimulationGenerator, Log, All);
@@ -151,24 +152,18 @@ namespace Dataflow
 		virtual ETickableTickType GetTickableTickType() const override { return ETickableTickType::Always; }
 		//~ End FTickableEditorObject Interface
 
-		/** Set the frame rate for cache recording */
-		void SetFrameRate(const int32 InFrameRate);
-
-		/** Set the time range for cache recording */
-		void SetTimeRange(const FVector2f& InTimeRange);
+		/** Set the blueprint class for cache recording */
+		void SetBlueprintClass(const TSubclassOf<AActor>& InBlueprintClass);
+		
+		/** Set the cache params for cache recording */
+		void SetCacheParams(const FDataflowPreviewCacheParams& InCacheParams);
 
 		/** Set the cache asset for cache recording */
 		void SetCacheAsset(const TObjectPtr<UChaosCacheCollection>& InCacheAsset);
 
-		/** Set the actor class for cache recording */
-		void SetActorClass(const TSubclassOf<AActor>& InActorClass);
-
 		/** Set the dataflow content */
 		void SetDataflowContent(const TObjectPtr<UDataflowBaseContent>& InDataflowContent);
 
-		/** Set the background task boolean */
-		void SetBackgroundTask(const bool bInBackgroundTask);
-
 		/** Enqueue a generator action to be processed on the async thread */
 		void RequestGeneratorAction(EDataflowGeneratorActions Action);
 		
@@ -186,17 +181,11 @@ namespace Dataflow
 		/** Cache asset to store the caches */
 		TObjectPtr<UChaosCacheCollection> CacheAsset = nullptr;
 
-		/** Frame rate used to record the cache */
-		int32 FrameRate = 30;
-
-		/** Time range used to generate the cache */
-		FVector2f TimeRange = FVector2f(0.0f, 5.0f);
-
-		/** Boolean to control if the task is going to be run in the background */
-		bool bBackgroundTask = false;
+		/** Cache params used to record simulation */
+		FDataflowPreviewCacheParams CacheParams;
 
-		/** Chaos cache manager BP class to be spawned */
-		TSubclassOf<AActor> ActorClass;
+		/** Blueprint class used to spawn the actor */
+		TSubclassOf<AActor> BlueprintClass = nullptr;
 
 		/** Dataflow content */
 		TObjectPtr<UDataflowBaseContent> DataflowContent;
diff --git a/Engine/Plugins/Experimental/Dataflow/Source/DataflowEditor/Public/Dataflow/DataflowSimulationScene.h b/Engine/Plugins/Experimental/Dataflow/Source/DataflowEditor/Public/Dataflow/DataflowSimulationScene.h
--- a/Engine/Plugins/Experimental/Dataflow/Source/DataflowEditor/Public/Dataflow/DataflowSimulationScene.h
+++ b/Engine/Plugins/Experimental/Dataflow/Source/DataflowEditor/Public/Dataflow/DataflowSimulationScene.h
@@ -28,25 +28,17 @@ public:
 	/** Set the simulation scene */
 	void SetSimulationScene(class FDataflowSimulationScene* SimulationScene);
 
+	/** Caching blueprint actor class to spawn */
+	UPROPERTY(EditAnywhere, Category = "Preview")
+	TSubclassOf<AActor> BlueprintClass = nullptr;
+
 	/** Caching asset to be used to record the simulation  */
 	UPROPERTY(EditAnywhere, Category="Caching")
 	TObjectPtr<UChaosCacheCollection> CacheAsset = nullptr;
 
-	/** Number of sampling frames per second for caching*/
-	UPROPERTY(EditAnywhere, Category="Caching")
-	int32 FrameRate = 30;
-
-	/** Number of sampling frames per second for caching*/
+	/** Caching params used to record the simulation */
 	UPROPERTY(EditAnywhere, Category="Caching")
-	FVector2f TimeRange = FVector2f(0.0f, 5.0f);
-
-	/** Caching actor class to spawn */
-	UPROPERTY(EditAnywhere, Category = "Caching")
-	TSubclassOf<AActor> ActorClass;
-
-	/** Boolean to check if the caching will be done on an async thread (if yes no GT dependency) */
-	UPROPERTY(EditAnywhere, AdvancedDisplay, Category = "Caching")
-	bool bBackgroundTask = true;
+	FDataflowPreviewCacheParams CacheParams;
 
 private:
 
diff --git a/Engine/Source/Runtime/Experimental/Dataflow/Engine/Private/Dataflow/DataflowContent.cpp b/Engine/Source/Runtime/Experimental/Dataflow/Engine/Private/Dataflow/DataflowContent.cpp
--- a/Engine/Source/Runtime/Experimental/Dataflow/Engine/Private/Dataflow/DataflowContent.cpp
+++ b/Engine/Source/Runtime/Experimental/Dataflow/Engine/Private/Dataflow/DataflowContent.cpp
@@ -18,6 +18,7 @@ bool bDataflowEnableContextCaching = true;
 FAutoConsoleVariableRef CVARDataflowEnableContextCaching(TEXT("p.Dataflow.Editor.ContextCaching"), bDataflowEnableContextCaching,
 	TEXT("Allow the Dataflow editor to crate and use a pre-evaluated graph when the dataflow editor is re-opened.[def:true]"));
 
+
 namespace DataflowContextHelpers
 {
 
@@ -307,7 +308,7 @@ void UDataflowBaseContent::UpdateContentDatas()
 	{
 		if(const IDataflowContentOwner* ContentOwner = Cast<IDataflowContentOwner>(GetDataflowOwner()))
 		{
-			ContentOwner->UpdateDataflowContent(this);
+			ContentOwner->WriteDataflowContent(this);
 		}
 	}
 }
@@ -344,6 +345,20 @@ TObjectPtr<UDataflow> UDataflowBaseContent::GetDataflowAsset() const
 	return DataflowContext ? DataflowContext->Graph : nullptr; 
 }
 
+#if WITH_EDITOR
+
+void UDataflowBaseContent::PostEditChangeProperty(FPropertyChangedEvent& PropertyChangedEvent) 
+{
+	Super::PostEditChangeProperty(PropertyChangedEvent);
+	
+	if(IDataflowContentOwner* ContentOwner = Cast<IDataflowContentOwner>(GetDataflowOwner()))
+	{
+		ContentOwner->ReadDataflowContent(this);
+	}
+}
+
+#endif //if WITH_EDITOR
+
 void UDataflowBaseContent::SetLastModifiedTimestamp(Dataflow::FTimestamp InTimestamp, bool bMakeDirty) 
 { 
 	if (InTimestamp.IsInvalid() || LastModifiedTimestamp < InTimestamp)
@@ -392,7 +407,7 @@ void UDataflowBaseContent::AddReferencedObjects(UObject* InThis, FReferenceColle
 	Super::AddReferencedObjects(InThis, Collector);
 }
 
-void OverrideActorProperty(TObjectPtr<AActor>& PreviewActor, const TObjectPtr<UObject>& PropertyValue, const FName& PropertyName)
+void OverrideActorProperty(const TObjectPtr<AActor>& PreviewActor, TObjectPtr<UObject> PropertyValue, const FName& PropertyName)
 {
 	if(PreviewActor && PropertyValue)
 	{
@@ -425,7 +440,7 @@ UDataflowSkeletalContent::UDataflowSkeletalContent() : Super()
 {
 }
 
-void UDataflowSkeletalContent::SetSkeletalMesh(const TObjectPtr<USkeletalMesh>& SkeletalMeshAsset)
+void UDataflowSkeletalContent::SetSkeletalMesh(const TObjectPtr<USkeletalMesh>& SkeletalMeshAsset, const bool bHideAsset)
 {
 	SkeletalMesh = SkeletalMeshAsset;
 	if(SkeletalMesh)
@@ -435,17 +450,19 @@ void UDataflowSkeletalContent::SetSkeletalMesh(const TObjectPtr<USkeletalMesh>&
 			SetSkeleton(SkeletalMesh->GetSkeleton());
 		}
 	}
+	bHideSkeletalMesh = bHideAsset;
 	SetConstructionDirty(true);
 	SetSimulationDirty(true);
 }
 
-void UDataflowSkeletalContent::SetAnimationAsset(const TObjectPtr<UAnimationAsset>& SkeletalAnimationAsset)
+void UDataflowSkeletalContent::SetAnimationAsset(const TObjectPtr<UAnimationAsset>& SkeletalAnimationAsset, const bool bHideAsset)
 {
 	AnimationAsset = SkeletalAnimationAsset;
 	if(AnimationAsset && (AnimationAsset->GetSkeleton()) != Skeleton)
 	{
 		SetSkeleton(AnimationAsset->GetSkeleton());
 	}
+	bHideAnimationAsset = bHideAsset;
 	SetConstructionDirty(true);
 	SetSimulationDirty(true);
 }
@@ -469,10 +486,8 @@ void UDataflowSkeletalContent::SetSkeleton(const TObjectPtr<USkeleton>& Skeleton
 
 void UDataflowSkeletalContent::PostEditChangeProperty(FPropertyChangedEvent& PropertyChangedEvent) 
 {
-	Super::PostEditChangeProperty(PropertyChangedEvent);
-	
 	const FName PropertyName = PropertyChangedEvent.Property->GetFName();
-	
+
 	if (PropertyName == GET_MEMBER_NAME_CHECKED(UDataflowSkeletalContent, SkeletalMesh))
 	{
 		SetSkeletalMesh(SkeletalMesh);
@@ -485,6 +500,27 @@ void UDataflowSkeletalContent::PostEditChangeProperty(FPropertyChangedEvent& Pro
 	{
 		SetSkeleton(Skeleton);
 	}
+	Super::PostEditChangeProperty(PropertyChangedEvent);
+}
+
+bool UDataflowSkeletalContent::CanEditChange(const FProperty* InProperty) const
+{
+	if (!Super::CanEditChange(InProperty))
+	{
+		return false;
+	}
+
+	const FName& Name = InProperty->GetFName();
+	if (Name == GET_MEMBER_NAME_CHECKED(ThisClass, SkeletalMesh))
+	{
+		return bHideSkeletalMesh == false;
+	}
+	else if(Name == GET_MEMBER_NAME_CHECKED(ThisClass, AnimationAsset))
+	{
+		return bHideAnimationAsset == false;
+	}
+
+	return true;
 }
 
 #endif //if WITH_EDITOR
@@ -497,9 +533,9 @@ void UDataflowSkeletalContent::AddContentObjects(FReferenceCollector& Collector)
 void UDataflowSkeletalContent::AddReferencedObjects(UObject* InThis, FReferenceCollector& Collector)
 {
 	UDataflowSkeletalContent* This = CastChecked<UDataflowSkeletalContent>(InThis);
+	Collector.AddReferencedObject(This->Skeleton);
 	Collector.AddReferencedObject(This->SkeletalMesh);
 	Collector.AddReferencedObject(This->AnimationAsset);
-	Collector.AddReferencedObject(This->Skeleton);
 	Super::AddReferencedObjects(InThis, Collector);
 }
 
@@ -510,4 +546,3 @@ void UDataflowSkeletalContent::SetActorProperties(TObjectPtr<AActor>& PreviewAct
 	OverrideActorProperty(PreviewActor, AnimationAsset, TEXT("AnimationAsset"));
 	OverrideActorProperty(PreviewActor, SkeletalMesh, TEXT("SkeletalMesh"));
 }
-
diff --git a/Engine/Source/Runtime/Experimental/Dataflow/Engine/Public/Dataflow/DataflowContent.h b/Engine/Source/Runtime/Experimental/Dataflow/Engine/Public/Dataflow/DataflowContent.h
--- a/Engine/Source/Runtime/Experimental/Dataflow/Engine/Public/Dataflow/DataflowContent.h
+++ b/Engine/Source/Runtime/Experimental/Dataflow/Engine/Public/Dataflow/DataflowContent.h
@@ -107,7 +107,10 @@ public:
 	}
 	
 	/** Interface to update a dataflow content instance from that owner */
-	virtual void UpdateDataflowContent(const TObjectPtr<UDataflowBaseContent>& DataflowContent) const = 0;
+	virtual void WriteDataflowContent(const TObjectPtr<UDataflowBaseContent>& DataflowContent) const = 0;
+
+	/** Interface to update a dataflow content instance from that owner */
+	virtual void ReadDataflowContent(const TObjectPtr<UDataflowBaseContent>& DataflowContent) = 0;
 	
 protected :
 
@@ -191,6 +194,10 @@ public:
 	//~ UObject interface
 	static void AddReferencedObjects(UObject* InThis, FReferenceCollector& Collector);
 
+#if WITH_EDITOR
+	virtual void PostEditChangeProperty(FPropertyChangedEvent& PropertyChangedEvent) override;
+#endif //if WITH_EDITOR
+
 protected:
 	
 	/** Data flow terminal path for evaluation */
@@ -235,6 +242,7 @@ public:
 	virtual ~UDataflowSkeletalContent() override{}
 	
 #if WITH_EDITOR
+	virtual bool CanEditChange(const FProperty* InProperty) const override;
 	virtual void PostEditChangeProperty(FPropertyChangedEvent& PropertyChangedEvent) override;
 #endif //if WITH_EDITOR
 
@@ -242,7 +250,7 @@ public:
 	virtual void AddContentObjects(FReferenceCollector& Collector) override;
 
 	/** Data flow skeletal mesh accessors */
-	void SetSkeletalMesh(const TObjectPtr<USkeletalMesh>& InMesh);
+	void SetSkeletalMesh(const TObjectPtr<USkeletalMesh>& InMesh, const bool bHideAsset = false);
 	const TObjectPtr<USkeletalMesh>& GetSkeletalMesh() const { return SkeletalMesh; }
 
 	/** Data flow skeleton accessors */
@@ -250,7 +258,7 @@ public:
 	const TObjectPtr<USkeleton>& GetSkeleton() const { return Skeleton; }
 
 	/** Data flow animation asset accessors */
-	void SetAnimationAsset(const TObjectPtr<UAnimationAsset>& InAnimation);
+	void SetAnimationAsset(const TObjectPtr<UAnimationAsset>& InAnimation, const bool bHideAsset = false);
 	const TObjectPtr<UAnimationAsset>& GetAnimationAsset() const { return AnimationAsset; }
 
 	//~ UObject interface
@@ -272,4 +280,10 @@ protected:
 	/** Animation asset to be used to preview simulation */
 	UPROPERTY(EditAnywhere, Category = "Preview", Transient, SkipSerialization)
 	TObjectPtr<UAnimationAsset> AnimationAsset = nullptr;
+
+	/** Boolean to control if the skeletal mesh could be edited or not */
+	bool bHideSkeletalMesh = false;
+
+	/** Boolean to control if the animation asset could be edited or not */
+	bool bHideAnimationAsset = false;
 };
diff --git a/Engine/Source/Runtime/Experimental/Dataflow/Engine/Public/Dataflow/DataflowObject.h b/Engine/Source/Runtime/Experimental/Dataflow/Engine/Public/Dataflow/DataflowObject.h
--- a/Engine/Source/Runtime/Experimental/Dataflow/Engine/Public/Dataflow/DataflowObject.h
+++ b/Engine/Source/Runtime/Experimental/Dataflow/Engine/Public/Dataflow/DataflowObject.h
@@ -4,6 +4,7 @@
 
 #include "Containers/Array.h"
 #include "CoreMinimal.h"
+#include "DataflowPreview.h"
 #include "Dataflow/DataflowCore.h"
 #include "EdGraph/EdGraph.h"
 #include "Templates/Function.h"
@@ -140,5 +141,31 @@ public:
 
 	const Dataflow::FTimestamp& GetRenderingTimestamp() const { return LastModifiedRenderTarget; }
 
+#if WITH_EDITORONLY_DATA
+
+	/*
+	* The following PreviewScene properties are modeled after PreviewSkeletalMesh in USkeleton
+	*	- they are inside WITH_EDITORONLY_DATA because they are not used at game runtime
+	*	- TSoftObjectPtrs since that will make it possible to avoid loading these assets until the PreviewScene asks for them
+	*	- DuplicateTransient so that if you copy a ClothAsset it won't copy these preview properties
+	*	- AssetRegistrySearchable makes it so that if the user searches the name of a PreviewScene asset in the Asset Browser
+	*/
+
+	/** Cachie params used in this asset */
+	UPROPERTY(DuplicateTransient, AssetRegistrySearchable)
+	FDataflowPreviewCacheParams PreviewCacheParams;
+
+	/** Cache asset used in this asset */
+	UPROPERTY(DuplicateTransient, AssetRegistrySearchable)
+	TSoftObjectPtr<UObject> PreviewCacheAsset = nullptr;
+
+	/** Caching blueprint actor class to spawn */
+	UPROPERTY(DuplicateTransient, AssetRegistrySearchable)
+	TSubclassOf<AActor> PreviewBlueprintClass = nullptr;
+
+#endif
+
+
+
 };
 
diff --git a/Engine/Source/Runtime/Experimental/Dataflow/Engine/Public/Dataflow/DataflowPreview.h b/Engine/Source/Runtime/Experimental/Dataflow/Engine/Public/Dataflow/DataflowPreview.h
new file mode 100644
--- /dev/null
+++ b/Engine/Source/Runtime/Experimental/Dataflow/Engine/Public/Dataflow/DataflowPreview.h
@@ -0,0 +1,32 @@
+﻿// Copyright Epic Games, Inc. All Rights Reserved.
+
+#pragma once
+
+#include "CoreMinimal.h"
+#include "Templates/SharedPointer.h"
+#include "UObject/ObjectMacros.h"
+#include "UObject/ObjectPtr.h"
+#include "UObject/UObjectGlobals.h"
+#include "DataflowPreview.generated.h"
+
+class UAnimationAsset;
+class USkeletalMesh;
+
+USTRUCT()
+struct FDataflowPreviewCacheParams
+{
+	GENERATED_BODY()
+public :
+
+	/** Number of sampling frames per second for caching*/
+	UPROPERTY(EditAnywhere, Category="Caching")
+	int32 FrameRate = 30;
+
+	/** Number of sampling frames per second for caching*/
+	UPROPERTY(EditAnywhere, Category="Caching")
+	FVector2f TimeRange = FVector2f(0.0f, 5.0f);
+
+	/** Boolean to check if the caching will be done on an async thread (if yes no GT dependency) */
+	UPROPERTY(EditAnywhere, Category = "Caching")
+	bool bBackgroundTask = true;
+};
diff --git a/Engine/Source/Runtime/Experimental/GeometryCollectionEngine/Private/GeometryCollection/GeometryCollectionObject.cpp b/Engine/Source/Runtime/Experimental/GeometryCollectionEngine/Private/GeometryCollection/GeometryCollectionObject.cpp
--- a/Engine/Source/Runtime/Experimental/GeometryCollectionEngine/Private/GeometryCollection/GeometryCollectionObject.cpp
+++ b/Engine/Source/Runtime/Experimental/GeometryCollectionEngine/Private/GeometryCollection/GeometryCollectionObject.cpp
@@ -46,6 +46,7 @@
 #include "Chaos/ChaosArchive.h"
 #include "Chaos/MassProperties.h"
 #include "GeometryCollectionProxyData.h"
+#include "Dataflow/DataflowObject.h"
 #include "GeometryCollection/Facades/CollectionHierarchyFacade.h"
 #include "GeometryCollection/Facades/CollectionInstancedMeshFacade.h"
 
@@ -2060,12 +2061,12 @@ TObjectPtr<UDataflowBaseContent> UGeometryCollection::CreateDataflowContent()
 	BaseContent->SetDataflowOwner(this);
 	BaseContent->SetTerminalAsset(this);
 	
-	UpdateDataflowContent(BaseContent);
+	WriteDataflowContent(BaseContent);
 	
 	return BaseContent;
 }
 
-void UGeometryCollection::UpdateDataflowContent(const TObjectPtr<UDataflowBaseContent>& DataflowContent) const
+void UGeometryCollection::WriteDataflowContent(const TObjectPtr<UDataflowBaseContent>& DataflowContent) const
 {
 	if(const TObjectPtr<UDataflowBaseContent> BaseContent = Cast<UDataflowBaseContent>(DataflowContent))
 	{
@@ -2074,6 +2075,9 @@ void UGeometryCollection::UpdateDataflowContent(const TObjectPtr<UDataflowBaseCo
 	}
 }
 
+void UGeometryCollection::ReadDataflowContent(const TObjectPtr<UDataflowBaseContent>& DataflowContent)
+{}
+
 #if WITH_EDITOR
 bool UGeometryCollection::CanEditChange(const FProperty* InProperty) const
 {
diff --git a/Engine/Source/Runtime/Experimental/GeometryCollectionEngine/Public/GeometryCollection/GeometryCollectionObject.h b/Engine/Source/Runtime/Experimental/GeometryCollectionEngine/Public/GeometryCollection/GeometryCollectionObject.h
--- a/Engine/Source/Runtime/Experimental/GeometryCollectionEngine/Public/GeometryCollection/GeometryCollectionObject.h
+++ b/Engine/Source/Runtime/Experimental/GeometryCollectionEngine/Public/GeometryCollection/GeometryCollectionObject.h
@@ -826,9 +826,11 @@ public:
 	UPROPERTY(EditAnywhere, BlueprintReadWrite, Category = "Dataflow", DisplayName = "DataFlow Overrides")
 	TMap<FString, FString> Overrides;
 
-	/** IDataflowContentOwner interface */
+	//~ Begin IDataflowContentOwner interface
 	virtual TObjectPtr<UDataflowBaseContent> CreateDataflowContent() override;
-	virtual void UpdateDataflowContent(const TObjectPtr<UDataflowBaseContent>& DataflowContent) const override;
+	virtual void WriteDataflowContent(const TObjectPtr<UDataflowBaseContent>& DataflowContent) const override;
+	virtual void ReadDataflowContent(const TObjectPtr<UDataflowBaseContent>& DataflowContent) override;
+	//~ End IDataflowContentOwner interface
 
 	GEOMETRYCOLLECTIONENGINE_API const TArray<int32>& GetBreadthFirstTransformIndices() const { return BreadthFirstTransformIndices; }
 
