Chaos Modular Vehicles:
Add IsWheelInContact to SuspensionBaseInterface.

diff --git a/Engine/Source/Runtime/Experimental/ChaosVehicles/ChaosVehiclesCore/Public/SimModule/SuspensionBaseInterface.h b/Engine/Source/Runtime/Experimental/ChaosVehicles/ChaosVehiclesCore/Public/SimModule/SuspensionBaseInterface.h
--- a/Engine/Source/Runtime/Experimental/ChaosVehicles/ChaosVehiclesCore/Public/SimModule/SuspensionBaseInterface.h
+++ b/Engine/Source/Runtime/Experimental/ChaosVehicles/ChaosVehiclesCore/Public/SimModule/SuspensionBaseInterface.h
@@ -44,6 +44,7 @@ public:
 	virtual void GetWorldRaycastLocation(const FTransform& BodyTransform, float WheelRadius, FSpringTrace& OutTrace) = 0;
 
 	void SetTargetPoint(const FVector& InTargetPoint, const FVector& InImpactNormal, float InHitDistance, bool InWheelInContact);
+	bool IsWheelInContact() const { return WheelInContact; }
 	void SetWheelSimTreeIndex(int WheelTreeIndexIn) { WheelSimTreeIndex = WheelTreeIndexIn; }
 	int GetWheelSimTreeIndex() const { return WheelSimTreeIndex; }
 
