Dataflow : Removed the UDataflowGraph from the Context

diff --git a/Engine/Plugins/ChaosClothAssetEditor/Source/ChaosClothAssetEditor/Private/ChaosClothAsset/ClothEditorToolkit.h b/Engine/Plugins/ChaosClothAssetEditor/Source/ChaosClothAssetEditor/Private/ChaosClothAsset/ClothEditorToolkit.h
--- a/Engine/Plugins/ChaosClothAssetEditor/Source/ChaosClothAssetEditor/Private/ChaosClothAsset/ClothEditorToolkit.h
+++ b/Engine/Plugins/ChaosClothAssetEditor/Source/ChaosClothAssetEditor/Private/ChaosClothAsset/ClothEditorToolkit.h
@@ -32,7 +32,7 @@ namespace Dataflow
 		DATAFLOW_CONTEXT_INTERNAL(TEngineContext<FContextSingle>, FClothAssetDataflowContext);
 
 		FClothAssetDataflowContext(UObject* InOwner, UDataflow* InGraph, FTimestamp InTimestamp)
-			: Super(InOwner, InGraph, InTimestamp)
+			: Super(InOwner, InTimestamp)
 		{}
 	};
 }
diff --git a/Engine/Plugins/Experimental/Dataflow/Source/DataflowEditor/Private/Dataflow/DataflowEditorToolkit.cpp b/Engine/Plugins/Experimental/Dataflow/Source/DataflowEditor/Private/Dataflow/DataflowEditorToolkit.cpp
--- a/Engine/Plugins/Experimental/Dataflow/Source/DataflowEditor/Private/Dataflow/DataflowEditorToolkit.cpp
+++ b/Engine/Plugins/Experimental/Dataflow/Source/DataflowEditor/Private/Dataflow/DataflowEditorToolkit.cpp
@@ -1020,7 +1020,7 @@ void FDataflowEditorToolkit::Tick(float DeltaTime)
 			Dataflow::FTimestamp TimeStamp = EditorContent->GetLastModifiedTimestamp();
 			if (!EditorContent->GetDataflowContext())
 			{
-				EditorContent->SetDataflowContext(MakeShared<Dataflow::FEngineContext>(EditorContent->GetDataflowOwner(), EditorContent->GetDataflowAsset(), Dataflow::FTimestamp::Invalid));
+				EditorContent->SetDataflowContext(MakeShared<Dataflow::FEngineContext>(EditorContent->GetDataflowOwner(), Dataflow::FTimestamp::Invalid));
 				TimeStamp = Dataflow::FTimestamp::Invalid;
 			}
 
diff --git a/Engine/Plugins/Experimental/Dataflow/Source/DataflowEditor/Private/Dataflow/DataflowEditorUtil.cpp b/Engine/Plugins/Experimental/Dataflow/Source/DataflowEditor/Private/Dataflow/DataflowEditorUtil.cpp
--- a/Engine/Plugins/Experimental/Dataflow/Source/DataflowEditor/Private/Dataflow/DataflowEditorUtil.cpp
+++ b/Engine/Plugins/Experimental/Dataflow/Source/DataflowEditor/Private/Dataflow/DataflowEditorUtil.cpp
@@ -123,13 +123,13 @@ namespace Dataflow
 		{
 			if (!Content->GetDataflowContext())
 			{
-				Content->SetDataflowContext(MakeShared<FEngineContext>(Content->GetDataflowOwner(), Content->GetDataflowAsset(), FTimestamp::Invalid));
+				Content->SetDataflowContext(MakeShared<FEngineContext>(Content->GetDataflowOwner(), FTimestamp::Invalid));
 			}
 			return Content->GetDataflowContext();
 		}
 
 		ensure(false);
-		return MakeShared<FEngineContext>(nullptr, nullptr, FTimestamp::Invalid);
+		return MakeShared<FEngineContext>(nullptr, FTimestamp::Invalid);
 	}
 
 	bool CanRenderNodeOutput(const UDataflowEdNode& EdNode, const UDataflowBaseContent& EditorContent, const IDataflowConstructionViewMode& ViewMode)
diff --git a/Engine/Plugins/Experimental/Dataflow/Source/DataflowEditor/Public/Dataflow/DataflowEditorToolkit.h b/Engine/Plugins/Experimental/Dataflow/Source/DataflowEditor/Public/Dataflow/DataflowEditorToolkit.h
--- a/Engine/Plugins/Experimental/Dataflow/Source/DataflowEditor/Public/Dataflow/DataflowEditorToolkit.h
+++ b/Engine/Plugins/Experimental/Dataflow/Source/DataflowEditor/Public/Dataflow/DataflowEditorToolkit.h
@@ -31,21 +31,6 @@ class FDataflowSkeletonView;
 class UDataflowEditor;
 
 
-namespace Dataflow
-{
-	class DATAFLOWEDITOR_API FAssetContext : public TEngineContext<FContextSingle>
-	{
-	public:
-		DATAFLOW_CONTEXT_INTERNAL(TEngineContext<FContextSingle>, FAssetContext); 
-
-		FAssetContext(UObject* InOwner, UDataflow* InGraph, FTimestamp InTimestamp)
-			: Super(InOwner, InGraph, InTimestamp)
-		{}
-	};
-}
-
-
-
 class DATAFLOWEDITOR_API FDataflowEditorToolkit final : public FBaseCharacterFXEditorToolkit, public FTickableEditorObject, public FNotifyHook, public FGCObject
 {
 	using FBaseCharacterFXEditorToolkit::ObjectScene;
diff --git a/Engine/Plugins/Experimental/Dataflow/Source/DataflowSimulation/Private/Dataflow/DataflowSimulationManager.cpp b/Engine/Plugins/Experimental/Dataflow/Source/DataflowSimulation/Private/Dataflow/DataflowSimulationManager.cpp
--- a/Engine/Plugins/Experimental/Dataflow/Source/DataflowSimulation/Private/Dataflow/DataflowSimulationManager.cpp
+++ b/Engine/Plugins/Experimental/Dataflow/Source/DataflowSimulation/Private/Dataflow/DataflowSimulationManager.cpp
@@ -1,4 +1,4 @@
-﻿// Copyright Epic Games, Inc. All Rights Reserved.
+// Copyright Epic Games, Inc. All Rights Reserved.
 
 #include "Dataflow/DataflowSimulationManager.h"
 #include "Dataflow/DataflowSimulationUtils.h"
@@ -333,7 +333,7 @@ void UDataflowSimulationManager::AddSimulationInterface(IDataflowSimulationInter
 			if(!DataflowData.SimulationContext.IsValid())
 			{
 				DataflowData.SimulationContext = MakeShared<Dataflow::FDataflowSimulationContext>(
-					DataflowAsset, DataflowAsset, Dataflow::FTimestamp::Invalid);
+					DataflowAsset, Dataflow::FTimestamp::Invalid);
 			}
 			DataflowData.SimulationInterfaces.FindOrAdd(SimulationInterface->GetSimulationType()).Add(SimulationInterface);
 		}
diff --git a/Engine/Plugins/Experimental/Dataflow/Source/DataflowSimulation/Public/Dataflow/DataflowSimulationContext.h b/Engine/Plugins/Experimental/Dataflow/Source/DataflowSimulation/Public/Dataflow/DataflowSimulationContext.h
--- a/Engine/Plugins/Experimental/Dataflow/Source/DataflowSimulation/Public/Dataflow/DataflowSimulationContext.h
+++ b/Engine/Plugins/Experimental/Dataflow/Source/DataflowSimulation/Public/Dataflow/DataflowSimulationContext.h
@@ -1,4 +1,4 @@
-﻿// Copyright Epic Games, Inc. All Rights Reserved.
+// Copyright Epic Games, Inc. All Rights Reserved.
 
 #pragma once
 
@@ -17,8 +17,8 @@ namespace Dataflow
 	public:
 		DATAFLOW_CONTEXT_INTERNAL(TEngineContext<Base>, TSimulationContext);
 
-		TSimulationContext(const TObjectPtr<UObject>& InOwner,const TObjectPtr<UDataflow>& InGraph, FTimestamp InTimestamp)
-				: Super(InOwner, InGraph, InTimestamp)
+		TSimulationContext(const TObjectPtr<UObject>& InOwner, FTimestamp InTimestamp)
+				: Super(InOwner, InTimestamp)
 		{}
 		
 		virtual ~TSimulationContext() override {};
diff --git a/Engine/Source/Runtime/Experimental/Dataflow/Engine/Private/Dataflow/DataflowBlueprintLibrary.cpp b/Engine/Source/Runtime/Experimental/Dataflow/Engine/Private/Dataflow/DataflowBlueprintLibrary.cpp
--- a/Engine/Source/Runtime/Experimental/Dataflow/Engine/Private/Dataflow/DataflowBlueprintLibrary.cpp
+++ b/Engine/Source/Runtime/Experimental/Dataflow/Engine/Private/Dataflow/DataflowBlueprintLibrary.cpp
@@ -14,7 +14,7 @@ void UDataflowBlueprintLibrary::EvaluateTerminalNodeByName(UDataflow* Dataflow,
 		{
 			if (const FDataflowTerminalNode* TerminalNode = Node->AsType<const FDataflowTerminalNode>())
 			{
-				Dataflow::FEngineContext Context(ResultAsset, Dataflow, FPlatformTime::Cycles64());
+				Dataflow::FEngineContext Context(ResultAsset, FPlatformTime::Cycles64());
 				TerminalNode->Evaluate(Context);
 				if (ResultAsset)
 				{
diff --git a/Engine/Source/Runtime/Experimental/Dataflow/Engine/Private/Dataflow/DataflowContent.cpp b/Engine/Source/Runtime/Experimental/Dataflow/Engine/Private/Dataflow/DataflowContent.cpp
--- a/Engine/Source/Runtime/Experimental/Dataflow/Engine/Private/Dataflow/DataflowContent.cpp
+++ b/Engine/Source/Runtime/Experimental/Dataflow/Engine/Private/Dataflow/DataflowContent.cpp
@@ -317,7 +317,7 @@ void UDataflowBaseContent::SetDataflowOwner(const TObjectPtr<UObject>& InOwner)
 {
 	if(!DataflowContext)
 	{
-		DataflowContext = MakeShared<Dataflow::FEngineContext>(nullptr, nullptr, Dataflow::FTimestamp::Invalid);
+		DataflowContext = MakeShared<Dataflow::FEngineContext>(nullptr, Dataflow::FTimestamp::Invalid);
 	}
 	DataflowContext->Owner = InOwner;  
 	SetConstructionDirty(true);
@@ -333,16 +333,16 @@ void UDataflowBaseContent::SetDataflowAsset(const TObjectPtr<UDataflow>& Dataflo
 {
 	if(!DataflowContext)
 	{
-		DataflowContext = MakeShared<Dataflow::FEngineContext>(nullptr, nullptr, Dataflow::FTimestamp::Invalid);
+		DataflowContext = MakeShared<Dataflow::FEngineContext>(nullptr, Dataflow::FTimestamp::Invalid);
 	}
-	DataflowContext->Graph = DataflowAsset;  
+	DataflowGraph = DataflowAsset;  
 	SetConstructionDirty(true);
 	SetSimulationDirty(true);
 }
 
 TObjectPtr<UDataflow> UDataflowBaseContent::GetDataflowAsset() const 
 {
-	return DataflowContext ? DataflowContext->Graph : nullptr; 
+	return DataflowGraph;
 }
 
 #if WITH_EDITOR
@@ -391,7 +391,7 @@ void UDataflowBaseContent::Serialize(FArchive& Ar)
 
 	if (!DataflowContext)
 	{
-		DataflowContext = MakeShared<Dataflow::FEngineContext>(nullptr, nullptr, LastModifiedTimestamp);
+		DataflowContext = MakeShared<Dataflow::FEngineContext>(nullptr, LastModifiedTimestamp);
 	}
 	DataflowContext->Serialize(Ar);
 }
@@ -402,7 +402,7 @@ void UDataflowBaseContent::AddReferencedObjects(UObject* InThis, FReferenceColle
 	if(This->DataflowContext)
 	{
 		Collector.AddReferencedObject(This->DataflowContext->Owner);
-		Collector.AddReferencedObject(This->DataflowContext->Graph);
+		Collector.AddReferencedObject(This->DataflowGraph);
 	}
 	Super::AddReferencedObjects(InThis, Collector);
 }
diff --git a/Engine/Source/Runtime/Experimental/Dataflow/Engine/Public/Dataflow/DataflowContent.h b/Engine/Source/Runtime/Experimental/Dataflow/Engine/Public/Dataflow/DataflowContent.h
--- a/Engine/Source/Runtime/Experimental/Dataflow/Engine/Public/Dataflow/DataflowContent.h
+++ b/Engine/Source/Runtime/Experimental/Dataflow/Engine/Public/Dataflow/DataflowContent.h
@@ -204,6 +204,10 @@ protected:
 	UPROPERTY(Transient, SkipSerialization)
 	FString DataflowTerminal = "";
 
+	/** Dataflow graph for evaluation */
+	UPROPERTY(Transient, SkipSerialization)
+	TObjectPtr<UDataflow> DataflowGraph;
+	
 	/** Data flow terminal path for evaluation */
 	UPROPERTY(Transient, SkipSerialization)
 	TObjectPtr<UObject> TerminalAsset = nullptr;
diff --git a/Engine/Source/Runtime/Experimental/Dataflow/Engine/Public/Dataflow/DataflowObjectInterface.h b/Engine/Source/Runtime/Experimental/Dataflow/Engine/Public/Dataflow/DataflowObjectInterface.h
--- a/Engine/Source/Runtime/Experimental/Dataflow/Engine/Public/Dataflow/DataflowObjectInterface.h
+++ b/Engine/Source/Runtime/Experimental/Dataflow/Engine/Public/Dataflow/DataflowObjectInterface.h
@@ -2,13 +2,9 @@
 
 #pragma once
 
-#include "Containers/UnrealString.h"
-#include "Dataflow/DataflowGraph.h"
-#include "GameFramework/Actor.h"
 #include "Dataflow/DataflowNodeParameters.h"
 #include "GeometryCollection/ManagedArrayCollection.h"
 
-class UDataflow;
 class UObject;
 class FArchive;
 
@@ -20,14 +16,12 @@ namespace Dataflow
 	public:
 		DATAFLOW_CONTEXT_INTERNAL(Base, TEngineContext);
 
-		TEngineContext(const TObjectPtr<UObject>& InOwner,const TObjectPtr<UDataflow>& InGraph, FTimestamp InTimestamp)
-				: Base(InTimestamp)
-				, Owner(InOwner)
-				, Graph(InGraph)
+		TEngineContext(const TObjectPtr<UObject>& InOwner, FTimestamp InTimestamp)
+			: Base(InTimestamp)
+			, Owner(InOwner)
 		{}
-	
+
 		TObjectPtr<UObject> Owner = nullptr;
-		TObjectPtr<UDataflow> Graph = nullptr;
 
 		virtual ~TEngineContext() {}
 
