Fix bundle manager caching uondemandtoc files.
Fix using different mount IDs when requesting a tagset for an already mounted container, which should fix containers being created twice in IOStoreOnDemand.
Add some comments for missing install/release APIs.
Add EOnDemandMountOptions::InstallSkipMount to just install from already mounted containers.

diff --git a/Engine/Source/Runtime/Experimental/IoStore/OnDemand/Private/OnDemandIoStore.cpp b/Engine/Source/Runtime/Experimental/IoStore/OnDemand/Private/OnDemandIoStore.cpp
--- a/Engine/Source/Runtime/Experimental/IoStore/OnDemand/Private/OnDemandIoStore.cpp
+++ b/Engine/Source/Runtime/Experimental/IoStore/OnDemand/Private/OnDemandIoStore.cpp
@@ -792,42 +792,45 @@ FIoStatus FOnDemandIoStore::TickMountRequest(FMountRequest& MountRequest)
 	Private::SplitHostUrl(Args.Url, Host, TocRelUrl);
 	const FStringView TocPath = FPathViews::GetPath(TocRelUrl);
 
-	if (Args.Toc)
+	if (!EnumHasAnyFlags(Args.Options, EOnDemandMountOptions::InstallSkipMount))
 	{
-		CreateContainersFromToc(Args.MountId, TocPath, *Args.Toc, MountRequest.Containers);
-	}
-	else if (Args.FilePath.IsEmpty() == false)
-	{
-		UE_LOG(LogIoStoreOnDemand, Log, TEXT("Loading TOC from file '%s'"), *Args.FilePath);
-
-		// TODO: Enable validation when the sentinal is included in all serialization paths
-		const bool bValidate = false;
-		TIoStatusOr<FOnDemandToc> TocStatus = FOnDemandToc::LoadFromFile(Args.FilePath, bValidate);
-		if (!TocStatus.IsOk())
+		if (Args.Toc)
 		{
-			return TocStatus.Status();
+			CreateContainersFromToc(Args.MountId, TocPath, *Args.Toc, MountRequest.Containers);
 		}
+		else if (Args.FilePath.IsEmpty() == false)
+		{
+			UE_LOG(LogIoStoreOnDemand, Log, TEXT("Loading TOC from file '%s'"), *Args.FilePath);
 
-		Args.Toc = MakeShared<FOnDemandToc>(TocStatus.ConsumeValueOrDie());
-
-		CreateContainersFromToc(Args.MountId, TocPath, *Args.Toc, MountRequest.Containers);
-	}
-	else if (Args.Url.IsEmpty() == false)
-	{
-		UE_LOG(LogIoStoreOnDemand, Log, TEXT("Loading TOC from URL '%s'"), *Args.Url);
+			// TODO: Enable validation when the sentinal is included in all serialization paths
+			const bool bValidate = false;
+			TIoStatusOr<FOnDemandToc> TocStatus = FOnDemandToc::LoadFromFile(Args.FilePath, bValidate);
+			if (!TocStatus.IsOk())
+			{
+				return TocStatus.Status();
+			}
 
-		const uint32 RetryCount				= 2;
-		const bool bFollowRedirects			= true;
-		TIoStatusOr<FOnDemandToc> TocStatus	= FOnDemandToc::LoadFromUrl(Args.Url, 2, bFollowRedirects);
+			Args.Toc = MakeUnique<FOnDemandToc>(TocStatus.ConsumeValueOrDie());
 
-		if (!TocStatus.IsOk())
-		{
-			return TocStatus.Status();
+			CreateContainersFromToc(Args.MountId, TocPath, *Args.Toc, MountRequest.Containers);
 		}
+		else if (Args.Url.IsEmpty() == false)
+		{
+			UE_LOG(LogIoStoreOnDemand, Log, TEXT("Loading TOC from URL '%s'"), *Args.Url);
+
+			const uint32 RetryCount = 2;
+			const bool bFollowRedirects = true;
+			TIoStatusOr<FOnDemandToc> TocStatus = FOnDemandToc::LoadFromUrl(Args.Url, 2, bFollowRedirects);
 
-		Args.Toc = MakeShared<FOnDemandToc>(TocStatus.ConsumeValueOrDie());
+			if (!TocStatus.IsOk())
+			{
+				return TocStatus.Status();
+			}
 
-		CreateContainersFromToc(Args.MountId, TocPath, *Args.Toc, MountRequest.Containers);
+			Args.Toc = MakeUnique<FOnDemandToc>(TocStatus.ConsumeValueOrDie());
+
+			CreateContainersFromToc(Args.MountId, TocPath, *Args.Toc, MountRequest.Containers);
+		}
 	}
 
 	// Remove already mounted containers
diff --git a/Engine/Source/Runtime/Experimental/IoStore/OnDemand/Public/IO/IoStoreOnDemand.h b/Engine/Source/Runtime/Experimental/IoStore/OnDemand/Public/IO/IoStoreOnDemand.h
--- a/Engine/Source/Runtime/Experimental/IoStore/OnDemand/Public/IO/IoStoreOnDemand.h
+++ b/Engine/Source/Runtime/Experimental/IoStore/OnDemand/Public/IO/IoStoreOnDemand.h
@@ -242,18 +242,20 @@ enum class EOnDemandInitResult
 enum class EOnDemandMountOptions
 {
 	/** The TOC is loaded but not installed or available for streaming. */
-	None			= 0,
+	None				= 0,
 	/** Make on-demand container(s) within a TOC available for streaming. */
-	StreamOnDemand	= 1 << 0,
+	StreamOnDemand		= 1 << 0,
 	/** Download and install on-demand contianer(s) to local storage. */
-	Install			= 1 << 1
+	Install				= 1 << 1,
+	/** Install only already mounted on-demand contianer(s) to local storage. */
+	InstallSkipMount	= 1 << 2,
 };
 ENUM_CLASS_FLAGS(EOnDemandMountOptions);
 
 struct FOnDemandMountArgs
 {
 	/** Mount an already serialized TOC. */
-	TSharedPtr<FOnDemandToc> Toc;
+	TUniquePtr<FOnDemandToc> Toc;
 	/** Mandatory ID to be used for unmounting all container file(s) included in the TOC. */
 	FString MountId;
 	/** Download the TOC from the specified URL. */
