Forked Server support for IAD cache

diff --git a/Engine/Source/Runtime/Core/Private/IO/IoDispatcher.cpp b/Engine/Source/Runtime/Core/Private/IO/IoDispatcher.cpp
--- a/Engine/Source/Runtime/Core/Private/IO/IoDispatcher.cpp
+++ b/Engine/Source/Runtime/Core/Private/IO/IoDispatcher.cpp
@@ -46,7 +46,8 @@ static const TCHAR* const GetIoErrorText_ErrorCodeTextArray[] =
 	TEXT("Invalid Parameter"),
 	TEXT("Signature Error"),
 	TEXT("Invalid Encryption Key"),
-	TEXT("Compression Error")
+	TEXT("Compression Error"),
+	TEXT("Pending Fork")
 };
 
 CORE_API const TCHAR* const* GetIoErrorText_ErrorCodeText = GetIoErrorText_ErrorCodeTextArray;
diff --git a/Engine/Source/Runtime/Core/Public/IO/IoStatus.h b/Engine/Source/Runtime/Core/Public/IO/IoStatus.h
--- a/Engine/Source/Runtime/Core/Public/IO/IoStatus.h
+++ b/Engine/Source/Runtime/Core/Public/IO/IoStatus.h
@@ -29,6 +29,7 @@ enum class EIoErrorCode
 	SignatureError,
 	InvalidEncryptionKey,
 	CompressionError,
+	PendingFork,
 };
 
 /** Get I/O error code description. */
diff --git a/Engine/Source/Runtime/Experimental/IoStore/OnDemand/Private/OnDemandIoStore.cpp b/Engine/Source/Runtime/Experimental/IoStore/OnDemand/Private/OnDemandIoStore.cpp
--- a/Engine/Source/Runtime/Experimental/IoStore/OnDemand/Private/OnDemandIoStore.cpp
+++ b/Engine/Source/Runtime/Experimental/IoStore/OnDemand/Private/OnDemandIoStore.cpp
@@ -22,6 +22,7 @@
 #include "IO/IoContainerHeader.h"
 #include "IO/PackageStore.h"
 #include "Misc/CommandLine.h"
+#include "Misc/CoreDelegates.h"
 #include "Misc/CoreDelegatesInternal.h"
 #include "Misc/CoreMisc.h"
 #include "Misc/EncryptionKeyManager.h"
@@ -58,7 +59,6 @@
 
 	if (IsRunningDedicatedServer())
 	{
-#if 0 // TODO: Fix this for forking servers
 		if (!FForkProcessHelper::IsForkRequested())
 		{
 			DirName = TEXT("InstallCacheServer");
@@ -72,9 +72,6 @@
 
 			DirName = FString::Printf(TEXT("InstallCacheServer-%u"), FPlatformProcess::GetCurrentProcessId());
 		}
-#else
-		DirName = TEXT("InstallCacheServer");
-#endif
 	}
 #if WITH_EDITOR
 	else if (GIsEditor)
@@ -334,6 +331,8 @@
 
 FOnDemandIoStore::~FOnDemandIoStore()
 {
+	FCoreDelegates::OnPostFork.RemoveAll(this);
+
 	FEncryptionKeyManager::Get().OnKeyAdded().RemoveAll(this);
 
 	if (OnMountPakHandle.IsValid())
@@ -349,38 +348,19 @@
 
 FIoStatus FOnDemandIoStore::Initialize()
 {
-	bool bUseInstallCache = GIoStoreOnDemandInstallCacheEnabled;
-#if !UE_BUILD_SHIPPING
-	bUseInstallCache = FParse::Param(FCommandLine::Get(), TEXT("NoIAD")) == false;
-#endif
-	if (bUseInstallCache)
+	const FIoStatus CacheStatus = InitializeOnDemandInstallCache();
+	if (CacheStatus.GetErrorCode() == EIoErrorCode::PendingFork)
 	{
-		FOnDemandInstallCacheConfig CacheConfig;
-		CacheConfig.RootDirectory = Private::GetInstallCacheDirectory(); 
-#if !UE_BUILD_SHIPPING
-		CacheConfig.bDropCache = FParse::Param(FCommandLine::Get(), TEXT("Iad.DropCache"));
-#endif
-		InstallCache = MakeOnDemandInstallCache(*this, CacheConfig);
-		if (InstallCache.IsValid())
+		UE_LOG(LogIoStoreOnDemand, Display, TEXT("Deferring install cache initialization until after process fork"));
+		if (!FCoreDelegates::OnPostFork.IsBoundToObject(this))
 		{
-			int32 BackendPriority = -5; // Lower than file (zero) but higher than streaming backend (-10)
-#if !UE_BUILD_SHIPPING
-			if (FParse::Param(FCommandLine::Get(), TEXT("Iad")))
-			{
-				// Bump the priority to be higher then the file system backend
-				BackendPriority = 5;
-			}
-#endif
-			FIoDispatcher::Get().Mount(InstallCache.ToSharedRef(), BackendPriority);
-			PackageStoreBackend = MakeOnDemandPackageStoreBackend();
-			FPackageStore::Get().Mount(PackageStoreBackend.ToSharedRef());
-		}
-		else
-		{
-			// Only warn until this is properly tested
-			UE_LOG(LogIoStoreOnDemand, Warning, TEXT("Failed to initialize install cache"));
+			FCoreDelegates::OnPostFork.AddRaw(this, &FOnDemandIoStore::OnPostFork);
 		}
 	}
+	else if (!CacheStatus.IsOk())
+	{
+		return CacheStatus;
+	}
 
 #if 0
 	OnMountPakHandle = FCoreInternalDelegates::GetOnPakMountOperation().AddLambda(
@@ -459,56 +439,6 @@
 	}
 #endif
 
-#if !(UE_BUILD_SHIPPING|UE_BUILD_TEST)
-	FString ParamValue;
-	if (FParse::Value(FCommandLine::Get(), TEXT("Iad.Fill="), ParamValue))
-	{
-		ParamValue.TrimStartAndEndInline();
-		int64 FillSize = -1;
-		LexFromString(FillSize, ParamValue);
-
-		if (FillSize > 0)
-		{
-			if (ParamValue.EndsWith(TEXT("GB")))
-			{
-				FillSize = FillSize << 30;
-			}
-			if (ParamValue.EndsWith(TEXT("MB")))
-			{
-				FillSize = FillSize << 20;
-			}
-
-			UE_LOG(LogIoStoreOnDemand, Log, TEXT("Filling install cache with %.2lf MiB of dummy data"), double(FillSize) / 1024.0 / 1024.0);
-
-			FIoStatus	Status = FIoStatus::Ok;
-			uint64		Seed = 1;
-			while (FillSize >= 0 && Status.IsOk())
-			{
-				const uint64		ChunkSize = 256 << 10;
-				FIoBuffer			Chunk(ChunkSize);
-				TArrayView<uint64>	Values(reinterpret_cast<uint64*>(Chunk.GetData()), ChunkSize / sizeof(uint64));
-
-				for (uint64& Value : Values)
-				{
-					Value = Seed;
-				}
-
-				const FIoHash ChunkHash = FIoHash::HashBuffer(Chunk.GetView());
-				Status = InstallCache->PutChunk(MoveTemp(Chunk), ChunkHash);
-				Seed++;
-				FillSize -= ChunkSize;
-			}
-
-			if (Status.IsOk())
-			{
-				Status = InstallCache->Flush();
-			}
-
-			UE_CLOG(!Status.IsOk(), LogIoStoreOnDemand, Warning, TEXT("Failed to fill install cache with dummy data"));
-		}
-	}
-#endif
-
 	return EIoErrorCode::Ok;
 }
 
@@ -586,6 +516,115 @@
 	return GetChunkInfo(ChunkId, EOnDemandContainerFlags::Mounted | EOnDemandContainerFlags::Installed);
 }
 
+void FOnDemandIoStore::OnPostFork(EForkProcessRole ProcessRole)
+{
+	FCoreDelegates::OnPostFork.RemoveAll(this);
+
+	if (ProcessRole != EForkProcessRole::Child)
+	{
+		return;
+	}
+
+	const FIoStatus Status = Initialize();
+	if (!Status.IsOk())
+	{
+		UE_LOG(LogIoStoreOnDemand, Fatal, TEXT("Failed to initialize I/O store on demand (post fork), reason '%s'"), *Status.ToString());
+	}
+}
+
+FIoStatus FOnDemandIoStore::InitializeOnDemandInstallCache()
+{
+	if (FForkProcessHelper::IsForkRequested() && !FForkProcessHelper::IsForkedChildProcess())
+	{
+		return FIoStatusBuilder(EIoErrorCode::PendingFork) << TEXT("Install cache waiting for fork");
+	}
+
+	bool bUseInstallCache = GIoStoreOnDemandInstallCacheEnabled;
+#if !UE_BUILD_SHIPPING
+	bUseInstallCache = FParse::Param(FCommandLine::Get(), TEXT("NoIAD")) == false;
+#endif
+	if (bUseInstallCache)
+	{
+		FOnDemandInstallCacheConfig CacheConfig;
+		CacheConfig.RootDirectory = Private::GetInstallCacheDirectory();
+#if !UE_BUILD_SHIPPING
+		CacheConfig.bDropCache = FParse::Param(FCommandLine::Get(), TEXT("Iad.DropCache"));
+#endif
+		InstallCache = MakeOnDemandInstallCache(*this, CacheConfig);
+		if (InstallCache.IsValid())
+		{
+			int32 BackendPriority = -5; // Lower than file (zero) but higher than streaming backend (-10)
+#if !UE_BUILD_SHIPPING
+			if (FParse::Param(FCommandLine::Get(), TEXT("Iad")))
+			{
+				// Bump the priority to be higher then the file system backend
+				BackendPriority = 5;
+			}
+#endif
+			FIoDispatcher::Get().Mount(InstallCache.ToSharedRef(), BackendPriority);
+			PackageStoreBackend = MakeOnDemandPackageStoreBackend();
+			FPackageStore::Get().Mount(PackageStoreBackend.ToSharedRef());
+		}
+		else
+		{
+			// Only warn until this is properly tested
+			UE_LOG(LogIoStoreOnDemand, Warning, TEXT("Failed to initialize install cache"));
+		}
+	}
+
+#if !(UE_BUILD_SHIPPING|UE_BUILD_TEST)
+	FString ParamValue;
+	if (FParse::Value(FCommandLine::Get(), TEXT("Iad.Fill="), ParamValue))
+	{
+		ParamValue.TrimStartAndEndInline();
+		int64 FillSize = -1;
+		LexFromString(FillSize, ParamValue);
+
+		if (FillSize > 0)
+		{
+			if (ParamValue.EndsWith(TEXT("GB")))
+			{
+				FillSize = FillSize << 30;
+			}
+			if (ParamValue.EndsWith(TEXT("MB")))
+			{
+				FillSize = FillSize << 20;
+			}
+
+			UE_LOG(LogIoStoreOnDemand, Log, TEXT("Filling install cache with %.2lf MiB of dummy data"), double(FillSize) / 1024.0 / 1024.0);
+
+			FIoStatus	Status = FIoStatus::Ok;
+			uint64		Seed = 1;
+			while (FillSize >= 0 && Status.IsOk())
+			{
+				const uint64		ChunkSize = 256 << 10;
+				FIoBuffer			Chunk(ChunkSize);
+				TArrayView<uint64>	Values(reinterpret_cast<uint64*>(Chunk.GetData()), ChunkSize / sizeof(uint64));
+
+				for (uint64& Value : Values)
+				{
+					Value = Seed;
+				}
+
+				const FIoHash ChunkHash = FIoHash::HashBuffer(Chunk.GetView());
+				Status = InstallCache->PutChunk(MoveTemp(Chunk), ChunkHash);
+				Seed++;
+				FillSize -= ChunkSize;
+			}
+
+			if (Status.IsOk())
+			{
+				Status = InstallCache->Flush();
+			}
+
+			UE_CLOG(!Status.IsOk(), LogIoStoreOnDemand, Warning, TEXT("Failed to fill install cache with dummy data"));
+		}
+	}
+#endif
+
+	return EIoErrorCode::Ok;
+}
+
 FOnDemandChunkInfo FOnDemandIoStore::GetChunkInfo(const FIoChunkId& ChunkId, EOnDemandContainerFlags ContainerFlags)
 {
 	TUniqueLock Lock(ContainerMutex);
@@ -877,6 +916,11 @@
 {
 	if (InstallCache.IsValid() == false || PackageStoreBackend.IsValid() == false)
 	{
+		if (FForkProcessHelper::IsForkRequested() && !FForkProcessHelper::IsForkedChildProcess())
+		{
+			return FIoStatusBuilder(EIoErrorCode::PendingFork) << TEXT("Install cache waiting for fork");
+		}
+
 		return FIoStatusBuilder(EIoErrorCode::InvalidCode) << TEXT("Install cache not configured");
 	}
 
diff --git a/Engine/Source/Runtime/Experimental/IoStore/OnDemand/Private/OnDemandIoStore.h b/Engine/Source/Runtime/Experimental/IoStore/OnDemand/Private/OnDemandIoStore.h
--- a/Engine/Source/Runtime/Experimental/IoStore/OnDemand/Private/OnDemandIoStore.h
+++ b/Engine/Source/Runtime/Experimental/IoStore/OnDemand/Private/OnDemandIoStore.h
@@ -14,6 +14,8 @@
 
 #include <atomic>
 
+enum class EForkProcessRole : uint8;
+
 struct FIoContainerHeader;
 using FSharedContainerHeader		= TSharedPtr<FIoContainerHeader>;
 
@@ -162,6 +164,8 @@ public:
 	TArray<FSharedOnDemandContainer> GetMountedContainers();
 
 private:
+	void					OnPostFork(EForkProcessRole ProcessRole);
+	FIoStatus				InitializeOnDemandInstallCache();
 	FOnDemandChunkInfo		GetChunkInfo(const FIoChunkId& ChunkId, EOnDemandContainerFlags ContainerFlags);
 	void					TryEnterTickLoop();
 	void					TickLoop();
