Chaos Modular Vehicles:

Add Leaf First BFS traversal option for the sim module tree.

diff --git a/Engine/Plugins/Experimental/ChaosModularVehicle/Source/ChaosModularVehicleEngine/Public/ChaosModularVehicle/ModularVehicleBaseComponent.h b/Engine/Plugins/Experimental/ChaosModularVehicle/Source/ChaosModularVehicleEngine/Public/ChaosModularVehicle/ModularVehicleBaseComponent.h
--- a/Engine/Plugins/Experimental/ChaosModularVehicle/Source/ChaosModularVehicleEngine/Public/ChaosModularVehicle/ModularVehicleBaseComponent.h
+++ b/Engine/Plugins/Experimental/ChaosModularVehicle/Source/ChaosModularVehicleEngine/Public/ChaosModularVehicle/ModularVehicleBaseComponent.h
@@ -110,8 +110,8 @@ public:
 
 	APlayerController* GetPlayerController() const;
 	bool IsLocallyControlled() const;
-	void SetTreeProcessingOrder(Chaos::ESimTreeProcessingOrder TreeProcessingOrderIn) { TreeProcessingOrder = TreeProcessingOrderIn; }
-	Chaos::ESimTreeProcessingOrder GetTreeProcessingOrder() { return TreeProcessingOrder; }
+	void SetTreeProcessingOrder(ESimTreeProcessingOrder TreeProcessingOrderIn) { TreeProcessingOrder = TreeProcessingOrderIn; }
+	ESimTreeProcessingOrder GetTreeProcessingOrder() { return TreeProcessingOrder; }
 
 	virtual void GetLifetimeReplicatedProps(TArray<FLifetimeProperty>& OutLifetimeProps) const override;
 
@@ -318,6 +318,9 @@ public:
 
 	UPROPERTY(EditAnywhere, Category = VehicleInput)
 	TArray<FModuleInputSetup> InputConfig;
+		
+	UPROPERTY(EditAnywhere, Category = "Game|Components|ModularVehicle")
+	TEnumAsByte<ESimTreeProcessingOrder> TreeProcessingOrder = ESimTreeProcessingOrder::LeafFirst;
 
 	UPROPERTY(Transient, Replicated)
 	TArray<FConstructionData> ConstructionDatas;
@@ -370,7 +373,6 @@ private:
 	int32 ClusteringCount = 0;
 
 	bool bIsLocallyControlled;
-	Chaos::ESimTreeProcessingOrder TreeProcessingOrder = Chaos::ESimTreeProcessingOrder::LeafFirst;
 
 	TArray<FModuleAnimationSetup> ModuleAnimationSetups;
 
@@ -380,6 +382,4 @@ private:
 	TObjectPtr<UVehicleInputProducerBase> InputProducer = nullptr;
 
 	FModuleInputContainer InputsContainer;
-
 };
-
diff --git a/Engine/Source/Runtime/Experimental/ChaosVehicles/ChaosVehiclesCore/Private/SimModule/SimModuleTree.cpp b/Engine/Source/Runtime/Experimental/ChaosVehicles/ChaosVehiclesCore/Private/SimModule/SimModuleTree.cpp
--- a/Engine/Source/Runtime/Experimental/ChaosVehicles/ChaosVehiclesCore/Private/SimModule/SimModuleTree.cpp
+++ b/Engine/Source/Runtime/Experimental/ChaosVehicles/ChaosVehiclesCore/Private/SimModule/SimModuleTree.cpp
@@ -236,10 +236,18 @@ void FSimModuleTree::Simulate(float DeltaTime, FAllInputs& Inputs, FClusterUnion
 		TArray<int> RootNodes;
 		GetRootNodes(RootNodes);
 
-		for (int RootIndex : RootNodes)
+		if (SimTreeProcessingOrder == ESimTreeProcessingOrder::LeafFirstBFS)
 		{
-			SimulateNode(DeltaTime, Inputs, RootIndex, PhysicsProxy);
+			SimulateNodeBFS(DeltaTime, Inputs, RootNodes, PhysicsProxy);
 		}
+		else
+		{
+			for (int RootIndex : RootNodes)
+			{
+				SimulateNode(DeltaTime, Inputs, RootIndex, PhysicsProxy);
+			}
+		}
+
 	}
 
 }
@@ -285,6 +293,55 @@ void FSimModuleTree::SimulateNode(float DeltaTime, FAllInputs& Inputs, int NodeI
 	}
 }
 
+void FSimModuleTree::SimulateNodeBFS(float DeltaTime, FAllInputs& Inputs, const TArray<int>& RootNodes, FClusterUnionPhysicsProxy* PhysicsProxy)
+{
+	TQueue<int> Queue;
+	TArray<int> Stack;
+
+	for (int Idx : RootNodes)
+	{
+		if (AccessSimModule(Idx))
+		{
+			Queue.Enqueue(Idx);
+		}
+	}
+
+	while (!Queue.IsEmpty())
+	{
+		int OutNode = -1;
+		Queue.Dequeue(OutNode);
+		if (OutNode >= 0)
+		{
+			Stack.Push(OutNode);
+			for (int Idx : GetChildren(OutNode))
+			{
+				if (AccessSimModule(Idx))
+				{
+					Queue.Enqueue(Idx);
+				}
+			}
+		}
+	}
+
+	while (!Stack.IsEmpty())
+	{
+		int Node = Stack.Pop();
+		if (ISimulationModuleBase* Module = AccessSimModule(Node))
+		{
+			if (Module->IsEnabled())
+			{
+				Module->Simulate(PhysicsProxy, DeltaTime, Inputs, *this);
+
+				if (IsAnimationEnabled() && Module->IsAnimationEnabled())
+				{
+					Module->Animate(PhysicsProxy);
+				}
+			}
+		}
+	}
+}
+
+
 void FSimModuleTree::DeleteNodesBelow(int AtIndex)
 {
 	if (IsValidNode(AtIndex))
diff --git a/Engine/Source/Runtime/Experimental/ChaosVehicles/ChaosVehiclesCore/Public/SimModule/SimModuleTree.h b/Engine/Source/Runtime/Experimental/ChaosVehicles/ChaosVehiclesCore/Public/SimModule/SimModuleTree.h
--- a/Engine/Source/Runtime/Experimental/ChaosVehicles/ChaosVehiclesCore/Public/SimModule/SimModuleTree.h
+++ b/Engine/Source/Runtime/Experimental/ChaosVehicles/ChaosVehiclesCore/Public/SimModule/SimModuleTree.h
@@ -6,9 +6,18 @@
 #include "SimModule/SimulationModuleBase.h"
 #include "SimModule/ModuleInput.h"
 
+#include "SimModuleTree.Generated.h"
 
 DECLARE_STATS_GROUP(TEXT("ModularVehicle.SimTree"), STATGROUP_ModularVehicleSimTree, STATGROUP_Advanced);
 
+UENUM(BlueprintType)
+enum ESimTreeProcessingOrder : int8
+{
+	ManualOverride = 0,	// User calls simulate on the child modules
+	LeafFirst = 1,		// modules simulation from the leaf first
+	RootFirst = 2,		// modules simulate from the root first
+	LeafFirstBFS = 3
+};
 
 class FGeometryCollectionPhysicsProxy;
 namespace Chaos
@@ -17,14 +26,6 @@ namespace Chaos
 	class FClusterUnionPhysicsProxy;
 	struct FAllInputs;
 
-
-	enum ESimTreeProcessingOrder : int8
-	{
-		ManualOverride = 0,	// User calls simulate on the child modules
-		LeafFirst = 1,		// modules simulation from the leaf first
-		RootFirst = 2		// modules simulate from the root first
-	};
-
 	struct FPendingModuleAdds
 	{
 		FPendingModuleAdds(int ParentIndexIn, ISimulationModuleBase* NewSimModuleIn)
@@ -203,6 +204,8 @@ namespace Chaos
 	protected:
 		void SimulateNode(float DeltaTime, FAllInputs& Inputs, int NodeIdx, FClusterUnionPhysicsProxy* PhysicsProxy);
 
+		void SimulateNodeBFS(float DeltaTime, FAllInputs& Inputs, const TArray<int>& RootNodes, FClusterUnionPhysicsProxy* PhysicsProxy);
+
 		void DeleteNodesBelow(int NodeIdx);
 
 		int GetNextIndex();
