Fix capital G in generated header name.

diff --git a/Engine/Source/Runtime/Experimental/ChaosVehicles/ChaosVehiclesCore/Public/SimModule/SimModuleTree.h b/Engine/Source/Runtime/Experimental/ChaosVehicles/ChaosVehiclesCore/Public/SimModule/SimModuleTree.h
--- a/Engine/Source/Runtime/Experimental/ChaosVehicles/ChaosVehiclesCore/Public/SimModule/SimModuleTree.h
+++ b/Engine/Source/Runtime/Experimental/ChaosVehicles/ChaosVehiclesCore/Public/SimModule/SimModuleTree.h
@@ -6,7 +6,7 @@
 #include "SimModule/SimulationModuleBase.h"
 #include "SimModule/ModuleInput.h"
 
-#include "SimModuleTree.Generated.h"
+#include "SimModuleTree.generated.h"
 
 DECLARE_STATS_GROUP(TEXT("ModularVehicle.SimTree"), STATGROUP_ModularVehicleSimTree, STATGROUP_Advanced);
 
