Add a commandline to set IOStoreOnDemand cache location on the server.

diff --git a/Engine/Source/Runtime/Experimental/IoStore/OnDemand/Private/OnDemandIoStore.cpp b/Engine/Source/Runtime/Experimental/IoStore/OnDemand/Private/OnDemandIoStore.cpp
--- a/Engine/Source/Runtime/Experimental/IoStore/OnDemand/Private/OnDemandIoStore.cpp
+++ b/Engine/Source/Runtime/Experimental/IoStore/OnDemand/Private/OnDemandIoStore.cpp
@@ -70,6 +70,27 @@ static FString GetInstallCacheDirectory()
 				UE_LOG(LogIoStoreOnDemand, Fatal, TEXT("Attempting to create IOStore cache before forking!"));
 			}
 
+			FString CommandLineDir;
+			bool bUsePathFromCommandLine = FParse::Value(FCommandLine::Get(), TEXT("ServerIOInstallCacheDir="), CommandLineDir);
+			if (bUsePathFromCommandLine)
+			{
+				if (!FPaths::ValidatePath(CommandLineDir))
+				{
+					bUsePathFromCommandLine = false;
+					UE_LOG(LogIoStoreOnDemand, Error, TEXT("Invalid ServerIOInstallCacheDir from command line: %s"), *CommandLineDir);
+				}
+				else if (!FPaths::IsRelative(CommandLineDir))
+				{
+					bUsePathFromCommandLine = false;
+					UE_LOG(LogIoStoreOnDemand, Error, TEXT("ServerIOInstallCacheDir from command line is not relative: %s"), *CommandLineDir);
+				}
+
+				if (bUsePathFromCommandLine)
+				{
+					return FPaths::ProjectPersistentDownloadDir() / CommandLineDir;
+				}
+			}
+
 			DirName = FString::Printf(TEXT("InstallCacheServer-%u"), FPlatformProcess::GetCurrentProcessId());
 		}
 	}
