Fix to ensure that under AutoRTFM, chaos locks interact correctly with transactions.

I've done this by:
- Making `FPhysSceneLock` transactionally safe under AutoRTFM.
- We pessimise taking the lock in a transaction to lock the write lock for the length of the transaction.
- Which also fixes the RegisterComponent/UnregisterComponent bug on UPrimitiveComponent whereby we had an incorrect on-commit that caused physics glitches.

#rnx

diff --git a/Engine/Source/Runtime/Core/Public/HAL/LowLevelMemTracker.h b/Engine/Source/Runtime/Core/Public/HAL/LowLevelMemTracker.h
--- a/Engine/Source/Runtime/Core/Public/HAL/LowLevelMemTracker.h
+++ b/Engine/Source/Runtime/Core/Public/HAL/LowLevelMemTracker.h
@@ -5,6 +5,7 @@
 #include "CoreTypes.h"
 #include "LowLevelMemTrackerDefines.h" // LLM_ENABLED_IN_CONFIG
 #include "ProfilingDebugging/TagTrace.h"
+#include "AutoRTFM/AutoRTFM.h"
 
 #ifndef PLATFORM_SUPPORTS_LLM
 #define PLATFORM_SUPPORTS_LLM 1
@@ -901,23 +902,32 @@ public:
 	{
 		if (FLowLevelMemTracker::IsEnabled())
 		{
-			Init(TagName, bIsStatTag, InTagSet, InTracker, bOverride);
+			// We run the init in the open, because we want to track LLM even in transactions.
+			UE_AUTORTFM_OPEN({ Init(TagName, bIsStatTag, InTagSet, InTracker, bOverride); });
+			// But remember that if we abort while we hold the scope, we need to destroy the scope.
+			AutoRTFM::PushOnAbortHandler(this, [this] { if (bEnabled) { Destruct(); } });
 		}
 	}
+
 	FLLMScope(ELLMTag TagEnum, bool bIsStatTag, ELLMTagSet InTagSet, ELLMTracker InTracker, bool bOverride = true)
 	{
 		if (FLowLevelMemTracker::IsEnabled())
 		{
-			Init(TagEnum, bIsStatTag, InTagSet, InTracker, bOverride);
+			// We run the init in the open, because we want to track LLM even in transactions.
+			UE_AUTORTFM_OPEN({ Init(TagEnum, bIsStatTag, InTagSet, InTracker, bOverride); });
+			// But remember that if we abort while we hold the scope, we need to destroy the scope.
+			AutoRTFM::PushOnAbortHandler(this, [this] { if (bEnabled) { Destruct(); } });
 		}
-
 	}
 
 	FLLMScope(const UE::LLMPrivate::FTagData* TagData, bool bIsStatTag, ELLMTagSet Set, ELLMTracker Tracker, bool bOverride = true)
 	{
 		if (FLowLevelMemTracker::IsEnabled())
 		{
-			Init(TagData, bIsStatTag, Set, Tracker, bOverride);
+			// We run the init in the open, because we want to track LLM even in transactions.
+			UE_AUTORTFM_OPEN({ Init(TagData, bIsStatTag, Set, Tracker, bOverride); });
+			// But remember that if we abort while we hold the scope, we need to destroy the scope.
+			AutoRTFM::PushOnAbortHandler(this, [this] { if (bEnabled) { Destruct(); } });
 		}
 	}
 
@@ -925,7 +935,10 @@ public:
 	{
 		if (bEnabled)
 		{
-			Destruct();
+			// We run the destroy in the open, because we want to track LLM even in transactions.
+			UE_AUTORTFM_OPEN({ Destruct(); });
+			// But remember to pop our on-abort handler (so that we don't try and double destroy).
+			AutoRTFM::PopOnAbortHandler(this);
 		}
 	}
 
diff --git a/Engine/Source/Runtime/Engine/Private/Components/ActorComponent.cpp b/Engine/Source/Runtime/Engine/Private/Components/ActorComponent.cpp
--- a/Engine/Source/Runtime/Engine/Private/Components/ActorComponent.cpp
+++ b/Engine/Source/Runtime/Engine/Private/Components/ActorComponent.cpp
@@ -1819,25 +1819,22 @@ void UActorComponent::ExecuteRegisterEvents(FRegisterComponentContext* Context)
 void UActorComponent::ExecuteUnregisterEvents()
 {
 	// Delay the destroying the physics, as well as the rendering state until after we commit
-	UE_AUTORTFM_ONCOMMIT2(=)
-	{
-		DestroyPhysicsState();
+	DestroyPhysicsState();
 
-		if (bRenderStateCreated)
-		{
-			SCOPE_CYCLE_COUNTER(STAT_ComponentDestroyRenderState);
-			checkf(bRegistered, TEXT("Component has render state when not registered (%s)"), *GetFullName());
-			DestroyRenderState_Concurrent();
-			checkf(!bRenderStateCreated, TEXT("Failed to route DestroyRenderState_Concurrent (%s)"), *GetFullName());
-		}
+	if (bRenderStateCreated)
+	{
+		SCOPE_CYCLE_COUNTER(STAT_ComponentDestroyRenderState);
+		checkf(bRegistered, TEXT("Component has render state when not registered (%s)"), *GetFullName());
+		DestroyRenderState_Concurrent();
+		checkf(!bRenderStateCreated, TEXT("Failed to route DestroyRenderState_Concurrent (%s)"), *GetFullName());
+	}
 
-		if (bRegistered)
-		{
-			SCOPE_CYCLE_COUNTER(STAT_ComponentOnUnregister);
-			OnUnregister();
-			checkf(!bRegistered, TEXT("Failed to route OnUnregister (%s)"), *GetFullName());
-		}
-	};
+	if (bRegistered)
+	{
+		SCOPE_CYCLE_COUNTER(STAT_ComponentOnUnregister);
+		OnUnregister();
+		checkf(!bRegistered, TEXT("Failed to route OnUnregister (%s)"), *GetFullName());
+	}
 }
 
 void UActorComponent::ReregisterComponent()
diff --git a/Engine/Source/Runtime/Engine/Private/Tests/AutoRTFM/AutoRTFMActorComponentTests.cpp b/Engine/Source/Runtime/Engine/Private/Tests/AutoRTFM/AutoRTFMActorComponentTests.cpp
new file mode 100644
--- /dev/null
+++ b/Engine/Source/Runtime/Engine/Private/Tests/AutoRTFM/AutoRTFMActorComponentTests.cpp
@@ -0,0 +1,110 @@
+// Copyright Epic Games, Inc. All Rights Reserved.
+
+#include "Misc/AutomationTest.h"
+#include "Misc/App.h"
+#include "AutoRTFM/AutoRTFM.h"
+#include "AutoRTFMTestActor.h"
+#include "AutoRTFMTestLevel.h"
+#include "AutoRTFMTestPrimitiveComponent.h"
+#include "PhysicsProxy/SingleParticlePhysicsProxy.h"
+#include "Physics/Experimental/PhysScene_Chaos.h"
+#include "PBDRigidsSolver.h"
+#include "Chaos/Core.h"
+#include "Engine/World.h"
+
+#if WITH_DEV_AUTOMATION_TESTS
+
+#define TEST_CHECK_TRUE(b) do \
+{ \
+	if (!(b)) \
+	{ \
+		FString String(FString::Printf(TEXT("FAILED: %s:%u"), TEXT(__FILE__), __LINE__)); \
+		ExecutionInfo.AddEvent(FAutomationEvent(EAutomationEventType::Info, String)); \
+		return false; \
+	} \
+} while(false)
+
+IMPLEMENT_SIMPLE_AUTOMATION_TEST(FAutoRTFMActorComponentTests, "AutoRTFM + UActorComponent", EAutomationTestFlags::EngineFilter | EAutomationTestFlags::ClientContext | EAutomationTestFlags::ServerContext | EAutomationTestFlags::CommandletContext)
+
+bool FAutoRTFMActorComponentTests::RunTest(const FString & Parameters)
+{
+	if (!AutoRTFM::ForTheRuntime::IsAutoRTFMRuntimeEnabled())
+	{
+		ExecutionInfo.AddEvent(FAutomationEvent(EAutomationEventType::Info, TEXT("SKIPPED 'FAutoRTFMActorComponentTests' test. AutoRTFM disabled.")));
+		return true;
+	}
+
+	UWorld* World = NewObject<UWorld>();
+	World->CreatePhysicsScene(nullptr);
+
+	UAutoRTFMTestLevel* Level = NewObject<UAutoRTFMTestLevel>();
+	Level->OwningWorld = World;
+	AAutoRTFMTestActor* Actor = NewObject<AAutoRTFMTestActor>(Level);
+	UAutoRTFMTestPrimitiveComponent* Component = NewObject<UAutoRTFMTestPrimitiveComponent>(Actor);
+
+	// Helper to re-add actor handles after each unregister.
+	auto ResetActorHandle = [&]
+		{
+			Component->BodyInstance.ActorHandle = Chaos::FSingleParticlePhysicsProxy::Create(Chaos::FGeometryParticle::CreateParticle());
+			Component->BodyInstance.ActorHandle->GetParticle_LowLevel()->SetGeometry(MakeImplicitObjectPtr<Chaos::FSphere>(Chaos::FVec3(1, 2, 3), 1));
+			World->GetPhysicsScene()->GetSolver()->RegisterObject(Component->BodyInstance.ActorHandle);
+		};
+
+	ResetActorHandle();
+
+	Component->RegisterComponent();
+	TEST_CHECK_TRUE(Component->IsRegistered());
+	Component->UnregisterComponent();
+	TEST_CHECK_TRUE(!Component->IsRegistered());
+	TEST_CHECK_TRUE(!Component->BodyInstance.ActorHandle);
+
+	ResetActorHandle();
+
+	AutoRTFM::ETransactionResult Result = AutoRTFM::Transact([&]
+		{
+			Component->RegisterComponent();
+
+			if (Component->IsRegistered())
+			{
+				AutoRTFM::AbortTransaction();
+			}
+		});
+
+	TEST_CHECK_TRUE(AutoRTFM::ETransactionResult::AbortedByRequest == Result);
+	TEST_CHECK_TRUE(!Component->IsRegistered());
+
+	bool bWasRegistered = false;
+
+	Result = AutoRTFM::Transact([&]
+		{
+			Component->RegisterComponent();
+			bWasRegistered = Component->IsRegistered();
+		});
+
+	TEST_CHECK_TRUE(AutoRTFM::ETransactionResult::Committed == Result);
+	TEST_CHECK_TRUE(bWasRegistered);
+	TEST_CHECK_TRUE(Component->IsRegistered());
+
+	Result = AutoRTFM::Transact([&]
+		{
+			Component->UnregisterComponent();
+			AutoRTFM::AbortTransaction();
+		});
+
+	TEST_CHECK_TRUE(AutoRTFM::ETransactionResult::AbortedByRequest == Result);
+	TEST_CHECK_TRUE(Component->IsRegistered());
+
+	Result = AutoRTFM::Transact([&]
+		{
+			Component->UnregisterComponent();
+		});
+
+	TEST_CHECK_TRUE(AutoRTFM::ETransactionResult::Committed == Result);
+	TEST_CHECK_TRUE(!Component->IsRegistered());
+
+	return true;
+}
+
+#undef TEST_CHECK_TRUE
+
+#endif //WITH_DEV_AUTOMATION_TESTS
diff --git a/Engine/Source/Runtime/Engine/Private/Tests/AutoRTFM/AutoRTFMTestActor.h b/Engine/Source/Runtime/Engine/Private/Tests/AutoRTFM/AutoRTFMTestActor.h
new file mode 100644
--- /dev/null
+++ b/Engine/Source/Runtime/Engine/Private/Tests/AutoRTFM/AutoRTFMTestActor.h
@@ -0,0 +1,15 @@
+// Copyright Epic Games, Inc. All Rights Reserved.
+
+#pragma once
+
+#include "GameFramework/Actor.h"
+#include "AutoRTFMTestActor.generated.h"
+
+UCLASS()
+class AAutoRTFMTestActor : public AActor
+{
+    GENERATED_BODY()
+
+public:
+    int Value = 42;
+};
diff --git a/Engine/Source/Runtime/Engine/Private/Tests/AutoRTFM/AutoRTFMTestLevel.h b/Engine/Source/Runtime/Engine/Private/Tests/AutoRTFM/AutoRTFMTestLevel.h
new file mode 100644
--- /dev/null
+++ b/Engine/Source/Runtime/Engine/Private/Tests/AutoRTFM/AutoRTFMTestLevel.h
@@ -0,0 +1,15 @@
+// Copyright Epic Games, Inc. All Rights Reserved.
+
+#pragma once
+
+#include "Engine/Level.h"
+#include "AutoRTFMTestLevel.generated.h"
+
+UCLASS()
+class UAutoRTFMTestLevel : public ULevel
+{
+    GENERATED_BODY()
+
+public:
+    int Value = 42;
+};
diff --git a/Engine/Source/Runtime/Engine/Private/Tests/AutoRTFM/AutoRTFMTestPrimitiveComponent.h b/Engine/Source/Runtime/Engine/Private/Tests/AutoRTFM/AutoRTFMTestPrimitiveComponent.h
new file mode 100644
--- /dev/null
+++ b/Engine/Source/Runtime/Engine/Private/Tests/AutoRTFM/AutoRTFMTestPrimitiveComponent.h
@@ -0,0 +1,15 @@
+// Copyright Epic Games, Inc. All Rights Reserved.
+
+#pragma once
+
+#include "Components/PrimitiveComponent.h"
+#include "AutoRTFMTestPrimitiveComponent.generated.h"
+
+UCLASS()
+class UAutoRTFMTestPrimitiveComponent : public UPrimitiveComponent
+{
+    GENERATED_BODY()
+
+public:
+    int Value = 42;
+};
diff --git a/Engine/Source/Runtime/Engine/Private/Tests/AutoRTFM/AutoRTFMTraceFilterTests.cpp b/Engine/Source/Runtime/Engine/Private/Tests/AutoRTFM/AutoRTFMTraceFilterTests.cpp
--- a/Engine/Source/Runtime/Engine/Private/Tests/AutoRTFM/AutoRTFMTraceFilterTests.cpp
+++ b/Engine/Source/Runtime/Engine/Private/Tests/AutoRTFM/AutoRTFMTraceFilterTests.cpp
@@ -154,4 +154,6 @@ bool FAutoRTFMTraceFilterTests::RunTest(const FString & Parameters)
 #endif
 }
 
+#undef TEST_CHECK_TRUE
+
 #endif //WITH_DEV_AUTOMATION_TESTS
diff --git a/Engine/Source/Runtime/Experimental/Chaos/Public/Framework/Threading.h b/Engine/Source/Runtime/Experimental/Chaos/Public/Framework/Threading.h
--- a/Engine/Source/Runtime/Experimental/Chaos/Public/Framework/Threading.h
+++ b/Engine/Source/Runtime/Experimental/Chaos/Public/Framework/Threading.h
@@ -11,6 +11,7 @@
 #include "Async/ParallelFor.h"
 #include <atomic>
 #include "HAL/CriticalSection.h"
+#include "AutoRTFM/AutoRTFM.h"
 
 #ifndef PHYSICS_THREAD_CONTEXT
 	#if (!UE_BUILD_SHIPPING && !UE_BUILD_TEST)
@@ -930,15 +931,140 @@ FORCEINLINE void EnsureIsInGameThreadContext()
 	};
 
 #if CHAOS_SCENE_LOCK_TYPE == CHAOS_SCENE_LOCK_SCENE_GUARD
-	using FPhysSceneLock = FPhysicsSceneGuard;
+	using FPhysSceneLockNonTransactional = FPhysicsSceneGuard;
 #elif CHAOS_SCENE_LOCK_TYPE == CHAOS_SCENE_LOCK_RWFIFO_SPINLOCK
-	using FPhysSceneLock = TRwFifoLock<FPhysSpinLock>;
+	using FPhysSceneLockNonTransactional = TRwFifoLock<FPhysSpinLock>;
 #elif CHAOS_SCENE_LOCK_TYPE == CHAOS_SCENE_LOCK_RWFIFO_CRITICALSECTION
-	using FPhysSceneLock = TRwFifoLock<FCriticalSection>;
+	using FPhysSceneLockNonTransactional = TRwFifoLock<FCriticalSection>;
 #elif CHAOS_SCENE_LOCK_TYPE == CHAOS_SCENE_LOCK_FRWLOCK
-	using FPhysSceneLock = FPhysicsRwLock;
+	using FPhysSceneLockNonTransactional = FPhysicsRwLock;
 #elif CHAOS_SCENE_LOCK_TYPE == CHAOS_SCENE_LOCK_SIMPLE_MUTEX
-	using FPhysSceneLock = FPhysicsSimpleMutexLock;
+	using FPhysSceneLockNonTransactional = FPhysicsSimpleMutexLock;
+#endif
+
+#if UE_AUTORTFM
+	// A transactionally safe lock that works in the following novel ways:
+	// - In the open (non-transactional):
+	//   - Take the lock like before. Simple!
+	//   - Free the lock like before too.
+	// - In the closed (transactional):
+	//   - During locking we query `TransactionalLockCount`:
+	//	   - 0 means we haven't taken the lock within our transaction nest and need to acquire the lock.
+	//     - Otherwise we already have the lock (and are preventing non-transactional code seeing any
+	//       modifications we've made while holding the lock), so just bump `TransactionalLockCount`.
+	//     - We also register an on-abort handler to release the lock should we abort (but we need to
+	//       query `TransactionalLockCount` even there because we could be aborting an inner transaction
+	//       and the parent transaction still wants to have the lock held!).
+	//   - During unlocking we defer doing the unlock until the transaction commits.
+	//
+	// Thus with this approach we will hold this lock for the *entirety* of the transactional nest should
+	// we take the lock during the transaction, thus preventing non-transactional code from seeing any
+	// modifications we should make.
+	//
+	// If we are within a transaction, we pessimise our read-lock to a write-lock. Note: that it should
+	// potentially be possible to have read-locks work correctly, but serious care will have to be taken to
+	// ensure that we don't have:
+	//   Open Thread     Closed Thread
+	//   -----------     ReadLock
+	//   -----------     ReadUnlock
+	//   WriteLock       -------------
+	//   WriteUnlock     -------------
+	//   -----------     ReadLock      <- Invalid because the transaction can potentially observe side
+	//                                    effects of the open-threads writes!
+	struct FPhysSceneLockTransactionallySafe final
+	{
+		void ReadLock()
+		{
+			if (AutoRTFM::IsTransactional())
+			{
+				// Transactionally pessimise ReadLock -> WriteLock.
+				WriteLock();
+			}
+			else
+			{
+				Lock.ReadLock();
+				check(0 == TransactionalLockCount);
+			}
+		}
+
+		void ReadUnlock()
+		{
+			if (AutoRTFM::IsTransactional())
+			{
+				// Transactionally pessimise ReadUnlock -> WriteUnlock.
+				WriteUnlock();
+			}
+			else
+			{
+				check(0 == TransactionalLockCount);
+				Lock.ReadUnlock();
+			}
+		}
+
+		void WriteLock()
+		{
+			if (AutoRTFM::IsTransactional())
+			{
+				AutoRTFM::Open([&]
+					{
+						// The transactional system which can increment TransactionalLockCount
+						// is always single-threaded, thus this is safe to check without atomicity.
+						if (0 == TransactionalLockCount)
+						{
+							Lock.WriteLock();
+						}
+
+						TransactionalLockCount += 1;
+					});
+
+				AutoRTFM::OnAbort([this]
+					{
+						check(0 != TransactionalLockCount);
+						TransactionalLockCount -= 1;
+
+						if (0 == TransactionalLockCount)
+						{
+							Lock.WriteUnlock();
+						}
+					});
+			}
+			else
+			{
+				Lock.WriteLock();
+				check(0 == TransactionalLockCount);
+			}
+		}
+
+		void WriteUnlock()
+		{
+			if (AutoRTFM::IsTransactional())
+			{
+				AutoRTFM::OnCommit([this]
+					{
+						check(0 != TransactionalLockCount);
+						TransactionalLockCount -= 1;
+
+						if (0 == TransactionalLockCount)
+						{
+							Lock.WriteUnlock();
+						}
+					});
+			}
+			else
+			{
+				check(0 == TransactionalLockCount);
+				Lock.WriteUnlock();
+			}
+		}
+
+	private:
+		FPhysSceneLockNonTransactional Lock;
+		uint32 TransactionalLockCount;
+	};
+
+	using FPhysSceneLock = FPhysSceneLockTransactionallySafe;
+#else
+	using FPhysSceneLock = FPhysSceneLockNonTransactional;
 #endif
 
 	// Stable types to use in calling code configured by the compiler switches above
