Open around ISPC calls into transform functions in Chaos

diff --git a/Engine/Source/Runtime/Experimental/Chaos/Private/Chaos/AABB.cpp b/Engine/Source/Runtime/Experimental/Chaos/Private/Chaos/AABB.cpp
--- a/Engine/Source/Runtime/Experimental/Chaos/Private/Chaos/AABB.cpp
+++ b/Engine/Source/Runtime/Experimental/Chaos/Private/Chaos/AABB.cpp
@@ -3,6 +3,7 @@
 #include "Chaos/Real.h"
 #include "Chaos/Sphere.h"
 #include "Chaos/Capsule.h"
+#include "AutoRTFM/AutoRTFM.h"
 #if INTEL_ISPC
 #include "AABB.ispc.generated.h"
 
@@ -409,7 +410,10 @@ inline TAABB<FReal, 3> TransformedAABBHelperISPC(const TAABB<FReal, 3>& AABB, co
 	check(bRealTypeCompatibleWithISPC);
 #if INTEL_ISPC
 	TVector<Chaos::FReal, 3> NewMin, NewMax;
-	ispc::TransformedAABB((const ispc::FTransform&)SpaceTransform, (const ispc::FVector&)AABB.Min(), (const ispc::FVector&)AABB.Max(), (ispc::FVector&)NewMin, (ispc::FVector&)NewMax);
+	UE_AUTORTFM_OPEN(
+	{
+		ispc::TransformedAABB((const ispc::FTransform&)SpaceTransform, (const ispc::FVector&)AABB.Min(), (const ispc::FVector&)AABB.Max(), (ispc::FVector&)NewMin, (ispc::FVector&)NewMax);
+	});
 
 	TAABB<FReal, 3> NewAABB(NewMin, NewMax);
 	return NewAABB;
@@ -432,7 +436,10 @@ inline TAABB<Chaos::FRealSingle, 3> TransformedAABBHelperISPC(const TAABB<Chaos:
 	static_assert(sizeof(ispc::FVector3f) == sizeof(Chaos::TVector<Chaos::FRealSingle, 3>), "sizeof(ispc::FVector3f) != sizeof(Chaos::TVector<Chaos::FRealSingle, 3>)");
 
 	TVector<Chaos::FRealSingle, 3> NewMin, NewMax;
-	ispc::TransformedAABBMixed((const ispc::FTransform&)SpaceTransform, (const ispc::FVector3f&)AABB.Min(), (const ispc::FVector3f&)AABB.Max(), (ispc::FVector3f&)NewMin, (ispc::FVector3f&)NewMax);
+	UE_AUTORTFM_OPEN(
+	{
+		ispc::TransformedAABBMixed((const ispc::FTransform&)SpaceTransform, (const ispc::FVector3f&)AABB.Min(), (const ispc::FVector3f&)AABB.Max(), (ispc::FVector3f&)NewMin, (ispc::FVector3f&)NewMax);
+	});
 
 	TAABB<Chaos::FRealSingle, 3> NewAABB(NewMin, NewMax);
 	return NewAABB;
@@ -453,7 +460,10 @@ inline TAABB<FReal, 3> TransformedAABBHelperISPC2(const TAABB<FReal, 3>& AABB, c
 	check(bRealTypeCompatibleWithISPC);
 #if INTEL_ISPC
 	TVector<Chaos::FReal, 3> NewMin, NewMax;
-	ispc::TransformedAABB2((const ispc::FTransform&)SpaceTransform, (const ispc::FVector&)AABB.Min(), (const ispc::FVector&)AABB.Max(), (ispc::FVector&)NewMin, (ispc::FVector&)NewMax);
+	UE_AUTORTFM_OPEN(
+	{
+		ispc::TransformedAABB2((const ispc::FTransform&)SpaceTransform, (const ispc::FVector&)AABB.Min(), (const ispc::FVector&)AABB.Max(), (ispc::FVector&)NewMin, (ispc::FVector&)NewMax);
+	});
 
 	TAABB<FReal, 3> NewAABB(NewMin, NewMax);
 	return NewAABB;
@@ -476,7 +486,10 @@ inline TAABB<Chaos::FRealSingle, 3> TransformedAABBHelperISPC2(const TAABB<Chaos
 	static_assert(sizeof(ispc::FVector3f) == sizeof(Chaos::TVector<Chaos::FRealSingle, 3>), "sizeof(ispc::FVector3f) != sizeof(Chaos::TVector<Chaos::FRealSingle, 3>)");
 
 	TVector<Chaos::FRealSingle, 3> NewMin, NewMax;
-	ispc::TransformedAABBMixed((const ispc::FTransform&)SpaceTransform, (const ispc::FVector3f&)AABB.Min(), (const ispc::FVector3f&)AABB.Max(), (ispc::FVector3f&)NewMin, (ispc::FVector3f&)NewMax);
+	UE_AUTORTFM_OPEN(
+	{
+		ispc::TransformedAABBMixed((const ispc::FTransform&)SpaceTransform, (const ispc::FVector3f&)AABB.Min(), (const ispc::FVector3f&)AABB.Max(), (ispc::FVector3f&)NewMin, (ispc::FVector3f&)NewMax);
+	});
 
 	TAABB<Chaos::FRealSingle, 3> NewAABB(NewMin, NewMax);
 	return NewAABB;
