Download soft package references when installing on-demand I/O store content.

#rnx

diff --git a/Engine/Source/Runtime/Experimental/IoStore/OnDemand/Private/OnDemandIoStore.cpp b/Engine/Source/Runtime/Experimental/IoStore/OnDemand/Private/OnDemandIoStore.cpp
--- a/Engine/Source/Runtime/Experimental/IoStore/OnDemand/Private/OnDemandIoStore.cpp
+++ b/Engine/Source/Runtime/Experimental/IoStore/OnDemand/Private/OnDemandIoStore.cpp
@@ -126,14 +126,16 @@ static void SplitHostUrl(const FStringView& Url, FStringView& OutHost, FStringVi
 }
 
 ///////////////////////////////////////////////////////////////////////////////
-using FPackageStoreEntryMap = TMap<FPackageId, const FFilePackageStoreEntry*>;
+using FPackageStoreEntryMap		= TMap<FPackageId, const FFilePackageStoreEntry*>;
+using FSoftPackageReferenceMap	= TMap<FPackageId, const FFilePackageStoreEntrySoftReferences*>;
 
 struct FContainerInstallData
 {
-	FPackageStoreEntryMap	PackageStoreEntries;
-	TSet<FPackageId>		PackageIds;
-	TSet<FIoChunkId>		ResolvedChunks;
-	uint64					TotalSize = 0;
+	FPackageStoreEntryMap		PackageStoreEntries;
+	FSoftPackageReferenceMap	SoftPackageReferences;
+	TSet<FPackageId>			PackageIds;
+	TSet<FIoChunkId>			ResolvedChunks;
+	uint64						TotalSize = 0;
 };
 
 using FInstallData = TMap<FSharedOnDemandContainer, FContainerInstallData>;
@@ -149,14 +151,14 @@ FIoStatus BuildInstallData(
 	// Setup the package information for each container
 	for (const FSharedOnDemandContainer& Container : Containers)
 	{
-		FContainerInstallData& Data = OutInstallData.FindOrAdd(Container);
-
 		if (!Container->Header.IsValid() || Container->Header->PackageIds.IsEmpty())
 		{
 			// The container contains no package data
 			continue;
 		}
 
+		FContainerInstallData& Data = OutInstallData.FindOrAdd(Container);
+
 		const FIoContainerHeader& Header = *Container->Header;
 		TConstArrayView<FFilePackageStoreEntry> Entries(
 			reinterpret_cast<const FFilePackageStoreEntry*>(Header.StoreEntries.GetData()),
@@ -164,11 +166,26 @@ FIoStatus BuildInstallData(
 		
 		Data.PackageStoreEntries.Reserve(Header.PackageIds.Num());
 
-		int32 Idx = 0;
-		for (const FFilePackageStoreEntry& Entry : Entries)
+		TConstArrayView<FFilePackageStoreEntrySoftReferences> SoftReferences;
+		if (Header.SoftPackageReferences.bContainsSoftPackageReferences)
 		{
-			const FPackageId PackageId = Header.PackageIds[Idx++];
+			Data.PackageStoreEntries.Reserve(Header.PackageIds.Num());
+			SoftReferences = MakeArrayView<const FFilePackageStoreEntrySoftReferences>(
+				reinterpret_cast<const FFilePackageStoreEntrySoftReferences*>(Header.SoftPackageReferences.PackageIndices.GetData()),
+				Header.PackageIds.Num());
+		}
+
+		for (int32 PackageIndex = 0; const FFilePackageStoreEntry& Entry : Entries)
+		{
+			const FPackageId PackageId = Header.PackageIds[PackageIndex];
 			Data.PackageStoreEntries.Add(PackageId, &Entry);
+
+			if (SoftReferences.IsEmpty() == false)
+			{
+				Data.SoftPackageReferences.Add(PackageId, &SoftReferences[PackageIndex]);
+			}
+
+			++PackageIndex;
 		}
 	}
 
@@ -202,13 +219,12 @@ FIoStatus BuildInstallData(
 		for (TPair<FSharedOnDemandContainer, FContainerInstallData>& Kv : OutInstallData)
 		{
 			FContainerInstallData& Data = Kv.Value;
-			if (const FFilePackageStoreEntry** Entry = Data.PackageStoreEntries.Find(PackageId))
+			if (const FFilePackageStoreEntry* Entry = Data.PackageStoreEntries.FindRef(PackageId))
 			{
 				Data.PackageIds.Add(PackageId);
 
-				// Add all imported packages not already processed
-				const FFilePackageStoreEntry& PackageStoreEntry = **Entry;
-				for (const FPackageId& ImportedPackageId : PackageStoreEntry.ImportedPackages)
+				// Add hard references 
+				for (const FPackageId& ImportedPackageId : Entry->ImportedPackages)
 				{
 					if (!Visitied.Contains(ImportedPackageId))
 					{
@@ -216,6 +232,20 @@ FIoStatus BuildInstallData(
 					}
 				}
 
+				// Add soft references 
+				const FIoContainerHeader& Header = *Kv.Key->Header;
+				if (const FFilePackageStoreEntrySoftReferences* SoftRefs = Data.SoftPackageReferences.FindRef(PackageId))
+				{
+					for (uint32 Index : SoftRefs->Indices)
+					{
+						const FPackageId SoftPackageReference = Header.SoftPackageReferences.PackageIds[int32(Index)];
+						if (!Visitied.Contains(SoftPackageReference))
+						{
+							Queue.Add(SoftPackageReference);
+						}
+					}
+				}
+
 				bFound = true;
 				break;
 			}
