New Blackboard interface to share data between modules

diff --git a/Engine/Source/Runtime/Experimental/ChaosVehicles/ChaosVehiclesCore/Private/SimModule/SimulationModuleBase.cpp b/Engine/Source/Runtime/Experimental/ChaosVehicles/ChaosVehiclesCore/Private/SimModule/SimulationModuleBase.cpp
--- a/Engine/Source/Runtime/Experimental/ChaosVehicles/ChaosVehiclesCore/Private/SimModule/SimulationModuleBase.cpp
+++ b/Engine/Source/Runtime/Experimental/ChaosVehicles/ChaosVehiclesCore/Private/SimModule/SimulationModuleBase.cpp
@@ -63,6 +63,11 @@ ISimulationModuleBase* ISimulationModuleBase::GetFirstChild()
 	return nullptr;
 }
 
+FVehicleBlackboard* ISimulationModuleBase::GetSimBlackboard()
+{
+	return SimModuleTree ? SimModuleTree->GetSimBlackboard() : nullptr;
+}
+
 FPBDRigidClusteredParticleHandle* ISimulationModuleBase::GetClusterParticle(Chaos::FClusterUnionPhysicsProxy* Proxy)
 { 
 	if (ParticleIdx.IsValid() && CachedParticle && (CachedParticle->UniqueIdx() == ParticleIdx))
diff --git a/Engine/Source/Runtime/Experimental/ChaosVehicles/ChaosVehiclesCore/Private/SimModule/VehicleBlackboard.cpp b/Engine/Source/Runtime/Experimental/ChaosVehicles/ChaosVehiclesCore/Private/SimModule/VehicleBlackboard.cpp
new file mode 100644
--- /dev/null
+++ b/Engine/Source/Runtime/Experimental/ChaosVehicles/ChaosVehiclesCore/Private/SimModule/VehicleBlackboard.cpp
@@ -0,0 +1,24 @@
+// Copyright Epic Games, Inc. All Rights Reserved.
+
+
+#include "SimModule/VehicleBlackboard.h"
+
+
+void FVehicleBlackboard::Invalidate(FName ObjName)
+{
+	ObjectsByName.Remove(ObjName);
+}
+
+void FVehicleBlackboard::Invalidate(EInvalidationReason Reason)
+{
+	switch (Reason)
+	{
+		default:
+		case EInvalidationReason::FullReset:
+			ObjectsByName.Empty();
+			break;
+
+		// TODO: Support other reasons
+	}
+}
+
diff --git a/Engine/Source/Runtime/Experimental/ChaosVehicles/ChaosVehiclesCore/Public/SimModule/SimModuleTree.h b/Engine/Source/Runtime/Experimental/ChaosVehicles/ChaosVehiclesCore/Public/SimModule/SimModuleTree.h
--- a/Engine/Source/Runtime/Experimental/ChaosVehicles/ChaosVehiclesCore/Public/SimModule/SimModuleTree.h
+++ b/Engine/Source/Runtime/Experimental/ChaosVehicles/ChaosVehiclesCore/Public/SimModule/SimModuleTree.h
@@ -5,6 +5,7 @@
 #include "SimModule/DeferredForcesModular.h"
 #include "SimModule/SimulationModuleBase.h"
 #include "SimModule/ModuleInput.h"
+#include "SimModule/VehicleBlackboard.h"
 
 #include "SimModuleTree.generated.h"
 
@@ -113,10 +114,12 @@ namespace Chaos
 			bAnimationEnabled = true;
 			bSimulationEnabled = true;
 			SimTreeProcessingOrder = ESimTreeProcessingOrder::LeafFirst;
+			SimBlackboard = MakeUnique<FVehicleBlackboard>();
 		}
 
 		~FSimModuleTree()
 		{
+			SimBlackboard.Reset();
 			DeleteNodesBelow(0);
 		}
 
@@ -201,6 +204,11 @@ namespace Chaos
 		void SetSimState(const Chaos::FModuleNetDataArray& ModuleDatas);
 		void InterpolateState(const float LerpFactor, Chaos::FModuleNetDataArray& LerpDatas, const Chaos::FModuleNetDataArray& MinDatas, const Chaos::FModuleNetDataArray& MaxDatas);
 
+		FVehicleBlackboard* GetSimBlackboard()
+		{
+			return SimBlackboard.Get();
+		}
+
 	protected:
 		void SimulateNode(float DeltaTime, FAllInputs& Inputs, int NodeIdx, FClusterUnionPhysicsProxy* PhysicsProxy);
 
@@ -226,6 +234,8 @@ namespace Chaos
 
 		FVehicleState VehicleState;
 		ESimTreeProcessingOrder SimTreeProcessingOrder;
+
+		TUniquePtr<FVehicleBlackboard> SimBlackboard;
 	};
 
 
diff --git a/Engine/Source/Runtime/Experimental/ChaosVehicles/ChaosVehiclesCore/Public/SimModule/SimulationModuleBase.h b/Engine/Source/Runtime/Experimental/ChaosVehicles/ChaosVehiclesCore/Public/SimModule/SimulationModuleBase.h
--- a/Engine/Source/Runtime/Experimental/ChaosVehicles/ChaosVehiclesCore/Public/SimModule/SimulationModuleBase.h
+++ b/Engine/Source/Runtime/Experimental/ChaosVehicles/ChaosVehiclesCore/Public/SimModule/SimulationModuleBase.h
@@ -7,6 +7,7 @@
 #include "Logging/LogMacros.h"
 #include "SimModule/ModuleFactoryRegister.h"
 #include "SimModule/ModuleInput.h"
+#include "SimModule/VehicleBlackboard.h"
 
 DECLARE_LOG_CATEGORY_EXTERN(LogSimulationModule, Warning, All);
 
@@ -321,6 +322,8 @@ namespace Chaos
 		ISimulationModuleBase* GetParent();
 		ISimulationModuleBase* GetFirstChild();
 
+		FVehicleBlackboard* GetSimBlackboard();
+
 		// for headless chaos testing
 		const FVector& GetAppliedForce() { return AppliedForce; }
 
diff --git a/Engine/Source/Runtime/Experimental/ChaosVehicles/ChaosVehiclesCore/Public/SimModule/VehicleBlackboard.h b/Engine/Source/Runtime/Experimental/ChaosVehicles/ChaosVehiclesCore/Public/SimModule/VehicleBlackboard.h
new file mode 100644
--- /dev/null
+++ b/Engine/Source/Runtime/Experimental/ChaosVehicles/ChaosVehiclesCore/Public/SimModule/VehicleBlackboard.h
@@ -0,0 +1,107 @@
+// Copyright Epic Games, Inc. All Rights Reserved.
+
+#pragma once
+
+#include "Templates/UniquePtr.h"
+#include "Templates/SharedPointer.h"
+#include "Templates/Casts.h"
+
+
+
+/** VehicleBlackboard: this is a simple generic map that can store any type, used as a way for decoupled modules to 
+ *  share calculations or transient state data (on the physics thread only). 
+ *  Values submitted are copy-in, copy-out. 
+ *  Unlike a traditional blackboard pattern, there is no support for subscribing to changes. 
+ */
+class CHAOSVEHICLESCORE_API FVehicleBlackboard
+{
+private:
+	class BlackboardObject
+	{
+		// untyped base
+		struct ObjectContainerBase
+		{
+		};
+
+		// typed container
+		template<typename T>
+		struct ObjectContainer : ObjectContainerBase
+		{
+			ObjectContainer(const T& t) : Object(t) {}
+
+			const T& Get() { return Object; }
+			T& GetMutable() { return Object; }
+
+		private:
+			T Object;
+		};
+
+
+	public:
+		template<typename T>
+		BlackboardObject(const T& obj) : ContainerPtr(MakeShared<ObjectContainer<T>>(obj)) {}
+
+		template<typename T>
+		const T& Get() const
+		{
+			ObjectContainer<T>* TypedContainer = static_cast<ObjectContainer<T>*>(ContainerPtr.Get());
+			return TypedContainer->Get();
+		}
+
+		template<typename T>
+		T& GetMutable() const
+		{
+			ObjectContainer<T>* TypedContainer = static_cast<ObjectContainer<T>*>(ContainerPtr.Get());
+			return TypedContainer->GetMutable();
+		}
+
+	private:
+		TSharedPtr<ObjectContainerBase> ContainerPtr;
+	};	// end BlackboardObject
+ 
+public:
+	enum class EInvalidationReason : uint8
+	{
+		FullReset,	// All blackboard objects should be invalidated
+		Rollback,	// Invalidate any rollback-sensitive objects
+	};
+
+	/** Attempt to retrieve an object from the blackboard. If found, OutFoundValue will be set. Returns true/false to indicate whether it was found. */
+ 	template<typename T>
+	bool TryGet(FName ObjName, T& OutFoundValue) const
+	{
+		if (const TSharedPtr<BlackboardObject>* ExistingObject = ObjectsByName.Find(ObjName))
+		{
+			OutFoundValue = ExistingObject->Get()->Get<T>();
+			return true;
+		}
+
+		return false;
+	}
+
+	/** Returns true/false to indicate if an object is stored with that name */
+	bool Contains(FName ObjName)
+	{
+		return ObjectsByName.Contains(ObjName);
+	}
+
+	/** Store object by a named key, overwriting any existing object */
+	template<typename T>
+	void Set(FName ObjName, T Obj)
+	{
+		ObjectsByName.Emplace(ObjName, MakeShared<BlackboardObject>(Obj));
+	}
+
+	/** Invalidate an object by name */
+	void Invalidate(FName ObjName);
+
+	/** Invalidate all objects that can be affected by a particular circumstance (such as a rollback) */
+	void Invalidate(EInvalidationReason Reason);
+	
+	/** Invalidate all objects */
+	void InvalidateAll() { Invalidate(EInvalidationReason::FullReset); }
+
+
+private:
+	TMap<FName, TSharedPtr<BlackboardObject>> ObjectsByName;
+};
\ No newline at end of file
