Slight refactor for better error handling for IAD errors.
- Mount callback always passes back MountId now
- Fix tagset array reserve
- Use unique function for completion callback

diff --git a/Engine/Source/Runtime/Experimental/IoStore/OnDemand/Private/IoStoreOnDemand.cpp b/Engine/Source/Runtime/Experimental/IoStore/OnDemand/Private/IoStoreOnDemand.cpp
--- a/Engine/Source/Runtime/Experimental/IoStore/OnDemand/Private/IoStoreOnDemand.cpp
+++ b/Engine/Source/Runtime/Experimental/IoStore/OnDemand/Private/IoStoreOnDemand.cpp
@@ -1048,7 +1048,7 @@ void FIoStoreOnDemandModule::ReportAnalytics(TArray<FAnalyticsEventAttribute>& O
 	}
 }
 
-void FIoStoreOnDemandModule::Mount(FOnDemandMountArgs&& Args, FOnDemandMountCompleted&& OnCompleted)
+void FIoStoreOnDemandModule::Mount(FOnDemandMountArgs&& Args, FOnDemandMountCompleted OnCompleted)
 {
 	if (IoStore.IsValid() == false)
 	{
@@ -1057,7 +1057,7 @@ void FIoStoreOnDemandModule::Mount(FOnDemandMountArgs&& Args, FOnDemandMountComp
 		{
 			UE_LOG(LogIas, Error, TEXT("Failed to initialize I/O store on-demand, reason '%s'"), *Status.ToString());
 			IoStore.Reset();
-			return OnCompleted(TIoStatusOr<FOnDemandMountResult>(Status));
+			return OnCompleted(Args.MountId, TIoStatusOr<FOnDemandMountResult>(Status));
 		}
 	}
 
@@ -1226,10 +1226,10 @@ void FIoStoreOnDemandModule::InitializeInternal()
 	{
 		IoStore->Mount(
 			MoveTemp(MountArgs.GetValue()),
-			[](TIoStatusOr<FOnDemandMountResult> MountResult)
+			[](FStringView MountId, TIoStatusOr<FOnDemandMountResult> MountResult)
 			{
 				UE_CLOG(!MountResult.IsOk(), LogIas, Error,
-					TEXT("Failed to mount TOC, reason '%s'"), *MountResult.Status().ToString());
+					TEXT("Failed to mount TOC for '%.*s', reason '%s'"), MountId.Len(), MountId.GetData(), *MountResult.Status().ToString());
 			});
 	}
 }
diff --git a/Engine/Source/Runtime/Experimental/IoStore/OnDemand/Private/OnDemandIoStore.cpp b/Engine/Source/Runtime/Experimental/IoStore/OnDemand/Private/OnDemandIoStore.cpp
--- a/Engine/Source/Runtime/Experimental/IoStore/OnDemand/Private/OnDemandIoStore.cpp
+++ b/Engine/Source/Runtime/Experimental/IoStore/OnDemand/Private/OnDemandIoStore.cpp
@@ -5,6 +5,7 @@
 #include "OnDemandInstallCache.h"
 #include "OnDemandPackageStoreBackend.h"
 
+#include "Algo/Accumulate.h"
 #include "Algo/Copy.h"
 #include "Algo/Find.h"
 #include "Algo/RemoveIf.h"
@@ -493,11 +494,11 @@ FIoStatus FOnDemandIoStore::Initialize()
 	return EIoErrorCode::Ok;
 }
 
-void FOnDemandIoStore::Mount(FOnDemandMountArgs&& Args, FOnDemandMountCompleted&& OnCompleted)
+void FOnDemandIoStore::Mount(FOnDemandMountArgs&& Args, FOnDemandMountCompleted OnCompleted)
 {
 	if (Args.MountId.IsEmpty())
 	{
-		return OnCompleted(FIoStatus(EIoErrorCode::InvalidParameter, TEXT("Invalid Mount ID")));
+		return OnCompleted(FStringView(), FIoStatus(EIoErrorCode::InvalidParameter, TEXT("Invalid Mount ID")));
 	}
 
 	{
@@ -507,7 +508,7 @@ void FOnDemandIoStore::Mount(FOnDemandMountArgs&& Args, FOnDemandMountCompleted&
 		if (MountRequest.IsValid())
 		{
 			UE_LOG(LogIoStoreOnDemand, Warning, TEXT("Mount request '%s' is already mounting"), *Args.MountId);
-			return OnCompleted(FIoStatus(EIoErrorCode::InvalidParameter));
+			return OnCompleted(Args.MountId, FIoStatus(EIoErrorCode::InvalidParameter));
 		}
 
 		UE_LOG(LogIoStoreOnDemand, Log, TEXT("Enqueing mount request, MountId='%s'"), *Args.MountId);
@@ -769,7 +770,7 @@ bool FOnDemandIoStore::Tick()
 				MountRequests.Remove(Request->MountArgs.MountId);
 			}
 			FOnDemandMountCompleted OnCompleted = MoveTemp(Request->OnCompleted);
-			OnCompleted(Status);
+			OnCompleted(Request->MountArgs.MountId, Status);
 			continue;
 		}
 
@@ -822,7 +823,7 @@ bool FOnDemandIoStore::Tick()
 					MountRequests.Remove(Request->MountArgs.MountId);
 				}
 				FOnDemandMountCompleted OnCompleted = MoveTemp(Request->OnCompleted);
-				OnCompleted(Status);
+				OnCompleted(Request->MountArgs.MountId, Status);
 				continue;
 			}
 		}
@@ -863,10 +864,7 @@ bool FOnDemandIoStore::Tick()
 			}
 
 			FOnDemandMountCompleted OnCompleted = MoveTemp(Request->OnCompleted);
-			OnCompleted(FOnDemandMountResult
-			{
-				.MountId = Request->MountArgs.MountId
-			});
+			OnCompleted(Request->MountArgs.MountId, FOnDemandMountResult());
 		}
 	}
 
@@ -1101,8 +1099,13 @@ FIoStatus FOnDemandIoStore::TickInstallRequest(FMountRequest& MountRequest)
 				continue;
 			}
 
-			FTagSet NewTagSet; 
+			FTagSet NewTagSet;
 			NewTagSet.Tag = TagSet->Tag;
+			NewTagSet.PackageIds.Reserve(Algo::TransformAccumulate(
+				TagSet->Packages,
+				[](const FOnDemandTocTagSetPackageList& PackageIndices) -> int32 { return PackageIndices.PackageIndicies.Num(); },
+				int32(0))
+			);
 			for (const FOnDemandTocTagSetPackageList& PackageIndices : TagSet->Packages)
 			{
 				if (MountRequest.Containers.IsEmpty() || int32(PackageIndices.ContainerIndex) > (MountRequest.Containers.Num() - 1))
@@ -1119,7 +1122,6 @@ FIoStatus FOnDemandIoStore::TickInstallRequest(FMountRequest& MountRequest)
 					continue;
 				}
 
-				NewTagSet.PackageIds.Reserve(PackageIndices.PackageIndicies.Num());
 				for (int32 IdIdx : PackageIndices.PackageIndicies)
 				{
 					NewTagSet.PackageIds.Add(Container.Header->PackageIds[IdIdx]);
diff --git a/Engine/Source/Runtime/Experimental/IoStore/OnDemand/Private/OnDemandIoStore.h b/Engine/Source/Runtime/Experimental/IoStore/OnDemand/Private/OnDemandIoStore.h
--- a/Engine/Source/Runtime/Experimental/IoStore/OnDemand/Private/OnDemandIoStore.h
+++ b/Engine/Source/Runtime/Experimental/IoStore/OnDemand/Private/OnDemandIoStore.h
@@ -157,7 +157,7 @@ public:
 	FOnDemandIoStore& operator=(FOnDemandIoStore&&) = delete;
 
 	FIoStatus				Initialize();
-	void					Mount(FOnDemandMountArgs&& Args, FOnDemandMountCompleted&& OnCompleted);
+	void					Mount(FOnDemandMountArgs&& Args, FOnDemandMountCompleted OnCompleted);
 	FIoStatus				Unmount(FStringView MountId);
 	FOnDemandChunkInfo		GetStreamingChunkInfo(const FIoChunkId& ChunkId);
 	FOnDemandChunkInfo		GetInstalledChunkInfo(const FIoChunkId& ChunkId);
diff --git a/Engine/Source/Runtime/Experimental/IoStore/OnDemand/Public/IO/IoStoreOnDemand.h b/Engine/Source/Runtime/Experimental/IoStore/OnDemand/Public/IO/IoStoreOnDemand.h
--- a/Engine/Source/Runtime/Experimental/IoStore/OnDemand/Public/IO/IoStoreOnDemand.h
+++ b/Engine/Source/Runtime/Experimental/IoStore/OnDemand/Public/IO/IoStoreOnDemand.h
@@ -270,10 +270,9 @@ struct FOnDemandMountArgs
 
 struct FOnDemandMountResult
 {
-	FString MountId;
 };
 
-using FOnDemandMountCompleted = TFunction<void(TIoStatusOr<FOnDemandMountResult>)>;
+using FOnDemandMountCompleted = TUniqueFunction<void(FStringView, TIoStatusOr<FOnDemandMountResult>)>;
 
 class FIoStoreOnDemandModule
 	: public IModuleInterface
@@ -296,7 +295,7 @@ public:
 
 	UE_API void ReportAnalytics(TArray<FAnalyticsEventAttribute>& OutAnalyticsArray) const;
 
-	UE_API void Mount(FOnDemandMountArgs&& Args, FOnDemandMountCompleted&& OnCompleted);
+	UE_API void Mount(FOnDemandMountArgs&& Args, FOnDemandMountCompleted OnCompleted);
 	UE_API FIoStatus Unmount(FStringView MountId);
 
 	virtual void StartupModule() override;
