Move both MountId and Status into FOnDemandMountResult struct.

diff --git a/Engine/Source/Runtime/Experimental/IoStore/OnDemand/Private/IoStoreOnDemand.cpp b/Engine/Source/Runtime/Experimental/IoStore/OnDemand/Private/IoStoreOnDemand.cpp
--- a/Engine/Source/Runtime/Experimental/IoStore/OnDemand/Private/IoStoreOnDemand.cpp
+++ b/Engine/Source/Runtime/Experimental/IoStore/OnDemand/Private/IoStoreOnDemand.cpp
@@ -1057,7 +1057,7 @@ void FIoStoreOnDemandModule::Mount(FOnDemandMountArgs&& Args, FOnDemandMountComp
 		{
 			UE_LOG(LogIas, Error, TEXT("Failed to initialize I/O store on-demand, reason '%s'"), *Status.ToString());
 			IoStore.Reset();
-			return OnCompleted(Args.MountId, TIoStatusOr<FOnDemandMountResult>(Status));
+			return OnCompleted(FOnDemandMountResult{ Args.MountId, Status });
 		}
 	}
 
@@ -1226,10 +1226,10 @@ void FIoStoreOnDemandModule::InitializeInternal()
 	{
 		IoStore->Mount(
 			MoveTemp(MountArgs.GetValue()),
-			[](FStringView MountId, TIoStatusOr<FOnDemandMountResult> MountResult)
+			[](FOnDemandMountResult MountResult)
 			{
-				UE_CLOG(!MountResult.IsOk(), LogIas, Error,
-					TEXT("Failed to mount TOC for '%.*s', reason '%s'"), MountId.Len(), MountId.GetData(), *MountResult.Status().ToString());
+				UE_CLOG(!MountResult.Status.IsOk(), LogIas, Error,
+					TEXT("Failed to mount TOC for '%s', reason '%s'"), *MountResult.MountId, *MountResult.Status.ToString());
 			});
 	}
 }
diff --git a/Engine/Source/Runtime/Experimental/IoStore/OnDemand/Private/OnDemandIoStore.cpp b/Engine/Source/Runtime/Experimental/IoStore/OnDemand/Private/OnDemandIoStore.cpp
--- a/Engine/Source/Runtime/Experimental/IoStore/OnDemand/Private/OnDemandIoStore.cpp
+++ b/Engine/Source/Runtime/Experimental/IoStore/OnDemand/Private/OnDemandIoStore.cpp
@@ -498,7 +498,7 @@ void FOnDemandIoStore::Mount(FOnDemandMountArgs&& Args, FOnDemandMountCompleted
 {
 	if (Args.MountId.IsEmpty())
 	{
-		return OnCompleted(FStringView(), FIoStatus(EIoErrorCode::InvalidParameter, TEXT("Invalid Mount ID")));
+		return OnCompleted(FOnDemandMountResult{ FString(), FIoStatus(EIoErrorCode::InvalidParameter, TEXT("Invalid Mount ID"))});
 	}
 
 	{
@@ -508,7 +508,7 @@ void FOnDemandIoStore::Mount(FOnDemandMountArgs&& Args, FOnDemandMountCompleted
 		if (MountRequest.IsValid())
 		{
 			UE_LOG(LogIoStoreOnDemand, Warning, TEXT("Mount request '%s' is already mounting"), *Args.MountId);
-			return OnCompleted(Args.MountId, FIoStatus(EIoErrorCode::InvalidParameter));
+			return OnCompleted(FOnDemandMountResult{ Args.MountId, FIoStatus(EIoErrorCode::InvalidParameter) });
 		}
 
 		UE_LOG(LogIoStoreOnDemand, Log, TEXT("Enqueing mount request, MountId='%s'"), *Args.MountId);
@@ -770,7 +770,7 @@ bool FOnDemandIoStore::Tick()
 				MountRequests.Remove(Request->MountArgs.MountId);
 			}
 			FOnDemandMountCompleted OnCompleted = MoveTemp(Request->OnCompleted);
-			OnCompleted(Request->MountArgs.MountId, Status);
+			OnCompleted(FOnDemandMountResult{ Request->MountArgs.MountId, Status });
 			continue;
 		}
 
@@ -823,7 +823,7 @@ bool FOnDemandIoStore::Tick()
 					MountRequests.Remove(Request->MountArgs.MountId);
 				}
 				FOnDemandMountCompleted OnCompleted = MoveTemp(Request->OnCompleted);
-				OnCompleted(Request->MountArgs.MountId, Status);
+				OnCompleted(FOnDemandMountResult{ Request->MountArgs.MountId, Status });
 				continue;
 			}
 		}
@@ -864,7 +864,7 @@ bool FOnDemandIoStore::Tick()
 			}
 
 			FOnDemandMountCompleted OnCompleted = MoveTemp(Request->OnCompleted);
-			OnCompleted(Request->MountArgs.MountId, FOnDemandMountResult());
+			OnCompleted(FOnDemandMountResult{ Request->MountArgs.MountId });
 		}
 	}
 
diff --git a/Engine/Source/Runtime/Experimental/IoStore/OnDemand/Public/IO/IoStoreOnDemand.h b/Engine/Source/Runtime/Experimental/IoStore/OnDemand/Public/IO/IoStoreOnDemand.h
--- a/Engine/Source/Runtime/Experimental/IoStore/OnDemand/Public/IO/IoStoreOnDemand.h
+++ b/Engine/Source/Runtime/Experimental/IoStore/OnDemand/Public/IO/IoStoreOnDemand.h
@@ -270,9 +270,11 @@ struct FOnDemandMountArgs
 
 struct FOnDemandMountResult
 {
+	FString MountId;
+	FIoStatus Status;
 };
 
-using FOnDemandMountCompleted = TUniqueFunction<void(FStringView, TIoStatusOr<FOnDemandMountResult>)>;
+using FOnDemandMountCompleted = TUniqueFunction<void(FOnDemandMountResult)>;
 
 class FIoStoreOnDemandModule
 	: public IModuleInterface
