Chaos : Make sure to remove all tracked collision constraint before disabling a cluster particle

diff --git a/Engine/Source/Runtime/Experimental/Chaos/Private/Chaos/PBDRigidClustering.cpp b/Engine/Source/Runtime/Experimental/Chaos/Private/Chaos/PBDRigidClustering.cpp
--- a/Engine/Source/Runtime/Experimental/Chaos/Private/Chaos/PBDRigidClustering.cpp
+++ b/Engine/Source/Runtime/Experimental/Chaos/Private/Chaos/PBDRigidClustering.cpp
@@ -2406,6 +2406,23 @@ namespace Chaos
 
 	void FRigidClustering::DisableCluster(FPBDRigidClusteredParticleHandle* ClusteredParticle)
 	{
+		// Before disabling the particle, we need to make sure to remove the recorded breaking collision from BreakingCollisions
+		// to avoid it crashing later when processing it 
+		ClusteredParticle->ParticleCollisions().VisitCollisions(
+			[this, ClusteredParticle](FPBDCollisionConstraint& Collision)
+			{
+				const uint8 OtherIdx = (Collision.GetParticle0() == ClusteredParticle) ? 1 : 0;
+				if (FGeometryParticleHandle* OtherGeometry = Collision.GetParticle(OtherIdx))
+				{
+					if (FPBDRigidParticleHandle* OtherRigid = OtherGeometry->CastToRigidParticle())
+					{
+						BreakingCollisions.Remove({ &Collision, ClusteredParticle });
+						BreakingCollisions.Remove({ &Collision, OtherRigid });
+					}
+				}
+				return ECollisionVisitorResult::Continue;
+			});
+
 		// #note: we don't recursively descend to the children
 		MEvolution.DisableParticle(ClusteredParticle);
 		TopLevelClusterParents.Remove(ClusteredParticle);
