ChaosFlesh: Collision with sphere repulsion

diff --git a/Engine/Plugins/Experimental/ChaosFlesh/Source/ChaosFleshEngine/Private/ChaosFlesh/ChaosDeformableSolverComponent.cpp b/Engine/Plugins/Experimental/ChaosFlesh/Source/ChaosFleshEngine/Private/ChaosFlesh/ChaosDeformableSolverComponent.cpp
--- a/Engine/Plugins/Experimental/ChaosFlesh/Source/ChaosFleshEngine/Private/ChaosFlesh/ChaosDeformableSolverComponent.cpp
+++ b/Engine/Plugins/Experimental/ChaosFlesh/Source/ChaosFleshEngine/Private/ChaosFlesh/ChaosDeformableSolverComponent.cpp
@@ -236,14 +236,17 @@ void UDeformableSolverComponent::BuildSimulationProxy()
 			, SolverConstraints.GaussSeidelConstraints.bUseSOR
 			, SolverConstraints.GaussSeidelConstraints.OmegaSOR
 			, SolverConstraints.GaussSeidelConstraints.bUseGSNeohookean
-			, SolverConstraints.GaussSeidelConstraints.CollisionSpring.bDoSelfCollision
-			, SolverConstraints.GaussSeidelConstraints.CollisionSpring.bDoInComponentSelfCollision
-			, SolverConstraints.GaussSeidelConstraints.CollisionSpring.NRingExcluded
-			, SolverConstraints.GaussSeidelConstraints.CollisionSpring.CollisionSearchRadius
-			, SolverConstraints.GaussSeidelConstraints.CollisionSpring.CollisionSpringStiffness
-			, SolverConstraints.GaussSeidelConstraints.CollisionSpring.bAllowSliding
+			, SolverConstraints.GaussSeidelConstraints.SpringCollision.bDoSpringCollision
+			, SolverConstraints.GaussSeidelConstraints.SpringCollision.InComponentSpringCollision.bDoInComponentSpringCollision
+			, SolverConstraints.GaussSeidelConstraints.SpringCollision.InComponentSpringCollision.NRingExcluded
+			, SolverConstraints.GaussSeidelConstraints.SpringCollision.CollisionSearchRadius
+			, SolverConstraints.GaussSeidelConstraints.SpringCollision.SpringCollisionStiffness
+			, SolverConstraints.GaussSeidelConstraints.SpringCollision.bAllowSliding
+			, SolverConstraints.GaussSeidelConstraints.SphereRepulsion.bDoSphereRepulsion
+			, SolverConstraints.GaussSeidelConstraints.SphereRepulsion.SphereRepulsionRadius
+			, SolverConstraints.GaussSeidelConstraints.SphereRepulsion.SphereRepulsionStiffness
 			, SolverMuscleActivation.bDoMuscleActivation
-			, SolverConstraints.GaussSeidelConstraints.CollisionSpring.bCollideWithFullmesh
+			, SolverConstraints.GaussSeidelConstraints.SpringCollision.bCollideWithFullmesh
 			, SolverConstraints.GaussSeidelConstraints.bEnableDynamicSprings
 		));
 
diff --git a/Engine/Plugins/Experimental/ChaosFlesh/Source/ChaosFleshEngine/Public/ChaosFlesh/ChaosDeformableSolverComponent.h b/Engine/Plugins/Experimental/ChaosFlesh/Source/ChaosFleshEngine/Public/ChaosFlesh/ChaosDeformableSolverComponent.h
--- a/Engine/Plugins/Experimental/ChaosFlesh/Source/ChaosFleshEngine/Public/ChaosFlesh/ChaosDeformableSolverComponent.h
+++ b/Engine/Plugins/Experimental/ChaosFlesh/Source/ChaosFleshEngine/Public/ChaosFlesh/ChaosDeformableSolverComponent.h
@@ -96,46 +96,78 @@ struct FSolverGridBasedCollisionsGroup
 };
 
 USTRUCT(BlueprintType)
-struct FCollisionSpringGroup
+struct FInComponentSpringCollisionGroup
 {
 	GENERATED_USTRUCT_BODY()
 	/**
-	* If uses spring self-collision
+	* If uses in-component spring self-collision
 	*/
-	UPROPERTY(EditAnywhere, Category = "CollisionSpring")
-	bool bDoSelfCollision = false;
+	UPROPERTY(EditAnywhere, Category = "InComponentSpringCollision")
+	bool bDoInComponentSpringCollision = false;
 	/**
-	* If uses in-component self-collision
+	* N ring to exclude for in-component spring self-collision
 	*/
-	UPROPERTY(EditAnywhere, Category = "CollisionSpring")
-	bool bDoInComponentSelfCollision = false;
+	UPROPERTY(EditAnywhere, Category = "InComponentSpringCollision")
+	int32 NRingExcluded = 1;
+};
+
+USTRUCT(BlueprintType)
+struct FSpringCollisionGroup
+{
+	GENERATED_USTRUCT_BODY()
 	/**
-	* N ring to exclude for in-component self-collision 
+	* If uses component-component spring collision
 	*/
-	UPROPERTY(EditAnywhere, Category = "CollisionSpring")
-	int32 NRingExcluded = 1;
+	UPROPERTY(EditAnywhere, Category = "SpringCollision")
+	bool bDoSpringCollision = false;
+	/**
+	* In-component spring self collision detection parameters
+	*/
+	UPROPERTY(EditAnywhere, BlueprintReadOnly, Category = "SpringCollision")
+	FInComponentSpringCollisionGroup InComponentSpringCollision;
 	/**
-	* Search radius for point triangle pairs
+	* Search radius for point triangle collision pairs
 	*/
-	UPROPERTY(EditAnywhere, Category = "CollisionSpring")
+	UPROPERTY(EditAnywhere, Category = "SpringCollision")
 	float CollisionSearchRadius = 0.f;
 	/**
 	* Collision spring stiffness; larger value will stop penetration better
 	*/
-	UPROPERTY(EditAnywhere, Category = "CollisionSpring")
-	float CollisionSpringStiffness = 500.f;
+	UPROPERTY(EditAnywhere, Category = "SpringCollision")
+	float SpringCollisionStiffness = 500.f;
 	/**
 	* Anisotropic springs will allow sliding on the triangle
 	*/
-	UPROPERTY(EditAnywhere, Category = "CollisionSpring")
+	UPROPERTY(EditAnywhere, Category = "SpringCollision")
 	bool bAllowSliding = true;
 	/**
 	* Do self collision with kinematic triangles as well
 	*/
-	UPROPERTY(EditAnywhere, Category = "CollisionSpring")
+	UPROPERTY(EditAnywhere, Category = "SpringCollision")
 	bool bCollideWithFullmesh = true;
 };
 
+USTRUCT(BlueprintType)
+struct FSphereRepulsionGroup
+{
+	GENERATED_USTRUCT_BODY()
+	/**
+	* If uses sphere repulsion for collision
+	*/
+	UPROPERTY(EditAnywhere, Category = "SphereRepulsion")
+	bool bDoSphereRepulsion = false;
+	/**
+	* Search radius for repulsion pairs
+	*/
+	UPROPERTY(EditAnywhere, Category = "SphereRepulsion")
+	float SphereRepulsionRadius = 0.f;
+	/**
+	* Stiffness for sphere repulsion
+	*/
+	UPROPERTY(EditAnywhere, Category = "SphereRepulsion")
+	float SphereRepulsionStiffness = 500.f;
+};
+
 USTRUCT(BlueprintType)
 struct FSolverGaussSeidelConstraintsGroup
 {
@@ -170,12 +202,18 @@ struct FSolverGaussSeidelConstraintsGroup
 	*/
 	UPROPERTY(EditAnywhere, Category = "GaussSeidelConstraints")
 	bool bEnableDynamicSprings = true;
+	
+	/**
+	* Component-component collision detection radius and stiffness
+	*/
+	UPROPERTY(EditAnywhere, BlueprintReadOnly, Category = "GaussSeidelConstraints")
+	FSpringCollisionGroup SpringCollision;
 
 	/**
-	* Collsion detection radius and stiffness
+	* Sphere repulsion parameters
 	*/
 	UPROPERTY(EditAnywhere, BlueprintReadOnly, Category = "GaussSeidelConstraints")
-	FCollisionSpringGroup CollisionSpring;
+	FSphereRepulsionGroup SphereRepulsion;
 };
 
 USTRUCT(BlueprintType)
diff --git a/Engine/Source/Runtime/Experimental/Chaos/Private/Chaos/Deformable/ChaosDeformableSolver.cpp b/Engine/Source/Runtime/Experimental/Chaos/Private/Chaos/Deformable/ChaosDeformableSolver.cpp
--- a/Engine/Source/Runtime/Experimental/Chaos/Private/Chaos/Deformable/ChaosDeformableSolver.cpp
+++ b/Engine/Source/Runtime/Experimental/Chaos/Private/Chaos/Deformable/ChaosDeformableSolver.cpp
@@ -150,14 +150,15 @@ namespace Chaos::Softs
 			Property.bDoQuasistatics, 
 			true));
 		Evolution->Particles().AddArray(&MObjects);
-		if (Property.bDoSelfCollision || Property.CacheToFile)
+		if (Property.bDoSpringCollision || Property.bDoSphereRepulsion || Property.CacheToFile)
 		{
 			SurfaceElements.Reset(new TArray<Chaos::TVec3<int32>>());
 			TetmeshSurfaceElements.Reset(new TArray<Chaos::TVec3<int32>>());
 		}
 
-		if (Property.bDoSelfCollision)
+		if (Property.bDoSpringCollision || Property.bDoSphereRepulsion)
 		{
+			ParticleComponentIndex.Reset(new TArray<int32>());
 			SurfaceTriangleMesh.Reset(new Chaos::FTriangleMesh());
 			SurfaceVertices.Reset(new TArray<int32>);
 		}
@@ -177,8 +178,9 @@ namespace Chaos::Softs
 			AllSkinEMeshArray.Reset(new TArray<FSolverReal>());
 			GSWeakConstraints.Reset(new FGaussSeidelWeakConstraints<FSolverReal, FSolverParticles>({}, {}, {}, {}, {}, GDeformableXPBDWeakConstraintParams));
 			GSDynamicWeakConstraints.Reset(new FGaussSeidelDynamicWeakConstraints<FSolverReal, FSolverParticles>(GDeformableXPBDWeakConstraintParams));
-			MuscleActivationConstraints.Reset(new FMuscleActivationConstraints<FSolverReal, FSolverParticles>());
+			GSSphereRepulsionConstraints.Reset(new FGaussSeidelSphereRepulsionConstraints<FSolverReal, FSolverParticles>(Property.SphereRepulsionRadius, Property.SphereRepulsionStiffness, Evolution->Particles(), GDeformableXPBDWeakConstraintParams));
 		}
+		MuscleActivationConstraints.Reset(new FMuscleActivationConstraints<FSolverReal, FSolverParticles>());
 		AllUnconstrainedSurfaceElementsCorotatedCod.Reset(new TArray<Chaos::TVec3<int32>>());
 		AllUnconstrainedSurfaceElementsSkin.Reset(new TArray<Chaos::TVec3<int32>>());
 
@@ -350,7 +352,7 @@ namespace Chaos::Softs
 
 				if (UninitializedProxys_Internal.Num() != 0)
 				{
-					if (Property.bDoSelfCollision)
+					if (Property.bDoSpringCollision || Property.bDoSphereRepulsion)
 					{
 						InitializeSelfCollisionVariables();
 					}
@@ -1357,6 +1359,7 @@ namespace Chaos::Softs
 		int32 VertexOffset = 0;
 		SurfaceElements->SetNum(0);
 		TetmeshSurfaceElements->SetNum(0);
+		ParticleComponentIndex->SetNum(0);
 		int ComponentOffset = 0;
 
 		for (FThreadingProxy* InProxy : UninitializedProxys_Internal)
@@ -1382,12 +1385,12 @@ namespace Chaos::Softs
 						VertexOffset += Vertex->Num();
 
 						
-						if (!Property.bDoInComponentSelfCollision) //Component-Component collisions
+						if (!Property.bDoInComponentSpringCollision || Property.bDoSphereRepulsion) //Component-Component collisions or sphere repulsion
 						{
-							int32 Offset = ParticleComponentIndex.Num();
-							ParticleComponentIndex.SetNum(ParticleComponentIndex.Num() + Vertex->Num());
+							int32 Offset = ParticleComponentIndex->Num();
+							ParticleComponentIndex->SetNum(ParticleComponentIndex->Num() + Vertex->Num());
 							for (int32 i = 0; i < Vertex->Num(); i++) {
-								ParticleComponentIndex[i + Offset] = ComponentOffset;
+								(*ParticleComponentIndex)[i + Offset] = ComponentOffset;
 							}
 							int NewComponentOffset = ComponentOffset;
 							GeometryCollection::Facades::FCollectionMeshFacade MeshFacade(*Rest);
@@ -1395,12 +1398,12 @@ namespace Chaos::Softs
 							for (int32 i = 0; i < ComponentIndex.Num(); i++) {
 								if (ComponentIndex[i] < 0)
 								{
-									ParticleComponentIndex[i + Offset] = ComponentIndex[i]; //Isolated Nodes
+									(*ParticleComponentIndex)[i + Offset] = ComponentIndex[i]; //Isolated Nodes
 								}
 								else
 								{
-									ParticleComponentIndex[i + Offset] = ComponentOffset + ComponentIndex[i];
-									NewComponentOffset = NewComponentOffset < ParticleComponentIndex[i + Offset] ? ParticleComponentIndex[i + Offset] : NewComponentOffset;
+									(*ParticleComponentIndex)[i + Offset] = ComponentOffset + ComponentIndex[i];
+									NewComponentOffset = NewComponentOffset < (*ParticleComponentIndex)[i + Offset] ? (*ParticleComponentIndex)[i + Offset] : NewComponentOffset;
 								}
 							}
 							ComponentOffset = NewComponentOffset + 1;
@@ -1482,13 +1485,13 @@ namespace Chaos::Softs
 		TriangleMeshCollisions.Reset(new FPBDTriangleMeshCollisions(
 			0, Evolution->Particles().Size(), *SurfaceTriangleMesh, false, false));
 		ParticleTriangleExclusionMap.Reset();
-		if (Property.bDoInComponentSelfCollision)
-		{	
+		if (Property.bDoInComponentSpringCollision)
+		{
 			int32 NRadius = Property.NRingExcluded;
 			SurfaceTriangleMesh->GetPointToNeighborsMap();
 			SurfaceTriangleMesh->GetPointToTriangleMap(); //Initialize Maps before parallel
 			ParticleTriangleExclusionMap.Reserve(SurfaceVertices->Num());
-			PhysicsParallelFor(SurfaceVertices->Num(), 
+			PhysicsParallelFor(SurfaceVertices->Num(),
 				[this, &NRadius](int32 i)
 				{
 					int32 VertexId = (*SurfaceVertices)[i];
@@ -1496,7 +1499,7 @@ namespace Chaos::Softs
 					TSet<int32> NRing;
 					if (NRadius > 1)
 					{
-						NRing = SurfaceTriangleMesh->GetNRing(VertexId, NRadius-1);
+						NRing = SurfaceTriangleMesh->GetNRing(VertexId, NRadius - 1);
 					}
 					else if (NRadius == 1)
 					{
@@ -1731,10 +1734,7 @@ namespace Chaos::Softs
 				};
 		}
 
-
-		
-
-		if (Property.bDoSelfCollision)
+		if (Property.bDoSpringCollision)
 		{
 			int32 TransientIndex = GSMainConstraint->AddTransientConstraintResidualAndHessianRange(1);
 			GSMainConstraint->TransientConstraintResidualAndHessian()[TransientIndex] = [this](const FSolverParticles& Particles, const int32 ConstraintIndex, const int32 ConstraintIndexLocal, const FSolverReal Dt, TVec3<FSolverReal>& ParticleResidual, Chaos::PMatrix<FSolverReal, 3, 3>& ParticleHessian)
@@ -1755,58 +1755,80 @@ namespace Chaos::Softs
 			Evolution->ConstraintInits()[InitIndex] =
 				[this](FSolverParticles& InParticles, const FSolverReal Dt)
 				{
-					this->TriangleMeshCollisions->InitFlesh(InParticles, Property.CollisionSearchRadius, this->Property.bCollideWithFullMesh);
-					TArray<FPBDTriangleMeshCollisions::FGIAColor> EmptyGIAColors;
-					if (Property.bDoInComponentSelfCollision)
+					this->TriangleMeshCollisions->InitFlesh(InParticles, Property.SpringCollisionSearchRadius, this->Property.bCollideWithFullMesh);
+					if (Property.bDoInComponentSpringCollision)
 					{
-						this->GSWeakConstraints->CollisionDetectionSpatialHashInComponent(this->Evolution->Particles(), *SurfaceVertices, *SurfaceTriangleMesh, ParticleTriangleExclusionMap, TriangleMeshCollisions->GetDynamicSpatialHash(), Property.CollisionSearchRadius, Property.CollisionSpringStiffness, Property.bAllowSliding);
+						this->GSWeakConstraints->CollisionDetectionSpatialHashInComponent(this->Evolution->Particles(), *SurfaceVertices, *SurfaceTriangleMesh, ParticleTriangleExclusionMap, TriangleMeshCollisions->GetDynamicSpatialHash(), Property.SpringCollisionSearchRadius, Property.SpringCollisionStiffness, Property.bAllowSliding);
 					}
 					else
 					{
-						this->GSWeakConstraints->CollisionDetectionSpatialHash(this->Evolution->Particles(), *SurfaceVertices, *SurfaceTriangleMesh, ParticleComponentIndex, TriangleMeshCollisions->GetDynamicSpatialHash(), Property.CollisionSearchRadius, Property.CollisionSpringStiffness, Property.bAllowSliding);
+						this->GSWeakConstraints->CollisionDetectionSpatialHash(this->Evolution->Particles(), *SurfaceVertices, *SurfaceTriangleMesh, *ParticleComponentIndex, TriangleMeshCollisions->GetDynamicSpatialHash(), Property.SpringCollisionSearchRadius, Property.SpringCollisionStiffness, Property.bAllowSliding);
 					}
-				};
-
-			int32 InitIndex1 = Evolution->AddConstraintInitRange(1, true);
-
-
-			Evolution->ConstraintInits()[InitIndex1] =
-				[this](FSolverParticles& InParticles, const FSolverReal Dt)
-				{
 					TArray<TArray<int32>> WCCollisionConstraints, WCCollisionIncidentElements, WCCollisionIncidentElementsLocal;
 					this->GSWeakConstraints->ComputeCollisionWCDataSimplified(WCCollisionConstraints, WCCollisionIncidentElements, WCCollisionIncidentElementsLocal);
 					this->GSMainConstraint->AddTransientConstraints(WCCollisionConstraints, WCCollisionIncidentElements, WCCollisionIncidentElementsLocal);
 					this->GSMainConstraint->InitTransientColor(InParticles);
 				};
-
 		}
+		if (Property.bDoSphereRepulsion)
+		{
+			GSSphereRepulsionConstraints.Reset(new FGaussSeidelSphereRepulsionConstraints<FSolverReal, FSolverParticles>(Property.SphereRepulsionRadius, Property.SphereRepulsionStiffness, Evolution->Particles(), GDeformableXPBDWeakConstraintParams));
+			int32 InitIndex = Evolution->AddConstraintInitRange(1, true);
+			
+			Evolution->ConstraintInits()[InitIndex] 
+				= [this](FSolverParticles& InParticles, const FSolverReal Dt)
+				{
+					this->GSSphereRepulsionConstraints->UpdateSphereRepulsionConstraints(this->Evolution->Particles(), *SurfaceVertices, *ParticleComponentIndex);
+					TArray<TArray<int32>> SphereRepulsionConstraints, SphereRepulsionIncidentElements, SphereRepulsionIncidentElementsLocal;
+					this->GSSphereRepulsionConstraints->ReturnSphereRepulsionConstraints(SphereRepulsionConstraints, SphereRepulsionIncidentElements, SphereRepulsionIncidentElementsLocal);
+					//this->GSMainConstraint->AddDynamicConstraints(SphereRepulsionConstraints, SphereRepulsionIncidentElements, SphereRepulsionIncidentElementsLocal);
+					//this->GSMainConstraint->InitDynamicColor(InParticles);
+					this->GSMainConstraint->AddTransientConstraints(SphereRepulsionConstraints, SphereRepulsionIncidentElements, SphereRepulsionIncidentElementsLocal);
+					this->GSMainConstraint->InitTransientColor(InParticles);
+					this->GSSphereRepulsionConstraints->Init(InParticles, Dt);
+				};
 
+			int32 TransientIndex = GSMainConstraint->AddTransientConstraintResidualAndHessianRange(1);
+			GSMainConstraint->TransientConstraintResidualAndHessian()[TransientIndex] =
+				[this](const FSolverParticles& Particles, const int32 ConstraintIndex, const int32 ConstraintIndexLocal, const FSolverReal Dt, TVec3<FSolverReal>& ParticleResidual, Chaos::PMatrix<FSolverReal, 3, 3>& ParticleHessian)
+				{
+					this->GSSphereRepulsionConstraints->AddSphereRepulsionResidualAndHessian(Particles, ConstraintIndex, ConstraintIndexLocal, Dt, ParticleResidual, ParticleHessian);
+				};
 
+			if (!Property.bEnablePositionTargets)
+			{
+				int32 PerNodeIndex = GSMainConstraint->AddPerNodeHessianRange(1);
+				GSMainConstraint->PerNodeHessian()[PerNodeIndex] = 
+					[this](const int32 p, const FSolverReal Dt, Chaos::PMatrix<FSolverReal, 3, 3>& ParticleHessian)
+					{
+						this->GSSphereRepulsionConstraints->AddSphereRepulsionHessian(p, Dt, ParticleHessian);
+					};
+			}
+		}
 	}
 
 	void FDeformableSolver::InitializeMuscleActivationVariables()
 	{
 		int32 InitIndex = Evolution->AddConstraintInitRange(1, true);
-		Evolution->ConstraintInits()[InitIndex] =
-			[this](FSolverParticles& InParticles, const FSolverReal Dt)
+		Evolution->ConstraintInits()[InitIndex] = [this](FSolverParticles& InParticles, const FSolverReal Dt)
+		{
+			this->MuscleActivationConstraints->UpdateLengthBasedMuscleActivation(InParticles);
+			if (Property.bUseGaussSeidelConstraints)
 			{
-				this->MuscleActivationConstraints->UpdateLengthBasedMuscleActivation(InParticles);
-				if (Property.bUseGaussSeidelConstraints)
+				if (Property.bUseGSNeohookean)
 				{
-					if (Property.bUseGSNeohookean)
-					{
-						this->MuscleActivationConstraints->ApplyMuscleActivation(*this->GSNeohookeanConstraints);
-					}
-					else
-					{
-						this->MuscleActivationConstraints->ApplyMuscleActivation(*this->GSCorotatedConstraints);
-					}
+					this->MuscleActivationConstraints->ApplyMuscleActivation(*this->GSNeohookeanConstraints);
 				}
 				else
 				{
-					this->MuscleActivationConstraints->ApplyMuscleActivation(*this->CorotatedConstraints[0]);
+					this->MuscleActivationConstraints->ApplyMuscleActivation(*this->GSCorotatedConstraints);
 				}
-			};
+			}
+			else
+			{
+				this->MuscleActivationConstraints->ApplyMuscleActivation(*this->CorotatedConstraints[0]);
+			}
+		};
 	}
 
 	void FDeformableSolver::InitializeMuscleActivation(FFleshThreadingProxy& Proxy)
diff --git a/Engine/Source/Runtime/Experimental/Chaos/Public/Chaos/Deformable/ChaosDeformableSolver.h b/Engine/Source/Runtime/Experimental/Chaos/Public/Chaos/Deformable/ChaosDeformableSolver.h
--- a/Engine/Source/Runtime/Experimental/Chaos/Public/Chaos/Deformable/ChaosDeformableSolver.h
+++ b/Engine/Source/Runtime/Experimental/Chaos/Public/Chaos/Deformable/ChaosDeformableSolver.h
@@ -15,6 +15,7 @@
 #include "Chaos/Deformable/GaussSeidelNeohookeanConstraints.h"
 #include "Chaos/Deformable/GaussSeidelWeakConstraints.h"
 #include "Chaos/Deformable/GaussSeidelDynamicWeakConstraints.h"
+#include "Chaos/Deformable/GaussSeidelSphereRepulsionConstraints.h"
 #include "Chaos/Deformable/GaussSeidelMainConstraint.h"
 #include "Chaos/Deformable/MuscleActivationConstraints.h"
 #include "Chaos/Deformable/GaussSeidelCorotatedCodimensionalConstraints.h"
@@ -193,6 +194,7 @@ namespace Chaos::Softs
 		TUniquePtr<Softs::FGaussSeidelLinearCodimensionalConstraints<Softs::FSolverReal, Softs::FSolverParticles>> GSLinearCodConstraints;
 		TUniquePtr<Softs::FGaussSeidelWeakConstraints<Softs::FSolverReal, Softs::FSolverParticles>> GSWeakConstraints;
 		TUniquePtr<Softs::FGaussSeidelDynamicWeakConstraints<Softs::FSolverReal, Softs::FSolverParticles>> GSDynamicWeakConstraints;
+		TUniquePtr<Softs::FGaussSeidelSphereRepulsionConstraints<Softs::FSolverReal, Softs::FSolverParticles>> GSSphereRepulsionConstraints;
 		TArray<TUniquePtr<Softs::FXPBDWeakConstraints<Softs::FSolverReal, Softs::FSolverParticles>>> WeakConstraints;
 		TArray<TUniquePtr<Softs::FBlendedXPBDCorotatedConstraints<Softs::FSolverReal, Softs::FSolverParticles>>> BlendedCorotatedConstraints;
 		TUniquePtr<Softs::FXPBDGridBasedCorotatedConstraints<Softs::FSolverReal, Softs::FSolverParticles>> GridBasedCorotatedConstraint;
@@ -218,8 +220,9 @@ namespace Chaos::Softs
 		TUniquePtr <TArray<TVec3<int32>>> AllUnconstrainedSurfaceElementsSkin;          //correspond to the triangle mesh elements that are simulated using linear cod constraints
 		TUniquePtr <TArray<FSolverReal>> AllCorotatedCodEMeshArray;
 		TUniquePtr <TArray<FSolverReal>> AllSkinEMeshArray;
-		TArray<int32> ParticleComponentIndex;
+		TUniquePtr <TArray<int32>> ParticleComponentIndex;
 		TMap<int32, TSet<int32>> ParticleTriangleExclusionMap;
+		TUniquePtr <FSolverParticles> ReferencePositions;
 		//Muscle Activation Variables
 		TUniquePtr<Softs::FMuscleActivationConstraints<Softs::FSolverReal, Softs::FSolverParticles>> MuscleActivationConstraints;
 
diff --git a/Engine/Source/Runtime/Experimental/Chaos/Public/Chaos/Deformable/ChaosDeformableSolverTypes.h b/Engine/Source/Runtime/Experimental/Chaos/Public/Chaos/Deformable/ChaosDeformableSolverTypes.h
--- a/Engine/Source/Runtime/Experimental/Chaos/Public/Chaos/Deformable/ChaosDeformableSolverTypes.h
+++ b/Engine/Source/Runtime/Experimental/Chaos/Public/Chaos/Deformable/ChaosDeformableSolverTypes.h
@@ -37,12 +37,15 @@ namespace Chaos::Softs
 			bool InbUseSOR = true,
 			FSolverReal InOmegaSOR = (FSolverReal)1.6,
 			bool InbUseGSNeohookean = false,
-			bool InbDoSelfCollision = false,
-			bool InbDoInComponentSelfCollision = false,
+			bool InbDoSpringCollision = false,
+			bool InbDoInComponentSpringCollision = false,
 			int32 InNRingExcluded = 1,
-			FSolverReal InCollisionSearchRadius = (FSolverReal)0,
-			FSolverReal InCollisionSpringStiffness = (FSolverReal)500.,
+			FSolverReal InSpringCollisionSearchRadius = (FSolverReal)0,
+			FSolverReal InSpringCollisionStiffness = (FSolverReal)500.,
 			bool InbAllowSliding = true,
+			bool InbDoSphereRepulsion = false,
+			FSolverReal InSphereRepulsionRadius = (FSolverReal)0,
+			FSolverReal InSphereRepulsionStiffness = (FSolverReal)500.,
 			bool InbDoMuscleActivation = false,
 			bool InbCollideWithFullMesh = false, 
 			bool InbEnableDynamicSprings = true)
@@ -67,12 +70,15 @@ namespace Chaos::Softs
 			, bUseSOR(InbUseSOR)
 			, OmegaSOR(InOmegaSOR)
 			, bUseGSNeohookean(InbUseGSNeohookean)
-			, bDoSelfCollision(InbDoSelfCollision)
-			, bDoInComponentSelfCollision(InbDoInComponentSelfCollision)
+			, bDoSpringCollision(InbDoSpringCollision)
+			, bDoInComponentSpringCollision(InbDoInComponentSpringCollision)
 			, NRingExcluded(InNRingExcluded)
-			, CollisionSearchRadius(InCollisionSearchRadius)
-			, CollisionSpringStiffness(InCollisionSpringStiffness)
+			, SpringCollisionSearchRadius(InSpringCollisionSearchRadius)
+			, SpringCollisionStiffness(InSpringCollisionStiffness)
 			, bAllowSliding(InbAllowSliding)
+			, bDoSphereRepulsion(InbDoSphereRepulsion)
+			, SphereRepulsionRadius(InSphereRepulsionRadius)
+			, SphereRepulsionStiffness(InSphereRepulsionStiffness)
 			, bDoMuscleActivation(InbDoMuscleActivation)
 			, bCollideWithFullMesh(InbCollideWithFullMesh)
 			, bEnableDynamicSprings(InbEnableDynamicSprings)
@@ -99,12 +105,15 @@ namespace Chaos::Softs
 		bool bUseSOR = true;
 		FSolverReal OmegaSOR = (FSolverReal)1.6;
 		bool bUseGSNeohookean = false;
-		bool bDoSelfCollision = false;
-		bool bDoInComponentSelfCollision = false;
+		bool bDoSpringCollision = false;
+		bool bDoInComponentSpringCollision = false;
 		int32 NRingExcluded = 1;
-		FSolverReal CollisionSearchRadius = (FSolverReal)0;
-		FSolverReal CollisionSpringStiffness = (FSolverReal)500.;
+		FSolverReal SpringCollisionSearchRadius = (FSolverReal)0;
+		FSolverReal SpringCollisionStiffness = (FSolverReal)500.;
 		bool bAllowSliding = true;
+		bool bDoSphereRepulsion = false;
+		FSolverReal SphereRepulsionRadius = (FSolverReal)0;
+		FSolverReal SphereRepulsionStiffness = (FSolverReal)500.;
 		bool bDoMuscleActivation = false;
 		bool bCollideWithFullMesh = false;
 		bool bEnableDynamicSprings = true; 
diff --git a/Engine/Source/Runtime/Experimental/Chaos/Public/Chaos/Deformable/GaussSeidelSphereRepulsionConstraints.h b/Engine/Source/Runtime/Experimental/Chaos/Public/Chaos/Deformable/GaussSeidelSphereRepulsionConstraints.h
new file mode 100644
--- /dev/null
+++ b/Engine/Source/Runtime/Experimental/Chaos/Public/Chaos/Deformable/GaussSeidelSphereRepulsionConstraints.h
@@ -0,0 +1,208 @@
+// Copyright Epic Games, Inc. All Rights Reserved.
+#pragma once
+
+#include "Chaos/BoundingVolumeHierarchy.h"
+#include "Chaos/PBDSpringConstraintsBase.h"
+#include "Chaos/XPBDCorotatedConstraints.h"
+#include "ChaosStats.h"
+#include "Chaos/ImplicitQRSVD.h"
+#include "Chaos/GraphColoring.h"
+#include "Chaos/NewtonCorotatedCache.h"
+#include "Chaos/Framework/Parallel.h"
+#include "Chaos/Utilities.h"
+#include "Chaos/DebugDrawQueue.h"
+#include "Chaos/XPBDWeakConstraints.h"
+#include "Chaos/Triangle.h"
+#include "Chaos/TriangleCollisionPoint.h"
+#include "Chaos/TriangleMesh.h"
+#include <unordered_map>
+#include "Chaos/HierarchicalSpatialHash.h"
+#include "Chaos/PBDSelfCollisionSphereConstraints.h"
+
+namespace Chaos::Softs
+{
+	namespace Private
+	{
+		struct FSphereSpatialEntry
+		{
+			const TConstArrayView<FSolverVec3>* Points;
+			int32 Index;
+
+			FSolverVec3 X() const
+			{
+				return (*Points)[Index];
+			}
+
+			template<typename TPayloadType>
+			int32 GetPayload(int32) const
+			{
+				return Index;
+			}
+		};
+	}
+
+	using Chaos::TVec3;
+	template <typename T, typename ParticleType>
+	class FGaussSeidelSphereRepulsionConstraints// : public FPBDSelfCollisionSphereConstraintsBase
+	{
+	public:
+		//TODO(Yushan): Add unittest for Gauss Seidel Sphere Repulsion Constraints
+		FGaussSeidelSphereRepulsionConstraints(FSolverReal InRadius, FSolverReal InStiffness, const ParticleType& InParticles, const FDeformableXPBDWeakConstraintParams& InParams): Radius(InRadius), Stiffness(InStiffness), DebugDrawParams(InParams)
+		{
+			const TArrayCollectionArray<FPAndInvM>& PAndInvM = InParticles.GetPAndInvM();
+			ReferencePositions.SetNum(PAndInvM.Num());
+			for (int32 i = 0; i < PAndInvM.Num(); ++i)
+			{
+				ReferencePositions[i] = PAndInvM[i].P;
+			}
+		}
+
+		virtual ~FGaussSeidelSphereRepulsionConstraints(){}
+
+		//Energy = k/2*(2r-d)^2
+		//Residual = de/dx = -force = -k*(2r-d)*dd/dx
+		//Hessian = de2/dx2 = k*dd/dx*dd/dx-k*(2r-d)*(-dd/dx*dd/dx^T+I)/d
+		void AddSphereRepulsionResidual(const ParticleType& InParticles, const int32 p, const T Dt, TVec3<T>& res)
+		{
+
+		}
+
+		void AddSphereRepulsionHessian(const int32 p, const T Dt, Chaos::PMatrix<T, 3, 3>& ParticleHessian)
+		{
+
+		}
+
+		void AddSphereRepulsionResidualAndHessian(const ParticleType& InParticles, const int32 ConstraintIndex, const int32 LocalIndex, const T Dt, TVec3<T>& ParticleResidual, Chaos::PMatrix<T, 3, 3>& ParticleHessian)
+		{
+			Chaos::TVec3<T> x0 = InParticles.P(Constraints[ConstraintIndex][0]);
+			Chaos::TVec3<T> x1 = InParticles.P(Constraints[ConstraintIndex][1]);
+			Chaos::TVec3<T> normal = (x1 - x0).GetSafeNormal();
+			normal = (LocalIndex ? normal : -normal); //dd/dx
+			T dist = (x1 - x0).Size();
+			FSolverReal penetration = 2 * Radius - dist; //2r-d
+
+			if (penetration > T(0)) {
+				T dist_inv = T(1) / (dist+T(1e-12));
+				Chaos::PMatrix<T, 3, 3> outer = Chaos::PMatrix<T, 3, 3>::OuterProduct(normal, normal);
+				Chaos::PMatrix<T, 3, 3> A = (Chaos::PMatrix<T, 3, 3>::Identity - outer) * dist_inv; //(-dd/dx*dd/dx^T+I)/d
+				ParticleHessian += Dt * Dt * ConstraintStiffness[ConstraintIndex] * (outer - penetration * A);
+				ParticleResidual += -Dt * Dt * penetration * ConstraintStiffness[ConstraintIndex] * normal;
+			}
+		}
+
+		void VisualizeAllBindings(const FSolverParticles& InParticles, const T Dt) const
+		{
+#if WITH_EDITOR
+			auto DoubleVert = [](Chaos::TVec3<T> V) { return FVector3d(V.X, V.Y, V.Z); };
+			for (int32 i = 0; i < Constraints.Num(); i++)
+			{
+				FVector3d SourcePos = DoubleVert(InParticles.P(Constraints[i][0]));
+				FVector3d TargetPos = DoubleVert(InParticles.P(Constraints[i][1]));
+
+				float ParticleThickness = DebugDrawParams.DebugParticleWidth;
+				float LineThickness = DebugDrawParams.DebugLineWidth;
+
+				Chaos::FDebugDrawQueue::GetInstance().DrawDebugPoint(SourcePos, FColor::Red, false, Dt, 0, ParticleThickness);
+				Chaos::FDebugDrawQueue::GetInstance().DrawDebugPoint(TargetPos, FColor::Red, false, Dt, 0, ParticleThickness);
+				Chaos::FDebugDrawQueue::GetInstance().DrawDebugLine(SourcePos, TargetPos, FColor::Green, false, Dt, 0, LineThickness);
+			}
+#endif
+		}
+
+		void Init(const FSolverParticles& InParticles, const T Dt) const
+		{
+			if (DebugDrawParams.bVisualizeBindings)
+			{
+				VisualizeAllBindings(InParticles, Dt);
+			}
+
+		}
+
+		void UpdateSphereRepulsionConstraints(const FSolverParticles& Particles, const TArray<int32>& SurfaceVertices, const TArray<int32>& ComponentIndex)
+		{
+			TRACE_CPUPROFILER_EVENT_SCOPE(STAT_GaussSeidelSphereRepulsionConstraints);
+			Constraints.Reset();
+			ConstraintStiffness.Reset();
+			if (SurfaceVertices.Num() == 0)
+			{
+				return;
+			}
+
+			// Build Spatial
+			TArray<Private::FSphereSpatialEntry> Entries;
+			const FSolverReal Diameter = 2.f * Radius;
+			TConstArrayView<FSolverVec3> Points = Particles.XArray();
+
+			Entries.Reset(SurfaceVertices.Num());
+			for (int32 Index = 0; Index < SurfaceVertices.Num(); ++Index)
+			{
+				Entries.Add({ &Points, SurfaceVertices[Index]});
+			}
+			
+			TSpatialHashGridPoints<int32, FSolverReal> SpatialHash(Diameter);
+			SpatialHash.InitializePoints(Entries);
+
+			const FSolverReal DiamSq = FMath::Square(Diameter);
+			constexpr int32 CellRadius = 1; // We set the cell size of the spatial hash such that we only need to look 1 cell away to find proximities.
+			constexpr int32 MaxNumExpectedConnectionsPerParticle = 3;
+			const int32 MaxNumExpectedConnections = MaxNumExpectedConnectionsPerParticle * Entries.Num();
+
+			Constraints = SpatialHash.FindAllSelfProximities(CellRadius, MaxNumExpectedConnections,
+				[this, &Particles, DiamSq, &ComponentIndex](const int32 i1, const int32 i2)
+				{
+					if (ComponentIndex[i1] == ComponentIndex[i2])
+					{
+						return false;
+					}
+					const FSolverReal CombinedMass = Particles.InvM(i1) + Particles.InvM(i2);
+					if (CombinedMass < (FSolverReal)1e-7)
+					{
+						return false;
+					}
+					if (FSolverVec3::DistSquared(this->ReferencePositions[i1], this->ReferencePositions[i2]) < DiamSq)
+					{
+						return false;
+					}
+					return true;
+				}
+			);
+
+			for (const TVec2<int32>& CollisionPair : Constraints)
+			{
+				T TotalMass = T(0);
+				if (Particles.InvM(CollisionPair[0]) > (FSolverReal)1e-7)
+				{
+					TotalMass += Particles.M(CollisionPair[0]);
+				}
+				if (Particles.InvM(CollisionPair[1]) > (FSolverReal)1e-7)
+				{
+					TotalMass += Particles.M(CollisionPair[1]);
+				}
+				ConstraintStiffness.Add(Stiffness*TotalMass/T(2));
+			}
+		}
+
+		void ReturnSphereRepulsionConstraints(TArray<TArray<int32>>& ExtraConstraints, TArray<TArray<int32>>& ExtraIncidentElements, TArray<TArray<int32>>& ExtraIncidentElementsLocal)
+		{
+			ExtraConstraints.Init(TArray<int32>(), Constraints.Num());
+			for (int32 i = 0; i < Constraints.Num(); i++)
+			{
+				ExtraConstraints[i].SetNum(2);
+				ExtraConstraints[i][0] = Constraints[i][0];
+				ExtraConstraints[i][1] = Constraints[i][1];
+			}
+			ExtraIncidentElements = Chaos::Utilities::ComputeIncidentElements(ExtraConstraints, &ExtraIncidentElementsLocal);
+		}
+
+		protected:
+			TArray<TVec2<int32>> Constraints;
+			FSolverReal Radius = FSolverReal(0);
+			FSolverReal Stiffness = FSolverReal(0);
+			TArray<T> ConstraintStiffness;
+		private:
+			TArray<FSolverVec3> ReferencePositions;
+			FDeformableXPBDWeakConstraintParams DebugDrawParams;
+	};
+
+
+}// End namespace Chaos::Softs
