ChaosFlesh: Generate fiber streamlines

diff --git a/Engine/Plugins/Experimental/ChaosFlesh/Source/ChaosFleshNodes/Private/Dataflow/ChaosFleshComputeFiberFieldNode.cpp b/Engine/Plugins/Experimental/ChaosFlesh/Source/ChaosFleshNodes/Private/Dataflow/ChaosFleshComputeFiberFieldNode.cpp
--- a/Engine/Plugins/Experimental/ChaosFlesh/Source/ChaosFleshNodes/Private/Dataflow/ChaosFleshComputeFiberFieldNode.cpp
+++ b/Engine/Plugins/Experimental/ChaosFlesh/Source/ChaosFleshNodes/Private/Dataflow/ChaosFleshComputeFiberFieldNode.cpp
@@ -5,6 +5,7 @@
 #include "Chaos/Math/Poisson.h"
 #include "ChaosFlesh/ChaosFlesh.h"
 #include "ChaosFlesh/TetrahedralCollection.h"
+#include "GeometryCollection/Facades/CollectionMuscleActivationFacade.h"
 
 #include UE_INLINE_GENERATED_CPP_BY_NAME(ChaosFleshComputeFiberFieldNode)
 
@@ -202,3 +203,109 @@ void FComputeFiberFieldNode::ComputeFiberField(
 		MaxIterations,
 		Tolerance);
 }
+
+void FComputeFiberStreamlineNode::Evaluate(Dataflow::FContext& Context, const FDataflowOutput* Out) const
+{
+	//
+	// Gather inputs
+	//
+
+	FManagedArrayCollection InCollection = GetValue<FManagedArrayCollection>(Context, &Collection);
+	FFieldCollection OutVectorField;
+	TArray<int32> InOriginIndices = GetValue<TArray<int32>>(Context, &OriginIndices);
+	TArray<int32> InInsertionIndices = GetValue<TArray<int32>>(Context, &InsertionIndices);
+
+	//
+	// Pull Origin & Insertion data out of the geometry collection.  We may want other ways of specifying
+	// these via an input on the node...
+	//
+
+	// Origin & Insertion
+	TManagedArray<int32>* Origin = nullptr;
+	TManagedArray<int32>* Insertion = nullptr;
+	if (InOriginIndices.IsEmpty() || InInsertionIndices.IsEmpty())
+	{
+		// Origin & Insertion group
+		if (OriginInsertionGroupName.IsEmpty())
+		{
+			UE_LOG(LogChaosFlesh, Warning, TEXT("ComputeFiberFieldNode: Attr 'OriginInsertionGroupName' cannot be empty."));
+			Out->SetValue(MoveTemp(InCollection), Context);
+			return;
+		}
+
+		// Origin vertices
+		if (InOriginIndices.IsEmpty())
+		{
+			if (OriginVertexFieldName.IsEmpty())
+			{
+				UE_LOG(LogChaosFlesh, Warning, TEXT("ComputeFiberFieldNode: Attr 'OriginVertexFieldName' cannot be empty."));
+				Out->SetValue(MoveTemp(InCollection), Context);
+				return;
+			}
+			Origin = InCollection.FindAttribute<int32>(FName(OriginVertexFieldName), FName(OriginInsertionGroupName));
+			if (!Origin)
+			{
+				UE_LOG(LogChaosFlesh, Warning,
+					TEXT("ComputeFiberFieldNode: Failed to find geometry collection attr '%s' in group '%s'"),
+					*OriginVertexFieldName, *OriginInsertionGroupName);
+				Out->SetValue(MoveTemp(InCollection), Context);
+				return;
+			}
+		}
+
+		// Insertion vertices
+		if (InInsertionIndices.IsEmpty())
+		{
+			if (InsertionVertexFieldName.IsEmpty())
+			{
+				UE_LOG(LogChaosFlesh, Warning, TEXT("ComputeFiberFieldNode: Attr 'InsertionVertexFieldName' cannot be empty."));
+				Out->SetValue(MoveTemp(InCollection), Context);
+				return;
+			}
+			Insertion = InCollection.FindAttribute<int32>(FName(InsertionVertexFieldName), FName(OriginInsertionGroupName));
+			if (!Insertion)
+			{
+				UE_LOG(LogChaosFlesh, Warning,
+					TEXT("ComputeFiberFieldNode: Failed to find geometry collection attr '%s' in group '%s'"),
+					*InsertionVertexFieldName, *OriginInsertionGroupName);
+				Out->SetValue(MoveTemp(InCollection), Context);
+				return;
+			}
+		}
+	}
+
+	InOriginIndices = Origin ? Origin->GetConstArray() : InOriginIndices;
+	InInsertionIndices = Insertion ? Insertion->GetConstArray() : InInsertionIndices;
+	if (InOriginIndices.Num() == 0 || InInsertionIndices.Num() == 0)
+	{
+		FindOutput(&VectorField)->SetValue(MoveTemp(OutVectorField), Context);
+		FindOutput(&Collection)->SetValue(MoveTemp(InCollection), Context);
+		return;
+	}
+	//
+	// Compute muscle fiber streamlines
+	// Save streamlines to muscle group
+	//
+	GeometryCollection::Facades::FMuscleActivationFacade MuscleActivation(InCollection);
+
+	TArray<TArray<TArray<FVector3f>>> Streamlines = MuscleActivation.BuildStreamlines(Origin ? Origin->GetConstArray() : InOriginIndices,
+		Insertion ? Insertion->GetConstArray() : InInsertionIndices, NumLinesMultiplier, MaxStreamlineIterations, MaxPointsPerLine);
+
+	//Render streamlines
+	for (int32 i = 0; i < Streamlines.Num(); ++i)
+	{
+		for (int32 j = 0; j < Streamlines[i].Num(); ++j)
+		{
+			for (int32 k = 1; k < Streamlines[i][j].Num(); ++k)
+			{
+				OutVectorField.AddVectorToField(Streamlines[i][j][k - 1], Streamlines[i][j][k]);
+			}
+		}
+	}
+
+	//
+	// Set output(s)
+	//
+	FindOutput(&VectorField)->SetValue(MoveTemp(OutVectorField), Context);
+	FindOutput(&Collection)->SetValue(MoveTemp(InCollection), Context);
+}
\ No newline at end of file
diff --git a/Engine/Plugins/Experimental/ChaosFlesh/Source/ChaosFleshNodes/Private/Dataflow/ChaosFleshComputeMuscleActivationNode.cpp b/Engine/Plugins/Experimental/ChaosFlesh/Source/ChaosFleshNodes/Private/Dataflow/ChaosFleshComputeMuscleActivationNode.cpp
--- a/Engine/Plugins/Experimental/ChaosFlesh/Source/ChaosFleshNodes/Private/Dataflow/ChaosFleshComputeMuscleActivationNode.cpp
+++ b/Engine/Plugins/Experimental/ChaosFlesh/Source/ChaosFleshNodes/Private/Dataflow/ChaosFleshComputeMuscleActivationNode.cpp
@@ -15,117 +15,9 @@ void FComputeMuscleActivationDataNode::Evaluate(Dataflow::FContext& Context, con
 		FManagedArrayCollection InCollection = GetValue<FManagedArrayCollection>(Context, &Collection);
 		TArray<int32> InOriginIndices = GetValue<TArray<int32>>(Context, &OriginIndicesIn);
 		TArray<int32> InInsertionIndices = GetValue<TArray<int32>>(Context, &InsertionIndicesIn);
-		TArray<float> OriginInsertionRestLength;
-		TManagedArray<float>* ContractionVolumeScalePerVertex = InCollection.FindAttribute<float>("ContractionVolumeScale", "Vertices");
-		if (InOriginIndices.Num() > 0 && InInsertionIndices.Num() > 0)
-		{
-			if (TManagedArray<FVector3f>* Vertex = InCollection.FindAttribute<FVector3f>("Vertex", "Vertices"))
-			{
-				if (TManagedArray<FIntVector4>* Elements = InCollection.FindAttribute<FIntVector4>(FTetrahedralCollection::TetrahedronAttribute, FTetrahedralCollection::TetrahedralGroup))
-				{
-					if (TManagedArray<FVector3f>* FiberDirections = InCollection.FindAttribute<FVector3f>("FiberDirection", FTetrahedralCollection::TetrahedralGroup))
-					{
-						GeometryCollection::Facades::FMuscleActivationFacade FMuscleActivation(InCollection);
-						TArray<TArray<int32>> MuscleActivationElements;
-						TArray<TArray<int32>> ComponentOrigins; //One origin node per muscle component
-						TArray<TArray<int32>> ComponentInsertions; //One insertion node per muscle component
-						GeometryCollection::Facades::FCollectionMeshFacade MeshFacade(InCollection);
-						TArray<int32> ComponentIndex = MeshFacade.GetGeometryGroupIndexArray();
-						TMap<int32, int32> ComponentToMuscleIndex; //Component index to muscle index
-						for (int32 i = 0; i < InOriginIndices.Num(); i++)
-						{
-							if (!ComponentToMuscleIndex.Contains(ComponentIndex[InOriginIndices[i]]))
-							{
-								ComponentToMuscleIndex.Add(ComponentIndex[InOriginIndices[i]], ComponentOrigins.Num());
-								ComponentOrigins.SetNum(ComponentOrigins.Num() + 1);
-								ComponentOrigins[ComponentOrigins.Num()-1].Add(InOriginIndices[i]);
-							}
-							else
-							{
-								ComponentOrigins[ComponentToMuscleIndex[ComponentIndex[InOriginIndices[i]]]].Add(InOriginIndices[i]);
-							}
-						}
-						ComponentInsertions.SetNum(ComponentOrigins.Num());
-						for (int32 i = 0; i < InInsertionIndices.Num(); i++)
-						{
-							if (!ComponentToMuscleIndex.Contains(ComponentIndex[InInsertionIndices[i]]))
-							{
-								ensureMsgf(false, TEXT("No origin in this component"));
-							}
-							else
-							{
-								ComponentInsertions[ComponentToMuscleIndex[ComponentIndex[InInsertionIndices[i]]]].Add(InInsertionIndices[i]);
-							}
-						}
-						MuscleActivationElements.SetNum(ComponentOrigins.Num());
-						for (int32 ElemIdx = 0; ElemIdx < Elements->Num(); ElemIdx++)
-						{
-							if (ComponentToMuscleIndex.Contains(ComponentIndex[(*Elements)[ElemIdx][0]]))
-							{
-								MuscleActivationElements[ComponentToMuscleIndex[ComponentIndex[(*Elements)[ElemIdx][0]]]].Add(ElemIdx);
-							}
-						}
-						//Choose one origin-insertion pair per muscle with largest distance apart
-						//use origin-insertion line segment length to estimate activation
-						for (int32 MuscleComponentIdx = 0; MuscleComponentIdx < ComponentOrigins.Num(); MuscleComponentIdx++)
-						{
-							if (ensureMsgf(ComponentOrigins.Num() > 0 && ComponentInsertions.Num() > 0, TEXT("Origin or Insertion missing in the muscle %d"), MuscleComponentIdx))
-							{
-								GeometryCollection::Facades::FMuscleActivationData MuscleActivationData;
-								MuscleActivationData.OriginInsertionRestLength = 0;
-								for (int32 OriginIdx : ComponentOrigins[MuscleComponentIdx])
-								{
-									for (int32 InsertionIdx : ComponentInsertions[MuscleComponentIdx])
-									{
-										float Dist = ((*Vertex)[OriginIdx] - (*Vertex)[InsertionIdx]).Size();
-										if (Dist > MuscleActivationData.OriginInsertionRestLength)
-										{
-											MuscleActivationData.OriginInsertionPair = FIntVector2(OriginIdx, InsertionIdx);
-											MuscleActivationData.OriginInsertionRestLength = Dist;
-										}
-									}
-								}
-								MuscleActivationData.MuscleActivationElement = MuscleActivationElements[MuscleComponentIdx];
-								MuscleActivationData.FiberDirectionMatrix.SetNum(MuscleActivationElements[MuscleComponentIdx].Num());
-								MuscleActivationData.ContractionVolumeScale.SetNum(MuscleActivationElements[MuscleComponentIdx].Num());
-								for (int32 LocalElemIdx = 0; LocalElemIdx < MuscleActivationElements[MuscleComponentIdx].Num(); LocalElemIdx++)
-								{
-									FVector3f V = (*FiberDirections)[MuscleActivationElements[MuscleComponentIdx][LocalElemIdx]];
-									// QR decomposition on vvT for orthogonal directions
-									FVector3f W = V;
-									if (V.X < V.Y)
-									{
-										W.X += 1.f;
-									}
-									else
-									{
-										W.Y += 1.f;
-									}
-									FVector3f U = (V ^ W).GetSafeNormal();
-									W = (U ^ V).GetSafeNormal();
-									MuscleActivationData.FiberDirectionMatrix[LocalElemIdx] = Chaos::PMatrix33d(V, W, U);
-									//Muscle contraction volume scale
-									MuscleActivationData.ContractionVolumeScale[LocalElemIdx] = ContractionVolumeScale;
-									if (ContractionVolumeScalePerVertex)
-									{
-										float AverageScale = 1.f;
-										for (int32 ie = 0; ie < 4; ie++)
-										{
-											AverageScale += (*ContractionVolumeScalePerVertex)[(*Elements)[MuscleActivationElements[MuscleComponentIdx][LocalElemIdx]][ie]];
-											AverageScale /= 4.f;
-										}
-										MuscleActivationData.ContractionVolumeScale[LocalElemIdx] *= AverageScale;
-									}
-								}
-								
-								
-								FMuscleActivation.AddMuscleActivationData(MuscleActivationData);
-							}
-						}
-					}
-				}
-			}
-		}
+		
+		GeometryCollection::Facades::FMuscleActivationFacade FMuscleActivation(InCollection);
+		FMuscleActivation.SetUpMuscleActivation(InOriginIndices, InInsertionIndices, ContractionVolumeScale);
 		Out->SetValue(MoveTemp(InCollection), Context);
 	}
 }
diff --git a/Engine/Plugins/Experimental/ChaosFlesh/Source/ChaosFleshNodes/Private/Dataflow/ChaosFleshNodesPlugin.cpp b/Engine/Plugins/Experimental/ChaosFlesh/Source/ChaosFleshNodes/Private/Dataflow/ChaosFleshNodesPlugin.cpp
--- a/Engine/Plugins/Experimental/ChaosFlesh/Source/ChaosFleshNodes/Private/Dataflow/ChaosFleshNodesPlugin.cpp
+++ b/Engine/Plugins/Experimental/ChaosFlesh/Source/ChaosFleshNodes/Private/Dataflow/ChaosFleshNodesPlugin.cpp
@@ -45,6 +45,7 @@ void IChaosFleshNodesPlugin::StartupModule()
 	DATAFLOW_NODE_REGISTER_CREATION_FACTORY(FAuthorSceneCollisionCandidates);
 	DATAFLOW_NODE_REGISTER_CREATION_FACTORY(FCalculateTetMetrics);
 	DATAFLOW_NODE_REGISTER_CREATION_FACTORY(FComputeFiberFieldNode);
+	DATAFLOW_NODE_REGISTER_CREATION_FACTORY(FComputeFiberStreamlineNode);
 	DATAFLOW_NODE_REGISTER_CREATION_FACTORY(FComputeIslandsNode);
 	DATAFLOW_NODE_REGISTER_CREATION_FACTORY(FComputeMuscleActivationDataNode);
 	DATAFLOW_NODE_REGISTER_CREATION_FACTORY(FCreateTetrahedronDataflowNode);
diff --git a/Engine/Plugins/Experimental/ChaosFlesh/Source/ChaosFleshNodes/Private/Dataflow/ChaosFleshVisualizeFiberFieldNode.cpp b/Engine/Plugins/Experimental/ChaosFlesh/Source/ChaosFleshNodes/Private/Dataflow/ChaosFleshVisualizeFiberFieldNode.cpp
--- a/Engine/Plugins/Experimental/ChaosFlesh/Source/ChaosFleshNodes/Private/Dataflow/ChaosFleshVisualizeFiberFieldNode.cpp
+++ b/Engine/Plugins/Experimental/ChaosFlesh/Source/ChaosFleshNodes/Private/Dataflow/ChaosFleshVisualizeFiberFieldNode.cpp
@@ -16,23 +16,31 @@ void FVisualizeFiberFieldNode::Evaluate(Dataflow::FContext& Context, const FData
 		
 		if (TManagedArray<FVector3f>* Vertex = InCollection.FindAttribute<FVector3f>("Vertex", "Vertices"))
 		{
-			if (TManagedArray<FIntVector4>* Elements = InCollection.FindAttribute<FIntVector4>(FTetrahedralCollection::TetrahedronAttribute, FTetrahedralCollection::TetrahedralGroup))
+			if (TManagedArray<FLinearColor>* Color = InCollection.FindAttribute<FLinearColor>("Color", "Vertices"))
 			{
-				if (TManagedArray<FVector3f>* FiberDirections = InCollection.FindAttribute<FVector3f>("FiberDirection", FTetrahedralCollection::TetrahedralGroup))
+				if (TManagedArray<FIntVector4>* Elements = InCollection.FindAttribute<FIntVector4>(FTetrahedralCollection::TetrahedronAttribute, FTetrahedralCollection::TetrahedralGroup))
 				{
-					ensureMsgf(Elements->Num() == FiberDirections->Num(), TEXT("Fiber direction has different size than elements"));
-					for (int32 ElemIndex = 0; ElemIndex < Elements->Num(); ElemIndex++)
+					if (TManagedArray<FVector3f>* FiberDirections = InCollection.FindAttribute<FVector3f>("FiberDirection", FTetrahedralCollection::TetrahedralGroup))
 					{
-						if ((*FiberDirections)[ElemIndex].Length() > 1e-10)
+						ensureMsgf(Elements->Num() == FiberDirections->Num(), TEXT("Fiber direction has different size than elements"));
+						for (int32 ElemIndex = 0; ElemIndex < Elements->Num(); ElemIndex++)
 						{
-							FVector3f VectorStart = { 0,0,0 };
-							for (int32 LocalIndex = 0; LocalIndex < 4; LocalIndex++)
+							if ((*FiberDirections)[ElemIndex].Length() > 1e-10)
 							{
-								VectorStart += (*Vertex)[(*Elements)[ElemIndex][LocalIndex]];
+								FVector3f VectorStart = { 0,0,0 };
+								FLinearColor VectorColor(EForceInit::ForceInitToZero);
+
+								for (int32 LocalIndex = 0; LocalIndex < 4; LocalIndex++)
+								{
+									VectorStart += (*Vertex)[(*Elements)[ElemIndex][LocalIndex]];
+									VectorColor += (*Color)[(*Elements)[ElemIndex][LocalIndex]];
+								}
+								VectorStart /= float(4);
+								VectorColor /= float(4);
+								FVector3f VectorEnd = VectorStart + (*FiberDirections)[ElemIndex] * VectorScale;
+								int32 VectorIndex = OutVectorField.AddVectorToField(VectorStart, VectorEnd);
+								OutVectorField.SetColorOnVector(VectorIndex, VectorColor);
 							}
-							VectorStart /= float(4);
-							FVector3f VectorEnd = VectorStart + (*FiberDirections)[ElemIndex] * VectorScale;
-							OutVectorField.AddVectorToField(VectorStart, VectorEnd);
 						}
 					}
 				}
diff --git a/Engine/Plugins/Experimental/ChaosFlesh/Source/ChaosFleshNodes/Public/Dataflow/ChaosFleshComputeFiberFieldNode.h b/Engine/Plugins/Experimental/ChaosFlesh/Source/ChaosFleshNodes/Public/Dataflow/ChaosFleshComputeFiberFieldNode.h
--- a/Engine/Plugins/Experimental/ChaosFlesh/Source/ChaosFleshNodes/Public/Dataflow/ChaosFleshComputeFiberFieldNode.h
+++ b/Engine/Plugins/Experimental/ChaosFlesh/Source/ChaosFleshNodes/Public/Dataflow/ChaosFleshComputeFiberFieldNode.h
@@ -5,6 +5,7 @@
 #include "CoreMinimal.h"
 #include "Dataflow/DataflowCore.h"
 #include "Dataflow/DataflowEngine.h"
+#include "Field/FieldSystemTypes.h"
 
 #include "ChaosFleshComputeFiberFieldNode.generated.h"
 
@@ -72,3 +73,59 @@ public:
 		TArray<FVector3f>& Directions,
 		TArray<float>& ScalarField) const;
 };
+
+/**
+* Computes fiber streamlines (line segments) flowing from muscle origins to insertions in the muscle fiber field.
+*/
+USTRUCT(meta = (DataflowFlesh))
+struct FComputeFiberStreamlineNode : public FDataflowNode
+{
+	GENERATED_USTRUCT_BODY()
+	DATAFLOW_NODE_DEFINE_INTERNAL(FComputeFiberStreamlineNode, "ComputeFiberStreamline", "Flesh", "")
+	DATAFLOW_NODE_RENDER_TYPE("VolumeRender", FFieldCollection::StaticType(), "VectorField")
+
+public:
+	//typedef FManagedArrayCollection DataType;
+
+	UPROPERTY(meta = (DataflowInput, DataflowOutput, DisplayName = "Collection", DataflowPassthrough = "Collection"))
+	FManagedArrayCollection Collection;
+
+	UPROPERTY(meta = (DataflowInput, DisplayName = "OriginVertexIndices"))
+	TArray<int32> OriginIndices;
+
+	UPROPERTY(meta = (DataflowInput, DisplayName = "InsertionVertexIndices"))
+	TArray<int32> InsertionIndices;
+
+	UPROPERTY(EditAnywhere, Category = "Dataflow")
+	FString OriginInsertionGroupName = FString();
+
+	UPROPERTY(EditAnywhere, Category = "Dataflow")
+	FString OriginVertexFieldName = FString("Origin");
+
+	UPROPERTY(EditAnywhere, Category = "Dataflow")
+	FString InsertionVertexFieldName = FString("Insertion");
+
+	UPROPERTY(EditAnywhere, Category = "Dataflow")
+	int32 MaxStreamlineIterations = 500;
+
+	UPROPERTY(EditAnywhere, Category = "Dataflow")
+	int32 MaxPointsPerLine = 20;
+
+	UPROPERTY(EditAnywhere, Category = "Dataflow")
+	int32 NumLinesMultiplier = 1;
+
+	UPROPERTY(meta = (DataflowOutput, DisplayName = "VectorField"))
+	FFieldCollection VectorField;
+
+	FComputeFiberStreamlineNode(const Dataflow::FNodeParameters& InParam, FGuid InGuid = FGuid::NewGuid())
+		: FDataflowNode(InParam, InGuid)
+	{
+		RegisterInputConnection(&Collection);
+		RegisterInputConnection(&OriginIndices);
+		RegisterInputConnection(&InsertionIndices);
+		RegisterOutputConnection(&Collection, &Collection);
+		RegisterOutputConnection(&VectorField);
+	}
+
+	virtual void Evaluate(Dataflow::FContext& Context, const FDataflowOutput* Out) const override;
+};
\ No newline at end of file
diff --git a/Engine/Plugins/Experimental/Dataflow/Source/DataflowEditor/Private/Dataflow/DataflowEngineRendering.cpp b/Engine/Plugins/Experimental/Dataflow/Source/DataflowEditor/Private/Dataflow/DataflowEngineRendering.cpp
--- a/Engine/Plugins/Experimental/Dataflow/Source/DataflowEditor/Private/Dataflow/DataflowEngineRendering.cpp
+++ b/Engine/Plugins/Experimental/Dataflow/Source/DataflowEditor/Private/Dataflow/DataflowEngineRendering.cpp
@@ -555,34 +555,33 @@ namespace Dataflow
 				FName PrimaryOutput = State.GetRenderOutputs()[0]; // "VectorField"
 				if (PrimaryOutput.IsEqual(FName("VectorField")))
 				{
-					FFieldCollection Default;
-					const FFieldCollection& Collection = State.GetValue<FFieldCollection>(PrimaryOutput, Default);
-					TArray<TPair<FVector3f, FVector3f>> VectorField = Collection.GetVectorField();
-					const int32 NumVertices = 3 * VectorField.Num();
-					const int32 NumTriangles = VectorField.Num();
-
-					TArray<FVector3f> Vertices; Vertices.AddUninitialized(NumVertices);
-					TArray<FIntVector> Tris; Tris.AddUninitialized(NumTriangles);
-					TArray<FVector3f> VertexNormals; VertexNormals.AddUninitialized(NumVertices);
-					TArray<FLinearColor> VertexColors; VertexColors.AddUninitialized(NumVertices);
-
-					for (int32 i = 0; i < VectorField.Num(); i++)
-					{
-
-						FVector3f Dir = VectorField[i].Value - VectorField[i].Key;
-						FVector3f DirAdd = Dir;
-						DirAdd.X += 1.f;
-						FVector3f OrthogonalDir = (Dir ^ DirAdd).GetSafeNormal();
-						Tris[i] = FIntVector(3 * i, 3 * i + 1, 3 * i + 2);
-						Vertices[3 * i] = VectorField[i].Key;
-						Vertices[3 * i + 1] = VectorField[i].Value;
-						Vertices[3 * i + 2] = VectorField[i].Key + float(0.1) * Dir.Size() * OrthogonalDir;
-						FVector3f TriangleNormal = (OrthogonalDir ^ Dir).GetSafeNormal();
-						VertexNormals[3 * i] = TriangleNormal;
-						VertexNormals[3 * i + 1] = TriangleNormal;
-						VertexNormals[3 * i + 2] = TriangleNormal;
-						VertexColors[i] = FLinearColor(FDataflowEditorModule::SurfaceColor);
-					}
+						FFieldCollection Default;
+						const FFieldCollection& Collection = State.GetValue<FFieldCollection>(PrimaryOutput, Default);
+						TArray<TPair<FVector3f, FVector3f>> VectorField = Collection.GetVectorField();
+						TArray<FLinearColor> VertexColors = Collection.GetVectorColor();
+						const int32 NumVertices = 3 * VectorField.Num();
+						const int32 NumTriangles = VectorField.Num();
+
+						TArray<FVector3f> Vertices; Vertices.AddUninitialized(NumVertices);
+						TArray<FIntVector> Tris; Tris.AddUninitialized(NumTriangles);
+						TArray<FVector3f> VertexNormals; VertexNormals.AddUninitialized(NumVertices);
+
+						for (int32 i = 0; i < VectorField.Num(); i++)
+						{
+							
+							FVector3f Dir = VectorField[i].Value - VectorField[i].Key;
+							FVector3f DirAdd = Dir;
+							DirAdd.X += 1.f;
+							FVector3f OrthogonalDir = (Dir^ DirAdd).GetSafeNormal();
+							Tris[i] = FIntVector(3*i, 3*i+1, 3*i+2);
+							Vertices[3*i] = VectorField[i].Key;
+							Vertices[3*i+1] = VectorField[i].Value;
+							Vertices[3*i+2] = VectorField[i].Key + float(0.1) * Dir.Size() * OrthogonalDir;
+							FVector3f TriangleNormal = (OrthogonalDir ^ Dir).GetSafeNormal();
+							VertexNormals[3*i] = TriangleNormal;
+							VertexNormals[3*i+1] = TriangleNormal;
+							VertexNormals[3*i+2] = TriangleNormal;
+						}
 
 					int32 GeometryIndex = RenderCollection.StartGeometryGroup(State.GetGuid().ToString());
 					RenderCollection.AddSurface(MoveTemp(Vertices), MoveTemp(Tris), MoveTemp(VertexNormals), MoveTemp(VertexColors));
diff --git a/Engine/Source/Runtime/Experimental/Chaos/Private/GeometryCollection/Facades/CollectionMuscleActivationFacade.cpp b/Engine/Source/Runtime/Experimental/Chaos/Private/GeometryCollection/Facades/CollectionMuscleActivationFacade.cpp
--- a/Engine/Source/Runtime/Experimental/Chaos/Private/GeometryCollection/Facades/CollectionMuscleActivationFacade.cpp
+++ b/Engine/Source/Runtime/Experimental/Chaos/Private/GeometryCollection/Facades/CollectionMuscleActivationFacade.cpp
@@ -5,26 +5,36 @@
 =============================================================================*/
 
 #include "GeometryCollection/Facades/CollectionMuscleActivationFacade.h"
-#include "GeometryCollection/GeometryCollection.h"
+
+#include "Chaos/Utilities.h"
+#include "Intersection/IntrRay3Triangle3.h"
+#include "TriangleTypes.h"
+#include "GeometryCollection/Facades/CollectionMeshFacade.h"
 
 namespace GeometryCollection::Facades
 {
 	// Attributes
 	const FName FMuscleActivationFacade::GroupName("MuscleActivation");
+	const FName FMuscleActivationFacade::GeometryGroupIndex("GeometryGroupIndex");
 	const FName FMuscleActivationFacade::MuscleActivationElement("MuscleActivationElement");
 	const FName FMuscleActivationFacade::OriginInsertionPair("OriginInsertionPair");
 	const FName FMuscleActivationFacade::OriginInsertionRestLength("OriginInsertionRestLength");
 	const FName FMuscleActivationFacade::FiberDirectionMatrix("FiberDirectionMatrix");
 	const FName FMuscleActivationFacade::ContractionVolumeScale("ContractionVolumeScale");
+	const FName FMuscleActivationFacade::FiberStreamline("FiberStreamline");
+	const FName FMuscleActivationFacade::FiberStreamlineRestLength("FiberStreamlineRestLength");
 
 	FMuscleActivationFacade::FMuscleActivationFacade(FManagedArrayCollection& InCollection)
 		: ConstCollection(InCollection)
 		, Collection(&InCollection)
+		, GeometryGroupIndexAttribute(InCollection, GeometryGroupIndex, GroupName, FGeometryCollection::GeometryGroup)
 		, MuscleActivationElementAttribute(InCollection, MuscleActivationElement, GroupName, "Tetrahedral")
 		, OriginInsertionPairAttribute(InCollection, OriginInsertionPair, GroupName, FGeometryCollection::VerticesGroup)
 		, OriginInsertionRestLengthAttribute(InCollection, OriginInsertionRestLength, GroupName)
 		, FiberDirectionMatrixAttribute(InCollection, FiberDirectionMatrix, GroupName)
 		, ContractionVolumeScaleAttribute(InCollection, ContractionVolumeScale, GroupName)
+		, FiberStreamlineAttribute(InCollection, FiberStreamline, GroupName)
+		, FiberStreamlineRestLengthAttribute(InCollection, FiberStreamlineRestLength, GroupName)
 	{
 		DefineSchema();
 	}
@@ -32,29 +42,36 @@ namespace GeometryCollection::Facades
 	FMuscleActivationFacade::FMuscleActivationFacade(const FManagedArrayCollection& InCollection)
 		: ConstCollection(InCollection)
 		, Collection(nullptr)
+		, GeometryGroupIndexAttribute(InCollection, GeometryGroupIndex, GroupName, FGeometryCollection::GeometryGroup)
 		, MuscleActivationElementAttribute(InCollection, MuscleActivationElement, GroupName)
 		, OriginInsertionPairAttribute(InCollection, OriginInsertionPair, GroupName)
 		, OriginInsertionRestLengthAttribute(InCollection, OriginInsertionRestLength, GroupName)
 		, FiberDirectionMatrixAttribute(InCollection, FiberDirectionMatrix, GroupName)
 		, ContractionVolumeScaleAttribute(InCollection, ContractionVolumeScale, GroupName)
+		, FiberStreamlineAttribute(InCollection, FiberStreamline, GroupName)
+		, FiberStreamlineRestLengthAttribute(InCollection, FiberStreamlineRestLength, GroupName)
 	{
 		
 	}
 
 	bool FMuscleActivationFacade::IsValid() const
 	{
-		return MuscleActivationElementAttribute.IsValid() && OriginInsertionPairAttribute.IsValid() && 
-			OriginInsertionRestLengthAttribute.IsValid() && FiberDirectionMatrixAttribute.IsValid() && ContractionVolumeScaleAttribute.IsValid();
+		return GeometryGroupIndexAttribute.IsValid() && MuscleActivationElementAttribute.IsValid() && OriginInsertionPairAttribute.IsValid() &&
+			OriginInsertionRestLengthAttribute.IsValid() && FiberDirectionMatrixAttribute.IsValid() && ContractionVolumeScaleAttribute.IsValid()
+			&& FiberStreamlineAttribute.IsValid() && FiberStreamlineRestLengthAttribute.IsValid();
 	}
 
 	void FMuscleActivationFacade::DefineSchema()
 	{
 		check(!IsConst());
+		GeometryGroupIndexAttribute.Add(ManageArrayAccessor::EPersistencePolicy::MakePersistent, FGeometryCollection::GeometryGroup);
 		MuscleActivationElementAttribute.Add(ManageArrayAccessor::EPersistencePolicy::MakePersistent, "Tetrahedral");
 		OriginInsertionPairAttribute.Add(ManageArrayAccessor::EPersistencePolicy::MakePersistent, FGeometryCollection::VerticesGroup);
 		OriginInsertionRestLengthAttribute.Add();
 		FiberDirectionMatrixAttribute.Add();
 		ContractionVolumeScaleAttribute.Add();
+		FiberStreamlineAttribute.Add();
+		FiberStreamlineRestLengthAttribute.Add();
 	}
 
 	int32 FMuscleActivationFacade::AddMuscleActivationData(const FMuscleActivationData& InputData)
@@ -63,30 +80,435 @@ namespace GeometryCollection::Facades
 		if (IsValid())
 		{
 			int32 NewIndex = MuscleActivationElementAttribute.AddElements(1);
+			GeometryGroupIndexAttribute.Modify()[NewIndex] = InputData.GeometryGroupIndex;
+			MuscleActivationElementAttribute.Modify()[NewIndex] = InputData.MuscleActivationElement;
 			MuscleActivationElementAttribute.Modify()[NewIndex] = InputData.MuscleActivationElement;
 			OriginInsertionPairAttribute.Modify()[NewIndex] = InputData.OriginInsertionPair;
 			OriginInsertionRestLengthAttribute.Modify()[NewIndex] = InputData.OriginInsertionRestLength;
 			FiberDirectionMatrixAttribute.Modify()[NewIndex] = InputData.FiberDirectionMatrix;
 			ContractionVolumeScaleAttribute.Modify()[NewIndex] = InputData.ContractionVolumeScale;
+			FiberStreamlineAttribute.Modify()[NewIndex] = InputData.FiberStreamline;
+			FiberStreamlineRestLengthAttribute.Modify()[NewIndex] = InputData.FiberStreamlineRestLength;
 			return NewIndex;
 		}
 		return INDEX_NONE;
 	}
 
-	const FMuscleActivationData FMuscleActivationFacade::GetMuscleActivationData(const int32 DataIndex) const
+	bool FMuscleActivationFacade::UpdateMuscleActivationData(int32 DataIndex, FMuscleActivationData& InputData)
+	{
+		check(!IsConst());
+		if (IsValid() && 0 <= DataIndex && DataIndex < NumMuscles())
+		{
+			GeometryGroupIndexAttribute.Modify()[DataIndex] = InputData.GeometryGroupIndex;
+			MuscleActivationElementAttribute.Modify()[DataIndex] = InputData.MuscleActivationElement;
+			OriginInsertionPairAttribute.Modify()[DataIndex] = InputData.OriginInsertionPair;
+			OriginInsertionRestLengthAttribute.Modify()[DataIndex] = InputData.OriginInsertionRestLength;
+			FiberDirectionMatrixAttribute.Modify()[DataIndex] = InputData.FiberDirectionMatrix;
+			ContractionVolumeScaleAttribute.Modify()[DataIndex] = InputData.ContractionVolumeScale;
+			FiberStreamlineAttribute.Modify()[DataIndex] = InputData.FiberStreamline;
+			FiberStreamlineRestLengthAttribute.Modify()[DataIndex] = InputData.FiberStreamlineRestLength;
+			return true;
+		}
+		return false;
+	}
+
+	FMuscleActivationData FMuscleActivationFacade::GetMuscleActivationData(const int32 DataIndex) const
 	{
 		FMuscleActivationData ReturnData;
 		if (IsValid())
 		{
-			if (NumMuscles() > DataIndex && DataIndex > -1)
+			if (0 <= DataIndex && DataIndex < NumMuscles())
 			{
+				ReturnData.GeometryGroupIndex = GeometryGroupIndexAttribute.Get()[DataIndex];
 				ReturnData.MuscleActivationElement = MuscleActivationElementAttribute.Get()[DataIndex];
 				ReturnData.OriginInsertionPair = OriginInsertionPairAttribute.Get()[DataIndex];
 				ReturnData.OriginInsertionRestLength = OriginInsertionRestLengthAttribute.Get()[DataIndex];
 				ReturnData.FiberDirectionMatrix = FiberDirectionMatrixAttribute.Get()[DataIndex];
 				ReturnData.ContractionVolumeScale = ContractionVolumeScaleAttribute.Get()[DataIndex];
+				ReturnData.FiberStreamline = FiberStreamlineAttribute.Get()[DataIndex];
+				ReturnData.FiberStreamlineRestLength = FiberStreamlineRestLengthAttribute.Get()[DataIndex];
 			}
 		}
 		return ReturnData;
 	}
+
+	bool FMuscleActivationFacade::SetUpMuscleActivation(const TArray<int32>& InOrigin, const TArray<int32>& Insertion,
+		float InContractionVolumeScale)
+	{
+		// Vertices and fiber field
+		if (!(InOrigin.Num() > 0
+			&& Insertion.Num() > 0
+			&& ConstCollection.FindAttribute<FVector3f>("Vertex", "Vertices")
+			&& ConstCollection.FindAttribute<FIntVector4>("Tetrahedron", "Tetrahedral")
+			&& ConstCollection.FindAttribute<FVector3f>("FiberDirection", "Tetrahedral")))
+		{
+			return false;
+		}
+		TArray<int32> Origin = InOrigin;
+		const TArray<FVector3f>& Vertices = ConstCollection.FindAttribute<FVector3f>("Vertex", "Vertices")->GetConstArray();
+		const TArray<FIntVector4>& Elements = ConstCollection.FindAttribute<FIntVector4>("Tetrahedron", "Tetrahedral")->GetConstArray();
+		const TArray<FVector3f>& FiberDirections = ConstCollection.FindAttribute<FVector3f>("FiberDirection", "Tetrahedral")->GetConstArray();
+		TArray<TArray<int32>> MuscleActivationElements;
+		TArray<TArray<int32>> ComponentOrigins; //One origin node per muscle component
+		TArray<TArray<int32>> ComponentInsertions; //One insertion node per muscle component
+		GeometryCollection::Facades::FCollectionMeshFacade MeshFacade(ConstCollection);
+		TArray<int32> ComponentIndex = MeshFacade.GetGeometryGroupIndexArray(); //Vertex index to geometry index
+		TMap<int32, int32> ComponentToMuscleIndex; //Component index to muscle index
+		TMap<int32, int32> MuscleToComponentIndex; //Muscle index to component index
+		Origin.Sort(); //For some order in muscle groups
+		for (int32 i = 0; i < Origin.Num(); i++)
+		{
+			if (!ComponentToMuscleIndex.Contains(ComponentIndex[Origin[i]]))
+			{
+				ComponentToMuscleIndex.Add(ComponentIndex[Origin[i]], ComponentOrigins.Num());
+				MuscleToComponentIndex.Add(ComponentOrigins.Num(), ComponentIndex[Origin[i]]);
+				ComponentOrigins.SetNum(ComponentOrigins.Num() + 1);
+				ComponentOrigins[ComponentOrigins.Num() - 1].Add(Origin[i]);
+			}
+			else
+			{
+				ComponentOrigins[ComponentToMuscleIndex[ComponentIndex[Origin[i]]]].Add(Origin[i]);
+			}
+		}
+		ComponentInsertions.SetNum(ComponentOrigins.Num());
+		for (int32 i = 0; i < Insertion.Num(); i++)
+		{
+			if (!ComponentToMuscleIndex.Contains(ComponentIndex[Insertion[i]]))
+			{
+				ensureMsgf(false, TEXT("No origin in this component"));
+			}
+			else
+			{
+				ComponentInsertions[ComponentToMuscleIndex[ComponentIndex[Insertion[i]]]].Add(Insertion[i]);
+			}
+		}
+		MuscleActivationElements.SetNum(ComponentOrigins.Num());
+		for (int32 ElemIdx = 0; ElemIdx < Elements.Num(); ElemIdx++)
+		{
+			if (ComponentToMuscleIndex.Contains(ComponentIndex[Elements[ElemIdx][0]]))
+			{
+				MuscleActivationElements[ComponentToMuscleIndex[ComponentIndex[Elements[ElemIdx][0]]]].Add(ElemIdx);
+			}
+		}
+		//Choose one origin-insertion pair per muscle that has largest distance apart within each muscle
+		// TODO: painted attribute directed origin-insertion pair
+		//use origin-insertion line segment length to estimate activation
+		for (int32 MuscleComponentIdx = 0; MuscleComponentIdx < ComponentOrigins.Num(); MuscleComponentIdx++)
+		{
+			if (ensureMsgf(ComponentOrigins.Num() > 0 && ComponentInsertions.Num() > 0, TEXT("Origin or Insertion missing in the muscle %d"), MuscleComponentIdx))
+			{
+				FMuscleActivationData MuscleActivationData;
+				MuscleActivationData.GeometryGroupIndex = MuscleToComponentIndex[MuscleComponentIdx];
+				MuscleActivationData.OriginInsertionRestLength = 0;
+				for (int32 OriginIdx : ComponentOrigins[MuscleComponentIdx])
+				{
+					for (int32 InsertionIdx : ComponentInsertions[MuscleComponentIdx])
+					{
+						float Dist = (Vertices[OriginIdx] - Vertices[InsertionIdx]).Size();
+						if (Dist > MuscleActivationData.OriginInsertionRestLength)
+						{
+							MuscleActivationData.OriginInsertionPair = FIntVector2(OriginIdx, InsertionIdx);
+							MuscleActivationData.OriginInsertionRestLength = Dist;
+						}
+					}
+				}
+				MuscleActivationData.MuscleActivationElement = MuscleActivationElements[MuscleComponentIdx];
+				MuscleActivationData.FiberDirectionMatrix.SetNum(MuscleActivationElements[MuscleComponentIdx].Num());
+				MuscleActivationData.ContractionVolumeScale.SetNum(MuscleActivationElements[MuscleComponentIdx].Num());
+				for (int32 LocalElemIdx = 0; LocalElemIdx < MuscleActivationElements[MuscleComponentIdx].Num(); LocalElemIdx++)
+				{
+					FVector3f V = FiberDirections[MuscleActivationElements[MuscleComponentIdx][LocalElemIdx]];
+					// QR decomposition on vvT for orthogonal directions
+					FVector3f W = V;
+					if (V.X < V.Y)
+					{
+						W.X += 1.f;
+					}
+					else
+					{
+						W.Y += 1.f;
+					}
+					FVector3f U = (V ^ W).GetSafeNormal();
+					W = (U ^ V).GetSafeNormal();
+					MuscleActivationData.FiberDirectionMatrix[LocalElemIdx] = Chaos::PMatrix33d(V, W, U);
+					//Muscle contraction volume scale
+					MuscleActivationData.ContractionVolumeScale[LocalElemIdx] = InContractionVolumeScale;
+				}
+				AddMuscleActivationData(MuscleActivationData);
+			}
+		}
+		return true;
+	}
+
+	TArray<TArray<TArray<FVector3f>>> FMuscleActivationFacade::BuildStreamlines(const TArray<int32>& Origin, const TArray<int32>& Insertion,
+		int32 NumLinesMultiplier, int32 MaxStreamlineIterations, int32 MaxPointsPerLine)
+	{
+		TArray<TArray<FVector3f>> LineSegments;
+		TArray<TArray<TArray<FVector3f>>> MuscleLineSegments;
+		TArray<int32> StreamlineStartElements;
+		// Vertices and fiber field
+		if (!(ConstCollection.FindAttribute<FVector3f>("Vertex", "Vertices")
+			&& ConstCollection.FindAttribute<FIntVector4>("Tetrahedron", "Tetrahedral")
+			&& ConstCollection.FindAttribute<FVector3f>("FiberDirection", "Tetrahedral")
+			&& ConstCollection.FindAttribute<int32>("TetrahedronStart", FGeometryCollection::GeometryGroup)
+			&& ConstCollection.FindAttribute<int32>("TetrahedronCount", FGeometryCollection::GeometryGroup)))
+		{
+			return MuscleLineSegments;
+		}
+		const TArray<FVector3f>& Vertices = ConstCollection.FindAttribute<FVector3f>("Vertex", "Vertices")->GetConstArray();
+		const TArray<FIntVector4>& Elements = ConstCollection.FindAttribute<FIntVector4>("Tetrahedron", "Tetrahedral")->GetConstArray();
+		const TArray<FVector3f>& FiberDirections = ConstCollection.FindAttribute<FVector3f>("FiberDirection", "Tetrahedral")->GetConstArray();
+		const TArray<int32>& TetrahedronStart = ConstCollection.FindAttribute<int32>("TetrahedronStart", FGeometryCollection::GeometryGroup)->GetConstArray();
+		const TArray<int32>& TetrahedronCount = ConstCollection.FindAttribute<int32>("TetrahedronCount", FGeometryCollection::GeometryGroup)->GetConstArray();
+
+		ensure(Elements.Num() == FiberDirections.Num());
+		StreamlineStartElements.Reset();
+		TArray<FVector3f> ModifiedFiberDirections = FiberDirections;
+		TArray<TArray<int32>> LocalIndex;
+		TArray<TArray<int32>> Mesh;
+		Mesh.SetNum(Elements.Num());
+		for (int32 i = 0; i < Elements.Num(); i++)
+		{
+			for (int32 j = 0; j < 4; j++)
+			{
+				Mesh[i].Add(Elements[i][j]);
+			}
+		}
+		TArray<TArray<int32>> IncidentElementsLocalIndex;
+		TArray<TArray<int32>> IncidentElements = Chaos::Utilities::ComputeIncidentElements(Mesh, &LocalIndex);
+		TArray<FIntVector2> Faces = Chaos::Utilities::ComputeTetMeshFacePairs(Elements);
+
+		TArray<TArray<int32>> FaceToTet;
+		FaceToTet.SetNum(Elements.Num() * 4);
+		for (int32 f = 0; f < Faces.Num(); ++f)
+		{
+			int32 q0 = Faces[f][0] / 4;
+			FaceToTet[Faces[f][0]].Add(q0);
+			if (Faces[f][1] > -1)
+			{
+				int32 q1 = Faces[f][1] / 4;
+				FaceToTet[Faces[f][0]].Add(q1);
+				FaceToTet[Faces[f][1]].Add(q0);
+				FaceToTet[Faces[f][1]].Add(q1);
+			}
+		}
+		TArray<bool> bOrigin, bInsertion;
+		bOrigin.Init(false, Vertices.Num());
+		bInsertion.Init(false, Vertices.Num());
+		for (int32 i : Origin)
+		{
+			bOrigin[i] = true;
+		}
+		for (int32 i : Insertion)
+		{
+			bInsertion[i] = true;
+		}
+		auto bConstrained = [&bOrigin, &bInsertion](int32 i) {
+			return bOrigin[i] || bInsertion[i];
+		};
+		for (int32 i = 0; i < FaceToTet.Num(); ++i)
+		{
+			if (FaceToTet[i].Num() == 1) //boundary face
+			{
+				int32 e = FaceToTet[i][0];
+				if (bConstrained(Elements[e][0]) || bConstrained(Elements[e][1]) || bConstrained(Elements[e][2]) || bConstrained(Elements[e][3]))
+					continue;
+				FIntVector3 LocalFace = Chaos::Utilities::TetFace(i % 4);
+				FIntVector3 Face = FIntVector3(Elements[e][LocalFace[0]], Elements[e][LocalFace[1]], Elements[e][LocalFace[2]]);
+				FVector3f Normal = ((Vertices[Face[1]] - Vertices[Face[0]]) ^ (Vertices[Face[2]] - Vertices[Face[0]])).GetSafeNormal();
+				ModifiedFiberDirections[e] = FiberDirections[e] - (FiberDirections[e] | Normal) * Normal;
+				ModifiedFiberDirections[e].Normalize();
+			}
+		}
+
+		TArray<bool> bEndElement;
+		bEndElement.Init(false, Elements.Num());
+		for (int32 e = 0; e < Elements.Num() / 4; ++e)
+		{
+			for (int32 ie = 0; ie < 4; ++ie)
+			{
+				if (bInsertion[Elements[e][ie]])
+				{
+					bEndElement[e] = true;
+				}
+			}
+		}
+		TArray<int32> SampleElements;
+
+		for (int32 i : Origin)
+		{
+			for (int32 e : IncidentElements[i])
+			{
+				if (!(bOrigin[Elements[e][0]] && bOrigin[Elements[e][1]] && bOrigin[Elements[e][2]] && bOrigin[Elements[e][3]]))
+				{
+					SampleElements.AddUnique(e);
+				}
+			}
+		}
+
+		TArray<TArray<FVector3f>> Origin_sampled = Chaos::Utilities::RandomPointsInTet(Vertices, Elements, SampleElements, NumLinesMultiplier);
+
+		for (int32 ij = 0; ij < Origin_sampled.Num(); ++ij)
+		{
+			for (FVector3f StartPosition : Origin_sampled[ij])
+			{
+				FVector3f StartDirection = ModifiedFiberDirections[SampleElements[ij]];
+				TArray<int32> StartTetCandidate = { SampleElements[ij] };
+				TArray<int32> NewStartTetCandidate;
+				TArray<FVector3f> CurrentLineSegment;
+				CurrentLineSegment.Add(StartPosition);
+				int32 Iter = 0;
+				FVector3f EndPosition;
+				bool bReachEnd = false;
+				while (((StartTetCandidate.Num() > 1 && Iter > 0) || (StartTetCandidate.Num() > 0 && Iter == 0)) && Iter < MaxStreamlineIterations)
+				{
+					if (CurrentLineSegment.Num() > 1 && (CurrentLineSegment[CurrentLineSegment.Num() - 1] - CurrentLineSegment[CurrentLineSegment.Num() - 2]).Size() < 1e-6)
+					{
+						CurrentLineSegment.Pop();
+						break;
+					}
+					NewStartTetCandidate.Reset();
+					bool NonTrivialIntersection = false;
+					for (int32 e : StartTetCandidate)
+					{
+						for (int32 f = 0; f < 4; f++)
+						{
+							FIntVector3 local_face = Chaos::Utilities::TetFace(f);
+
+							UE::Math::TRay<float> RayIn(StartPosition, StartDirection);
+							UE::Geometry::TTriangle3<float> TriangleIn(Vertices[Elements[e][local_face[0]]], Vertices[Elements[e][local_face[1]]], Vertices[Elements[e][local_face[2]]]);
+							UE::Geometry::TIntrRay3Triangle3<float> Intersection(RayIn, TriangleIn);
+							if (Intersection.Find() && Intersection.IntersectionType == EIntersectionType::Point)
+							{
+								FVector3f IntersectionPosition = TriangleIn.BarycentricPoint(float(Intersection.TriangleBaryCoords[0]), float(Intersection.TriangleBaryCoords[1]), float(Intersection.TriangleBaryCoords[2]));
+								if ((StartPosition - IntersectionPosition).Size() > 1e-6)
+								{
+									NonTrivialIntersection = true;
+									EndPosition = IntersectionPosition;
+									StartPosition = EndPosition;
+									CurrentLineSegment.Add(EndPosition);
+									for (int32 NewTetCandidate : FaceToTet[4 * e + f])
+									{
+										if (NewTetCandidate != e)
+										{
+											NewStartTetCandidate.Add(NewTetCandidate);
+											StartDirection = ModifiedFiberDirections[NewTetCandidate];
+											bReachEnd = bEndElement[NewTetCandidate];
+											break;
+										}
+									}
+									NewStartTetCandidate.Add(e);
+									break;
+								}
+							}
+						}
+						if (NonTrivialIntersection)
+						{
+							NonTrivialIntersection = false;
+							break;
+						}
+					}
+					StartTetCandidate = NewStartTetCandidate;
+					Iter++;
+					if (bReachEnd)
+					{
+						LineSegments.Add(CurrentLineSegment);
+						StreamlineStartElements.Add(SampleElements[ij]);
+						break;
+					}
+				}
+			}
+		}
+
+		//Coarsen streamlines
+		for (int32 i = 0; i < LineSegments.Num(); ++i)
+		{
+			if (LineSegments[i].Num() > MaxPointsPerLine)
+			{
+				float TotalLength = 0;
+				for (int32 j = 1; j < LineSegments[i].Num(); ++j)
+				{
+					TotalLength += (LineSegments[i][j - 1] - LineSegments[i][j]).Size();
+				}
+				float MinLength = TotalLength / float(MaxPointsPerLine - 1);
+				TArray<FVector3f> NewLine;
+				NewLine.Add(LineSegments[i][0]);
+				float EndLength = 0;
+				int32 EndIndex = LineSegments[i].Num() - 1;
+				for (int32 j = LineSegments[i].Num() - 1; j >= 0; --j)
+				{
+					EndLength += (LineSegments[i][j - 1] - LineSegments[i][j]).Size();
+					if (EndLength > MinLength)
+					{
+						EndIndex = j - 1;
+						break;
+					}
+				}
+				int32 End = 1;
+				float CurrentLength = 0;
+				while (End <= EndIndex)
+				{
+					CurrentLength += (LineSegments[i][End - 1] - LineSegments[i][End]).Size();
+					if (CurrentLength > MinLength)
+					{
+						NewLine.Add(LineSegments[i][End]);
+						CurrentLength = 0;
+					}
+					End += 1;
+				}
+				if (CurrentLength > 0)
+				{
+					NewLine.Add(LineSegments[i][End]);
+				}
+				NewLine.Add(LineSegments[i][LineSegments[i].Num() - 1]);
+				LineSegments[i] = NewLine;
+			}
+		}
+
+		//Split line segments by muscle groups
+		TArray<TArray<float>> MuscleLineSegmentRestLength;
+		MuscleLineSegments.SetNum(NumMuscles());
+		MuscleLineSegmentRestLength.SetNum(NumMuscles());
+		TArray<int32> ElementToMuscleIndexArray;
+		TArray<int32> GroupIndexToMuscleIndexArray;
+		GroupIndexToMuscleIndexArray.Init(INDEX_NONE, TetrahedronStart.Num());
+		ElementToMuscleIndexArray.Init(INDEX_NONE, Elements.Num());
+		for (int32 MuscleIndex = 0; MuscleIndex < NumMuscles(); ++MuscleIndex)
+		{
+			FMuscleActivationData MuscleActivationData = GetMuscleActivationData(MuscleIndex);	
+			if (0 <= MuscleActivationData.GeometryGroupIndex && MuscleActivationData.GeometryGroupIndex < GroupIndexToMuscleIndexArray.Num())
+				GroupIndexToMuscleIndexArray[MuscleActivationData.GeometryGroupIndex] = MuscleIndex;
+		}
+		for (int32 GroupIndex = 0; GroupIndex < TetrahedronStart.Num(); ++GroupIndex)
+		{
+			for (int32 LocalIdx = 0; LocalIdx < TetrahedronCount[GroupIndex]; ++LocalIdx)
+			{
+				ElementToMuscleIndexArray[TetrahedronStart[GroupIndex] + LocalIdx] = GroupIndexToMuscleIndexArray[GroupIndex];
+			}
+		}
+		for (int32 LineIndex = 0; LineIndex < StreamlineStartElements.Num(); ++LineIndex)
+		{
+			int32 MuscleIndex = ElementToMuscleIndexArray[StreamlineStartElements[LineIndex]];
+			if (MuscleIndex >= 0)
+			{
+				MuscleLineSegments[MuscleIndex].Add(LineSegments[LineIndex]);
+				float TotalLength = 0;
+				for (int32 j = 1; j < LineSegments[LineIndex].Num(); ++j)
+				{
+					TotalLength += (LineSegments[LineIndex][j - 1] - LineSegments[LineIndex][j]).Size();
+				}
+				MuscleLineSegmentRestLength[MuscleIndex].Add(TotalLength);
+			}
+		}
+		
+		//Save streamline data
+		for (int32 MuscleIndex = 0; MuscleIndex < NumMuscles(); ++MuscleIndex)
+		{
+			FMuscleActivationData MuscleActivationData = GetMuscleActivationData(MuscleIndex);
+			MuscleActivationData.FiberStreamline = MuscleLineSegments[MuscleIndex];
+			MuscleActivationData.FiberStreamlineRestLength = MuscleLineSegmentRestLength[MuscleIndex];
+			UpdateMuscleActivationData(MuscleIndex, MuscleActivationData);
+		}
+		return MuscleLineSegments;
+	}
 }
diff --git a/Engine/Source/Runtime/Experimental/Chaos/Public/Chaos/Deformable/MuscleActivationConstraints.h b/Engine/Source/Runtime/Experimental/Chaos/Public/Chaos/Deformable/MuscleActivationConstraints.h
--- a/Engine/Source/Runtime/Experimental/Chaos/Public/Chaos/Deformable/MuscleActivationConstraints.h
+++ b/Engine/Source/Runtime/Experimental/Chaos/Public/Chaos/Deformable/MuscleActivationConstraints.h
@@ -28,18 +28,27 @@ namespace Chaos::Softs
 			for (int MuscleIdx = 0; MuscleIdx < FMuscleActivation.NumMuscles(); MuscleIdx++)
 			{
 				Data MuscleActivationData = FMuscleActivation.GetMuscleActivationData(MuscleIdx);
-				for (int32 i = 0; i < MuscleActivationData.MuscleActivationElement.Num(); i++)
+				if (FMuscleActivation.IsValidGeometryIndex(MuscleActivationData.GeometryGroupIndex))
 				{
-					MuscleActivationData.MuscleActivationElement[i] += ElementOffset;
+					int32 OldSize = MuscleActivationElement.Num();
+					MuscleActivationElement.SetNum(OldSize + 1);
+					FiberDirectionMatrix.SetNum(OldSize + 1);
+					ContractionVolumeScale.SetNum(OldSize + 1);
+					for (int32 i = 0; i < MuscleActivationData.MuscleActivationElement.Num(); i++)
+					{
+						if (FMuscleActivation.IsValidElementIndex(MuscleActivationData.MuscleActivationElement[i]))
+						{
+							MuscleActivationElement[OldSize].Add(MuscleActivationData.MuscleActivationElement[i] + ElementOffset);
+							FiberDirectionMatrix[OldSize].Add(MuscleActivationData.FiberDirectionMatrix[i]);
+							ContractionVolumeScale[OldSize].Add(MuscleActivationData.ContractionVolumeScale[i]);
+						}
+					}
+					MuscleActivationData.OriginInsertionPair[0] += VertexOffset;
+					MuscleActivationData.OriginInsertionPair[1] += VertexOffset;
+					OriginInsertionPair.Add(MuscleActivationData.OriginInsertionPair);
+					OriginInsertionRestLength.Add(MuscleActivationData.OriginInsertionRestLength);
+					Activation.Add(0.f);
 				}
-				MuscleActivationElement.Add(MuscleActivationData.MuscleActivationElement);
-				MuscleActivationData.OriginInsertionPair[0] += VertexOffset;
-				MuscleActivationData.OriginInsertionPair[1] += VertexOffset;
-				OriginInsertionPair.Add(MuscleActivationData.OriginInsertionPair);
-				OriginInsertionRestLength.Add(MuscleActivationData.OriginInsertionRestLength);
-				Activation.Add(0.f);
-				FiberDirectionMatrix.Add(MuscleActivationData.FiberDirectionMatrix);
-				ContractionVolumeScale.Add(MuscleActivationData.ContractionVolumeScale);
 			}
 		}
 
diff --git a/Engine/Source/Runtime/Experimental/Chaos/Public/Chaos/Utilities.h b/Engine/Source/Runtime/Experimental/Chaos/Public/Chaos/Utilities.h
--- a/Engine/Source/Runtime/Experimental/Chaos/Public/Chaos/Utilities.h
+++ b/Engine/Source/Runtime/Experimental/Chaos/Public/Chaos/Utilities.h
@@ -1216,5 +1216,152 @@ namespace Chaos
 				}
 			}
 		}
+
+		inline FIntVector3 SortFIntVector3(FIntVector3& V)
+		{
+			if (V[1] < V[0])
+			{
+				int32 V0 = V[0];
+				V[0] = V[1];
+				V[1] = V0;
+			}
+			if (V[2] < V[0])
+			{
+				int32 V2 = V[2];
+				V[2] = V[1];
+				V[1] = V[0];
+				V[0] = V2;
+			}
+			else if (V[2] < V[1])
+			{
+				int32 V1 = V[1];
+				V[1] = V[2];
+				V[2] = V1;
+			}
+			return V;
+		}
+
+		// For each of the 4 triangle faces in a tet, return local vertex indices
+		// right hand rule, triangle normal pointing outwards
+		inline FIntVector3 TetFace(int32 f) {
+			switch (f) {
+			case 0:
+				return { 1, 2, 3 };
+			case 1:
+				return { 0, 3, 2 };
+			case 2:
+				return { 0, 1, 3 };
+			case 3:
+				return { 0, 2, 1 };
+			default:
+				return { -1, -1, -1 };
+			}
+		}
+
+		template <typename Func1, typename Func2>
+		TArray<FIntVector2> ComputeMeshFacePairs(const TArray<FIntVector4>& Mesh, Func1 GreaterThan, Func2 Equal) {
+			int32 NumFacesPerElement = 4;
+			int32 face_number = Mesh.Num() * NumFacesPerElement;
+			TArray<int32> AllFaces, ranges;
+			AllFaces.SetNum(face_number);
+			for (int32 f = 0; f < face_number; f++) {
+				AllFaces[f] = f;
+			}
+
+			AllFaces.Sort([&GreaterThan, &Mesh](int32 face_a, int32 face_b) { return GreaterThan(Mesh, face_a, face_b); });
+			TArray<FIntVector2> CommonFacePairs;
+			int32 f1 = 0, f2 = 1;
+			while (f2 < AllFaces.Num())
+			{
+				if (Equal(Mesh, AllFaces[f1], AllFaces[f2])) //common face from two tetrahedra
+				{
+					CommonFacePairs.Add(FIntVector2(AllFaces[f1], AllFaces[f2]));
+					f1 += 2;
+					f2 += 2;
+				}
+				else //boundary face
+				{
+					CommonFacePairs.Add(FIntVector2(AllFaces[f1], -1));
+					f1 += 1;
+					f2 += 1;
+				}
+			}
+			if (f1 < AllFaces.Num())
+			{
+				CommonFacePairs.Add(FIntVector2(AllFaces[f1], -1));
+			}
+			return CommonFacePairs;
+		}
+
+		//returns pairs of face indices {FaceA, FaceB} that are shared by two adjacent tetrahedra or boundary face {FaceA, -1}
+		inline TArray<FIntVector2> ComputeTetMeshFacePairs(const TArray<FIntVector4>& TetMesh) {
+			auto FaceGreaterThan = [](FIntVector3& i1, FIntVector3& i2)
+			{
+				SortFIntVector3(i1);
+				SortFIntVector3(i2);
+				if (i1[0] > i2[0])
+					return true;
+				else if (i1[0] == i2[0] && i1[1] > i2[1])
+					return true;
+				return (i1[0] == i2[0] && i1[1] == i2[1] && i1[2] > i2[2]);
+			};
+			auto FaceEqual = [](FIntVector3& i1, FIntVector3& i2)
+			{
+				SortFIntVector3(i1);
+				SortFIntVector3(i2);
+				return i1 == i2;
+			};
+			auto GreaterThan = [&FaceGreaterThan](const TArray<FIntVector4>& Mesh, int32 e1, int32 e2)
+			{
+				int32 Local1 = e1 % 4;
+				int32 Element1 = e1 / 4;
+				FIntVector3 LocalFace1 = TetFace(Local1);
+				FIntVector3 Edge1 = FIntVector3(Mesh[Element1][LocalFace1[0]], Mesh[Element1][LocalFace1[1]], Mesh[Element1][LocalFace1[2]]);
+				int32 Local2 = e2 % 4;
+				int32 Element2 = e2 / 4;
+				FIntVector3 LocalFace2 = TetFace(Local2);
+				FIntVector3 Edge2 = FIntVector3(Mesh[Element2][LocalFace2[0]], Mesh[Element2][LocalFace2[1]], Mesh[Element2][LocalFace2[2]]);
+				return FaceGreaterThan(Edge1, Edge2);
+			};
+			auto Equal = [&FaceEqual](const TArray<FIntVector4>& Mesh, int32 e1, int32 e2)
+			{
+				int32 Local1 = e1 % 4;
+				int32 Element1 = e1 / 4;
+				FIntVector3 LocalFace1 = TetFace(Local1);
+				FIntVector3 Edge1 = FIntVector3(Mesh[Element1][LocalFace1[0]], Mesh[Element1][LocalFace1[1]], Mesh[Element1][LocalFace1[2]]);
+				int32 Local2 = e2 % 4;
+				int32 Element2 = e2 / 4;
+				FIntVector3 LocalFace2 = TetFace(Local2);
+				FIntVector3 Edge2 = FIntVector3(Mesh[Element2][LocalFace2[0]], Mesh[Element2][LocalFace2[1]], Mesh[Element2][LocalFace2[2]]);
+				return FaceEqual(Edge1, Edge2);
+			};
+			return ComputeMeshFacePairs(TetMesh, GreaterThan, Equal);
+		}
+
+		//Return a randomly sampled point in selected tetrahedra following uniform distribution
+		inline TArray<TArray<FVector3f>> RandomPointsInTet(const TArray<FVector3f>& x, const TArray<FIntVector4>& TetMesh, const TArray<int32>& SampleElements, int32 NumRandomPointsPerElement = 1) {
+			srand(1);
+			TArray<TArray<FVector3f>> SampledPoints;
+			SampledPoints.SetNum(SampleElements.Num());
+			for (int32 ElemIdx = 0; ElemIdx < SampleElements.Num(); ++ElemIdx)
+			{
+				int32 e = SampleElements[ElemIdx];
+				for (int32 i = 0; i < NumRandomPointsPerElement; ++i) {
+					FVector4f Weights(0);
+					FVector3f Point(0);
+					float TotalWeight = 0;
+					for (int32 j = 0; j < 4; ++j) {
+						Weights[j] = float(rand()) / float(RAND_MAX);
+						TotalWeight += Weights[j];
+					}
+					for (int32 j = 0; j < 4; ++j) {
+						Weights[j] /= TotalWeight;
+						Point += Weights[j] * x[TetMesh[e][j]];
+					}
+					SampledPoints[ElemIdx].Add(Point);
+				}
+			}
+			return SampledPoints;
+		}
 	} // namespace Utilities
 } // namespace Chaos
\ No newline at end of file
diff --git a/Engine/Source/Runtime/Experimental/Chaos/Public/Field/FieldSystemTypes.h b/Engine/Source/Runtime/Experimental/Chaos/Public/Field/FieldSystemTypes.h
--- a/Engine/Source/Runtime/Experimental/Chaos/Public/Field/FieldSystemTypes.h
+++ b/Engine/Source/Runtime/Experimental/Chaos/Public/Field/FieldSystemTypes.h
@@ -415,27 +415,54 @@ public:
 		AddGroup("VectorField"); 
 		AddAttribute<FVector3f>("Start", "VectorField");
 		AddAttribute<FVector3f>("End", "VectorField");
+		AddAttribute<FLinearColor>("Color", "VectorField");
 	}
 	
-	void AddVectorToField(FVector3f Start, FVector3f End)
+	int32 AddVectorToField(FVector3f Start, FVector3f End)
 	{
 		int32 Size = AddElements(1, "VectorField");
 		ModifyAttribute<FVector3f>("Start", "VectorField")[Size] = Start;
 		ModifyAttribute<FVector3f>("End", "VectorField")[Size] = End;
+		ModifyAttribute<FLinearColor>("Color", "VectorField")[Size] = FLinearColor(0.6f, 0.6f, 0.6f);
+		return Size;
+	}
+
+	void SetColorOnVector(int32 Index, FLinearColor Color)
+	{
+		if (0 <= Index && Index < NumElements("VectorField"))
+		{
+			ModifyAttribute<FLinearColor>("Color", "VectorField")[Index] = Color;
+		}
 	}
 	
 	TArray<TPair<FVector3f, FVector3f>> GetVectorField() const
 	{
 		TArray<TPair<FVector3f, FVector3f>> VectorField;
-		const TManagedArray<FVector3f>& VectorFieldStart = GetAttribute<FVector3f>("Start", "VectorField");
-		const TManagedArray<FVector3f>& VectorFieldEnd = GetAttribute<FVector3f>("End", "VectorField");
-		for (int32 i = 0; i < NumElements("VectorField"); i++)
+		if (FindAttribute<FVector3f>("Start", "VectorField") != nullptr && FindAttribute<FVector3f>("End", "VectorField") != nullptr)
 		{
-			VectorField.Add(TPair<FVector3f, FVector3f>(VectorFieldStart[i], VectorFieldEnd[i]));
+			const TManagedArray<FVector3f>& VectorFieldStart = GetAttribute<FVector3f>("Start", "VectorField");
+			const TManagedArray<FVector3f>& VectorFieldEnd = GetAttribute<FVector3f>("End", "VectorField");
+			for (int32 i = 0; i < NumElements("VectorField"); i++)
+			{
+				VectorField.Add(TPair<FVector3f, FVector3f>(VectorFieldStart[i], VectorFieldEnd[i]));
+			}
 		}
 		return VectorField;
 	}
 
+	TArray<FLinearColor> GetVectorColor() const
+	{
+		if (FindAttribute<FLinearColor>("Color", "VectorField") != nullptr)
+		{
+			const TManagedArray<FLinearColor>& Color = GetAttribute<FLinearColor>("Color", "VectorField");
+			return Color.GetConstArray();
+		}
+		else
+		{
+			return TArray<FLinearColor>();
+		}
+	}
+
 };
 
 
diff --git a/Engine/Source/Runtime/Experimental/Chaos/Public/GeometryCollection/Facades/CollectionMuscleActivationFacade.h b/Engine/Source/Runtime/Experimental/Chaos/Public/GeometryCollection/Facades/CollectionMuscleActivationFacade.h
--- a/Engine/Source/Runtime/Experimental/Chaos/Public/GeometryCollection/Facades/CollectionMuscleActivationFacade.h
+++ b/Engine/Source/Runtime/Experimental/Chaos/Public/GeometryCollection/Facades/CollectionMuscleActivationFacade.h
@@ -3,32 +3,40 @@
 #pragma once
 
 #include "CoreMinimal.h"
+#include "GeometryCollection/GeometryCollection.h"
 #include "GeometryCollection/ManagedArrayAccessor.h"
 #include "GeometryCollection/ManagedArrayCollection.h"
 #include "GeometryCollection/Facades/CollectionSelectionFacade.h"
 
 namespace GeometryCollection::Facades
 {
+	//Activation data for each muscle
 	struct FMuscleActivationData
 	{
-		TArray<int32> MuscleActivationElement;
-		FIntVector2 OriginInsertionPair;
-		float OriginInsertionRestLength;
-		TArray<Chaos::PMatrix33d> FiberDirectionMatrix;
-		TArray<float> ContractionVolumeScale;
+		int32 GeometryGroupIndex; //Geometry group index of the muscle
+		TArray<int32> MuscleActivationElement; //Contractible tetrahedra
+		FIntVector2 OriginInsertionPair; //Muscle origin point and insertion point (to determine muscle length)
+		float OriginInsertionRestLength; //Muscle origin-insertion rest length
+		TArray<Chaos::PMatrix33d> FiberDirectionMatrix; //Per-element fiber direction orthogonal matrix: [v, w1, w2], v is the fiber direction
+		TArray<float> ContractionVolumeScale; // Per-element volume scale for muscle contraction
+		TArray<TArray<FVector3f>> FiberStreamline; //Fiber streamline(s) for inverse dynamics
+		TArray<float> FiberStreamlineRestLength; //Fiber streamline rest length(s)
 	};
 
-	/** Kinematic Facade */
+	/** Muscle Activation Facade */
 	class FMuscleActivationFacade
 	{
 	public:
 
 		static CHAOS_API const FName GroupName;
+		static CHAOS_API const FName GeometryGroupIndex;
 		static CHAOS_API const FName MuscleActivationElement;
 		static CHAOS_API const FName OriginInsertionPair;
 		static CHAOS_API const FName OriginInsertionRestLength;
 		static CHAOS_API const FName FiberDirectionMatrix;
 		static CHAOS_API const FName ContractionVolumeScale;
+		static CHAOS_API const FName FiberStreamline;
+		static CHAOS_API const FName FiberStreamlineRestLength;
 
 		CHAOS_API FMuscleActivationFacade(FManagedArrayCollection& InCollection);
 		CHAOS_API FMuscleActivationFacade(const FManagedArrayCollection& InCollection);
@@ -46,17 +54,24 @@ namespace GeometryCollection::Facades
 		//  Skeletal Mesh Bone Bindings
 		//
 		CHAOS_API int32 AddMuscleActivationData(const FMuscleActivationData& InputData);
-		CHAOS_API const FMuscleActivationData GetMuscleActivationData(const int32 DataIndex) const;
-		int32 NumMuscles() const { return MuscleActivationElementAttribute.Num(); }
-
+		CHAOS_API bool UpdateMuscleActivationData(const int32 DataIndex, FMuscleActivationData& InputData);
+		CHAOS_API FMuscleActivationData GetMuscleActivationData(const int32 DataIndex) const;
+		CHAOS_API bool IsValidGeometryIndex(const int32 Index) const { return 0 <= Index && Index < ConstCollection.NumElements(FGeometryCollection::GeometryGroup); };
+		CHAOS_API bool IsValidElementIndex(const int32 Index) const { return 0 <= Index && Index < ConstCollection.NumElements("Tetrahedral"); };
+		CHAOS_API int32 NumMuscles() const { return MuscleActivationElementAttribute.Num(); }
+		CHAOS_API bool SetUpMuscleActivation(const TArray<int32>& Origin, const TArray<int32>& Insertion, float ContractionVolumeScale);
+		CHAOS_API TArray<TArray<TArray<FVector3f>>> BuildStreamlines(const TArray<int32>& Origin, const TArray<int32>& Insertion,
+			 int32 NumLinesMultiplier, int32 MaxStreamlineIterations, int32 MaxPointsPerLine);
 	private:
 		const FManagedArrayCollection& ConstCollection;
 		FManagedArrayCollection* Collection = nullptr;
-
+		TManagedArrayAccessor<int32> GeometryGroupIndexAttribute;
 		TManagedArrayAccessor<TArray<int32>> MuscleActivationElementAttribute;
 		TManagedArrayAccessor<FIntVector2> OriginInsertionPairAttribute;
 		TManagedArrayAccessor<float> OriginInsertionRestLengthAttribute;
 		TManagedArrayAccessor<TArray<Chaos::PMatrix33d>> FiberDirectionMatrixAttribute;
 		TManagedArrayAccessor<TArray<float>> ContractionVolumeScaleAttribute;
+		TManagedArrayAccessor<TArray<TArray<FVector3f>>> FiberStreamlineAttribute;
+		TManagedArrayAccessor<TArray<float>> FiberStreamlineRestLengthAttribute;
 	};
 }
diff --git a/Engine/Source/Runtime/Experimental/Chaos/Public/GeometryCollection/ManagedArrayTypeValues.inl b/Engine/Source/Runtime/Experimental/Chaos/Public/GeometryCollection/ManagedArrayTypeValues.inl
--- a/Engine/Source/Runtime/Experimental/Chaos/Public/GeometryCollection/ManagedArrayTypeValues.inl
+++ b/Engine/Source/Runtime/Experimental/Chaos/Public/GeometryCollection/ManagedArrayTypeValues.inl
@@ -58,6 +58,7 @@ MANAGED_ARRAY_TYPE(TArray<FIntVector3>, IntVector3Array)
 MANAGED_ARRAY_TYPE(TArray<FVector4f>, Vector4fArray)
 MANAGED_ARRAY_TYPE(Chaos::PMatrix33d, PMatrix33d)
 MANAGED_ARRAY_TYPE(TArray<Chaos::PMatrix33d>, PMatrix33dArray)
+MANAGED_ARRAY_TYPE(TArray<TArray<FVector3f>>, FVector3fNestedArray)
 
 // NOTE: new types must be added at the bottom to keep serialization from breaking
 
