Add an error message when uploading data if we detect a chunk that is too large and will not work at runtime.

#rnx

- The toc stored the sizes as 64bit but at runtime we only store 32bit, as it saves space and attempting to download data sizes this large wouldn't make much sense anyway. However we were allowing these chunks to be uploaded without any issue and would only error at runtime.
- This change will highlight the problem when a build is being made and will clearly identify the file causing the issue.

diff --git a/Engine/Source/Runtime/Experimental/IoStore/OnDemand/Private/Tool/Upload.cpp b/Engine/Source/Runtime/Experimental/IoStore/OnDemand/Private/Tool/Upload.cpp
--- a/Engine/Source/Runtime/Experimental/IoStore/OnDemand/Private/Tool/Upload.cpp
+++ b/Engine/Source/Runtime/Experimental/IoStore/OnDemand/Private/Tool/Upload.cpp
@@ -585,6 +585,15 @@ static TIoStatusOr<FUploadResult> UploadContainerFiles(
 				return FIoStatus(EIoErrorCode::ReadError, TEXT("Encoded chunk size does not match buffer"));
 			}
 
+			// At runtime we are limited to MAX_uint32 for chunk lengths to save space and because anything larger than that
+			// is not reasonable to load via IoStoreOnDemand anyway. So we need to check for this now and fail the upload if
+			// there is a chunk that will fail at runtime.
+			// Note that EncodedSize should always be <= RawChunkSize but test both to be safe.
+			if (RawChunkSize > MAX_uint32 || EncodedChunkSize > MAX_uint32)
+			{
+				return FIoStatus(EIoErrorCode::InvalidParameter, WriteToString<512>(*ChunkInfo.FileName, TEXT(": Chunk size should not exceed MAX_uint32")));
+			}
+
 			FOnDemandTocEntry& TocEntry = ContainerEntry.Entries.AddDefaulted_GetRef();
 			TocEntry.ChunkId = ChunkInfo.Id;
 			TocEntry.Hash = ChunkHash;
