IASTool - Fix missing headers

diff --git a/Engine/Source/Runtime/Experimental/IoStore/OnDemand/Private/Tool/UploadQueue.cpp b/Engine/Source/Runtime/Experimental/IoStore/OnDemand/Private/Tool/UploadQueue.cpp
--- a/Engine/Source/Runtime/Experimental/IoStore/OnDemand/Private/Tool/UploadQueue.cpp
+++ b/Engine/Source/Runtime/Experimental/IoStore/OnDemand/Private/Tool/UploadQueue.cpp
@@ -4,7 +4,7 @@
 
 #include "UploadQueue.h"
 
-#include <Containers/UnrealString.h>
+#include <Async/Async.h>
 #include <CoreGlobals.h>
 #include <Misc/OutputDeviceRedirector.h>
 #include <IO/IoStoreOnDemand.h>
diff --git a/Engine/Source/Runtime/Experimental/IoStore/OnDemand/Private/Tool/UploadQueue.h b/Engine/Source/Runtime/Experimental/IoStore/OnDemand/Private/Tool/UploadQueue.h
--- a/Engine/Source/Runtime/Experimental/IoStore/OnDemand/Private/Tool/UploadQueue.h
+++ b/Engine/Source/Runtime/Experimental/IoStore/OnDemand/Private/Tool/UploadQueue.h
@@ -7,9 +7,11 @@
 
 #include "IO/IoBuffer.h"
 
+#include <Async/Future.h>
 #include <Containers/Array.h>
-#include <Containers/StringConv.h>
+#include <Containers/Queue.h>
 #include <Containers/StringView.h>
+#include <Containers/UnrealString.h>
 #include <String/LexFromString.h>
 
 
