ChaosFlesh: UE build fix

diff --git a/Engine/Source/Runtime/Experimental/Chaos/Public/Chaos/Deformable/GaussSeidelSphereRepulsionConstraints.h b/Engine/Source/Runtime/Experimental/Chaos/Public/Chaos/Deformable/GaussSeidelSphereRepulsionConstraints.h
--- a/Engine/Source/Runtime/Experimental/Chaos/Public/Chaos/Deformable/GaussSeidelSphereRepulsionConstraints.h
+++ b/Engine/Source/Runtime/Experimental/Chaos/Public/Chaos/Deformable/GaussSeidelSphereRepulsionConstraints.h
@@ -21,32 +21,12 @@
 
 namespace Chaos::Softs
 {
-	namespace Private
-	{
-		struct FSphereSpatialEntry
-		{
-			const TConstArrayView<FSolverVec3>* Points;
-			int32 Index;
-
-			FSolverVec3 X() const
-			{
-				return (*Points)[Index];
-			}
-
-			template<typename TPayloadType>
-			int32 GetPayload(int32) const
-			{
-				return Index;
-			}
-		};
-	}
-
 	using Chaos::TVec3;
 	template <typename T, typename ParticleType>
-	class FGaussSeidelSphereRepulsionConstraints// : public FPBDSelfCollisionSphereConstraintsBase
+	class FGaussSeidelSphereRepulsionConstraints
 	{
 	public:
-		//TODO(Yushan): Add unittest for Gauss Seidel Sphere Repulsion Constraints
+		//TODO(Yushan): Add unit tests for Gauss Seidel Sphere Repulsion Constraints
 		FGaussSeidelSphereRepulsionConstraints(FSolverReal InRadius, FSolverReal InStiffness, const ParticleType& InParticles, const FDeformableXPBDWeakConstraintParams& InParams): Radius(InRadius), Stiffness(InStiffness), DebugDrawParams(InParams)
 		{
 			const TArrayCollectionArray<FPAndInvM>& PAndInvM = InParticles.GetPAndInvM();
@@ -129,7 +109,7 @@ namespace Chaos::Softs
 			}
 
 			// Build Spatial
-			TArray<Private::FSphereSpatialEntry> Entries;
+			TArray<FSphereSpatialEntry> Entries;
 			const FSolverReal Diameter = 2.f * Radius;
 			TConstArrayView<FSolverVec3> Points = Particles.XArray();
 
@@ -200,6 +180,22 @@ namespace Chaos::Softs
 			FSolverReal Stiffness = FSolverReal(0);
 			TArray<T> ConstraintStiffness;
 		private:
+			struct FSphereSpatialEntry
+			{
+				const TConstArrayView<FSolverVec3>* Points;
+				int32 Index;
+
+				FSolverVec3 X() const
+				{
+					return (*Points)[Index];
+				}
+
+				template<typename TPayloadType>
+				int32 GetPayload(int32) const
+				{
+					return Index;
+				}
+			};
 			TArray<FSolverVec3> ReferencePositions;
 			FDeformableXPBDWeakConstraintParams DebugDrawParams;
 	};
