Add torque limit implementation for angular swing and twist drives in PBDJointSolverGaussSeidel.

diff --git a/Engine/Source/Runtime/Experimental/Chaos/Private/Chaos/Joint/PBDJointSolverGaussSeidel.cpp b/Engine/Source/Runtime/Experimental/Chaos/Private/Chaos/Joint/PBDJointSolverGaussSeidel.cpp
--- a/Engine/Source/Runtime/Experimental/Chaos/Private/Chaos/Joint/PBDJointSolverGaussSeidel.cpp
+++ b/Engine/Source/Runtime/Experimental/Chaos/Private/Chaos/Joint/PBDJointSolverGaussSeidel.cpp
@@ -1396,6 +1396,7 @@ namespace Chaos
 		const FVec3& Axis,
 		const FReal Angle,
 		const FReal AngVelTarget,
+		const FReal LambdaLimit,
 		FReal& Lambda)
 	{
 		check(!IsDynamic(KIndex));
@@ -1431,7 +1432,11 @@ namespace Chaos
 			const FReal S = SpringMassScale * JointStiffness * Dt * Dt;
 			const FReal D = SpringMassScale * JointDamping * Dt;
 			const FReal Multiplier = (FReal)1 / ((S + D) * II + (FReal)1);
-			const FReal DLambda = SolverStiffness * Multiplier * (S * Angle - D * AngVelDt - Lambda);
+			FReal DLambda = SolverStiffness * Multiplier * (S * Angle - D * AngVelDt - Lambda);
+			if (LambdaLimit > SMALL_NUMBER)
+			{
+				DLambda = FMath::Clamp(DLambda, -LambdaLimit - Lambda, LambdaLimit - Lambda);
+			}
 
 			//const FVec3 DR1 = IA1 * -DLambda;
 			const FVec3 DR1 = Axis * -(DLambda * II1);
@@ -1452,6 +1457,7 @@ namespace Chaos
 		const FVec3& Axis,
 		const FReal Angle,
 		const FReal AngVelTarget,
+		const FReal LambdaLimit,
 		FReal& Lambda)
 	{
 		check(IsDynamic(0));
@@ -1489,7 +1495,11 @@ namespace Chaos
 			const FReal S = SpringMassScale * JointStiffness * Dt * Dt;
 			const FReal D = SpringMassScale * JointDamping * Dt;
 			const FReal Multiplier = (FReal)1 / ((S + D) * II + (FReal)1);
-			const FReal DLambda = SolverStiffness * Multiplier * (S * Angle - D * AngVelDt - Lambda);
+			FReal DLambda = SolverStiffness * Multiplier * (S * Angle - D * AngVelDt - Lambda);
+			if (LambdaLimit > SMALL_NUMBER)
+			{
+				DLambda = FMath::Clamp(DLambda, -LambdaLimit - Lambda, LambdaLimit - Lambda);
+			}
 
 			//const FVec3 DR0 = IA0 * DLambda;
 			//const FVec3 DR1 = IA1 * -DLambda;
@@ -1511,19 +1521,20 @@ namespace Chaos
 		const FVec3& Axis,
 		const FReal Angle,
 		const FReal AngVelTarget,
+		const FReal LambdaLimit,
 		FReal& Lambda)
 	{
 		if (!IsDynamic(0))
 		{
-			ApplyRotationConstraintSoftKD(0, 1, Dt, JointStiffness, JointDamping, bAccelerationMode, Axis, Angle, AngVelTarget, Lambda);
+			ApplyRotationConstraintSoftKD(0, 1, Dt, JointStiffness, JointDamping, bAccelerationMode, Axis, Angle, AngVelTarget, LambdaLimit, Lambda);
 		}
 		else if (!IsDynamic(1))
 		{
-			ApplyRotationConstraintSoftKD(1, 0, Dt, JointStiffness, JointDamping, bAccelerationMode, Axis, -Angle, -AngVelTarget, Lambda);
+			ApplyRotationConstraintSoftKD(1, 0, Dt, JointStiffness, JointDamping, bAccelerationMode, Axis, -Angle, -AngVelTarget, LambdaLimit, Lambda);
 		}
 		else
 		{
-			ApplyRotationConstraintSoftDD(Dt, JointStiffness, JointDamping, bAccelerationMode, Axis, Angle, AngVelTarget, Lambda);
+			ApplyRotationConstraintSoftDD(Dt, JointStiffness, JointDamping, bAccelerationMode, Axis, Angle, AngVelTarget, LambdaLimit, Lambda);
 		}
 	}
 
@@ -1593,7 +1604,7 @@ namespace Chaos
 				const FReal TwistStiffness = FPBDJointUtilities::GetSoftTwistStiffness(SolverSettings, JointSettings);
 				const FReal TwistDamping = FPBDJointUtilities::GetSoftTwistDamping(SolverSettings, JointSettings);
 				const bool bAccelerationMode = FPBDJointUtilities::GetAngularSoftAccelerationMode(SolverSettings, JointSettings);
-				ApplyRotationConstraintSoft(Dt, TwistStiffness, TwistDamping, bAccelerationMode, TwistAxis, DTwistAngle, 0.0f, TwistSoftLambda);
+				ApplyRotationConstraintSoft(Dt, TwistStiffness, TwistDamping, bAccelerationMode, TwistAxis, DTwistAngle, 0.0f, 0.0f, TwistSoftLambda);
 			}
 			else
 			{
@@ -1628,7 +1639,7 @@ namespace Chaos
 				const FReal SoftSwingStiffness = FPBDJointUtilities::GetSoftSwingStiffness(SolverSettings, JointSettings);
 				const FReal SoftSwingDamping = FPBDJointUtilities::GetSoftSwingDamping(SolverSettings, JointSettings);
 				const bool bAccelerationMode = FPBDJointUtilities::GetAngularSoftAccelerationMode(SolverSettings, JointSettings);
-				ApplyRotationConstraintSoft(Dt, SoftSwingStiffness, SoftSwingDamping, bAccelerationMode, SwingAxis, DSwingAngle, 0.0f, SwingSoftLambda);
+				ApplyRotationConstraintSoft(Dt, SoftSwingStiffness, SoftSwingDamping, bAccelerationMode, SwingAxis, DSwingAngle, 0.0f, 0.0f, SwingSoftLambda);
 			}
 			else
 			{
@@ -1664,7 +1675,7 @@ namespace Chaos
 				const FReal SoftSwingStiffness = FPBDJointUtilities::GetSoftSwingStiffness(SolverSettings, JointSettings);
 				const FReal SoftSwingDamping = FPBDJointUtilities::GetSoftSwingDamping(SolverSettings, JointSettings);
 				const bool bAccelerationMode = FPBDJointUtilities::GetAngularSoftAccelerationMode(SolverSettings, JointSettings);
-				ApplyRotationConstraintSoft(Dt, SoftSwingStiffness, SoftSwingDamping, bAccelerationMode, SwingAxis, DSwingAngle, 0.0f, SwingSoftLambda);
+				ApplyRotationConstraintSoft(Dt, SoftSwingStiffness, SoftSwingDamping, bAccelerationMode, SwingAxis, DSwingAngle, 0.0f, 0.0f, SwingSoftLambda);
 			}
 			else
 			{
@@ -1710,7 +1721,7 @@ namespace Chaos
 				const FReal SoftSwingStiffness = FPBDJointUtilities::GetSoftSwingStiffness(SolverSettings, JointSettings);
 				const FReal SoftSwingDamping = FPBDJointUtilities::GetSoftSwingDamping(SolverSettings, JointSettings);
 				const bool bAccelerationMode = FPBDJointUtilities::GetAngularSoftAccelerationMode(SolverSettings, JointSettings);
-				ApplyRotationConstraintSoft(Dt, SoftSwingStiffness, SoftSwingDamping, bAccelerationMode, SwingAxis, DSwingAngle, 0.0f, SwingSoftLambda);
+				ApplyRotationConstraintSoft(Dt, SoftSwingStiffness, SoftSwingDamping, bAccelerationMode, SwingAxis, DSwingAngle, 0.0f, 0.0f, SwingSoftLambda);
 			}
 			else
 			{
@@ -1756,7 +1767,7 @@ namespace Chaos
 				const FReal SoftSwingStiffness = FPBDJointUtilities::GetSoftSwingStiffness(SolverSettings, JointSettings);
 				const FReal SoftSwingDamping = FPBDJointUtilities::GetSoftSwingDamping(SolverSettings, JointSettings);
 				const bool bAccelerationMode = FPBDJointUtilities::GetAngularSoftAccelerationMode(SolverSettings, JointSettings);
-				ApplyRotationConstraintSoft(Dt, SoftSwingStiffness, SoftSwingDamping, bAccelerationMode, SwingAxis, DSwingAngle, 0.0f, SwingSoftLambda);
+				ApplyRotationConstraintSoft(Dt, SoftSwingStiffness, SoftSwingDamping, bAccelerationMode, SwingAxis, DSwingAngle, 0.0f, 0.0f, SwingSoftLambda);
 			}
 			else
 			{
@@ -1787,16 +1798,20 @@ namespace Chaos
 
 		const FReal AngularTwistDriveStiffness = FPBDJointUtilities::GetAngularTwistDriveStiffness(SolverSettings, JointSettings);
 		const FReal AngularTwistDriveDamping = FPBDJointUtilities::GetAngularTwistDriveDamping(SolverSettings, JointSettings);
+		const FReal AngularTwistDriveTorqueLimit = FPBDJointUtilities::GetAngularTwistDriveTorqueLimit(SolverSettings, JointSettings);
 		const FReal AngularSwingDriveStiffness = FPBDJointUtilities::GetAngularSwingDriveStiffness(SolverSettings, JointSettings);
 		const FReal AngularSwingDriveDamping = FPBDJointUtilities::GetAngularSwingDriveDamping(SolverSettings, JointSettings);
+		const FReal AngularSwingDriveTorqueLimit = FPBDJointUtilities::GetAngularSwingDriveTorqueLimit(SolverSettings, JointSettings);
 		const bool bAccelerationMode = FPBDJointUtilities::GetAngularDriveAccelerationMode(SolverSettings, JointSettings);
+		const FReal AngularTwistDriveLambdaLimit = AngularTwistDriveTorqueLimit * Dt * Dt;
+		const FReal AngularSwingDriveLambdaLimit = AngularSwingDriveTorqueLimit * Dt * Dt;
 
 		const bool bUseTwistDrive = bTwistDriveEnabled && (((FMath::Abs(DTwistAngle) > AngleTolerance) && (AngularTwistDriveStiffness > 0.0f)) || (AngularTwistDriveDamping > 0.0f));
 		if (bUseTwistDrive)
 		{
 			FReal ReturnedLambda = RotationDriveLambdas[(int32)EJointAngularConstraintIndex::Twist];
 			const FVec3 TwistAxis = ConnectorRs[1] * FJointConstants::TwistAxis();
-			ApplyRotationConstraintSoft(Dt, AngularTwistDriveStiffness, AngularTwistDriveDamping, bAccelerationMode, TwistAxis, DTwistAngle, JointSettings.AngularDriveVelocityTarget[(int32)EJointAngularConstraintIndex::Twist], ReturnedLambda);
+			ApplyRotationConstraintSoft(Dt, AngularTwistDriveStiffness, AngularTwistDriveDamping, bAccelerationMode, TwistAxis, DTwistAngle, JointSettings.AngularDriveVelocityTarget[(int32)EJointAngularConstraintIndex::Twist], AngularTwistDriveLambdaLimit, ReturnedLambda);
 			RotationDriveLambdas[(int32)EJointAngularConstraintIndex::Twist] = ReturnedLambda;
 		}
 
@@ -1805,7 +1820,7 @@ namespace Chaos
 		{
 			FReal ReturnedLambda = RotationDriveLambdas[(int32)EJointAngularConstraintIndex::Swing1];
 			const FVec3 Swing1Axis = ConnectorRs[1] * FJointConstants::Swing1Axis();
-			ApplyRotationConstraintSoft(Dt, AngularSwingDriveStiffness, AngularSwingDriveDamping, bAccelerationMode, Swing1Axis, DSwing1Angle, JointSettings.AngularDriveVelocityTarget[(int32)EJointAngularConstraintIndex::Swing1], ReturnedLambda);
+			ApplyRotationConstraintSoft(Dt, AngularSwingDriveStiffness, AngularSwingDriveDamping, bAccelerationMode, Swing1Axis, DSwing1Angle, JointSettings.AngularDriveVelocityTarget[(int32)EJointAngularConstraintIndex::Swing1], AngularSwingDriveLambdaLimit, ReturnedLambda);
 			RotationDriveLambdas[(int32)EJointAngularConstraintIndex::Swing1] = ReturnedLambda;
 		}
 
@@ -1814,7 +1829,7 @@ namespace Chaos
 		{
 			FReal ReturnedLambda = RotationDriveLambdas[(int32)EJointAngularConstraintIndex::Swing2];
 			const FVec3 Swing2Axis = ConnectorRs[1] * FJointConstants::Swing2Axis();
-			ApplyRotationConstraintSoft(Dt, AngularSwingDriveStiffness, AngularSwingDriveDamping, bAccelerationMode, Swing2Axis, DSwing2Angle, JointSettings.AngularDriveVelocityTarget[(int32)EJointAngularConstraintIndex::Swing2], ReturnedLambda);
+			ApplyRotationConstraintSoft(Dt, AngularSwingDriveStiffness, AngularSwingDriveDamping, bAccelerationMode, Swing2Axis, DSwing2Angle, JointSettings.AngularDriveVelocityTarget[(int32)EJointAngularConstraintIndex::Swing2], AngularSwingDriveLambdaLimit, ReturnedLambda);
 			RotationDriveLambdas[(int32)EJointAngularConstraintIndex::Swing2] = ReturnedLambda;
 		}
 	}
@@ -1852,9 +1867,9 @@ namespace Chaos
 			AxisAngVel[2] = FVec3::DotProduct(TargetAngVel, Axes[2]);
 		}
 
-		ApplyRotationConstraintSoft(Dt, AngularDriveStiffness, AngularDriveDamping, bAccelerationMode, Axes[0], AxisAngles[0], AxisAngVel[0], RotationDriveLambdas[0]);
-		ApplyRotationConstraintSoft(Dt, AngularDriveStiffness, AngularDriveDamping, bAccelerationMode, Axes[1], AxisAngles[1], AxisAngVel[1], RotationDriveLambdas[1]);
-		ApplyRotationConstraintSoft(Dt, AngularDriveStiffness, AngularDriveDamping, bAccelerationMode, Axes[2], AxisAngles[2], AxisAngVel[2], RotationDriveLambdas[2]);
+		ApplyRotationConstraintSoft(Dt, AngularDriveStiffness, AngularDriveDamping, bAccelerationMode, Axes[0], AxisAngles[0], AxisAngVel[0], 0, RotationDriveLambdas[0]);
+		ApplyRotationConstraintSoft(Dt, AngularDriveStiffness, AngularDriveDamping, bAccelerationMode, Axes[1], AxisAngles[1], AxisAngVel[1], 0, RotationDriveLambdas[1]);
+		ApplyRotationConstraintSoft(Dt, AngularDriveStiffness, AngularDriveDamping, bAccelerationMode, Axes[2], AxisAngles[2], AxisAngVel[2], 0, RotationDriveLambdas[2]);
 	}
 
 
diff --git a/Engine/Source/Runtime/Experimental/Chaos/Private/Chaos/PBDJointConstraintUtilities.cpp b/Engine/Source/Runtime/Experimental/Chaos/Private/Chaos/PBDJointConstraintUtilities.cpp
--- a/Engine/Source/Runtime/Experimental/Chaos/Private/Chaos/PBDJointConstraintUtilities.cpp
+++ b/Engine/Source/Runtime/Experimental/Chaos/Private/Chaos/PBDJointConstraintUtilities.cpp
@@ -522,6 +522,17 @@ namespace Chaos
 		return 0.0f;
 	}
 
+	FReal FPBDJointUtilities::GetAngularTwistDriveTorqueLimit(
+		const FPBDJointSolverSettings& SolverSettings,
+		const FPBDJointSettings& JointSettings)
+	{
+		if (JointSettings.bAngularTwistVelocityDriveEnabled)
+		{
+			return JointSettings.AngularDriveMaxTorque[(int)EJointAngularConstraintIndex::Twist];
+		}
+		return 0.0f;
+	}
+
 	FReal FPBDJointUtilities::GetAngularSwingDriveStiffness(
 		const FPBDJointSolverSettings& SolverSettings,
 		const FPBDJointSettings& JointSettings)
@@ -544,6 +555,17 @@ namespace Chaos
 		return 0.0f;
 	}
 
+	FReal FPBDJointUtilities::GetAngularSwingDriveTorqueLimit(
+		const FPBDJointSolverSettings& SolverSettings,
+		const FPBDJointSettings& JointSettings)
+	{
+		if (JointSettings.bAngularSwingVelocityDriveEnabled)
+		{
+			return JointSettings.AngularDriveMaxTorque[(int)EJointAngularConstraintIndex::Swing1];
+		}
+		return 0.0f;
+	}
+
 	FReal FPBDJointUtilities::GetAngularSLerpDriveStiffness(
 		const FPBDJointSolverSettings& SolverSettings,
 		const FPBDJointSettings& JointSettings)
diff --git a/Engine/Source/Runtime/Experimental/Chaos/Public/Chaos/Joint/PBDJointSolverGaussSeidel.h b/Engine/Source/Runtime/Experimental/Chaos/Public/Chaos/Joint/PBDJointSolverGaussSeidel.h
--- a/Engine/Source/Runtime/Experimental/Chaos/Public/Chaos/Joint/PBDJointSolverGaussSeidel.h
+++ b/Engine/Source/Runtime/Experimental/Chaos/Public/Chaos/Joint/PBDJointSolverGaussSeidel.h
@@ -322,6 +322,7 @@ namespace Chaos
 			const FVec3& Axis,
 			const FReal Angle,
 			const FReal AngVelTarget,
+			const FReal LambdaLimit,
 			FReal& Lambda);
 
 		void ApplyRotationConstraintSoftKD(
@@ -334,6 +335,7 @@ namespace Chaos
 			const FVec3& Axis,
 			const FReal Angle,
 			const FReal AngVelTarget,
+			const FReal LambdaLimit,
 			FReal& Lambda);
 
 		void ApplyRotationConstraintSoftDD(
@@ -344,6 +346,7 @@ namespace Chaos
 			const FVec3& Axis,
 			const FReal Angle,
 			const FReal AngVelTarget,
+			const FReal LambdaLimit,
 			FReal& Lambda);
 
 		void ApplyLockedRotationConstraints(
diff --git a/Engine/Source/Runtime/Experimental/Chaos/Public/Chaos/PBDJointConstraintUtilities.h b/Engine/Source/Runtime/Experimental/Chaos/Public/Chaos/PBDJointConstraintUtilities.h
--- a/Engine/Source/Runtime/Experimental/Chaos/Public/Chaos/PBDJointConstraintUtilities.h
+++ b/Engine/Source/Runtime/Experimental/Chaos/Public/Chaos/PBDJointConstraintUtilities.h
@@ -231,6 +231,10 @@ namespace Chaos
 			const FPBDJointSolverSettings& SolverSettings,
 			const FPBDJointSettings& JointSettings);
 
+		static FReal GetAngularTwistDriveTorqueLimit(
+			const FPBDJointSolverSettings& SolverSettings,
+			const FPBDJointSettings& JointSettings);
+
 		static FReal GetAngularSwingDriveStiffness(
 			const FPBDJointSolverSettings& SolverSettings,
 			const FPBDJointSettings& JointSettings);
@@ -239,6 +243,10 @@ namespace Chaos
 			const FPBDJointSolverSettings& SolverSettings,
 			const FPBDJointSettings& JointSettings);
 
+		static FReal GetAngularSwingDriveTorqueLimit(
+			const FPBDJointSolverSettings& SolverSettings,
+			const FPBDJointSettings& JointSettings);
+
 		static FReal GetAngularSLerpDriveStiffness(
 			const FPBDJointSolverSettings& SolverSettings,
 			const FPBDJointSettings& JointSettings);
