Iris - Don't call memcpy unless size is greater than zero.

#rnx

diff --git a/Engine/Source/Runtime/Experimental/Iris/Core/Private/Iris/Serialization/NetSerializerArrayStorage.cpp b/Engine/Source/Runtime/Experimental/Iris/Core/Private/Iris/Serialization/NetSerializerArrayStorage.cpp
--- a/Engine/Source/Runtime/Experimental/Iris/Core/Private/Iris/Serialization/NetSerializerArrayStorage.cpp
+++ b/Engine/Source/Runtime/Experimental/Iris/Core/Private/Iris/Serialization/NetSerializerArrayStorage.cpp
@@ -38,7 +38,10 @@ void FNetSerializerAlignedStorage::AdjustSize(FNetSerializationContext& Context,
 		void* NewData = Context.GetInternalContext()->Alloc(InNum, InAlignment);
 		FMemory::Memzero(NewData, InNum);
 		// Copy old data
-		FMemory::Memcpy(NewData, Data, StorageNum);
+		if (StorageNum > 0)
+		{
+			FMemory::Memcpy(NewData, Data, StorageNum);
+		}
 		Context.GetInternalContext()->Free(Data);
 
 		Data = static_cast<uint8*>(NewData);
