Add density debug draw for chaos

diff --git a/Engine/Source/Runtime/Experimental/Chaos/Private/Chaos/ChaosDebugDraw.cpp b/Engine/Source/Runtime/Experimental/Chaos/Private/Chaos/ChaosDebugDraw.cpp
--- a/Engine/Source/Runtime/Experimental/Chaos/Private/Chaos/ChaosDebugDraw.cpp
+++ b/Engine/Source/Runtime/Experimental/Chaos/Private/Chaos/ChaosDebugDraw.cpp
@@ -1606,6 +1606,43 @@ namespace Chaos
 			}
 		}
 
+		void DrawParticleDensityImpl(const FRigidTransform3& SpaceTransform, const FGeometryParticleHandle* InParticle, const FChaosDebugDrawSettings& Settings)
+		{
+			if (ShouldIgnoreParticle(InParticle))
+			{
+				return;
+			}
+
+			if (const auto RigidParticle = InParticle->CastToRigidParticle())
+			{
+				if (RigidParticle->Disabled())
+				{
+					return;
+				}
+
+				const FColor Color = RigidParticle->IsDynamic() ? FColor::Yellow : FColor::Cyan;
+
+				const FVec3 PCOM = SpaceTransform.TransformPosition(FParticleUtilities::GetCoMWorldPosition(RigidParticle));
+
+				// density is per material so it is set per shape
+				for (const TUniquePtr<FPerShapeData>& Shape: RigidParticle->ShapesArray())
+				{
+					if (Shape && Shape->NumMaterials() > 0)
+					{
+						// use the first one for now
+						if (const Chaos::FChaosPhysicsMaterial* ChaosMaterial = Shape->GetMaterial(0).Get())
+						{
+							const double Density = ChaosMaterial->Density;
+							auto FormatString = TEXT("{0}");
+							const FString DensityStr = FString::Format(FormatString, { FString::SanitizeFloat(Density, 1) });
+							FDebugDrawQueue::GetInstance().DrawDebugString(PCOM, DensityStr, nullptr, Color, UE_KINDA_SMALL_NUMBER, false, Settings.FontScale);
+
+						}
+					}
+				}
+			}
+		}
+
 		void DrawJointConstraintImpl(const FRigidTransform3& SpaceTransform, const FPBDJointConstraintHandle* ConstraintHandle, Chaos::FRealSingle ColorScale, const FChaosDebugDrawJointFeatures& FeatureMask, const FChaosDebugDrawSettings& Settings)
 		{
 			if (!ConstraintHandle->IsEnabled())
@@ -2314,6 +2351,28 @@ namespace Chaos
 			}
 		}
 
+		void DrawParticleDensity(const FRigidTransform3& SpaceTransform, const TParticleView<FKinematicGeometryParticles>& ParticlesView, const FChaosDebugDrawSettings* Settings)
+		{
+			if (FDebugDrawQueue::IsDebugDrawingEnabled())
+			{
+				for (auto& Particle : ParticlesView)
+				{
+					DrawParticleDensityImpl(SpaceTransform, GetHandleHelper(&Particle), GetChaosDebugDrawSettings(Settings));
+				}
+			}
+		}
+
+		void DrawParticleDensity(const FRigidTransform3& SpaceTransform, const TParticleView<FPBDRigidParticles>& ParticlesView, const FChaosDebugDrawSettings* Settings)
+		{
+			if (FDebugDrawQueue::IsDebugDrawingEnabled())
+			{
+				for (auto& Particle : ParticlesView)
+				{
+					DrawParticleDensityImpl(SpaceTransform, GetHandleHelper(&Particle), GetChaosDebugDrawSettings(Settings));
+				}
+			}
+		}
+
 		void DrawJointConstraints(const FRigidTransform3& SpaceTransform, const TArray<FPBDJointConstraintHandle*>& ConstraintHandles, Chaos::FRealSingle ColorScale, const FChaosDebugDrawJointFeatures& FeatureMask, const FChaosDebugDrawSettings* Settings)
 		{
 			if (FDebugDrawQueue::IsDebugDrawingEnabled())
diff --git a/Engine/Source/Runtime/Experimental/Chaos/Private/PBDRigidsSolver.cpp b/Engine/Source/Runtime/Experimental/Chaos/Private/PBDRigidsSolver.cpp
--- a/Engine/Source/Runtime/Experimental/Chaos/Private/PBDRigidsSolver.cpp
+++ b/Engine/Source/Runtime/Experimental/Chaos/Private/PBDRigidsSolver.cpp
@@ -124,6 +124,7 @@ namespace Chaos
 	{
 		int32 ChaosSolverDebugDrawShapes = CHAOS_SOLVER_ENABLE_DEBUG_DRAW;
 		int32 ChaosSolverDebugDrawMass = 0;
+		int32 ChaosSolverDebugDrawDensity = 0;
 		int32 ChaosSolverDebugDrawBVHs = 0;
 		int32 ChaosSolverDebugDrawCollisions = CHAOS_SOLVER_ENABLE_DEBUG_DRAW;
 		int32 ChaosSolverDebugDrawCollidingShapes = 0;
@@ -152,6 +153,7 @@ namespace Chaos
 		DebugDraw::FChaosDebugDrawJointFeatures ChaosSolverDrawJointFeatures = DebugDraw::FChaosDebugDrawJointFeatures::MakeDefault();
 		FAutoConsoleVariableRef CVarChaosSolverDrawShapes(TEXT("p.Chaos.Solver.DebugDrawShapes"), ChaosSolverDebugDrawShapes, TEXT("Draw Shapes (0 = never; 1 = end of frame)."));
 		FAutoConsoleVariableRef CVarChaosSolverDebugDrawMass(TEXT("p.Chaos.Solver.DebugDrawMass"), ChaosSolverDebugDrawMass, TEXT("Draw Mass values in Kg (0 = never; 1 = end of frame)."));
+		FAutoConsoleVariableRef CVarChaosSolverDebugDrawDensity(TEXT("p.Chaos.Solver.DebugDrawDensity"), ChaosSolverDebugDrawDensity, TEXT("Draw Density values in Kg/cm3 (0 = never; 1 = end of frame)."));
 		FAutoConsoleVariableRef CVarChaosSolverDrawBVHs(TEXT("p.Chaos.Solver.DebugDrawBVHs"), ChaosSolverDebugDrawBVHs, TEXT("Draw Particle BVHs where applicable (0 = never; 1 = end of frame)."));
 		FAutoConsoleVariableRef CVarChaosSolverDrawCollisions(TEXT("p.Chaos.Solver.DebugDrawCollisions"), ChaosSolverDebugDrawCollisions, TEXT("Draw Collisions (0 = never; 1 = end of frame)."));
 		FAutoConsoleVariableRef CVarChaosSolverDrawCollidingShapes(TEXT("p.Chaos.Solver.DebugDrawCollidingShapes"), ChaosSolverDebugDrawCollidingShapes, TEXT("Draw Shapes that have collisions on them (0 = never; 1 = end of frame)."));
@@ -2512,6 +2514,11 @@ TRACE_COUNTER_SET(ChaosTraceCounter_##Name, Value)
 			DebugDraw::DrawParticleMass(FRigidTransform3(), Particles.GetActiveKinematicParticlesView(), &ChaosSolverDebugDebugDrawSettings);
 			DebugDraw::DrawParticleMass(FRigidTransform3(), Particles.GetNonDisabledDynamicView(), &ChaosSolverDebugDebugDrawSettings);
 		}
+		if (ChaosSolverDebugDrawDensity == 1)
+		{
+			DebugDraw::DrawParticleDensity(FRigidTransform3(), Particles.GetActiveKinematicParticlesView(), &ChaosSolverDebugDebugDrawSettings);
+			DebugDraw::DrawParticleDensity(FRigidTransform3(), Particles.GetNonDisabledDynamicView(), &ChaosSolverDebugDebugDrawSettings);
+		}
 		if (ChaosSolverDebugDrawBVHs == 1)
 		{
 			for (const FTransientGeometryParticleHandle& Particle : Particles.GetNonDisabledView())
diff --git a/Engine/Source/Runtime/Experimental/Chaos/Public/Chaos/ChaosDebugDraw.h b/Engine/Source/Runtime/Experimental/Chaos/Public/Chaos/ChaosDebugDraw.h
--- a/Engine/Source/Runtime/Experimental/Chaos/Public/Chaos/ChaosDebugDraw.h
+++ b/Engine/Source/Runtime/Experimental/Chaos/Public/Chaos/ChaosDebugDraw.h
@@ -219,6 +219,8 @@ namespace Chaos
 		CHAOS_API void DrawCollisions(const FRigidTransform3& SpaceTransform, const Private::FCollisionConstraintAllocator& CollisionAllocator, FRealSingle ColorScale, const FChaosDebugDrawSettings* Settings = nullptr);
 		CHAOS_API void DrawParticleMass(const FRigidTransform3& SpaceTransform, const TParticleView<FKinematicGeometryParticles>& ParticlesView, const FChaosDebugDrawSettings* Settings = nullptr);
 		CHAOS_API void DrawParticleMass(const FRigidTransform3& SpaceTransform, const TParticleView<FPBDRigidParticles>& ParticlesView, const FChaosDebugDrawSettings* Settings = nullptr);
+		CHAOS_API void DrawParticleDensity(const FRigidTransform3& SpaceTransform, const TParticleView<FKinematicGeometryParticles>& ParticlesView, const FChaosDebugDrawSettings* Settings = nullptr);
+		CHAOS_API void DrawParticleDensity(const FRigidTransform3& SpaceTransform, const TParticleView<FPBDRigidParticles>& ParticlesView, const FChaosDebugDrawSettings* Settings = nullptr);
 		CHAOS_API void DrawJointConstraints(const FRigidTransform3& SpaceTransform, const TArray<FPBDJointConstraintHandle*>& ConstraintHandles, FRealSingle ColorScale, const FChaosDebugDrawJointFeatures& FeatureMask = FChaosDebugDrawJointFeatures::MakeDefault(), const FChaosDebugDrawSettings* Settings = nullptr);
 		CHAOS_API void DrawJointConstraints(const FRigidTransform3& SpaceTransform, const FPBDJointConstraints& Constraints, FRealSingle ColorScale, const FChaosDebugDrawJointFeatures& FeatureMask = FChaosDebugDrawJointFeatures::MakeDefault(), const FChaosDebugDrawSettings* Settings = nullptr);
 		CHAOS_API void DrawCharacterGroundConstraints(const FRigidTransform3& SpaceTransform, const FCharacterGroundConstraintContainer& Constraints, const FChaosDebugDrawSettings* Settings = nullptr);
