Add debug command to see which classes are fully push-based.

diff --git a/Engine/Source/Runtime/Experimental/Iris/Core/Private/Iris/ReplicationSystem/ObjectReplicationBridgeDebugging.cpp b/Engine/Source/Runtime/Experimental/Iris/Core/Private/Iris/ReplicationSystem/ObjectReplicationBridgeDebugging.cpp
--- a/Engine/Source/Runtime/Experimental/Iris/Core/Private/Iris/ReplicationSystem/ObjectReplicationBridgeDebugging.cpp
+++ b/Engine/Source/Runtime/Experimental/Iris/Core/Private/Iris/ReplicationSystem/ObjectReplicationBridgeDebugging.cpp
@@ -19,6 +19,8 @@
 #include "Net/Core/NetBitArrayPrinter.h"
 #include "Net/Core/Trace/NetDebugName.h"
 
+#include "UObject/CoreNet.h"
+
 /**
  * This class contains misc console commands that log the state of different Iris systems.
  * 
@@ -715,4 +717,112 @@ void UObjectReplicationBridge::PrintNetCullDistances(const TArray<FString>& Args
 	
 	UE_LOG(LogIrisBridge, Display, TEXT(""));
 	UE_LOG(LogIrisBridge, Display, TEXT("################ Stop Printing NetCullDistance Values ################"));
-}
\ No newline at end of file
+}
+
+//-----------------------------------------------
+FAutoConsoleCommand ObjectBridgePrintPushBasedStatuses(
+	TEXT("Net.Iris.PrintPushBasedStatuses"), 
+	TEXT("Prints the push-based statuses of all classes."), 
+	FConsoleCommandWithArgsDelegate::CreateLambda([](const TArray< FString >& Args)
+{
+	using namespace UE::Net::Private;
+	using namespace UE::Net::Private::ObjectBridgeDebugging;
+
+	UReplicationSystem* RepSystem = FindReplicationSystemFromArg(Args);
+	if (!RepSystem)
+	{
+		UE_LOG(LogIrisBridge, Error, TEXT("Could not find ReplicationSystem."));
+		return;
+	}
+
+	UObjectReplicationBridge* ObjectBridge = RepSystem->GetReplicationBridgeAs<UObjectReplicationBridge>();
+	if (!ObjectBridge)
+	{
+		UE_LOG(LogIrisBridge, Error, TEXT("Could not find ObjectReplicationBridge."));
+		return;
+	}
+
+	ObjectBridge->PrintPushBasedStatuses();
+}));
+
+void UObjectReplicationBridge::PrintPushBasedStatuses() const
+{
+	using namespace UE::Net;
+	using namespace UE::Net::Private;
+
+	FReplicationProtocolManager* ProtocolManager = GetReplicationProtocolManager();
+	if (!ProtocolManager)
+	{
+		UE_LOG(LogIrisBridge, Error, TEXT("Could not find ReplicationProtocolManager."));
+		return;
+	}
+
+	struct FPushBasedInfo
+	{
+		const UClass* Class = nullptr;
+		int32 RefCount = 0;
+		bool bIsFullyPushBased = false;
+	};
+
+	TArray<FPushBasedInfo> PushBasedInfos;
+	ProtocolManager->ForEachProtocol([&](const FReplicationProtocol* Protocol, const UObject* ArchetypeOrCDOUsedAsKey)
+	{
+		if (!ArchetypeOrCDOUsedAsKey)
+		{
+			return;
+		}
+
+		for (const FReplicationStateDescriptor* StateDescriptor : MakeArrayView(Protocol->ReplicationStateDescriptors, Protocol->ReplicationStateCount))
+		{
+			if (!EnumHasAnyFlags(StateDescriptor->Traits, EReplicationStateTraits::HasPushBasedDirtiness))
+			{
+				PushBasedInfos.Add({ArchetypeOrCDOUsedAsKey->GetClass(), Protocol->GetRefCount(), false});
+				return;
+			}
+		}
+
+		PushBasedInfos.Add({ArchetypeOrCDOUsedAsKey->GetClass(), Protocol->GetRefCount(), true});
+	});
+
+	// Print by push-based status (not push-based first), then by ref count, then by name.
+	Algo::Sort(PushBasedInfos, [](const FPushBasedInfo& A, const FPushBasedInfo& B) 
+	{
+		if (A.bIsFullyPushBased != B.bIsFullyPushBased)
+		{
+			return B.bIsFullyPushBased;
+		}
+		else if (A.RefCount != B.RefCount)
+		{
+			return A.RefCount > B.RefCount;
+		}
+		return A.Class->GetName() < B.Class->GetName();
+	});
+
+	UE_LOG(LogIrisBridge, Display, TEXT("################ Start Printing Push-Based Statuses ################"));
+	UE_LOG(LogIrisBridge, Display, TEXT(""));
+
+	for (const FPushBasedInfo& Info : PushBasedInfos)
+	{
+		UE_LOG(LogIrisBridge, Display, TEXT("%s (RefCount: %d) (PushBased: %d)"), ToCStr(Info.Class->GetName()), Info.RefCount, (int32)Info.bIsFullyPushBased);
+		if (!Info.bIsFullyPushBased)
+		{
+			UE_LOG(LogIrisBridge, Display, TEXT("\tPrinting properties that aren't push-based:"));
+
+			TArray<FLifetimeProperty> LifetimeProps;
+			LifetimeProps.Reserve(Info.Class->ClassReps.Num());
+			Info.Class->GetDefaultObject()->GetLifetimeReplicatedProps(LifetimeProps);
+			for (const FLifetimeProperty& LifetimeProp : LifetimeProps)
+			{
+				if (!LifetimeProp.bIsPushBased && LifetimeProp.Condition != COND_Never)
+				{
+					const FRepRecord& RepRecord = Info.Class->ClassReps[LifetimeProp.RepIndex];
+					const FProperty* Prop = RepRecord.Property;
+					UE_LOG(LogIrisBridge, Display, TEXT("\t\t%s"), ToCStr(Prop->GetPathName()));
+				}
+			}
+		}
+	}
+
+	UE_LOG(LogIrisBridge, Display, TEXT(""));
+	UE_LOG(LogIrisBridge, Display, TEXT("################ Stop Printing Push-Based Statuses ################"));
+}
diff --git a/Engine/Source/Runtime/Experimental/Iris/Core/Private/Iris/ReplicationSystem/ReplicationProtocolManager.h b/Engine/Source/Runtime/Experimental/Iris/Core/Private/Iris/ReplicationSystem/ReplicationProtocolManager.h
--- a/Engine/Source/Runtime/Experimental/Iris/Core/Private/Iris/ReplicationSystem/ReplicationProtocolManager.h
+++ b/Engine/Source/Runtime/Experimental/Iris/Core/Private/Iris/ReplicationSystem/ReplicationProtocolManager.h
@@ -32,6 +32,10 @@ public:
 	/** Iterate over all protocols matching the ProtocolId, mostly used by debug functionality with no sideeffects */
 	template<typename T>
 	void ForEachProtocol(FReplicationProtocolIdentifier ProtocolId, T&& Functor) const;
+
+	/** Iterate over all protocols. */
+	template<typename T>
+	void ForEachProtocol(T&& Functor) const;
 	
 	/* Destroy existing replication protocol */
 	IRISCORE_API void DestroyReplicationProtocol(const FReplicationProtocol* ReplicationProtocol);
@@ -88,5 +92,14 @@ void FReplicationProtocolManager::ForEachProtocol(FReplicationProtocolIdentifier
 	}
 }
 
+template<typename T>
+void FReplicationProtocolManager::ForEachProtocol(T&& Functor) const
+{
+	for (auto It = RegisteredProtocols.CreateConstIterator(); It; ++It)
+	{
+		const FRegisteredProtocolInfo& Info = It.Value();
+		Functor(Info.Protocol, Info.ArchetypeOrCDOUsedAsKey);
+	}
+}
 
 }
diff --git a/Engine/Source/Runtime/Experimental/Iris/Core/Public/Iris/ReplicationSystem/ObjectReplicationBridge.h b/Engine/Source/Runtime/Experimental/Iris/Core/Public/Iris/ReplicationSystem/ObjectReplicationBridge.h
--- a/Engine/Source/Runtime/Experimental/Iris/Core/Public/Iris/ReplicationSystem/ObjectReplicationBridge.h
+++ b/Engine/Source/Runtime/Experimental/Iris/Core/Public/Iris/ReplicationSystem/ObjectReplicationBridge.h
@@ -199,6 +199,7 @@ public:
 	void PrintRelevantObjects(uint32 ArgTraits) const;
 	void PrintRelevantObjectsForConnections(const TArray<FString>& Args) const;
 	void PrintNetCullDistances(const TArray<FString>& Args) const;
+	void PrintPushBasedStatuses() const;
 
 protected:
 	IRISCORE_API virtual ~UObjectReplicationBridge();
