Geometry collection : add support for collision visualization mode  ( player collision & Visibility collision )

diff --git a/Engine/Source/Runtime/Engine/Classes/Components/MeshComponent.h b/Engine/Source/Runtime/Engine/Classes/Components/MeshComponent.h
--- a/Engine/Source/Runtime/Engine/Classes/Components/MeshComponent.h
+++ b/Engine/Source/Runtime/Engine/Classes/Components/MeshComponent.h
@@ -144,6 +144,9 @@ public:
 	/** Generate streaming data for all materials. */
 	ENGINE_API void GetStreamingTextureInfoInner(FStreamingTextureLevelContext& LevelContext, const TArray<FStreamingTextureBuildInfo>* PreBuiltData, float ComponentScaling, TArray<FStreamingRenderAssetPrimitiveInfo>& OutStreamingTextures) const;
 
+	/** Returns the wireframe color to use for this component. */
+	ENGINE_API FColor GetWireframeColorForSceneProxy() const;
+
 #if !(UE_BUILD_SHIPPING || UE_BUILD_TEST)
 	/**
 	 * Output to the log which materials and textures are used by this component.
diff --git a/Engine/Source/Runtime/Engine/Private/Components/MeshComponent.cpp b/Engine/Source/Runtime/Engine/Private/Components/MeshComponent.cpp
--- a/Engine/Source/Runtime/Engine/Private/Components/MeshComponent.cpp
+++ b/Engine/Source/Runtime/Engine/Private/Components/MeshComponent.cpp
@@ -599,6 +599,29 @@ void UMeshComponent::GetStreamingTextureInfoInner(FStreamingTextureLevelContext&
 	}
 }
 
+FColor UMeshComponent::GetWireframeColorForSceneProxy() const
+{
+	if (Mobility == EComponentMobility::Static)
+	{
+		return FColor(0, 255, 255, 255);
+	}
+	else if (Mobility == EComponentMobility::Stationary)
+	{
+		return FColor(128, 128, 255, 255);
+	}
+	else // Movable
+	{
+		if (BodyInstance.bSimulatePhysics)
+		{
+			return FColor(0, 255, 128, 255);
+		}
+		else
+		{
+			return FColor(255, 0, 255, 255);
+		}
+	}
+}
+
 #if !(UE_BUILD_SHIPPING || UE_BUILD_TEST)
 
 void UMeshComponent::LogMaterialsAndTextures(FOutputDevice& Ar, int32 Indent) const
diff --git a/Engine/Source/Runtime/Engine/Private/Components/StaticMeshComponent.cpp b/Engine/Source/Runtime/Engine/Private/Components/StaticMeshComponent.cpp
--- a/Engine/Source/Runtime/Engine/Private/Components/StaticMeshComponent.cpp
+++ b/Engine/Source/Runtime/Engine/Private/Components/StaticMeshComponent.cpp
@@ -1167,28 +1167,7 @@ FColor UStaticMeshComponent::GetWireframeColor() const
 	{
 		return WireframeColorOverride;
 	}
-	else
-	{
-		if(Mobility == EComponentMobility::Static)
-		{
-			return FColor(0, 255, 255, 255);
-		}
-		else if(Mobility == EComponentMobility::Stationary)
-		{
-			return FColor(128, 128, 255, 255);
-		}
-		else // Movable
-		{
-			if(BodyInstance.bSimulatePhysics)
-			{
-				return FColor(0, 255, 128, 255);
-			}
-			else
-			{
-				return FColor(255, 0, 255, 255);
-			}
-		}
-	}
+	return Super::GetWireframeColorForSceneProxy();
 }
 
 
diff --git a/Engine/Source/Runtime/Experimental/GeometryCollectionEngine/Private/GeometryCollection/GeometryCollectionDebugDraw.cpp b/Engine/Source/Runtime/Experimental/GeometryCollectionEngine/Private/GeometryCollection/GeometryCollectionDebugDraw.cpp
new file mode 100644
--- /dev/null
+++ b/Engine/Source/Runtime/Experimental/GeometryCollectionEngine/Private/GeometryCollection/GeometryCollectionDebugDraw.cpp
@@ -0,0 +1,179 @@
+// Copyright Epic Games, Inc. All Rights Reserved.
+
+#include "GeometryCollection/GeometryCollectionDebugDraw.h"
+
+#include "GeometryCollection/GeometryCollection.h"
+#include "GeometryCollection/Facades/CollectionTransformFacade.h"
+#include "GeometryCollection/Facades/CollectionHierarchyFacade.h"
+#include "GeometryCollectionProxyData.h"
+
+#include "SceneManagement.h"
+#include "PhysicsEngine/SphereElem.h"
+#include "PhysicsEngine/SphylElem.h"
+#include "PhysicsEngine/BoxElem.h"
+#include "PhysicsEngine/ConvexElem.h"
+#include "PhysicsEngine/AggregateGeom.h"
+
+#include "Chaos/Sphere.h"
+#include "Chaos/Capsule.h"
+#include "Chaos/Box.h"
+#include "Chaos/Convex.h"
+#include "Chaos/ImplicitObjectScaled.h"
+#include "Chaos/TriangleMeshImplicitObject.h"
+#include "Chaos/HeightField.h"
+
+namespace GeometryCollectionDebugDraw
+{
+	void AddSphere(FKAggregateGeom& AggGeom, const Chaos::FImplicitSphere3& Sphere, const Chaos::FRigidTransform3& Transform)
+	{
+		FKSphereElem& SphereElem = AggGeom.SphereElems.AddDefaulted_GetRef();
+		SphereElem.Center = Transform.TransformPosition(Sphere.GetCenter());
+		SphereElem.Radius = (Transform.GetScale3D().Z * Sphere.GetRadius());
+	}
+
+	void AddCapsule(FKAggregateGeom& AggGeom, const Chaos::FImplicitCapsule3& Capsule, const Chaos::FRigidTransform3& Transform)
+	{
+		FKSphylElem& CapsuleElem = AggGeom.SphylElems.AddDefaulted_GetRef();
+		CapsuleElem.Center = Transform.TransformPosition(Capsule.GetCenter());
+		CapsuleElem.Rotation = (Transform.GetRotation() * FRotationMatrix::MakeFromZ(Capsule.GetAxis())).Rotator();
+		CapsuleElem.Length = (Transform.GetScale3D().Z * Capsule.GetHeight());
+		CapsuleElem.Radius = (Transform.GetScale3D().X * Capsule.GetRadius());
+	}
+
+	void AddBox(FKAggregateGeom& AggGeom, const Chaos::FImplicitBox3& Box, const Chaos::FRigidTransform3& Transform)
+	{
+		FKBoxElem& BoxElem = AggGeom.BoxElems.AddDefaulted_GetRef();
+		BoxElem.Center = Transform.TransformPosition(Box.Center());
+		BoxElem.Rotation = Transform.Rotator();
+		BoxElem.X = Transform.GetScale3D().X * Box.Extents().X;
+		BoxElem.Y = Transform.GetScale3D().Y * Box.Extents().Y;
+		BoxElem.Z = Transform.GetScale3D().Z * Box.Extents().Z;
+	}
+
+	void AddConvex(FKAggregateGeom& AggGeom, const Chaos::FImplicitConvex3& Convex, const Chaos::FRigidTransform3& Transform)
+	{
+		FKConvexElem& ConvexElem = AggGeom.ConvexElems.AddDefaulted_GetRef();
+		ConvexElem.VertexData.Reserve(Convex.GetVertices().Num());
+		for (const FVector3f& Vtx : Convex.GetVertices())
+		{
+			ConvexElem.VertexData.Emplace(Transform.TransformPosition(FVector(Vtx)));
+		}
+		for (int32 PlaneIndex = 0; PlaneIndex < Convex.NumPlanes(); PlaneIndex++)
+		{
+			const int32 NumVertices = Convex.NumPlaneVertices(PlaneIndex);
+			if (NumVertices >= 3)
+			{
+				const int32 NumTriangles = (NumVertices - 2);
+				ConvexElem.IndexData.Reserve(ConvexElem.IndexData.Num() + (NumTriangles * 3));
+				for (int32 VtxIndex = 2; VtxIndex < NumVertices; VtxIndex++)
+				{
+					ConvexElem.IndexData.Emplace(Convex.GetPlaneVertex(PlaneIndex, 0));
+					ConvexElem.IndexData.Emplace(Convex.GetPlaneVertex(PlaneIndex, VtxIndex-1));
+					ConvexElem.IndexData.Emplace(Convex.GetPlaneVertex(PlaneIndex, VtxIndex));
+				}
+			}
+		}
+	}
+
+	void AddImplicitObject(FKAggregateGeom& AggGeom, const Chaos::FImplicitObject* Implicit, const Chaos::FRigidTransform3& Transform)
+	{
+		using namespace Chaos;
+
+		if (const FImplicitSphere3* Sphere = Implicit->AsA<FImplicitSphere3>())
+		{
+			AddSphere(AggGeom, *Sphere, Transform);
+		}
+		else if (const FImplicitCapsule3* Capsule = Implicit->AsA<FImplicitCapsule3>())
+		{
+			AddCapsule(AggGeom, *Capsule, Transform);
+		}
+		else if (const FImplicitBox3* Box = Implicit->AsA<FImplicitBox3>())
+		{
+			AddBox(AggGeom, *Box, Transform);
+		}
+		else if (const FImplicitConvex3* Convex = Implicit->AsA<FImplicitConvex3>())
+		{
+			AddConvex(AggGeom, *Convex, Transform);
+		}
+		else if (const FTriangleMeshImplicitObject* TriMesh = Implicit->AsA<FTriangleMeshImplicitObject>())
+		{
+			ensure(false); // unsupported for geometry collection
+		}
+		else if (const FHeightField* HeightField = Implicit->AsA<FHeightField>())
+		{
+			ensure(false); // unsupported for geometry collection
+		}
+		else if (const FImplicitObjectInstanced* Instanced = Implicit->AsA<FImplicitObjectInstanced>())
+		{
+			AddImplicitObject(AggGeom, Instanced->GetInnerObject().Get(), Transform);
+		}
+		else if (const FImplicitObjectScaled* Scaled = Implicit->AsA<FImplicitObjectScaled>())
+		{
+			const FRigidTransform3 ScaledTransform = FRigidTransform3(Transform.GetTranslation(), Transform.GetRotation(), Transform.GetScale3D() * Scaled->GetScale());
+			AddImplicitObject(AggGeom, Scaled->GetInnerObject().Get(), ScaledTransform);
+		}
+		else if (const FImplicitObjectUnion* Union = Implicit->AsA<FImplicitObjectUnion>())
+		{
+			ensure(false);
+		}
+	}
+
+
+	void Draw(const FGeometryCollection& Collection, const FTransform& CollectionWorldTransform, FMeshElementCollector& MeshCollector, int32 ViewIndex, const FMaterialRenderProxy* MaterialProxy, const FColor& Color, bool bDrawSolid)
+	{
+		const FName MassToLocalAttributeName = "MassToLocal";
+
+		const TManagedArray<FTransform>* MassToLocalPtr = Collection.FindAttribute<FTransform>(MassToLocalAttributeName, FTransformCollection::TransformGroup);
+		const TManagedArray<Chaos::FImplicitObjectPtr>* ImplicitsPtr = Collection.FindAttribute<Chaos::FImplicitObjectPtr>(FGeometryDynamicCollection::ImplicitsAttribute, FTransformCollection::TransformGroup);
+		if (MassToLocalPtr && ImplicitsPtr)
+		{
+			const TManagedArray<FTransform>& MassToLocal = *MassToLocalPtr;
+			const TManagedArray<Chaos::FImplicitObjectPtr>& Implicits = *ImplicitsPtr;
+
+			// let's take advantage of the existing method for rendering Aggregate geometry 
+			FKAggregateGeom AggGeom;
+
+			const GeometryCollection::Facades::FCollectionTransformFacade TransformFacade(Collection);
+			const Chaos::Facades::FCollectionHierarchyFacade HierarchyFacade(Collection);
+			TArray<int32> TransformIndices = TransformFacade.GetRootIndices();
+
+			for (int32 Index = 0; Index < TransformIndices.Num(); Index++)
+			{
+				const int32 TransformIndex = TransformIndices[Index];
+				if (const Chaos::FImplicitObjectPtr ImplicitPtr = Implicits[TransformIndex])
+				{
+					// TODO : use the root bone transform * component transform
+					const Chaos::FRigidTransform3 CollectionSpaceTransform = TransformFacade.ComputeCollectionSpaceTransform(TransformIndex);
+					const Chaos::FRigidTransform3 CollectionSpaceParticleTransform = MassToLocal[TransformIndex] * CollectionSpaceTransform;
+
+					ImplicitPtr->VisitLeafObjects(
+						[&](const Chaos::FImplicitObject* LeafImplicitObject, const Chaos::FRigidTransform3& LeafRelativeTransform, const int32 UnusedRootObjectIndex, const int32 UnusedObjectIndex, const int32 UnusedLeafObjectIndex)
+						{
+							const Chaos::FRigidTransform3 LeafTransform = LeafRelativeTransform * CollectionSpaceParticleTransform;
+							AddImplicitObject(AggGeom, LeafImplicitObject, LeafTransform);
+						});
+				}
+				else
+				{
+					// no implicit , so we need to use the children ones 
+					TransformIndices.Append(HierarchyFacade.GetChildrenAsArray(TransformIndex));
+				}
+
+			}
+
+			// collect the rendering data through the collector 
+			AggGeom.GetAggGeom(CollectionWorldTransform, Color, MaterialProxy, /*bPerHullColor*/false, bDrawSolid, /*bOutputVelocity*/false, ViewIndex, MeshCollector);
+		}
+	}
+
+	void DrawSolid(const FGeometryCollection& Collection, const FTransform& CollectionWorldTransform, FMeshElementCollector& MeshCollector, int32 ViewIndex, const FMaterialRenderProxy* MaterialProxy)
+	{
+		Draw(Collection, CollectionWorldTransform, MeshCollector, ViewIndex, MaterialProxy, /*Color*/FColor::White, /*bDrawSolid*/true);
+	}
+
+	void DrawWireframe(const FGeometryCollection& Collection, const FTransform& CollectionWorldTransform, FMeshElementCollector& MeshCollector, int32 ViewIndex, const FColor& Color)
+	{
+		Draw(Collection, CollectionWorldTransform, MeshCollector, ViewIndex, /*MaterialProxy*/nullptr, Color, /*bDrawSolid*/false);
+	}
+
+}
diff --git a/Engine/Source/Runtime/Experimental/GeometryCollectionEngine/Private/GeometryCollection/GeometryCollectionSceneProxy.cpp b/Engine/Source/Runtime/Experimental/GeometryCollectionEngine/Private/GeometryCollection/GeometryCollectionSceneProxy.cpp
--- a/Engine/Source/Runtime/Experimental/GeometryCollectionEngine/Private/GeometryCollection/GeometryCollectionSceneProxy.cpp
+++ b/Engine/Source/Runtime/Experimental/GeometryCollectionEngine/Private/GeometryCollection/GeometryCollectionSceneProxy.cpp
@@ -17,6 +17,7 @@
 #include "GeometryCollection/GeometryCollectionComponent.h"
 #include "GeometryCollection/GeometryCollectionAlgo.h"
 #include "GeometryCollection/GeometryCollectionHitProxy.h"
+#include "GeometryCollection/GeometryCollectionDebugDraw.h"
 #include "RHIDefinitions.h"
 #include "ComponentReregisterContext.h"
 #include "ComponentRecreateRenderStateContext.h"
@@ -127,6 +128,11 @@ FGeometryCollectionSceneProxy::FGeometryCollectionSceneProxy(UGeometryCollection
 	, VertexFactoryDebugColor(GetScene().GetFeatureLevel())
 #endif
 {
+	if (Component->GetRestCollection())
+	{
+		GeometryCollection = Component->GetRestCollection()->GetGeometryCollection();
+	}
+
 	EnableGPUSceneSupportFlags();
 
 	Materials.Empty();
@@ -205,6 +211,9 @@ FGeometryCollectionSceneProxy::FGeometryCollectionSceneProxy(UGeometryCollection
 	bAlwaysHasVelocity = true;
 
 	DynamicData = Component->InitDynamicData(true);
+
+	SetWireframeColor(Component->GetWireframeColorForSceneProxy());
+	CollisionResponse = Component->GetCollisionResponseToChannels();
 }
 
 FGeometryCollectionSceneProxy::~FGeometryCollectionSceneProxy()
@@ -619,6 +628,26 @@ FVertexFactory const* FGeometryCollectionSceneProxy::GetVertexFactory() const
 #endif
 }
 
+bool FGeometryCollectionSceneProxy::ShowCollisionMeshes(const FEngineShowFlags& EngineShowFlags) const
+{
+	if (IsCollisionEnabled())
+	{
+		if (EngineShowFlags.CollisionPawn && CollisionResponse.GetResponse(ECC_Pawn) != ECR_Ignore)
+		{
+			return true;
+		}
+		if (EngineShowFlags.CollisionVisibility && CollisionResponse.GetResponse(ECC_Visibility) != ECR_Ignore)
+		{
+			return true;
+		}
+		if (EngineShowFlags.Collision)
+		{
+			return true;
+		}
+	}
+	return false;
+}
+
 void FGeometryCollectionSceneProxy::GetDynamicMeshElements(const TArray<const FSceneView*>& Views, const FSceneViewFamily& ViewFamily, uint32 VisibilityMap, FMeshElementCollector& Collector) const
 {
 	QUICK_SCOPE_CYCLE_COUNTER(STAT_GeometryCollectionSceneProxy_GetDynamicMeshElements);
@@ -630,6 +659,8 @@ void FGeometryCollectionSceneProxy::GetDynamicMeshElements(const TArray<const FS
 	const FEngineShowFlags& EngineShowFlags = ViewFamily.EngineShowFlags;
 	const bool bWireframe = AllowDebugViewmodes() && EngineShowFlags.Wireframe;
 	const bool bProxyIsSelected = IsSelected();
+	const bool bDrawOnlyCollisionMeshes = EngineShowFlags.CollisionPawn || EngineShowFlags.CollisionVisibility;
+	const bool bDrawWireframeCollision = EngineShowFlags.Collision && IsCollisionEnabled();
 
 	auto SetDebugMaterial = [this, &Collector, &EngineShowFlags, bProxyIsSelected](FMeshBatch& Mesh) -> void
 	{
@@ -672,7 +703,7 @@ void FGeometryCollectionSceneProxy::GetDynamicMeshElements(const TArray<const FS
 			// but whole component selection seems ok for now
 			bool bSectionIsSelected = bProxyIsSelected;
 
-			auto VertexColorVisualizationMaterialInstance = new FColoredMaterialRenderProxy(
+			FMaterialRenderProxy* VertexColorVisualizationMaterialInstance = new FColoredMaterialRenderProxy(
 				VertexColorVisualizationMaterial->GetRenderProxy(),
 				GetSelectionColor(FLinearColor::White, bSectionIsSelected, IsHovered())
 			);
@@ -685,96 +716,125 @@ void FGeometryCollectionSceneProxy::GetDynamicMeshElements(const TArray<const FS
 #endif
 	};
 
-	for (int32 ViewIndex = 0; ViewIndex < Views.Num(); ViewIndex++)
+	const bool bDrawGeometryCollectionMesh = !bDrawOnlyCollisionMeshes;
+
+	if (bDrawGeometryCollectionMesh)
 	{
-		if ((VisibilityMap & (1 << ViewIndex)) == 0) 
+		for (int32 ViewIndex = 0; ViewIndex < Views.Num(); ViewIndex++)
 		{
-			continue; 
-		}
+			if ((VisibilityMap & (1 << ViewIndex)) == 0)
+			{
+				continue;
+			}
 
-		// If not dynamic then use the section array with interior fracture surfaces removed.
-		bool bRemoveInternalFaces = DynamicData != nullptr && !DynamicData->IsDynamic && MeshDescription.SectionsNoInternal.Num();
+			// If not dynamic then use the section array with interior fracture surfaces removed.
+			bool bRemoveInternalFaces = DynamicData != nullptr && !DynamicData->IsDynamic && MeshDescription.SectionsNoInternal.Num();
 
 #if WITH_EDITOR
-		// If hiding geometry in editor then we don't remove hidden faces.
-		bRemoveInternalFaces &= HiddenTransforms.Num() == 0;
+			// If hiding geometry in editor then we don't remove hidden faces.
+			bRemoveInternalFaces &= HiddenTransforms.Num() == 0;
 #endif
 
 #if GEOMETRYCOLLECTION_EDITOR_SELECTION
-		// If using subsections then use the subsection array. 
-		TArray<FGeometryCollectionMeshElement> const& SectionArray = bUsesSubSections 
-			? MeshDescription.SubSections 
-			: bRemoveInternalFaces ? MeshDescription.SectionsNoInternal : MeshDescription.Sections;
+			// If using subsections then use the subsection array. 
+			TArray<FGeometryCollectionMeshElement> const& SectionArray = bUsesSubSections
+				? MeshDescription.SubSections
+				: bRemoveInternalFaces ? MeshDescription.SectionsNoInternal : MeshDescription.Sections;
 #else
-		TArray<FGeometryCollectionMeshElement> const& SectionArray = bRemoveInternalFaces ? MeshDescription.SectionsNoInternal : MeshDescription.Sections;
+			TArray<FGeometryCollectionMeshElement> const& SectionArray = bRemoveInternalFaces ? MeshDescription.SectionsNoInternal : MeshDescription.Sections;
 #endif
 
-		// Grab the material proxies we'll be using for each section.
-		TArray<FMaterialRenderProxy*, TInlineAllocator<32>> MaterialProxies;
-		for (int32 SectionIndex = 0; SectionIndex < SectionArray.Num(); ++SectionIndex)
-		{
-			const FGeometryCollectionMeshElement& Section = SectionArray[SectionIndex];
-			FMaterialRenderProxy* MaterialProxy = GetMaterial(Collector, Section.MaterialIndex);
-			MaterialProxies.Add(MaterialProxy);
-		}
-
-		// Draw the meshes.
-		for (int32 SectionIndex = 0; SectionIndex < SectionArray.Num(); ++SectionIndex)
-		{
-			const FGeometryCollectionMeshElement& Section = SectionArray[SectionIndex];
-
-			FMeshBatch& Mesh = Collector.AllocateMesh();
-			Mesh.bWireframe = bWireframe;
-			Mesh.VertexFactory = GetVertexFactory();
-			Mesh.MaterialRenderProxy = MaterialProxies[SectionIndex];
-			Mesh.ReverseCulling = IsLocalToWorldDeterminantNegative();
-			Mesh.Type = PT_TriangleList;
-			Mesh.DepthPriorityGroup = SDPG_World;
-			Mesh.bCanApplyViewModeOverrides = true;
-			SetDebugMaterial(Mesh);
+			// Grab the material proxies we'll be using for each section.
+			TArray<FMaterialRenderProxy*, TInlineAllocator<32>> MaterialProxies;
+			for (int32 SectionIndex = 0; SectionIndex < SectionArray.Num(); ++SectionIndex)
+			{
+				const FGeometryCollectionMeshElement& Section = SectionArray[SectionIndex];
+				FMaterialRenderProxy* MaterialProxy = GetMaterial(Collector, Section.MaterialIndex);
+				MaterialProxies.Add(MaterialProxy);
+			}
 
-			FMeshBatchElement& BatchElement = Mesh.Elements[0];
-			BatchElement.IndexBuffer = &MeshResource.IndexBuffer;
-			BatchElement.PrimitiveUniformBuffer = GetUniformBuffer();
-			BatchElement.FirstIndex = Section.TriangleStart * 3;
-			BatchElement.NumPrimitives = Section.TriangleCount;
-			BatchElement.MinVertexIndex = Section.VertexStart;
-			BatchElement.MaxVertexIndex = Section.VertexEnd;
+			// Draw the meshes.
+			for (int32 SectionIndex = 0; SectionIndex < SectionArray.Num(); ++SectionIndex)
+			{
+				const FGeometryCollectionMeshElement& Section = SectionArray[SectionIndex];
+
+				FMeshBatch& Mesh = Collector.AllocateMesh();
+				Mesh.bWireframe = bWireframe;
+				Mesh.VertexFactory = GetVertexFactory();
+				Mesh.MaterialRenderProxy = MaterialProxies[SectionIndex];
+				Mesh.ReverseCulling = IsLocalToWorldDeterminantNegative();
+				Mesh.Type = PT_TriangleList;
+				Mesh.DepthPriorityGroup = SDPG_World;
+				Mesh.bCanApplyViewModeOverrides = true;
+				SetDebugMaterial(Mesh);
+
+				FMeshBatchElement& BatchElement = Mesh.Elements[0];
+				BatchElement.IndexBuffer = &MeshResource.IndexBuffer;
+				BatchElement.PrimitiveUniformBuffer = GetUniformBuffer();
+				BatchElement.FirstIndex = Section.TriangleStart * 3;
+				BatchElement.NumPrimitives = Section.TriangleCount;
+				BatchElement.MinVertexIndex = Section.VertexStart;
+				BatchElement.MaxVertexIndex = Section.VertexEnd;
+
+				Collector.AddMesh(ViewIndex, Mesh);
+			}
 
-			Collector.AddMesh(ViewIndex, Mesh);
+#if GEOMETRYCOLLECTION_EDITOR_SELECTION
+			// Highlight selected bone using specialized material.
+			// #note: This renders the geometry again but with the bone selection material.  Ideally we'd have one render pass and one material.
+			if (bEnableBoneSelection && !bSuppressSelectionMaterial && BoneSelectedMaterial)
+			{
+				FMaterialRenderProxy* MaterialRenderProxy = BoneSelectedMaterial->GetRenderProxy();
+
+				FMeshBatch& Mesh = Collector.AllocateMesh();
+				Mesh.bWireframe = bWireframe;
+				Mesh.VertexFactory = &VertexFactoryDebugColor;
+				Mesh.MaterialRenderProxy = MaterialRenderProxy;
+				Mesh.ReverseCulling = IsLocalToWorldDeterminantNegative();
+				Mesh.Type = PT_TriangleList;
+				Mesh.DepthPriorityGroup = SDPG_World;
+				Mesh.bCanApplyViewModeOverrides = false;
+
+				FMeshBatchElement& BatchElement = Mesh.Elements[0];
+				BatchElement.IndexBuffer = &MeshResource.IndexBuffer;
+				BatchElement.PrimitiveUniformBuffer = GetUniformBuffer();
+				BatchElement.FirstIndex = 0;
+				BatchElement.NumPrimitives = MeshDescription.NumTriangles;
+				BatchElement.MinVertexIndex = 0;
+				BatchElement.MaxVertexIndex = MeshDescription.NumVertices;
+
+				Collector.AddMesh(ViewIndex, Mesh);
+			}
+#endif // GEOMETRYCOLLECTION_EDITOR_SELECTION
 		}
+	}
 
-#if GEOMETRYCOLLECTION_EDITOR_SELECTION
-		// Highlight selected bone using specialized material.
-		// #note: This renders the geometry again but with the bone selection material.  Ideally we'd have one render pass and one material.
-		if (bEnableBoneSelection && !bSuppressSelectionMaterial && BoneSelectedMaterial)
+	// draw extra stuff ( collision , bounds ... )
+	for (int32 ViewIndex = 0; ViewIndex < Views.Num(); ViewIndex++)
+	{
+		if (VisibilityMap & (1 << ViewIndex))
 		{
-			FMaterialRenderProxy* MaterialRenderProxy = BoneSelectedMaterial->GetRenderProxy();
-
-			FMeshBatch& Mesh = Collector.AllocateMesh();
-			Mesh.bWireframe = bWireframe;
-			Mesh.VertexFactory = &VertexFactoryDebugColor;
-			Mesh.MaterialRenderProxy = MaterialRenderProxy;
-			Mesh.ReverseCulling = IsLocalToWorldDeterminantNegative();
-			Mesh.Type = PT_TriangleList;
-			Mesh.DepthPriorityGroup = SDPG_World;
-			Mesh.bCanApplyViewModeOverrides = false;
-
-			FMeshBatchElement& BatchElement = Mesh.Elements[0];
-			BatchElement.IndexBuffer = &MeshResource.IndexBuffer;
-			BatchElement.PrimitiveUniformBuffer = GetUniformBuffer();
-			BatchElement.FirstIndex = 0;
-			BatchElement.NumPrimitives = MeshDescription.NumTriangles;
-			BatchElement.MinVertexIndex = 0;
-			BatchElement.MaxVertexIndex = MeshDescription.NumVertices;
-
-			Collector.AddMesh(ViewIndex, Mesh);
-		}
-#endif // GEOMETRYCOLLECTION_EDITOR_SELECTION
+			// collision modes
+			if (ShowCollisionMeshes(EngineShowFlags) && GeometryCollection && AllowDebugViewmodes())
+			{
+				FTransform GeomTransform(GetLocalToWorld());
+				if (bDrawWireframeCollision)
+				{
+					GeometryCollectionDebugDraw::DrawWireframe(*GeometryCollection, GeomTransform, Collector, ViewIndex, GetWireframeColor().ToFColor(true));
+				}
+				else
+				{
+					FMaterialRenderProxy* CollisionMaterialInstance = new FColoredMaterialRenderProxy(GEngine->ShadedLevelColorationUnlitMaterial->GetRenderProxy(), GetWireframeColor());
+					Collector.RegisterOneFrameMaterialProxy(CollisionMaterialInstance);
+					GeometryCollectionDebugDraw::DrawSolid(*GeometryCollection, GeomTransform, Collector, ViewIndex, CollisionMaterialInstance);
+				}
+			}
 
+			// render bounds
 #if !(UE_BUILD_SHIPPING || UE_BUILD_TEST)
-		RenderBounds(Collector.GetPDI(ViewIndex), ViewFamily.EngineShowFlags, GetBounds(), IsSelected());
+			RenderBounds(Collector.GetPDI(ViewIndex), ViewFamily.EngineShowFlags, GetBounds(), IsSelected());
 #endif
+		}
 	}
 }
 
@@ -1123,6 +1183,8 @@ FNaniteGeometryCollectionSceneProxy::FNaniteGeometryCollectionSceneProxy(UGeomet
 		}
 	}
 
+	SetWireframeColor(Component->GetWireframeColorForSceneProxy());
+
 	// Initialize to rest transforms.
 	TArray<FMatrix44f> RestTransforms;
 	Component->GetRestTransforms(RestTransforms);
diff --git a/Engine/Source/Runtime/Experimental/GeometryCollectionEngine/Private/GeometryCollection/GeometryCollectionSceneProxy.h b/Engine/Source/Runtime/Experimental/GeometryCollectionEngine/Private/GeometryCollection/GeometryCollectionSceneProxy.h
--- a/Engine/Source/Runtime/Experimental/GeometryCollectionEngine/Private/GeometryCollection/GeometryCollectionSceneProxy.h
+++ b/Engine/Source/Runtime/Experimental/GeometryCollectionEngine/Private/GeometryCollection/GeometryCollectionSceneProxy.h
@@ -218,6 +218,9 @@ class FGeometryCollectionSceneProxy final : public FPrimitiveSceneProxy
 
 	int32 NumTransforms = 0;
 	TArray<FMatrix44f> RestTransforms;
+	TSharedPtr<FGeometryCollection, ESPMode::ThreadSafe> GeometryCollection;
+
+	FCollisionResponseContainer CollisionResponse;
 
 	FBoxSphereBounds PreSkinnedBounds;
 
@@ -321,6 +324,9 @@ protected:
 #if RHI_RAYTRACING
 	void UpdatingRayTracingGeometry_RenderingThread(TArray<FGeometryCollectionMeshElement> const& InSectionArray);
 #endif
+
+private:
+	bool ShowCollisionMeshes(const FEngineShowFlags& EngineShowFlags) const;
 };
 
 
diff --git a/Engine/Source/Runtime/Experimental/GeometryCollectionEngine/Public/GeometryCollection/GeometryCollectionDebugDraw.h b/Engine/Source/Runtime/Experimental/GeometryCollectionEngine/Public/GeometryCollection/GeometryCollectionDebugDraw.h
new file mode 100644
--- /dev/null
+++ b/Engine/Source/Runtime/Experimental/GeometryCollectionEngine/Public/GeometryCollection/GeometryCollectionDebugDraw.h
@@ -0,0 +1,16 @@
+// Copyright Epic Games, Inc. All Rights Reserved.
+
+#pragma once
+
+#include "Math/MathFwd.h"
+
+class FMeshElementCollector;
+class FMaterialRenderProxy;
+class FGeometryCollection;
+struct FColor;
+
+namespace GeometryCollectionDebugDraw
+{
+	void DrawSolid(const FGeometryCollection& Collection, const FTransform& CollectionWorldTransform, FMeshElementCollector& MeshCollector, int32 ViewIndex, const FMaterialRenderProxy* MaterialProxy);
+	void DrawWireframe(const FGeometryCollection& Collection, const FTransform& CollectionWorldTransform, FMeshElementCollector& MeshCollector, int32 ViewIndex, const FColor& Color);
+};
