Fix compilation error when UE_NET_ENABLE_REPLICATIONREADER_LOG is enabled.

diff --git a/Engine/Source/Runtime/Experimental/Iris/Core/Private/Iris/ReplicationSystem/ReplicationReader.cpp b/Engine/Source/Runtime/Experimental/Iris/Core/Private/Iris/ReplicationSystem/ReplicationReader.cpp
--- a/Engine/Source/Runtime/Experimental/Iris/Core/Private/Iris/ReplicationSystem/ReplicationReader.cpp
+++ b/Engine/Source/Runtime/Experimental/Iris/Core/Private/Iris/ReplicationSystem/ReplicationReader.cpp
@@ -2086,7 +2086,7 @@ void FReplicationReader::ResolveAndDispatchUnresolvedReferences()
 
 	if (NumHandlesPendingResolveLastUpdate != VisitedUnresolvedHandles.Num() || ObjectsWithAttachmentPendingResolve.Num() > 0)
 	{
-		UE_LOG_REPLICATIONREADER(TEXT("FReplicationReader::ResolveAndDispatchUnresolvedReferences NetHandles pending: %u Attachments pending: %u)"), UpdatedHandles.Num(), ObjectsWithAttachmentPendingResolve.Num());
+		UE_LOG_REPLICATIONREADER(TEXT("FReplicationReader::ResolveAndDispatchUnresolvedReferences NetHandles pending: %u Attachments pending: %u)"), VisitedUnresolvedHandles.Num(), ObjectsWithAttachmentPendingResolve.Num());
 		NumHandlesPendingResolveLastUpdate = VisitedUnresolvedHandles.Num();
 	}
 }
