IOStore command parsing
- Fix crash on empty optional parameter

diff --git a/Engine/Source/Runtime/Experimental/IoStore/OnDemand/Private/Tool/Command.cpp b/Engine/Source/Runtime/Experimental/IoStore/OnDemand/Private/Tool/Command.cpp
--- a/Engine/Source/Runtime/Experimental/IoStore/OnDemand/Private/Tool/Command.cpp
+++ b/Engine/Source/Runtime/Experimental/IoStore/OnDemand/Private/Tool/Command.cpp
@@ -200,7 +200,7 @@ FContext::FContext(const FArguments& InArguments, int32 ArgC, const TCHAR* const
 	{
 		if (int32 i = PendingOptional; i >= 0)
 		{
-			if (Arg[0] == '-')
+			if (Arg.IsEmpty() || Arg[0] == '-')
 			{
 				FArgument Argument = Arguments[i];
 				FatalError<Error::FMissingValue>(Argument.Name);
@@ -270,7 +270,7 @@ FContext::FContext(const FArguments& InArguments, int32 ArgC, const TCHAR* const
 			break;
 		}
 
-		(Arg[0] == '-' || PendingOptional >= 0)
+		(PendingOptional >= 0 || Arg[0] == '-')
 			? HandleOptional(Arg)
 			: HandlePositional(Arg);
 	}
