IASTool
- Add -OutputStatsJson option, currently just includes total chunk size

diff --git a/Engine/Source/Runtime/Experimental/IoStore/OnDemand/Private/Tool/ChunkPlugin.cpp b/Engine/Source/Runtime/Experimental/IoStore/OnDemand/Private/Tool/ChunkPlugin.cpp
--- a/Engine/Source/Runtime/Experimental/IoStore/OnDemand/Private/Tool/ChunkPlugin.cpp
+++ b/Engine/Source/Runtime/Experimental/IoStore/OnDemand/Private/Tool/ChunkPlugin.cpp
@@ -20,6 +20,8 @@
 #include "S3/S3Client.h"
 #include "Serialization/JsonReader.h"
 #include "Serialization/JsonSerializer.h"
+#include "Serialization/JsonSerializerMacros.h"
+#include "Serialization/JsonWriter.h"
 #include "Serialization/LargeMemoryWriter.h"
 #include "Serialization/MemoryReader.h"
 #include "UploadQueue.h"
@@ -27,7 +29,14 @@
 namespace UE::IoStore::Tool
 {
 
+struct FChunkPluginStats : FJsonSerializable
+{
+	double TotalChunksSizeKb = -1.0;
 
+BEGIN_JSON_SERIALIZER
+	JSON_SERIALIZE("TotalChunksSizeKb", TotalChunksSizeKb);
+END_JSON_SERIALIZER
+};
 
 struct FS3Params
 {
@@ -234,9 +243,10 @@ static int32 ChunkPluginCommandEntry(const FContext& Context)
 	const FString BuildVersion			= FString(Context.Get<FStringView>(TEXT("-BuildVersion"), FString()));
 	const FString OnDemandTocName		= FString(Context.Get<FStringView>(TEXT("-OnDemandTocName"), FString()));
 	const FString InputFolder			= FString(Context.Get<FStringView>(TEXT("-InputFolder"), FString()));
-	const FString InOutputFolder			= FString(Context.Get<FStringView>(TEXT("-OutputFolder"), FString()));
+	const FString InOutputFolder		= FString(Context.Get<FStringView>(TEXT("-OutputFolder"), FString()));
 	const FString IntermediateFolder	= FString(Context.Get<FStringView>(TEXT("-IntermediateFolder"), FString()));
 	FString SettingsFile				= FString(Context.Get<FStringView>(TEXT("-SettingsFile"), FString()));
+	FString OutputStatsJson				= FString(Context.Get<FStringView>(TEXT("-OutputStatsJson"), FString()));
 	
 	FS3Params S3Params(Context);
 
@@ -250,6 +260,7 @@ static int32 ChunkPluginCommandEntry(const FContext& Context)
 	FPaths::NormalizeDirectoryName(ContainerFolder);
 	FPaths::NormalizeDirectoryName(OutputFolder);
 	FPaths::NormalizeFilename(SettingsFile);
+	FPaths::NormalizeFilename(OutputStatsJson);
 
 	UE_LOG(LogIoStore, Display, TEXT("I/O store chunk plugin:"));
 	UE_LOG(LogIoStore, Display, TEXT("----------------------------------------"));
@@ -260,6 +271,7 @@ static int32 ChunkPluginCommandEntry(const FContext& Context)
 	UE_LOG(LogIoStore, Display, TEXT("\tOutputFolder: %s"), *OutputFolder);
 	UE_LOG(LogIoStore, Display, TEXT("\tIntermediateFolder: %s"), *IntermediateFolder);
 	UE_LOG(LogIoStore, Display, TEXT("\tSettingsFile: %s"), *SettingsFile);
+	UE_LOG(LogIoStore, Display, TEXT("\tOutputStatsJson: %s"), *OutputStatsJson);
 	UE_LOG(LogIoStore, Display, TEXT("\tIncludeSigPak: %s"), bIncludeSigPak ? TEXT("true") : TEXT("false"));
 	UE_LOG(LogIoStore, Display, TEXT("\tDeleteContainerFiles: %s"), bDeleteContainerFiles ? TEXT("true") : TEXT("false"));
 
@@ -332,6 +344,8 @@ static int32 ChunkPluginCommandEntry(const FContext& Context)
 		}
 	}
 
+	FChunkPluginStats Stats;
+
 	TMap<FGuid, FAES::FAESKey> EncryptionKeys;
 	if (FileMgr.DirectoryExists(*ContainerFolder) == false)
 	{
@@ -502,8 +516,6 @@ static int32 ChunkPluginCommandEntry(const FContext& Context)
 				return -1;
 			}
 
-
-
 			FIoStatus WriteStatus = ChunkWriter->WriteChunk(ChunksRelativeFolder, ReadResult.IoBuffer, ChunkHash);
 			if (WriteStatus.IsOk() == false)
 			{
@@ -518,6 +530,8 @@ static int32 ChunkPluginCommandEntry(const FContext& Context)
 			TocEntry.EncodedSize = EncodedChunkSize;
 			TocEntry.BlockOffset = BlockOffset;
 			TocEntry.BlockCount = BlockCount;
+
+			Stats.TotalChunksSizeKb += (double)EncodedChunkSize / 1024.0;
 		}
 
 		if (bDeleteContainerFiles)
@@ -535,14 +549,13 @@ static int32 ChunkPluginCommandEntry(const FContext& Context)
 		TocTagSet.Packages = MoveTemp(Pair.Value);
 	}
 
-	IFileManager& FileMan = IFileManager::Get();
 	for (const FString& Path : FilesToDelete)
 	{
 		//UE_LOG(LogIas, Display, TEXT("Attempt Deleting '%s'"), *Path);
-		if (FileMan.FileExists(*Path))
+		if (FileMgr.FileExists(*Path))
 		{
 			UE_LOG(LogIas, Display, TEXT("Deleting '%s'"), *Path);
-			if (!FileMan.Delete(*Path, /*RequireExists*/true))
+			if (!FileMgr.Delete(*Path, /*RequireExists*/true))
 			{
 				UE_LOG(LogIas, Error, TEXT("Failed to delete '%s'"), *Path);
 			}
@@ -616,7 +629,7 @@ static int32 ChunkPluginCommandEntry(const FContext& Context)
 		// todo, we actually want this file in the base directory
 		// const FString TocPath = OutputFolder / Filename; // like this yo
 		const FString TocPath = OutputFolder / IoStoreRelativeFolder / Filename;
-		if (TUniquePtr<FArchive> Ar(IFileManager::Get().CreateFileWriter(*TocPath)); Ar.IsValid())
+		if (TUniquePtr<FArchive> Ar(FileMgr.CreateFileWriter(*TocPath)); Ar.IsValid())
 		{
 			*Ar << OnDemandToc;
 			Ar->Close();
@@ -688,9 +701,20 @@ static int32 ChunkPluginCommandEntry(const FContext& Context)
 	if (!ChunkWriter->Flush())
 	{
 		UE_LOG(LogIoStore, Error, TEXT("Writer error: Failed to upload chunk(s)"));
-		return 1; 
+		return -1; 
 	}
 
+	if (!OutputStatsJson.IsEmpty())
+	{
+		TUniquePtr<FArchive> FileWriter(FileMgr.CreateFileWriter(*OutputStatsJson));
+		if (!FileWriter)
+		{
+			UE_LOG(LogIoStore, Display, TEXT("Failed writing stats file '%s'"), *OutputStatsJson);
+			return -1;
+		}
+		TSharedRef<TJsonWriter<UTF8CHAR>> JsonWriter = TJsonWriterFactory<UTF8CHAR>::Create(FileWriter.Get());
+		Stats.ToJson(JsonWriter, false);
+	}
 
 	return 0;
 }
@@ -710,7 +734,7 @@ static FCommand ChunkPluginCommand(
 		TArgument<FStringView>(TEXT("-OutputFolder"),		TEXT("Ouptut folder.")),
 		TArgument<FStringView>(TEXT("-IntermediateFolder"),	TEXT("Intermediate folder.")),
 		TArgument<FStringView>(TEXT("-SettingsFile"),		TEXT("Optional settings file.")),
-		TArgument<FStringView>(TEXT("-ErrorOutput"),		TEXT("Error output.")),
+		TArgument<FStringView>(TEXT("-OutputStatsJson"),	TEXT("Path to write a json file with statistics.")),
 		TArgument<bool>(TEXT("-IncludeSigPak"),				TEXT("Include .sig and .pak file in the uondemandtoc")),
 		TArgument<bool>(TEXT("-KeepContainerFiles"),		TEXT("Should we keep the container files after processing them.")),
 		TArgument<FStringView>(TEXT("-BucketPrefix"),		TEXT("Path to prefix to bucket objects")),
