Chaos - fix out of date comment and misleading in FGenericParticlePairMidPhase::FindOrCreateConstraint
- invert key check logic to match other check

diff --git a/Engine/Source/Runtime/Experimental/Chaos/Private/Chaos/Collision/ParticlePairMidPhase.cpp b/Engine/Source/Runtime/Experimental/Chaos/Private/Chaos/Collision/ParticlePairMidPhase.cpp
--- a/Engine/Source/Runtime/Experimental/Chaos/Private/Chaos/Collision/ParticlePairMidPhase.cpp
+++ b/Engine/Source/Runtime/Experimental/Chaos/Private/Chaos/Collision/ParticlePairMidPhase.cpp
@@ -1703,7 +1703,7 @@ namespace Chaos
 		// would require up to two Sphere-Box contacts, with the particles in opposite orders.
 #if !UE_BUILD_TEST && !UE_BUILD_SHIPPING
 		const bool bIsCorrectParticles = ((InParticle0 == Particle0) && (InParticle1 == Particle1)) || ((InParticle0 == Particle1) && (InParticle1 == Particle0));
-		if (!ensureMsgf(bIsCorrectParticles, TEXT("Attempt to us MidPhase for particles %d - %d with particles %d - %d"), Particle0->ParticleID().LocalID, Particle1->ParticleID().LocalID, InParticle0->ParticleID().LocalID, InParticle1->ParticleID().LocalID))
+		if (!ensureMsgf(bIsCorrectParticles, TEXT("Attempt to use MidPhase for particles %d - %d with particles %d - %d"), Particle0->ParticleID().LocalID, Particle1->ParticleID().LocalID, InParticle0->ParticleID().LocalID, InParticle1->ParticleID().LocalID))
 		{
 			// We somehow received a callback for the wrong particle pair...this should not happen
 			return nullptr;
@@ -1713,17 +1713,15 @@ namespace Chaos
 		const FParticlePairMidPhaseCollisionKey CollisionKey = FParticlePairMidPhaseCollisionKey(InImplicitId0, InImplicitId1);
 		FPBDCollisionConstraint* Constraint = FindConstraint(CollisionKey);
 
-		// @todo(chaos): fix key uniqueness guarantee.  We need a truly unique key gen function
+		// Check that we have the collision constraint for the expected pair of implicit objects. The key is assigned by the 
+		// midphase based on the indices of the geometry in each particle and should be unique as long as those indices do not change.
 #if !UE_BUILD_TEST && !UE_BUILD_SHIPPING
-		const bool bIsKeyCollision = (Constraint != nullptr) && ((Constraint->GetImplicit0() != InImplicit0) || (Constraint->GetImplicit1() != InImplicit1) || (Constraint->GetCollisionParticles0() != InBVHParticles0) || (Constraint->GetCollisionParticles1() != InBVHParticles1));
-		if (bIsKeyCollision)
-		{
-			// If we get here, we have a key collision. The key uses a hash of pointers which is very likely to be unique for different implicit pairs, 
-			// especially since it only needs to be unique for this particle pair, but it is not guaranteed.
-			// Creating a new constraint with the same key could cause fatal problems (the original constraint will be deleted when we add the new one 
-			// to the map, but if it is asleep it will be referenced in the contact graph) so we just abort and accept we will miss collisions. 
-			// It is extremely unlikely to happen but we should fix it at some point.
-			ensure(false);
+		const bool bIsCorrectImplicits = (Constraint == nullptr) || ((Constraint->GetImplicit0() == InImplicit0) && (Constraint->GetImplicit1() == InImplicit1) && (Constraint->GetCollisionParticles0() == InBVHParticles0) && (Constraint->GetCollisionParticles1() == InBVHParticles1));
+		if (!ensureMsgf(bIsCorrectImplicits, TEXT("Collision Constraint for geometry pair was created for a different pair. Collisions will be dropped.")))
+		{
+			// If we get here, the collision constraint is not for this pair of implicit objects, which is a bug.
+			// This might be because we added/removed collision geometry to a particle without resetting the collisions. 
+			// The fix is to call Evolution->InvalidateParticle() after a geometry change.
 			return nullptr;
 		}
 #endif
