add simple empty light cache interface to nanite geometry collection scene proxy, since we never have baked light maps for them

diff --git a/Engine/Source/Runtime/Experimental/GeometryCollectionEngine/Private/GeometryCollection/GeometryCollectionSceneProxy.cpp b/Engine/Source/Runtime/Experimental/GeometryCollectionEngine/Private/GeometryCollection/GeometryCollectionSceneProxy.cpp
--- a/Engine/Source/Runtime/Experimental/GeometryCollectionEngine/Private/GeometryCollection/GeometryCollectionSceneProxy.cpp
+++ b/Engine/Source/Runtime/Experimental/GeometryCollectionEngine/Private/GeometryCollection/GeometryCollectionSceneProxy.cpp
@@ -1482,3 +1482,21 @@ void FGeometryCollectionTransformBuffer::UpdateDynamicData(FRHICommandListBase&
 	FMemory::Memcpy(VertexBufferData, Transforms.GetData(), Transforms.Num() * sizeof(FMatrix44f));
 	RHICmdList.UnlockBuffer(VertexBufferRHI);
 }
+
+FNaniteGeometryCollectionSceneProxy::FEmptyLightCacheInfo FNaniteGeometryCollectionSceneProxy::EmptyLightCacheInfo;
+
+FLightInteraction FNaniteGeometryCollectionSceneProxy::FEmptyLightCacheInfo::GetInteraction(const FLightSceneProxy* LightSceneProxy) const
+{
+	// Ask base class
+	TArray<FGuid> Empty_IrrelevantLights;
+	ELightInteractionType LightInteraction = GetStaticInteraction(LightSceneProxy, Empty_IrrelevantLights);
+
+	if (LightInteraction != LIT_MAX)
+	{
+		return FLightInteraction(LightInteraction);
+	}
+
+	// Use dynamic lighting if the light doesn't have static lighting.
+	return FLightInteraction::Dynamic();
+}
+
diff --git a/Engine/Source/Runtime/Experimental/GeometryCollectionEngine/Private/GeometryCollection/GeometryCollectionSceneProxy.h b/Engine/Source/Runtime/Experimental/GeometryCollectionEngine/Private/GeometryCollection/GeometryCollectionSceneProxy.h
--- a/Engine/Source/Runtime/Experimental/GeometryCollectionEngine/Private/GeometryCollection/GeometryCollectionSceneProxy.h
+++ b/Engine/Source/Runtime/Experimental/GeometryCollectionEngine/Private/GeometryCollection/GeometryCollectionSceneProxy.h
@@ -377,6 +377,12 @@ public:
 	void OnMotionBegin();
 	void OnMotionEnd();
 
+	inline virtual void GetLCIs(FLCIArray& LCIs) override
+	{
+		FLightCacheInterface* LCI = &EmptyLightCacheInfo;
+		LCIs.Add(LCI);
+	}
+
 protected:
 	// TODO : Copy required data from UObject instead of using unsafe object pointer.
 	const UGeometryCollection* GeometryCollection = nullptr;
@@ -402,4 +408,16 @@ protected:
 	uint32 bRequiresGPUSceneUpdate : 1;
 
 	FInstanceSceneDataBuffers InstanceSceneDataBuffersImpl;
+
+	// Geometry collection doesn't currently support baked light maps, so we use this simple empty light cache info for all nanite geometry collection proxies
+	class FEmptyLightCacheInfo : public FLightCacheInterface
+	{
+	public:
+
+		// FLightCacheInterface.
+		GEOMETRYCOLLECTIONENGINE_API virtual FLightInteraction GetInteraction(const FLightSceneProxy* LightSceneProxy) const override;
+	};
+
+private:
+	static FEmptyLightCacheInfo EmptyLightCacheInfo;
 };
