Make it easier to see when the IAS cache is disabled.

#rnx

- Previously when the cache was disabled we'd continue to write out and display stats for it but set at zero, so it looked like the cache was present but nothing was in it and nothing was written into it which could be confusing.
- With this change we now log/write out a specific "cache disabled" message where the stats would normally be to make it more obvious. The user can then check earlier in the log to see why the cache is disabled.

diff --git a/Engine/Source/Runtime/Experimental/IoStore/OnDemand/Private/OnDemandIoDispatcherBackend.cpp b/Engine/Source/Runtime/Experimental/IoStore/OnDemand/Private/OnDemandIoDispatcherBackend.cpp
--- a/Engine/Source/Runtime/Experimental/IoStore/OnDemand/Private/OnDemandIoDispatcherBackend.cpp
+++ b/Engine/Source/Runtime/Experimental/IoStore/OnDemand/Private/OnDemandIoDispatcherBackend.cpp
@@ -1097,6 +1097,7 @@ private:
 FOnDemandIoBackend::FOnDemandIoBackend(const FOnDemandEndpointConfig& Config, FOnDemandIoStore& InIoStore, TUniquePtr<IIasCache>&& InCache)
 	: IoStore(InIoStore)
 	, Cache(MoveTemp(InCache))
+	, Stats(Cache.IsValid() ? EStatsFlags::None : EStatsFlags::CachingDisabled)
 	, CurrentThreadPriority(EThreadPriority::TPri_Num)
 {
 	EndpointTestPath = Config.TocPath;
diff --git a/Engine/Source/Runtime/Experimental/IoStore/OnDemand/Private/Statistics.cpp b/Engine/Source/Runtime/Experimental/IoStore/OnDemand/Private/Statistics.cpp
--- a/Engine/Source/Runtime/Experimental/IoStore/OnDemand/Private/Statistics.cpp
+++ b/Engine/Source/Runtime/Experimental/IoStore/OnDemand/Private/Statistics.cpp
@@ -341,7 +341,8 @@ static FOnDemandIoBackendStats* GStatistics = nullptr;
 static FDelegateHandle GStatisticsEndFrameDelegateHandle;
 static FDelegateHandle GStatisticsOnScreenDelegateHandle;
 
-FOnDemandIoBackendStats::FOnDemandIoBackendStats()
+FOnDemandIoBackendStats::FOnDemandIoBackendStats(EStatsFlags InFlags)
+	: Flags(InFlags)
 {
 	check(GStatistics == nullptr);
 	GStatistics = this;
@@ -391,8 +392,16 @@ FOnDemandIoBackendStats::FOnDemandIoBackendStats()
 			static double LastLogTime = 0.0;
 			if (double Time = FPlatformTime::Seconds(); Time - LastLogTime > (double)GIasStatisticsLogInterval)
 			{
-				UE_LOG(LogIas, Log, TEXT("CacheStats: CachedKiB=%d, WrittenKiB=%d, ReadKiB=%d, RejectedKiB=%d, Get=%d, Error=%d, Put=%d, PutReject=%d, PutExisting=%d"),
-					CCachedKiB, CWrittenKiB, CReadKiB, CRejectedKiB, CGetCount, CErrorCount, CPutCount, CPutRejectCount, CPutExistingCount);
+				if (!EnumHasAllFlags(Flags,EStatsFlags::CachingDisabled))
+				{
+					UE_LOG(LogIas, Log, TEXT("CacheStats: CachedKiB=%d, WrittenKiB=%d, ReadKiB=%d, RejectedKiB=%d, Get=%d, Error=%d, Put=%d, PutReject=%d, PutExisting=%d"),
+						CCachedKiB, CWrittenKiB, CReadKiB, CRejectedKiB, CGetCount, CErrorCount, CPutCount, CPutRejectCount, CPutExistingCount);
+				}
+				else
+				{
+					UE_LOG(LogIas, Log, TEXT("CacheStats: Disabled"));
+				}
+
 				UE_LOG(LogIas, Log, TEXT("HttpStats: DownloadedKiB=%d, Get=%d, Retry=%d, Cancel=%d, Error=%d, CurPending=%d, CurDurationMsAvg=%d, CurDurationMsMax=%d"),
 					HDownloadedKiB, HGetCount, HRetryCount, HCancelCount, HErrorCount, HPendingCount, HDurationMsAvg, HDurationMsMax);
 				LastLogTime = Time;
@@ -407,18 +416,33 @@ FOnDemandIoBackendStats::FOnDemandIoBackendStats()
 			if (GIasDisplayOnScreenStatistics)
 			{
 				const bool bIsConnected = GHttpConnectCount.Get() > GHttpDisconnectCount.Get();
+
+				FText CachingText;
+				if (!EnumHasAllFlags(Flags, EStatsFlags::CachingDisabled))
+				{
+					CachingText = FText::Format(
+						LOCTEXT("IASCacheEnabled", "Cached:{0} KiB | Read:{1} KiB ({2})"),
+						GCacheCachedBytes.Get() >> 10,
+						GCacheReadBytes.Get() >> 10,
+						GCacheGetCount.Get()
+					);
+				}
+				else
+				{
+					CachingText = LOCTEXT("IASCacheDisabled", "Caching Disabled");
+				}
+
 				FText Message = FText::Format(
-					LOCTEXT("IAS", "IAS - {0}: Cached:{1} KiB | Read:{2} KiB ({3}) | Downloaded:{4} KiB ({5}) {6} ms | Retries:{7} | Pending:{8}"),
+					LOCTEXT("IAS", "IAS - {0}: {1} | Downloaded:{2} KiB ({3}) {4} ms | Retries:{5} | Pending:{6}"),
 					bIsConnected ? LOCTEXT("IASConnect", "Connected") : LOCTEXT("IASDisconnect", "Disconnected"),
-					GCacheCachedBytes.Get() >> 10,
-					GCacheReadBytes.Get() >> 10,
-					GCacheGetCount.Get(),
+					CachingText,
 					GHttpDownloadedBytes.Get() >> 10,
 					GHttpGetCount.Get(),
 					GHttpDurationMsAvg.Get(),
 					GHttpRetryCount.Get(),
 					GHttpPendingCount.Get()
 				);
+
 				Out.Add(FCoreDelegates::EOnScreenMessageSeverity::Info, MoveTemp(Message));
 
 				if (GIoDecodeErrorCount.Get() > 0 || GHttpErrorCount.Get() > 0)
diff --git a/Engine/Source/Runtime/Experimental/IoStore/OnDemand/Private/Statistics.h b/Engine/Source/Runtime/Experimental/IoStore/OnDemand/Private/Statistics.h
--- a/Engine/Source/Runtime/Experimental/IoStore/OnDemand/Private/Statistics.h
+++ b/Engine/Source/Runtime/Experimental/IoStore/OnDemand/Private/Statistics.h
@@ -25,12 +25,20 @@ namespace UE::IoStore
 #	define IAS_STATISTICS_IMPL(...) { return __VA_ARGS__; }
 #endif
 
+enum class EStatsFlags
+{
+	None				= 0,
+	CachingDisabled		= 1 << 0
+};
+
+ENUM_CLASS_FLAGS(EStatsFlags);
+
 class FOnDemandIoBackendStats
 {
 public:
 	static FOnDemandIoBackendStats* Get() IAS_STATISTICS_IMPL(nullptr)
 
-	FOnDemandIoBackendStats() IAS_STATISTICS_IMPL()
+	FOnDemandIoBackendStats(EStatsFlags Flags) IAS_STATISTICS_IMPL()
 	~FOnDemandIoBackendStats() IAS_STATISTICS_IMPL()
 
 	/** Report analytics not directly associated with a specific endpoint */
@@ -64,6 +72,12 @@ public:
 	void OnHttpGet(uint64 SizeBytes, uint64 DurationMs) IAS_STATISTICS_IMPL()
 	void OnHttpRetry() IAS_STATISTICS_IMPL()
 	void OnHttpError() IAS_STATISTICS_IMPL()
+
+private:
+
+#if IAS_WITH_STATISTICS
+	EStatsFlags Flags;
+#endif //IAS_WITH_STATISTICS
 };
 
 #undef IAS_STATISTICS_IMPL
