Add support for per-bone hit proxies to geometry collection nanite scene proxies.
#tests Use fracture mode with nanite geometry collection to confirm that selection works.

diff --git a/Engine/Source/Runtime/Experimental/GeometryCollectionEngine/Private/GeometryCollection/GeometryCollectionSceneProxy.cpp b/Engine/Source/Runtime/Experimental/GeometryCollectionEngine/Private/GeometryCollection/GeometryCollectionSceneProxy.cpp
--- a/Engine/Source/Runtime/Experimental/GeometryCollectionEngine/Private/GeometryCollection/GeometryCollectionSceneProxy.cpp
+++ b/Engine/Source/Runtime/Experimental/GeometryCollectionEngine/Private/GeometryCollection/GeometryCollectionSceneProxy.cpp
@@ -1040,6 +1040,7 @@ FNaniteGeometryCollectionSceneProxy::FNaniteGeometryCollectionSceneProxy(UGeomet
 , GeometryCollection(Component->GetRestCollection())
 , bCurrentlyInMotion(false)
 , bRequiresGPUSceneUpdate(false)
+, bEnableBoneSelection(false)
 {
 	LLM_SCOPE_BYTAG(Nanite);
 
@@ -1050,11 +1051,16 @@ FNaniteGeometryCollectionSceneProxy::FNaniteGeometryCollectionSceneProxy(UGeomet
 
 	MaterialRelevance = Component->GetMaterialRelevance(Component->GetScene()->GetFeatureLevel());
 
+#if GEOMETRYCOLLECTION_EDITOR_SELECTION
+	bEnableBoneSelection = Component->GetEnableBoneSelection();
+#endif
+
 	FInstanceSceneDataBuffers::FAccessTag AccessTag(PointerHash(this));
 	FInstanceSceneDataBuffers::FWriteView ProxyData = InstanceSceneDataBuffersImpl.BeginWriteAccess(AccessTag);
 	ProxyData.Flags.bHasPerInstanceHierarchyOffset = true;
 	ProxyData.Flags.bHasPerInstanceLocalBounds = true;
 	ProxyData.Flags.bHasPerInstanceDynamicData = true;
+	ProxyData.Flags.bHasPerInstanceEditorData = bEnableBoneSelection;
 	InstanceSceneDataBuffersImpl.EndWriteAccess(AccessTag);
 
 	// Note: ideally this would be picked up from the Flags.bHasPerInstanceDynamicData above, but that path is not great at the moment.
@@ -1185,6 +1191,32 @@ FNaniteGeometryCollectionSceneProxy::FNaniteGeometryCollectionSceneProxy(UGeomet
 
 	SetWireframeColor(Component->GetWireframeColorForSceneProxy());
 
+#if GEOMETRYCOLLECTION_EDITOR_SELECTION
+	if (bEnableBoneSelection)
+	{
+		// Generate a hit proxy per geometry section so that we can perform per bone hit tests.
+		HitProxyMode = EHitProxyMode::PerInstance;
+		HitProxies.Reserve(NumGeometry);
+		for (int32 GeometryIndex = 0; GeometryIndex < NumGeometry; ++GeometryIndex)
+		{
+			HGeometryCollection* HitProxy = new HGeometryCollection(Component, GeometryIndex);
+			HitProxies.Add(HitProxy);
+		}
+	}
+	else if (AActor* Actor = Component->GetOwner())
+	{
+		// Generate default material hit proxies for simple selection.
+		HitProxyMode = Nanite::FSceneProxyBase::EHitProxyMode::MaterialSection;
+		for (int32 SectionIndex = 0; SectionIndex < MaterialSections.Num(); ++SectionIndex)
+		{
+			FMaterialSection& Section = MaterialSections[SectionIndex];
+			HHitProxy* HitProxy = new HActor(Actor, Component, SectionIndex, SectionIndex);
+			Section.HitProxy = HitProxy;
+			HitProxies.Add(HitProxy);
+		}
+	}
+#endif
+
 	// Initialize to rest transforms.
 	TArray<FMatrix44f> RestTransforms;
 	Component->GetRestTransforms(RestTransforms);
@@ -1243,20 +1275,7 @@ FPrimitiveViewRelevance FNaniteGeometryCollectionSceneProxy::GetViewRelevance(co
 HHitProxy* FNaniteGeometryCollectionSceneProxy::CreateHitProxies(UPrimitiveComponent* Component, TArray<TRefCountPtr<HHitProxy>>& OutHitProxies)
 {
 	LLM_SCOPE_BYTAG(Nanite);
-
-	if (Component->GetOwner())
-	{
-		// Generate separate hit proxies for each material section, so that we can perform hit tests against each one.
-		for (int32 SectionIndex = 0; SectionIndex < MaterialSections.Num(); ++SectionIndex)
-		{
-			FMaterialSection& Section = MaterialSections[SectionIndex];
-			HHitProxy* ActorHitProxy = new HActor(Component->GetOwner(), Component, SectionIndex, SectionIndex);
-			check(!Section.HitProxy);
-			Section.HitProxy = ActorHitProxy;
-			OutHitProxies.Add(ActorHitProxy);
-		}
-	}
-
+	OutHitProxies.Append(HitProxies);
 	return Super::CreateHitProxies(Component, OutHitProxies);
 }
 #endif
@@ -1328,6 +1347,10 @@ void FNaniteGeometryCollectionSceneProxy::SetDynamicData_RenderThread(FGeometryC
 		ProxyData.PrevInstanceToPrimitiveRelative.Reset(TransformCount);
 		ProxyData.InstanceLocalBounds.Reset(TransformCount);
 		ProxyData.InstanceHierarchyOffset.Reset(TransformCount);
+		
+#if GEOMETRYCOLLECTION_EDITOR_SELECTION
+		ProxyData.InstanceEditorData.Reset(bEnableBoneSelection ? TransformCount : 0);
+#endif
 
 		ProxyData.Flags.bHasPerInstanceDynamicData = true;
 		ProxyData.Flags.bHasPerInstanceLocalBounds = true;
@@ -1358,6 +1381,13 @@ void FNaniteGeometryCollectionSceneProxy::SetDynamicData_RenderThread(FGeometryC
 
 			ProxyData.InstanceLocalBounds.Emplace(PadInstanceLocalBounds(NaniteData.LocalBounds));
 			ProxyData.InstanceHierarchyOffset.Emplace(NaniteData.HierarchyOffset);
+			
+#if GEOMETRYCOLLECTION_EDITOR_SELECTION
+			if (bEnableBoneSelection)
+			{
+				ProxyData.InstanceEditorData.Emplace(FInstanceEditorData::Pack(HitProxies[TransformToGeometryIndex]->Id.GetColor(), false));
+			}
+#endif
 		}
 		InstanceSceneDataBuffersImpl.EndWriteAccess(AccessTag);
 	}
diff --git a/Engine/Source/Runtime/Experimental/GeometryCollectionEngine/Private/GeometryCollection/GeometryCollectionSceneProxy.h b/Engine/Source/Runtime/Experimental/GeometryCollectionEngine/Private/GeometryCollection/GeometryCollectionSceneProxy.h
--- a/Engine/Source/Runtime/Experimental/GeometryCollectionEngine/Private/GeometryCollection/GeometryCollectionSceneProxy.h
+++ b/Engine/Source/Runtime/Experimental/GeometryCollectionEngine/Private/GeometryCollection/GeometryCollectionSceneProxy.h
@@ -343,7 +343,7 @@ public:
 	virtual void CreateRenderThreadResources(FRHICommandListBase& RHICmdList) override;
 	virtual SIZE_T GetTypeHash() const override;
 	virtual FPrimitiveViewRelevance	GetViewRelevance(const FSceneView* View) const override;
-#if WITH_EDITOR
+#if GEOMETRYCOLLECTION_EDITOR_SELECTION
 	virtual HHitProxy* CreateHitProxies(UPrimitiveComponent* Component, TArray<TRefCountPtr<HHitProxy> >& OutHitProxies) override;
 #endif
 	virtual void DrawStaticElements(FStaticPrimitiveDrawInterface* PDI) override;
@@ -406,6 +406,11 @@ protected:
 	uint32 bHasMaterialErrors : 1;
 	uint32 bCurrentlyInMotion : 1;
 	uint32 bRequiresGPUSceneUpdate : 1;
+	uint32 bEnableBoneSelection : 1;
+	
+#if GEOMETRYCOLLECTION_EDITOR_SELECTION
+	TArray<TRefCountPtr<HHitProxy>> HitProxies;
+#endif
 
 	FInstanceSceneDataBuffers InstanceSceneDataBuffersImpl;
 
