Chaos - fix Generic ParticlePair MidPhase collision key generation

The order of the two particles in a MidPhase do not necessarily match the order in its collision constraints. In FGenericParticlePairMidPhase::FindOrCreateConstraint we did not this take into account when generating the key for the collision. So it was possible for two collisions to generate the same key.

This is a cpp-only fix. We can do a cleaner version later...

diff --git a/Engine/Source/Runtime/Experimental/Chaos/Private/Chaos/Collision/ParticlePairMidPhase.cpp b/Engine/Source/Runtime/Experimental/Chaos/Private/Chaos/Collision/ParticlePairMidPhase.cpp
--- a/Engine/Source/Runtime/Experimental/Chaos/Private/Chaos/Collision/ParticlePairMidPhase.cpp
+++ b/Engine/Source/Runtime/Experimental/Chaos/Private/Chaos/Collision/ParticlePairMidPhase.cpp
@@ -1710,7 +1710,13 @@ namespace Chaos
 		}
 #endif
 
-		const FParticlePairMidPhaseCollisionKey CollisionKey = FParticlePairMidPhaseCollisionKey(InImplicitId0, InImplicitId1);
+		// @todo(chaos): bParticlesInExpectedOrder is a temporary cpp-only fix for a key collision in FindConstraint.
+		// When we have swapped the order of the particles, we need to make sure that the key uses the shapes
+		// in the order that the particles are in the MidPhase. We can fix this more cleanly by passing the 
+		// FParticlePairMidPhaseCollisionKey in from GenerateCollisionsImplicitLeafImplicitLeaf where we already
+		// know if we swapped particle order or not.
+		const bool bParticlesInExpectedOrder = (InParticle0 == Particle0) && (InParticle1 == Particle1);
+		const FParticlePairMidPhaseCollisionKey CollisionKey = bParticlesInExpectedOrder ? FParticlePairMidPhaseCollisionKey(InImplicitId0, InImplicitId1) : FParticlePairMidPhaseCollisionKey(InImplicitId1, InImplicitId0);
 		FPBDCollisionConstraint* Constraint = FindConstraint(CollisionKey);
 
 		// Check that we have the collision constraint for the expected pair of implicit objects. The key is assigned by the 
