Add FOnDemandIoStore::GetSizeForPackages

diff --git a/Engine/Source/Runtime/Experimental/IoStore/OnDemand/Private/IoStoreOnDemand.cpp b/Engine/Source/Runtime/Experimental/IoStore/OnDemand/Private/IoStoreOnDemand.cpp
--- a/Engine/Source/Runtime/Experimental/IoStore/OnDemand/Private/IoStoreOnDemand.cpp
+++ b/Engine/Source/Runtime/Experimental/IoStore/OnDemand/Private/IoStoreOnDemand.cpp
@@ -41,6 +41,8 @@ namespace UE::IoStore
 ////////////////////////////////////////////////////////////////////////////////
 FString GIasOnDemandTocExt = TEXT(".uondemandtoc");
 
+static const TCHAR* NotInitializedError = TEXT("I/O store on-demand not initialized");
+
 bool GIasSuspendSystem = false;
 static FAutoConsoleVariableRef CVar_SuspendSystemEnabled(
 	TEXT("ias.SuspendSystem"),
@@ -1070,7 +1072,17 @@ FIoStatus FIoStoreOnDemandModule::Unmount(FStringView MountId)
 	{
 		return IoStore->Unmount(MountId);
 	}
-	return FIoStatus(EIoErrorCode::InvalidCode, TEXT("I/O store on-demand not initialized"));
+	return FIoStatus(EIoErrorCode::InvalidCode, NotInitializedError);
+}
+
+TIoStatusOr<uint64> FIoStoreOnDemandModule::GetSizeForPackages(const FOnDemandSizeForPackagesArgs& Args) const
+{
+	if (IoStore)
+	{
+		return IoStore->GetSizeForPackages(Args);
+	}
+
+	return FIoStatus(EIoErrorCode::InvalidCode, NotInitializedError);
 }
 
 void FIoStoreOnDemandModule::InitializeInternal()
diff --git a/Engine/Source/Runtime/Experimental/IoStore/OnDemand/Private/OnDemandIoStore.cpp b/Engine/Source/Runtime/Experimental/IoStore/OnDemand/Private/OnDemandIoStore.cpp
--- a/Engine/Source/Runtime/Experimental/IoStore/OnDemand/Private/OnDemandIoStore.cpp
+++ b/Engine/Source/Runtime/Experimental/IoStore/OnDemand/Private/OnDemandIoStore.cpp
@@ -558,6 +558,34 @@ FIoStatus FOnDemandIoStore::Unmount(FStringView MountId)
 	return EIoErrorCode::Ok;
 }
 
+TIoStatusOr<uint64> FOnDemandIoStore::GetSizeForPackages(const FOnDemandSizeForPackagesArgs& Args) const
+{
+	TSet<FSharedOnDemandContainer> AllContainers(
+		const_cast<FOnDemandIoStore*>(this)->GetMountedContainers([](const FSharedOnDemandContainer& Container)
+			{
+				return EnumHasAnyFlags(Container->Flags, EOnDemandContainerFlags::Installed);
+			})
+	);
+
+	TSet<FPackageId> PackageIdsToInstall;
+	PackageIdsToInstall.Append(Args.Packages);
+
+	Private::FInstallData InstallData;
+	TSet<FPackageId> Missing;
+
+	FIoStatus Status = BuildInstallData(AllContainers, PackageIdsToInstall, InstallData, Missing);
+	if (Status.IsOk() == false)
+	{
+		return Status;
+	}
+
+	const uint64 RetSize = Algo::TransformAccumulate(InstallData,
+		[](const Private::FInstallData::ElementType& Pair) { return Pair.Value.TotalSize; },
+		uint64(0));
+
+	return RetSize;
+}
+
 FOnDemandChunkInfo FOnDemandIoStore::GetStreamingChunkInfo(const FIoChunkId& ChunkId)
 {
 	return GetChunkInfo(ChunkId, EOnDemandContainerFlags::Mounted | EOnDemandContainerFlags::Streaming);
@@ -862,10 +890,10 @@ bool FOnDemandIoStore::Tick()
 					EnumAddFlags(Container->Flags, EOnDemandContainerFlags::Mounted);
 				}
 			}
-
-			FOnDemandMountCompleted OnCompleted = MoveTemp(Request->OnCompleted);
-			OnCompleted(FOnDemandMountResult{ Request->MountArgs.MountId });
 		}
+
+		FOnDemandMountCompleted OnCompleted = MoveTemp(Request->OnCompleted);
+		OnCompleted(FOnDemandMountResult{ Request->MountArgs.MountId });
 	}
 
 	return true;
@@ -1059,11 +1087,12 @@ FIoStatus FOnDemandIoStore::TickInstallRequest(FMountRequest& MountRequest)
 		return EIoErrorCode::Ok;
 	};
 
-	TArray<FSharedOnDemandContainer> MountedContainers = GetMountedContainers();
-	TSet<FSharedOnDemandContainer> AllContainers;
-
-	Algo::CopyIf(MountedContainers, AllContainers, 
-		[](const FSharedOnDemandContainer& Container){ return EnumHasAnyFlags(Container->Flags, EOnDemandContainerFlags::Installed); });
+	TSet<FSharedOnDemandContainer> AllContainers(
+		GetMountedContainers([](const FSharedOnDemandContainer& Container) 
+		{ 
+			return EnumHasAnyFlags(Container->Flags, EOnDemandContainerFlags::Installed); 
+		})
+	);
 	Algo::Copy(MountRequest.Containers, AllContainers);
 
 	// Fetch all container headers
@@ -1421,4 +1450,14 @@ TArray<FSharedOnDemandContainer> FOnDemandIoStore::GetMountedContainers()
 	return Containers;
 }
 
+TArray<FSharedOnDemandContainer> FOnDemandIoStore::GetMountedContainers(TFunctionRef<bool(const FSharedOnDemandContainer& Container)> Predicate)
+{
+	TArray<FSharedOnDemandContainer> Ret;
+
+	UE::TUniqueLock Lock(ContainerMutex);
+	Algo::CopyIf(Containers, Ret, Predicate);
+
+	return Ret;
+}
+
 } // namespace UE::IoStore
diff --git a/Engine/Source/Runtime/Experimental/IoStore/OnDemand/Private/OnDemandIoStore.h b/Engine/Source/Runtime/Experimental/IoStore/OnDemand/Private/OnDemandIoStore.h
--- a/Engine/Source/Runtime/Experimental/IoStore/OnDemand/Private/OnDemandIoStore.h
+++ b/Engine/Source/Runtime/Experimental/IoStore/OnDemand/Private/OnDemandIoStore.h
@@ -145,8 +145,6 @@ class FOnDemandIoStore
 
 	using FSharedMountRequest	= TSharedPtr<FMountRequest>;
 	using FMountRequestMap		= TMap<FString, FSharedMountRequest>;
-	// List of (unique container name, package index list) pairs
-	using FPackageFilter		= TArray<TPair<FString, TConstArrayView<uint32>>>;
 
 public:
 	FOnDemandIoStore();
@@ -159,9 +157,12 @@ public:
 	FIoStatus				Initialize();
 	void					Mount(FOnDemandMountArgs&& Args, FOnDemandMountCompleted OnCompleted);
 	FIoStatus				Unmount(FStringView MountId);
+	TIoStatusOr<uint64>		GetSizeForPackages(const FOnDemandSizeForPackagesArgs& Args) const;
 	FOnDemandChunkInfo		GetStreamingChunkInfo(const FIoChunkId& ChunkId);
 	FOnDemandChunkInfo		GetInstalledChunkInfo(const FIoChunkId& ChunkId);
 	TArray<FSharedOnDemandContainer> GetMountedContainers();
+	TArray<FSharedOnDemandContainer> GetMountedContainers(
+		TFunctionRef<bool(const FSharedOnDemandContainer& Container)> Predicate);
 
 private:
 	void					OnPostFork(EForkProcessRole ProcessRole);
diff --git a/Engine/Source/Runtime/Experimental/IoStore/OnDemand/Public/IO/IoStoreOnDemand.h b/Engine/Source/Runtime/Experimental/IoStore/OnDemand/Public/IO/IoStoreOnDemand.h
--- a/Engine/Source/Runtime/Experimental/IoStore/OnDemand/Public/IO/IoStoreOnDemand.h
+++ b/Engine/Source/Runtime/Experimental/IoStore/OnDemand/Public/IO/IoStoreOnDemand.h
@@ -276,6 +276,25 @@ struct FOnDemandMountResult
 
 using FOnDemandMountCompleted = TUniqueFunction<void(FOnDemandMountResult)>;
 
+struct FOnDemandSizeForPackagesArgs
+{
+	/* Includes all packages from matching containers */
+	struct FPackagesWhere
+	{
+		/* The MountId to collect packages from*/
+		FString MountId;
+		/* Includes only packages from the specificed tag sets if not empty */
+		TArray<FString> TagSets;
+	};
+
+	/* Includes all packages from matching containers */
+	// TODO: implement this
+	// TArray<FPackagesWhere> PackagesWhere;
+
+	/* Explicit packages to include */
+	TArray<FPackageId> Packages;
+};
+
 class FIoStoreOnDemandModule
 	: public IModuleInterface
 {
@@ -300,6 +319,8 @@ public:
 	UE_API void Mount(FOnDemandMountArgs&& Args, FOnDemandMountCompleted OnCompleted);
 	UE_API FIoStatus Unmount(FStringView MountId);
 
+	UE_API TIoStatusOr<uint64> GetSizeForPackages(const FOnDemandSizeForPackagesArgs& Args) const;
+
 	virtual void StartupModule() override;
 	virtual void ShutdownModule() override;
 	
