Add info to the IAS anayltics event detailing if the file cache is enabled or not.

#rnx

diff --git a/Engine/Source/Runtime/Experimental/IoStore/OnDemand/Private/Statistics.cpp b/Engine/Source/Runtime/Experimental/IoStore/OnDemand/Private/Statistics.cpp
--- a/Engine/Source/Runtime/Experimental/IoStore/OnDemand/Private/Statistics.cpp
+++ b/Engine/Source/Runtime/Experimental/IoStore/OnDemand/Private/Statistics.cpp
@@ -520,6 +520,8 @@ void FOnDemandIoBackendStats::ReportEndPointAnalytics(TArray<FAnalyticsEventAttr
 
 		AppendAnalyticsEventAttributeArray(OutAnalyticsArray
 
+			,TEXT("IasCacheEnabled"), !EnumHasAllFlags(Flags, EStatsFlags::CachingDisabled)
+
 			,TRACK_DELTA("IasCacheTotalCount", CacheTotalCount)
 			,TRACK_DELTA("IasCacheErrorCount", GCacheErrorCount.Get())
 			,TRACK_DELTA("IasCacheGetCount", GCacheGetCount.Get())
