Implement review feedback for 34881161

diff --git a/Engine/Source/Runtime/Experimental/IoStore/OnDemand/Private/IoStoreOnDemand.cpp b/Engine/Source/Runtime/Experimental/IoStore/OnDemand/Private/IoStoreOnDemand.cpp
--- a/Engine/Source/Runtime/Experimental/IoStore/OnDemand/Private/IoStoreOnDemand.cpp
+++ b/Engine/Source/Runtime/Experimental/IoStore/OnDemand/Private/IoStoreOnDemand.cpp
@@ -1075,11 +1075,11 @@ FIoStatus FIoStoreOnDemandModule::Unmount(FStringView MountId)
 	return FIoStatus(EIoErrorCode::InvalidCode, NotInitializedError);
 }
 
-TIoStatusOr<uint64> FIoStoreOnDemandModule::GetSizeForPackages(const FOnDemandSizeForPackagesArgs& Args) const
+TIoStatusOr<uint64> FIoStoreOnDemandModule::GetInstallSize(const FOnDemandGetInstallSizeArgs& Args) const
 {
 	if (IoStore)
 	{
-		return IoStore->GetSizeForPackages(Args);
+		return IoStore->GetInstallSize(Args);
 	}
 
 	return FIoStatus(EIoErrorCode::InvalidCode, NotInitializedError);
diff --git a/Engine/Source/Runtime/Experimental/IoStore/OnDemand/Private/OnDemandIoStore.cpp b/Engine/Source/Runtime/Experimental/IoStore/OnDemand/Private/OnDemandIoStore.cpp
--- a/Engine/Source/Runtime/Experimental/IoStore/OnDemand/Private/OnDemandIoStore.cpp
+++ b/Engine/Source/Runtime/Experimental/IoStore/OnDemand/Private/OnDemandIoStore.cpp
@@ -558,8 +558,10 @@ FIoStatus FOnDemandIoStore::Unmount(FStringView MountId)
 	return EIoErrorCode::Ok;
 }
 
-TIoStatusOr<uint64> FOnDemandIoStore::GetSizeForPackages(const FOnDemandSizeForPackagesArgs& Args) const
+TIoStatusOr<uint64> FOnDemandIoStore::GetInstallSize(const FOnDemandGetInstallSizeArgs& Args) const
 {
+	TRACE_CPUPROFILER_EVENT_SCOPE(FOnDemandIoStore::GetInstallSize);
+
 	TSet<FSharedOnDemandContainer> AllContainers(
 		const_cast<FOnDemandIoStore*>(this)->GetMountedContainers([](const FSharedOnDemandContainer& Container)
 			{
diff --git a/Engine/Source/Runtime/Experimental/IoStore/OnDemand/Private/OnDemandIoStore.h b/Engine/Source/Runtime/Experimental/IoStore/OnDemand/Private/OnDemandIoStore.h
--- a/Engine/Source/Runtime/Experimental/IoStore/OnDemand/Private/OnDemandIoStore.h
+++ b/Engine/Source/Runtime/Experimental/IoStore/OnDemand/Private/OnDemandIoStore.h
@@ -157,7 +157,7 @@ public:
 	FIoStatus				Initialize();
 	void					Mount(FOnDemandMountArgs&& Args, FOnDemandMountCompleted OnCompleted);
 	FIoStatus				Unmount(FStringView MountId);
-	TIoStatusOr<uint64>		GetSizeForPackages(const FOnDemandSizeForPackagesArgs& Args) const;
+	TIoStatusOr<uint64>		GetInstallSize(const FOnDemandGetInstallSizeArgs& Args) const;
 	FOnDemandChunkInfo		GetStreamingChunkInfo(const FIoChunkId& ChunkId);
 	FOnDemandChunkInfo		GetInstalledChunkInfo(const FIoChunkId& ChunkId);
 	TArray<FSharedOnDemandContainer> GetMountedContainers();
diff --git a/Engine/Source/Runtime/Experimental/IoStore/OnDemand/Public/IO/IoStoreOnDemand.h b/Engine/Source/Runtime/Experimental/IoStore/OnDemand/Public/IO/IoStoreOnDemand.h
--- a/Engine/Source/Runtime/Experimental/IoStore/OnDemand/Public/IO/IoStoreOnDemand.h
+++ b/Engine/Source/Runtime/Experimental/IoStore/OnDemand/Public/IO/IoStoreOnDemand.h
@@ -276,7 +276,7 @@ struct FOnDemandMountResult
 
 using FOnDemandMountCompleted = TUniqueFunction<void(FOnDemandMountResult)>;
 
-struct FOnDemandSizeForPackagesArgs
+struct FOnDemandGetInstallSizeArgs
 {
 	/* Includes all packages from matching containers */
 	struct FPackagesWhere
@@ -319,7 +319,7 @@ public:
 	UE_API void Mount(FOnDemandMountArgs&& Args, FOnDemandMountCompleted OnCompleted);
 	UE_API FIoStatus Unmount(FStringView MountId);
 
-	UE_API TIoStatusOr<uint64> GetSizeForPackages(const FOnDemandSizeForPackagesArgs& Args) const;
+	UE_API TIoStatusOr<uint64> GetInstallSize(const FOnDemandGetInstallSizeArgs& Args) const;
 
 	virtual void StartupModule() override;
 	virtual void ShutdownModule() override;
