Modular Vehicle (suspension surface type);

Adding a way to track the current surface type that our springs are currently touching.

-- Code --
ModularVehicleSimulationCU:
PerformAdditionalSimWork() - Creating a FSuspensionTargetPoint and passing to the suspension's SetTargetPoint, instead of individual values.

SuspensionBaseInterface:
FSuspensionTargetPoint - Created this struct to hold target position values instead of individual properties.
Ctor() - Removed the list of properties being init'd.
SetTargetPoint() - Changed to take an FSuspensionTargetPoint param.
IsWheelInContact() - Use the TargetPoint value.
GetSurfaceType() - Simple getter.
SetImpactNormal() - Simple setter.
GetImpactNormal() - Simple getter.
SetHitDistance() - Simple setter.
GetHitDistance() - Simple getter.
SetTargetPosition() - Simple setter.
GetTargetPosition() - Simple getter.
TargetPoint - FSuspensionTargetPoint property. Storing suspension target point data.

SuspensionSimModule:
FSuspensionSimModule::UpdateConstraint() - Getting the values instead of accessing private properties.
FSuspensionOutputData::FillOutputState() - Getting the values instead of accessing private properties.

#tests PIE, sandbox, make sure everything still works as expected.

diff --git a/Engine/Plugins/Experimental/ChaosModularVehicle/Source/ChaosModularVehicleEngine/Private/ChaosModularVehicle/ModularVehicleSimulationCU.cpp b/Engine/Plugins/Experimental/ChaosModularVehicle/Source/ChaosModularVehicleEngine/Private/ChaosModularVehicle/ModularVehicleSimulationCU.cpp
--- a/Engine/Plugins/Experimental/ChaosModularVehicle/Source/ChaosModularVehicleEngine/Private/ChaosModularVehicle/ModularVehicleSimulationCU.cpp
+++ b/Engine/Plugins/Experimental/ChaosModularVehicle/Source/ChaosModularVehicleEngine/Private/ChaosModularVehicle/ModularVehicleSimulationCU.cpp
@@ -358,7 +358,16 @@ void FModularVehicleSimulationCU::PerformAdditionalSimWork(UWorld* InWorld, cons
 									HitDistance = HitResult.Distance - WheelRadius;
 								}
 
-								Suspension->SetTargetPoint(HitPoint, HitResult.ImpactNormal, HitDistance, HitResult.bBlockingHit);
+								const TEnumAsByte<EPhysicalSurface> DefaultSurfaceType = EPhysicalSurface::SurfaceType_Default;
+								FSuspensionTargetPoint TargetPoint(
+									HitPoint
+									, HitResult.ImpactNormal
+									, HitDistance
+									, HitResult.bBlockingHit
+									, HitResult.PhysMaterial.IsValid() ? HitResult.PhysMaterial->SurfaceType : DefaultSurfaceType
+								);
+
+								Suspension->SetTargetPoint(TargetPoint);
 							}
 
 						}
diff --git a/Engine/Plugins/Experimental/ChaosModularVehicle/Source/ChaosModularVehicleEngine/Private/ChaosModularVehicle/SuspensionSimModule.cpp b/Engine/Plugins/Experimental/ChaosModularVehicle/Source/ChaosModularVehicleEngine/Private/ChaosModularVehicle/SuspensionSimModule.cpp
--- a/Engine/Plugins/Experimental/ChaosModularVehicle/Source/ChaosModularVehicleEngine/Private/ChaosModularVehicle/SuspensionSimModule.cpp
+++ b/Engine/Plugins/Experimental/ChaosModularVehicle/Source/ChaosModularVehicleEngine/Private/ChaosModularVehicle/SuspensionSimModule.cpp
@@ -179,7 +179,10 @@ using namespace Chaos;
 				if (FSuspensionConstraintPhysicsProxy* Proxy = Constraint->GetProxy<FSuspensionConstraintPhysicsProxy>())
 				{
 					Chaos::FPhysicsSolver* Solver = Proxy->GetSolver<Chaos::FPhysicsSolver>();
-					Solver->SetSuspensionTarget(Constraint, TargetPos, ImpactNormal, WheelInContact);
+					const FVector& CurrentTargetPosition = GetTargetPosition();
+					const FVector& CurrentImpactNormal = GetImpactNormal();
+					const bool bCurrentWheelInContact = IsWheelInContact();
+					Solver->SetSuspensionTarget(Constraint, CurrentTargetPosition, CurrentImpactNormal, bCurrentWheelInContact);
 				}
 			}
 		}
@@ -229,7 +232,7 @@ using namespace Chaos;
 			SpringDisplacement = Sim->SpringDisplacement;
 			SpringDisplacementVector = -Sim->Setup().SuspensionAxis * Sim->SpringDisplacement + Sim->GetAnimationOffset();
 			SpringSpeed = Sim->SpringSpeed;
-			ImpactNormal = Sim->ImpactNormal;
+			ImpactNormal = Sim->GetImpactNormal();
 
 			AnimFlags = EAnimationFlags::AnimatePosition;
 			AnimationLocOffset = SpringDisplacementVector;
diff --git a/Engine/Source/Runtime/Experimental/ChaosVehicles/ChaosVehiclesCore/Private/SimModule/SuspensionBaseInterface.cpp b/Engine/Source/Runtime/Experimental/ChaosVehicles/ChaosVehiclesCore/Private/SimModule/SuspensionBaseInterface.cpp
--- a/Engine/Source/Runtime/Experimental/ChaosVehicles/ChaosVehiclesCore/Private/SimModule/SuspensionBaseInterface.cpp
+++ b/Engine/Source/Runtime/Experimental/ChaosVehicles/ChaosVehiclesCore/Private/SimModule/SuspensionBaseInterface.cpp
@@ -5,27 +5,14 @@
 namespace Chaos
 {
 
-FSuspensionBaseInterface::FSuspensionBaseInterface()
-	: WheelSimTreeIndex(INVALID_IDX)
-	, TargetPos(FVector::ZeroVector)
-	, ImpactNormal(FVector::ZeroVector)
-	, HitDistance(0.f)
-	, WheelInContact(false)
-{
-}
-
 bool FSuspensionBaseInterface::IsBehaviourType(eSimModuleTypeFlags InType) const
 { 
 	return (InType & Raycast); 
 }
 
-void FSuspensionBaseInterface::SetTargetPoint(const FVector& InTargetPoint, const FVector& InImpactNormal, float InHitDistance, bool InWheelInContact)
+void FSuspensionBaseInterface::SetTargetPoint(const FSuspensionTargetPoint& InTargetPoint)
 {
-	TargetPos = InTargetPoint;
-	ImpactNormal = InImpactNormal;
-	HitDistance = InHitDistance;
-	WheelInContact = InWheelInContact;
+	TargetPoint = InTargetPoint;
 }
 
-
 } // namespace Chaos
diff --git a/Engine/Source/Runtime/Experimental/ChaosVehicles/ChaosVehiclesCore/Public/SimModule/SuspensionBaseInterface.h b/Engine/Source/Runtime/Experimental/ChaosVehicles/ChaosVehiclesCore/Public/SimModule/SuspensionBaseInterface.h
--- a/Engine/Source/Runtime/Experimental/ChaosVehicles/ChaosVehiclesCore/Public/SimModule/SuspensionBaseInterface.h
+++ b/Engine/Source/Runtime/Experimental/ChaosVehicles/ChaosVehiclesCore/Public/SimModule/SuspensionBaseInterface.h
@@ -2,6 +2,7 @@
 
 #pragma once
 
+#include "Chaos/ChaosEngineInterface.h"
 #include "SimModule/SimulationModuleBase.h"
 
 namespace Chaos
@@ -26,35 +27,71 @@ struct CHAOSVEHICLESCORE_API FSpringTrace
 	}
 };
 
+/** Suspension target point data. */
+struct CHAOSVEHICLESCORE_API FSuspensionTargetPoint
+{
+	FSuspensionTargetPoint() {}
+
+	FSuspensionTargetPoint(const FVector& InTargetPosition, const FVector& InImpactNormal, const float InHitDistance, const bool bInWheelInContact, const TEnumAsByte<EPhysicalSurface> InSurfaceType)
+		: TargetPosition(InTargetPosition)
+		, ImpactNormal(InImpactNormal)
+		, HitDistance(InHitDistance)
+		, bWheelInContact(bInWheelInContact)
+		, SurfaceType(InSurfaceType)
+	{}
+
+	FVector TargetPosition = FVector::ZeroVector;
+	FVector ImpactNormal = FVector::ZeroVector;
+	float HitDistance = 0.0f;
+	bool bWheelInContact = false;
+	TEnumAsByte<EPhysicalSurface> SurfaceType = EPhysicalSurface::SurfaceType_Default;
+};
 
-class CHAOSVEHICLESCORE_API FSuspensionBaseInterface : public ISimulationModuleBase, public TSimulationModuleTypeable<FSuspensionBaseInterface>
+class CHAOSVEHICLESCORE_API FSuspensionBaseInterface 
+	: public ISimulationModuleBase
+	, public TSimulationModuleTypeable<FSuspensionBaseInterface>
 {
 public:
+
 	DEFINE_CHAOSSIMTYPENAME(FSuspensionBaseInterface);
-	FSuspensionBaseInterface();
+	FSuspensionBaseInterface() {}
 
 	virtual ~FSuspensionBaseInterface() {}
 
 	virtual bool IsBehaviourType(eSimModuleTypeFlags InType) const override;
 
 	virtual float GetMaxSpringLength() const = 0;
+	
 	virtual float GetSpringLength() const = 0;
 	virtual void SetSpringLength(float InLength, float WheelRadius) = 0;
+	
 	virtual void GetWorldRaycastLocation(const FTransform& BodyTransform, float WheelRadius, FSpringTrace& OutTrace) = 0;
 
-	void SetTargetPoint(const FVector& InTargetPoint, const FVector& InImpactNormal, float InHitDistance, bool InWheelInContact);
-	bool IsWheelInContact() const { return WheelInContact; }
+	void SetTargetPoint(const FSuspensionTargetPoint& InTargetPoint);
+	
+	bool IsWheelInContact() const { return TargetPoint.bWheelInContact; }
+	
 	void SetWheelSimTreeIndex(int WheelTreeIndexIn) { WheelSimTreeIndex = WheelTreeIndexIn; }
 	int GetWheelSimTreeIndex() const { return WheelSimTreeIndex; }
+	
+	TEnumAsByte<EPhysicalSurface> GetSurfaceType() const { return TargetPoint.SurfaceType; }
+	
+	void SetImpactNormal(const FVector& NewValue) { TargetPoint.ImpactNormal = NewValue; }
+	const FVector& GetImpactNormal() const { return TargetPoint.ImpactNormal; }
+
+	void SetHitDistance(const float NewValue) { TargetPoint.HitDistance = NewValue; }
+	float GetHitDistance() const { return TargetPoint.HitDistance; }
+
+	void SetTargetPosition(const FVector& NewValue) { TargetPoint.TargetPosition = NewValue; }
+	const FVector& GetTargetPosition() const { return TargetPoint.TargetPosition; }
 
 protected:
-	int WheelSimTreeIndex;
 
-	FVector TargetPos;
-	FVector ImpactNormal;
-	float HitDistance;
-	bool WheelInContact;
+	int WheelSimTreeIndex = INVALID_IDX;
+
+private:
 
+	FSuspensionTargetPoint TargetPoint;
 };
 
 
