Fixed Geometry collection memory reporting for GetResourceSizeEx if editor data is filtered and the cooking data options are removed.

diff --git a/Engine/Source/Runtime/Experimental/GeometryCollectionEngine/Private/GeometryCollection/GeometryCollectionObject.cpp b/Engine/Source/Runtime/Experimental/GeometryCollectionEngine/Private/GeometryCollection/GeometryCollectionObject.cpp
--- a/Engine/Source/Runtime/Experimental/GeometryCollectionEngine/Private/GeometryCollection/GeometryCollectionObject.cpp
+++ b/Engine/Source/Runtime/Experimental/GeometryCollectionEngine/Private/GeometryCollection/GeometryCollectionObject.cpp
@@ -1034,7 +1034,7 @@ void UGeometryCollection::Serialize(FArchive& Ar)
 			EnsureSimulationDataIsCooked(bIsTransacting, bAllowCopyFromDDC);
 		}
 #endif
-		if (Ar.IsLoading() || Ar.IsCountingMemory())
+		if (Ar.IsLoading() || (Ar.IsCountingMemory() && !Ar.IsFilterEditorOnly()))
 		{
 			GeometryCollection->Serialize(ChaosAr);
 		}
