IOStore OnDemand Install API

diff --git a/Engine/Source/Runtime/Core/Public/IO/IoChunkId.h b/Engine/Source/Runtime/Core/Public/IO/IoChunkId.h
--- a/Engine/Source/Runtime/Core/Public/IO/IoChunkId.h
+++ b/Engine/Source/Runtime/Core/Public/IO/IoChunkId.h
@@ -87,6 +87,16 @@ public:
 		return !(*this == Rhs);
 	}
 
+	inline bool operator <(const FIoChunkId& Rhs) const
+	{
+		return FMemory::Memcmp(Id, Rhs.Id, sizeof Id) < 0;
+	}
+
+	inline bool operator >(const FIoChunkId& Rhs) const
+	{
+		return FMemory::Memcmp(Id, Rhs.Id, sizeof Id) > 0;
+	}
+
 	void Set(const void* InIdPtr, SIZE_T InSize)
 	{
 		check(InSize == sizeof Id);
diff --git a/Engine/Source/Runtime/Experimental/IoStore/OnDemand/Private/IoStoreOnDemand.cpp b/Engine/Source/Runtime/Experimental/IoStore/OnDemand/Private/IoStoreOnDemand.cpp
--- a/Engine/Source/Runtime/Experimental/IoStore/OnDemand/Private/IoStoreOnDemand.cpp
+++ b/Engine/Source/Runtime/Experimental/IoStore/OnDemand/Private/IoStoreOnDemand.cpp
@@ -1050,11 +1050,11 @@ void FIoStoreOnDemandModule::ReportAnalytics(TArray<FAnalyticsEventAttribute>& O
 	}
 }
 
-void FIoStoreOnDemandModule::Mount(FOnDemandMountArgs&& Args, FOnDemandMountCompleted OnCompleted)
+void FIoStoreOnDemandModule::Mount(FOnDemandMountArgs&& Args, FOnDemandMountCompleted&& OnCompleted)
 {
 	if (IoStore.IsValid() == false)
 	{
-		IoStore = MakeUnique<FOnDemandIoStore>();
+		IoStore = MakeShared<FOnDemandIoStore>();
 		if (FIoStatus Status = IoStore->Initialize(); !Status.IsOk())
 		{
 			UE_LOG(LogIas, Error, TEXT("Failed to initialize I/O store on-demand, reason '%s'"), *Status.ToString());
@@ -1066,6 +1066,22 @@ void FIoStoreOnDemandModule::Mount(FOnDemandMountArgs&& Args, FOnDemandMountComp
 	IoStore->Mount(MoveTemp(Args), MoveTemp(OnCompleted));
 }
 
+void FIoStoreOnDemandModule::Install(FOnDemandInstallArgs&& Args, FOnDemandInstallCompleted&& OnCompleted)
+{
+	if (IoStore.IsValid() == false)
+	{
+		IoStore = MakeShared<FOnDemandIoStore>();
+		if (FIoStatus Status = IoStore->Initialize(); !Status.IsOk())
+		{
+			UE_LOG(LogIas, Error, TEXT("Failed to initialize I/O store on-demand, reason '%s'"), *Status.ToString());
+			IoStore.Reset();
+			return OnCompleted(FOnDemandInstallResult{ .Status = Status });
+		}
+	}
+
+	IoStore->Install(MoveTemp(Args), MoveTemp(OnCompleted));
+}
+
 FIoStatus FIoStoreOnDemandModule::Unmount(FStringView MountId)
 {
 	if (IoStore.IsValid())
@@ -1110,7 +1126,7 @@ void FIoStoreOnDemandModule::InitializeInternal()
 
 	if (IoStore.IsValid() == false)
 	{
-		IoStore = MakeUnique<FOnDemandIoStore>();
+		IoStore = MakeShared<FOnDemandIoStore>();
 		if (FIoStatus Status = IoStore->Initialize(); !Status.IsOk())
 		{
 			UE_LOG(LogIas, Error, TEXT("Failed to initialize I/O store on demand, reason '%s'"), *Status.ToString());
@@ -1221,16 +1237,21 @@ void FIoStoreOnDemandModule::InitializeInternal()
 	}
 
 #if !UE_BUILD_SHIPPING
+	TOptional<FOnDemandInstallArgs> InstallArgs;
 	if (FParse::Param(FCommandLine::Get(), TEXT("Iad")))
 	{
-		// Temporary switch for testing installation to local storage (IAD)
-		MountArgs.Emplace(FOnDemandMountArgs
+		if (MountArgs)
 		{
-			.MountId = EndpointConfig.TocFilePath,
-			.Url = EndpointConfig.ServiceUrls[0] / EndpointConfig.TocPath,
-			.FilePath = EndpointConfig.TocFilePath,
-			.Options = EOnDemandMountOptions::Install
-		});
+			MountArgs.GetValue().Options = EOnDemandMountOptions::InstallOnDemand;
+
+			static FOnDemandContentHandle ContentHandle = FOnDemandContentHandle::Create(TEXT("AllContent"));
+			InstallArgs.Emplace(FOnDemandInstallArgs
+			{
+				.Url = EndpointConfig.ServiceUrls[0] / EndpointConfig.TocPath,
+				.MountId = MountArgs.GetValue().MountId,
+				.ContentHandle = ContentHandle
+			});
+		}
 	}
 #endif
 
@@ -1243,6 +1264,18 @@ void FIoStoreOnDemandModule::InitializeInternal()
 				UE_CLOG(!MountResult.Status.IsOk(), LogIas, Error,
 					TEXT("Failed to mount TOC for '%s', reason '%s'"), *MountResult.MountId, *MountResult.Status.ToString());
 			});
+#if !UE_BUILD_SHIPPING
+		if (InstallArgs)
+		{
+			IoStore->Install(
+				MoveTemp(InstallArgs.GetValue()),
+				[](FOnDemandInstallResult InstallResult)
+				{
+					UE_CLOG(!InstallResult.Status.IsOk(), LogIoStoreOnDemand, Error,
+						TEXT("Failed to install content, reason '%s'"), *InstallResult.Status.ToString());
+				});
+		}
+#endif
 	}
 }
 	
diff --git a/Engine/Source/Runtime/Experimental/IoStore/OnDemand/Private/OnDemandInstallCache.cpp b/Engine/Source/Runtime/Experimental/IoStore/OnDemand/Private/OnDemandInstallCache.cpp
--- a/Engine/Source/Runtime/Experimental/IoStore/OnDemand/Private/OnDemandInstallCache.cpp
+++ b/Engine/Source/Runtime/Experimental/IoStore/OnDemand/Private/OnDemandInstallCache.cpp
@@ -1207,15 +1207,24 @@ FIoStatus FOnDemandInstallCache::Purge(TMap<FIoHash, uint64>&& ChunksToIntall)
 		return FIoStatus::Ok;
 	}
 
-	TSet<FCasBlockId>					ReferencedBlocks;
-	TArray<FSharedOnDemandContainer>	MountedContainers = IoStore.GetMountedContainers();
+	TArray<FSharedOnDemandContainer>	Containers; 
+	TArray<TBitArray<>>					ChunkEntryIndices;
+
+	IoStore.GetReferencedContent(Containers, ChunkEntryIndices);
+	check(Containers.Num() == ChunkEntryIndices.Num());
 
 	//TODO: Compute fragmentation metric and redownload chunks when this number gets too high
-	for (const FSharedOnDemandContainer& Container : MountedContainers)
+	for (int32 Index = 0; FSharedOnDemandContainer Container : Containers)
 	{
-		for (const TPair<FIoChunkId, FOnDemandChunkEntry>& Kv : Container->ChunkEntries)
+		const TBitArray<>& IsReferenced = ChunkEntryIndices[Index++];
+		for (int32 EntryIndex = 0; const FOnDemandChunkEntry& Entry : Container->ChunkEntries)
 		{
-			if (FCasLocation Loc = Cas.FindChunk(Kv.Value.Hash); Loc.IsValid())
+			if (bool bIsReferenced = IsReferenced[EntryIndex++]; bIsReferenced == false)
+			{
+				continue;
+			}
+
+			if (FCasLocation Loc = Cas.FindChunk(Entry.Hash); Loc.IsValid())
 			{
 				if (FCasBlockInfo* Info = BlockInfo.Find(Loc.BlockId))
 				{
diff --git a/Engine/Source/Runtime/Experimental/IoStore/OnDemand/Private/OnDemandIoStore.cpp b/Engine/Source/Runtime/Experimental/IoStore/OnDemand/Private/OnDemandIoStore.cpp
--- a/Engine/Source/Runtime/Experimental/IoStore/OnDemand/Private/OnDemandIoStore.cpp
+++ b/Engine/Source/Runtime/Experimental/IoStore/OnDemand/Private/OnDemandIoStore.cpp
@@ -7,13 +7,13 @@
 
 #include "Algo/Accumulate.h"
 #include "Algo/Copy.h"
-#include "Algo/Find.h"
 #include "Algo/RemoveIf.h"
 #include "Algo/Transform.h"
 #include "Async/ManualResetEvent.h"
 #include "Async/UniqueLock.h"
 #include "Containers/RingBuffer.h"
 #include "Containers/StringConv.h"
+#include "Containers/Ticker.h"
 #include "HAL/FileManager.h"
 #include "HAL/PlatformFile.h"
 #include "HAL/PlatformFileManager.h"
@@ -126,6 +126,86 @@ static void SplitHostUrl(const FStringView& Url, FStringView& OutHost, FStringVi
 	OutRemainder = Url.RightChop(OutHost.Len());
 }
 
+///////////////////////////////////////////////////////////////////////////////
+static FAnsiStringBuilderBase& GetChunkUrl(
+	const FStringView& Host,
+	const FOnDemandContainer& Container,
+	const FOnDemandChunkEntry& Entry,
+	FAnsiStringBuilderBase& OutUrl)
+{
+	OutUrl.Reset();
+	if (Host.IsEmpty() == false)
+	{
+		OutUrl << Host;
+	}
+
+	if (!Container.ChunksDirectory.IsEmpty())
+	{
+		OutUrl << "/" << Container.ChunksDirectory;
+	}
+
+	const FString HashString = LexToString(Entry.Hash);
+	OutUrl << "/" << HashString.Left(2) << "/" << HashString << ".iochunk";
+
+	return OutUrl;
+}
+
+static FIoStatus FetchContainerHeader(
+	FStringView Host,
+	FAnsiStringBuilderBase& ChunkUrlBuilder,
+	FSharedOnDemandContainer Container,
+	FIoContainerHeader& OutHeader)
+{
+	const FIoChunkId			ChunkId = CreateContainerHeaderChunkId(Container->ContainerId);
+	const FOnDemandChunkEntry*	Entry = Container->FindChunkEntry(ChunkId);
+
+	if (Entry == nullptr)
+	{
+		return EIoErrorCode::Ok;
+	}
+
+	UE_LOG(LogIoStoreOnDemand, Verbose, TEXT("Fetching container header, ContainerName='%s'"), *Container->Name);
+	TIoStatusOr<FIoBuffer> Response = FHttpClient::Get(GetChunkUrl(Host, *Container, *Entry, ChunkUrlBuilder).ToView(), 2, EHttpRedirects::Follow);
+
+	if (Response.IsOk() == false)
+	{
+		FIoStatus Status = FIoStatusBuilder(EIoErrorCode::ReadError) << TEXT("Failed to fetch container header chunk");
+		return Status;
+	}
+
+	FOnDemandChunkInfo ChunkInfo(Container, *Entry);
+
+	FIoChunkDecodingParams Params;
+	Params.CompressionFormat = ChunkInfo.CompressionFormat();
+	Params.EncryptionKey = ChunkInfo.EncryptionKey();
+	Params.BlockSize = ChunkInfo.BlockSize();
+	Params.TotalRawSize = ChunkInfo.RawSize();
+	Params.RawOffset = 0;
+	Params.EncodedOffset = 0;
+	Params.EncodedBlockSize = ChunkInfo.Blocks();
+	Params.BlockHash = ChunkInfo.BlockHashes();
+
+	FIoBuffer EncodedChunk = Response.ConsumeValueOrDie();
+	FIoBuffer RawChunk(ChunkInfo.RawSize());
+	if (FIoChunkEncoding::Decode(Params, EncodedChunk.GetView(), RawChunk.GetMutableView()) == false)
+	{
+		FIoStatus Status = FIoStatusBuilder(EIoErrorCode::ReadError)
+			<< TEXT("Failed to decode container header chunk");
+	}
+
+	FMemoryReaderView Ar(RawChunk.GetView());
+	Ar << OutHeader;
+	Ar.Close();
+
+	if (Ar.IsError() || Ar.IsCriticalError())
+	{
+		FIoStatus Status = FIoStatusBuilder(EIoErrorCode::FileNotOpen) << TEXT("Failed to serialize container header");
+		return Status;
+	}
+
+	return EIoErrorCode::Ok;
+}
+
 ///////////////////////////////////////////////////////////////////////////////
 using FPackageStoreEntryMap		= TMap<FPackageId, const FFilePackageStoreEntry*>;
 using FSoftPackageReferenceMap	= TMap<FPackageId, const FFilePackageStoreEntrySoftReferences*>;
@@ -135,7 +215,7 @@ struct FContainerInstallData
 	FPackageStoreEntryMap		PackageStoreEntries;
 	FSoftPackageReferenceMap	SoftPackageReferences;
 	TSet<FPackageId>			PackageIds;
-	TSet<FIoChunkId>			ResolvedChunks;
+	TSet<uint32>				ResolvedChunks;
 	uint64						TotalSize = 0;
 };
 
@@ -265,17 +345,17 @@ FIoStatus BuildInstallData(
 
 		for (const FPackageId& PackageId : Data.PackageIds)
 		{
-			const FIoChunkId PackageChunkId					= CreatePackageDataChunkId(PackageId);
-			const FOnDemandChunkEntry* PackageChunkEntry	= Container.ChunkEntries.Find(PackageChunkId);
+			const FIoChunkId PackageChunkId	= CreatePackageDataChunkId(PackageId);
+			int32 EntryIndex				= Container.FindChunkEntryIndex(PackageChunkId);
 
-			if (PackageChunkEntry == nullptr)
+			if (EntryIndex == INDEX_NONE)
 			{
 				UE_LOG(LogIoStoreOnDemand, Warning, TEXT("Missing package data"));
 				continue;
 			}
 
-			Data.ResolvedChunks.Add(PackageChunkId);
-			Data.TotalSize += PackageChunkEntry->EncodedSize;
+			Data.ResolvedChunks.Add(EntryIndex);
+			Data.TotalSize += Container.ChunkEntries[EntryIndex].EncodedSize;
 
 			const EIoChunkType AdditionalPackageChunkTypes[] =
 			{
@@ -287,32 +367,30 @@ FIoStatus BuildInstallData(
 			for (EIoChunkType ChunkType : AdditionalPackageChunkTypes)
 			{
 				const FIoChunkId ChunkId = CreateIoChunkId(PackageId.Value(), 0, ChunkType);
-				if (const FOnDemandChunkEntry* ChunkEntry = Container.ChunkEntries.Find(ChunkId); ChunkEntry != nullptr)
+				if (EntryIndex = Container.FindChunkEntryIndex(ChunkId); EntryIndex != INDEX_NONE)
 				{
-					Data.ResolvedChunks.Add(ChunkId);
-					Data.TotalSize += ChunkEntry->EncodedSize;
+					Data.ResolvedChunks.Add(EntryIndex);
+					Data.TotalSize += Container.ChunkEntries[EntryIndex].EncodedSize;
 				}
 			}
 		}
 
 		// For now we always download these chunks
-		for (const TPair<FIoChunkId, FOnDemandChunkEntry>& IdEntry : Container.ChunkEntries)
+		for (int32 EntryIndex = 0; const FIoChunkId& ChunkId : Container.ChunkIds)
 		{
-			const FIoChunkId& ChunkId				= IdEntry.Key;
-			const FOnDemandChunkEntry& ChunkEntry	= IdEntry.Value;
-
 			switch(ChunkId.GetChunkType())
 			{
 				case EIoChunkType::ExternalFile:
 				case EIoChunkType::ShaderCodeLibrary:
 				case EIoChunkType::ShaderCode:
 				{
-					Data.ResolvedChunks.Add(ChunkId);
-					Data.TotalSize += ChunkEntry.EncodedSize;
+					Data.ResolvedChunks.Add(EntryIndex);
+					Data.TotalSize += Container.ChunkEntries[EntryIndex].EncodedSize;
 				}
 				default:
 					break;
 			}
+			++EntryIndex;
 		}
 	}
 
@@ -329,8 +407,8 @@ void LexToString(EOnDemandContainerFlags Flags, FStringBuilderBase& Out)
 		TEXT("None"),
 		TEXT("PendingEncryptionKey"),
 		TEXT("Mounted"),
-		TEXT("Streaming"),
-		TEXT("Installed"),
+		TEXT("StreamOnDemand"),
+		TEXT("InstallOnDemand"),
 		TEXT("Encrypted")
 	};
 
@@ -494,11 +572,11 @@ FIoStatus FOnDemandIoStore::Initialize()
 	return EIoErrorCode::Ok;
 }
 
-void FOnDemandIoStore::Mount(FOnDemandMountArgs&& Args, FOnDemandMountCompleted OnCompleted)
+void FOnDemandIoStore::Mount(FOnDemandMountArgs&& Args, FOnDemandMountCompleted&& OnCompleted)
 {
 	if (Args.MountId.IsEmpty())
 	{
-		return OnCompleted(FOnDemandMountResult{ FString(), FIoStatus(EIoErrorCode::InvalidParameter, TEXT("Invalid Mount ID"))});
+		return OnCompleted(FOnDemandMountResult{ .Status = FIoStatus(EIoErrorCode::InvalidParameter, TEXT("Invalid mount ID")) });
 	}
 
 	{
@@ -508,19 +586,33 @@ void FOnDemandIoStore::Mount(FOnDemandMountArgs&& Args, FOnDemandMountCompleted
 		if (MountRequest.IsValid())
 		{
 			UE_LOG(LogIoStoreOnDemand, Warning, TEXT("Mount request '%s' is already mounting"), *Args.MountId);
-			return OnCompleted(FOnDemandMountResult{ Args.MountId, FIoStatus(EIoErrorCode::InvalidParameter) });
+			return OnCompleted(FOnDemandMountResult{ .Status = FIoStatus(EIoErrorCode::InvalidParameter, TEXT("Mount in progress")) });
 		}
 
 		UE_LOG(LogIoStoreOnDemand, Log, TEXT("Enqueing mount request, MountId='%s'"), *Args.MountId);
 		
 		MountRequest = MakeShared<FMountRequest>();
-		MountRequest->MountArgs		= MoveTemp(Args);
+		MountRequest->Args			= MoveTemp(Args);
 		MountRequest->OnCompleted	= MoveTemp(OnCompleted);
 	}
 
 	TryEnterTickLoop();
 }
 
+void FOnDemandIoStore::Install(FOnDemandInstallArgs&& Args, FOnDemandInstallCompleted&& OnCompleted)
+{
+	FSharedInstallRequest InstallRequest = MakeShared<FInstallRequest>();
+	InstallRequest->Args		= MoveTemp(Args);
+	InstallRequest->OnCompleted = MoveTemp(OnCompleted);
+
+	{
+		UE::TUniqueLock Lock(MountRequestMutex);
+		InstallRequests.Add(MoveTemp(InstallRequest));
+	}
+
+	TryEnterTickLoop();
+}
+
 FIoStatus FOnDemandIoStore::Unmount(FStringView MountId)
 {
 	UE_LOG(LogIoStoreOnDemand, Log, TEXT("Unmounting '%s'"), *WriteToString<256>(MountId));
@@ -563,9 +655,9 @@ TIoStatusOr<uint64> FOnDemandIoStore::GetInstallSize(const FOnDemandGetInstallSi
 	TRACE_CPUPROFILER_EVENT_SCOPE(FOnDemandIoStore::GetInstallSize);
 
 	TSet<FSharedOnDemandContainer> AllContainers(
-		const_cast<FOnDemandIoStore*>(this)->GetMountedContainers([](const FSharedOnDemandContainer& Container)
+		const_cast<FOnDemandIoStore*>(this)->GetContainers([](const FSharedOnDemandContainer& Container)
 			{
-				return EnumHasAnyFlags(Container->Flags, EOnDemandContainerFlags::Installed);
+				return EnumHasAnyFlags(Container->Flags, EOnDemandContainerFlags::InstallOnDemand);
 			})
 	);
 
@@ -590,12 +682,12 @@ TIoStatusOr<uint64> FOnDemandIoStore::GetInstallSize(const FOnDemandGetInstallSi
 
 FOnDemandChunkInfo FOnDemandIoStore::GetStreamingChunkInfo(const FIoChunkId& ChunkId)
 {
-	return GetChunkInfo(ChunkId, EOnDemandContainerFlags::Mounted | EOnDemandContainerFlags::Streaming);
+	return GetChunkInfo(ChunkId, EOnDemandContainerFlags::Mounted | EOnDemandContainerFlags::StreamOnDemand);
 }
 
 FOnDemandChunkInfo FOnDemandIoStore::GetInstalledChunkInfo(const FIoChunkId& ChunkId)
 {
-	return GetChunkInfo(ChunkId, EOnDemandContainerFlags::Mounted | EOnDemandContainerFlags::Installed);
+	return GetChunkInfo(ChunkId, EOnDemandContainerFlags::Mounted | EOnDemandContainerFlags::InstallOnDemand);
 }
 
 void FOnDemandIoStore::OnPostFork(EForkProcessRole ProcessRole)
@@ -715,7 +807,7 @@ FOnDemandChunkInfo FOnDemandIoStore::GetChunkInfo(const FIoChunkId& ChunkId, EOn
 	{
 		if (EnumHasAllFlags(Container->Flags, ContainerFlags))
 		{
-			if (const FOnDemandChunkEntry* Entry = Container->ChunkEntries.Find(ChunkId))
+			if (const FOnDemandChunkEntry* Entry = Container->FindChunkEntry(ChunkId))
 			{
 				return FOnDemandChunkInfo(Container, *Entry);
 			}
@@ -776,207 +868,174 @@ void FOnDemandIoStore::TickLoop()
 
 bool FOnDemandIoStore::Tick()
 {
-	TArray<FSharedMountRequest> Requests;
+	TArray<FSharedMountRequest> LocalMountRequests;
 	{
 		UE::TUniqueLock Lock(MountRequestMutex);
 		for (TPair<FString, FSharedMountRequest>& Kv : MountRequests)
 		{
-			Requests.Add(Kv.Value);
+			LocalMountRequests.Add(Kv.Value);
 		}
 	}
 
-	if (Requests.IsEmpty())
-	{
-		return false;
-	}
+	bool bTicked = LocalMountRequests.IsEmpty() == false;
 
-	for (FSharedMountRequest& Request : Requests)
+	// Tick mount request(s)
+	for (FSharedMountRequest& Request : LocalMountRequests)
 	{
-		// Tick mount
-		if (FIoStatus Status = TickMountRequest(*Request); !Status.IsOk())
 		{
-			{
-				UE::TUniqueLock Lock(MountRequestMutex);
-				MountRequests.Remove(Request->MountArgs.MountId);
-			}
-			FOnDemandMountCompleted OnCompleted = MoveTemp(Request->OnCompleted);
-			OnCompleted(FOnDemandMountResult{ Request->MountArgs.MountId, Status });
-			continue;
+			UE::TUniqueLock Lock(MountRequestMutex);
+			MountRequests.Remove(Request->Args.MountId);
 		}
 
-		// Retrieve encryption key(s)
-		bool bPendingEncryptionKeys = false;
-		for (const FSharedOnDemandContainer& Container : Request->Containers)
+		if (FIoStatus Status = TickMountRequest(*Request); !Status.IsOk())
 		{
-			if (EnumHasAnyFlags(Request->MountArgs.Options, EOnDemandMountOptions::StreamOnDemand))
+			FOnDemandMountResult MountResult
 			{
-				EnumAddFlags(Container->Flags, EOnDemandContainerFlags::Streaming);
-			}
-
-			if (EnumHasAnyFlags(Container->Flags, EOnDemandContainerFlags::Encrypted) &&
-				Container->EncryptionKey.IsValid() == false)
-			{
-				FGuid KeyGuid;
-				ensure(FGuid::Parse(Container->EncryptionKeyGuid, KeyGuid));
-				if (FEncryptionKeyManager::Get().TryGetKey(KeyGuid, Container->EncryptionKey))
-				{
-					EnumRemoveFlags(Container->Flags, EOnDemandContainerFlags::PendingEncryptionKey);
-				}
-				else
-				{
-					EnumAddFlags(Container->Flags, EOnDemandContainerFlags::PendingEncryptionKey);
-					bPendingEncryptionKeys = true;
-				}
-			}
-		}
+				.Status = MoveTemp(Status),
+				.DurationInSeconds = Request->DurationInSeconds 
+			};
 
-		// Tick install
-		if (EnumHasAnyFlags(Request->MountArgs.Options, EOnDemandMountOptions::Install))
-		{
-			if (bPendingEncryptionKeys)
+			if (EnumHasAnyFlags(Request->Args.Options, EOnDemandMountOptions::CallbackOnGameThread))
 			{
-				UE_LOG(LogIoStoreOnDemand, Log, TEXT("Deferring install request until encryption key(s) becomes available"));
-				continue;
+				ExecuteOnGameThread(
+					UE_SOURCE_LOCATION,
+					[OnCompleted = MoveTemp(Request->OnCompleted), MountResult = MoveTemp(MountResult)]() mutable
+					{
+						OnCompleted(MoveTemp(MountResult));
+					});
 			}
-
-			if (FIoStatus Status = TickInstallRequest(*Request); !Status.IsOk())
+			else
 			{
-				// Unknown == Pending
-				if (Status.GetErrorCode() == EIoErrorCode::Unknown)
-				{
-					UE_LOG(LogIoStoreOnDemand, Log, TEXT("%s"), *Status.ToString());
-					continue;
-				}
-
-				{
-					UE::TUniqueLock Lock(MountRequestMutex);
-					MountRequests.Remove(Request->MountArgs.MountId);
-				}
 				FOnDemandMountCompleted OnCompleted = MoveTemp(Request->OnCompleted);
-				OnCompleted(FOnDemandMountResult{ Request->MountArgs.MountId, Status });
-				continue;
+				OnCompleted(MoveTemp(MountResult));
 			}
+			continue;
 		}
 
-		// Complete mount/install request
+		TStringBuilder<128> Sb;
 		{
-			{
-				UE::TUniqueLock Lock(MountRequestMutex);
-				MountRequests.Remove(Request->MountArgs.MountId);
-			}
-
-			TStringBuilder<128> Sb;
 			UE::TUniqueLock Lock(ContainerMutex);
 			for (const FSharedOnDemandContainer& Container : Request->Containers)
 			{
 				if (EnumHasAnyFlags(Container->Flags, EOnDemandContainerFlags::Mounted))
 				{
+					UE_LOG(LogIoStoreOnDemand, Log, TEXT("Ignoring already mounted container, ContainerName='%s'"), *Container->Name);
 					continue;
 				}
 
 				Containers.Add(Container);
 
-				// To conform with PAK mount behavior, requests are completed even if encryption keys are pending
-				if (EnumHasAnyFlags(Container->Flags, EOnDemandContainerFlags::PendingEncryptionKey))
+				// To conform with PAK mount behavior, mount requests are completed even if encryption keys are pending
+				if (EnumHasAnyFlags(Container->Flags, EOnDemandContainerFlags::Encrypted) &&
+					Container->EncryptionKey.IsValid() == false)
 				{
-					UE_LOG(LogIoStoreOnDemand, Log, TEXT("Deferring container '%s' until encryption key '%s' becomes available"),
-						*Container->Name, *Container->EncryptionKeyGuid);
+					FGuid KeyGuid;
+					ensure(FGuid::Parse(Container->EncryptionKeyGuid, KeyGuid));
+					if (FEncryptionKeyManager::Get().TryGetKey(KeyGuid, Container->EncryptionKey) == false)
+					{
+						EnumAddFlags(Container->Flags, EOnDemandContainerFlags::PendingEncryptionKey);
+						UE_LOG(LogIoStoreOnDemand, Log, TEXT("Deferring container '%s' until encryption key '%s' becomes available"),
+							*Container->Name, *Container->EncryptionKeyGuid);
+					}
 				}
-				else
+
+				if (EnumHasAnyFlags(Container->Flags, EOnDemandContainerFlags::PendingEncryptionKey) == false)
 				{
+					EnumAddFlags(Container->Flags, EOnDemandContainerFlags::Mounted);
+					Container->EncryptionKeyGuid.Empty();
 					Sb.Reset();
 					LexToString(Container->Flags, Sb);
 					UE_LOG(LogIoStoreOnDemand, Log, TEXT("Mounting container '%s', Entries=%d, Flags='%s'"),
 						*Container->Name, Container->ChunkEntries.Num(), Sb.ToString());
-
-					EnumAddFlags(Container->Flags, EOnDemandContainerFlags::Mounted);
 				}
 			}
 		}
 
-		FOnDemandMountCompleted OnCompleted = MoveTemp(Request->OnCompleted);
-		OnCompleted(FOnDemandMountResult{ Request->MountArgs.MountId });
-	}
-
-	return true;
-}
+		FOnDemandMountResult MountResult
+		{
+			.MountId = MoveTemp(Request->Args.MountId),
+			.DurationInSeconds = Request->DurationInSeconds
+		};
 
-FIoStatus FOnDemandIoStore::TickMountRequest(FMountRequest& MountRequest)
-{
-	UE_LOG(LogIoStoreOnDemand, Verbose, TEXT("Ticking mount request, MountId='%s'"), *MountRequest.MountArgs.MountId);
+		if (EnumHasAnyFlags(Request->Args.Options, EOnDemandMountOptions::CallbackOnGameThread))
+		{
+			ExecuteOnGameThread(
+				UE_SOURCE_LOCATION,
+				[OnCompleted = MoveTemp(Request->OnCompleted), MountResult = MoveTemp(MountResult)]() mutable
+				{
+					OnCompleted(MoveTemp(MountResult));
+				});
+		}
+		else
+		{
+			FOnDemandMountCompleted OnCompleted = MoveTemp(Request->OnCompleted);
+			OnCompleted(MoveTemp(MountResult));
+		}
+	}
 
-	FOnDemandMountArgs& Args = MountRequest.MountArgs;
-	FStringView Host, TocRelUrl;
-	Private::SplitHostUrl(Args.Url, Host, TocRelUrl);
-	const FStringView TocPath = FPathViews::GetPath(TocRelUrl);
+	TArray<FSharedInstallRequest> LocalInstallRequests;
+	{
+		UE::TUniqueLock Lock(MountRequestMutex);
+		LocalInstallRequests = InstallRequests;
+	}
 
-	if (!EnumHasAnyFlags(Args.Options, EOnDemandMountOptions::InstallSkipMount))
+	// Tick install request(s)
+	for (FSharedInstallRequest& Request : LocalInstallRequests)
 	{
-		if (Args.Toc)
+		FIoStatus Status = TickInstallRequest(*Request);
+		if (Status.GetErrorCode() == EIoErrorCode::Unknown)
 		{
-			CreateContainersFromToc(Args.MountId, TocPath, *Args.Toc, MountRequest.Containers);
+			UE_LOG(LogIoStoreOnDemand, Log, TEXT("Deferring install request, reason '%s'"), *Status.ToString());
+			continue;
 		}
-		else if (Args.FilePath.IsEmpty() == false)
-		{
-			UE_LOG(LogIoStoreOnDemand, Log, TEXT("Loading TOC from file '%s'"), *Args.FilePath);
 
-			// TODO: Enable validation when the sentinal is included in all serialization paths
-			const bool bValidate = false;
-			TIoStatusOr<FOnDemandToc> TocStatus = FOnDemandToc::LoadFromFile(Args.FilePath, bValidate);
-			if (!TocStatus.IsOk())
-			{
-				return TocStatus.Status();
-			}
+		{
+			UE::TUniqueLock Lock(MountRequestMutex);
+			InstallRequests.Remove(Request);
+		}
 
-			Args.Toc = MakeUnique<FOnDemandToc>(TocStatus.ConsumeValueOrDie());
+		FOnDemandInstallResult InstallResult
+		{
+			.Status = MoveTemp(Status),
+			.DurationInSeconds = Request->DurationInSeconds,
+			.TotalContentSize = Request->TotalContentSize,
+			.TotalInstallSize = Request->TotalInstallSize
+		};
 
-			CreateContainersFromToc(Args.MountId, TocPath, *Args.Toc, MountRequest.Containers);
+		if (EnumHasAnyFlags(Request->Args.Options, EOnDemandInstallOptions::CallbackOnGameThread))
+		{
+			ExecuteOnGameThread(
+				UE_SOURCE_LOCATION,
+				[OnCompleted = MoveTemp(Request->OnCompleted), InstallResult = MoveTemp(InstallResult)]() mutable
+				{
+					OnCompleted(MoveTemp(InstallResult));
+				});
 		}
-		else if (Args.Url.IsEmpty() == false)
+		else
 		{
-			UE_LOG(LogIoStoreOnDemand, Log, TEXT("Loading TOC from URL '%s'"), *Args.Url);
-
-			const uint32 RetryCount = 2;
-			const bool bFollowRedirects = true;
-			TIoStatusOr<FOnDemandToc> TocStatus = FOnDemandToc::LoadFromUrl(Args.Url, 2, bFollowRedirects);
+			FOnDemandInstallCompleted OnCompleted = MoveTemp(Request->OnCompleted);
+			OnCompleted(MoveTemp(InstallResult));
+		}
 
-			if (!TocStatus.IsOk())
-			{
-				return TocStatus.Status();
-			}
+		bTicked = true;
+	}
 
-			Args.Toc = MakeUnique<FOnDemandToc>(TocStatus.ConsumeValueOrDie());
+	return bTicked;
+}
 
-			CreateContainersFromToc(Args.MountId, TocPath, *Args.Toc, MountRequest.Containers);
-		}
-	}
+FIoStatus FOnDemandIoStore::TickMountRequest(FMountRequest& MountRequest)
+{
+	UE_LOG(LogIoStoreOnDemand, Verbose, TEXT("Ticking mount request, MountId='%s'"), *MountRequest.Args.MountId);
 
-	// Remove already mounted containers
-	// TODO: Treat this as an error?
+	const double StartTime = FPlatformTime::Seconds();
+	ON_SCOPE_EXIT
 	{
-		MountRequest.Containers.SetNum(
-			Algo::RemoveIf(
-				MountRequest.Containers,
-				[MountedContainers = GetMountedContainers()](const FSharedOnDemandContainer& Container)
-				{
-					const FSharedOnDemandContainer* Existing =
-						Algo::FindBy(
-							MountedContainers,
-							Container->UniqueName(),
-							[](const FSharedOnDemandContainer& C) { return C->UniqueName(); });
+		MountRequest.DurationInSeconds = FPlatformTime::Seconds() - StartTime;
+	};
 
-					if (Existing != nullptr)
-					{
-						UE_LOG(LogIoStoreOnDemand, Warning, TEXT("Container already mounted, ContainerName='%s', MountId='%s'"),
-							*WriteToString<128>(Container->Name), *WriteToString<128>(Container->MountId));
-						return true;
-					}
-					return false;
-				}));
-	}
+	FOnDemandMountArgs& Args = MountRequest.Args;
 
-	// Find containers matching the mount ID if this is a request to install/download already mounted containers 
-	if (MountRequest.Containers.IsEmpty())
+	// Find containers matching the mount ID
 	{
 		UE::TUniqueLock Lock(ContainerMutex);
 		for (FSharedOnDemandContainer& Container : Containers)
@@ -988,117 +1047,112 @@ FIoStatus FOnDemandIoStore::TickMountRequest(FMountRequest& MountRequest)
 		}
 	}
 
-	return EIoErrorCode::Ok;
-}
-
-FIoStatus FOnDemandIoStore::TickInstallRequest(FMountRequest& MountRequest)
-{
-	if (InstallCache.IsValid() == false || PackageStoreBackend.IsValid() == false)
-	{
-		if (FForkProcessHelper::IsForkRequested() && !FForkProcessHelper::IsForkedChildProcess())
-		{
-			return FIoStatusBuilder(EIoErrorCode::PendingFork) << TEXT("Install cache waiting for fork");
-		}
-
-		return FIoStatusBuilder(EIoErrorCode::InvalidCode) << TEXT("Install cache not configured");
-	}
-
-	if (MountRequest.bCancelled)
+	// TODO: Treat this as an error?
+	if (MountRequest.Containers.IsEmpty() == false)
 	{
-		return EIoErrorCode::Cancelled;
+		return FIoStatus::Ok;
 	}
 
-	TAnsiStringBuilder<512> ChunkUrl;
 	FStringView Host, TocRelUrl;
-	Private::SplitHostUrl(MountRequest.MountArgs.Url, Host, TocRelUrl);
+	Private::SplitHostUrl(Args.Url, Host, TocRelUrl);
+	const FStringView TocPath = FPathViews::GetPath(TocRelUrl);
 
-	auto GetChunkUrl = [](
-		const FStringView& Host,
-		const FOnDemandContainer& Container,
-		const FOnDemandChunkEntry& Entry,
-		FAnsiStringBuilderBase& OutUrl) -> FAnsiStringBuilderBase&
+	if (Args.Toc)
 	{
-		OutUrl.Reset();
-		if (Host.IsEmpty() == false)
-		{
-			OutUrl << Host;
-		}
+		CreateContainersFromToc(Args.MountId, TocPath, *Args.Toc, MountRequest.Containers);
+	}
+	else if (Args.FilePath.IsEmpty() == false)
+	{
+		UE_LOG(LogIoStoreOnDemand, Log, TEXT("Loading TOC from file '%s'"), *Args.FilePath);
 
-		if (!Container.ChunksDirectory.IsEmpty())
+		// TODO: Enable validation when the sentinal is included in all serialization paths
+		const bool bValidate = false;
+		TIoStatusOr<FOnDemandToc> TocStatus = FOnDemandToc::LoadFromFile(Args.FilePath, bValidate);
+		if (!TocStatus.IsOk())
 		{
-			OutUrl << "/" << Container.ChunksDirectory;
+			return TocStatus.Status();
 		}
 
-		const FString HashString = LexToString(Entry.Hash);
-		OutUrl << "/" << HashString.Left(2) << "/" << HashString << ".iochunk";
+		Args.Toc = MakeUnique<FOnDemandToc>(TocStatus.ConsumeValueOrDie());
 
-		return OutUrl;
-	};
-
-	auto FetchContainerHeader = [&Host, &GetChunkUrl, &ChunkUrl](
-		FSharedOnDemandContainer Container,
-		FIoContainerHeader& Header) -> FIoStatus 
+		CreateContainersFromToc(Args.MountId, TocPath, *Args.Toc, MountRequest.Containers);
+	}
+	else if (Args.Url.IsEmpty() == false)
 	{
-		const FIoChunkId			ChunkId = CreateContainerHeaderChunkId(Container->ContainerId);
-		const FOnDemandChunkEntry*	Entry = Container->ChunkEntries.Find(ChunkId);
+		UE_LOG(LogIoStoreOnDemand, Log, TEXT("Loading TOC from URL '%s'"), *Args.Url);
 
-		if (Entry == nullptr)
+		const uint32 RetryCount = 2;
+		const bool bFollowRedirects = true;
+		TIoStatusOr<FOnDemandToc> TocStatus = FOnDemandToc::LoadFromUrl(Args.Url, 2, bFollowRedirects);
+
+		if (!TocStatus.IsOk())
 		{
-			return EIoErrorCode::Ok;
+			return TocStatus.Status();
 		}
 
-		UE_LOG(LogIoStoreOnDemand, Verbose, TEXT("Fetching container header, ContainerName='%s'"), *Container->Name);
-		TIoStatusOr<FIoBuffer> Response = FHttpClient::Get(GetChunkUrl(Host, *Container, *Entry, ChunkUrl).ToView(), 2, EHttpRedirects::Follow);
+		Args.Toc = MakeUnique<FOnDemandToc>(TocStatus.ConsumeValueOrDie());
+
+		CreateContainersFromToc(Args.MountId, TocPath, *Args.Toc, MountRequest.Containers);
+	}
 
-		if (Response.IsOk() == false)
+	for (const FSharedOnDemandContainer& Container : MountRequest.Containers)
+	{
+		if (EnumHasAnyFlags(Args.Options, EOnDemandMountOptions::StreamOnDemand))
 		{
-			FIoStatus Status = FIoStatusBuilder(EIoErrorCode::ReadError) << TEXT("Failed to fetch container header chunk");
-			return Status;
+			EnumAddFlags(Container->Flags, EOnDemandContainerFlags::StreamOnDemand);
 		}
-
-		FOnDemandChunkInfo ChunkInfo(Container, *Entry);
-
-		FIoChunkDecodingParams Params;
-		Params.CompressionFormat	= ChunkInfo.CompressionFormat();
-		Params.EncryptionKey		= ChunkInfo.EncryptionKey();
-		Params.BlockSize			= ChunkInfo.BlockSize();
-		Params.TotalRawSize			= ChunkInfo.RawSize();
-		Params.RawOffset			= 0; 
-		Params.EncodedOffset		= 0; 
-		Params.EncodedBlockSize		= ChunkInfo.Blocks();
-		Params.BlockHash			= ChunkInfo.BlockHashes();
-
-		FIoBuffer EncodedChunk = Response.ConsumeValueOrDie();
-		FIoBuffer RawChunk(ChunkInfo.RawSize());
-		if (FIoChunkEncoding::Decode(Params, EncodedChunk.GetView(), RawChunk.GetMutableView()) == false)
+		else if (EnumHasAnyFlags(Args.Options, EOnDemandMountOptions::InstallOnDemand))
 		{
-			FIoStatus Status = FIoStatusBuilder(EIoErrorCode::ReadError)
-				<< TEXT("Failed to decode container header chunk");
+			EnumAddFlags(Container->Flags, EOnDemandContainerFlags::InstallOnDemand);
 		}
+	}
+
+	return EIoErrorCode::Ok;
+}
+
+FIoStatus FOnDemandIoStore::TickInstallRequest(FInstallRequest& InstallRequest)
+{
+	UE_LOG(LogIoStoreOnDemand, Log, TEXT("Ticking install request, ContentHandle='%s'"),
+		*LexToString(InstallRequest.Args.ContentHandle));
+
+	const double StartTime = FPlatformTime::Seconds();
+	ON_SCOPE_EXIT
+	{
+		InstallRequest.DurationInSeconds = FPlatformTime::Seconds() - StartTime;
+	};
 
-		FMemoryReaderView Ar(RawChunk.GetView());
-		Ar << Header;
-		Ar.Close();
+	if (InstallRequest.Args.ContentHandle.IsValid() == false)
+	{
+		return FIoStatusBuilder(EIoErrorCode::InvalidParameter) << TEXT("Invalid content handle");
+	}
 
-		if (Ar.IsError() || Ar.IsCriticalError())
+	if (InstallCache.IsValid() == false || PackageStoreBackend.IsValid() == false)
+	{
+		if (FForkProcessHelper::IsForkRequested() && !FForkProcessHelper::IsForkedChildProcess())
 		{
-			FIoStatus Status = FIoStatusBuilder(EIoErrorCode::FileNotOpen) << TEXT("Failed to serialize container header");
-			return Status;
+			return FIoStatusBuilder(EIoErrorCode::PendingFork) << TEXT("Install cache waiting for fork");
 		}
 
-		return EIoErrorCode::Ok;
-	};
+		return FIoStatusBuilder(EIoErrorCode::InvalidCode) << TEXT("Install cache not configured");
+	}
 
-	TSet<FSharedOnDemandContainer> AllContainers(
-		GetMountedContainers([](const FSharedOnDemandContainer& Container) 
+	//TODO: Implement cancellation
+
+	TAnsiStringBuilder<512> ChunkUrl;
+	FStringView Host, TocRelUrl;
+	Private::SplitHostUrl(InstallRequest.Args.Url, Host, TocRelUrl);
+
+	// Only install content from non-streaming container(s)
+	TSet<FSharedOnDemandContainer> ContainersForInstallation(
+		GetContainers([](const FSharedOnDemandContainer& Container)
 		{ 
-			return EnumHasAnyFlags(Container->Flags, EOnDemandContainerFlags::Installed); 
+			return EnumHasAnyFlags(Container->Flags, EOnDemandContainerFlags::InstallOnDemand);
 		})
 	);
-	Algo::Copy(MountRequest.Containers, AllContainers);
 
+	// TODO: move this to mount
 	// Fetch all container headers
-	for (const FSharedOnDemandContainer& Container : AllContainers)
+	for (const FSharedOnDemandContainer& Container : ContainersForInstallation)
 	{
 		if (EnumHasAnyFlags(Container->Flags, EOnDemandContainerFlags::PendingEncryptionKey))
 		{
@@ -1109,7 +1163,7 @@ FIoStatus FOnDemandIoStore::TickInstallRequest(FMountRequest& MountRequest)
 		if (!Container->Header.IsValid())
 		{
 			FSharedContainerHeader Header = MakeShared<FIoContainerHeader>();
-			if (FIoStatus Status = FetchContainerHeader(Container, *Header); !Status.IsOk())
+			if (FIoStatus Status = Private::FetchContainerHeader(Host, ChunkUrl, Container, *Header); !Status.IsOk())
 			{
 				return Status; 
 			}
@@ -1117,61 +1171,19 @@ FIoStatus FOnDemandIoStore::TickInstallRequest(FMountRequest& MountRequest)
 		}
 	}
 
-	// Parse the tag sets
-	if (MountRequest.MountArgs.TagSets.IsEmpty() == false && MountRequest.TagSets.IsEmpty())
+	TSet<FPackageId> PackageIdsToInstall;
+	for (const FPackageId& PackageId : InstallRequest.Args.PackageIds)
 	{
-		const FOnDemandToc& Toc = *MountRequest.MountArgs.Toc;
-		for (const FString& Tag : MountRequest.MountArgs.TagSets)
-		{
-			const FOnDemandTocTagSet* TagSet = Algo::FindBy(Toc.TagSets, Tag, &FOnDemandTocTagSet::Tag);
-			if (TagSet == nullptr)
-			{
-				UE_LOG(LogIoStoreOnDemand, Warning, TEXT("Tag set '%s' doesn't exists"), *Tag);
-				continue;
-			}
-
-			FTagSet NewTagSet;
-			NewTagSet.Tag = TagSet->Tag;
-			NewTagSet.PackageIds.Reserve(Algo::TransformAccumulate(
-				TagSet->Packages,
-				[](const FOnDemandTocTagSetPackageList& PackageIndices) -> int32 { return PackageIndices.PackageIndicies.Num(); },
-				int32(0))
-			);
-			for (const FOnDemandTocTagSetPackageList& PackageIndices : TagSet->Packages)
-			{
-				if (MountRequest.Containers.IsEmpty() || int32(PackageIndices.ContainerIndex) > (MountRequest.Containers.Num() - 1))
-				{
-					UE_LOG(LogIoStoreOnDemand, Warning, TEXT("Invalid container index in tag set '%s'"), *Tag);
-					continue;
-				}
-
-				FOnDemandContainer& Container = *MountRequest.Containers[PackageIndices.ContainerIndex];
-				if (Container.Header.IsValid() == false)
-				{
-					UE_LOG(LogIoStoreOnDemand, Warning, TEXT("Tag set specified for invalid container header, Tag='%s', ContainerName='%s'"),
-						*Tag, *Container.UniqueName());
-					continue;
-				}
-
-				for (int32 IdIdx : PackageIndices.PackageIndicies)
-				{
-					NewTagSet.PackageIds.Add(Container.Header->PackageIds[IdIdx]);
-				}
-			}
-
-			MountRequest.TagSets.Emplace(MoveTemp(NewTagSet));
-		}
+		PackageIdsToInstall.Add(PackageId);
 	}
 
-	TSet<FPackageId> PackageIdsToInstall;
-	// Install all packages if no tag set was specified
-	if (MountRequest.MountArgs.TagSets.IsEmpty())
+	// Install all packages if no tag set(s) was specified
+	if (InstallRequest.Args.TagSets.IsEmpty())
 	{
-		for (const FSharedOnDemandContainer& Container : MountRequest.Containers)
+		for (const FSharedOnDemandContainer& Container : ContainersForInstallation)
 		{
-			if (Container->Header.IsValid())
+			if (Container->Header.IsValid() && Container->MountId == InstallRequest.Args.MountId)
 			{
-				PackageIdsToInstall.Reserve(Container->Header->PackageIds.Num());
 				for (const FPackageId& PackageId : Container->Header->PackageIds)
 				{
 					PackageIdsToInstall.Add(PackageId);
@@ -1181,34 +1193,48 @@ FIoStatus FOnDemandIoStore::TickInstallRequest(FMountRequest& MountRequest)
 	}
 	else
 	{
-		for (const FString& Tag : MountRequest.MountArgs.TagSets)
+		for (const FSharedOnDemandContainer& Container : ContainersForInstallation)
 		{
-			const FTagSet* TagSet = nullptr;
-			for (const FTagSet& Set : MountRequest.TagSets)
+			if (Container->Header.IsValid() == false)
 			{
-				if (Set.Tag == Tag)
+				continue;
+			}
+
+			if (InstallRequest.Args.MountId.IsEmpty() == false && InstallRequest.Args.MountId != Container->MountId)
+			{
+				continue;
+			}
+
+			for (const FString& Tag : InstallRequest.Args.TagSets)
+			{
+				for (const FOnDemandTagSet& TagSet : Container->TagSets)
 				{
-					PackageIdsToInstall.Reserve(Set.PackageIds.Num());
-					for (const FPackageId& PackageId : Set.PackageIds)
+					if (TagSet.Tag == Tag)
 					{
-						PackageIdsToInstall.Add(PackageId);
+						for (const uint32 PackageIndex : TagSet.PackageIndicies)
+						{
+							const FPackageId PackageId = Container->Header->PackageIds[IntCastChecked<int32>(PackageIndex)];
+							PackageIdsToInstall.Add(PackageId);
+						}
 					}
 				}
 			}
 		}
 	}
 
+	// Its OK for PackageIdsToInstall to be empty at this point. Any chunks not referenced by a package must still be installed.
+
 	// Find all I/O chunks fo the specified list of packages
 	Private::FInstallData InstallData;
 	TSet<FPackageId> Missing;
 
-	FIoStatus Status = BuildInstallData(AllContainers, PackageIdsToInstall, InstallData, Missing);
+	FIoStatus Status = BuildInstallData(ContainersForInstallation, PackageIdsToInstall, InstallData, Missing);
 	if (Status.IsOk() == false)
 	{
 		return Status;
 	}
 
-	// Check the other I/O backends for missing packge chunks
+	// Check the other I/O backends for missing package chunks
 	for (const FPackageId& PackageId : Missing)
 	{
 		const FIoChunkId ChunkId = CreatePackageDataChunkId(PackageId);
@@ -1227,9 +1253,9 @@ FIoStatus FOnDemandIoStore::TickInstallRequest(FMountRequest& MountRequest)
 			FOnDemandContainer&						Container = *Kv.Key;
 			const Private::FContainerInstallData&	Data = Kv.Value;
 
-			for (const FIoChunkId& ChunkId : Data.ResolvedChunks)
+			for (int32 EntryIndex : Data.ResolvedChunks)
 			{
-				const FOnDemandChunkEntry& Entry = Container.ChunkEntries.FindRef(ChunkId);
+				const FOnDemandChunkEntry& Entry = Container.ChunkEntries[EntryIndex];
 				ChunksToInstall.Add(Entry.Hash, Entry.EncodedSize);
 			}
 		}
@@ -1265,24 +1291,25 @@ FIoStatus FOnDemandIoStore::TickInstallRequest(FMountRequest& MountRequest)
 		FOnDemandContainer& Container				= *Kv.Key;
 		const Private::FContainerInstallData& Data	= Kv.Value;
 
-		for (const FIoChunkId& ChunkId : Data.ResolvedChunks)
+		for (int32 EntryIndex : Data.ResolvedChunks)
 		{
-			FOnDemandChunkEntry* ChunkEntry = Container.ChunkEntries.Find(ChunkId);
+			const FIoChunkId& ChunkId				= Container.ChunkIds[EntryIndex];
+			const FOnDemandChunkEntry& ChunkEntry	= Container.ChunkEntries[EntryIndex];
 
 			++TotalChunkCount;
-			TotalBytes += ChunkEntry->EncodedSize;
+			TotalBytes += ChunkEntry.EncodedSize;
 
-			if (InstallCache->IsChunkCached(ChunkEntry->Hash))
+			if (InstallCache->IsChunkCached(ChunkEntry.Hash))
 			{
 				continue;
 			}
 
 			++DownloadedChunkCount;
-			DownloadedBytes += ChunkEntry->EncodedSize;
+			DownloadedBytes += ChunkEntry.EncodedSize;
 
 			ConcurrentRequests++;
 			HttpClient->Get(
-				GetChunkUrl(FStringView(), Container, *ChunkEntry, ChunkUrl).ToView(),
+				Private::GetChunkUrl(FStringView(), Container, ChunkEntry, ChunkUrl).ToView(),
 				[this, &Status, ChunkId, ChunkEntry, &ConcurrentRequests]
 				(TIoStatusOr<FIoBuffer> ChunkStatus, uint64 DurationMs) mutable
 				{
@@ -1295,7 +1322,7 @@ FIoStatus FOnDemandIoStore::TickInstallRequest(FMountRequest& MountRequest)
 
 					FIoBuffer Chunk			= ChunkStatus.ConsumeValueOrDie();
 					const FIoHash ChunkHash = FIoHash::HashBuffer(Chunk.GetView());
-					if (ChunkHash != ChunkEntry->Hash)
+					if (ChunkHash != ChunkEntry.Hash)
 					{
 						Status = FIoStatus(EIoErrorCode::ReadError, TEXTVIEW("Hash mismatch"));
 						return;
@@ -1320,7 +1347,14 @@ FIoStatus FOnDemandIoStore::TickInstallRequest(FMountRequest& MountRequest)
 	
 	if (Status = InstallCache->Flush(); Status.IsOk() == false)
 	{
-		return FIoStatusBuilder(EIoErrorCode::WriteError) << TEXT("Failed to flush downloded content to cache");
+		return FIoStatusBuilder(EIoErrorCode::WriteError) << TEXT("Failed to flush downloaded content to cache");
+	}
+
+	TSharedPtr<FOnDemandInternalContentHandle, ESPMode::ThreadSafe>& ContentHandle = InstallRequest.Args.ContentHandle.Handle;
+	if (ContentHandle->IoStore.IsValid() == false)
+	{
+		// First time this content handle is used
+		ContentHandle->IoStore = AsWeak(); 
 	}
 
 	// TODO: Only mount what has been installed
@@ -1335,9 +1369,19 @@ FIoStatus FOnDemandIoStore::TickInstallRequest(FMountRequest& MountRequest)
 			check(MountStatus.IsOk());
 		}
 
-		EnumAddFlags(Container.Flags, EOnDemandContainerFlags::Installed);
+		{
+			UE::TUniqueLock Lock(ContainerMutex);
+			FOnDemandChunkEntryReferences& References = Container.FindOrAddChunkEntryReferences(*ContentHandle);
+			for (int32 EntryIndex : Data.ResolvedChunks)
+			{
+				References.Indices[EntryIndex] = true;
+			}
+		}
 	}
 
+	InstallRequest.TotalContentSize = TotalBytes;
+	InstallRequest.TotalInstallSize = DownloadedBytes;
+
 	UE_LOG(LogIoStoreOnDemand, Log, TEXT("Downloaded %llu (%.2lf MiB) of total %llu (%.2lf MiB) requested chunk(s)'"),
 		DownloadedChunkCount, double(DownloadedBytes) / 1024.0 / 1024.0, TotalChunkCount, double(TotalBytes) / 1024.0 / 1024.0);
 
@@ -1363,12 +1407,13 @@ void FOnDemandIoStore::OnEncryptionKeyAdded(const FGuid& Id, const FAES::FAESKey
 
 					EnumRemoveFlags(Container->Flags, EOnDemandContainerFlags::PendingEncryptionKey);
 					EnumAddFlags(Container->Flags, EOnDemandContainerFlags::Mounted);
+					Container->EncryptionKeyGuid.Empty();
 				}
 			}
 		}
 	}
 
-	// Tick pending mount/install request waiting on encryption key(s)
+	// Tick pending install request waiting on encryption key(s)
 	TryEnterTickLoop();
 }
 
@@ -1426,12 +1471,25 @@ void FOnDemandIoStore::CreateContainersFromToc(
 			EnumAddFlags(Container->Flags, EOnDemandContainerFlags::Encrypted);
 		}
 
-		Container->ChunkEntries.Reserve(ContainerEntry.Entries.Num());
-		for (const FOnDemandTocEntry& TocEntry : ContainerEntry.Entries)
+		//TODO: Sort before uploading
+		ContainerEntry.Entries.Sort([](const FOnDemandTocEntry& Lhs, const FOnDemandTocEntry& Rhs)
+		{
+			return Lhs.ChunkId < Rhs.ChunkId;
+		});
+
+		const int32	EntryCount		= ContainerEntry.Entries.Num();
+		const uint64 TotalEntrySize = (sizeof(FIoChunkId) + sizeof(FOnDemandTocEntry)) * EntryCount;
+		Container->ChunkEntryData	= MakeUnique<uint8[]>(TotalEntrySize);
+		Container->ChunkIds			= MakeArrayView<FIoChunkId>(reinterpret_cast<FIoChunkId*>(Container->ChunkEntryData.Get()), EntryCount);
+		Container->ChunkEntries		= MakeArrayView<FOnDemandChunkEntry>(
+			reinterpret_cast<FOnDemandChunkEntry*>(Container->ChunkIds.GetData() + EntryCount), EntryCount);
+
+		for (int32 EntryIndex = 0; EntryIndex < EntryCount; EntryIndex++)
 		{
-			check(TocEntry.RawSize <= 0xffff'ffffull);
-			check(TocEntry.EncodedSize <= 0xffff'ffffull);
-			Container->ChunkEntries.Add(TocEntry.ChunkId, FOnDemandChunkEntry
+			const FOnDemandTocEntry& TocEntry = ContainerEntry.Entries[EntryIndex];
+
+			Container->ChunkIds[EntryIndex]		= TocEntry.ChunkId;
+			Container->ChunkEntries[EntryIndex]	= FOnDemandChunkEntry
 			{
 				.Hash					= TocEntry.Hash,
 				.RawSize				= uint32(TocEntry.RawSize),
@@ -1439,20 +1497,35 @@ void FOnDemandIoStore::CreateContainersFromToc(
 				.BlockOffset			= TocEntry.BlockOffset,
 				.BlockCount				= TocEntry.BlockCount,
 				.CompressionFormatIndex	= 0
-			});
+			};
+		}
+
+		const uint32 ContainerIndex = Out.Num();
+		for (FOnDemandTocTagSet& TagSet : Toc.TagSets)
+		{
+			for (FOnDemandTocTagSetPackageList& ContainerPackageIndices : TagSet.Packages)
+			{
+				if (ContainerPackageIndices.ContainerIndex == ContainerIndex)
+				{
+					FOnDemandTagSet& NewTagSet = Container->TagSets.AddDefaulted_GetRef();
+					NewTagSet.Tag = TagSet.Tag;
+					NewTagSet.PackageIndicies = MoveTemp(ContainerPackageIndices.PackageIndicies);
+					break;
+				}
+			}
 		}
 
 		Out.Add(MoveTemp(Container));
 	}
 }
 
-TArray<FSharedOnDemandContainer> FOnDemandIoStore::GetMountedContainers()
+TArray<FSharedOnDemandContainer> FOnDemandIoStore::GetContainers()
 {
 	UE::TUniqueLock Lock(ContainerMutex);
 	return Containers;
 }
 
-TArray<FSharedOnDemandContainer> FOnDemandIoStore::GetMountedContainers(TFunctionRef<bool(const FSharedOnDemandContainer& Container)> Predicate)
+TArray<FSharedOnDemandContainer> FOnDemandIoStore::GetContainers(TFunctionRef<bool(const FSharedOnDemandContainer& Container)> Predicate)
 {
 	TArray<FSharedOnDemandContainer> Ret;
 
@@ -1462,4 +1535,92 @@ TArray<FSharedOnDemandContainer> FOnDemandIoStore::GetMountedContainers(TFunctio
 	return Ret;
 }
 
+void FOnDemandIoStore::ReleaseContent(FOnDemandInternalContentHandle& ContentHandle)
+{
+	UE_LOG(LogIoStoreOnDemand, Log, TEXT("Releasing content handle '%s'"), *LexToString(ContentHandle));
+
+	UE::TUniqueLock Lock(ContainerMutex);
+	for (FSharedOnDemandContainer& Container : Containers)
+	{
+		if (EnumHasAnyFlags(Container->Flags, EOnDemandContainerFlags::StreamOnDemand))
+		{
+			continue;
+		}
+
+		Container->ChunkEntryReferences.SetNum(
+			Algo::RemoveIf(
+				Container->ChunkEntryReferences,
+				[ContentHandleId = ContentHandle.HandleId()](const FOnDemandChunkEntryReferences& Refs)
+				{ 
+					return Refs.ContentHandleId == ContentHandleId;
+				}));
+	}
+}
+
+void FOnDemandIoStore::GetReferencedContent(TArray<FSharedOnDemandContainer>& OutContainers, TArray<TBitArray<>>& OutChunkEntryIndices)
+{
+	UE::TUniqueLock Lock(ContainerMutex);
+	for (FSharedOnDemandContainer& Container : Containers)
+	{
+		if (EnumHasAnyFlags(Container->Flags, EOnDemandContainerFlags::StreamOnDemand))
+		{
+			continue;
+		}
+
+		TBitArray<> Indices = Container->GetReferencedChunkEntries();
+		if (Indices.IsEmpty() == false)
+		{
+			OutContainers.Add(Container);
+			OutChunkEntryIndices.Add(MoveTemp(Indices));
+		}
+	}
+}
+
+///////////////////////////////////////////////////////////////////////////////
+FOnDemandInternalContentHandle::~FOnDemandInternalContentHandle()
+{
+	if (TSharedPtr<FOnDemandIoStore, ESPMode::ThreadSafe> Pinned = IoStore.Pin(); Pinned.IsValid())
+	{
+		Pinned->ReleaseContent(*this);
+	}
+}
+
+FString LexToString(const FOnDemandInternalContentHandle& Handle)
+{
+	return FString::Printf(TEXT("0x%llX (%s)"), Handle.HandleId(), *Handle.DebugName);
+}
+
+///////////////////////////////////////////////////////////////////////////////
+FOnDemandContentHandle::FOnDemandContentHandle()
+{
+}
+
+FOnDemandContentHandle::~FOnDemandContentHandle()
+{
+}
+
+FOnDemandContentHandle FOnDemandContentHandle::Create()
+{
+	FOnDemandContentHandle NewHandle;
+	NewHandle.Handle = MakeShared<FOnDemandInternalContentHandle, ESPMode::ThreadSafe>();
+	return NewHandle;
+}
+
+FOnDemandContentHandle FOnDemandContentHandle::Create(FSharedString DebugName)
+{
+	FOnDemandContentHandle NewHandle;
+	NewHandle.Handle = MakeShared<FOnDemandInternalContentHandle, ESPMode::ThreadSafe>(DebugName);
+	return NewHandle;
+}
+
+FOnDemandContentHandle FOnDemandContentHandle::Create(FStringView DebugName)
+{
+	return Create(FSharedString(DebugName));
+}
+
+FString LexToString(const FOnDemandContentHandle& Handle)
+{
+	return Handle.IsValid() ? LexToString(*Handle.Handle) : TEXT("Invalid");
+}
+
 } // namespace UE::IoStore
diff --git a/Engine/Source/Runtime/Experimental/IoStore/OnDemand/Private/OnDemandIoStore.h b/Engine/Source/Runtime/Experimental/IoStore/OnDemand/Private/OnDemandIoStore.h
--- a/Engine/Source/Runtime/Experimental/IoStore/OnDemand/Private/OnDemandIoStore.h
+++ b/Engine/Source/Runtime/Experimental/IoStore/OnDemand/Private/OnDemandIoStore.h
@@ -2,9 +2,11 @@
 
 #pragma once
 
+#include "Algo/BinarySearch.h"
 #include "Async/Async.h"
 #include "Async/Mutex.h"
 #include "Containers/AnsiString.h"
+#include "Containers/BitArray.h"
 #include "IO/IoStoreOnDemand.h"
 #include "IO/IoHash.h"
 #include "IO/IoChunkEncoding.h"
@@ -26,6 +28,27 @@ class IOnDemandPackageStoreBackend;
 class IOnDemandInstallCache;
 using FSharedPackageStoreBackend	= TSharedPtr<IOnDemandPackageStoreBackend>;
 using FSharedInstallCache			= TSharedPtr<IOnDemandInstallCache>;
+using FWeakOnDemandIoStore			= TWeakPtr<FOnDemandIoStore, ESPMode::ThreadSafe>;
+
+///////////////////////////////////////////////////////////////////////////////
+class FOnDemandInternalContentHandle
+{
+public:
+	FOnDemandInternalContentHandle()
+		: DebugName(TEXT("NoName"))
+	{ }
+	FOnDemandInternalContentHandle(FSharedString InDebugName)
+		: DebugName(InDebugName)
+	{ }
+	~FOnDemandInternalContentHandle();
+
+	UPTRINT HandleId() const { return UPTRINT(this); }
+
+	FSharedString			DebugName;
+	FWeakOnDemandIoStore	IoStore;
+};
+
+FString LexToString(const FOnDemandInternalContentHandle& Handle);
 
 ///////////////////////////////////////////////////////////////////////////////
 enum class EOnDemandContainerFlags : uint8
@@ -33,8 +56,8 @@ enum class EOnDemandContainerFlags : uint8
 	None					= 0,
 	PendingEncryptionKey	= (1 << 0),
 	Mounted					= (1 << 1),
-	Streaming				= (1 << 2),
-	Installed				= (1 << 3),
+	StreamOnDemand			= (1 << 2),
+	InstallOnDemand			= (1 << 3),
 	Encrypted				= (1 << 4),
 	Count
 };
@@ -55,33 +78,125 @@ struct FOnDemandChunkEntry
 	uint32	BlockCount = 0;
 	uint8	CompressionFormatIndex = 0;
 };
+static_assert(sizeof(FOnDemandChunkEntry) == 40);
+
+///////////////////////////////////////////////////////////////////////////////
+struct FOnDemandTagSet
+{
+	FString			Tag;
+	TArray<uint32>	PackageIndicies;
+};
+
+///////////////////////////////////////////////////////////////////////////////
+struct FOnDemandChunkEntryReferences
+{
+	UPTRINT ContentHandleId = 0;	
+	TBitArray<>	Indices;
+};
 
 ///////////////////////////////////////////////////////////////////////////////
 struct FOnDemandContainer
 {
-	using FChunkEntryMap = TMap<FIoChunkId, FOnDemandChunkEntry>;
-
-	FAES::FAESKey			EncryptionKey;
-	FIoHash					PakHash;
-	FChunkEntryMap 			ChunkEntries;
-	FSharedContainerHeader	Header;
-	FString					EncryptionKeyGuid;
-	FString					Name;
-	FString					MountId;
-	FAnsiString				ChunksDirectory;
-	TArray<FName>			CompressionFormats;
-	TArray<uint32>			BlockSizes;
-	TArray<FIoBlockHash>	BlockHashes;
-	FIoContainerId			ContainerId;
-	uint32					BlockSize = 0;
-	EOnDemandContainerFlags Flags = EOnDemandContainerFlags ::None;
-
-	FString					UniqueName() const;
+	FString									UniqueName() const;
+	inline int32							FindChunkEntryIndex(const FIoChunkId& ChunkId) const;
+	inline const FOnDemandChunkEntry*		FindChunkEntry(const FIoChunkId& ChunkId) const;
+	inline FOnDemandChunkEntry*				FindChunkEntry(const FIoChunkId& ChunkId);
+	inline FOnDemandChunkEntryReferences&	FindOrAddChunkEntryReferences(const FOnDemandInternalContentHandle& ContentHandle);
+	inline TBitArray<>						GetReferencedChunkEntries() const;
+
+	FAES::FAESKey							EncryptionKey;
+	FSharedContainerHeader					Header;
+	FString									EncryptionKeyGuid;
+	FString									Name;
+	FString									MountId;
+	FAnsiString								ChunksDirectory;
+	TArray<FName>							CompressionFormats;
+	TArray<uint32>							BlockSizes;
+	TArray<FIoBlockHash>					BlockHashes;
+	TArray<FOnDemandTagSet>					TagSets;
+	TUniquePtr<uint8[]>						ChunkEntryData;
+	TArrayView<FIoChunkId>					ChunkIds;
+	TArrayView<FOnDemandChunkEntry> 		ChunkEntries;
+	TArray<FOnDemandChunkEntryReferences>	ChunkEntryReferences;
+	FIoContainerId							ContainerId;
+	uint32									BlockSize = 0;
+	EOnDemandContainerFlags 				Flags = EOnDemandContainerFlags::None;
 };
 
 using FSharedOnDemandContainer = TSharedPtr<FOnDemandContainer, ESPMode::ThreadSafe>;
 
+int32 FOnDemandContainer::FindChunkEntryIndex(const FIoChunkId& ChunkId) const
+{
+	if (const int32 Index = Algo::LowerBound(ChunkIds, ChunkId); Index < ChunkIds.Num())
+	{
+		if (ChunkIds[Index] == ChunkId)
+		{
+			return Index; 
+		}
+	}
+
+	return INDEX_NONE;
+}
+
+const FOnDemandChunkEntry* FOnDemandContainer::FindChunkEntry(const FIoChunkId& ChunkId) const
+{
+	if (int32 Index = FindChunkEntryIndex(ChunkId); Index != INDEX_NONE)
+	{
+		return &ChunkEntries[Index];
+	}
+
+	return nullptr;
+}
+
+FOnDemandChunkEntry* FOnDemandContainer::FindChunkEntry(const FIoChunkId& ChunkId)
+{
+	if (int32 Index = FindChunkEntryIndex(ChunkId); Index != INDEX_NONE)
+	{
+		return &ChunkEntries[Index];
+	}
+
+	return nullptr;
+}
+
+FOnDemandChunkEntryReferences& FOnDemandContainer::FindOrAddChunkEntryReferences(const FOnDemandInternalContentHandle& ContentHandle)
+{
+	const UPTRINT ContentHandleId = ContentHandle.HandleId(); 
+	for (FOnDemandChunkEntryReferences& Refs : ChunkEntryReferences)
+	{
+		if (Refs.ContentHandleId == ContentHandleId)
+		{
+			return Refs;
+		}
+	}
+
+	FOnDemandChunkEntryReferences& NewRef = ChunkEntryReferences.AddDefaulted_GetRef();
+	NewRef.ContentHandleId = ContentHandleId;
+	NewRef.Indices.SetNum(ChunkEntries.Num(), false);
+	return NewRef;
+}
+
+TBitArray<> FOnDemandContainer::GetReferencedChunkEntries() const
+{
+	TBitArray<> Indices;
+	for (const FOnDemandChunkEntryReferences& Refs : ChunkEntryReferences)
+	{
+		check(Refs.Indices.Num() == ChunkEntries.Num());
+		Indices.CombineWithBitwiseOR(Refs.Indices, EBitwiseOperatorFlags::MaxSize);
+	}
+
+	return Indices;
+}
+
 ///////////////////////////////////////////////////////////////////////////////
+namespace Private
+{
+	static FIoStatus FetchContainerHeader(
+		FStringView Host,
+		FAnsiStringBuilderBase& ChunkUrlBuilder,
+		FSharedOnDemandContainer Container,
+		FIoContainerHeader& OutHeader);
+}
+
 struct FOnDemandChunkInfo
 {
 	FOnDemandChunkInfo()
@@ -103,6 +218,11 @@ struct FOnDemandChunkInfo
 
 private:
 	friend class FOnDemandIoStore;
+	friend FIoStatus Private::FetchContainerHeader(
+		FStringView Host,
+		FAnsiStringBuilderBase& ChunkUrlBuilder,
+		FSharedOnDemandContainer Container,
+		FIoContainerHeader& OutHeader);
 
 	FOnDemandChunkInfo(FSharedOnDemandContainer InContainer, const FOnDemandChunkEntry& InEntry)
 		: SharedContainer(InContainer)
@@ -127,25 +247,31 @@ TConstArrayView<FIoBlockHash> FOnDemandChunkInfo::BlockHashes() const
 
 ///////////////////////////////////////////////////////////////////////////////
 class FOnDemandIoStore
+	: public TSharedFromThis<FOnDemandIoStore, ESPMode::ThreadSafe>
 {
-	struct FTagSet
-	{
-		FString Tag;
-		TArray<FPackageId> PackageIds;
-	};
-
 	struct FMountRequest
 	{
-		FOnDemandMountArgs					MountArgs;
+		FOnDemandMountArgs					Args;
 		FOnDemandMountCompleted				OnCompleted;
 		TArray<FSharedOnDemandContainer>	Containers;
-		TArray<FTagSet>						TagSets;
 		std::atomic_bool					bCancelled{false};
+		double								DurationInSeconds = 0.0;
 	};
 
 	using FSharedMountRequest	= TSharedPtr<FMountRequest>;
 	using FMountRequestMap		= TMap<FString, FSharedMountRequest>;
 
+	struct FInstallRequest
+	{
+		FOnDemandInstallArgs		Args;
+		FOnDemandInstallCompleted	OnCompleted;
+		double						DurationInSeconds = 0.0;
+		uint64						TotalContentSize = 0;
+		uint64						TotalInstallSize = 0;
+	};
+
+	using FSharedInstallRequest	= TSharedPtr<FInstallRequest>;
+
 public:
 	FOnDemandIoStore();
 	~FOnDemandIoStore();
@@ -155,14 +281,17 @@ public:
 	FOnDemandIoStore& operator=(FOnDemandIoStore&&) = delete;
 
 	FIoStatus				Initialize();
-	void					Mount(FOnDemandMountArgs&& Args, FOnDemandMountCompleted OnCompleted);
+	void					Mount(FOnDemandMountArgs&& Args, FOnDemandMountCompleted&& OnCompleted);
+	void					Install(FOnDemandInstallArgs&& Args, FOnDemandInstallCompleted&& OnCompleted);
 	FIoStatus				Unmount(FStringView MountId);
 	TIoStatusOr<uint64>		GetInstallSize(const FOnDemandGetInstallSizeArgs& Args) const;
 	FOnDemandChunkInfo		GetStreamingChunkInfo(const FIoChunkId& ChunkId);
 	FOnDemandChunkInfo		GetInstalledChunkInfo(const FIoChunkId& ChunkId);
-	TArray<FSharedOnDemandContainer> GetMountedContainers();
-	TArray<FSharedOnDemandContainer> GetMountedContainers(
+	TArray<FSharedOnDemandContainer> GetContainers();
+	TArray<FSharedOnDemandContainer> GetContainers(
 		TFunctionRef<bool(const FSharedOnDemandContainer& Container)> Predicate);
+	void					ReleaseContent(FOnDemandInternalContentHandle& ContentHandle);
+	void					GetReferencedContent(TArray<FSharedOnDemandContainer>& OutContainers, TArray<TBitArray<>>& OutChunkEntryIndices);
 
 private:
 	void					OnPostFork(EForkProcessRole ProcessRole);
@@ -172,7 +301,7 @@ private:
 	void					TickLoop();
 	bool					Tick();
 	FIoStatus				TickMountRequest(FMountRequest& MountRequest);
-	FIoStatus				TickInstallRequest(FMountRequest& MountRequest);
+	FIoStatus				TickInstallRequest(FInstallRequest& InstallRequest);
 	void					OnEncryptionKeyAdded(const FGuid& Id, const FAES::FAESKey& Key);
 	static void				CreateContainersFromToc(
 								FStringView MountId,
@@ -187,6 +316,7 @@ private:
 	UE::FMutex							ContainerMutex;
 
 	FMountRequestMap					MountRequests;
+	TArray<FSharedInstallRequest>		InstallRequests;
 	UE::FMutex							MountRequestMutex;
 
 	bool								bTicking = false;
diff --git a/Engine/Source/Runtime/Experimental/IoStore/OnDemand/Public/IO/IoStoreOnDemand.h b/Engine/Source/Runtime/Experimental/IoStore/OnDemand/Public/IO/IoStoreOnDemand.h
--- a/Engine/Source/Runtime/Experimental/IoStore/OnDemand/Public/IO/IoStoreOnDemand.h
+++ b/Engine/Source/Runtime/Experimental/IoStore/OnDemand/Public/IO/IoStoreOnDemand.h
@@ -3,9 +3,12 @@
 #pragma once
 
 #include "Containers/UnrealString.h"
+#include "Containers/StringFwd.h"
+#include "Containers/SharedString.h"
 #include "IO/IoChunkId.h"
 #include "IO/IoContainerId.h"
 #include "IO/IoHash.h"
+#include "IO/PackageId.h"
 #include "IO/IoStatus.h"
 #include "Misc/EnumClassFlags.h"
 #include "Misc/Guid.h"
@@ -28,6 +31,7 @@ struct FAnalyticsEventAttribute;
 struct FIoContainerSettings;
 struct FIoStoreWriterSettings;
 namespace UE::IoStore { struct FOnDemandEndpoint; }
+namespace UE::IoStore { class FOnDemandInternalContentHandle; }
 namespace UE::IoStore { class FOnDemandIoStore; }
 namespace UE::IoStore { class IOnDemandIoDispatcherBackend; }
 using FIoBlockHash = uint32;
@@ -238,20 +242,47 @@ enum class EOnDemandInitResult
 
 #endif // UE_IAS_CUSTOM_INITIALIZATION
 
-/** Options for controlling the behavior of mount requests. */
+/**
+ * Keeps referenced data pinned in the cache until released.
+ */
+class FOnDemandContentHandle
+{
+public:
+	/** Creates a new invalid content handle. */ 
+	UE_API FOnDemandContentHandle();
+	/** Destroy the handle and release any referenced content. */ 
+	UE_API ~FOnDemandContentHandle();
+	/** Destroy the handle and release any referenced content. */ 
+	void Reset() { Handle.Reset(); }
+	/** Returns whether the handle is valid. */
+	bool IsValid() const { return Handle.IsValid(); }
+	/** Create a new content handle .*/
+	UE_API static FOnDemandContentHandle Create();
+	/** Create a new content handle with a debug name. */
+	UE_API static FOnDemandContentHandle Create(FSharedString DebugName);
+	/** Create a new content handle with a debug name. */
+	UE_API static FOnDemandContentHandle Create(FStringView DebugName);
+	/** Returns a string representing the content handle. */
+	UE_API friend FString LexToString(const FOnDemandContentHandle& Handle);
+
+private:
+	friend class FOnDemandIoStore;
+	TSharedPtr<FOnDemandInternalContentHandle, ESPMode::ThreadSafe> Handle;
+};
+
+/** Options for controlling the behavior of mounted container(s). */
 enum class EOnDemandMountOptions
 {
-	/** The TOC is loaded but not installed or available for streaming. */
-	None				= 0,
-	/** Make on-demand container(s) within a TOC available for streaming. */
-	StreamOnDemand		= 1 << 0,
-	/** Download and install on-demand contianer(s) to local storage. */
-	Install				= 1 << 1,
-	/** Install only already mounted on-demand contianer(s) to local storage. */
-	InstallSkipMount	= 1 << 2,
+	/** Mount containers with the purpose of streaming the content on-demand. */
+	StreamOnDemand			= 1 << 0,
+	/** Mount containers with the purpose of installing/downloading the content on-demand. */
+	InstallOnDemand			= 1 << 1,
+	/** Trigger callback on game thread. */
+	CallbackOnGameThread	= 1 << 2,
 };
 ENUM_CLASS_FLAGS(EOnDemandMountOptions);
 
+/** Arguments for mounting on-demand container TOC(s). */
 struct FOnDemandMountArgs
 {
 	/** Mount an already serialized TOC. */
@@ -262,20 +293,67 @@ struct FOnDemandMountArgs
 	FString Url;
 	/** Serialize the TOC from the specified file path. */
 	FString FilePath;
-	/** Used with EOnDemandMountOptions::Install to fliter requested content */
-	TArray<FString> TagSets;
 	/** Mount options. */
-	EOnDemandMountOptions Options = EOnDemandMountOptions::None;
+	EOnDemandMountOptions Options = EOnDemandMountOptions::StreamOnDemand;
 };
 
+/** Holds information about a mount request. */
 struct FOnDemandMountResult
 {
+	/** The mount ID used for mounting the container(s). */
 	FString MountId;
+	/** The status of the mount request. */
 	FIoStatus Status;
+	/** Duration in seconds. */
+	double DurationInSeconds = 0.0;
 };
 
+/** Mount completion callback. */
 using FOnDemandMountCompleted = TUniqueFunction<void(FOnDemandMountResult)>;
 
+/** Options for controlling the behavior of the install request. */
+enum class EOnDemandInstallOptions
+{
+	/** No additional options. */
+	None				 = 0,
+	/** Trigger callback on game thread. */
+	CallbackOnGameThread = 1 << 0,
+};
+ENUM_CLASS_FLAGS(EOnDemandInstallOptions);
+
+/** Arguments for installing/downloading on-demand content. */
+struct FOnDemandInstallArgs
+{
+	/** URL from where to download the chunks. */
+	FString Url;
+	/** Install all content from containers matching this mount ID. */
+	FString MountId;
+	/** Install content matching a set of tag(s) and optionally the mount ID. */
+	TArray<FString> TagSets;
+	/** Package ID's to install. */
+	TArray<FPackageId> PackageIds;
+	/** Content handle. */
+	FOnDemandContentHandle ContentHandle;
+	/** Install options. */
+	EOnDemandInstallOptions Options = EOnDemandInstallOptions::None;
+};
+
+/** Holds information about an install request. */
+struct FOnDemandInstallResult
+{
+	/** The status of the install request. */
+	FIoStatus Status;
+	/** Duration in seconds. */
+	double DurationInSeconds = 0.0;
+	/** The total size of the requested content. */
+	uint64 TotalContentSize = 0;
+	/** The total installed/downloaded size. */
+	uint64 TotalInstallSize = 0;
+};
+
+/** Install completion callback. */
+using FOnDemandInstallCompleted = TUniqueFunction<void(FOnDemandInstallResult)>;
+
 struct FOnDemandGetInstallSizeArgs
 {
 	/* Includes all packages from matching containers */
@@ -306,7 +384,7 @@ private:
 	TOptional<bool> DeferredEnabled;
 	TOptional<bool> DeferredAbandonCache;
 	TOptional<bool> DeferredBulkOptionalEnabled;
-	TUniquePtr<FOnDemandIoStore> IoStore;
+	TSharedPtr<FOnDemandIoStore, ESPMode::ThreadSafe> IoStore;
 
 public:
 	UE_API void SetBulkOptionalEnabled(bool bInEnabled);
@@ -316,7 +394,8 @@ public:
 
 	UE_API void ReportAnalytics(TArray<FAnalyticsEventAttribute>& OutAnalyticsArray) const;
 
-	UE_API void Mount(FOnDemandMountArgs&& Args, FOnDemandMountCompleted OnCompleted);
+	UE_API void Mount(FOnDemandMountArgs&& Args, FOnDemandMountCompleted&& OnCompleted);
+	UE_API void Install(FOnDemandInstallArgs&& Args, FOnDemandInstallCompleted&& OnCompleted);
 	UE_API FIoStatus Unmount(FStringView MountId);
 
 	UE_API TIoStatusOr<uint64> GetInstallSize(const FOnDemandGetInstallSizeArgs& Args) const;
