Make sure that winsock is initialized for IAS on windows

#rnx

- Technically this should be in the IoStoreHttpClient module but that doesn't implement IModuleInterface and even if it did we'd then need to invoke the load via FIoStoreOnDemandModule anyway. This will naturally be solved when/if we convert IoStoreOnDemand to a plugin as it can be set to load both module interfaces.
- Currently only an issue when IAS is started with the 'ias.tocurl' cmdline as it's the only path that will start making network calls before other modules that initialize winsock have loaded.
-- In theory this change will allow us to remove other instances of initializing winsock found in the IasTool but that can be saved for a future submit.

diff --git a/Engine/Source/Runtime/Experimental/IoStore/OnDemand/Private/IoStoreOnDemand.cpp b/Engine/Source/Runtime/Experimental/IoStore/OnDemand/Private/IoStoreOnDemand.cpp
--- a/Engine/Source/Runtime/Experimental/IoStore/OnDemand/Private/IoStoreOnDemand.cpp
+++ b/Engine/Source/Runtime/Experimental/IoStore/OnDemand/Private/IoStoreOnDemand.cpp
@@ -32,6 +32,12 @@
 #include "Statistics.h"
 #include "String/LexFromString.h"
 
+#if PLATFORM_WINDOWS
+#	include <Windows/AllowWindowsPlatformTypes.h>
+#		include <winsock2.h>
+#	include <Windows/HideWindowsPlatformTypes.h>
+#endif //PLATFORM_WINDOWS
+
 DEFINE_LOG_CATEGORY(LogIoStoreOnDemand);
 DEFINE_LOG_CATEGORY(LogIas);
 
@@ -1281,6 +1287,24 @@ void FIoStoreOnDemandModule::InitializeInternal()
 	
 void FIoStoreOnDemandModule::StartupModule()
 {
+#if PLATFORM_WINDOWS
+	{
+		WSADATA WsaData;
+		int Result = WSAStartup(MAKEWORD(2, 2), &WsaData);
+		if (Result == 0)
+		{
+			bPlatformSpecificSetup = true;
+		}
+		else
+		{
+			TCHAR SystemErrorMsg[MAX_SPRINTF] = { 0 };
+			FPlatformMisc::GetSystemErrorMessage(SystemErrorMsg, sizeof(SystemErrorMsg), Result);
+
+			UE_LOG(LogIas, Error, TEXT("WSAStartup failed due to: %s (%d)"), SystemErrorMsg, Result);
+		}
+	}
+#endif //PLATFORM_WINDOWS
+
 #if !UE_IAS_CUSTOM_INITIALIZATION
 
 	if (!GIasSuspendSystem)
@@ -1297,6 +1321,22 @@ void FIoStoreOnDemandModule::StartupModule()
 
 void FIoStoreOnDemandModule::ShutdownModule()
 {
+#if PLATFORM_WINDOWS
+	if (bPlatformSpecificSetup)
+	{
+		if (WSACleanup() != 0)
+		{
+			const uint32 SystemError = FPlatformMisc::GetLastError();
+
+			TCHAR SystemErrorMsg[MAX_SPRINTF] = { 0 };
+			FPlatformMisc::GetSystemErrorMessage(SystemErrorMsg, sizeof(SystemErrorMsg), SystemError);
+
+			UE_LOG(LogIas, Error, TEXT("WSACleanup failed due to: %s (%u)"), SystemErrorMsg, SystemError);
+		}
+
+		bPlatformSpecificSetup = false;
+	}
+#endif //PLATFORM_WINDOWS
 }
 
 #if UE_IAS_CUSTOM_INITIALIZATION
diff --git a/Engine/Source/Runtime/Experimental/IoStore/OnDemand/Public/IO/IoStoreOnDemand.h b/Engine/Source/Runtime/Experimental/IoStore/OnDemand/Public/IO/IoStoreOnDemand.h
--- a/Engine/Source/Runtime/Experimental/IoStore/OnDemand/Public/IO/IoStoreOnDemand.h
+++ b/Engine/Source/Runtime/Experimental/IoStore/OnDemand/Public/IO/IoStoreOnDemand.h
@@ -386,6 +386,9 @@ private:
 	TOptional<bool> DeferredBulkOptionalEnabled;
 	TSharedPtr<FOnDemandIoStore, ESPMode::ThreadSafe> IoStore;
 
+	/** Indicates that platform specific setup was invoked when the module was started and needs to be cleaned up on shutdown */
+	bool bPlatformSpecificSetup = false;
+
 public:
 	UE_API void SetBulkOptionalEnabled(bool bInEnabled);
 	UE_API void SetEnabled(bool bInEnabled);
