[Backout] - CL35028627

Original CL Desc
-----------------------------------------------------------------
Iris
* Fix dormant objects getting replicated if one of their push model property was set dirty or force net update was called on them.
* Added dormancy iris unit test

#rnx

diff --git a/Engine/Source/Runtime/Experimental/Iris/Core/Private/Iris/ReplicationSystem/ObjectReplicationBridge.cpp b/Engine/Source/Runtime/Experimental/Iris/Core/Private/Iris/ReplicationSystem/ObjectReplicationBridge.cpp
--- a/Engine/Source/Runtime/Experimental/Iris/Core/Private/Iris/ReplicationSystem/ObjectReplicationBridge.cpp
+++ b/Engine/Source/Runtime/Experimental/Iris/Core/Private/Iris/ReplicationSystem/ObjectReplicationBridge.cpp
@@ -1003,14 +1003,14 @@ void UObjectReplicationBridge::BuildPollList(UE::Net::FNetBitArrayView ObjectsCo
 		ObjectsConsideredForPolling.Copy(RelevantObjects);
 	}
 
-	// Mask off objects pending dormancy as we do not want to poll/pre-update them unless they are marked for flush
+	// Mask off objects pending dormancy as we do not want to poll/pre-update them unless they are marked for flush or are dirty
 	if (bUseDormancyToFilterPolling)
 	{
 		IRIS_PROFILER_SCOPE(BuildPollList_Dormancy);
 
-		// Mask off objects pending dormancy
+		// Mask off objects pending dormancy that are not dirty
 		const FNetBitArrayView AccumulatedDirtyObjects = DirtyNetObjectTracker.GetAccumulatedDirtyNetObjects();
-		ObjectsConsideredForPolling.Combine(WantToBeDormantObjects, FNetBitArrayView::AndNotOp);
+		ObjectsConsideredForPolling.CombineMultiple(FNetBitArrayView::AndNotOp, WantToBeDormantObjects, FNetBitArrayView::AndNotOp, AccumulatedDirtyObjects);
 
 		FNetBitArrayView ForceNetUpdateObjects = ReplicationSystemInternal->GetDirtyNetObjectTracker().GetForceNetUpdateObjects();
 
