Suppress a false positive warning with -NoPCH with the 14.40 compiler

#rnx

diff --git a/Engine/Source/Runtime/Experimental/ChaosVisualDebugger/Public/DataWrappers/ChaosVDAccelerationStructureDataWrappers.h b/Engine/Source/Runtime/Experimental/ChaosVisualDebugger/Public/DataWrappers/ChaosVDAccelerationStructureDataWrappers.h
--- a/Engine/Source/Runtime/Experimental/ChaosVisualDebugger/Public/DataWrappers/ChaosVDAccelerationStructureDataWrappers.h
+++ b/Engine/Source/Runtime/Experimental/ChaosVisualDebugger/Public/DataWrappers/ChaosVDAccelerationStructureDataWrappers.h
@@ -54,7 +54,10 @@ struct FChaosVDBVCellElementDataWrapper : public FChaosVDWrapperDataBase
 	CHAOSVDRUNTIME_API bool Serialize(FArchive& Ar);
 };
 
+// TODO: Disabling C4996 due to MSVC 14.30 bug
+#pragma warning(disable : 4996)
 CVD_IMPLEMENT_SERIALIZER(FChaosVDBVCellElementDataWrapper)
+#pragma warning(default : 4996)
 
 USTRUCT()
 struct FChaosVDBoundingVolumeDataWrapper : public FChaosVDAccelerationStructureBase
