Fixed heightfield queries missing when they begin outside of the heightfield grid and the clipped query ends up within a single cell.

When we hit the case of a query existing in a single cell we skip the rasterization phase of the query and just query the cell directly. Previously we took the the bounds inflation and wrapped a new box around the start of the query and used that to grab the set of overlapping cells. This is incorrect we need to take the inflated bounds of the clipped query as the intersection bounds.

diff --git a/Engine/Source/Runtime/Experimental/Chaos/Private/Chaos/HeightField.cpp b/Engine/Source/Runtime/Experimental/Chaos/Private/Chaos/HeightField.cpp
--- a/Engine/Source/Runtime/Experimental/Chaos/Private/Chaos/HeightField.cpp
+++ b/Engine/Source/Runtime/Experimental/Chaos/Private/Chaos/HeightField.cpp
@@ -1446,10 +1446,12 @@ namespace Chaos
 			if (bSameCell)
 			{
 				//start and end in same cell, so test it and any cells that the InHalfExtents overlap with
-				const FVec2 MinPoint = Start2D - Inflation2D;
-				const FVec2 MaxPoint = Start2D + Inflation2D;
-				TVec2<int32> MinCell = FlatGrid.Cell(MinPoint / Scale2D);
-				TVec2<int32> MaxCell = FlatGrid.Cell(MaxPoint / Scale2D);
+				FBounds2D ClippedBounds = FBounds2D::FromPoints(ClippedStart, ClippedEnd);
+				ClippedBounds.Inflate(Inflation2D);
+
+				TVec2<int32> MinCell = FlatGrid.Cell(ClippedBounds.Min / Scale2D);
+				TVec2<int32> MaxCell = FlatGrid.Cell(ClippedBounds.Max / Scale2D);
+
 				for (int32 Y = MinCell[1]; Y <= MaxCell[1]; ++Y)
 				{
 					for (int32 X = MinCell[0]; X <= MaxCell[0]; ++X)
diff --git a/Engine/Source/Runtime/Experimental/Chaos/Public/Chaos/HeightField.h b/Engine/Source/Runtime/Experimental/Chaos/Public/Chaos/HeightField.h
--- a/Engine/Source/Runtime/Experimental/Chaos/Public/Chaos/HeightField.h
+++ b/Engine/Source/Runtime/Experimental/Chaos/Public/Chaos/HeightField.h
@@ -895,6 +895,12 @@ namespace Chaos
 				Max = { FMath::Max(Max.X, InPoint.X), FMath::Max(Max.Y, InPoint.Y) };
 			}
 
+			void Inflate(const FVec2& InInflation)
+			{
+				Min -= InInflation;
+				Max += InInflation;
+			}
+
 			FVec2 Clamp(const FVec2& InToClamp, FReal InNudge = UE_SMALL_NUMBER) const
 			{
 				const FVec2 NudgeVec(InNudge, InNudge);
