OnDemandIoStore
- Header chunks are now downloaded during mount
- Mount requests will now fail with PendingEncryptionKey if no encryption key is found, but containers will still be created and will mount when keys become availble.
- If encryption keys aren't available, the encoded header is stored until they become available
- Install requests are no longer deferred until encryption keys are available, they will fail with UnknownChunkID or PendingEncryptionKey
- Fully implemented GetInstallSize() and fixed thread safety issues

diff --git a/Engine/Source/Runtime/Core/Private/IO/IoDispatcher.cpp b/Engine/Source/Runtime/Core/Private/IO/IoDispatcher.cpp
--- a/Engine/Source/Runtime/Core/Private/IO/IoDispatcher.cpp
+++ b/Engine/Source/Runtime/Core/Private/IO/IoDispatcher.cpp
@@ -47,7 +47,8 @@ static const TCHAR* const GetIoErrorText_ErrorCodeTextArray[] =
 	TEXT("Signature Error"),
 	TEXT("Invalid Encryption Key"),
 	TEXT("Compression Error"),
-	TEXT("Pending Fork")
+	TEXT("Pending Fork"),
+	TEXT("Pending Encryption Key")
 };
 
 CORE_API const TCHAR* const* GetIoErrorText_ErrorCodeText = GetIoErrorText_ErrorCodeTextArray;
diff --git a/Engine/Source/Runtime/Core/Public/IO/IoStatus.h b/Engine/Source/Runtime/Core/Public/IO/IoStatus.h
--- a/Engine/Source/Runtime/Core/Public/IO/IoStatus.h
+++ b/Engine/Source/Runtime/Core/Public/IO/IoStatus.h
@@ -30,6 +30,7 @@ enum class EIoErrorCode
 	InvalidEncryptionKey,
 	CompressionError,
 	PendingFork,
+	PendingEncryptionKey,
 };
 
 /** Get I/O error code description. */
diff --git a/Engine/Source/Runtime/Experimental/IoStore/OnDemand/Private/IoStoreOnDemand.cpp b/Engine/Source/Runtime/Experimental/IoStore/OnDemand/Private/IoStoreOnDemand.cpp
--- a/Engine/Source/Runtime/Experimental/IoStore/OnDemand/Private/IoStoreOnDemand.cpp
+++ b/Engine/Source/Runtime/Experimental/IoStore/OnDemand/Private/IoStoreOnDemand.cpp
@@ -1251,12 +1251,10 @@ void FIoStoreOnDemandModule::InitializeInternal()
 			MountArgs.GetValue().Options = EOnDemandMountOptions::InstallOnDemand;
 
 			static FOnDemandContentHandle ContentHandle = FOnDemandContentHandle::Create(TEXT("AllContent"));
-			InstallArgs.Emplace(FOnDemandInstallArgs
-			{
-				.Url = EndpointConfig.ServiceUrls[0] / EndpointConfig.TocPath,
-				.MountId = MountArgs.GetValue().MountId,
-				.ContentHandle = ContentHandle
-			});
+			InstallArgs.Emplace();
+			InstallArgs->Url = EndpointConfig.ServiceUrls[0] / EndpointConfig.TocPath;
+			InstallArgs->MountId = MountArgs.GetValue().MountId;
+			InstallArgs->ContentHandle = ContentHandle;
 		}
 	}
 #endif
diff --git a/Engine/Source/Runtime/Experimental/IoStore/OnDemand/Private/OnDemandIoStore.cpp b/Engine/Source/Runtime/Experimental/IoStore/OnDemand/Private/OnDemandIoStore.cpp
--- a/Engine/Source/Runtime/Experimental/IoStore/OnDemand/Private/OnDemandIoStore.cpp
+++ b/Engine/Source/Runtime/Experimental/IoStore/OnDemand/Private/OnDemandIoStore.cpp
@@ -6,6 +6,7 @@
 #include "OnDemandPackageStoreBackend.h"
 
 #include "Algo/Accumulate.h"
+#include "Algo/AnyOf.h"
 #include "Algo/Copy.h"
 #include "Algo/RemoveIf.h"
 #include "Algo/Transform.h"
@@ -150,30 +151,57 @@ static FAnsiStringBuilderBase& GetChunkUrl(
 	return OutUrl;
 }
 
-static FIoStatus FetchContainerHeader(
+static FIoStatus FetchEncodedContainerChunk(
 	FStringView Host,
 	FAnsiStringBuilderBase& ChunkUrlBuilder,
 	FSharedOnDemandContainer Container,
-	FIoContainerHeader& OutHeader)
+	const FOnDemandChunkEntry& Entry,
+	FIoBuffer& OutEncodedChunk)
+{
+	TIoStatusOr<FIoBuffer> Response = FHttpClient::Get(GetChunkUrl(Host, *Container, Entry, ChunkUrlBuilder).ToView(), 2, EHttpRedirects::Follow);
+
+	if (Response.IsOk() == false)
+	{
+		FIoStatus Status = FIoStatusBuilder(EIoErrorCode::ReadError) << TEXT("Failed to fetch container chunk");
+		return Status;
+	}
+
+	OutEncodedChunk = Response.ConsumeValueOrDie();
+
+	return EIoErrorCode::Ok;
+}
+
+static FIoStatus FetchEncodedContainerHeader(
+	FStringView Host,
+	FAnsiStringBuilderBase& ChunkUrlBuilder, 
+	FSharedOnDemandContainer Container,
+	FIoBuffer& OutEncodedChunk)
 {
 	const FIoChunkId			ChunkId = CreateContainerHeaderChunkId(Container->ContainerId);
 	const FOnDemandChunkEntry*	Entry = Container->FindChunkEntry(ChunkId);
 
 	if (Entry == nullptr)
 	{
-		return EIoErrorCode::Ok;
+		return FIoStatusBuilder(EIoErrorCode::NotFound) << TEXT("Failed to find header chunk");
 	}
 
-	UE_LOG(LogIoStoreOnDemand, Verbose, TEXT("Fetching container header, ContainerName='%s'"), *Container->Name);
-	TIoStatusOr<FIoBuffer> Response = FHttpClient::Get(GetChunkUrl(Host, *Container, *Entry, ChunkUrlBuilder).ToView(), 2, EHttpRedirects::Follow);
-
-	if (Response.IsOk() == false)
+	UE_LOG(LogIoStoreOnDemand, Verbose, TEXT("Fetching encoded container header, ContainerName='%s'"), *Container->Name);
+	FIoStatus Status = FetchEncodedContainerChunk(Host, ChunkUrlBuilder, Container, *Entry, OutEncodedChunk);
+	if (!Status.IsOk())
 	{
-		FIoStatus Status = FIoStatusBuilder(EIoErrorCode::ReadError) << TEXT("Failed to fetch container header chunk");
 		return Status;
 	}
 
-	FOnDemandChunkInfo ChunkInfo(Container, *Entry);
+	return EIoErrorCode::Ok;
+}
+
+static FIoStatus DecodeContainerChunk(
+	FSharedOnDemandContainer Container,
+	const FOnDemandChunkEntry& Entry,
+	FMemoryView EncodedChunk,
+	FIoBuffer& OutRawChunk)
+{
+	FOnDemandChunkInfo ChunkInfo(Container, Entry);
 
 	FIoChunkDecodingParams Params;
 	Params.CompressionFormat = ChunkInfo.CompressionFormat();
@@ -185,12 +213,72 @@ static FIoStatus FetchContainerHeader(
 	Params.EncodedBlockSize = ChunkInfo.Blocks();
 	Params.BlockHash = ChunkInfo.BlockHashes();
 
-	FIoBuffer EncodedChunk = Response.ConsumeValueOrDie();
-	FIoBuffer RawChunk(ChunkInfo.RawSize());
-	if (FIoChunkEncoding::Decode(Params, EncodedChunk.GetView(), RawChunk.GetMutableView()) == false)
+	OutRawChunk = FIoBuffer(ChunkInfo.RawSize());
+	if (FIoChunkEncoding::Decode(Params, EncodedChunk, OutRawChunk.GetMutableView()) == false)
 	{
 		FIoStatus Status = FIoStatusBuilder(EIoErrorCode::ReadError)
-			<< TEXT("Failed to decode container header chunk");
+			<< TEXT("Failed to decode container chunk");
+		return Status;
+	}
+
+	return EIoErrorCode::Ok;
+}
+
+static FIoStatus DecodeContainerHeader(
+	FSharedOnDemandContainer Container,
+	FMemoryView EncodedHeaderChunk,
+	FIoContainerHeader& OutHeader)
+{
+	const FIoChunkId			ChunkId = CreateContainerHeaderChunkId(Container->ContainerId);
+	const FOnDemandChunkEntry*	Entry = Container->FindChunkEntry(ChunkId);
+
+	FIoBuffer RawChunk;
+	if (FIoStatus Status = DecodeContainerChunk(Container, *Entry, EncodedHeaderChunk, RawChunk); !Status.IsOk())
+	{
+		return Status;
+	}
+
+	FMemoryReaderView Ar(RawChunk.GetView());
+	Ar << OutHeader;
+	Ar.Close();
+
+	if (Ar.IsError() || Ar.IsCriticalError())
+	{
+		FIoStatus Status = FIoStatusBuilder(EIoErrorCode::FileNotOpen) << TEXT("Failed to serialize container header");
+		return Status;
+	}
+
+	return EIoErrorCode::Ok;
+}
+
+static FIoStatus FetchContainerHeader(
+	FStringView Host,
+	FAnsiStringBuilderBase& ChunkUrlBuilder,
+	FSharedOnDemandContainer Container,
+	FIoContainerHeader& OutHeader)
+{
+	const FIoChunkId			ChunkId = CreateContainerHeaderChunkId(Container->ContainerId);
+	const FOnDemandChunkEntry*	Entry = Container->FindChunkEntry(ChunkId);
+
+	if (Entry == nullptr)
+	{
+		return FIoStatusBuilder(EIoErrorCode::NotFound) << TEXT("Failed to find header chunk");
+	}
+
+	UE_LOG(LogIoStoreOnDemand, Verbose, TEXT("Fetching container header, ContainerName='%s'"), *Container->Name);
+	FIoBuffer EncodedChunk;
+	{
+		FIoStatus Status = FetchEncodedContainerChunk(Host, ChunkUrlBuilder, Container, *Entry, EncodedChunk);
+		if (!Status.IsOk())
+		{
+			return Status;
+		}
+	}
+
+	FIoBuffer RawChunk;
+	if (FIoStatus Status = DecodeContainerChunk(Container, *Entry, EncodedChunk.GetView(), RawChunk); !Status.IsOk())
+	{
+		return Status;
 	}
 
 	FMemoryReaderView Ar(RawChunk.GetView());
@@ -221,7 +309,7 @@ struct FContainerInstallData
 
 using FInstallData = TMap<FSharedOnDemandContainer, FContainerInstallData>;
 
-FIoStatus BuildInstallData(
+static FIoStatus BuildInstallData(
 	const TSet<FSharedOnDemandContainer>& Containers,
 	const TSet<FPackageId>& PackageIds,
 	FInstallData& OutInstallData,
@@ -232,7 +320,8 @@ FIoStatus BuildInstallData(
 	// Setup the package information for each container
 	for (const FSharedOnDemandContainer& Container : Containers)
 	{
-		if (!Container->Header.IsValid() || Container->Header->PackageIds.IsEmpty())
+		check(Container->Header.IsValid());
+		if (Container->Header->PackageIds.IsEmpty())
 		{
 			// The container contains no package data
 			continue;
@@ -376,6 +465,8 @@ FIoStatus BuildInstallData(
 		}
 
 		// For now we always download these chunks
+		// TODO: this should only be done for containers that were visited above
+		// or containers with the mount ID of the install request.
 		for (int32 EntryIndex = 0; const FIoChunkId& ChunkId : Container.ChunkIds)
 		{
 			switch(ChunkId.GetChunkType())
@@ -580,20 +671,14 @@ void FOnDemandIoStore::Mount(FOnDemandMountArgs&& Args, FOnDemandMountCompleted&
 	}
 
 	{
-		UE::TUniqueLock Lock(MountRequestMutex);
-		FSharedMountRequest& MountRequest = MountRequests.FindOrAdd(Args.MountId);
+		FSharedMountRequest MountRequest = MakeShared<FMountRequest>();
+		MountRequest->Args = MoveTemp(Args);
+		MountRequest->OnCompleted = MoveTemp(OnCompleted);
 
-		if (MountRequest.IsValid())
-		{
-			UE_LOG(LogIoStoreOnDemand, Warning, TEXT("Mount request '%s' is already mounting"), *Args.MountId);
-			return OnCompleted(FOnDemandMountResult{ .Status = FIoStatus(EIoErrorCode::InvalidParameter, TEXT("Mount in progress")) });
-		}
+		UE_LOG(LogIoStoreOnDemand, Log, TEXT("Enqueing mount request, MountId='%s'"), *MountRequest->Args.MountId);
 
-		UE_LOG(LogIoStoreOnDemand, Log, TEXT("Enqueing mount request, MountId='%s'"), *Args.MountId);
-		
-		MountRequest = MakeShared<FMountRequest>();
-		MountRequest->Args			= MoveTemp(Args);
-		MountRequest->OnCompleted	= MoveTemp(OnCompleted);
+		UE::TUniqueLock Lock(MountRequestMutex);
+		MountRequests.Add(MoveTemp(MountRequest));
 	}
 
 	TryEnterTickLoop();
@@ -617,12 +702,17 @@ FIoStatus FOnDemandIoStore::Unmount(FStringView MountId)
 {
 	UE_LOG(LogIoStoreOnDemand, Log, TEXT("Unmounting '%s'"), *WriteToString<256>(MountId));
 
+	bool bPendingMount = false;
+
 	{
 		UE::TUniqueLock Lock(MountRequestMutex);
-		if (FSharedMountRequest* MountRequest = MountRequests.Find(FString(MountId)))
-		{
-			(*MountRequest)->bCancelled = true;
-		}
+		bPendingMount = Algo::AnyOf(MountRequests, 
+			[MountId](const FSharedMountRequest& Request) { return Request->Args.MountId == MountId; });
+	}
+
+	if (bPendingMount)
+	{
+		return FIoStatusBuilder(EIoErrorCode::InvalidParameter) << TEXT("Mount requests pending for MountId");
 	}
 
 	{
@@ -635,11 +725,14 @@ FIoStatus FOnDemandIoStore::Unmount(FStringView MountId)
 				UE_LOG(LogIoStoreOnDemand, Log, TEXT("Unmounting container, ContainerName='%s', MountId='%s'"),
 					*WriteToString<128>(Container->Name), *WriteToString<128>(Container->MountId));
 
+				const FString ContainerName = Container->UniqueName();
 				if (PackageStoreBackend.IsValid())
 				{
-					PackageStoreBackend->Unmount(Container->UniqueName());
+					PackageStoreBackend->Unmount(ContainerName);
 				}
 
+				EncodedContainerHeaders.Remove(ContainerName);
+
 				return true;
 			}
 
@@ -654,20 +747,19 @@ TIoStatusOr<uint64> FOnDemandIoStore::GetInstallSize(const FOnDemandGetInstallSi
 {
 	TRACE_CPUPROFILER_EVENT_SCOPE(FOnDemandIoStore::GetInstallSize);
 
-	TSet<FSharedOnDemandContainer> AllContainers(
-		const_cast<FOnDemandIoStore*>(this)->GetContainers([](const FSharedOnDemandContainer& Container)
-			{
-				return EnumHasAnyFlags(Container->Flags, EOnDemandContainerFlags::InstallOnDemand);
-			})
-	);
-
+	TSet<FSharedOnDemandContainer> ContainersForInstallation;
 	TSet<FPackageId> PackageIdsToInstall;
-	PackageIdsToInstall.Append(Args.Packages);
+	FIoStatus GetContainersStatus = const_cast<FOnDemandIoStore*>(this)->GetContainersAndPackagesForInstall(
+		Args, ContainersForInstallation, PackageIdsToInstall);
+	if (!GetContainersStatus.IsOk())
+	{
+		return GetContainersStatus;
+	}
 
 	Private::FInstallData InstallData;
 	TSet<FPackageId> Missing;
 
-	FIoStatus Status = BuildInstallData(AllContainers, PackageIdsToInstall, InstallData, Missing);
+	FIoStatus Status = BuildInstallData(ContainersForInstallation, PackageIdsToInstall, InstallData, Missing);
 	if (Status.IsOk() == false)
 	{
 		return Status;
@@ -871,10 +963,7 @@ bool FOnDemandIoStore::Tick()
 	TArray<FSharedMountRequest> LocalMountRequests;
 	{
 		UE::TUniqueLock Lock(MountRequestMutex);
-		for (TPair<FString, FSharedMountRequest>& Kv : MountRequests)
-		{
-			LocalMountRequests.Add(Kv.Value);
-		}
+		LocalMountRequests = MountRequests;
 	}
 
 	bool bTicked = LocalMountRequests.IsEmpty() == false;
@@ -882,95 +971,20 @@ bool FOnDemandIoStore::Tick()
 	// Tick mount request(s)
 	for (FSharedMountRequest& Request : LocalMountRequests)
 	{
-		{
-			UE::TUniqueLock Lock(MountRequestMutex);
-			MountRequests.Remove(Request->Args.MountId);
-		}
+		FIoStatus MountStatus = TickMountRequest(*Request);
 
-		if (FIoStatus Status = TickMountRequest(*Request); !Status.IsOk())
 		{
-			FOnDemandMountResult MountResult
-			{
-				.Status = MoveTemp(Status),
-				.DurationInSeconds = Request->DurationInSeconds 
-			};
-
-			if (EnumHasAnyFlags(Request->Args.Options, EOnDemandMountOptions::CallbackOnGameThread))
-			{
-				ExecuteOnGameThread(
-					UE_SOURCE_LOCATION,
-					[OnCompleted = MoveTemp(Request->OnCompleted), MountResult = MoveTemp(MountResult)]() mutable
-					{
-						OnCompleted(MoveTemp(MountResult));
-					});
-			}
-			else
-			{
-				FOnDemandMountCompleted OnCompleted = MoveTemp(Request->OnCompleted);
-				OnCompleted(MoveTemp(MountResult));
-			}
-			continue;
+			UE::TUniqueLock Lock(MountRequestMutex);
+			MountRequests.Remove(Request);
 		}
 
-		TStringBuilder<128> Sb;
-		{
-			UE::TUniqueLock Lock(ContainerMutex);
-			for (const FSharedOnDemandContainer& Container : Request->Containers)
+		OnMountRequestComplete(*Request, 
+			FOnDemandMountResult
 			{
-				if (EnumHasAnyFlags(Container->Flags, EOnDemandContainerFlags::Mounted))
-				{
-					UE_LOG(LogIoStoreOnDemand, Log, TEXT("Ignoring already mounted container, ContainerName='%s'"), *Container->Name);
-					continue;
-				}
-
-				Containers.Add(Container);
-
-				// To conform with PAK mount behavior, mount requests are completed even if encryption keys are pending
-				if (EnumHasAnyFlags(Container->Flags, EOnDemandContainerFlags::Encrypted) &&
-					Container->EncryptionKey.IsValid() == false)
-				{
-					FGuid KeyGuid;
-					ensure(FGuid::Parse(Container->EncryptionKeyGuid, KeyGuid));
-					if (FEncryptionKeyManager::Get().TryGetKey(KeyGuid, Container->EncryptionKey) == false)
-					{
-						EnumAddFlags(Container->Flags, EOnDemandContainerFlags::PendingEncryptionKey);
-						UE_LOG(LogIoStoreOnDemand, Log, TEXT("Deferring container '%s' until encryption key '%s' becomes available"),
-							*Container->Name, *Container->EncryptionKeyGuid);
-					}
-				}
-
-				if (EnumHasAnyFlags(Container->Flags, EOnDemandContainerFlags::PendingEncryptionKey) == false)
-				{
-					EnumAddFlags(Container->Flags, EOnDemandContainerFlags::Mounted);
-					Container->EncryptionKeyGuid.Empty();
-					Sb.Reset();
-					LexToString(Container->Flags, Sb);
-					UE_LOG(LogIoStoreOnDemand, Log, TEXT("Mounting container '%s', Entries=%d, Flags='%s'"),
-						*Container->Name, Container->ChunkEntries.Num(), Sb.ToString());
-				}
-			}
-		}
-
-		FOnDemandMountResult MountResult
-		{
-			.MountId = MoveTemp(Request->Args.MountId),
-			.DurationInSeconds = Request->DurationInSeconds
-		};
-
-		if (EnumHasAnyFlags(Request->Args.Options, EOnDemandMountOptions::CallbackOnGameThread))
-		{
-			ExecuteOnGameThread(
-				UE_SOURCE_LOCATION,
-				[OnCompleted = MoveTemp(Request->OnCompleted), MountResult = MoveTemp(MountResult)]() mutable
-				{
-					OnCompleted(MoveTemp(MountResult));
-				});
-		}
-		else
-		{
-			FOnDemandMountCompleted OnCompleted = MoveTemp(Request->OnCompleted);
-			OnCompleted(MoveTemp(MountResult));
-		}
+				.MountId = MoveTemp(Request->Args.MountId),
+				.Status = MoveTemp(MountStatus),
+				.DurationInSeconds = Request->DurationInSeconds
+			});
 	}
 
 	TArray<FSharedInstallRequest> LocalInstallRequests;
@@ -983,39 +997,20 @@ bool FOnDemandIoStore::Tick()
 	for (FSharedInstallRequest& Request : LocalInstallRequests)
 	{
 		FIoStatus Status = TickInstallRequest(*Request);
-		if (Status.GetErrorCode() == EIoErrorCode::Unknown)
-		{
-			UE_LOG(LogIoStoreOnDemand, Log, TEXT("Deferring install request, reason '%s'"), *Status.ToString());
-			continue;
-		}
 
 		{
 			UE::TUniqueLock Lock(MountRequestMutex);
 			InstallRequests.Remove(Request);
 		}
 
-		FOnDemandInstallResult InstallResult
-		{
-			.Status = MoveTemp(Status),
-			.DurationInSeconds = Request->DurationInSeconds,
-			.TotalContentSize = Request->TotalContentSize,
-			.TotalInstallSize = Request->TotalInstallSize
-		};
-
-		if (EnumHasAnyFlags(Request->Args.Options, EOnDemandInstallOptions::CallbackOnGameThread))
-		{
-			ExecuteOnGameThread(
-				UE_SOURCE_LOCATION,
-				[OnCompleted = MoveTemp(Request->OnCompleted), InstallResult = MoveTemp(InstallResult)]() mutable
-				{
-					OnCompleted(MoveTemp(InstallResult));
-				});
-		}
-		else
-		{
-			FOnDemandInstallCompleted OnCompleted = MoveTemp(Request->OnCompleted);
-			OnCompleted(MoveTemp(InstallResult));
-		}
+		OnInstallRequestComplete(*Request,
+			FOnDemandInstallResult
+			{
+				.Status = MoveTemp(Status),
+				.DurationInSeconds = Request->DurationInSeconds,
+				.TotalContentSize = Request->TotalContentSize,
+				.TotalInstallSize = Request->TotalInstallSize
+			});
 
 		bTicked = true;
 	}
@@ -1035,31 +1030,47 @@ FIoStatus FOnDemandIoStore::TickMountRequest(FMountRequest& MountRequest)
 
 	FOnDemandMountArgs& Args = MountRequest.Args;
 
+	bool bAnyPendingEncryptionKey = false;
+	bool bFoundContainers = false;
+
 	// Find containers matching the mount ID
 	{
 		UE::TUniqueLock Lock(ContainerMutex);
 		for (FSharedOnDemandContainer& Container : Containers)
 		{
-			if (Container->MountId == Args.MountId)
+			if (Container->MountId != Args.MountId)
 			{
-				MountRequest.Containers.Add(Container);
+				continue;
+			}
+			
+			bFoundContainers = true;
+			
+			if (EnumHasAnyFlags(Container->Flags, EOnDemandContainerFlags::PendingEncryptionKey))
+			{
+				bAnyPendingEncryptionKey = true;
 			}
 		}
 	}
 
-	// TODO: Treat this as an error?
-	if (MountRequest.Containers.IsEmpty() == false)
+	if (bFoundContainers)
 	{
-		return FIoStatus::Ok;
+		if (bAnyPendingEncryptionKey)
+		{
+			return EIoErrorCode::PendingEncryptionKey;
+		}
+		return EIoErrorCode::Ok;
 	}
 
+	// Containers haven't been created yet, do it now
+	TArray<FSharedOnDemandContainer> MountRequestContainers;
+	
 	FStringView Host, TocRelUrl;
 	Private::SplitHostUrl(Args.Url, Host, TocRelUrl);
 	const FStringView TocPath = FPathViews::GetPath(TocRelUrl);
 
 	if (Args.Toc)
 	{
-		CreateContainersFromToc(Args.MountId, TocPath, *Args.Toc, MountRequest.Containers);
+		CreateContainersFromToc(Args.MountId, TocPath, *Args.Toc, MountRequestContainers);
 	}
 	else if (Args.FilePath.IsEmpty() == false)
 	{
@@ -1075,7 +1086,7 @@ FIoStatus FOnDemandIoStore::TickMountRequest(FMountRequest& MountRequest)
 
 		Args.Toc = MakeUnique<FOnDemandToc>(TocStatus.ConsumeValueOrDie());
 
-		CreateContainersFromToc(Args.MountId, TocPath, *Args.Toc, MountRequest.Containers);
+		CreateContainersFromToc(Args.MountId, TocPath, *Args.Toc, MountRequestContainers);
 	}
 	else if (Args.Url.IsEmpty() == false)
 	{
@@ -1092,10 +1103,13 @@ FIoStatus FOnDemandIoStore::TickMountRequest(FMountRequest& MountRequest)
 
 		Args.Toc = MakeUnique<FOnDemandToc>(TocStatus.ConsumeValueOrDie());
 
-		CreateContainersFromToc(Args.MountId, TocPath, *Args.Toc, MountRequest.Containers);
+		CreateContainersFromToc(Args.MountId, TocPath, *Args.Toc, MountRequestContainers);
 	}
+	
+	TAnsiStringBuilder<512> ChunkUrlBuilder;
+	TMap<FString, FIoBuffer> MountRequestEncodedContainerHeaders;
 
-	for (const FSharedOnDemandContainer& Container : MountRequest.Containers)
+	for (const FSharedOnDemandContainer& Container : MountRequestContainers)
 	{
 		if (EnumHasAnyFlags(Args.Options, EOnDemandMountOptions::StreamOnDemand))
 		{
@@ -1105,11 +1119,84 @@ FIoStatus FOnDemandIoStore::TickMountRequest(FMountRequest& MountRequest)
 		{
 			EnumAddFlags(Container->Flags, EOnDemandContainerFlags::InstallOnDemand);
 		}
+
+		if (EnumHasAnyFlags(Container->Flags, EOnDemandContainerFlags::Encrypted) &&
+			Container->EncryptionKey.IsValid() == false)
+		{
+			FGuid KeyGuid;
+			ensure(FGuid::Parse(Container->EncryptionKeyGuid, KeyGuid));
+			if (FEncryptionKeyManager::Get().TryGetKey(KeyGuid, Container->EncryptionKey) == false)
+			{
+				EnumAddFlags(Container->Flags, EOnDemandContainerFlags::PendingEncryptionKey);
+				UE_LOG(LogIoStoreOnDemand, Log, TEXT("Deferring container '%s' until encryption key '%s' becomes available"),
+					*Container->Name, *Container->EncryptionKeyGuid);
+
+				bAnyPendingEncryptionKey = true;
+			}
+		}
+
+		if (EnumHasAnyFlags(Container->Flags, EOnDemandContainerFlags::PendingEncryptionKey))
+		{
+			//Fetch encoded container header and store it until it can be decrypted
+			FIoBuffer EncodedHeader;
+			if (FIoStatus Status = Private::FetchEncodedContainerHeader(Host, ChunkUrlBuilder, Container, EncodedHeader); !Status.IsOk())
+			{
+				return Status;
+			}
+			MountRequestEncodedContainerHeaders.Add(Container->UniqueName(), MoveTemp(EncodedHeader));
+
+			continue;
+		}
+
+		TStringBuilder<128> Sb;
+
+		EnumAddFlags(Container->Flags, EOnDemandContainerFlags::Mounted);
+		Container->EncryptionKeyGuid.Empty();
+		Sb.Reset();
+		LexToString(Container->Flags, Sb);
+		UE_LOG(LogIoStoreOnDemand, Log, TEXT("Mounting container '%s', Entries=%d, Flags='%s'"),
+			*Container->Name, Container->ChunkEntries.Num(), Sb.ToString());
+
+		// Fetch container header
+		FSharedContainerHeader Header = MakeShared<FIoContainerHeader>();
+		if (FIoStatus Status = Private::FetchContainerHeader(Host, ChunkUrlBuilder, Container, *Header); !Status.IsOk())
+		{
+			return Status;
+		}
+		Container->Header = MoveTemp(Header);
 	}
 
+	{
+		UE::TUniqueLock Lock(ContainerMutex);
+		Containers.Append(MoveTemp(MountRequestContainers));
+		EncodedContainerHeaders.Append(MoveTemp(MountRequestEncodedContainerHeaders));
+	}
+
+	if (bAnyPendingEncryptionKey)
+	{
+		return EIoErrorCode::PendingEncryptionKey;
+	}
 	return EIoErrorCode::Ok;
 }
 
+void FOnDemandIoStore::OnMountRequestComplete(FMountRequest& Request, FOnDemandMountResult&& MountResult)
+{
+	if (EnumHasAnyFlags(Request.Args.Options, EOnDemandMountOptions::CallbackOnGameThread))
+	{
+		ExecuteOnGameThread(
+			UE_SOURCE_LOCATION,
+			[OnCompleted = MoveTemp(Request.OnCompleted), MountResult = MoveTemp(MountResult)]() mutable
+			{
+				OnCompleted(MoveTemp(MountResult));
+			});
+	}
+	else
+	{
+		FOnDemandMountCompleted OnCompleted = MoveTemp(Request.OnCompleted);
+		OnCompleted(MoveTemp(MountResult));
+	}
+}
+
 FIoStatus FOnDemandIoStore::TickInstallRequest(FInstallRequest& InstallRequest)
 {
 	UE_LOG(LogIoStoreOnDemand, Log, TEXT("Ticking install request, ContentHandle='%s'"),
@@ -1142,84 +1229,13 @@ FIoStatus FOnDemandIoStore::TickInstallRequest(FInstallRequest& InstallRequest)
 	FStringView Host, TocRelUrl;
 	Private::SplitHostUrl(InstallRequest.Args.Url, Host, TocRelUrl);
 
-	// Only install content from non-streaming container(s)
-	TSet<FSharedOnDemandContainer> ContainersForInstallation(
-		GetContainers([](const FSharedOnDemandContainer& Container)
-		{ 
-			return EnumHasAnyFlags(Container->Flags, EOnDemandContainerFlags::InstallOnDemand);
-		})
-	);
-
-	// TODO: move this to mount
-	// Fetch all container headers
-	for (const FSharedOnDemandContainer& Container : ContainersForInstallation)
-	{
-		if (EnumHasAnyFlags(Container->Flags, EOnDemandContainerFlags::PendingEncryptionKey))
-		{
-			return FIoStatusBuilder(EIoErrorCode::Unknown)
-				<< TEXT("Deferring install request until encryption key(s) becomes available");
-		}
-
-		if (!Container->Header.IsValid())
-		{
-			FSharedContainerHeader Header = MakeShared<FIoContainerHeader>();
-			if (FIoStatus Status = Private::FetchContainerHeader(Host, ChunkUrl, Container, *Header); !Status.IsOk())
-			{
-				return Status; 
-			}
-			Container->Header = MoveTemp(Header);
-		}
-	}
-
+	TSet<FSharedOnDemandContainer> ContainersForInstallation;
 	TSet<FPackageId> PackageIdsToInstall;
-	for (const FPackageId& PackageId : InstallRequest.Args.PackageIds)
-	{
-		PackageIdsToInstall.Add(PackageId);
-	}
-
-	// Install all packages if no tag set(s) was specified
-	if (InstallRequest.Args.TagSets.IsEmpty())
+	FIoStatus GetContainersStatus = GetContainersAndPackagesForInstall(
+		InstallRequest.Args, ContainersForInstallation, PackageIdsToInstall);
+	if (!GetContainersStatus.IsOk())
 	{
-		for (const FSharedOnDemandContainer& Container : ContainersForInstallation)
-		{
-			if (Container->Header.IsValid() && Container->MountId == InstallRequest.Args.MountId)
-			{
-				for (const FPackageId& PackageId : Container->Header->PackageIds)
-				{
-					PackageIdsToInstall.Add(PackageId);
-				}
-			}
-		}
-	}
-	else
-	{
-		for (const FSharedOnDemandContainer& Container : ContainersForInstallation)
-		{
-			if (Container->Header.IsValid() == false)
-			{
-				continue;
-			}
-
-			if (InstallRequest.Args.MountId.IsEmpty() == false && InstallRequest.Args.MountId != Container->MountId)
-			{
-				continue;
-			}
-
-			for (const FString& Tag : InstallRequest.Args.TagSets)
-			{
-				for (const FOnDemandTagSet& TagSet : Container->TagSets)
-				{
-					if (TagSet.Tag == Tag)
-					{
-						for (const uint32 PackageIndex : TagSet.PackageIndicies)
-						{
-							const FPackageId PackageId = Container->Header->PackageIds[IntCastChecked<int32>(PackageIndex)];
-							PackageIdsToInstall.Add(PackageId);
-						}
-					}
-				}
-			}
-		}
+		return GetContainersStatus;
 	}
 
 	// Its OK for PackageIdsToInstall to be empty at this point. Any chunks not referenced by a package must still be installed.
@@ -1240,8 +1256,9 @@ FIoStatus FOnDemandIoStore::TickInstallRequest(FInstallRequest& InstallRequest)
 		const FIoChunkId ChunkId = CreatePackageDataChunkId(PackageId);
 		if (FIoDispatcher::Get().DoesChunkExist(ChunkId) == false)
 		{
-			// TODO: Fail the install request or just try to continue
-			UE_LOG(LogIoStoreOnDemand, Warning, TEXT("Missing package chunk '%s'"), *LexToString(ChunkId));
+			UE_LOG(LogIoStoreOnDemand, Error, TEXT("Missing package chunk '%s'"), *LexToString(ChunkId));
+			return FIoStatusBuilder(EIoErrorCode::UnknownChunkID) << 
+				TEXT("Missing package chunk '") << LexToString(ChunkId) << TEXT("'");
 		}
 	}
 
@@ -1388,33 +1405,75 @@ FIoStatus FOnDemandIoStore::TickInstallRequest(FInstallRequest& InstallRequest)
 	return EIoErrorCode::Ok;
 }
 
+void FOnDemandIoStore::OnInstallRequestComplete(FInstallRequest& Request, FOnDemandInstallResult&& InstallResult)
+{
+	if (EnumHasAnyFlags(Request.Args.Options, EOnDemandInstallOptions::CallbackOnGameThread))
+	{
+		ExecuteOnGameThread(
+			UE_SOURCE_LOCATION,
+			[OnCompleted = MoveTemp(Request.OnCompleted), InstallResult = MoveTemp(InstallResult)]() mutable
+			{
+				OnCompleted(MoveTemp(InstallResult));
+			});
+	}
+	else
+	{
+		FOnDemandInstallCompleted OnCompleted = MoveTemp(Request.OnCompleted);
+		OnCompleted(MoveTemp(InstallResult));
+	}
+}
+
 void FOnDemandIoStore::OnEncryptionKeyAdded(const FGuid& Id, const FAES::FAESKey& Key)
 {
+	TRACE_CPUPROFILER_EVENT_SCOPE(FOnDemandIoStore::OnEncryptionKeyAdded);
+
+	TUniqueLock Lock(ContainerMutex);
+
+	for (FSharedOnDemandContainer& Container : Containers)
 	{
-		TUniqueLock Lock(ContainerMutex);
+		if (EnumHasAnyFlags(Container->Flags, EOnDemandContainerFlags::PendingEncryptionKey) == false)
+		{
+			continue;
+		}
 
-		for (FSharedOnDemandContainer& Container : Containers)
+		FGuid KeyGuid;
+		ensure(FGuid::Parse(Container->EncryptionKeyGuid, KeyGuid));
+
+		if (FEncryptionKeyManager::Get().TryGetKey(KeyGuid, Container->EncryptionKey) == false)
 		{
-			if (EnumHasAnyFlags(Container->Flags, EOnDemandContainerFlags::PendingEncryptionKey))
-			{
-				FGuid KeyGuid;
-				ensure(FGuid::Parse(FString(StringCast<TCHAR>(*Container->EncryptionKeyGuid)), KeyGuid));
+			continue;
+		}
 
-				if (FEncryptionKeyManager::Get().TryGetKey(KeyGuid, Container->EncryptionKey))
-				{
-					UE_LOG(LogIoStoreOnDemand, Log, TEXT("Mounting container '%s', Entries=%d, Flags='%s'"),
-							*Container->Name, Container->ChunkEntries.Num(), *LexToString(Container->Flags));
+		UE_LOG(LogIoStoreOnDemand, Log, TEXT("Mounting container (found encryption key) '%s', Entries=%d, Flags='%s'"),
+				*Container->Name, Container->ChunkEntries.Num(), *LexToString(Container->Flags));
 
-					EnumRemoveFlags(Container->Flags, EOnDemandContainerFlags::PendingEncryptionKey);
-					EnumAddFlags(Container->Flags, EOnDemandContainerFlags::Mounted);
-					Container->EncryptionKeyGuid.Empty();
-				}
+		EnumRemoveFlags(Container->Flags, EOnDemandContainerFlags::PendingEncryptionKey);
+		
+		Container->EncryptionKeyGuid.Empty();
+
+		const FString ContainerName = Container->UniqueName();
+		const uint32 ContainerNameHash = GetTypeHash(ContainerName);
+
+		FIoBuffer* EncodedHeader = EncodedContainerHeaders.FindByHash(ContainerNameHash, ContainerName);
+		if (ensure(EncodedHeader))
+		{
+			FSharedContainerHeader Header = MakeShared<FIoContainerHeader>();
+			FIoStatus DecodeStatus = Private::DecodeContainerHeader(Container, EncodedHeader->GetView(), *Header);
+			EncodedHeader = nullptr;
+			EncodedContainerHeaders.RemoveByHash(ContainerNameHash, ContainerName);
+
+			if (ensure(DecodeStatus.IsOk()))
+			{
+				Container->Header = MoveTemp(Header);
+				EnumAddFlags(Container->Flags, EOnDemandContainerFlags::Mounted);
+			}
+			else
+			{
+				UE_LOG(LogIoStoreOnDemand, Error, TEXT("Failed to decode header when mounting container '%s', Entries=%d, Flags='%s'"),
+					*Container->Name, Container->ChunkEntries.Num(), *LexToString(Container->Flags));
 			}
 		}
 	}
-
-	// Tick pending install request waiting on encryption key(s)
-	TryEnterTickLoop();
 }
 
 void FOnDemandIoStore::CreateContainersFromToc(
@@ -1519,20 +1578,106 @@ void FOnDemandIoStore::CreateContainersFromToc(
 	}
 }
 
-TArray<FSharedOnDemandContainer> FOnDemandIoStore::GetContainers()
+FIoStatus FOnDemandIoStore::GetContainersForInstall(
+	FStringView MountId,
+	TSet<FSharedOnDemandContainer>& OutContainersForInstallation,
+	TSet<FSharedOnDemandContainer>& OutContainersWithMountId)
 {
+	// Only install content from non-streaming container(s)
+	TSet<FSharedOnDemandContainer> ContainersForInstallation;
+	TSet<FSharedOnDemandContainer> ContainersWithMountId;
+
 	UE::TUniqueLock Lock(ContainerMutex);
-	return Containers;
+
+	for (const FSharedOnDemandContainer& Container : Containers)
+	{
+		if (EnumHasAnyFlags(Container->Flags, EOnDemandContainerFlags::InstallOnDemand) == false)
+		{
+			continue;
+		}
+
+		// Check that the containers for the request have been mounted
+		if (MountId.IsEmpty() == false && Container->MountId == MountId)
+		{
+			if (EnumHasAnyFlags(Container->Flags, EOnDemandContainerFlags::PendingEncryptionKey) ||
+				EnumHasAnyFlags(Container->Flags, EOnDemandContainerFlags::Mounted) == false)
+			{
+				return EIoErrorCode::PendingEncryptionKey;
+			}
+
+			ContainersWithMountId.Add(Container);
+		}
+
+		// Header will not be valid until the container is fully decrypted/mounted
+		// We must check this under the containter lock so it doesn't race with the 
+		// decryption callback.
+		// After mount, the header can be accessed from multiple threads because it will not be changed again.
+		if (EnumHasAnyFlags(Container->Flags, EOnDemandContainerFlags::Mounted))
+		{
+			ContainersForInstallation.Add(Container);
+		}
+	}
+
+	OutContainersForInstallation = MoveTemp(ContainersForInstallation);
+	OutContainersWithMountId = MoveTemp(ContainersWithMountId);
+
+	return EIoErrorCode::Ok;
 }
 
-TArray<FSharedOnDemandContainer> FOnDemandIoStore::GetContainers(TFunctionRef<bool(const FSharedOnDemandContainer& Container)> Predicate)
+FIoStatus FOnDemandIoStore::GetContainersAndPackagesForInstall(
+	const FOnDemandInstallArgsCommon& Args, 
+	TSet<FSharedOnDemandContainer>& OutContainersForInstallation, 
+	TSet<FPackageId>& OutPackageIdsToInstall)
 {
-	TArray<FSharedOnDemandContainer> Ret;
+	TSet<FSharedOnDemandContainer> ContainersForInstallation;
+	TSet<FSharedOnDemandContainer> ContainersWithMountId;
+	FIoStatus GetContainersStatus = GetContainersForInstall(Args.MountId, ContainersForInstallation, ContainersWithMountId);
+	if (!GetContainersStatus.IsOk())
+	{
+		return GetContainersStatus;
+	}
 
-	UE::TUniqueLock Lock(ContainerMutex);
-	Algo::CopyIf(Containers, Ret, Predicate);
+	TSet<FPackageId> PackageIdsToInstall(Args.PackageIds);
 
-	return Ret;
+	// Install all packages if no tag set(s) was specified
+	if (Args.TagSets.IsEmpty())
+	{
+		for (const FSharedOnDemandContainer& Container : ContainersWithMountId)
+		{
+			for (const FPackageId& PackageId : Container->Header->PackageIds)
+			{
+				PackageIdsToInstall.Add(PackageId);
+			}
+		}
+	}
+	else
+	{
+		const TSet<FSharedOnDemandContainer>& SearchContainers =
+			Args.MountId.IsEmpty() ? ContainersForInstallation : ContainersWithMountId;
+
+		for (const FSharedOnDemandContainer& Container : SearchContainers)
+		{
+			for (const FString& Tag : Args.TagSets)
+			{
+				for (const FOnDemandTagSet& TagSet : Container->TagSets)
+				{
+					if (TagSet.Tag == Tag)
+					{
+						for (const uint32 PackageIndex : TagSet.PackageIndicies)
+						{
+							const FPackageId PackageId = Container->Header->PackageIds[IntCastChecked<int32>(PackageIndex)];
+							PackageIdsToInstall.Add(PackageId);
+						}
+					}
+				}
+			}
+		}
+	}
+
+	OutContainersForInstallation = MoveTemp(ContainersForInstallation);
+	OutPackageIdsToInstall = MoveTemp(PackageIdsToInstall);
+
+	return EIoErrorCode::Ok;
 }
 
 void FOnDemandIoStore::ReleaseContent(FOnDemandInternalContentHandle& ContentHandle)
diff --git a/Engine/Source/Runtime/Experimental/IoStore/OnDemand/Private/OnDemandIoStore.h b/Engine/Source/Runtime/Experimental/IoStore/OnDemand/Private/OnDemandIoStore.h
--- a/Engine/Source/Runtime/Experimental/IoStore/OnDemand/Private/OnDemandIoStore.h
+++ b/Engine/Source/Runtime/Experimental/IoStore/OnDemand/Private/OnDemandIoStore.h
@@ -7,6 +7,7 @@
 #include "Async/Mutex.h"
 #include "Containers/AnsiString.h"
 #include "Containers/BitArray.h"
+#include "IO/IoBuffer.h"
 #include "IO/IoStoreOnDemand.h"
 #include "IO/IoHash.h"
 #include "IO/IoChunkEncoding.h"
@@ -190,11 +191,11 @@ TBitArray<> FOnDemandContainer::GetReferencedChunkEntries() const
 ///////////////////////////////////////////////////////////////////////////////
 namespace Private
 {
-	static FIoStatus FetchContainerHeader(
-		FStringView Host,
-		FAnsiStringBuilderBase& ChunkUrlBuilder,
+	static FIoStatus DecodeContainerChunk(
 		FSharedOnDemandContainer Container,
-		FIoContainerHeader& OutHeader);
+		const FOnDemandChunkEntry& Entry,
+		FMemoryView EncodedHeaderChunk,
+		FIoBuffer& OutRawChunk);
 }
 
 struct FOnDemandChunkInfo
@@ -218,11 +219,11 @@ struct FOnDemandChunkInfo
 
 private:
 	friend class FOnDemandIoStore;
-	friend FIoStatus Private::FetchContainerHeader(
-		FStringView Host,
-		FAnsiStringBuilderBase& ChunkUrlBuilder,
+	friend FIoStatus Private::DecodeContainerChunk(
 		FSharedOnDemandContainer Container,
-		FIoContainerHeader& OutHeader);
+		const FOnDemandChunkEntry& Entry,
+		FMemoryView EncodedHeaderChunk,
+		FIoBuffer& OutRawChunk);
 
 	FOnDemandChunkInfo(FSharedOnDemandContainer InContainer, const FOnDemandChunkEntry& InEntry)
 		: SharedContainer(InContainer)
@@ -251,15 +252,12 @@ class FOnDemandIoStore
 {
 	struct FMountRequest
 	{
-		FOnDemandMountArgs					Args;
-		FOnDemandMountCompleted				OnCompleted;
-		TArray<FSharedOnDemandContainer>	Containers;
-		std::atomic_bool					bCancelled{false};
-		double								DurationInSeconds = 0.0;
+		FOnDemandMountArgs			Args;
+		FOnDemandMountCompleted		OnCompleted;
+		double						DurationInSeconds = 0.0;
 	};
 
 	using FSharedMountRequest	= TSharedPtr<FMountRequest>;
-	using FMountRequestMap		= TMap<FString, FSharedMountRequest>;
 
 	struct FInstallRequest
 	{
@@ -287,13 +285,18 @@ public:
 	TIoStatusOr<uint64>		GetInstallSize(const FOnDemandGetInstallSizeArgs& Args) const;
 	FOnDemandChunkInfo		GetStreamingChunkInfo(const FIoChunkId& ChunkId);
 	FOnDemandChunkInfo		GetInstalledChunkInfo(const FIoChunkId& ChunkId);
-	TArray<FSharedOnDemandContainer> GetContainers();
-	TArray<FSharedOnDemandContainer> GetContainers(
-		TFunctionRef<bool(const FSharedOnDemandContainer& Container)> Predicate);
 	void					ReleaseContent(FOnDemandInternalContentHandle& ContentHandle);
 	void					GetReferencedContent(TArray<FSharedOnDemandContainer>& OutContainers, TArray<TBitArray<>>& OutChunkEntryIndices);
 
 private:
+	FIoStatus				GetContainersForInstall(
+		FStringView MountId,
+		TSet<FSharedOnDemandContainer>& OutContainersForInstallation,
+		TSet<FSharedOnDemandContainer>& OutContainersWithMountId);
+	FIoStatus				GetContainersAndPackagesForInstall(
+		const FOnDemandInstallArgsCommon& Args,
+		TSet<FSharedOnDemandContainer>& OutContainersForInstallation,
+		TSet<FPackageId>& OutPackageIdsToInstall);
 	void					OnPostFork(EForkProcessRole ProcessRole);
 	FIoStatus				InitializeOnDemandInstallCache();
 	FOnDemandChunkInfo		GetChunkInfo(const FIoChunkId& ChunkId, EOnDemandContainerFlags ContainerFlags);
@@ -301,7 +304,9 @@ private:
 	void					TickLoop();
 	bool					Tick();
 	FIoStatus				TickMountRequest(FMountRequest& MountRequest);
+	void					OnMountRequestComplete(FMountRequest& MountRequest, FOnDemandMountResult&& MountResult);
 	FIoStatus				TickInstallRequest(FInstallRequest& InstallRequest);
+	void					OnInstallRequestComplete(FInstallRequest& InstallRequest, FOnDemandInstallResult&& InstallResult);
 	void					OnEncryptionKeyAdded(const FGuid& Id, const FAES::FAESKey& Key);
 	static void				CreateContainersFromToc(
 								FStringView MountId,
@@ -313,9 +318,10 @@ private:
 	FSharedPackageStoreBackend			PackageStoreBackend;
 	FDelegateHandle						OnMountPakHandle;
 	TArray<FSharedOnDemandContainer>	Containers;
+	TMap<FString, FIoBuffer>			EncodedContainerHeaders;
 	UE::FMutex							ContainerMutex;
 
-	FMountRequestMap					MountRequests;
+	TArray<FSharedMountRequest>			MountRequests;
 	TArray<FSharedInstallRequest>		InstallRequests;
 	UE::FMutex							MountRequestMutex;
 
diff --git a/Engine/Source/Runtime/Experimental/IoStore/OnDemand/Public/IO/IoStoreOnDemand.h b/Engine/Source/Runtime/Experimental/IoStore/OnDemand/Public/IO/IoStoreOnDemand.h
--- a/Engine/Source/Runtime/Experimental/IoStore/OnDemand/Public/IO/IoStoreOnDemand.h
+++ b/Engine/Source/Runtime/Experimental/IoStore/OnDemand/Public/IO/IoStoreOnDemand.h
@@ -321,17 +321,21 @@ enum class EOnDemandInstallOptions
 };
 ENUM_CLASS_FLAGS(EOnDemandInstallOptions);
 
-/** Arguments for installing/downloading on-demand content. */
-struct FOnDemandInstallArgs
+struct FOnDemandInstallArgsCommon
 {
-	/** URL from where to download the chunks. */
-	FString Url;
 	/** Install all content from containers matching this mount ID. */
 	FString MountId;
 	/** Install content matching a set of tag(s) and optionally the mount ID. */
 	TArray<FString> TagSets;
 	/** Package ID's to install. */
 	TArray<FPackageId> PackageIds;
+};
+
+/** Arguments for installing/downloading on-demand content. */
+struct FOnDemandInstallArgs : public FOnDemandInstallArgsCommon
+{
+	/** URL from where to download the chunks. */
+	FString Url;
 	/** Content handle. */
 	FOnDemandContentHandle ContentHandle;
 	/** Install options. */
@@ -354,23 +358,9 @@ struct FOnDemandInstallResult
 /** Install completion callback. */
 using FOnDemandInstallCompleted = TUniqueFunction<void(FOnDemandInstallResult)>;
 
-struct FOnDemandGetInstallSizeArgs
+/** Arguments for getting the size of on-demand content */
+struct FOnDemandGetInstallSizeArgs : public FOnDemandInstallArgsCommon
 {
-	/* Includes all packages from matching containers */
-	struct FPackagesWhere
-	{
-		/* The MountId to collect packages from*/
-		FString MountId;
-		/* Includes only packages from the specificed tag sets if not empty */
-		TArray<FString> TagSets;
-	};
-
-	/* Includes all packages from matching containers */
-	// TODO: implement this
-	// TArray<FPackagesWhere> PackagesWhere;
-
-	/* Explicit packages to include */
-	TArray<FPackageId> Packages;
 };
 
 class FIoStoreOnDemandModule
