Geometry collection : expose a way for blueprint to set the component space transform of a specific roiot proxy mesh

diff --git a/Engine/Source/Runtime/Experimental/GeometryCollectionEngine/Private/GeometryCollection/GeometryCollectionComponent.cpp b/Engine/Source/Runtime/Experimental/GeometryCollectionEngine/Private/GeometryCollection/GeometryCollectionComponent.cpp
--- a/Engine/Source/Runtime/Experimental/GeometryCollectionEngine/Private/GeometryCollection/GeometryCollectionComponent.cpp
+++ b/Engine/Source/Runtime/Experimental/GeometryCollectionEngine/Private/GeometryCollection/GeometryCollectionComponent.cpp
@@ -6280,6 +6280,12 @@ void UGeometryCollectionComponent::ForceBrokenForCustomRenderer(bool bForceBroke
 	RefreshCustomRenderer();
 }
 
+void UGeometryCollectionComponent::SetRootProxyComponentSpaceTransform(int32 Index, const FTransform& RootProxyTransform)
+{
+	SetRootProxyLocalTransform(Index, FTransform3f(RootProxyTransform));
+	RefreshRootProxies();
+}
+
 void UGeometryCollectionComponent::SetRootProxyLocalTransform(int32 Index, const FTransform3f& RootProxyTransform)
 {
 	if (RootProxyLocalTransforms.IsEmpty() && RestCollection)
diff --git a/Engine/Source/Runtime/Experimental/GeometryCollectionEngine/Public/GeometryCollection/GeometryCollectionComponent.h b/Engine/Source/Runtime/Experimental/GeometryCollectionEngine/Public/GeometryCollection/GeometryCollectionComponent.h
--- a/Engine/Source/Runtime/Experimental/GeometryCollectionEngine/Public/GeometryCollection/GeometryCollectionComponent.h
+++ b/Engine/Source/Runtime/Experimental/GeometryCollectionEngine/Public/GeometryCollection/GeometryCollectionComponent.h
@@ -1333,6 +1333,14 @@ public:
 	UFUNCTION(BlueprintCallable, Category = "Physics")
 	GEOMETRYCOLLECTIONENGINE_API void ForceBrokenForCustomRenderer(bool bForceBroken);
 
+	/** 
+	* blueprint function to set a specific root proxy local transform 
+	* warning: when called from C++ and calling it on multiple root proxies it is recommended to use SetRootProxyLocalTransform instead and then call RefreshCustomRenderer
+	* to avoid the cost of refreshing the renderer each time
+	*/
+	UFUNCTION(BlueprintCallable, Category = "Physics")
+	GEOMETRYCOLLECTIONENGINE_API void SetRootProxyComponentSpaceTransform(int32 Index, const FTransform& RootProxyTransform);
+
 	/** Set a specific root proxy local transform */
 	GEOMETRYCOLLECTIONENGINE_API void SetRootProxyLocalTransform(int32 Index, const FTransform3f& RootProxyTransform);
 
