Iris
- Prevent NetSerialization readjournal from being overwritten after an error is reported.

diff --git a/Engine/Source/Runtime/Experimental/Iris/Core/Private/Iris/Serialization/NetSerializationContext.cpp b/Engine/Source/Runtime/Experimental/Iris/Core/Private/Iris/Serialization/NetSerializationContext.cpp
--- a/Engine/Source/Runtime/Experimental/Iris/Core/Private/Iris/Serialization/NetSerializationContext.cpp
+++ b/Engine/Source/Runtime/Experimental/Iris/Core/Private/Iris/Serialization/NetSerializationContext.cpp
@@ -25,13 +25,19 @@ bool FNetSerializationContext::IsBitStreamOverflown() const
 }
 
 void FNetSerializationContext::AddReadJournalEntry(const TCHAR* Name)
-{ 
-	ReadJournal.AddEntry(Name, BitStreamReader->GetPosBits(), ErrorContext.GetObjectHandle());
+{
+	if (!HasErrorOrOverflow())
+	{
+		ReadJournal.AddEntry(Name, BitStreamReader->GetPosBits(), ErrorContext.GetObjectHandle());
+	}
 }
 
 void FNetSerializationContext::AddReadJournalEntry(const FNetDebugName* DebugName)
 {
-	ReadJournal.AddEntry(DebugName->Name, BitStreamReader->GetPosBits(), ErrorContext.GetObjectHandle());
+	if (!HasErrorOrOverflow())
+	{
+		ReadJournal.AddEntry(DebugName->Name, BitStreamReader->GetPosBits(), ErrorContext.GetObjectHandle());
+	}
 }
 
 FString FNetSerializationContext::PrintReadJournal()
