Initialize Chaos struct variables to fix build errors. CVar command to confirm: CoreUObject.AttemptToFindUninitializedScriptStructMembers
#rnx

diff --git a/Engine/Source/Runtime/Experimental/ChaosVisualDebugger/Public/DataWrappers/ChaosVDAccelerationStructureDataWrappers.h b/Engine/Source/Runtime/Experimental/ChaosVisualDebugger/Public/DataWrappers/ChaosVDAccelerationStructureDataWrappers.h
--- a/Engine/Source/Runtime/Experimental/ChaosVisualDebugger/Public/DataWrappers/ChaosVDAccelerationStructureDataWrappers.h
+++ b/Engine/Source/Runtime/Experimental/ChaosVisualDebugger/Public/DataWrappers/ChaosVDAccelerationStructureDataWrappers.h
@@ -46,7 +46,7 @@ struct FChaosVDBVCellElementDataWrapper : public FChaosVDWrapperDataBase
 	UPROPERTY(VisibleAnywhere, Category="Data")
 	FBox Bounds;
 	UPROPERTY(VisibleAnywhere, Category="Data")
-	int32 ParticleIndex;
+	int32 ParticleIndex = INDEX_NONE;
 
 	FIntVector3 StartIdx = FIntVector3::ZeroValue;
 	FIntVector3 EndIdx = FIntVector3::ZeroValue;
diff --git a/Engine/Source/Runtime/Experimental/ChaosVisualDebugger/Public/DataWrappers/ChaosVDCollisionDataWrappers.h b/Engine/Source/Runtime/Experimental/ChaosVisualDebugger/Public/DataWrappers/ChaosVDCollisionDataWrappers.h
--- a/Engine/Source/Runtime/Experimental/ChaosVisualDebugger/Public/DataWrappers/ChaosVDCollisionDataWrappers.h
+++ b/Engine/Source/Runtime/Experimental/ChaosVisualDebugger/Public/DataWrappers/ChaosVDCollisionDataWrappers.h
@@ -400,7 +400,7 @@ struct CHAOSVDRUNTIME_API FChaosVDParticlePairMidPhase
 	int32 SolverID = INDEX_NONE;
 
 	UPROPERTY(VisibleAnywhere, Category=General)
-	EChaosVDMidPhaseType MidPhaseType;
+	EChaosVDMidPhaseType MidPhaseType = EChaosVDMidPhaseType::Unknown;
 
 	UPROPERTY(VisibleAnywhere, Category=Flags)
 	uint8 bIsActive:1 = false;
