Fix struct initialize errors
#rnx

diff --git a/Engine/Plugins/Experimental/DaySequence/Source/DaySequence/Public/DaySequenceCollectionAsset.h b/Engine/Plugins/Experimental/DaySequence/Source/DaySequence/Public/DaySequenceCollectionAsset.h
--- a/Engine/Plugins/Experimental/DaySequence/Source/DaySequence/Public/DaySequenceCollectionAsset.h
+++ b/Engine/Plugins/Experimental/DaySequence/Source/DaySequence/Public/DaySequenceCollectionAsset.h
@@ -30,7 +30,7 @@ struct FDaySequenceCollectionEntry
 	
 	/* The offset hierarchical bias assigned to this collection entry. */
 	UPROPERTY(EditAnywhere, Category="Day Sequence")
-	int BiasOffset;
+	int BiasOffset = 0;
 	
 	/* The set of conditions which must evaluate to their expected values for this entry to be active. */
 	UPROPERTY(EditAnywhere, Category="Day Sequence")
diff --git a/Engine/Plugins/Experimental/EditorDataStorageFeatures/Source/TedsAssetData/Public/TedsAssetDataColumns.h b/Engine/Plugins/Experimental/EditorDataStorageFeatures/Source/TedsAssetData/Public/TedsAssetDataColumns.h
--- a/Engine/Plugins/Experimental/EditorDataStorageFeatures/Source/TedsAssetData/Public/TedsAssetDataColumns.h
+++ b/Engine/Plugins/Experimental/EditorDataStorageFeatures/Source/TedsAssetData/Public/TedsAssetDataColumns.h
@@ -96,7 +96,7 @@ struct FDiskSizeColumn final : public FTypedElementDataStorageColumn
 
 	// Size on disk (in bytes)
 	UPROPERTY()
-	int64 DiskSize;
+	int64 DiskSize = 0;
 };
 
 // Column to store the verse path of an asset
diff --git a/Engine/Plugins/Runtime/ComputeFramework/Source/ComputeFramework/Public/ComputeFramework/ShaderParamTypeDefinition.h b/Engine/Plugins/Runtime/ComputeFramework/Source/ComputeFramework/Public/ComputeFramework/ShaderParamTypeDefinition.h
--- a/Engine/Plugins/Runtime/ComputeFramework/Source/ComputeFramework/Public/ComputeFramework/ShaderParamTypeDefinition.h
+++ b/Engine/Plugins/Runtime/ComputeFramework/Source/ComputeFramework/Public/ComputeFramework/ShaderParamTypeDefinition.h
@@ -349,7 +349,7 @@ public:
 	EShaderResourceType	ResourceType = EShaderResourceType::None;
 
 	UPROPERTY()
-	EShaderParamModifier Modifier;
+	EShaderParamModifier Modifier = EShaderParamModifier::None;
 	
 	bool IsAnyBufferType() const
 	{
diff --git a/Engine/Source/Runtime/Experimental/ChaosVisualDebugger/Public/DataWrappers/ChaosVDAccelerationStructureDataWrappers.h b/Engine/Source/Runtime/Experimental/ChaosVisualDebugger/Public/DataWrappers/ChaosVDAccelerationStructureDataWrappers.h
--- a/Engine/Source/Runtime/Experimental/ChaosVisualDebugger/Public/DataWrappers/ChaosVDAccelerationStructureDataWrappers.h
+++ b/Engine/Source/Runtime/Experimental/ChaosVisualDebugger/Public/DataWrappers/ChaosVDAccelerationStructureDataWrappers.h
@@ -44,7 +44,8 @@ struct FChaosVDBVCellElementDataWrapper : public FChaosVDWrapperDataBase
 	GENERATED_BODY()
 
 	UPROPERTY(VisibleAnywhere, Category="Data")
-	FBox Bounds;
+	FBox Bounds = FBox(EForceInit::ForceInit);
+
 	UPROPERTY(VisibleAnywhere, Category="Data")
 	int32 ParticleIndex = INDEX_NONE;
 
