Iris - Crash fix for removing InstancedStruct from a replicated array and adding the same struct type again.

diff --git a/Engine/Plugins/Experimental/StructUtils/Source/StructUtils/Private/InstancedStructNetSerializer.cpp b/Engine/Plugins/Experimental/StructUtils/Source/StructUtils/Private/InstancedStructNetSerializer.cpp
--- a/Engine/Plugins/Experimental/StructUtils/Source/StructUtils/Private/InstancedStructNetSerializer.cpp
+++ b/Engine/Plugins/Experimental/StructUtils/Source/StructUtils/Private/InstancedStructNetSerializer.cpp
@@ -497,6 +497,8 @@ void FInstancedStructNetSerializer::InternalFreeStructInstance(FNetSerialization
 		FNetFreeDynamicStateArgs FreeArgs;
 		FreeArgs.NetSerializerConfig = &StructConfig;
 		FreeArgs.Source = NetSerializerValuePointer(Value.StructData.GetData());
+
+		StructNetSerializer->FreeDynamicState(Context, FreeArgs);
 	}
 }
 
diff --git a/Engine/Source/Runtime/Experimental/Iris/Core/Private/Iris/Serialization/ArrayPropertyNetSerializer.cpp b/Engine/Source/Runtime/Experimental/Iris/Core/Private/Iris/Serialization/ArrayPropertyNetSerializer.cpp
--- a/Engine/Source/Runtime/Experimental/Iris/Core/Private/Iris/Serialization/ArrayPropertyNetSerializer.cpp
+++ b/Engine/Source/Runtime/Experimental/Iris/Core/Private/Iris/Serialization/ArrayPropertyNetSerializer.cpp
@@ -833,6 +833,15 @@ void FArrayPropertyNetSerializer::AdjustArraySize(FNetSerializationContext& Cont
 	// If element count is within the allocated capacity we just change the number of elements
 	else
 	{
+		// Zero out data just to be sure, to always give new elements the same initial state.
+		if (NewElementCount > Array.ElementCount)
+		{
+			const FReplicationStateDescriptor* ElementStateDescriptor = Config->StateDescriptor;
+			const SIZE_T ElementSize = ElementStateDescriptor->InternalSize;
+			void* ElementsToZeroOut = (void*)(NetSerializerValuePointer(Array.ElementStorage) + (Array.ElementCount * ElementSize));
+
+			FMemory::Memzero(ElementsToZeroOut, ElementSize * (NewElementCount - Array.ElementCount));
+		}
 		Array.ElementCount = NewElementCount;
 	}
 }
