[ChaosVD] Adding properties to the serialization code that were being filled but not serialized.

diff --git a/Engine/Source/Runtime/Core/Public/UObject/FortniteSeasonBranchObjectVersion.h b/Engine/Source/Runtime/Core/Public/UObject/FortniteSeasonBranchObjectVersion.h
--- a/Engine/Source/Runtime/Core/Public/UObject/FortniteSeasonBranchObjectVersion.h
+++ b/Engine/Source/Runtime/Core/Public/UObject/FortniteSeasonBranchObjectVersion.h
@@ -48,6 +48,9 @@ struct CORE_API FFortniteSeasonBranchObjectVersion
 		// Changing Chaos::FImplicitObjectUnion to store an int32 vs a uint16 for NumLeafObjects.
 		ChaosImplicitObjectUnionLeafObjectsToInt32,
 
+		// Chaos Visual Debugger : Adding serialization for properties that were being recorded, but not serialized
+		CVDSerializationFixMissingSerializationProperties,
+
 		// -----<new versions can be added above this line>-------------------------------------------------
 		VersionPlusOne,
 		LatestVersion = VersionPlusOne - 1
diff --git a/Engine/Source/Runtime/Experimental/ChaosVisualDebugger/Private/DataWrappers/ChaosVDAccelerationStructureDataWrappers.cpp b/Engine/Source/Runtime/Experimental/ChaosVisualDebugger/Private/DataWrappers/ChaosVDAccelerationStructureDataWrappers.cpp
--- a/Engine/Source/Runtime/Experimental/ChaosVisualDebugger/Private/DataWrappers/ChaosVDAccelerationStructureDataWrappers.cpp
+++ b/Engine/Source/Runtime/Experimental/ChaosVisualDebugger/Private/DataWrappers/ChaosVDAccelerationStructureDataWrappers.cpp
@@ -4,6 +4,8 @@
 
 #include "DataWrappers/ChaosVDDataSerializationMacros.h"
 
+#include "UObject/FortniteSeasonBranchObjectVersion.h"
+
 
 bool FChaosVDBVCellElementDataWrapper::Serialize(FArchive& Ar)
 {
@@ -86,8 +88,11 @@ bool FChaosVDAABBTreePayloadBoundsElement::Serialize(FArchive& Ar)
 	Ar << ParticleIndex;
 	Ar << Bounds;
 
-	// TODO: This will be uncommented in the nex CL with the required object version bump
-	//Ar << ActualBounds;
+	Ar.UsingCustomVersion(FFortniteSeasonBranchObjectVersion::GUID);
+	if (Ar.CustomVer(FFortniteSeasonBranchObjectVersion::GUID) >= FFortniteSeasonBranchObjectVersion::CVDSerializationFixMissingSerializationProperties)
+	{
+		Ar << ActualBounds;
+	}
 
 	return !Ar.IsError();
 }
diff --git a/Engine/Source/Runtime/Experimental/ChaosVisualDebugger/Private/DataWrappers/ChaosVDParticleDataWrapper.cpp b/Engine/Source/Runtime/Experimental/ChaosVisualDebugger/Private/DataWrappers/ChaosVDParticleDataWrapper.cpp
--- a/Engine/Source/Runtime/Experimental/ChaosVisualDebugger/Private/DataWrappers/ChaosVDParticleDataWrapper.cpp
+++ b/Engine/Source/Runtime/Experimental/ChaosVisualDebugger/Private/DataWrappers/ChaosVDParticleDataWrapper.cpp
@@ -2,12 +2,25 @@
 
 #include "DataWrappers/ChaosVDParticleDataWrapper.h"
 
+#include "UObject/FortniteSeasonBranchObjectVersion.h"
+
 // @note: Tracing an scene with 1000 particles moving, Manually serializing the structs is ~20% faster
 // than normal UStruct serialization in unversioned mode. As we do this at runtime when tracing in development builds, this is important.
 // One of the downside is that this will be more involved to maintain as any versioning needs to be done by hand.
 
 bool FChaosVDFRigidParticleControlFlags::Serialize(FArchive& Ar)
 {
+	Ar.UsingCustomVersion(FFortniteSeasonBranchObjectVersion::GUID);
+	if (Ar.CustomVer(FFortniteSeasonBranchObjectVersion::GUID) >= FFortniteSeasonBranchObjectVersion::CVDSerializationFixMissingSerializationProperties)
+	{
+		Ar << bHasValidData;
+
+		if (!bHasValidData)
+		{
+			return !Ar.IsError();
+		}
+	}
+	
 	Ar << bGravityEnabled;
 	Ar << bCCDEnabled;
 	Ar << bOneWayInteractionEnabled;
