Iris
- Fix for issue where replicated overridden blueprint events/functions would not call the overrriden event/function

diff --git a/Engine/Source/Runtime/Experimental/Iris/Core/Private/Iris/ReplicationSystem/NetBlob/NetRPC.cpp b/Engine/Source/Runtime/Experimental/Iris/Core/Private/Iris/ReplicationSystem/NetBlob/NetRPC.cpp
--- a/Engine/Source/Runtime/Experimental/Iris/Core/Private/Iris/ReplicationSystem/NetBlob/NetRPC.cpp
+++ b/Engine/Source/Runtime/Experimental/Iris/Core/Private/Iris/ReplicationSystem/NetBlob/NetRPC.cpp
@@ -559,22 +559,30 @@ void FNetRPC::CallFunction(FNetRPCCallContext& CallContext)
 		FReplicationStateOperations::Dequantize(Context, FunctionParameters, QuantizedBlobState.GetStateBuffer(), BlobDescriptor);
 	}
 
+	// Since the replicated FunctionLocator references the SuperFunction, we must lookup the actual function to call from the target object to properly support BP derived functions.
+	// See: JIRA: UE-220400
+	const UFunction* ActualFunction = Object->FindFunction(Function->GetFName());
+	if (ActualFunction == nullptr)
+	{
+		ActualFunction = Function;
+	}
+
 	// Forward function
 	if (const FForwardNetRPCCallMulticastDelegate& Delegate = CallContext.GetForwardNetRPCCallDelegate(); Delegate.IsBound())
 	{
 		UObject* RootObject = NetRPC_GetRootObject(Context, NetObjectReference);
 		UObject* SubObject = (Object != RootObject ? Object : static_cast<UObject*>(nullptr));
-		Delegate.Broadcast(RootObject, SubObject, const_cast<UFunction*>(Function), FunctionParameters);
+		Delegate.Broadcast(RootObject, SubObject, const_cast<UFunction*>(ActualFunction), FunctionParameters);
 	}
 
 	// Call function
 	{
 #if IRIS_CLIENT_PROFILER_ENABLE
-		UE::Net::FClientProfiler::RecordRPC(Function->GetFName());
+		UE::Net::FClientProfiler::RecordRPC(ActualFunction->GetFName());
 #endif
 
 		UE::Net::FScopedNetContextRPC CallingRPC;
-		Object->ProcessEvent(const_cast<UFunction*>(Function), FunctionParameters);
+		Object->ProcessEvent(const_cast<UFunction*>(ActualFunction), FunctionParameters);
 	}
 
 	// Deinitialize function parameters
