IOStoreOnDemand - Add GetInstallSizesByMountId()

diff --git a/Engine/Source/Runtime/Experimental/IoStore/OnDemand/Private/IoStoreOnDemand.cpp b/Engine/Source/Runtime/Experimental/IoStore/OnDemand/Private/IoStoreOnDemand.cpp
--- a/Engine/Source/Runtime/Experimental/IoStore/OnDemand/Private/IoStoreOnDemand.cpp
+++ b/Engine/Source/Runtime/Experimental/IoStore/OnDemand/Private/IoStoreOnDemand.cpp
@@ -1107,6 +1107,16 @@ TIoStatusOr<uint64> FIoStoreOnDemandModule::GetInstallSize(const FOnDemandGetIns
 	return FIoStatus(EIoErrorCode::InvalidCode, NotInitializedError);
 }
 
+FIoStatus FIoStoreOnDemandModule::GetInstallSizesByMountId(const FOnDemandGetInstallSizeArgs& Args, TMap<FString, uint64>& OutSizesByMountId) const
+{
+	if (IoStore)
+	{
+		return IoStore->GetInstallSizesByMountId(Args, OutSizesByMountId);
+	}
+
+	return FIoStatus(EIoErrorCode::InvalidCode, NotInitializedError);
+}
+
 void FIoStoreOnDemandModule::InitializeInternal()
 {
 	LLM_SCOPE_BYTAG(Ias);
diff --git a/Engine/Source/Runtime/Experimental/IoStore/OnDemand/Private/OnDemandIoStore.cpp b/Engine/Source/Runtime/Experimental/IoStore/OnDemand/Private/OnDemandIoStore.cpp
--- a/Engine/Source/Runtime/Experimental/IoStore/OnDemand/Private/OnDemandIoStore.cpp
+++ b/Engine/Source/Runtime/Experimental/IoStore/OnDemand/Private/OnDemandIoStore.cpp
@@ -772,6 +772,36 @@ TIoStatusOr<uint64> FOnDemandIoStore::GetInstallSize(const FOnDemandGetInstallSi
 	return RetSize;
 }
 
+FIoStatus FOnDemandIoStore::GetInstallSizesByMountId(const FOnDemandGetInstallSizeArgs& Args, TMap<FString, uint64>& OutSizesByMountId) const
+{
+	TRACE_CPUPROFILER_EVENT_SCOPE(FOnDemandIoStore::GetInstallSizesByMountId);
+
+	TSet<FSharedOnDemandContainer> ContainersForInstallation;
+	TSet<FPackageId> PackageIdsToInstall;
+	FIoStatus GetContainersStatus = const_cast<FOnDemandIoStore*>(this)->GetContainersAndPackagesForInstall(
+		Args, ContainersForInstallation, PackageIdsToInstall);
+	if (!GetContainersStatus.IsOk())
+	{
+		return GetContainersStatus;
+	}
+
+	Private::FInstallData InstallData;
+	TSet<FPackageId> Missing;
+
+	FIoStatus Status = BuildInstallData(ContainersForInstallation, PackageIdsToInstall, InstallData, Missing);
+	if (Status.IsOk() == false)
+	{
+		return Status;
+	}
+
+	for (const TPair<FSharedOnDemandContainer, Private::FContainerInstallData>& Pair : InstallData)
+	{
+		OutSizesByMountId.FindOrAdd(Pair.Key->MountId, 0) += Pair.Value.TotalSize;
+	}
+
+	return EIoErrorCode::Ok;
+}
+
 FOnDemandChunkInfo FOnDemandIoStore::GetStreamingChunkInfo(const FIoChunkId& ChunkId)
 {
 	return GetChunkInfo(ChunkId, EOnDemandContainerFlags::Mounted | EOnDemandContainerFlags::StreamOnDemand);
diff --git a/Engine/Source/Runtime/Experimental/IoStore/OnDemand/Private/OnDemandIoStore.h b/Engine/Source/Runtime/Experimental/IoStore/OnDemand/Private/OnDemandIoStore.h
--- a/Engine/Source/Runtime/Experimental/IoStore/OnDemand/Private/OnDemandIoStore.h
+++ b/Engine/Source/Runtime/Experimental/IoStore/OnDemand/Private/OnDemandIoStore.h
@@ -283,6 +283,7 @@ public:
 	void					Install(FOnDemandInstallArgs&& Args, FOnDemandInstallCompleted&& OnCompleted);
 	FIoStatus				Unmount(FStringView MountId);
 	TIoStatusOr<uint64>		GetInstallSize(const FOnDemandGetInstallSizeArgs& Args) const;
+	FIoStatus				GetInstallSizesByMountId(const FOnDemandGetInstallSizeArgs& Args, TMap<FString, uint64>& OutSizesByMountId) const;
 	FOnDemandChunkInfo		GetStreamingChunkInfo(const FIoChunkId& ChunkId);
 	FOnDemandChunkInfo		GetInstalledChunkInfo(const FIoChunkId& ChunkId);
 	void					ReleaseContent(FOnDemandInternalContentHandle& ContentHandle);
diff --git a/Engine/Source/Runtime/Experimental/IoStore/OnDemand/Public/IO/IoStoreOnDemand.h b/Engine/Source/Runtime/Experimental/IoStore/OnDemand/Public/IO/IoStoreOnDemand.h
--- a/Engine/Source/Runtime/Experimental/IoStore/OnDemand/Public/IO/IoStoreOnDemand.h
+++ b/Engine/Source/Runtime/Experimental/IoStore/OnDemand/Public/IO/IoStoreOnDemand.h
@@ -393,6 +393,8 @@ public:
 
 	UE_API TIoStatusOr<uint64> GetInstallSize(const FOnDemandGetInstallSizeArgs& Args) const;
 
+	UE_API FIoStatus GetInstallSizesByMountId(const FOnDemandGetInstallSizeArgs& Args, TMap<FString, uint64>& OutSizesByMountId) const;
+
 	virtual void StartupModule() override;
 	virtual void ShutdownModule() override;
 	
