Iris
* Fix dormant objects still getting replicated when one of their push model property was dirty.
* Now caching NetFlush request of dormant objects instead of simply setting them dirty.
* Dormant objects will now be NetFlushed when an RPC is called on them.  This makes Iris behave the same as generic replication.
* Added cvars to modify these behaviors: net.Iris.RPC.AllowOnDormantObjects (default true), net.Iris.RPC.AutoNetFlushOnDormantObjects (default true)

diff --git a/Engine/Plugins/Runtime/ReplicationSystemTestPlugin/Source/Private/Tests/ReplicationSystem/TestDormancy.cpp b/Engine/Plugins/Runtime/ReplicationSystemTestPlugin/Source/Private/Tests/ReplicationSystem/TestDormancy.cpp
--- a/Engine/Plugins/Runtime/ReplicationSystemTestPlugin/Source/Private/Tests/ReplicationSystem/TestDormancy.cpp
+++ b/Engine/Plugins/Runtime/ReplicationSystemTestPlugin/Source/Private/Tests/ReplicationSystem/TestDormancy.cpp
@@ -51,7 +51,7 @@ UE_NET_TEST_FIXTURE(FReplicationSystemServerClientTestFixture, TestDormancyWithP
 	UE_NET_ASSERT_NE(ClientObject->GetIntA(), 3);
 
 	// Flush dormancy on the object
-	Server->GetReplicationBridge()->ForceUpdateWantsToBeDormantObject(ServerObject->NetRefHandle);
+	Server->GetReplicationBridge()->NetFlushDormantObject(ServerObject->NetRefHandle);
 
 	// The property on the client should now be received
 	Server->UpdateAndSend({ Client });
diff --git a/Engine/Source/Runtime/Engine/Private/Net/Iris/ReplicationSystem/ReplicationSystemUtil.cpp b/Engine/Source/Runtime/Engine/Private/Net/Iris/ReplicationSystem/ReplicationSystemUtil.cpp
--- a/Engine/Source/Runtime/Engine/Private/Net/Iris/ReplicationSystem/ReplicationSystemUtil.cpp
+++ b/Engine/Source/Runtime/Engine/Private/Net/Iris/ReplicationSystem/ReplicationSystemUtil.cpp
@@ -478,23 +478,18 @@ void FReplicationSystemUtil::FlushNetDormancy(UReplicationSystem* ReplicationSys
 
 	if (!Actor->IsActorInitialized())
 	{
-		UE_LOG(LogIris, Verbose, TEXT("FReplicationSystemUtil::FlushNetDormancy called on %s that isn't fully initialized yet. Ingoring."), ToCStr(GetFullNameSafe(Actor)));
+		UE_LOG(LogIris, Verbose, TEXT("FReplicationSystemUtil::FlushNetDormancy called on %s that isn't fully initialized yet. Ignoring."), ToCStr(GetFullNameSafe(Actor)));
 		return;
 	}
 
-	FNetHandle ActorHandle = GetNetHandle(Actor);
-	if (ActorHandle.IsValid())
+	if (!Actor->GetIsReplicated())
 	{
-		if (UObjectReplicationBridge* Bridge = ReplicationSystem->GetReplicationBridgeAs<UObjectReplicationBridge>())
-		{
-			const FNetRefHandle ActorRefHandle = Bridge->GetReplicatedRefHandle(ActorHandle);
-			if (ActorRefHandle.IsValid())
-			{
-				Bridge->ForceUpdateWantsToBeDormantObject(ActorRefHandle);
-			}
-		}
+		ensureMsgf(Actor->GetIsReplicated(), TEXT("FReplicationSystemUtil::FlushNetDormancy Actor: %s is not replicated"), ToCStr(GetFullNameSafe(Actor)));
+		return;
 	}
-	else if (Actor->HasActorBegunPlay() || Actor->IsActorBeginningPlay())
+
+	// Handle DormInitial actors
+	if (bWasDormInitial && (Actor->HasActorBegunPlay() || Actor->IsActorBeginningPlay()))
 	{
 		// Call BeginReplication for DORM_Initial actors the first time their dormancy is flushed
 		// (since it's not called when they BeginPlay).
@@ -503,12 +498,19 @@ void FReplicationSystemUtil::FlushNetDormancy(UReplicationSystem* ReplicationSys
 		// -The actor and its components/subobjects may not be completely set up for replication yet
 		// -If the actor is DormInitial, and is flushed before BeginPlay, its dormancy state will change to DormantAll and it will BeginReplication normally
 
-		UE_CLOG(!Actor->GetIsReplicated(), LogIris, Warning, TEXT("FReplicationSystemUtil::FlushNetDormancy Actor %s that is not replicated"), ToCStr(Actor->GetName()));
-		UE_CLOG(!bWasDormInitial, LogIris, Display, TEXT("FReplicationSystemUtil::FlushNetDormancy For not replicated Actor %s is not initially dormant"), ToCStr(Actor->GetName()));
+		Actor->BeginReplication();
+	}
 
-		if (Actor->GetIsReplicated())
+	const FNetHandle ActorHandle = GetNetHandle(Actor);
+	if (ActorHandle.IsValid())
+	{
+		if (UObjectReplicationBridge* Bridge = ReplicationSystem->GetReplicationBridgeAs<UObjectReplicationBridge>())
 		{
- 			Actor->BeginReplication();
+			const FNetRefHandle ActorRefHandle = Bridge->GetReplicatedRefHandle(ActorHandle);
+			if (ActorRefHandle.IsValid())
+			{
+				Bridge->NetFlushDormantObject(ActorRefHandle);
+			}
 		}
 	}
 }
diff --git a/Engine/Source/Runtime/Experimental/Iris/Core/Private/Iris/ReplicationSystem/NetBlob/NetBlobManager.cpp b/Engine/Source/Runtime/Experimental/Iris/Core/Private/Iris/ReplicationSystem/NetBlob/NetBlobManager.cpp
--- a/Engine/Source/Runtime/Experimental/Iris/Core/Private/Iris/ReplicationSystem/NetBlob/NetBlobManager.cpp
+++ b/Engine/Source/Runtime/Experimental/Iris/Core/Private/Iris/ReplicationSystem/NetBlob/NetBlobManager.cpp
@@ -9,6 +9,7 @@
 #include "Iris/ReplicationSystem/ReplicationSystem.h"
 #include "Iris/ReplicationSystem/ReplicationSystemInternal.h"
 #include "Iris/ReplicationSystem/ReplicationWriter.h"
+#include "Iris/ReplicationSystem/ObjectReplicationBridge.h"
 #include "Iris/Core/NetObjectReference.h"
 #include "Iris/Core/IrisLog.h"
 #include "Net/Core/Trace/NetDebugName.h"
@@ -21,10 +22,16 @@ static TAutoConsoleVariable<int32> CVarEnableIrisRPCs(TEXT("net.Iris.EnableRPCs"
 static bool bThrottleRPCWarnings = true;
 static FAutoConsoleVariableRef CVarThrottleRPCWarnings(TEXT("net.Iris.ThrottleRPCWarnings"), bThrottleRPCWarnings, TEXT("Only log send failure warnings once per RPC type."));
 
+static bool bAllowRPCsOnDormantObjects = true;
+static FAutoConsoleVariableRef CVarAllowRPCsOnDormantObjects(TEXT("net.Iris.RPC.AllowOnDormantObjects"), bAllowRPCsOnDormantObjects, TEXT("When true allow RPCs to be sent on dormant objects. When false block all RPCs from the moment the dormant change is requested."));
+
+static bool bAutoNetFlushOnDormantRPC = true;
+static FAutoConsoleVariableRef CVarAutoNetFlushOnDormantRPC(TEXT("net.Iris.RPC.AutoNetFlushOnDormantObjects"), bAutoNetFlushOnDormantRPC, TEXT("When true we will make an implicit NetFlush request when an RPC is sent on a dormant object. When false no replicated properties would be sent along with the RPC."));
+
 FNetBlobManager::FNetBlobManager()
 {
 }
-	
+
 void FNetBlobManager::Init(FNetBlobManagerInitParams& InitParams)
 {
 	BlobHandlerManager.Init();
@@ -97,7 +104,7 @@ bool FNetBlobManager::QueueNetObjectAttachment(uint32 ConnectionId, const FNetOb
 	return true;
 }
 
-bool FNetBlobManager::SendRPC(const UObject* Object, const UObject* SubObject, const UFunction* Function, const void* Parameters, UE::Net::ENetObjectAttachmentSendPolicyFlags SendFlags)
+bool FNetBlobManager::SendMulticastRPC(const FSendRPCContext& Context, const void* Parameters, UE::Net::ENetObjectAttachmentSendPolicyFlags SendFlags)
 {
 	if (CVarEnableIrisRPCs.GetValueOnGameThread() <= 0)
 	{
@@ -111,9 +118,9 @@ bool FNetBlobManager::SendRPC(const UObject* Object, const UObject* SubObject, c
 	}
 
 	// May the RPC be sent?
-	if ((Function->FunctionFlags & (bIsServer ? (FUNC_NetClient | FUNC_NetMulticast) : FUNC_NetServer)) == 0)
+	if ((Context.Function->FunctionFlags & (bIsServer ? (FUNC_NetClient | FUNC_NetMulticast) : FUNC_NetServer)) == 0)
 	{
-		checkf(false, TEXT("Trying to call RPC %s in the wrong direction."), ToCStr(Function->GetName()));
+		checkf(false, TEXT("Trying to call RPC %s in the wrong direction."), ToCStr(Context.Function->GetName()));
 		return true;
 	}
 
@@ -125,23 +132,44 @@ bool FNetBlobManager::SendRPC(const UObject* Object, const UObject* SubObject, c
 	}
 
 	FRPCOwner OwnerInfo;
-	if (!GetRPCOwner(OwnerInfo, Object, SubObject, Function))
+	if (!GetRPCOwner(OwnerInfo, Context))
 	{
 		return false;
 	}
+	
+	// If we prevent RPCs on dormant objects (allowed by default)
+	const bool bIsObjectDormant = NetRefHandleManager->GetWantToBeDormantInternalIndices().IsBitSet(OwnerInfo.RootObjectIndex);
+	if (!bAllowRPCsOnDormantObjects && bIsObjectDormant)
+	{
+		// If the object is dormant but requested a FlushNet we still allow him to send RPCs
+		// This also means objects that went dormant on the current frame can also send RPCs.
+		if (!NetRefHandleManager->GetDormantObjectsPendingFlushNet().IsBitSet(OwnerInfo.RootObjectIndex))
+		{
+			return false;
+		}
+	}
 
-	const TRefCountPtr<FNetRPC>& RPC = Handler->CreateRPC(OwnerInfo.CallerRef, Function, Parameters);
+	const TRefCountPtr<FNetRPC>& RPC = Handler->CreateRPC(OwnerInfo.CallerRef, Context.Function, Parameters);
 	if (!RPC.IsValid())
 	{
+		UE_LOG(LogIris, Warning, TEXT("Unable to create RPC for function %s."), ToCStr(Context.Function->GetName()));
 		return true;
 	}
 
+	// Force a NetFlush when sending an RPC on a dormant object
+	if (bAutoNetFlushOnDormantRPC && bIsObjectDormant)
+	{
+		UObjectReplicationBridge* Bridge = ReplicationSystem->GetReplicationBridgeAs<UObjectReplicationBridge>();
+		Bridge->NetFlushDormantObject(ObjectReferenceCache->GetObjectReferenceHandleFromObject(Context.RootObject));
+	}
+
 	RPC->SetNetObjectReference(OwnerInfo.CallerRef, OwnerInfo.TargetRef);
 	AttachmentSendQueue.Enqueue(OwnerInfo.RootObjectIndex, OwnerInfo.SubObjectIndex, reinterpret_cast<const TRefCountPtr<FNetObjectAttachment>&>(RPC), SendFlags, Connections->GetOpenConnections());
+
 	return true;
 }
 
-bool FNetBlobManager::SendRPC(uint32 ConnectionId, const UObject* Object, const UObject* SubObject, const UFunction* Function, const void* Parameters, UE::Net::ENetObjectAttachmentSendPolicyFlags SendFlags)
+bool FNetBlobManager::SendUnicastRPC(uint32 ConnectionId, const FSendRPCContext& Context, const void* Parameters, UE::Net::ENetObjectAttachmentSendPolicyFlags SendFlags)
 {
 	if (CVarEnableIrisRPCs.GetValueOnGameThread() <= 0)
 	{
@@ -154,9 +182,9 @@ bool FNetBlobManager::SendRPC(uint32 ConnectionId, const UObject* Object, const
 		return false;
 	}
 
-	if ((Function->FunctionFlags & (bIsServer ? (FUNC_NetClient | FUNC_NetMulticast) : FUNC_NetServer)) == 0)
+	if ((Context.Function->FunctionFlags & (bIsServer ? (FUNC_NetClient | FUNC_NetMulticast) : FUNC_NetServer)) == 0)
 	{
-		checkf(false, TEXT("Trying to call RPC %s in the wrong direction."), ToCStr(Function->GetName()));
+		checkf(false, TEXT("Trying to call RPC %s in the wrong direction."), ToCStr(Context.Function->GetName()));
 		return true;
 	}
 
@@ -173,31 +201,50 @@ bool FNetBlobManager::SendRPC(uint32 ConnectionId, const UObject* Object, const
 	}
 
 	FRPCOwner OwnerInfo;
-	if (!GetRPCOwner(OwnerInfo, Object, SubObject, Function))
+	if (!GetRPCOwner(OwnerInfo, Context))
 	{
 		return false;
 	}
 
-	const TRefCountPtr<FNetRPC>& RPC = Handler->CreateRPC(OwnerInfo.CallerRef, Function, Parameters);
+	// If we prevent RPCs on dormant objects (allowed by default)
+	const bool bIsObjectDormant = NetRefHandleManager->GetWantToBeDormantInternalIndices().IsBitSet(OwnerInfo.RootObjectIndex);
+	if (!bAllowRPCsOnDormantObjects && bIsObjectDormant)
+	{
+		// If the object is dormant but requested a FlushNet we still allow him to send RPCs
+		// This also means objects that went dormant on the current frame can also send RPCs.
+		if (!NetRefHandleManager->GetDormantObjectsPendingFlushNet().IsBitSet(OwnerInfo.RootObjectIndex))
+		{
+			return false;
+		}
+	}
+
+	const TRefCountPtr<FNetRPC>& RPC = Handler->CreateRPC(OwnerInfo.CallerRef, Context.Function, Parameters);
 	if (!RPC.IsValid())
 	{
-		UE_LOG(LogIris, Warning, TEXT("Unable to create RPC for function %s."), ToCStr(Function->GetName()));
+		UE_LOG(LogIris, Warning, TEXT("Unable to create RPC for function %s."), ToCStr(Context.Function->GetName()));
 		return true;
 	}
 
+	// Force a NetFlush when sending an RPC on a dormant object
+	if (bAutoNetFlushOnDormantRPC && bIsObjectDormant)
+	{
+		UObjectReplicationBridge* Bridge = ReplicationSystem->GetReplicationBridgeAs<UObjectReplicationBridge>();
+		Bridge->NetFlushDormantObject(ObjectReferenceCache->GetObjectReferenceHandleFromObject(Context.RootObject));
+	}
+
 	RPC->SetNetObjectReference(OwnerInfo.CallerRef, OwnerInfo.TargetRef);
 	AttachmentSendQueue.Enqueue(ConnectionId, OwnerInfo.RootObjectIndex, OwnerInfo.SubObjectIndex, reinterpret_cast<const TRefCountPtr<FNetObjectAttachment>&>(RPC), SendFlags);
 	return true;
 }
 
-bool FNetBlobManager::GetRPCOwner(FRPCOwner& OutOwnerInfo, const UObject* RootObject, const UObject* SubObject, const UFunction* Function) const
+bool FNetBlobManager::GetRPCOwner(FRPCOwner& OutOwnerInfo, const FSendRPCContext& Context) const
 {
 	bool bCanSendRpc = false;
 
 	// If a root object is sending an RPC
-	if (SubObject == nullptr)
+	if (Context.SubObject == nullptr)
 	{
-		OutOwnerInfo.TargetRef = ObjectReferenceCache->GetOrCreateObjectReference(RootObject);
+		OutOwnerInfo.TargetRef = ObjectReferenceCache->GetOrCreateObjectReference(Context.RootObject);
 		OutOwnerInfo.CallerRef = OutOwnerInfo.TargetRef;
 
 		bCanSendRpc = GetRootObjectIndicesFromHandle(OutOwnerInfo.TargetRef.GetRefHandle(), OutOwnerInfo.RootObjectIndex);
@@ -207,23 +254,23 @@ bool FNetBlobManager::GetRPCOwner(FRPCOwner& OutOwnerInfo, const UObject* RootOb
 			bool bLogRPCFailed = true;
 			if (UE::Net::Private::bThrottleRPCWarnings)
 			{
-				bool& bWasAlreadyLogged = RPCWarningThrottler.FindOrAdd(Function->GetFName(), false);
+				bool& bWasAlreadyLogged = RPCWarningThrottler.FindOrAdd(Context.Function->GetFName(), false);
 				bLogRPCFailed = !bWasAlreadyLogged;
 				bWasAlreadyLogged = true;
 			}
 			UE_CLOG(bLogRPCFailed, LogIris, Warning, TEXT("SendRPC %s for %s Failed. This rootobject is not yet replicated (RefHandle: %s Index: %u)."),
-				ToCStr(Function->GetName()), *GetNameSafe(RootObject), ToCStr(OutOwnerInfo.CallerRef.GetRefHandle().ToString()), OutOwnerInfo.RootObjectIndex);
+					ToCStr(Context.Function->GetName()), *GetNameSafe(Context.RootObject), ToCStr(OutOwnerInfo.CallerRef.GetRefHandle().ToString()), OutOwnerInfo.RootObjectIndex);
 		}
 	}
 	// If a subobject is sending an RPC
 	else
 	{
-		const FNetRefHandle SubObjectNetRef = ObjectReferenceCache->GetObjectReferenceHandleFromObject(SubObject);
+		const FNetRefHandle SubObjectNetRef = ObjectReferenceCache->GetObjectReferenceHandleFromObject(Context.SubObject);
 
 		// If the subobject can be referenced
 		if (SubObjectNetRef.IsValid())
 		{
-			OutOwnerInfo.TargetRef = ObjectReferenceCache->GetOrCreateObjectReference(SubObject);
+			OutOwnerInfo.TargetRef = ObjectReferenceCache->GetOrCreateObjectReference(Context.SubObject);
 			OutOwnerInfo.CallerRef = OutOwnerInfo.TargetRef;
 
 			check(OutOwnerInfo.TargetRef.GetRefHandle() == SubObjectNetRef);
@@ -234,8 +281,8 @@ bool FNetBlobManager::GetRPCOwner(FRPCOwner& OutOwnerInfo, const UObject* RootOb
 		// Send the RPC via the Root object if the subobject is not capable
 		if (!bCanSendRpc)
 		{
-			OutOwnerInfo.TargetRef = ObjectReferenceCache->GetOrCreateObjectReference(SubObject);
-			OutOwnerInfo.CallerRef = ObjectReferenceCache->GetOrCreateObjectReference(RootObject);
+			OutOwnerInfo.TargetRef = ObjectReferenceCache->GetOrCreateObjectReference(Context.SubObject);
+			OutOwnerInfo.CallerRef = ObjectReferenceCache->GetOrCreateObjectReference(Context.RootObject);
 
 			bCanSendRpc = GetRootObjectIndicesFromHandle(OutOwnerInfo.CallerRef.GetRefHandle(), OutOwnerInfo.RootObjectIndex);
 		}
@@ -245,13 +292,13 @@ bool FNetBlobManager::GetRPCOwner(FRPCOwner& OutOwnerInfo, const UObject* RootOb
 			bool bLogRPCFailed = true;
 			if (UE::Net::Private::bThrottleRPCWarnings)
 			{
-				bool& bWasAlreadyLogged = RPCWarningThrottler.FindOrAdd(Function->GetFName(), false);
+				bool& bWasAlreadyLogged = RPCWarningThrottler.FindOrAdd(Context.Function->GetFName(), false);
 				bLogRPCFailed = !bWasAlreadyLogged;
 				bWasAlreadyLogged = true;
 			}
 
 			UE_CLOG(bLogRPCFailed, LogIris, Warning, TEXT("SendRPC %s for %s::%s Failed. The root object (RefHandle: %s Index: %u) and subobject (RefHandle: %s Index: %u) is not yet replicated."),
-				ToCStr(Function->GetName()), *GetNameSafe(RootObject), *GetNameSafe(SubObject), ToCStr(OutOwnerInfo.CallerRef.ToString()), OutOwnerInfo.RootObjectIndex, ToCStr(OutOwnerInfo.TargetRef.ToString()), OutOwnerInfo.SubObjectIndex);
+					ToCStr(Context.Function->GetName()), *GetNameSafe(Context.RootObject), *GetNameSafe(Context.SubObject), ToCStr(OutOwnerInfo.CallerRef.ToString()), OutOwnerInfo.RootObjectIndex, ToCStr(OutOwnerInfo.TargetRef.ToString()), OutOwnerInfo.SubObjectIndex);
 		}
 	}
 
diff --git a/Engine/Source/Runtime/Experimental/Iris/Core/Private/Iris/ReplicationSystem/NetBlob/NetBlobManager.h b/Engine/Source/Runtime/Experimental/Iris/Core/Private/Iris/ReplicationSystem/NetBlob/NetBlobManager.h
--- a/Engine/Source/Runtime/Experimental/Iris/Core/Private/Iris/ReplicationSystem/NetBlob/NetBlobManager.h
+++ b/Engine/Source/Runtime/Experimental/Iris/Core/Private/Iris/ReplicationSystem/NetBlob/NetBlobManager.h
@@ -46,8 +46,19 @@ public:
 	bool AllowObjectReplication() const { return bAllowObjectReplication; }
 
 	bool QueueNetObjectAttachment(uint32 ConnectionId, const FNetObjectReference& TargetRef, const TRefCountPtr<FNetObjectAttachment>& Attachment, ENetObjectAttachmentSendPolicyFlags SendFlags = ENetObjectAttachmentSendPolicyFlags::None);
-	bool SendRPC(const UObject* Object, const UObject* SubObject, const UFunction* Function, const void* Parameters, ENetObjectAttachmentSendPolicyFlags SendFlags = ENetObjectAttachmentSendPolicyFlags::None);
-	bool SendRPC(uint32 ConnectionId, const UObject* Object, const UObject* SubObject, const UFunction* Function, const void* Parameters, ENetObjectAttachmentSendPolicyFlags SendFlags = ENetObjectAttachmentSendPolicyFlags::None);
+
+	struct FSendRPCContext
+	{
+		const UObject* RootObject = nullptr;
+		const UObject* SubObject = nullptr;
+		const UFunction* Function = nullptr;
+	};
+
+	// Multicast RPC
+	bool SendMulticastRPC(const FSendRPCContext& Context, const void* Parameters, ENetObjectAttachmentSendPolicyFlags SendFlags = ENetObjectAttachmentSendPolicyFlags::None);
+
+	// Unicast RPC
+	bool SendUnicastRPC(uint32 ConnectionId, const FSendRPCContext& Context, const void* Parameters, ENetObjectAttachmentSendPolicyFlags SendFlags = ENetObjectAttachmentSendPolicyFlags::None);
 
 	bool HasUnprocessedReliableAttachments(FInternalNetRefIndex InternalIndex) const;
 	bool HasAnyUnprocessedReliableAttachments() const;
@@ -75,6 +86,7 @@ public:
 	void RemoveConnection(uint32 ConnectionId);
 
 private:
+
 	void RegisterDefaultHandlers();
 
 	struct FRPCOwner
@@ -88,7 +100,7 @@ private:
 		FInternalNetRefIndex RootObjectIndex = FNetRefHandleManager::InvalidInternalIndex;
 		FInternalNetRefIndex SubObjectIndex = FNetRefHandleManager::InvalidInternalIndex;
 	};
-	bool GetRPCOwner(FRPCOwner& OutOwnerInfo, const UObject* RootObject, const UObject* SubObject, const UFunction* Function) const;
+	bool GetRPCOwner(FRPCOwner& OutOwnerInfo, const FSendRPCContext& Context) const;
 
 	/** 
 	* Validates that RootObjectRefHandle is a true root object and return it's index
diff --git a/Engine/Source/Runtime/Experimental/Iris/Core/Private/Iris/ReplicationSystem/NetRefHandleManager.cpp b/Engine/Source/Runtime/Experimental/Iris/Core/Private/Iris/ReplicationSystem/NetRefHandleManager.cpp
--- a/Engine/Source/Runtime/Experimental/Iris/Core/Private/Iris/ReplicationSystem/NetRefHandleManager.cpp
+++ b/Engine/Source/Runtime/Experimental/Iris/Core/Private/Iris/ReplicationSystem/NetRefHandleManager.cpp
@@ -72,6 +72,7 @@ void FNetRefHandleManager::Init(const FInitParams& InitParams)
 		InitNetBitArray(&DestroyedStartupObjectInternalIndices);
 		InitNetBitArray(&WantToBeDormantInternalIndices);
 		InitNetBitArray(&ObjectsWithPreUpdate);
+		InitNetBitArray(&DormantObjectsPendingFlushNet);
 	}
 
 	// Mark the invalid index as used
@@ -397,9 +398,9 @@ FNetRefHandle FNetRefHandleManager::CreateNetObjectFromRemote(FNetRefHandle Want
 
 void FNetRefHandleManager::InternalDestroyNetObject(FInternalNetRefIndex InternalIndex)
 {
-	FReplicatedObjectData& Data = ReplicatedObjectData[InternalIndex];
+	UE_LOG(LogIris, Verbose, TEXT("FNetRefHandleManager::InternalDestroyNetObject: %s"), *PrintObjectFromIndex(InternalIndex));
 
-	UE_LOG(LogIris, Verbose, TEXT("FNetRefHandleManager::InternalDestroyNetObject: (InternalIndex: %u) (%s)"), InternalIndex, *Data.RefHandle.ToString());
+	FReplicatedObjectData& Data = ReplicatedObjectData[InternalIndex];
 
 	uint8* StateBuffer = ReplicatedObjectStateBuffers[InternalIndex];
 	// Free any allocated resources
@@ -466,9 +467,7 @@ void FNetRefHandleManager::InternalDestroyNetObject(FInternalNetRefIndex Interna
 	AssignedInternalIndices.ClearBit(InternalIndex);
 
 	// Restore internal state
-	SubObjectInternalIndices.ClearBit(InternalIndex);
-	WantToBeDormantInternalIndices.ClearBit(InternalIndex);
-	ObjectsWithDependentObjectsInternalIndices.ClearBit(InternalIndex);
+	ClearStateForFreedInternalIndex(InternalIndex);
 
 	// Cleanup cross reference to destruction info
 	if (DestroyedStartupObjectInternalIndices.GetBit(InternalIndex))
@@ -484,6 +483,16 @@ void FNetRefHandleManager::InternalDestroyNetObject(FInternalNetRefIndex Interna
 	--ActiveObjectCount;
 }
 
+void FNetRefHandleManager::ClearStateForFreedInternalIndex(FInternalNetRefIndex FreedInternalIndex)
+{
+	GlobalScopableInternalIndices.ClearBit(FreedInternalIndex);
+	ObjectsWithPreUpdate.ClearBit(FreedInternalIndex);
+	SubObjectInternalIndices.ClearBit(FreedInternalIndex);
+	ObjectsWithDependentObjectsInternalIndices.ClearBit(FreedInternalIndex);
+	WantToBeDormantInternalIndices.ClearBit(FreedInternalIndex);
+	DormantObjectsPendingFlushNet.ClearBit(FreedInternalIndex);
+}
+
 FNetRefHandle FNetRefHandleManager::CreateHandleForDestructionInfo(FNetRefHandle Handle, const FReplicationProtocol* DestroyedObjectProtocol)
 {
 	// Create destruction info handle carrying destruction info
diff --git a/Engine/Source/Runtime/Experimental/Iris/Core/Private/Iris/ReplicationSystem/NetRefHandleManager.h b/Engine/Source/Runtime/Experimental/Iris/Core/Private/Iris/ReplicationSystem/NetRefHandleManager.h
--- a/Engine/Source/Runtime/Experimental/Iris/Core/Private/Iris/ReplicationSystem/NetRefHandleManager.h
+++ b/Engine/Source/Runtime/Experimental/Iris/Core/Private/Iris/ReplicationSystem/NetRefHandleManager.h
@@ -302,10 +302,21 @@ public:
 	const FNetBitArrayView GetWantToBeDormantInternalIndices() const { return MakeNetBitArrayView(WantToBeDormantInternalIndices); }
 	FNetBitArrayView GetWantToBeDormantInternalIndices() { return MakeNetBitArrayView(WantToBeDormantInternalIndices); }
 
+	/** List of dormant objects that requested a FlushNet */
+	const FNetBitArrayView GetDormantObjectsPendingFlushNet() const { return MakeNetBitArrayView(DormantObjectsPendingFlushNet); }
+	FNetBitArrayView GetDormantObjectsPendingFlushNet() { return MakeNetBitArrayView(DormantObjectsPendingFlushNet); }
+
+	/** Get Objects that are flagged for PreUpdate (aka PreReplication) */
+	FNetBitArrayView GetObjectsWithPreUpdate() const { return MakeNetBitArrayView(ObjectsWithPreUpdate); }
+
+public:
+
 	/** Return a string to identify the object linked to an index in logs */
 	[[nodiscard]] FString PrintObjectFromIndex(FInternalNetRefIndex ObjectIndex) const;
 	[[nodiscard]] FString PrintObjectFromNetRefHandle(FNetRefHandle ObjectHandle) const;
 
+public:
+
 	/** Delegate that will notify when the highest NetChunkedArray internal index has increased (Highest = Max-1)*/
 	DECLARE_MULTICAST_DELEGATE_OneParam(FOnNetChunkedArrayIncrease, FInternalNetRefIndex HighestInternalIndex);
 	FOnNetChunkedArrayIncrease& GetOnNetChunkedArrayIncreaseDelegate() const { return OnNetChunkedArrayIncrease; };
@@ -317,9 +328,6 @@ public:
 	DECLARE_MULTICAST_DELEGATE_OneParam(FOnMaxInternalNetRefIndexIncreased, FInternalNetRefIndex MaxInternalIndex);
 	FOnMaxInternalNetRefIndexIncreased& GetOnMaxInternalNetRefIndexIncreasedDelegate() const { return OnMaxInternalNetRefIndexIncreased; };
 
-	/** Get Objects that is flagged for PreUpdate (aka PreReplication) */
-	FNetBitArrayView GetObjectsWithPreUpdate() const { return MakeNetBitArrayView(ObjectsWithPreUpdate); }
-
 public:
 
 	// Iterate over all dependent objects and their dependent objects
@@ -344,6 +352,9 @@ private:
 	FInternalNetRefIndex InternalCreateNetObject(const FNetRefHandle NetRefHandle, const FNetHandle GlobalHandle, const FReplicationProtocol* ReplicationProtocol);
 	void InternalDestroyNetObject(FInternalNetRefIndex InternalIndex);
 
+	/** Clear the status of all internal bit arrays when an internal index is put back in the free list */
+	void ClearStateForFreedInternalIndex(FInternalNetRefIndex FreedInternalIndex);
+
 	static uint64 MakeNetRefHandleId(uint64 Seed, bool bIsStatic);
 	uint64 GetNextNetRefHandleId(uint64 HandleIndex) const;
 
@@ -440,6 +451,9 @@ private:
 	// Bitset marking internal indices that wants to be dormant
 	FNetBitArray WantToBeDormantInternalIndices;
 
+	// Set of dormant objects that have requested a NetFlushDormancy
+	FNetBitArray DormantObjectsPendingFlushNet;
+
 	// Bitset marking internal indices that is flagged as requiring PreUpdate
 	FNetBitArray ObjectsWithPreUpdate;
 
diff --git a/Engine/Source/Runtime/Experimental/Iris/Core/Private/Iris/ReplicationSystem/ObjectReplicationBridge.cpp b/Engine/Source/Runtime/Experimental/Iris/Core/Private/Iris/ReplicationSystem/ObjectReplicationBridge.cpp
--- a/Engine/Source/Runtime/Experimental/Iris/Core/Private/Iris/ReplicationSystem/ObjectReplicationBridge.cpp
+++ b/Engine/Source/Runtime/Experimental/Iris/Core/Private/Iris/ReplicationSystem/ObjectReplicationBridge.cpp
@@ -515,7 +515,7 @@ UE::Net::FNetRefHandle UObjectReplicationBridge::BeginReplication(FNetRefHandle
 		SetPollWithObject(OwnerRefHandle, SubObjectRefHandle);
 
 		// Copy pending dormancy from owner
-		SetObjectWantsToBeDormant(SubObjectRefHandle, GetObjectWantsToBeDormant(OwnerRefHandle));
+		SetSubObjectDormancyStatus(SubObjectRefHandle, OwnerRefHandle);
 	
 		return SubObjectRefHandle;
 	}
@@ -523,6 +523,19 @@ UE::Net::FNetRefHandle UObjectReplicationBridge::BeginReplication(FNetRefHandle
 	return FNetRefHandle::GetInvalid();
 }
 
+void UObjectReplicationBridge::SetSubObjectDormancyStatus(FNetRefHandle SubObjectRefHandle, FNetRefHandle OwnerRefHandle)
+{
+	using namespace UE::Net;
+	using namespace UE::Net::Private;
+
+	const FNetRefHandleManager& LocalNetRefHandleManager = *NetRefHandleManager;
+	const FInternalNetRefIndex SubObjectInternalIndex = LocalNetRefHandleManager.GetInternalIndex(SubObjectRefHandle);
+	const FInternalNetRefIndex OwnerInternalIndex = LocalNetRefHandleManager.GetInternalIndex(OwnerRefHandle);
+	
+	FNetBitArrayView DormantObjects = LocalNetRefHandleManager.GetWantToBeDormantInternalIndices();
+	DormantObjects.SetBitValue(SubObjectInternalIndex, DormantObjects.GetBit(OwnerInternalIndex));
+}
+
 void UObjectReplicationBridge::SetSubObjectNetCondition(FNetRefHandle SubObjectRefHandle, ELifetimeCondition Condition)
 {
 	using namespace UE::Net::Private;
@@ -1003,35 +1016,35 @@ void UObjectReplicationBridge::BuildPollList(UE::Net::FNetBitArrayView ObjectsCo
 		ObjectsConsideredForPolling.Copy(RelevantObjects);
 	}
 
-	// Mask off objects pending dormancy as we do not want to poll/pre-update them unless they are marked for flush or are dirty
+	// Mask off objects pending dormancy as we do not want to poll/pre-update them unless they are marked for flush
 	if (bUseDormancyToFilterPolling)
 	{
 		IRIS_PROFILER_SCOPE(BuildPollList_Dormancy);
 
-		// Mask off objects pending dormancy that are not dirty
-		const FNetBitArrayView AccumulatedDirtyObjects = DirtyNetObjectTracker.GetAccumulatedDirtyNetObjects();
-		ObjectsConsideredForPolling.CombineMultiple(FNetBitArrayView::AndNotOp, WantToBeDormantObjects, FNetBitArrayView::AndNotOp, AccumulatedDirtyObjects);
+		// Mask off all dormant objects
+		ObjectsConsideredForPolling.Combine(WantToBeDormantObjects, FNetBitArrayView::AndNotOp);
 
+		// Force a poll on objects that requested a FlushNet
+		FNetBitArrayView DormantObjectsPendingFlushNet = LocalNetRefHandleManager.GetDormantObjectsPendingFlushNet();
 		FNetBitArrayView ForceNetUpdateObjects = ReplicationSystemInternal->GetDirtyNetObjectTracker().GetForceNetUpdateObjects();
-
-		// Force poll objects that have requested to flush dormancy
-		for (FNetRefHandle HandlePendingFlush : MakeArrayView(DormantHandlesPendingFlush))
+		
+		uint32 FlushNetCount = 0;		
+		auto FlushNetDormancyForRelevantObjects = [&](uint32 DormantObjectIndex)
 		{
-			if (const uint32 InternalObjectIndex = LocalNetRefHandleManager.GetInternalIndex(HandlePendingFlush))
-			{
-				// If HandlePendingFlush is relevant, poll it this frame, and treat it as a forcenetupdate in order to also schedule subobjects correctly
-				// Out of scope dormant objects will be Marked Dirty to ensure that they are scheduled for polling once relevant
-				if (RelevantObjects.IsBitSet(InternalObjectIndex))
-				{
-					ObjectsConsideredForPolling.SetBit(InternalObjectIndex);
-					ForceNetUpdateObjects.SetBit(InternalObjectIndex);
-				}
-			}
-		}
+			// Poll the dormant object this frame and treat it as a forcenetupdate in order to also schedule subobjects correctly
+			ObjectsConsideredForPolling.SetBit(DormantObjectIndex);
+			ForceNetUpdateObjects.SetBit(DormantObjectIndex);
+			++FlushNetCount;
+		};
 
-		UE_NET_TRACE_FRAME_STATSCOUNTER(ReplicationSystem->GetId(), FlushDormancyObjectCount, DormantHandlesPendingFlush.Num(), ENetTraceVerbosity::Trace);
+		FNetBitArrayView::ForAllSetBits(DormantObjectsPendingFlushNet, RelevantObjects, FNetBitArrayView::AndOp, FlushNetDormancyForRelevantObjects);
+		
+		// Remove FlushNetDormancy request for the objects that were relevant
+		DormantObjectsPendingFlushNet.Combine(RelevantObjects, FNetBitArrayView::AndNotOp);
+
+		UE_NET_TRACE_FRAME_STATSCOUNTER(ReplicationSystem->GetId(), FlushDormancyObjectCount, FlushNetCount, ENetTraceVerbosity::Trace);
 	}
-	DormantHandlesPendingFlush.Reset();
+	
 
 	/**
 	* Make sure to propagate polling for owners to subobjects and vice versa. If an actor is not due to update due to
@@ -1279,44 +1292,57 @@ void UObjectReplicationBridge::SetObjectWantsToBeDormant(FNetRefHandle Handle, b
 	using namespace UE::Net;
 	using namespace UE::Net::Private;
 
-	FReplicationSystemInternal* ReplicationSystemInternal = GetReplicationSystem()->GetReplicationSystemInternal();
-	FNetRefHandleManager& LocalNetRefHandleManager = ReplicationSystemInternal->GetNetRefHandleManager();
+	FNetRefHandleManager& LocalNetRefHandleManager = GetReplicationSystem()->GetReplicationSystemInternal()->GetNetRefHandleManager();
 
-	if (const FInternalNetRefIndex InternalObjectIndex = LocalNetRefHandleManager.GetInternalIndex(Handle))
+	const FInternalNetRefIndex InternalObjectIndex = LocalNetRefHandleManager.GetInternalIndex(Handle);
+	if (InternalObjectIndex == FNetRefHandleManager::InvalidInternalIndex)
 	{
-		FNetBitArrayView WantToBeDormantObjects = LocalNetRefHandleManager.GetWantToBeDormantInternalIndices();
+		return;
+	}
 
-		// Update pending dormancy status
-		WantToBeDormantObjects.SetBitValue(InternalObjectIndex, bWantsToBeDormant);
+	UE_LOG_OBJECTREPLICATIONBRIDGE(Verbose, TEXT("SetObjectWantsToBeDormant: %s wants to be %s "), *LocalNetRefHandleManager.PrintObjectFromIndex(InternalObjectIndex), bWantsToBeDormant ? TEXT("Dormant"):TEXT("Awake"));
 
-		UE_LOG_OBJECTREPLICATIONBRIDGE(VeryVerbose, TEXT("SetObjectWantsToBeDormant for %s ( InternalIndex: %u ) %d "), *Handle.ToString(), InternalObjectIndex, bWantsToBeDormant ? 1 : 0);
+	// Only RootObjects can trigger dormancy changes
+	ensureMsgf(!LocalNetRefHandleManager.GetSubObjectInternalIndices().IsBitSet(InternalObjectIndex), TEXT("Only root objects can become dormant: %s "), *LocalNetRefHandleManager.PrintObjectFromIndex(InternalObjectIndex));
+	
+	FNetBitArrayView WantToBeDormantObjects = LocalNetRefHandleManager.GetWantToBeDormantInternalIndices();
 
-		// If we want to be dormant we want to make sure we poll the object one more time
-		if (bWantsToBeDormant)
-		{
-			// Note: this will override update frequency forcing a last poll/pre-replication-update if necessary
-			// No need to do this for subobjects as they always are handled with owner
-			DormantHandlesPendingFlush.Add(Handle);
-		}
+	// Update pending dormancy status
+	WantToBeDormantObjects.SetBitValue(InternalObjectIndex, bWantsToBeDormant);
 
-		// Since we use this as a mask when updating objects we must include subobjects as well
-		// Subobjects added later will copy status from owner when they are added
-		for (const FInternalNetRefIndex SubObjectInternalIndex : LocalNetRefHandleManager.GetSubObjects(InternalObjectIndex))
-		{
-			WantToBeDormantObjects.SetBitValue(SubObjectInternalIndex, bWantsToBeDormant);
-		}
+	// If we want to be dormant we want to make sure we poll the object immediately
+	LocalNetRefHandleManager.GetDormantObjectsPendingFlushNet().SetBitValue(InternalObjectIndex, bWantsToBeDormant);
 
-		// Request frequent world location updates for non-dormant spatially filtered objects.
-		OptionallySetObjectRequiresFrequentWorldLocationUpdate(Handle, !bWantsToBeDormant);
+	// Since we use this as a mask when updating objects we must include subobjects as well
+	// Subobjects added later will copy status from owner when they are added
+	for (const FInternalNetRefIndex SubObjectInternalIndex : LocalNetRefHandleManager.GetSubObjects(InternalObjectIndex))
+	{
+		WantToBeDormantObjects.SetBitValue(SubObjectInternalIndex, bWantsToBeDormant);
 	}
+
+	// Request frequent world location updates for non-dormant spatially filtered objects.
+	OptionallySetObjectRequiresFrequentWorldLocationUpdate(Handle, !bWantsToBeDormant);
 }
 
-void UObjectReplicationBridge::ForceUpdateWantsToBeDormantObject(FNetRefHandle Handle)
+void UObjectReplicationBridge::NetFlushDormantObject(FNetRefHandle Handle)
 {
-	UE_LOG_OBJECTREPLICATIONBRIDGE(VeryVerbose, TEXT("ForceUpdateWantsToBeDormantObject for %s"), *Handle.ToString());
+	using namespace UE::Net;
+	using namespace UE::Net::Private;
+
+	FNetRefHandleManager& LocalNetRefHandleManager = GetReplicationSystem()->GetReplicationSystemInternal()->GetNetRefHandleManager();
+
+	const FInternalNetRefIndex InternalObjectIndex = LocalNetRefHandleManager.GetInternalIndex(Handle);
+	if (InternalObjectIndex == FNetRefHandleManager::InvalidInternalIndex)
+	{
+		return;
+	}
+
+	UE_LOG_OBJECTREPLICATIONBRIDGE(Verbose, TEXT("NetFlushDormantObject: %s"), *LocalNetRefHandleManager.PrintObjectFromIndex(InternalObjectIndex));
+
+	LocalNetRefHandleManager.GetDormantObjectsPendingFlushNet().SetBit(InternalObjectIndex);
 
+	// Mark the object dirty in order to trigger an update of its WorldLocation and to accumulate dirty flags for when he comes back out of dormancy.
 	ReplicationSystem->MarkDirty(Handle);
-	DormantHandlesPendingFlush.Add(Handle);
 }
 
 void UObjectReplicationBridge::SetNetPushIdOnInstance(UE::Net::FReplicationInstanceProtocol* InstanceProtocol, FNetHandle NetHandle)
diff --git a/Engine/Source/Runtime/Experimental/Iris/Core/Private/Iris/ReplicationSystem/ReplicationSystem.cpp b/Engine/Source/Runtime/Experimental/Iris/Core/Private/Iris/ReplicationSystem/ReplicationSystem.cpp
--- a/Engine/Source/Runtime/Experimental/Iris/Core/Private/Iris/ReplicationSystem/ReplicationSystem.cpp
+++ b/Engine/Source/Runtime/Experimental/Iris/Core/Private/Iris/ReplicationSystem/ReplicationSystem.cpp
@@ -1038,19 +1038,32 @@ bool UReplicationSystem::QueueNetObjectAttachment(uint32 ConnectionId, const UE:
 	return NetBlobManager.QueueNetObjectAttachment(ConnectionId, TargetRef, Attachment);
 }
 
-bool UReplicationSystem::SendRPC(const UObject* Object, const UObject* SubObject, const UFunction* Function, const void* Parameters)
+bool UReplicationSystem::SendRPC(const UObject* RootObject, const UObject* SubObject, const UFunction* Function, const void* Parameters)
 {
-	UE::Net::ENetObjectAttachmentSendPolicyFlags SendFlags = UE::Net::ENetObjectAttachmentSendPolicyFlags::None;
+	using namespace UE::Net;
+	using namespace UE::Net::Private;
+
+	ENetObjectAttachmentSendPolicyFlags SendFlags = ENetObjectAttachmentSendPolicyFlags::None;
 	if (ReplicationSystemCVars::bAllowAttachmentSendPolicyFlags)
 	{
-		if (UE::Net::ENetObjectAttachmentSendPolicyFlags* Flags = Impl->AttachmentSendPolicyFlags.Find(FObjectKey(Function)))
+		if (ENetObjectAttachmentSendPolicyFlags* Flags = Impl->AttachmentSendPolicyFlags.Find(FObjectKey(Function)))
 		{
 			SendFlags = *Flags;
 		}
 	}
 
-	UE::Net::Private::FNetBlobManager& NetBlobManager = Impl->ReplicationSystemInternal.GetNetBlobManager();
-	return NetBlobManager.SendRPC(Object, SubObject, Function, Parameters, SendFlags);
+	FNetBlobManager::FSendRPCContext RPCContext = { .RootObject = RootObject, .SubObject = SubObject, .Function = Function };
+
+	FNetBlobManager& NetBlobManager = Impl->ReplicationSystemInternal.GetNetBlobManager();
+	return NetBlobManager.SendMulticastRPC(RPCContext, Parameters, SendFlags);
+}
+
+bool UReplicationSystem::SendRPC(uint32 ConnectionId, const UObject* RootObject, const UObject* SubObject, const UFunction* Function, const void* Parameters)
+{
+	using namespace UE::Net::Private;
+	FNetBlobManager::FSendRPCContext RPCContext = { .RootObject = RootObject, .SubObject = SubObject, .Function = Function };
+	FNetBlobManager& NetBlobManager = Impl->ReplicationSystemInternal.GetNetBlobManager();
+	return NetBlobManager.SendUnicastRPC(ConnectionId, RPCContext, Parameters);
 }
 
 bool UReplicationSystem::SetRPCSendPolicyFlags(const UFunction* Function, UE::Net::ENetObjectAttachmentSendPolicyFlags SendFlags)
@@ -1077,12 +1090,6 @@ void UReplicationSystem::ResetRPCSendPolicyFlags()
 	Impl->AttachmentSendPolicyFlags.Reset();
 }
 
-bool UReplicationSystem::SendRPC(uint32 ConnectionId, const UObject* Object, const UObject* SubObject, const UFunction* Function, const void* Parameters)
-{
-	UE::Net::Private::FNetBlobManager& NetBlobManager = Impl->ReplicationSystemInternal.GetNetBlobManager();
-	return NetBlobManager.SendRPC(ConnectionId, Object, SubObject, Function, Parameters);
-}
-
 void UReplicationSystem::InitDataStreams(uint32 ConnectionId, UDataStreamManager* DataStreamManager)
 {
 	UE::Net::Private::FReplicationConnections& Connections = Impl->ReplicationSystemInternal.GetConnections();
diff --git a/Engine/Source/Runtime/Experimental/Iris/Core/Public/Iris/ReplicationSystem/ObjectReplicationBridge.h b/Engine/Source/Runtime/Experimental/Iris/Core/Public/Iris/ReplicationSystem/ObjectReplicationBridge.h
--- a/Engine/Source/Runtime/Experimental/Iris/Core/Public/Iris/ReplicationSystem/ObjectReplicationBridge.h
+++ b/Engine/Source/Runtime/Experimental/Iris/Core/Public/Iris/ReplicationSystem/ObjectReplicationBridge.h
@@ -173,8 +173,8 @@ public:
 	/** Returns whether the object wants to be dormant. */
 	IRISCORE_API bool GetObjectWantsToBeDormant(FNetRefHandle Handle) const;
 
-	/** Trigger replication of dirty state for object wanting to be dormant. */
-	IRISCORE_API void ForceUpdateWantsToBeDormantObject(FNetRefHandle Handle);	
+	/** Trigger a single poll to refresh the dirty values of a dormant object. */
+	IRISCORE_API void NetFlushDormantObject(FNetRefHandle Handle);	
 
 	/** Set poll frequency on root object and its subobjects. They will be polled on the same frame. */
 	IRISCORE_API void SetPollFrequency(FNetRefHandle RootHandle, float PollFrequency);
@@ -390,6 +390,9 @@ private:
 	/** Returns true if instances of this class should be delta compressed */
 	bool ShouldClassBeDeltaCompressed(const UClass* Class);
 
+	/** Set the initial dormancy status of a subobject */
+	void SetSubObjectDormancyStatus(FNetRefHandle SubObjectRefHandle, FNetRefHandle OwnerRefHandle);
+
 	/** Marks a spatially filtered object as requiring or not requiring frequent world location updates independent of it having dirty replicated properties */
 	void OptionallySetObjectRequiresFrequentWorldLocationUpdate(FNetRefHandle RefHandle, bool bDesiresFrequentWorldLocationUpdate);
 
@@ -465,9 +468,6 @@ private:
 	// Type stats
 	TMap<FName, FName> ClassesWithTypeStats;
 
-	// Array of dormant objects that has requested a flush
-	TArray<FNetRefHandle> DormantHandlesPendingFlush;
-
 	// When we flush objects, we might need to defer sending creation info.
 	TMap<FNetRefHandle, TUniquePtr<const FCreationHeader>> CachedCreationHeaders;
 
diff --git a/Engine/Source/Runtime/Experimental/Iris/Core/Public/Iris/ReplicationSystem/ReplicationSystem.h b/Engine/Source/Runtime/Experimental/Iris/Core/Public/Iris/ReplicationSystem/ReplicationSystem.h
--- a/Engine/Source/Runtime/Experimental/Iris/Core/Public/Iris/ReplicationSystem/ReplicationSystem.h
+++ b/Engine/Source/Runtime/Experimental/Iris/Core/Public/Iris/ReplicationSystem/ReplicationSystem.h
@@ -289,13 +289,26 @@ public:
 
 	/**
 	 * Multicast an RPC targeting a object/subobject. 
-	 * @param Object A valid Owner/Actor. If no SubObject is specified the function will be called in this instance on the remote side.
-	 * @param SubObject Optional SubObject that the function will be called in on the remote side.
+	 * @param RootObject A valid Owner/Actor. If no subobject is specified the function is called on the root object on the remote side.
+	 * @param SubObject Optional subobject on whom the function is called on the remote side.
 	 * @param Function The function to call.
 	 * @param Parameters The function parameters.
+	 * 
+	 * @return Whether the RPC was successfully queued for replication or not.
+	 */
+	IRISCORE_API bool SendRPC(const UObject* RootObject, const UObject* SubObject, const UFunction* Function, const void* Parameters);
+
+	/**
+	 * Unicast an RPC targeting a object/subobject.
+	 * @param ConnectionId A valid connection ID. Only this connection will receive the RPC.
+	 * @param RootObject A valid Owner/Actor. If no subobject is specified the function is called on the root object on the remote side.
+	 * @param SubObject Optional subobject on whom the function is called on the remote side.
+	 * @param Function The function to call.
+	 * @param Parameters The function parameters.
+	 * 
 	 * @return Whether the RPC was successfully queued for replication or not.
 	 */
-	IRISCORE_API bool SendRPC(const UObject* Object, const UObject* SubObject, const UFunction* Function, const void* Parameters);
+	IRISCORE_API bool SendRPC(uint32 ConnectionId, const UObject* RootObject, const UObject* SubObject, const UFunction* Function, const void* Parameters);
 
 	/**
 	 * Set the policy flags for an RPC identified by its function
@@ -307,17 +320,6 @@ public:
 
 	/** Resets all set RPCSendPolicy flags */
 	IRISCORE_API void ResetRPCSendPolicyFlags();
-	
-	/**
-	 * Unicast an RPC targeting a object/subobject.
-	 * @param ConnectionId A valid connection ID. Only this connection will replicate the RPC.
-	 * @param Object A valid Owner/Actor. If no SubObject is specified the function will be called in this instance on the remote side.
-	 * @param SubObject Optional SubObject that the function will be called in on the remote side.
-	 * @param Function The function to call.
-	 * @param Parameters The function parameters.
-	 * @return Whether the RPC was successfully queued for replication or not.
-	 */
-	IRISCORE_API bool SendRPC(uint32 ConnectionId, const UObject* Object, const UObject* SubObject, const UFunction* Function, const void* Parameters);
 
 	/** @return The UReplicationBridge that was passed with the system creation parameters. */
 	IRISCORE_API UReplicationBridge* GetReplicationBridge() const;
@@ -775,5 +777,5 @@ inline UReplicationSystem* GetReplicationSystem(uint32 Id)
 	return Id >= FReplicationSystemFactory::MaxReplicationSystemCount ? nullptr : FReplicationSystemFactory::ReplicationSystems[Id];
 }
 
-}
+} // end namespace UE::Net
 
