Reverted requirement for valid container header when installing I/O store content and some minor cleanup

#rnx

diff --git a/Engine/Source/Runtime/Experimental/IoStore/OnDemand/Private/IoStoreOnDemand.cpp b/Engine/Source/Runtime/Experimental/IoStore/OnDemand/Private/IoStoreOnDemand.cpp
--- a/Engine/Source/Runtime/Experimental/IoStore/OnDemand/Private/IoStoreOnDemand.cpp
+++ b/Engine/Source/Runtime/Experimental/IoStore/OnDemand/Private/IoStoreOnDemand.cpp
@@ -1259,6 +1259,7 @@ void FIoStoreOnDemandModule::InitializeInternal()
 		if (MountArgs)
 		{
 			MountArgs.GetValue().Options = EOnDemandMountOptions::InstallOnDemand;
+			MountArgs.GetValue().Url = EndpointConfig.ServiceUrls[0] / EndpointConfig.TocPath;
 
 			static FOnDemandContentHandle ContentHandle = FOnDemandContentHandle::Create(TEXT("AllContent"));
 			InstallArgs.Emplace();
diff --git a/Engine/Source/Runtime/Experimental/IoStore/OnDemand/Private/OnDemandIoStore.cpp b/Engine/Source/Runtime/Experimental/IoStore/OnDemand/Private/OnDemandIoStore.cpp
--- a/Engine/Source/Runtime/Experimental/IoStore/OnDemand/Private/OnDemandIoStore.cpp
+++ b/Engine/Source/Runtime/Experimental/IoStore/OnDemand/Private/OnDemandIoStore.cpp
@@ -151,58 +151,9 @@ static FAnsiStringBuilderBase& GetChunkUrl(
 	return OutUrl;
 }
 
-static FIoStatus FetchEncodedContainerChunk(
-	FStringView Host,
-	FAnsiStringBuilderBase& ChunkUrlBuilder,
-	FSharedOnDemandContainer Container,
-	const FOnDemandChunkEntry& Entry,
-	FIoBuffer& OutEncodedChunk)
-{
-	TIoStatusOr<FIoBuffer> Response = FHttpClient::Get(GetChunkUrl(Host, *Container, Entry, ChunkUrlBuilder).ToView(), 2, EHttpRedirects::Follow);
-
-	if (Response.IsOk() == false)
-	{
-		FIoStatus Status = FIoStatusBuilder(EIoErrorCode::ReadError) << TEXT("Failed to fetch container chunk");
-		return Status;
-	}
-
-	OutEncodedChunk = Response.ConsumeValueOrDie();
-
-	return EIoErrorCode::Ok;
-}
-
-static FIoStatus FetchEncodedContainerHeader(
-	FStringView Host,
-	FAnsiStringBuilderBase& ChunkUrlBuilder, 
-	FSharedOnDemandContainer Container,
-	FIoBuffer& OutEncodedChunk)
-{
-	const FIoChunkId			ChunkId = CreateContainerHeaderChunkId(Container->ContainerId);
-	const FOnDemandChunkEntry*	Entry = Container->FindChunkEntry(ChunkId);
-
-	if (Entry == nullptr)
-	{
-		return FIoStatusBuilder(EIoErrorCode::NotFound) << TEXT("Failed to find header chunk");
-	}
-
-	UE_LOG(LogIoStoreOnDemand, Verbose, TEXT("Fetching encoded container header, ContainerName='%s'"), *Container->Name);
-	FIoStatus Status = FetchEncodedContainerChunk(Host, ChunkUrlBuilder, Container, *Entry, OutEncodedChunk);
-	if (!Status.IsOk())
-	{
-		return Status;
-	}
-
-	return EIoErrorCode::Ok;
-}
-
-static FIoStatus DecodeContainerChunk(
-	FSharedOnDemandContainer Container,
-	const FOnDemandChunkEntry& Entry,
-	FMemoryView EncodedChunk,
-	FIoBuffer& OutRawChunk)
+///////////////////////////////////////////////////////////////////////////////
+static TIoStatusOr<FIoBuffer> DecodeChunk(const FOnDemandChunkInfo& ChunkInfo, FMemoryView EncodedChunk)
 {
-	FOnDemandChunkInfo ChunkInfo(Container, Entry);
-
 	FIoChunkDecodingParams Params;
 	Params.CompressionFormat = ChunkInfo.CompressionFormat();
 	Params.EncryptionKey = ChunkInfo.EncryptionKey();
@@ -213,7 +164,7 @@ static FIoStatus DecodeContainerChunk(
 	Params.EncodedBlockSize = ChunkInfo.Blocks();
 	Params.BlockHash = ChunkInfo.BlockHashes();
 
-	OutRawChunk = FIoBuffer(ChunkInfo.RawSize());
+	FIoBuffer OutRawChunk = FIoBuffer(ChunkInfo.RawSize());
 	if (FIoChunkEncoding::Decode(Params, EncodedChunk, OutRawChunk.GetMutableView()) == false)
 	{
 		FIoStatus Status = FIoStatusBuilder(EIoErrorCode::ReadError)
@@ -221,68 +172,21 @@ static FIoStatus DecodeContainerChunk(
 		return Status;
 	}
 
-	return EIoErrorCode::Ok;
+	return OutRawChunk;
 }
 
-static FIoStatus DecodeContainerHeader(
-	FSharedOnDemandContainer Container,
-	FMemoryView EncodedHeaderChunk,
-	FIoContainerHeader& OutHeader)
-{
-	const FIoChunkId			ChunkId = CreateContainerHeaderChunkId(Container->ContainerId);
-	const FOnDemandChunkEntry*	Entry = Container->FindChunkEntry(ChunkId);
-
-	FIoBuffer RawChunk;
-	if (FIoStatus Status = DecodeContainerChunk(Container, *Entry, EncodedHeaderChunk, RawChunk); !Status.IsOk())
-	{
-		return Status;
-	}
-
-	FMemoryReaderView Ar(RawChunk.GetView());
-	Ar << OutHeader;
-	Ar.Close();
-
-	if (Ar.IsError() || Ar.IsCriticalError())
-	{
-		FIoStatus Status = FIoStatusBuilder(EIoErrorCode::FileNotOpen) << TEXT("Failed to serialize container header");
-		return Status;
-	}
-
-	return EIoErrorCode::Ok;
-}
-
-static FIoStatus FetchContainerHeader(
-	FStringView Host,
-	FAnsiStringBuilderBase& ChunkUrlBuilder,
-	FSharedOnDemandContainer Container,
-	FIoContainerHeader& OutHeader)
+///////////////////////////////////////////////////////////////////////////////
+static TIoStatusOr<FSharedContainerHeader> DeserializeContainerHeader(const FOnDemandChunkInfo& ChunkInfo, FMemoryView EncodedHeaderChunk)
 {
-	const FIoChunkId			ChunkId = CreateContainerHeaderChunkId(Container->ContainerId);
-	const FOnDemandChunkEntry*	Entry = Container->FindChunkEntry(ChunkId);
-
-	if (Entry == nullptr)
-	{
-		return FIoStatusBuilder(EIoErrorCode::NotFound) << TEXT("Failed to find header chunk");
-	}
-
-	UE_LOG(LogIoStoreOnDemand, Verbose, TEXT("Fetching container header, ContainerName='%s'"), *Container->Name);
-	FIoBuffer EncodedChunk;
-	{
-		FIoStatus Status = FetchEncodedContainerChunk(Host, ChunkUrlBuilder, Container, *Entry, EncodedChunk);
-		if (!Status.IsOk())
-		{
-			return Status;
-		}
-	}
-
-	FIoBuffer RawChunk;
-	if (FIoStatus Status = DecodeContainerChunk(Container, *Entry, EncodedChunk.GetView(), RawChunk); !Status.IsOk())
+	TIoStatusOr<FIoBuffer> Chunk = DecodeChunk(ChunkInfo, EncodedHeaderChunk);
+	if (Chunk.IsOk() == false)
 	{
-		return Status;
+		return Chunk.Status();
 	}
 
-	FMemoryReaderView Ar(RawChunk.GetView());
-	Ar << OutHeader;
+	FSharedContainerHeader OutHeader = MakeShared<FIoContainerHeader>();
+	FMemoryReaderView Ar(Chunk.ValueOrDie().GetView());
+	Ar << *OutHeader;
 	Ar.Close();
 
 	if (Ar.IsError() || Ar.IsCriticalError())
@@ -291,7 +195,7 @@ static FIoStatus FetchContainerHeader(
 		return Status;
 	}
 
-	return EIoErrorCode::Ok;
+	return OutHeader; 
 }
 
 ///////////////////////////////////////////////////////////////////////////////
@@ -320,8 +224,7 @@ static FIoStatus BuildInstallData(
 	// Setup the package information for each container
 	for (const FSharedOnDemandContainer& Container : Containers)
 	{
-		check(Container->Header.IsValid());
-		if (Container->Header->PackageIds.IsEmpty())
+		if (Container->Header.IsValid() == false || Container->Header->PackageIds.IsEmpty())
 		{
 			// The container contains no package data
 			continue;
@@ -731,7 +634,7 @@ FIoStatus FOnDemandIoStore::Unmount(FStringView MountId)
 					PackageStoreBackend->Unmount(ContainerName);
 				}
 
-				EncodedContainerHeaders.Remove(ContainerName);
+				PendingContainerHeaders.Remove(ContainerName);
 
 				return true;
 			}
@@ -1008,7 +911,7 @@ bool FOnDemandIoStore::Tick()
 			MountRequests.Remove(Request);
 		}
 
-		OnMountRequestComplete(*Request, 
+		CompleteMountRequest(*Request, 
 			FOnDemandMountResult
 			{
 				.MountId = MoveTemp(Request->Args.MountId),
@@ -1033,7 +936,7 @@ bool FOnDemandIoStore::Tick()
 			InstallRequests.Remove(Request);
 		}
 
-		OnInstallRequestComplete(*Request,
+		CompleteInstallRequest(*Request,
 			FOnDemandInstallResult
 			{
 				.Status = MoveTemp(Status),
@@ -1092,7 +995,7 @@ FIoStatus FOnDemandIoStore::TickMountRequest(FMountRequest& MountRequest)
 	}
 
 	// Containers haven't been created yet, do it now
-	TArray<FSharedOnDemandContainer> MountRequestContainers;
+	TArray<FSharedOnDemandContainer> RequestedContainers;
 	
 	FStringView Host, TocRelUrl;
 	Private::SplitHostUrl(Args.Url, Host, TocRelUrl);
@@ -1100,7 +1003,7 @@ FIoStatus FOnDemandIoStore::TickMountRequest(FMountRequest& MountRequest)
 
 	if (Args.Toc)
 	{
-		CreateContainersFromToc(Args.MountId, TocPath, *Args.Toc, MountRequestContainers);
+		CreateContainersFromToc(Args.MountId, TocPath, *Args.Toc, RequestedContainers);
 	}
 	else if (Args.FilePath.IsEmpty() == false)
 	{
@@ -1116,7 +1019,7 @@ FIoStatus FOnDemandIoStore::TickMountRequest(FMountRequest& MountRequest)
 
 		Args.Toc = MakeUnique<FOnDemandToc>(TocStatus.ConsumeValueOrDie());
 
-		CreateContainersFromToc(Args.MountId, TocPath, *Args.Toc, MountRequestContainers);
+		CreateContainersFromToc(Args.MountId, TocPath, *Args.Toc, RequestedContainers);
 	}
 	else if (Args.Url.IsEmpty() == false)
 	{
@@ -1133,13 +1036,13 @@ FIoStatus FOnDemandIoStore::TickMountRequest(FMountRequest& MountRequest)
 
 		Args.Toc = MakeUnique<FOnDemandToc>(TocStatus.ConsumeValueOrDie());
 
-		CreateContainersFromToc(Args.MountId, TocPath, *Args.Toc, MountRequestContainers);
+		CreateContainersFromToc(Args.MountId, TocPath, *Args.Toc, RequestedContainers);
 	}
 	
 	TAnsiStringBuilder<512> ChunkUrlBuilder;
-	TMap<FString, FIoBuffer> MountRequestEncodedContainerHeaders;
+	TMap<FString, FIoBuffer> RequestedContainerHeaders;
 
-	for (const FSharedOnDemandContainer& Container : MountRequestContainers)
+	for (const FSharedOnDemandContainer& Container : RequestedContainers)
 	{
 		if (EnumHasAnyFlags(Args.Options, EOnDemandMountOptions::StreamOnDemand))
 		{
@@ -1165,57 +1068,64 @@ FIoStatus FOnDemandIoStore::TickMountRequest(FMountRequest& MountRequest)
 			}
 		}
 
-		if (EnumHasAnyFlags(Container->Flags, EOnDemandContainerFlags::PendingEncryptionKey))
+		// Try fetch and deserialize the container header if the container is used for installing content.
+		if (EnumHasAnyFlags(Container->Flags, EOnDemandContainerFlags::InstallOnDemand))
 		{
-			if (EnumHasAnyFlags(Container->Flags, EOnDemandContainerFlags::InstallOnDemand))
+			const FIoChunkId ChunkId = CreateContainerHeaderChunkId(Container->ContainerId);
+			if (const FOnDemandChunkInfo ChunkInfo = FOnDemandChunkInfo::Find(Container, ChunkId); ChunkInfo.IsValid())
 			{
-				//Fetch encoded container header and store it until it can be decrypted
-				FIoBuffer EncodedHeader;
-				if (FIoStatus Status = Private::FetchEncodedContainerHeader(Host, ChunkUrlBuilder, Container, EncodedHeader); !Status.IsOk())
+				UE_LOG(LogIoStoreOnDemand, Verbose, TEXT("Fetching container header, ContainerName='%s', ChunkId='%s'"),
+					*Container->Name, *LexToString(ChunkId));
+
+				const uint32 RetryCount = 2;
+				TIoStatusOr<FIoBuffer> ContainerHeaderChunk = FHttpClient::Get(
+					Private::GetChunkUrl(Host, *Container, ChunkInfo.ChunkEntry(),
+					ChunkUrlBuilder).ToView(),
+					RetryCount,
+					EHttpRedirects::Follow);
+
+				if (ContainerHeaderChunk.IsOk() == false)
 				{
-					return Status;
+					return ContainerHeaderChunk.Status();
 				}
-				MountRequestEncodedContainerHeaders.Add(Container->UniqueName(), MoveTemp(EncodedHeader));
-			}
 
-			continue;
+				if (EnumHasAnyFlags(Container->Flags, EOnDemandContainerFlags::PendingEncryptionKey))
+				{
+					RequestedContainerHeaders.Add(Container->UniqueName(), ContainerHeaderChunk.ConsumeValueOrDie());
+					continue;
+				}
+
+				TIoStatusOr<FSharedContainerHeader> Header = Private::DeserializeContainerHeader(
+					ChunkInfo,
+					ContainerHeaderChunk.ValueOrDie().GetView());
+
+				if (Header.IsOk() == false)
+				{
+					return Header.Status();
+				}
+
+				Container->Header = Header.ConsumeValueOrDie(); 
+			}
 		}
 
 		TStringBuilder<128> Sb;
-
-		EnumAddFlags(Container->Flags, EOnDemandContainerFlags::Mounted);
-		Container->EncryptionKeyGuid.Empty();
-		Sb.Reset();
 		LexToString(Container->Flags, Sb);
 		UE_LOG(LogIoStoreOnDemand, Log, TEXT("Mounting container '%s', Entries=%d, Flags='%s'"),
 			*Container->Name, Container->ChunkEntries.Num(), Sb.ToString());
-
-		if (EnumHasAnyFlags(Container->Flags, EOnDemandContainerFlags::InstallOnDemand))
-		{
-			// Fetch container header
-			FSharedContainerHeader Header = MakeShared<FIoContainerHeader>();
-			if (FIoStatus Status = Private::FetchContainerHeader(Host, ChunkUrlBuilder, Container, *Header); !Status.IsOk())
-			{
-				return Status;
-			}
-			Container->Header = MoveTemp(Header);
-		}
+		EnumAddFlags(Container->Flags, EOnDemandContainerFlags::Mounted);
+		Container->EncryptionKeyGuid.Empty();
 	}
 
 	{
 		UE::TUniqueLock Lock(ContainerMutex);
-		Containers.Append(MoveTemp(MountRequestContainers));
-		EncodedContainerHeaders.Append(MoveTemp(MountRequestEncodedContainerHeaders));
+		Containers.Append(MoveTemp(RequestedContainers));
+		PendingContainerHeaders.Append(MoveTemp(RequestedContainerHeaders));
 	}
 
-	if (bAnyPendingEncryptionKey)
-	{
-		return EIoErrorCode::PendingEncryptionKey;
-	}
-	return EIoErrorCode::Ok;
+	return bAnyPendingEncryptionKey ? EIoErrorCode::PendingEncryptionKey : EIoErrorCode::Ok;
 }
 
-void FOnDemandIoStore::OnMountRequestComplete(FMountRequest& Request, FOnDemandMountResult&& MountResult)
+void FOnDemandIoStore::CompleteMountRequest(FMountRequest& Request, FOnDemandMountResult&& MountResult)
 {
 	if (EnumHasAnyFlags(Request.Args.Options, EOnDemandMountOptions::CallbackOnGameThread))
 	{
@@ -1267,11 +1177,13 @@ FIoStatus FOnDemandIoStore::TickInstallRequest(FInstallRequest& InstallRequest)
 
 	TSet<FSharedOnDemandContainer> ContainersForInstallation;
 	TSet<FPackageId> PackageIdsToInstall;
-	FIoStatus GetContainersStatus = GetContainersAndPackagesForInstall(
-		InstallRequest.Args, ContainersForInstallation, PackageIdsToInstall);
-	if (!GetContainersStatus.IsOk())
+
+	if (FIoStatus Status = GetContainersAndPackagesForInstall(
+		InstallRequest.Args,
+		ContainersForInstallation,
+		PackageIdsToInstall); !Status.IsOk())
 	{
-		return GetContainersStatus;
+		return Status;
 	}
 
 	// Its OK for PackageIdsToInstall to be empty at this point. Any chunks not referenced by a package must still be installed.
@@ -1441,7 +1353,7 @@ FIoStatus FOnDemandIoStore::TickInstallRequest(FInstallRequest& InstallRequest)
 	return EIoErrorCode::Ok;
 }
 
-void FOnDemandIoStore::OnInstallRequestComplete(FInstallRequest& Request, FOnDemandInstallResult&& InstallResult)
+void FOnDemandIoStore::CompleteInstallRequest(FInstallRequest& Request, FOnDemandInstallResult&& InstallResult)
 {
 	if (EnumHasAnyFlags(Request.Args.Options, EOnDemandInstallOptions::CallbackOnGameThread))
 	{
@@ -1487,27 +1399,33 @@ void FOnDemandIoStore::OnEncryptionKeyAdded(const FGuid& Id, const FAES::FAESKey
 		
 		Container->EncryptionKeyGuid.Empty();
 
-		const FString ContainerName = Container->UniqueName();
-		const uint32 ContainerNameHash = GetTypeHash(ContainerName);
+		const FString ContainerName		= Container->UniqueName();
+		const uint32 ContainerNameHash	= GetTypeHash(ContainerName);
 
-		FIoBuffer* EncodedHeader = EncodedContainerHeaders.FindByHash(ContainerNameHash, ContainerName);
-		if (ensure(EncodedHeader))
+		if (FIoBuffer* EncodedHeader = PendingContainerHeaders.FindByHash(ContainerNameHash, ContainerName))
 		{
-			FSharedContainerHeader Header = MakeShared<FIoContainerHeader>();
-			FIoStatus DecodeStatus = Private::DecodeContainerHeader(Container, EncodedHeader->GetView(), *Header);
-			EncodedHeader = nullptr;
-			EncodedContainerHeaders.RemoveByHash(ContainerNameHash, ContainerName);
+			const FIoChunkId ChunkId			= CreateContainerHeaderChunkId(Container->ContainerId);
+			const FOnDemandChunkInfo ChunkInfo	= FOnDemandChunkInfo::Find(Container, ChunkId);
 
-			if (ensure(DecodeStatus.IsOk()))
-			{
-				Container->Header = MoveTemp(Header);
-				EnumAddFlags(Container->Flags, EOnDemandContainerFlags::Mounted);
-			}
-			else
+			if (ensure(ChunkInfo.IsValid()))
 			{
-				UE_LOG(LogIoStoreOnDemand, Error, TEXT("Failed to decode header when mounting container '%s', Entries=%d, Flags='%s'"),
-					*Container->Name, Container->ChunkEntries.Num(), *LexToString(Container->Flags));
+				TIoStatusOr<FSharedContainerHeader> Header = Private::DeserializeContainerHeader(
+					ChunkInfo,
+					EncodedHeader->GetView());
+
+				if (ensure(Header.IsOk()))
+				{
+					Container->Header = Header.ConsumeValueOrDie(); 
+					EnumAddFlags(Container->Flags, EOnDemandContainerFlags::Mounted);
+				}
+				else
+				{
+					UE_LOG(LogIoStoreOnDemand, Error, TEXT("Failed to deserialize header when mounting container '%s', Entries=%d, Flags='%s'"),
+						*Container->Name, Container->ChunkEntries.Num(), *LexToString(Container->Flags));
+				}
 			}
+
+			PendingContainerHeaders.RemoveByHash(ContainerNameHash, ContainerName);
 		}
 	}
 }
@@ -1667,10 +1585,13 @@ FIoStatus FOnDemandIoStore::GetContainersAndPackagesForInstall(
 {
 	TSet<FSharedOnDemandContainer> ContainersForInstallation;
 	TSet<FSharedOnDemandContainer> ContainersWithMountId;
-	FIoStatus GetContainersStatus = GetContainersForInstall(Args.MountId, ContainersForInstallation, ContainersWithMountId);
-	if (!GetContainersStatus.IsOk())
+
+	if (FIoStatus Status = GetContainersForInstall(
+		Args.MountId,
+		ContainersForInstallation,
+		ContainersWithMountId); !Status.IsOk())
 	{
-		return GetContainersStatus;
+		return Status;
 	}
 
 	TSet<FPackageId> PackageIdsToInstall(Args.PackageIds);
@@ -1680,6 +1601,11 @@ FIoStatus FOnDemandIoStore::GetContainersAndPackagesForInstall(
 	{
 		for (const FSharedOnDemandContainer& Container : ContainersWithMountId)
 		{
+			if (Container->Header.IsValid() == false)
+			{
+				continue;;
+			}
+
 			for (const FPackageId& PackageId : Container->Header->PackageIds)
 			{
 				PackageIdsToInstall.Add(PackageId);
@@ -1693,6 +1619,11 @@ FIoStatus FOnDemandIoStore::GetContainersAndPackagesForInstall(
 
 		for (const FSharedOnDemandContainer& Container : SearchContainers)
 		{
+			if (Container->Header.IsValid() == false)
+			{
+				continue;;
+			}
+
 			for (const FString& Tag : Args.TagSets)
 			{
 				for (const FOnDemandTagSet& TagSet : Container->TagSets)
diff --git a/Engine/Source/Runtime/Experimental/IoStore/OnDemand/Private/OnDemandIoStore.h b/Engine/Source/Runtime/Experimental/IoStore/OnDemand/Private/OnDemandIoStore.h
--- a/Engine/Source/Runtime/Experimental/IoStore/OnDemand/Private/OnDemandIoStore.h
+++ b/Engine/Source/Runtime/Experimental/IoStore/OnDemand/Private/OnDemandIoStore.h
@@ -189,15 +189,6 @@ TBitArray<> FOnDemandContainer::GetReferencedChunkEntries() const
 }
 
 ///////////////////////////////////////////////////////////////////////////////
-namespace Private
-{
-	static FIoStatus DecodeContainerChunk(
-		FSharedOnDemandContainer Container,
-		const FOnDemandChunkEntry& Entry,
-		FMemoryView EncodedHeaderChunk,
-		FIoBuffer& OutRawChunk);
-}
-
 struct FOnDemandChunkInfo
 {
 	FOnDemandChunkInfo()
@@ -213,17 +204,15 @@ struct FOnDemandChunkInfo
 	inline TConstArrayView<uint32>			Blocks() const;
 	inline TConstArrayView<FIoBlockHash>	BlockHashes() const;
 	FAnsiStringView							ChunksDirectory() const { return SharedContainer->ChunksDirectory; }
+	const FOnDemandChunkEntry&				ChunkEntry() const { return Entry; }
 
 	bool									IsValid() const { return SharedContainer.IsValid(); }
 	operator								bool() const { return IsValid(); }
+	
+	inline static FOnDemandChunkInfo		Find(FSharedOnDemandContainer Container, const FIoChunkId& ChunkId);
 
 private:
 	friend class FOnDemandIoStore;
-	friend FIoStatus Private::DecodeContainerChunk(
-		FSharedOnDemandContainer Container,
-		const FOnDemandChunkEntry& Entry,
-		FMemoryView EncodedHeaderChunk,
-		FIoBuffer& OutRawChunk);
 
 	FOnDemandChunkInfo(FSharedOnDemandContainer InContainer, const FOnDemandChunkEntry& InEntry)
 		: SharedContainer(InContainer)
@@ -246,6 +235,17 @@ TConstArrayView<FIoBlockHash> FOnDemandChunkInfo::BlockHashes() const
 		: TConstArrayView<FIoBlockHash>(SharedContainer->BlockHashes.GetData() + Entry.BlockOffset, Entry.BlockCount);
 }
 
+FOnDemandChunkInfo FOnDemandChunkInfo::Find(FSharedOnDemandContainer Container, const FIoChunkId& ChunkId)
+{
+	check(Container.IsValid());
+	if (const FOnDemandChunkEntry* Entry = Container->FindChunkEntry(ChunkId))
+	{
+		return FOnDemandChunkInfo(Container, *Entry);
+	}
+
+	return FOnDemandChunkInfo();
+}
+
 ///////////////////////////////////////////////////////////////////////////////
 class FOnDemandIoStore
 	: public TSharedFromThis<FOnDemandIoStore, ESPMode::ThreadSafe>
@@ -291,13 +291,13 @@ public:
 
 private:
 	FIoStatus				GetContainersForInstall(
-		FStringView MountId,
-		TSet<FSharedOnDemandContainer>& OutContainersForInstallation,
-		TSet<FSharedOnDemandContainer>& OutContainersWithMountId);
+								FStringView MountId,
+								TSet<FSharedOnDemandContainer>& OutContainersForInstallation,
+								TSet<FSharedOnDemandContainer>& OutContainersWithMountId);
 	FIoStatus				GetContainersAndPackagesForInstall(
-		const FOnDemandInstallArgsCommon& Args,
-		TSet<FSharedOnDemandContainer>& OutContainersForInstallation,
-		TSet<FPackageId>& OutPackageIdsToInstall);
+								const FOnDemandInstallArgsCommon& Args,
+								TSet<FSharedOnDemandContainer>& OutContainersForInstallation,
+								TSet<FPackageId>& OutPackageIdsToInstall);
 	void					OnPostFork(EForkProcessRole ProcessRole);
 	FIoStatus				InitializeOnDemandInstallCache();
 	FOnDemandChunkInfo		GetChunkInfo(const FIoChunkId& ChunkId, EOnDemandContainerFlags ContainerFlags);
@@ -305,9 +305,9 @@ private:
 	void					TickLoop();
 	bool					Tick();
 	FIoStatus				TickMountRequest(FMountRequest& MountRequest);
-	void					OnMountRequestComplete(FMountRequest& MountRequest, FOnDemandMountResult&& MountResult);
+	void					CompleteMountRequest(FMountRequest& MountRequest, FOnDemandMountResult&& MountResult);
 	FIoStatus				TickInstallRequest(FInstallRequest& InstallRequest);
-	void					OnInstallRequestComplete(FInstallRequest& InstallRequest, FOnDemandInstallResult&& InstallResult);
+	void					CompleteInstallRequest(FInstallRequest& InstallRequest, FOnDemandInstallResult&& InstallResult);
 	void					OnEncryptionKeyAdded(const FGuid& Id, const FAES::FAESKey& Key);
 	static void				CreateContainersFromToc(
 								FStringView MountId,
@@ -319,7 +319,7 @@ private:
 	FSharedPackageStoreBackend			PackageStoreBackend;
 	FDelegateHandle						OnMountPakHandle;
 	TArray<FSharedOnDemandContainer>	Containers;
-	TMap<FString, FIoBuffer>			EncodedContainerHeaders;
+	TMap<FString, FIoBuffer>			PendingContainerHeaders;
 	UE::FMutex							ContainerMutex;
 
 	TArray<FSharedMountRequest>			MountRequests;
