Iris
- No longer disconnect when encountering nan (still disconnects but now we zero out data in quantize if we detect nan) for PackedVectors

diff --git a/Engine/Plugins/Runtime/ReplicationSystemTestPlugin/Source/Private/Tests/Serialization/TestPackedVectorNetSerializers.cpp b/Engine/Plugins/Runtime/ReplicationSystemTestPlugin/Source/Private/Tests/Serialization/TestPackedVectorNetSerializers.cpp
--- a/Engine/Plugins/Runtime/ReplicationSystemTestPlugin/Source/Private/Tests/Serialization/TestPackedVectorNetSerializers.cpp
+++ b/Engine/Plugins/Runtime/ReplicationSystemTestPlugin/Source/Private/Tests/Serialization/TestPackedVectorNetSerializers.cpp
@@ -25,6 +25,7 @@ protected:
 
 	void TestIsEqual();
 	void TestSerialize(float MaxAbsDiff, int32 MaxUlpDiff);
+	void TestSerializeBadValues();
 	void TestSerializeDelta();
 	void TestValidate();
 
@@ -39,6 +40,7 @@ protected:
 	};
 
 	static const FFloatTriplet BadValues[];
+	static const FFloatTriplet ExpectedBadValues[];
 	static const SIZE_T BadValueCount;
 	static const FVector Values[];
 	static const SIZE_T ValueCount;
@@ -55,6 +57,7 @@ const FTestPackedVectorNetSerializerBase::FFloatTriplet FTestPackedVectorNetSeri
 	{0.0f, std::numeric_limits<ScalarType>::signaling_NaN(), 0.0f},
 };
 const SIZE_T FTestPackedVectorNetSerializerBase::BadValueCount = sizeof(FTestPackedVectorNetSerializerBase::BadValues)/sizeof(FTestPackedVectorNetSerializerBase::BadValues[0]);
+const FTestPackedVectorNetSerializerBase::FFloatTriplet FTestPackedVectorNetSerializerBase::ExpectedBadValues[FTestPackedVectorNetSerializerBase::BadValueCount] = {};
 
 const FVector FTestPackedVectorNetSerializerBase::Values[] =
 {
@@ -165,6 +168,11 @@ UE_NET_TEST_FIXTURE(FTestVectorNetQuantizeNetSerializer, TestSerialize)
 	TestSerialize(1.0f, 1);
 }
 
+UE_NET_TEST_FIXTURE(FTestVectorNetQuantizeNetSerializer, TestSerializeBadValues)
+{
+	TestSerializeBadValues();
+}
+
 UE_NET_TEST_FIXTURE(FTestVectorNetQuantizeNetSerializer, TestSerializeDelta)
 {
 	TestSerializeDelta();
@@ -265,6 +273,14 @@ void FTestPackedVectorNetSerializerBase::TestSerialize(float MaxAbsDiff, int32 M
 	}
 }
 
+void FTestPackedVectorNetSerializerBase::TestSerializeBadValues()
+{
+	for (const bool bQuantizedCompare : {false, true})
+	{
+		Super::TestSerialize(reinterpret_cast<const FVector*>(BadValues), reinterpret_cast<const FVector*>(ExpectedBadValues), BadValueCount, NetSerializerConfig, bQuantizedCompare);
+	}
+}
+
 void FTestPackedVectorNetSerializerBase::TestSerializeDelta()
 {
 	Super::TestSerializeDelta(Values, ValueCount, NetSerializerConfig);
diff --git a/Engine/Source/Runtime/Experimental/Iris/Core/Private/Iris/Serialization/PackedVectorNetSerializers.cpp b/Engine/Source/Runtime/Experimental/Iris/Core/Private/Iris/Serialization/PackedVectorNetSerializers.cpp
--- a/Engine/Source/Runtime/Experimental/Iris/Core/Private/Iris/Serialization/PackedVectorNetSerializers.cpp
+++ b/Engine/Source/Runtime/Experimental/Iris/Core/Private/Iris/Serialization/PackedVectorNetSerializers.cpp
@@ -6,6 +6,7 @@
 #include "Iris/Serialization/NetBitStreamWriter.h"
 #include "Iris/Core/BitTwiddling.h"
 #include "Math/Vector.h"
+#include "Logging/LogMacros.h"
 
 namespace UE::Net::Private
 {
@@ -242,9 +243,8 @@ void FPackedVectorNetSerializerBase::Deserialize(uint32 ScaleBitCount, FNetSeria
 			memcpy(&Vector, &TempValue.X, 3U*sizeof(double)); //-V512
 			if (Vector.ContainsNaN())
 			{
-				// While we could detect this at send time it's very likely that a NaN or infinite value
-				// indicates something is very wrong with the simulation and that clients are
-				// better off being disconnected so they can join another game.
+				// While we detect and zero out nan data at quantize time if we get bad data here it
+				// indicates serialization issues and client should disconnect.
 				Context.SetError(GNetError_InvalidValue);
 				return;
 			}
@@ -260,9 +260,8 @@ void FPackedVectorNetSerializerBase::Deserialize(uint32 ScaleBitCount, FNetSeria
 			memcpy(&Vector, &Components, sizeof(Components));
 			if (Vector.ContainsNaN())
 			{
-				// While we could detect this at send time it's very likely that a NaN or infinite value
-				// indicates something is very wrong with the simulation and that clients are
-				// better off being disconnected so they can join another game.
+				// While we detect and zero out nan data at quantize time if we get bad data here it
+				// indicates serialization issues and client should disconnect.
 				Context.SetError(GNetError_InvalidValue);
 				return;
 			}
@@ -521,9 +520,8 @@ void FPackedVectorNetSerializerBase::DeserializeDelta(uint32 ScaleBitCount, FNet
 			memcpy(&Vector, &TempValue.X, 3U*sizeof(double)); //-V512
 			if (Vector.ContainsNaN())
 			{
-				// While we could detect this at send time it's very likely that a NaN or infinite value
-				// indicates something is very wrong with the simulation and that clients are
-				// better off being disconnected so they can join another game.
+				// While we detect and zero out nan data at quantize time if we get bad data here it
+				// indicates serialization issues and client should disconnect.
 				Context.SetError(GNetError_InvalidValue);
 				return;
 			}
@@ -562,9 +560,8 @@ void FPackedVectorNetSerializerBase::DeserializeDelta(uint32 ScaleBitCount, FNet
 			memcpy(&Vector, &Components, sizeof(Components));
 			if (Vector.ContainsNaN())
 			{
-				// While we could detect this at send time it's very likely that a NaN or infinite value
-				// indicates something is very wrong with the simulation and that clients are
-				// better off being disconnected so they can join another game.
+				// While we detect and zero out nan data at quantize time if we get bad data here it
+				// indicates serialization issues and client should disconnect.
 				Context.SetError(GNetError_InvalidValue);
 				return;
 			}
@@ -597,6 +594,19 @@ void FPackedVectorNetSerializerBase::Quantize(uint32 ScaleBitCount, FNetSerializ
 	const SourceType& Source = *reinterpret_cast<const SourceType*>(Args.Source);
 	QuantizedType& Target = *reinterpret_cast<QuantizedType*>(Args.Target);
 
+	if (Source.ContainsNaN())
+	{
+		logOrEnsureNanError(TEXT("%s"), TEXT("PackedVectorNetSerializeBase::Quantize Value isn't finite. Clearing for safety."));		
+
+		constexpr uint32 ComponentBitCountForZeroValue = 1U;
+
+		QuantizedType TempValue = {};
+		TempValue.ComponentBitCountAndExtraInfo = IsScaledValueMask | ComponentBitCountForZeroValue;
+		Target = TempValue;
+
+		return;
+	}
+
 	const ScalarType Scale = ScalarType(IntType(1) << ScaleBitCount);
 	
 	SourceType ScaledValue;
