Added command line argument for ignoring filtering containers based on the OnDemand flag when uploading I/O store content.

diff --git a/Engine/Source/Runtime/Experimental/IoStore/OnDemand/Private/Tool/Upload.cpp b/Engine/Source/Runtime/Experimental/IoStore/OnDemand/Private/Tool/Upload.cpp
--- a/Engine/Source/Runtime/Experimental/IoStore/OnDemand/Private/Tool/Upload.cpp
+++ b/Engine/Source/Runtime/Experimental/IoStore/OnDemand/Private/Tool/Upload.cpp
@@ -54,6 +54,7 @@ struct FUploadParams
 	bool bDeleteContainerFiles = true;
 	bool bDeletePakFiles = true;
 	bool bPerContainerTocs = false;
+	bool bIgnoreContainerFlags = false;
 
 	/** If we should write out the .iochunktoc to disk as well as uploading it. */
 	bool bWriteTocToDisk = false;
@@ -125,8 +126,9 @@ static FUploadParams BuildUploadParams(const FContext& Context)
 	Ret.bPerContainerTocs		= Context.Get<bool>(TEXT("-PerContainerTocs"),				Ret.bPerContainerTocs);
 	Ret.MaxConcurrentUploads	= Context.Get<int32>(TEXT("-MaxConcurrentUploads"),			Ret.MaxConcurrentUploads);
 
-	Ret.bDeleteContainerFiles	= !Context.Get<bool>(TEXT("-KeepContainerFiles"),	!Ret.bDeleteContainerFiles);
-	Ret.bDeletePakFiles			= !Context.Get<bool>(TEXT("-KeepPakFiles"),			!Ret.bDeletePakFiles);
+	Ret.bDeleteContainerFiles	= !Context.Get<bool>(TEXT("-KeepContainerFiles"),			!Ret.bDeleteContainerFiles);
+	Ret.bDeletePakFiles			= !Context.Get<bool>(TEXT("-KeepPakFiles"),					!Ret.bDeletePakFiles);
+	Ret.bIgnoreContainerFlags	= Context.Get<bool>(TEXT("-IgnoreContainerFlags"),			Ret.bIgnoreContainerFlags);
 
 	if (Ret.bWriteTocToDisk)
 	{
@@ -510,7 +512,8 @@ static TIoStatusOr<FUploadResult> UploadContainerFiles(
 			}
 		}
 
-		if (EnumHasAnyFlags(ContainerFileReader.GetContainerFlags(), EIoContainerFlags::OnDemand) == false)
+		if (UploadParams.bIgnoreContainerFlags == false &&
+			EnumHasAnyFlags(ContainerFileReader.GetContainerFlags(), EIoContainerFlags::OnDemand) == false)
 		{
 			continue;
 		}
@@ -880,6 +883,7 @@ static FCommand UploadCommand(
 		TArgument<bool>(TEXT("-KeepPakFiles"),			TEXT("Do not delete the springboard pak files")),
 		TArgument<bool>(TEXT("-WriteTocToDisk"),		TEXT("Output the TOC to disk as well as uploading")),
 		TArgument<bool>(TEXT("-PerContainerTocs"),		TEXT("Whether to generate TOC's for each container file(s)")),
+		TArgument<bool>(TEXT("-IgnoreContainerFlags"),	TEXT("Whether to ignore the OnDemand container flag")),
 		TArgument<int32>(TEXT("-MaxConcurrentUploads"),	TEXT("Number of simultaneous uploads")),
 		S3Arguments,
 	}
