Don't sacrifice designated initializer in public API

#rnx

diff --git a/Engine/Source/Runtime/Experimental/IoStore/OnDemand/Private/OnDemandIoStore.cpp b/Engine/Source/Runtime/Experimental/IoStore/OnDemand/Private/OnDemandIoStore.cpp
--- a/Engine/Source/Runtime/Experimental/IoStore/OnDemand/Private/OnDemandIoStore.cpp
+++ b/Engine/Source/Runtime/Experimental/IoStore/OnDemand/Private/OnDemandIoStore.cpp
@@ -652,11 +652,15 @@ TIoStatusOr<uint64> FOnDemandIoStore::GetInstallSize(const FOnDemandGetInstallSi
 
 	TSet<FSharedOnDemandContainer> ContainersForInstallation;
 	TSet<FPackageId> PackageIdsToInstall;
-	FIoStatus GetContainersStatus = const_cast<FOnDemandIoStore*>(this)->GetContainersAndPackagesForInstall(
-		Args, ContainersForInstallation, PackageIdsToInstall);
-	if (!GetContainersStatus.IsOk())
+
+	if (FIoStatus Status = GetContainersAndPackagesForInstall(
+		Args.MountId,
+		Args.TagSets,
+		Args.PackageIds,
+		ContainersForInstallation,
+		PackageIdsToInstall); !Status.IsOk())
 	{
-		return GetContainersStatus;
+		return Status;
 	}
 
 	Private::FInstallData InstallData;
@@ -681,11 +685,15 @@ FIoStatus FOnDemandIoStore::GetInstallSizesByMountId(const FOnDemandGetInstallSi
 
 	TSet<FSharedOnDemandContainer> ContainersForInstallation;
 	TSet<FPackageId> PackageIdsToInstall;
-	FIoStatus GetContainersStatus = const_cast<FOnDemandIoStore*>(this)->GetContainersAndPackagesForInstall(
-		Args, ContainersForInstallation, PackageIdsToInstall);
-	if (!GetContainersStatus.IsOk())
+
+	if (FIoStatus Status = GetContainersAndPackagesForInstall(
+		Args.MountId,
+		Args.TagSets,
+		Args.PackageIds,
+		ContainersForInstallation,
+		PackageIdsToInstall); !Status.IsOk())
 	{
-		return GetContainersStatus;
+		return Status;
 	}
 
 	Private::FInstallData InstallData;
@@ -1179,7 +1187,9 @@ FIoStatus FOnDemandIoStore::TickInstallRequest(FInstallRequest& InstallRequest)
 	TSet<FPackageId> PackageIdsToInstall;
 
 	if (FIoStatus Status = GetContainersAndPackagesForInstall(
-		InstallRequest.Args,
+		InstallRequest.Args.MountId,
+		InstallRequest.Args.TagSets,
+		InstallRequest.Args.PackageIds,
 		ContainersForInstallation,
 		PackageIdsToInstall); !Status.IsOk())
 	{
@@ -1535,7 +1545,7 @@ void FOnDemandIoStore::CreateContainersFromToc(
 FIoStatus FOnDemandIoStore::GetContainersForInstall(
 	FStringView MountId,
 	TSet<FSharedOnDemandContainer>& OutContainersForInstallation,
-	TSet<FSharedOnDemandContainer>& OutContainersWithMountId)
+	TSet<FSharedOnDemandContainer>& OutContainersWithMountId) const
 {
 	// Only install content from non-streaming container(s)
 	TSet<FSharedOnDemandContainer> ContainersForInstallation;
@@ -1579,25 +1589,27 @@ FIoStatus FOnDemandIoStore::GetContainersForInstall(
 }
 
 FIoStatus FOnDemandIoStore::GetContainersAndPackagesForInstall(
-	const FOnDemandInstallArgsCommon& Args, 
+	FStringView MountId,
+	const TArray<FString>& TagSets,
+	const TArray<FPackageId>& PackageIds,
 	TSet<FSharedOnDemandContainer>& OutContainersForInstallation, 
-	TSet<FPackageId>& OutPackageIdsToInstall)
+	TSet<FPackageId>& OutPackageIdsToInstall) const
 {
 	TSet<FSharedOnDemandContainer> ContainersForInstallation;
 	TSet<FSharedOnDemandContainer> ContainersWithMountId;
 
 	if (FIoStatus Status = GetContainersForInstall(
-		Args.MountId,
+		MountId,
 		ContainersForInstallation,
 		ContainersWithMountId); !Status.IsOk())
 	{
 		return Status;
 	}
 
-	TSet<FPackageId> PackageIdsToInstall(Args.PackageIds);
+	TSet<FPackageId> PackageIdsToInstall(PackageIds);
 
 	// Install all packages if no tag set(s) was specified
-	if (Args.TagSets.IsEmpty())
+	if (TagSets.IsEmpty())
 	{
 		for (const FSharedOnDemandContainer& Container : ContainersWithMountId)
 		{
@@ -1615,7 +1627,7 @@ FIoStatus FOnDemandIoStore::GetContainersAndPackagesForInstall(
 	else
 	{
 		const TSet<FSharedOnDemandContainer>& SearchContainers =
-			Args.MountId.IsEmpty() ? ContainersForInstallation : ContainersWithMountId;
+			MountId.IsEmpty() ? ContainersForInstallation : ContainersWithMountId;
 
 		for (const FSharedOnDemandContainer& Container : SearchContainers)
 		{
@@ -1624,7 +1636,7 @@ FIoStatus FOnDemandIoStore::GetContainersAndPackagesForInstall(
 				continue;;
 			}
 
-			for (const FString& Tag : Args.TagSets)
+			for (const FString& Tag : TagSets)
 			{
 				for (const FOnDemandTagSet& TagSet : Container->TagSets)
 				{
diff --git a/Engine/Source/Runtime/Experimental/IoStore/OnDemand/Private/OnDemandIoStore.h b/Engine/Source/Runtime/Experimental/IoStore/OnDemand/Private/OnDemandIoStore.h
--- a/Engine/Source/Runtime/Experimental/IoStore/OnDemand/Private/OnDemandIoStore.h
+++ b/Engine/Source/Runtime/Experimental/IoStore/OnDemand/Private/OnDemandIoStore.h
@@ -293,11 +293,13 @@ private:
 	FIoStatus				GetContainersForInstall(
 								FStringView MountId,
 								TSet<FSharedOnDemandContainer>& OutContainersForInstallation,
-								TSet<FSharedOnDemandContainer>& OutContainersWithMountId);
+								TSet<FSharedOnDemandContainer>& OutContainersWithMountId) const;
 	FIoStatus				GetContainersAndPackagesForInstall(
-								const FOnDemandInstallArgsCommon& Args,
+								FStringView MountId,
+								const TArray<FString>& TagSets,
+								const TArray<FPackageId>& PackageIds,
 								TSet<FSharedOnDemandContainer>& OutContainersForInstallation,
-								TSet<FPackageId>& OutPackageIdsToInstall);
+								TSet<FPackageId>& OutPackageIdsToInstall) const;
 	void					OnPostFork(EForkProcessRole ProcessRole);
 	FIoStatus				InitializeOnDemandInstallCache();
 	FOnDemandChunkInfo		GetChunkInfo(const FIoChunkId& ChunkId, EOnDemandContainerFlags ContainerFlags);
@@ -320,7 +322,7 @@ private:
 	FDelegateHandle						OnMountPakHandle;
 	TArray<FSharedOnDemandContainer>	Containers;
 	TMap<FString, FIoBuffer>			PendingContainerHeaders;
-	UE::FMutex							ContainerMutex;
+	mutable UE::FMutex					ContainerMutex;
 
 	TArray<FSharedMountRequest>			MountRequests;
 	TArray<FSharedInstallRequest>		InstallRequests;
diff --git a/Engine/Source/Runtime/Experimental/IoStore/OnDemand/Public/IO/IoStoreOnDemand.h b/Engine/Source/Runtime/Experimental/IoStore/OnDemand/Public/IO/IoStoreOnDemand.h
--- a/Engine/Source/Runtime/Experimental/IoStore/OnDemand/Public/IO/IoStoreOnDemand.h
+++ b/Engine/Source/Runtime/Experimental/IoStore/OnDemand/Public/IO/IoStoreOnDemand.h
@@ -321,7 +321,8 @@ enum class EOnDemandInstallOptions
 };
 ENUM_CLASS_FLAGS(EOnDemandInstallOptions);
 
-struct FOnDemandInstallArgsCommon
+/** Arguments for installing/downloading on-demand content. */
+struct FOnDemandInstallArgs 
 {
 	/** Install all content from containers matching this mount ID. */
 	FString MountId;
@@ -329,11 +330,6 @@ struct FOnDemandInstallArgsCommon
 	TArray<FString> TagSets;
 	/** Package ID's to install. */
 	TArray<FPackageId> PackageIds;
-};
-
-/** Arguments for installing/downloading on-demand content. */
-struct FOnDemandInstallArgs : public FOnDemandInstallArgsCommon
-{
 	/** URL from where to download the chunks. */
 	FString Url;
 	/** Content handle. */
@@ -358,9 +354,15 @@ struct FOnDemandInstallResult
 /** Install completion callback. */
 using FOnDemandInstallCompleted = TUniqueFunction<void(FOnDemandInstallResult)>;
 
-/** Arguments for getting the size of on-demand content */
-struct FOnDemandGetInstallSizeArgs : public FOnDemandInstallArgsCommon
+/** Arguments for getting the size of on-demand content. */
+struct FOnDemandGetInstallSizeArgs 
 {
+	/** Get total install size for containers matching this mount ID. */
+	FString MountId;
+	/** Get total install size for the specified tag(s) and optionally matching the mount ID. */
+	TArray<FString> TagSets;
+	/** Get total intall size for the specified package IDs. */
+	TArray<FPackageId> PackageIds;
 };
 
 class FIoStoreOnDemandModule
