Dataflow - Allows shorter name for array and structure pins by adding a SkipInDisplayNameChain property Meta.

#rnx

diff --git a/Engine/Plugins/ChaosClothAssetEditor/Source/ChaosClothAssetDataflowNodes/Public/ChaosClothAsset/ConnectableValue.h b/Engine/Plugins/ChaosClothAssetEditor/Source/ChaosClothAssetDataflowNodes/Public/ChaosClothAsset/ConnectableValue.h
--- a/Engine/Plugins/ChaosClothAssetEditor/Source/ChaosClothAssetDataflowNodes/Public/ChaosClothAsset/ConnectableValue.h
+++ b/Engine/Plugins/ChaosClothAssetEditor/Source/ChaosClothAssetDataflowNodes/Public/ChaosClothAsset/ConnectableValue.h
@@ -13,7 +13,7 @@ struct FChaosClothAssetConnectableStringValue
 PRAGMA_ENABLE_DEPRECATION_WARNINGS
 
 	/** The value for this property. */
-	UPROPERTY(EditAnywhere, Category = "Value")
+	UPROPERTY(EditAnywhere, Category = "Value", Meta = (SkipInDisplayNameChain))
 	FString StringValue;
 
 	UE_DEPRECATED(5.5, "Override properties are no longer used.")
@@ -29,7 +29,7 @@ struct FChaosClothAssetConnectableIStringValue
 PRAGMA_ENABLE_DEPRECATION_WARNINGS
 
 	/** The value for this property. */
-	UPROPERTY(EditAnywhere, Category = "Value", Meta = (DataflowInput))
+	UPROPERTY(EditAnywhere, Category = "Value", Meta = (DataflowInput, SkipInDisplayNameChain))
 	FString StringValue;
 
 	/**
@@ -55,7 +55,7 @@ struct FChaosClothAssetConnectableOStringValue
 	GENERATED_BODY()
 
 	/** The value for this property. */
-	UPROPERTY(EditAnywhere, Category = "Value", Meta = (DataflowOutput))
+	UPROPERTY(EditAnywhere, Category = "Value", Meta = (DataflowOutput, SkipInDisplayNameChain))
 	FString StringValue;
 };
 
@@ -67,7 +67,7 @@ struct FChaosClothAssetConnectableIOStringValue
 PRAGMA_ENABLE_DEPRECATION_WARNINGS
 
 	/** The value for this property. */
-	UPROPERTY(EditAnywhere, Category = "Value", Meta = (DataflowInput, DataflowOutput, DataflowPassthrough = "StringValue"))
+	UPROPERTY(EditAnywhere, Category = "Value", Meta = (DataflowInput, DataflowOutput, DataflowPassthrough = "StringValue", SkipInDisplayNameChain))
 	FString StringValue;
 
 	UE_DEPRECATED(5.5, "Override properties are no longer used.")
diff --git a/Engine/Plugins/ChaosClothAssetEditor/Source/ChaosClothAssetDataflowNodes/Public/ChaosClothAsset/ProxyDeformerNode.h b/Engine/Plugins/ChaosClothAssetEditor/Source/ChaosClothAssetDataflowNodes/Public/ChaosClothAsset/ProxyDeformerNode.h
--- a/Engine/Plugins/ChaosClothAssetEditor/Source/ChaosClothAssetDataflowNodes/Public/ChaosClothAsset/ProxyDeformerNode.h
+++ b/Engine/Plugins/ChaosClothAssetEditor/Source/ChaosClothAssetDataflowNodes/Public/ChaosClothAsset/ProxyDeformerNode.h
@@ -35,7 +35,7 @@ public:
 	FManagedArrayCollection Collection;
 
 	/** Selection filter set used to restrict a selection of render vertices to a selection of simulation mesh triangles. Right click and do AddOptionPin to add more selection sets .*/
-	UPROPERTY(EditAnywhere, EditFixedSize, Category = "ProxyDeformer")
+	UPROPERTY(EditAnywhere, EditFixedSize, Category = "ProxyDeformer", Meta = (SkipInDisplayNameChain))
 	TArray<FChaosClothAssetSelectionFilterSet> SelectionFilterSets;
 
 	/** Whether using multiple simulation mesh triangles to influence the position of the deformed render vertex. */
diff --git a/Engine/Source/Runtime/Experimental/Dataflow/Core/Private/Dataflow/DataflowNode.cpp b/Engine/Source/Runtime/Experimental/Dataflow/Core/Private/Dataflow/DataflowNode.cpp
--- a/Engine/Source/Runtime/Experimental/Dataflow/Core/Private/Dataflow/DataflowNode.cpp
+++ b/Engine/Source/Runtime/Experimental/Dataflow/Core/Private/Dataflow/DataflowNode.cpp
@@ -701,30 +701,32 @@ FString FDataflowNode::StripContainerIndexFromPropertyFullName(const FString& In
 FText FDataflowNode::GetPropertyDisplayNameText(const TArray<const FProperty*>& PropertyChain, int32 ContainerIndex)
 {
 #if WITH_EDITORONLY_DATA  // GetDisplayNameText() is only available if WITH_EDITORONLY_DATA
+	static const FTextFormat TextFormat(NSLOCTEXT("DataflowNode", "PropertyDisplayNameTextConcatenator", "{0}.{1}"));
 	FText PropertyText;
+	bool bIsPropertyTextEmpty = true;
 	bool bFoundArrayProperty = false;
 	for (int32 Index = PropertyChain.Num() - 1; Index >= 0; --Index)
 	{
 		const FProperty* const Property = PropertyChain[Index];
-		FText PropertyDisplayName = Property->GetDisplayNameText();
+		if (!Property->HasMetaData(FName(TEXT("SkipInDisplayNameChain"))))
+		{
+			FText PropertyDisplayName = Property->GetDisplayNameText();
+			PropertyText = bIsPropertyTextEmpty ?
+				MoveTemp(PropertyDisplayName) :
+				FText::Format(TextFormat, PropertyText, MoveTemp(PropertyDisplayName));
+			bIsPropertyTextEmpty = false;
+		}
 		if (const FArrayProperty* const ArrayProperty = CastField<FArrayProperty>(Property))
 		{
-			if (ContainerIndex != INDEX_NONE)
-			{
-				check(!bFoundArrayProperty); // We only expect to find one array to substitute in.
-				bFoundArrayProperty = true;	
-				static const FTextFormat TextFormatContainer(NSLOCTEXT("DataflowNode", "PropertyDisplayNameTextContainer", "{0}[{1}]"));
-				PropertyDisplayName = FText::Format(TextFormatContainer, PropertyDisplayName, ContainerIndex);
-			}
-
-			// Skip ElemProperty. Otherwise you get names like "MyFloatArray[0].MyFloatArray" when you just want "MyFloatArray[0]"
-			--Index;
+			check(!bFoundArrayProperty); // We only expect to find one array to substitute in.
+			bFoundArrayProperty = (ContainerIndex != INDEX_NONE);
+			--Index;  // Skip ElemProperty. Otherwise you get names like "MyFloatArray[0].MyFloatArray" when you just want "MyFloatArray[0]"
 		}
-
-		static const FTextFormat TextFormat(NSLOCTEXT("DataflowNode", "PropertyDisplayNameTextConcatenator", "{0}.{1}"));
-		PropertyText = PropertyText.IsEmpty() ?
-			PropertyDisplayName :
-			FText::Format(TextFormat, PropertyText, PropertyDisplayName);
+	}
+	if (bFoundArrayProperty)
+	{
+		static const FTextFormat TextFormatContainer(NSLOCTEXT("DataflowNode", "PropertyDisplayNameTextContainer", "{0}[{1}]"));
+		PropertyText = FText::Format(TextFormatContainer, PropertyText, ContainerIndex);
 	}
 
 	return PropertyText;
