Fix height field raycast
Compiler seems to generate a infinite loop, when a function is inline.
It seems such error could be generated by inline function having a lot of parameters which is the case there.  Registers and parameter are around the same number there.  This is a blind fix, I couldn't reproduce it, but it cleans up the code and remove some tech debt.

diff --git a/Engine/Source/Runtime/Experimental/Chaos/Private/Chaos/HeightField.cpp b/Engine/Source/Runtime/Experimental/Chaos/Private/Chaos/HeightField.cpp
--- a/Engine/Source/Runtime/Experimental/Chaos/Private/Chaos/HeightField.cpp
+++ b/Engine/Source/Runtime/Experimental/Chaos/Private/Chaos/HeightField.cpp
@@ -1060,27 +1060,25 @@ namespace Chaos
 
 	// Raycast is navigating slowly one cell at the time, on the high resolution height field
 	// This navigation is the most precise and will check if we need to compute the raycast at the triangle level
-	FORCEINLINE bool FHeightField::WalkSlow(TVec2<int32>& CellIdx, FHeightfieldRaycastVisitor& Visitor, FReal CurrentLength, const VectorRegister4Float& CurrentLengthSimd,
-		const FVec2& ScaledMin, FReal ZMidPoint, const FVec3& Dir, const FVec3& InvDir, bool bParallel[3], const FVec2& ScaledDx2D, FVec3& NextStart, 
-		const FVec3& ScaleSign, const FVec3& ScaledDx, int32 IndexLowResX, int32 IndexLowResY) const
+	FORCEINLINE bool FHeightField::WalkSlow(FWalkingData& WalkingData, int32 IndexLowResX, int32 IndexLowResY) const
 	{
-		const int32 LowResInc = Visitor.GeomData->LowResInc;
+		const int32 LowResInc = WalkingData.Visitor.GeomData->LowResInc;
 		while (true)
 		{
-			if (FlatGrid.IsValid(CellIdx))
+			if (FlatGrid.IsValid(WalkingData.CellIdx))
 			{
 				PHYSICS_CSV_CUSTOM_VERY_EXPENSIVE(PhysicsCounters, NumRayHeightfieldCellVisited, 1, ECsvCustomStatOp::Accumulate);
 				// Test for the cell bounding box is done in the visitor at the same time as fetching the points for the triangles
 				// this avoid fetching the points twice ( here and in the visitor )
 
-				const int32 Payload = CellIdx[1] * (GeomData.NumCols - 1) + CellIdx[0];
+				const int32 Payload = WalkingData.CellIdx[1] * (GeomData.NumCols - 1) + WalkingData.CellIdx[0];
 				const int32 SubY = (Payload) / (GeomData.NumCols - 1);
 				const int32 FullIndex = Payload + SubY;
 				FAABBVectorized Bounds;
-				Visitor.GeomData->GetBoundsScaledSimd(FullIndex, Bounds);
-				if (Bounds.RaycastFast(Visitor.StartPointSimd, Visitor.InvDirSimd, Visitor.Parallel, CurrentLengthSimd))
+				WalkingData.Visitor.GeomData->GetBoundsScaledSimd(FullIndex, Bounds);
+				if (Bounds.RaycastFast(WalkingData.Visitor.StartPointSimd, WalkingData.Visitor.InvDirSimd, WalkingData.Visitor.Parallel, WalkingData.CurrentLengthSimd))
 				{
-					const bool bContinue = Visitor.VisitRaycast(CellIdx[1] * (GeomData.NumCols - 1) + CellIdx[0], CurrentLength);
+					const bool bContinue = WalkingData.Visitor.VisitRaycast(WalkingData.CellIdx[1] * (GeomData.NumCols - 1) + WalkingData.CellIdx[0], WalkingData.CurrentLength);
 					if (!bContinue)
 					{
 						return false;
@@ -1090,19 +1088,19 @@ namespace Chaos
 
 			//find next cell
 			//We want to know which plane we used to cross into next cell
-			const FVec2 ScaledCellCenter2D = ScaledMin + FVec2(static_cast<FReal>(CellIdx[0]) + 0.5f, static_cast<FReal>(CellIdx[1]) + 0.5f) * ScaledDx2D;
-			const FVec3 ScaledCellCenter(ScaledCellCenter2D[0], ScaledCellCenter2D[1], ZMidPoint);
+			const FVec2 ScaledCellCenter2D = WalkingData.ScaledMin + FVec2(static_cast<FReal>(WalkingData.CellIdx[0]) + 0.5f, static_cast<FReal>(WalkingData.CellIdx[1]) + 0.5f) * WalkingData.ScaledDx2D;
+			const FVec3 ScaledCellCenter(ScaledCellCenter2D[0], ScaledCellCenter2D[1], WalkingData.ZMidPoint);
 
 			FReal Times[3];
-			FReal BestTime = CurrentLength;
+			FReal BestTime = WalkingData.CurrentLength;
 			bool bTerminate = true;
 			for (int Axis = 0; Axis < 3; ++Axis)
 			{
-				if (!bParallel[Axis])
+				if (!WalkingData.bParallel[Axis])
 				{
-					const FReal CrossPoint = (Dir[Axis] * ScaleSign[Axis]) > 0 ? ScaledCellCenter[Axis] + ScaledDx[Axis] / 2 : ScaledCellCenter[Axis] - ScaledDx[Axis] / 2;
-					const FReal Distance = CrossPoint - NextStart[Axis];	//note: CellCenter already has /2, we probably want to use the corner instead
-					const FReal Time = Distance * InvDir[Axis];
+					const FReal CrossPoint = (WalkingData.Dir[Axis] * WalkingData.ScaleSign[Axis]) > 0 ? ScaledCellCenter[Axis] + WalkingData.ScaledDx[Axis] / 2 : ScaledCellCenter[Axis] - WalkingData.ScaledDx[Axis] / 2;
+					const FReal Distance = CrossPoint - WalkingData.NextStart[Axis];	//note: CellCenter already has /2, we probably want to use the corner instead
+					const FReal Time = Distance * WalkingData.InvDir[Axis];
 					Times[Axis] = Time;
 					if (Time < BestTime)
 					{
@@ -1121,26 +1119,26 @@ namespace Chaos
 				return false;
 			}
 
-			const TVec2<int32> PrevIdx = CellIdx;
+			const TVec2<int32> PrevIdx = WalkingData.CellIdx;
 
 			for (int Axis = 0; Axis < 2; ++Axis)
 			{
-				CellIdx[Axis] += (Times[Axis] <= BestTime) ? ((Dir[Axis] * ScaleSign[Axis]) > 0 ? 1 : -1) : 0;
-				if (CellIdx[Axis] < 0 || CellIdx[Axis] >= FlatGrid.Counts()[Axis])
+				WalkingData.CellIdx[Axis] += (Times[Axis] <= BestTime) ? ((WalkingData.Dir[Axis] * WalkingData.ScaleSign[Axis]) > 0 ? 1 : -1) : 0;
+				if (WalkingData.CellIdx[Axis] < 0 || WalkingData.CellIdx[Axis] >= FlatGrid.Counts()[Axis])
 				{
 					return false;
 				}
 			}
 
-			if (PrevIdx == CellIdx)
+			if (PrevIdx == WalkingData.CellIdx)
 			{
 				//crossed on z plane which means no longer in heightfield bounds
 				return false;
 			}
 
-			NextStart = NextStart + Dir * BestTime;
+			WalkingData.NextStart = WalkingData.NextStart + WalkingData.Dir * BestTime;
 
-			if (IndexLowResX != CellIdx[0] / LowResInc || IndexLowResY != CellIdx[1] / LowResInc)
+			if (IndexLowResX != WalkingData.CellIdx[0] / LowResInc || IndexLowResY != WalkingData.CellIdx[1] / LowResInc)
 			{
 				return true;
 			}
@@ -1152,51 +1150,49 @@ namespace Chaos
 	// Raycast is navigating fast on the high resolution height field
 	// This navigation is more precise than the low resolution walk. It would have more cache miss because it grabs several cells not necessarily contiguous in memory.
 	// It navigates only on the cell that need to be checked.
-	FORCEINLINE bool FHeightField::WalkFast(TVec2<int32>& CellIdx, FHeightfieldRaycastVisitor& Visitor, FReal CurrentLength, const VectorRegister4Float& CurrentLengthSimd,
-		const FVec2& ScaledMin, FReal ZMidPoint, const FVec3& Dir, const FVec3& InvDir, bool bParallel[3], const FVec2& ScaledDx2D, FVec3& NextStart,
-		const FVec3& ScaleSign, const FVec3& ScaledDx, const FVec2& Scale2D, const FVec3& DirScaled) const 
+	FORCEINLINE bool FHeightField::WalkFast(FWalkingData& WalkingData, const FVec2& Scale2D, const FVec3& DirScaled) const
 	{
 		const FReal BoundsMinZ = CachedBounds.Min().Z;
 		const FReal BoundsMaxZ = CachedBounds.Max().Z;
 
-		const int32 LowResInc = Visitor.GeomData->LowResInc;
-		const int32 IndexLowResX = CellIdx[0] / LowResInc;
-		const int32 IndexLowResY = CellIdx[1] / LowResInc;
+		const int32 LowResInc = WalkingData.Visitor.GeomData->LowResInc;
+		const int32 IndexLowResX = WalkingData.CellIdx[0] / LowResInc;
+		const int32 IndexLowResY = WalkingData.CellIdx[1] / LowResInc;
 
-		const TVec2<int32> NextStartInt = TVec2<int32>(static_cast<int32>(NextStart[0] / Scale2D[0]), static_cast<int32>(NextStart[1] / Scale2D[1]));
-		CellIdx = FlatGrid.Cell(NextStartInt);
+		const TVec2<int32> NextStartInt = TVec2<int32>(static_cast<int32>(WalkingData.NextStart[0] / Scale2D[0]), static_cast<int32>(WalkingData.NextStart[1] / Scale2D[1]));
+		WalkingData.CellIdx = FlatGrid.Cell(NextStartInt);
 
-		const FVec3 NextStartOri = NextStart;
+		const FVec3 NextStartOri = WalkingData.NextStart;
 		const FReal Length2D = FMath::Sqrt(DirScaled[0] * DirScaled[0] + DirScaled[1] * DirScaled[1]);
 		constexpr FReal FastInc2D = 2.0;  // Length in cell dimension
 		const FReal FastIncScaled = FastInc2D / Length2D;
 
-		const FVec3 InspectionStepVector = DirScaled * FastIncScaled * ScaledDx;
-		FReal NextStartZ = FMath::Clamp(NextStart.Z, BoundsMinZ, BoundsMaxZ); // Numerical errors can violate this condition, so force it here.
+		const FVec3 InspectionStepVector = DirScaled * FastIncScaled * WalkingData.ScaledDx;
+		FReal NextStartZ = FMath::Clamp(WalkingData.NextStart.Z, BoundsMinZ, BoundsMaxZ); // Numerical errors can violate this condition, so force it here.
 		const FReal FastIncScaled3D = InspectionStepVector.Length();
-		const FReal CellsToInspectZ = FastIncScaled3D * Dir.Z;
+		const FReal CellsToInspectZ = FastIncScaled3D * WalkingData.Dir.Z;
 		FReal DistanceProcessed = 0.0;
 
 		while (true)
 		{
-			if (!FlatGrid.IsValid(CellIdx))
+			if (!FlatGrid.IsValid(WalkingData.CellIdx))
 			{
 				return false;
 			}
 
-			const FVec3 NewNextStart = (InspectionStepVector + NextStart);
+			const FVec3 NewNextStart = (InspectionStepVector + WalkingData.NextStart);
 			const TVec2<int32> NewNextStartInt = TVec2<int32>(static_cast<int32>(NewNextStart[0] / Scale2D[0]), static_cast<int32>(NewNextStart[1] / Scale2D[1]));
-			const TVec2<int32> NextStartInt2 = TVec2<int32>(static_cast<int32>(NextStart[0] / Scale2D[0]), static_cast<int32>(NextStart[1] / Scale2D[1]));
+			const TVec2<int32> NextStartInt2 = TVec2<int32>(static_cast<int32>(WalkingData.NextStart[0] / Scale2D[0]), static_cast<int32>(WalkingData.NextStart[1] / Scale2D[1]));
 
 			const TVec2<int32> DiffInt = NewNextStartInt - NextStartInt2;
 			const TVec2<int32> AddedCellIdx = TVec2<int32>(static_cast<int32>(DiffInt[0]), static_cast<int32>(DiffInt[1]));
 
 			FAABBVectorized Bounds;
-			Visitor.GeomData->GetBoundsScaled(CellIdx, AddedCellIdx, Bounds);
-			if (Bounds.RaycastFast(Visitor.StartPointSimd, Visitor.InvDirSimd, Visitor.Parallel, CurrentLengthSimd))
+			WalkingData.Visitor.GeomData->GetBoundsScaled(WalkingData.CellIdx, AddedCellIdx, Bounds);
+			if (Bounds.RaycastFast(WalkingData.Visitor.StartPointSimd, WalkingData.Visitor.InvDirSimd, WalkingData.Visitor.Parallel, WalkingData.CurrentLengthSimd))
 			{
-				NextStart = NextStartOri + Dir * DistanceProcessed;
-				bool bContinue = WalkSlow(CellIdx, Visitor, CurrentLength, CurrentLengthSimd, ScaledMin, ZMidPoint, Dir, InvDir, bParallel, ScaledDx2D, NextStart, ScaleSign, ScaledDx, IndexLowResX, IndexLowResY);
+				WalkingData.NextStart = NextStartOri + WalkingData.Dir * DistanceProcessed;
+				bool bContinue = WalkSlow(WalkingData, IndexLowResX, IndexLowResY);
 				if (bContinue)
 				{
 					return true;
@@ -1204,22 +1200,22 @@ namespace Chaos
 				return false;
 			}
 
-			NextStart = NewNextStart;
+			WalkingData.NextStart = NewNextStart;
 
-			const TVec2<int32> PrevIdx = CellIdx;
-			CellIdx = FlatGrid.Cell(NewNextStartInt);
-			if (PrevIdx == CellIdx)
+			const TVec2<int32> PrevIdx = WalkingData.CellIdx;
+			WalkingData.CellIdx = FlatGrid.Cell(NewNextStartInt);
+			if (PrevIdx == WalkingData.CellIdx)
 			{
 				return false;
 			}
 
-			if (DistanceProcessed > CurrentLength || NextStartZ < BoundsMinZ || NextStartZ > BoundsMaxZ)
+			if (DistanceProcessed > WalkingData.CurrentLength || NextStartZ < BoundsMinZ || NextStartZ > BoundsMaxZ)
 			{
 				return false;
 			}
 			DistanceProcessed += FastIncScaled3D;
 			NextStartZ += CellsToInspectZ;
-			if (IndexLowResX != CellIdx[0] / LowResInc || IndexLowResY != CellIdx[1] / LowResInc)
+			if (IndexLowResX != WalkingData.CellIdx[0] / LowResInc || IndexLowResY != WalkingData.CellIdx[1] / LowResInc)
 			{
 				// Walk on low resolution height field
 				return true;
@@ -1231,24 +1227,21 @@ namespace Chaos
 	// Raycast is navigating a low resolution height field
 	// This navigation is less precise can have lot of false positive intersection but navigate fast with few cache miss.
 	// It navigates on a squared cell number.
-	FORCEINLINE bool FHeightField::WalkOnLowRes(TVec2<int32>& CellIdx, FHeightfieldRaycastVisitor& Visitor, FReal CurrentLength, const VectorRegister4Float& CurrentLengthSimd,
-		const FVec2& ScaledMin, FReal ZMidPoint, const FVec3& Dir, const FVec3& InvDir, bool bParallel[3], const FVec2& ScaledDx2D, FVec3& NextStart,
-		const FVec3& ScaleSign, const FVec3& ScaledDx, const FVec2& Scale2D, const FVec3& DirScaled) const
+	FORCEINLINE bool FHeightField::WalkOnLowRes(FWalkingData& WalkingData, const FVec2& Scale2D, const FVec3& DirScaled) const
 	{
-		const int32 LowResInc = Visitor.GeomData->LowResInc;
+		const int32 LowResInc = WalkingData.Visitor.GeomData->LowResInc;
 		while (true)
 		{
-			if (!FlatGrid.IsValid(CellIdx))
+			if (!FlatGrid.IsValid(WalkingData.CellIdx))
 			{
 				return false;
 			}
 
 			FAABBVectorized Bounds;
-			Visitor.GeomData->GetLowResBoundsScaled(CellIdx, Bounds);
-			if (Bounds.RaycastFast(Visitor.StartPointSimd, Visitor.InvDirSimd, Visitor.Parallel, CurrentLengthSimd))
+			WalkingData.Visitor.GeomData->GetLowResBoundsScaled(WalkingData.CellIdx, Bounds);
+			if (Bounds.RaycastFast(WalkingData.Visitor.StartPointSimd, WalkingData.Visitor.InvDirSimd, WalkingData.Visitor.Parallel, WalkingData.CurrentLengthSimd))
 			{
-				bool bContinue = WalkFast(CellIdx, Visitor, CurrentLength, CurrentLengthSimd, ScaledMin, ZMidPoint, Dir, InvDir, bParallel, ScaledDx2D, NextStart,
-					ScaleSign, ScaledDx, Scale2D, DirScaled);
+				bool bContinue = WalkFast(WalkingData, Scale2D, DirScaled);
 				if (bContinue)
 				{
 					continue; 
@@ -1256,19 +1249,19 @@ namespace Chaos
 				return false;
 			}
 
-			const FVec2 ScaledCellCenter2D = ScaledMin + FVec2(static_cast<FReal>(CellIdx[0] / LowResInc) + 0.5f, static_cast<FReal>(CellIdx[1] / LowResInc) + 0.5f) * (ScaledDx2D * (FReal)LowResInc);
-			const FVec3 ScaledCellCenter(ScaledCellCenter2D[0], ScaledCellCenter2D[1], ZMidPoint);
+			const FVec2 ScaledCellCenter2D = WalkingData.ScaledMin + FVec2(static_cast<FReal>(WalkingData.CellIdx[0] / LowResInc) + 0.5f, static_cast<FReal>(WalkingData.CellIdx[1] / LowResInc) + 0.5f) * (WalkingData.ScaledDx2D * (FReal)LowResInc);
+			const FVec3 ScaledCellCenter(ScaledCellCenter2D[0], ScaledCellCenter2D[1], WalkingData.ZMidPoint);
 
 			FReal Times[3];
-			FReal BestTime = CurrentLength;
+			FReal BestTime = WalkingData.CurrentLength;
 			bool bTerminate = true;
 			for (int Axis = 0; Axis < 3; ++Axis)
 			{
-				if (!bParallel[Axis])
+				if (!WalkingData.bParallel[Axis])
 				{
-					const FReal CrossPoint = (Dir[Axis] * ScaleSign[Axis]) > 0 ? ScaledCellCenter[Axis] + ScaledDx[Axis] * (FReal)LowResInc / 2 : ScaledCellCenter[Axis] - ScaledDx[Axis] * (FReal)LowResInc / 2;
-					const FReal Distance = CrossPoint - NextStart[Axis];	//note: CellCenter already has /2, we probably want to use the corner instead
-					const FReal Time = Distance * InvDir[Axis];
+					const FReal CrossPoint = (WalkingData.Dir[Axis] * WalkingData.ScaleSign[Axis]) > 0 ? ScaledCellCenter[Axis] + WalkingData.ScaledDx[Axis] * (FReal)LowResInc / 2 : ScaledCellCenter[Axis] - WalkingData.ScaledDx[Axis] * (FReal)LowResInc / 2;
+					const FReal Distance = CrossPoint - WalkingData.NextStart[Axis];	//note: CellCenter already has /2, we probably want to use the corner instead
+					const FReal Time = Distance * WalkingData.InvDir[Axis];
 					Times[Axis] = Time;
 					if (Time < BestTime)
 					{
@@ -1287,21 +1280,21 @@ namespace Chaos
 				return false;
 			}
 
-			const TVec2<int32> PrevIdx = CellIdx;
+			const TVec2<int32> PrevIdx = WalkingData.CellIdx;
 
 			for (int Axis = 0; Axis < 2; ++Axis)
 			{
-				CellIdx[Axis] += (Times[Axis] <= BestTime) ? ((Dir[Axis] * ScaleSign[Axis]) > 0 ? LowResInc : -LowResInc) : 0;
+				WalkingData.CellIdx[Axis] += (Times[Axis] <= BestTime) ? ((WalkingData.Dir[Axis] * WalkingData.ScaleSign[Axis]) > 0 ? LowResInc : -LowResInc) : 0;
 				// Make sure CellIdx doesn't go beyond the grid
-				CellIdx[Axis] = FMath::Min<int32>(CellIdx[Axis], FlatGrid.Counts()[Axis] - 1);
-				CellIdx[Axis] = FMath::Max<int32>(CellIdx[Axis], 0);
+				WalkingData.CellIdx[Axis] = FMath::Min<int32>(WalkingData.CellIdx[Axis], FlatGrid.Counts()[Axis] - 1);
+				WalkingData.CellIdx[Axis] = FMath::Max<int32>(WalkingData.CellIdx[Axis], 0);
 			}
 
-			if (PrevIdx == CellIdx)
+			if (PrevIdx == WalkingData.CellIdx)
 			{
 				return false;
 			}
-			NextStart = NextStart + Dir * BestTime;
+			WalkingData.NextStart = WalkingData.NextStart + WalkingData.Dir * BestTime;
 		}
 		return false;
 	}
@@ -1314,49 +1307,48 @@ namespace Chaos
 			return false;
 		}
 
-		FReal CurrentLength = Length;
+		FWalkingData WalkingData(Visitor);
+		WalkingData.Dir = Dir;
+		WalkingData.CurrentLength = Length;
 
-		// Data for fast box cast
-		bool bParallel[3];
-		FVec3 InvDir;
-
-		FReal InvCurrentLength = 1 / CurrentLength;
+		FReal InvCurrentLength = 1 / WalkingData.CurrentLength;
 		for(int Axis = 0; Axis < 3; ++Axis)
 		{
-			bParallel[Axis] = FMath::IsNearlyZero(Dir[Axis], (FReal)1.e-8);
-			InvDir[Axis] = bParallel[Axis] ? 0 : 1 / Dir[Axis];
+			WalkingData.bParallel[Axis] = FMath::IsNearlyZero(Dir[Axis], (FReal)1.e-8);
+			WalkingData.InvDir[Axis] = WalkingData.bParallel[Axis] ? 0 : 1 / Dir[Axis];
 		}
 		
 		FReal RayEntryTime;
 		FReal RayExitTime;
-		if(CachedBounds.RaycastFast(StartPoint, Dir, InvDir, bParallel, Length, InvCurrentLength, RayEntryTime, RayExitTime))
+		if(CachedBounds.RaycastFast(StartPoint, Dir, WalkingData.InvDir, WalkingData.bParallel, Length, InvCurrentLength, RayEntryTime, RayExitTime))
 		{
-			CurrentLength = RayExitTime + 1e-2; // to account for precision errors 
-			FVec3 NextStart = StartPoint + (Dir * RayEntryTime);
+			WalkingData.CurrentLength = RayExitTime + 1e-2; // to account for precision errors 
+			WalkingData.NextStart = StartPoint + (Dir * RayEntryTime);
 
 			const FVec2 Scale2D(GeomData.Scale[0], GeomData.Scale[1]);
-			TVec2<int32> CellIdx = FlatGrid.Cell(TVec2<int32>(static_cast<int32>(NextStart[0] / Scale2D[0]), static_cast<int32>(NextStart[1] / Scale2D[1])));
+			WalkingData.CellIdx = FlatGrid.Cell(TVec2<int32>(static_cast<int32>(WalkingData.NextStart[0] / Scale2D[0]), static_cast<int32>(WalkingData.NextStart[1] / Scale2D[1])));
 			const FReal ZDx = CachedBounds.Extents()[2];
-			const FReal ZMidPoint = CachedBounds.Min()[2] + ZDx * 0.5f;
-			const FVec3 ScaledDx(FlatGrid.Dx()[0] * Scale2D[0],FlatGrid.Dx()[1] * Scale2D[1],ZDx);
-			const FVec2 ScaledDx2D(ScaledDx[0],ScaledDx[1]);
-			const FVec2 ScaledMin = FlatGrid.MinCorner() * Scale2D;
-			const FVec3 ScaleSign = GeomData.Scale.GetSignVector();
+			WalkingData.ZMidPoint = CachedBounds.Min()[2] + ZDx * 0.5f;
+			WalkingData.ScaledDx = FVec3(FlatGrid.Dx()[0] * Scale2D[0], FlatGrid.Dx()[1] * Scale2D[1], ZDx);
+			WalkingData.ScaledDx2D = FVec2(WalkingData.ScaledDx[0], WalkingData.ScaledDx[1]);
+			WalkingData.ScaledMin = FlatGrid.MinCorner() * Scale2D;
+			WalkingData.ScaleSign = GeomData.Scale.GetSignVector();
 
-			const FVec3 DirScaled = Dir / ScaledDx;
+			FVec3 DirScaled = Dir / WalkingData.ScaledDx;
 			const FReal SumPlaneAxis = FMath::Abs(DirScaled[0]) + FMath::Abs(DirScaled[1]);
 			const bool bCanWalk = SumPlaneAxis > UE_SMALL_NUMBER;
-			const VectorRegister4Float CurrentLengthSimd = VectorSet1(static_cast<FRealSingle>(CurrentLength));
+			WalkingData.CurrentLengthSimd = VectorSet1(static_cast<FRealSingle>(WalkingData.CurrentLength));
+
 			if (bCanWalk)
 			{
-				WalkOnLowRes(CellIdx, Visitor, CurrentLength, CurrentLengthSimd, ScaledMin, ZMidPoint, Dir, InvDir, bParallel, ScaledDx2D, NextStart, ScaleSign, ScaledDx, Scale2D, DirScaled);
+				WalkOnLowRes(WalkingData, Scale2D, DirScaled);
 			}
 			else 
 			{
 				const int32 LowResInc = Visitor.GeomData->LowResInc;
-				const int32 IndexLowResX = CellIdx[0] / LowResInc;
-				const int32 IndexLowResY = CellIdx[1] / LowResInc;
-				const bool bContinue = WalkSlow(CellIdx, Visitor, CurrentLength, CurrentLengthSimd, ScaledMin, ZMidPoint, Dir, InvDir, bParallel, ScaledDx2D, NextStart, ScaleSign, ScaledDx, IndexLowResX, IndexLowResY);
+				const int32 IndexLowResX = WalkingData.CellIdx[0] / LowResInc;
+				const int32 IndexLowResY = WalkingData.CellIdx[1] / LowResInc;
+				const bool bContinue = WalkSlow(WalkingData, IndexLowResX, IndexLowResY);
 				if (!bContinue)
 				{
 					return false;
diff --git a/Engine/Source/Runtime/Experimental/Chaos/Public/Chaos/HeightField.h b/Engine/Source/Runtime/Experimental/Chaos/Public/Chaos/HeightField.h
--- a/Engine/Source/Runtime/Experimental/Chaos/Public/Chaos/HeightField.h
+++ b/Engine/Source/Runtime/Experimental/Chaos/Public/Chaos/HeightField.h
@@ -1076,17 +1076,27 @@ namespace Chaos
 		template<typename SQVisitor>
 		bool GridSweep(const FVec3& StartPoint, const FVec3& Dir, const FReal Length, const FVec3 InHalfExtents, SQVisitor& Visitor) const;
 		
-		CHAOS_API bool WalkSlow(TVec2<int32>& CellIdx, FHeightfieldRaycastVisitor& Visitor, FReal CurrentLength, const VectorRegister4Float& CurrentLengthSimd,
-			const FVec2& ScaledMin, FReal ZMidPoint, const FVec3& Dir, const FVec3& InvDir, bool bParallel[3], const FVec2& ScaledDx2D, FVec3& NextStart,
-			const FVec3& ScaleSign, const FVec3& ScaledDx, int32 IndexLowResX, int32 IndexLowResY) const;
-		
-		CHAOS_API bool WalkFast(TVec2<int32>& CellIdx, FHeightfieldRaycastVisitor& Visitor, FReal CurrentLength, const VectorRegister4Float& CurrentLengthSimd,
-			const FVec2& ScaledMin, FReal ZMidPoint, const FVec3& Dir, const FVec3& InvDir, bool bParallel[3], const FVec2& ScaledDx2D, FVec3& NextStart,
-			const FVec3& ScaleSign, const FVec3& ScaledDx, const FVec2& Scale2D, const FVec3& DirScaled) const;
+		struct FWalkingData {
+
+			explicit FWalkingData(FHeightfieldRaycastVisitor& VisitorIn) : Visitor(VisitorIn) {}
+			VectorRegister4Float CurrentLengthSimd;
+			FVec2 ScaledMin;
+			FReal CurrentLength;
+			FReal ZMidPoint;
+			FVec3 Dir;
+			FVec3 InvDir; 
+			FVec3 NextStart;
+			FVec3 ScaleSign;
+			FVec3 ScaledDx;
+			FVec2 ScaledDx2D;
+			TVec2<int32> CellIdx;
+			bool bParallel[3];
+			FHeightfieldRaycastVisitor& Visitor;
+		};
 
-		CHAOS_API bool WalkOnLowRes(TVec2<int32>& CellIdx, FHeightfieldRaycastVisitor& Visitor, FReal CurrentLength, const VectorRegister4Float& CurrentLengthSimd,
-			const FVec2& ScaledMin, FReal ZMidPoint, const FVec3& Dir, const FVec3& InvDir, bool bParallel[3], const FVec2& ScaledDx2D, FVec3& NextStart,
-			const FVec3& ScaleSign, const FVec3& ScaledDx, const FVec2& Scale2D, const FVec3& DirScaled) const;
+		CHAOS_API bool WalkSlow(FWalkingData& WalkingData, int32 IndexLowResX, int32 IndexLowResY) const;
+		CHAOS_API bool WalkFast(FWalkingData& WalkingData, const FVec2& Scale2D, const FVec3& DirScaled) const;
+		CHAOS_API bool WalkOnLowRes(FWalkingData& WalkingData, const FVec2& Scale2D, const FVec3& DirScaled) const;
 
 		CHAOS_API bool GridCast(const FVec3& StartPoint, const FVec3& Dir, const FReal Length, FHeightfieldRaycastVisitor& Visitor) const;
 		CHAOS_API bool GetGridIntersections(FBounds2D InFlatBounds, TArray<TVec2<int32>>& OutInterssctions) const;
