OnDemandIOStore install progress callback

diff --git a/Engine/Source/Runtime/Experimental/IoStore/OnDemand/Private/IoStoreOnDemand.cpp b/Engine/Source/Runtime/Experimental/IoStore/OnDemand/Private/IoStoreOnDemand.cpp
--- a/Engine/Source/Runtime/Experimental/IoStore/OnDemand/Private/IoStoreOnDemand.cpp
+++ b/Engine/Source/Runtime/Experimental/IoStore/OnDemand/Private/IoStoreOnDemand.cpp
@@ -1072,7 +1072,7 @@ void FIoStoreOnDemandModule::Mount(FOnDemandMountArgs&& Args, FOnDemandMountComp
 	IoStore->Mount(MoveTemp(Args), MoveTemp(OnCompleted));
 }
 
-void FIoStoreOnDemandModule::Install(FOnDemandInstallArgs&& Args, FOnDemandInstallCompleted&& OnCompleted)
+void FIoStoreOnDemandModule::Install(FOnDemandInstallArgs&& Args, FOnDemandInstallCompleted&& OnCompleted, FOnDemandInstallProgressed&& OnProgress /*= nullptr*/)
 {
 	if (IoStore.IsValid() == false)
 	{
@@ -1085,7 +1085,7 @@ void FIoStoreOnDemandModule::Install(FOnDemandInstallArgs&& Args, FOnDemandInsta
 		}
 	}
 
-	IoStore->Install(MoveTemp(Args), MoveTemp(OnCompleted));
+	IoStore->Install(MoveTemp(Args), MoveTemp(OnCompleted), MoveTemp(OnProgress));
 }
 
 FIoStatus FIoStoreOnDemandModule::Unmount(FStringView MountId)
diff --git a/Engine/Source/Runtime/Experimental/IoStore/OnDemand/Private/OnDemandInstallCache.cpp b/Engine/Source/Runtime/Experimental/IoStore/OnDemand/Private/OnDemandInstallCache.cpp
--- a/Engine/Source/Runtime/Experimental/IoStore/OnDemand/Private/OnDemandInstallCache.cpp
+++ b/Engine/Source/Runtime/Experimental/IoStore/OnDemand/Private/OnDemandInstallCache.cpp
@@ -1178,7 +1178,7 @@ FIoStatus FOnDemandInstallCache::PutChunk(FIoBuffer&& Chunk, const FIoHash& Chun
 	}
 
 	PendingChunks->Append(MoveTemp(Chunk), ChunkHash);
-	return FIoStatus::Ok;;
+	return FIoStatus::Ok;
 }
 
 FIoStatus FOnDemandInstallCache::Purge(TMap<FIoHash, uint64>&& ChunksToIntall) 
@@ -1308,7 +1308,7 @@ FIoStatus FOnDemandInstallCache::Flush()
 	}
 
 	Cas.Compact();
-	return FIoStatus::Ok;;
+	return FIoStatus::Ok;
 }
 
 FIoStatus FOnDemandInstallCache::FlushPendingChunks(FPendingChunks& Chunks)
@@ -1396,7 +1396,7 @@ FIoStatus FOnDemandInstallCache::FlushPendingChunks(FPendingChunks& Chunks)
 		}
 	}
 
-	return FIoStatus::Ok;;
+	return FIoStatus::Ok;
 }
 
 void FOnDemandInstallCache::CompleteRequest(FChunkRequest& ChunkRequest)
diff --git a/Engine/Source/Runtime/Experimental/IoStore/OnDemand/Private/OnDemandIoStore.cpp b/Engine/Source/Runtime/Experimental/IoStore/OnDemand/Private/OnDemandIoStore.cpp
--- a/Engine/Source/Runtime/Experimental/IoStore/OnDemand/Private/OnDemandIoStore.cpp
+++ b/Engine/Source/Runtime/Experimental/IoStore/OnDemand/Private/OnDemandIoStore.cpp
@@ -16,6 +16,7 @@
 #include "Containers/StringConv.h"
 #include "Containers/Ticker.h"
 #include "HAL/FileManager.h"
+#include "HAL/IConsoleManager.h"
 #include "HAL/PlatformFile.h"
 #include "HAL/PlatformFileManager.h"
 #include "HAL/PlatformTime.h"
@@ -202,13 +203,20 @@ static TIoStatusOr<FSharedContainerHeader> DeserializeContainerHeader(const FOnD
 using FPackageStoreEntryMap		= TMap<FPackageId, const FFilePackageStoreEntry*>;
 using FSoftPackageReferenceMap	= TMap<FPackageId, const FFilePackageStoreEntrySoftReferences*>;
 
+enum class EChunkCacheState : uint8
+{
+	None = 0,
+	Cached,
+	NotCached
+};
+
 struct FContainerInstallData
 {
-	FPackageStoreEntryMap		PackageStoreEntries;
-	FSoftPackageReferenceMap	SoftPackageReferences;
-	TSet<FPackageId>			PackageIds;
-	TSet<uint32>				ResolvedChunks;
-	uint64						TotalSize = 0;
+	FPackageStoreEntryMap			PackageStoreEntries;
+	FSoftPackageReferenceMap		SoftPackageReferences;
+	TSet<FPackageId>				PackageIds;
+	TMap<uint32, EChunkCacheState>	ResolvedChunks;
+	uint64							TotalSize = 0;
 };
 
 using FInstallData = TMap<FSharedOnDemandContainer, FContainerInstallData>;
@@ -346,7 +354,7 @@ static FIoStatus BuildInstallData(
 				continue;
 			}
 
-			Data.ResolvedChunks.Add(EntryIndex);
+			Data.ResolvedChunks.Add(EntryIndex, EChunkCacheState::None);
 			Data.TotalSize += Container.ChunkEntries[EntryIndex].EncodedSize;
 
 			const EIoChunkType AdditionalPackageChunkTypes[] =
@@ -361,7 +369,7 @@ static FIoStatus BuildInstallData(
 				const FIoChunkId ChunkId = CreateIoChunkId(PackageId.Value(), 0, ChunkType);
 				if (EntryIndex = Container.FindChunkEntryIndex(ChunkId); EntryIndex != INDEX_NONE)
 				{
-					Data.ResolvedChunks.Add(EntryIndex);
+					Data.ResolvedChunks.Add(EntryIndex, EChunkCacheState::None);
 					Data.TotalSize += Container.ChunkEntries[EntryIndex].EncodedSize;
 				}
 			}
@@ -378,7 +386,7 @@ static FIoStatus BuildInstallData(
 				case EIoChunkType::ShaderCodeLibrary:
 				case EIoChunkType::ShaderCode:
 				{
-					Data.ResolvedChunks.Add(EntryIndex);
+					Data.ResolvedChunks.Add(EntryIndex, EChunkCacheState::None);
 					Data.TotalSize += Container.ChunkEntries[EntryIndex].EncodedSize;
 				}
 				default:
@@ -587,11 +595,12 @@ void FOnDemandIoStore::Mount(FOnDemandMountArgs&& Args, FOnDemandMountCompleted&
 	TryEnterTickLoop();
 }
 
-void FOnDemandIoStore::Install(FOnDemandInstallArgs&& Args, FOnDemandInstallCompleted&& OnCompleted)
+void FOnDemandIoStore::Install(FOnDemandInstallArgs&& Args, FOnDemandInstallCompleted&& OnCompleted, FOnDemandInstallProgressed&& OnProgress /*= nullptr*/)
 {
 	FSharedInstallRequest InstallRequest = MakeShared<FInstallRequest>();
 	InstallRequest->Args		= MoveTemp(Args);
 	InstallRequest->OnCompleted = MoveTemp(OnCompleted);
+	InstallRequest->OnProgressed = MoveTemp(OnProgress);
 
 	{
 		UE::TUniqueLock Lock(MountRequestMutex);
@@ -937,21 +946,14 @@ bool FOnDemandIoStore::Tick()
 	// Tick install request(s)
 	for (FSharedInstallRequest& Request : LocalInstallRequests)
 	{
-		FIoStatus Status = TickInstallRequest(*Request);
+		FOnDemandInstallResult Result = TickInstallRequest(*Request);
 
 		{
 			UE::TUniqueLock Lock(MountRequestMutex);
 			InstallRequests.Remove(Request);
 		}
 
-		CompleteInstallRequest(*Request,
-			FOnDemandInstallResult
-			{
-				.Status = MoveTemp(Status),
-				.DurationInSeconds = Request->DurationInSeconds,
-				.TotalContentSize = Request->TotalContentSize,
-				.TotalInstallSize = Request->TotalInstallSize
-			});
+		CompleteInstallRequest(*Request, MoveTemp(Result));
 
 		bTicked = true;
 	}
@@ -1135,6 +1137,11 @@ FIoStatus FOnDemandIoStore::TickMountRequest(FMountRequest& MountRequest)
 
 void FOnDemandIoStore::CompleteMountRequest(FMountRequest& Request, FOnDemandMountResult&& MountResult)
 {
+	if (!Request.OnCompleted)
+	{
+		return;
+	}
+
 	if (EnumHasAnyFlags(Request.Args.Options, EOnDemandMountOptions::CallbackOnGameThread))
 	{
 		ExecuteOnGameThread(
@@ -1151,30 +1158,37 @@ void FOnDemandIoStore::CompleteMountRequest(FMountRequest& Request, FOnDemandMou
 	}
 }
 
-FIoStatus FOnDemandIoStore::TickInstallRequest(FInstallRequest& InstallRequest)
+FOnDemandInstallResult FOnDemandIoStore::TickInstallRequest(const FInstallRequest& InstallRequest)
 {
 	UE_LOG(LogIoStoreOnDemand, Log, TEXT("Ticking install request, ContentHandle='%s'"),
 		*LexToString(InstallRequest.Args.ContentHandle));
 
+	FOnDemandInstallResult OutResult;
+	FIoStatus& Status					= OutResult.Status;
+	FOnDemandInstallProgress& Progress	= OutResult.Progress;
+
 	const double StartTime = FPlatformTime::Seconds();
 	ON_SCOPE_EXIT
 	{
-		InstallRequest.DurationInSeconds = FPlatformTime::Seconds() - StartTime;
+		OutResult.DurationInSeconds = FPlatformTime::Seconds() - StartTime;
 	};
 
 	if (InstallRequest.Args.ContentHandle.IsValid() == false)
 	{
-		return FIoStatusBuilder(EIoErrorCode::InvalidParameter) << TEXT("Invalid content handle");
+		Status = FIoStatusBuilder(EIoErrorCode::InvalidParameter) << TEXT("Invalid content handle");
+		return OutResult;
 	}
 
 	if (InstallCache.IsValid() == false || PackageStoreBackend.IsValid() == false)
 	{
 		if (FForkProcessHelper::IsForkRequested() && !FForkProcessHelper::IsForkedChildProcess())
 		{
-			return FIoStatusBuilder(EIoErrorCode::PendingFork) << TEXT("Install cache waiting for fork");
+			Status = FIoStatusBuilder(EIoErrorCode::PendingFork) << TEXT("Install cache waiting for fork");
+			return OutResult;
 		}
 
-		return FIoStatusBuilder(EIoErrorCode::InvalidCode) << TEXT("Install cache not configured");
+		Status = FIoStatusBuilder(EIoErrorCode::InvalidCode) << TEXT("Install cache not configured");
+		return OutResult;
 	}
 
 	//TODO: Implement cancellation
@@ -1186,14 +1200,14 @@ FIoStatus FOnDemandIoStore::TickInstallRequest(FInstallRequest& InstallRequest)
 	TSet<FSharedOnDemandContainer> ContainersForInstallation;
 	TSet<FPackageId> PackageIdsToInstall;
 
-	if (FIoStatus Status = GetContainersAndPackagesForInstall(
+	if (Status = GetContainersAndPackagesForInstall(
 		InstallRequest.Args.MountId,
 		InstallRequest.Args.TagSets,
 		InstallRequest.Args.PackageIds,
 		ContainersForInstallation,
 		PackageIdsToInstall); !Status.IsOk())
 	{
-		return Status;
+		return OutResult;
 	}
 
 	// Its OK for PackageIdsToInstall to be empty at this point. Any chunks not referenced by a package must still be installed.
@@ -1202,10 +1216,10 @@ FIoStatus FOnDemandIoStore::TickInstallRequest(FInstallRequest& InstallRequest)
 	Private::FInstallData InstallData;
 	TSet<FPackageId> Missing;
 
-	FIoStatus Status = BuildInstallData(ContainersForInstallation, PackageIdsToInstall, InstallData, Missing);
+	Status = BuildInstallData(ContainersForInstallation, PackageIdsToInstall, InstallData, Missing);
 	if (Status.IsOk() == false)
 	{
-		return Status;
+		return OutResult;
 	}
 
 	// Check the other I/O backends for missing package chunks
@@ -1215,8 +1229,9 @@ FIoStatus FOnDemandIoStore::TickInstallRequest(FInstallRequest& InstallRequest)
 		if (FIoDispatcher::Get().DoesChunkExist(ChunkId) == false)
 		{
 			UE_LOG(LogIoStoreOnDemand, Error, TEXT("Missing package chunk '%s'"), *LexToString(ChunkId));
-			return FIoStatusBuilder(EIoErrorCode::UnknownChunkID) << 
+			Status = FIoStatusBuilder(EIoErrorCode::UnknownChunkID) <<
 				TEXT("Missing package chunk '") << LexToString(ChunkId) << TEXT("'");
+			return OutResult;
 		}
 	}
 
@@ -1228,8 +1243,10 @@ FIoStatus FOnDemandIoStore::TickInstallRequest(FInstallRequest& InstallRequest)
 			FOnDemandContainer&						Container = *Kv.Key;
 			const Private::FContainerInstallData&	Data = Kv.Value;
 
-			for (int32 EntryIndex : Data.ResolvedChunks)
+			for (const TPair<uint32, Private::EChunkCacheState>& ChunkPair : Data.ResolvedChunks)
 			{
+				const int32 EntryIndex = ChunkPair.Key;
+
 				const FOnDemandChunkEntry& Entry = Container.ChunkEntries[EntryIndex];
 				ChunksToInstall.Add(Entry.Hash, Entry.EncodedSize);
 			}
@@ -1237,7 +1254,7 @@ FIoStatus FOnDemandIoStore::TickInstallRequest(FInstallRequest& InstallRequest)
 
 		if (Status = InstallCache->Purge(MoveTemp(ChunksToInstall)); Status.IsOk() == false)
 		{
-			return Status;
+			return OutResult;
 		}
 	}
 
@@ -1253,42 +1270,84 @@ FIoStatus FOnDemandIoStore::TickInstallRequest(FInstallRequest& InstallRequest)
 	TUniquePtr<FHttpClient> HttpClient = FHttpClient::Create(MoveTemp(HttpConfig));
 	if (HttpClient.IsValid() == false)
 	{
-		return FIoStatusBuilder(EIoErrorCode::InvalidCode) << TEXT("Failed to initialize HTTP client");
+		Status = FIoStatusBuilder(EIoErrorCode::InvalidCode) << TEXT("Failed to initialize HTTP client");
+		return OutResult;
 	}
 
 	uint64 TotalChunkCount		= 0;
 	uint64 TotalBytes			= 0;
-	uint64 DownloadedChunkCount	= 0;
-	uint64 DownloadedBytes		= 0;
-
-	for (const auto& Kv : InstallData)
+	uint64 ToDownloadChunkCount	= 0;
+	uint64 ToDownloadBytes		= 0;
+	
+	for (Private::FInstallData::ElementType& Kv : InstallData)
 	{
-		FOnDemandContainer& Container				= *Kv.Key;
-		const Private::FContainerInstallData& Data	= Kv.Value;
+		const FOnDemandContainer& Container		= *Kv.Key;
+		Private::FContainerInstallData& Data	= Kv.Value;
 
-		for (int32 EntryIndex : Data.ResolvedChunks)
+		for (TPair<uint32, Private::EChunkCacheState>& ChunkPair : Data.ResolvedChunks)
 		{
-			const FIoChunkId& ChunkId				= Container.ChunkIds[EntryIndex];
-			const FOnDemandChunkEntry& ChunkEntry	= Container.ChunkEntries[EntryIndex];
+			const int32 EntryIndex						= ChunkPair.Key;
+			Private::EChunkCacheState& ChunkCacheState	= ChunkPair.Value;
+
+			const FIoChunkId& ChunkId = Container.ChunkIds[EntryIndex];
+			const FOnDemandChunkEntry& ChunkEntry = Container.ChunkEntries[EntryIndex];
 
 			++TotalChunkCount;
 			TotalBytes += ChunkEntry.EncodedSize;
 
 			if (InstallCache->IsChunkCached(ChunkEntry.Hash))
 			{
+				ChunkCacheState = Private::EChunkCacheState::Cached;
 				continue;
 			}
 
-			++DownloadedChunkCount;
-			DownloadedBytes += ChunkEntry.EncodedSize;
+			ChunkCacheState = Private::EChunkCacheState::NotCached;
+
+			++ToDownloadChunkCount;
+			ToDownloadBytes += ChunkEntry.EncodedSize;
+		}
+	}
+
+	Progress.TotalContentSize = TotalBytes;
+	Progress.TotalInstallSize = ToDownloadBytes;
+
+	ProgressInstallRequest(InstallRequest, Progress);
+
+	for (const Private::FInstallData::ElementType& Kv : InstallData)
+	{
+		const FOnDemandContainer& Container			= *Kv.Key;
+		const Private::FContainerInstallData& Data	= Kv.Value;
+
+		for (const TPair<uint32, Private::EChunkCacheState>& ChunkPair : Data.ResolvedChunks)
+		{
+			const int32 EntryIndex							= ChunkPair.Key;
+			const Private::EChunkCacheState ChunkCacheState	= ChunkPair.Value;
+
+			const FIoChunkId& ChunkId						= Container.ChunkIds[EntryIndex];
+			const FOnDemandChunkEntry& ChunkEntry			= Container.ChunkEntries[EntryIndex];
+
+			if (ChunkCacheState == Private::EChunkCacheState::Cached)
+			{
+				continue;
+			}
 
 			ConcurrentRequests++;
 			HttpClient->Get(
 				Private::GetChunkUrl(FStringView(), Container, ChunkEntry, ChunkUrl).ToView(),
-				[this, &Status, ChunkId, ChunkEntry, &ConcurrentRequests]
+				[this, &InstallRequest, &OutResult, ChunkId, ChunkEntry, &ConcurrentRequests]
 				(TIoStatusOr<FIoBuffer> ChunkStatus, uint64 DurationMs) mutable
 				{
 					ConcurrentRequests--;
+
+					FIoStatus& Status					= OutResult.Status;
+					FOnDemandInstallProgress& Progress	= OutResult.Progress;
+
+					if (!Status.IsOk())
+					{
+						// a different request already failed, don't stomp the status if this one succeeded
+						return;
+					}
+					
 					if (!ChunkStatus.IsOk())
 					{
 						Status = ChunkStatus.Status();
@@ -1303,6 +1362,14 @@ FIoStatus FOnDemandIoStore::TickInstallRequest(FInstallRequest& InstallRequest)
 						return;
 					}
 					Status = InstallCache->PutChunk(MoveTemp(Chunk), ChunkHash);
+					if (!Status.IsOk())
+					{
+						return;
+					}
+
+					// TODO: Is this good enough progress or should it be fine grained bytes from the HttpClient?
+					Progress.CurrentInstallSize += ChunkEntry.EncodedSize;
+					ProgressInstallRequest(InstallRequest, Progress);
 				});
 
 			while (ConcurrentRequests >= MaxConcurrentRequests)
@@ -1312,20 +1379,25 @@ FIoStatus FOnDemandIoStore::TickInstallRequest(FInstallRequest& InstallRequest)
 
 			if (Status.IsOk() == false)
 			{
-				return Status;
+				return OutResult;
 			}
 		}
 	}
 
 	while (HttpClient->Tick())
 		;
+
+	if (Status.IsOk() == false)
+	{
+		return OutResult;
+	}
 	
 	if (Status = InstallCache->Flush(); Status.IsOk() == false)
 	{
-		return FIoStatusBuilder(EIoErrorCode::WriteError) << TEXT("Failed to flush downloaded content to cache");
+		return OutResult;
 	}
 
-	TSharedPtr<FOnDemandInternalContentHandle, ESPMode::ThreadSafe>& ContentHandle = InstallRequest.Args.ContentHandle.Handle;
+	const TSharedPtr<FOnDemandInternalContentHandle, ESPMode::ThreadSafe>& ContentHandle = InstallRequest.Args.ContentHandle.Handle;
 	if (ContentHandle->IoStore.IsValid() == false)
 	{
 		// First time this content handle is used
@@ -1347,24 +1419,27 @@ FIoStatus FOnDemandIoStore::TickInstallRequest(FInstallRequest& InstallRequest)
 		{
 			UE::TUniqueLock Lock(ContainerMutex);
 			FOnDemandChunkEntryReferences& References = Container.FindOrAddChunkEntryReferences(*ContentHandle);
-			for (int32 EntryIndex : Data.ResolvedChunks)
+			for (const TPair<uint32, Private::EChunkCacheState>& ChunkPair : Data.ResolvedChunks)
 			{
+				const int32 EntryIndex = ChunkPair.Key;
 				References.Indices[EntryIndex] = true;
 			}
 		}
 	}
 
-	InstallRequest.TotalContentSize = TotalBytes;
-	InstallRequest.TotalInstallSize = DownloadedBytes;
-
 	UE_LOG(LogIoStoreOnDemand, Log, TEXT("Downloaded %llu (%.2lf MiB) of total %llu (%.2lf MiB) requested chunk(s)'"),
-		DownloadedChunkCount, double(DownloadedBytes) / 1024.0 / 1024.0, TotalChunkCount, double(TotalBytes) / 1024.0 / 1024.0);
+		ToDownloadChunkCount, double(ToDownloadBytes) / 1024.0 / 1024.0, TotalChunkCount, double(TotalBytes) / 1024.0 / 1024.0);
 
-	return EIoErrorCode::Ok;
+	return OutResult;
 }
 
 void FOnDemandIoStore::CompleteInstallRequest(FInstallRequest& Request, FOnDemandInstallResult&& InstallResult)
 {
+	if (!Request.OnCompleted)
+	{
+		return;
+	}
+
 	if (EnumHasAnyFlags(Request.Args.Options, EOnDemandInstallOptions::CallbackOnGameThread))
 	{
 		ExecuteOnGameThread(
@@ -1381,6 +1456,31 @@ void FOnDemandIoStore::CompleteInstallRequest(FInstallRequest& Request, FOnDeman
 	}
 }
 
+void FOnDemandIoStore::ProgressInstallRequest(const FInstallRequest& Request, const FOnDemandInstallProgress& Progress)
+{
+	if (!Request.OnProgressed)
+	{
+		return;
+	}
+
+	if (EnumHasAnyFlags(Request.Args.Options, EOnDemandInstallOptions::CallbackOnGameThread))
+	{
+		ExecuteOnGameThread(
+			UE_SOURCE_LOCATION,
+			[OnProgressed = Request.OnProgressed, Progress]()
+			{
+				OnProgressed(Progress);
+			});
+	}
+	else
+	{
+		if (Request.OnProgressed)
+		{
+			Request.OnProgressed(Progress);
+		}
+	}
+}
+
 void FOnDemandIoStore::OnEncryptionKeyAdded(const FGuid& Id, const FAES::FAESKey& Key)
 {
 	TRACE_CPUPROFILER_EVENT_SCOPE(FOnDemandIoStore::OnEncryptionKeyAdded);
diff --git a/Engine/Source/Runtime/Experimental/IoStore/OnDemand/Private/OnDemandIoStore.h b/Engine/Source/Runtime/Experimental/IoStore/OnDemand/Private/OnDemandIoStore.h
--- a/Engine/Source/Runtime/Experimental/IoStore/OnDemand/Private/OnDemandIoStore.h
+++ b/Engine/Source/Runtime/Experimental/IoStore/OnDemand/Private/OnDemandIoStore.h
@@ -263,9 +263,7 @@ class FOnDemandIoStore
 	{
 		FOnDemandInstallArgs		Args;
 		FOnDemandInstallCompleted	OnCompleted;
-		double						DurationInSeconds = 0.0;
-		uint64						TotalContentSize = 0;
-		uint64						TotalInstallSize = 0;
+		FOnDemandInstallProgressed	OnProgressed;
 	};
 
 	using FSharedInstallRequest	= TSharedPtr<FInstallRequest>;
@@ -280,7 +278,7 @@ public:
 
 	FIoStatus				Initialize();
 	void					Mount(FOnDemandMountArgs&& Args, FOnDemandMountCompleted&& OnCompleted);
-	void					Install(FOnDemandInstallArgs&& Args, FOnDemandInstallCompleted&& OnCompleted);
+	void					Install(FOnDemandInstallArgs&& Args, FOnDemandInstallCompleted&& OnCompleted, FOnDemandInstallProgressed&& OnProgress = nullptr);
 	FIoStatus				Unmount(FStringView MountId);
 	TIoStatusOr<uint64>		GetInstallSize(const FOnDemandGetInstallSizeArgs& Args) const;
 	FIoStatus				GetInstallSizesByMountId(const FOnDemandGetInstallSizeArgs& Args, TMap<FString, uint64>& OutSizesByMountId) const;
@@ -308,8 +306,9 @@ private:
 	bool					Tick();
 	FIoStatus				TickMountRequest(FMountRequest& MountRequest);
 	void					CompleteMountRequest(FMountRequest& MountRequest, FOnDemandMountResult&& MountResult);
-	FIoStatus				TickInstallRequest(FInstallRequest& InstallRequest);
+	FOnDemandInstallResult	TickInstallRequest(const FInstallRequest& InstallRequest);
 	void					CompleteInstallRequest(FInstallRequest& InstallRequest, FOnDemandInstallResult&& InstallResult);
+	void					ProgressInstallRequest(const FInstallRequest& InstallRequest, const FOnDemandInstallProgress& Progress);
 	void					OnEncryptionKeyAdded(const FGuid& Id, const FAES::FAESKey& Key);
 	static void				CreateContainersFromToc(
 								FStringView MountId,
diff --git a/Engine/Source/Runtime/Experimental/IoStore/OnDemand/Public/IO/IoStoreOnDemand.h b/Engine/Source/Runtime/Experimental/IoStore/OnDemand/Public/IO/IoStoreOnDemand.h
--- a/Engine/Source/Runtime/Experimental/IoStore/OnDemand/Public/IO/IoStoreOnDemand.h
+++ b/Engine/Source/Runtime/Experimental/IoStore/OnDemand/Public/IO/IoStoreOnDemand.h
@@ -338,6 +338,20 @@ struct FOnDemandInstallArgs
 	EOnDemandInstallOptions Options = EOnDemandInstallOptions::None;
 };
 
+/** Holds information about progress for an install request. */
+struct FOnDemandInstallProgress
+{
+	/** The total size of the requested content. */
+	uint64 TotalContentSize = 0;
+	/** The total size to be installed/downloaded (<= TotalContentSize). */
+	uint64 TotalInstallSize = 0;
+	/** The size currently installed/downloaded (<= TotalInstallSize). */
+	uint64 CurrentInstallSize = 0;
+};
+
+/** Install Progress callback. */
+using FOnDemandInstallProgressed = TFunction<void(FOnDemandInstallProgress)>;
+
 /** Holds information about an install request. */
 struct FOnDemandInstallResult
 {
@@ -345,10 +359,8 @@ struct FOnDemandInstallResult
 	FIoStatus Status;
 	/** Duration in seconds. */
 	double DurationInSeconds = 0.0;
-	/** The total size of the requested content. */
-	uint64 TotalContentSize = 0;
-	/** The total installed/downloaded size. */
-	uint64 TotalInstallSize = 0;
+	/** Final progress for the install request. */
+	FOnDemandInstallProgress Progress;
 };
 
 /** Install completion callback. */
@@ -390,7 +402,7 @@ public:
 	UE_API void ReportAnalytics(TArray<FAnalyticsEventAttribute>& OutAnalyticsArray) const;
 
 	UE_API void Mount(FOnDemandMountArgs&& Args, FOnDemandMountCompleted&& OnCompleted);
-	UE_API void Install(FOnDemandInstallArgs&& Args, FOnDemandInstallCompleted&& OnCompleted);
+	UE_API void Install(FOnDemandInstallArgs&& Args, FOnDemandInstallCompleted&& OnCompleted, FOnDemandInstallProgressed&& OnProgress = nullptr);
 	UE_API FIoStatus Unmount(FStringView MountId);
 
 	UE_API TIoStatusOr<uint64> GetInstallSize(const FOnDemandGetInstallSizeArgs& Args) const;
