Dataflow - Added DataflowNode virtual function to reconnect missing inputs on serialization.

#tests Reloaded a Dataflow asset with a redirected connection.
#rnx

diff --git a/Engine/Source/Runtime/Experimental/Dataflow/Core/Private/Dataflow/DataflowGraph.cpp b/Engine/Source/Runtime/Experimental/Dataflow/Core/Private/Dataflow/DataflowGraph.cpp
--- a/Engine/Source/Runtime/Experimental/Dataflow/Core/Private/Dataflow/DataflowGraph.cpp
+++ b/Engine/Source/Runtime/Experimental/Dataflow/Core/Private/Dataflow/DataflowGraph.cpp
@@ -342,7 +342,15 @@ namespace Dataflow
 								Ar << bIsHidden;
 							}
 
-							if (FDataflowOutput* Output = Node->FindOutput(ArName))
+							FDataflowOutput* Output = Node->FindOutput(ArName);
+							if (!Output)
+							{
+								// Find out if the output has recently been redirected
+								Output = Node->RedirectSerializedOutput(ArName);
+								UE_CLOG(Output, LogChaos, Display, TEXT("Output (%s) has been redirected to output (%s) in Dataflow node (%s).")
+									, *ArName.ToString(), *Output->GetName().ToString(), *ArNodeName.ToString());
+							}
+							if (Output)
 							{
 								if (bIsAnyType)
 								{
@@ -384,7 +392,15 @@ namespace Dataflow
 								Ar << bIsHidden;
 							}
 
-							if (FDataflowInput* Input = Node->FindInput(ArName))
+							FDataflowInput* Input = Node->FindInput(ArName);
+							if (!Input)
+							{
+								// Find out if the input has recently been redirected
+								Input = Node->RedirectSerializedInput(ArName);
+								UE_CLOG(Input, LogChaos, Display, TEXT("Input (%s) has been redirected to input (%s) in Dataflow node (%s).")
+									, *ArName.ToString(), *Input->GetName().ToString(), *ArNodeName.ToString());
+							}
+							if (Input)
 							{
 								if (bIsAnyType)
 								{
diff --git a/Engine/Source/Runtime/Experimental/Dataflow/Core/Public/Dataflow/DataflowGraph.h b/Engine/Source/Runtime/Experimental/Dataflow/Core/Public/Dataflow/DataflowGraph.h
--- a/Engine/Source/Runtime/Experimental/Dataflow/Core/Public/Dataflow/DataflowGraph.h
+++ b/Engine/Source/Runtime/Experimental/Dataflow/Core/Public/Dataflow/DataflowGraph.h
@@ -167,6 +167,7 @@ namespace Dataflow
 
 		DATAFLOWCORE_API void RemoveNode(TSharedPtr<FDataflowNode> Node);
 
+		const TArray<FLink>& GetConnections() const { return Connections; }
 		DATAFLOWCORE_API void ClearConnections(FDataflowConnection* ConnectionBase);
 		DATAFLOWCORE_API void ClearConnections(FDataflowInput* Input);
 		DATAFLOWCORE_API void ClearConnections(FDataflowOutput* Output);
diff --git a/Engine/Source/Runtime/Experimental/Dataflow/Core/Public/Dataflow/DataflowNode.h b/Engine/Source/Runtime/Experimental/Dataflow/Core/Public/Dataflow/DataflowNode.h
--- a/Engine/Source/Runtime/Experimental/Dataflow/Core/Public/Dataflow/DataflowNode.h
+++ b/Engine/Source/Runtime/Experimental/Dataflow/Core/Public/Dataflow/DataflowNode.h
@@ -203,6 +203,11 @@ PRAGMA_ENABLE_DEPRECATION_WARNINGS
 	/** Override this method to provide custom serialization for this node. */
 	virtual void Serialize(FArchive& Ar) {}
 
+	/** Override this method to provide custom reconnections when a node inputs has been deprecated and removed. */
+	virtual FDataflowInput* RedirectSerializedInput(const FName& MissingInputName) { return nullptr; }
+	/** Override this method to provide custom reconnections when a node outputs has been deprecated and removed. */
+	virtual FDataflowOutput* RedirectSerializedOutput(const FName& MissingOutputName) { return nullptr; }
+
 	/** Called by editor toolkits when the node is selected, or already selected and invalidated. */
 	virtual void OnSelected(Dataflow::FContext& Context) {}
 	/** Called by editor toolkits when the node is deselected. */
diff --git a/Engine/Source/Runtime/Experimental/Dataflow/Engine/Private/Dataflow/DataflowObject.cpp b/Engine/Source/Runtime/Experimental/Dataflow/Engine/Private/Dataflow/DataflowObject.cpp
--- a/Engine/Source/Runtime/Experimental/Dataflow/Engine/Private/Dataflow/DataflowObject.cpp
+++ b/Engine/Source/Runtime/Experimental/Dataflow/Engine/Private/Dataflow/DataflowObject.cpp
@@ -5,7 +5,9 @@
 #include "Dataflow/DataflowEdNode.h"
 #include "Dataflow/DataflowNodeParameters.h"
 #include "Dataflow/DataflowObjectInterface.h"
-
+#if WITH_EDITOR
+#include "EdGraph/EdGraphPin.h"
+#endif
 #include UE_INLINE_GENERATED_CPP_BY_NAME(DataflowObject)
 
 #define LOCTEXT_NAMESPACE "UDataflow"
@@ -99,6 +101,37 @@ void UDataflow::PostLoad()
 			EdNode->SetEnabledState(ENodeEnabledState::Disabled);
 		}
 	}
+
+	// Resync connections (nodes might have redirected connections
+	for (const Dataflow::FLink& Link : Dataflow->GetConnections())
+	{
+		TSharedPtr<const FDataflowNode> OutputNode = Dataflow->FindBaseNode(Link.OutputNode);
+		TSharedPtr<const FDataflowNode> InputNode = Dataflow->FindBaseNode(Link.InputNode);
+		if (ensure(OutputNode && InputNode))
+		{
+			const FDataflowOutput* const Output = OutputNode->FindOutput(Link.Output);
+			const FDataflowInput* const Input = InputNode->FindInput(Link.Input);
+			if (ensure(Output && Input))
+			{
+				TObjectPtr<UDataflowEdNode> OutputEdNode = FindEdNodeByDataflowNodeGuid(Link.OutputNode);
+				TObjectPtr<UDataflowEdNode> InputEdNode = FindEdNodeByDataflowNodeGuid(Link.InputNode);
+
+				if (ensure(OutputEdNode && InputEdNode))
+				{
+					UEdGraphPin* const OutputPin = OutputEdNode->FindPin(Output->GetName(), EEdGraphPinDirection::EGPD_Output);
+					UEdGraphPin* const InputPin = InputEdNode->FindPin(Input->GetName(), EEdGraphPinDirection::EGPD_Input);
+
+					if (ensure(OutputPin && InputPin))
+					{
+						if (OutputPin->LinkedTo.Find(InputPin) == INDEX_NONE)
+						{
+							OutputPin->MakeLinkTo(InputPin);
+						}
+					}
+				}
+			}
+		}
+	}
 #endif
 
 	LastModifiedRenderTarget = Dataflow::FTimestamp::Current();
