Added support for cancelling I/O store install request using a cancellation token

#rnx
[RN]

diff --git a/Engine/Source/Runtime/Experimental/IoStore/OnDemand/Private/IoStoreOnDemand.cpp b/Engine/Source/Runtime/Experimental/IoStore/OnDemand/Private/IoStoreOnDemand.cpp
--- a/Engine/Source/Runtime/Experimental/IoStore/OnDemand/Private/IoStoreOnDemand.cpp
+++ b/Engine/Source/Runtime/Experimental/IoStore/OnDemand/Private/IoStoreOnDemand.cpp
@@ -1072,7 +1072,11 @@ void FIoStoreOnDemandModule::Mount(FOnDemandMountArgs&& Args, FOnDemandMountComp
 	IoStore->Mount(MoveTemp(Args), MoveTemp(OnCompleted));
 }
 
-void FIoStoreOnDemandModule::Install(FOnDemandInstallArgs&& Args, FOnDemandInstallCompleted&& OnCompleted, FOnDemandInstallProgressed&& OnProgress /*= nullptr*/)
+void FIoStoreOnDemandModule::Install(
+	FOnDemandInstallArgs&& Args,
+	FOnDemandInstallCompleted&& OnCompleted,
+	FOnDemandInstallProgressed&& OnProgress /*= nullptr*/,
+	const FOnDemandCancellationToken* CancellationToken)
 {
 	if (IoStore.IsValid() == false)
 	{
@@ -1085,7 +1089,7 @@ void FIoStoreOnDemandModule::Install(FOnDemandInstallArgs&& Args, FOnDemandInsta
 		}
 	}
 
-	IoStore->Install(MoveTemp(Args), MoveTemp(OnCompleted), MoveTemp(OnProgress));
+	IoStore->Install(MoveTemp(Args), MoveTemp(OnCompleted), MoveTemp(OnProgress), CancellationToken);
 }
 
 FIoStatus FIoStoreOnDemandModule::Unmount(FStringView MountId)
diff --git a/Engine/Source/Runtime/Experimental/IoStore/OnDemand/Private/OnDemandIoStore.cpp b/Engine/Source/Runtime/Experimental/IoStore/OnDemand/Private/OnDemandIoStore.cpp
--- a/Engine/Source/Runtime/Experimental/IoStore/OnDemand/Private/OnDemandIoStore.cpp
+++ b/Engine/Source/Runtime/Experimental/IoStore/OnDemand/Private/OnDemandIoStore.cpp
@@ -595,12 +595,17 @@ void FOnDemandIoStore::Mount(FOnDemandMountArgs&& Args, FOnDemandMountCompleted&
 	TryEnterTickLoop();
 }
 
-void FOnDemandIoStore::Install(FOnDemandInstallArgs&& Args, FOnDemandInstallCompleted&& OnCompleted, FOnDemandInstallProgressed&& OnProgress /*= nullptr*/)
+void FOnDemandIoStore::Install(
+	FOnDemandInstallArgs&& Args,
+	FOnDemandInstallCompleted&& OnCompleted,
+	FOnDemandInstallProgressed&& OnProgress /*= nullptr*/,
+	const FOnDemandCancellationToken* CancellationToken)
 {
 	FSharedInstallRequest InstallRequest = MakeShared<FInstallRequest>();
 	InstallRequest->Args		= MoveTemp(Args);
 	InstallRequest->OnCompleted = MoveTemp(OnCompleted);
 	InstallRequest->OnProgressed = MoveTemp(OnProgress);
+	InstallRequest->CancellationToken = CancellationToken;
 
 	{
 		UE::TUniqueLock Lock(MountRequestMutex);
@@ -1173,6 +1178,19 @@ FOnDemandInstallResult FOnDemandIoStore::TickInstallRequest(const FInstallReques
 		OutResult.DurationInSeconds = FPlatformTime::Seconds() - StartTime;
 	};
 
+	auto CheckAndSetCancelled = [&InstallRequest, &Status]() -> bool
+	{
+		if (InstallRequest.CancellationToken && InstallRequest.CancellationToken->IsCanceled())
+		{
+			UE_LOG(LogIoStoreOnDemand, Log, TEXT("Cancelling install request, ContentHandle='%s'"),
+				*LexToString(InstallRequest.Args.ContentHandle));
+			Status = FIoStatus(EIoErrorCode::Cancelled);
+			return true;
+		}
+
+		return false;
+	};
+
 	if (InstallRequest.Args.ContentHandle.IsValid() == false)
 	{
 		Status = FIoStatusBuilder(EIoErrorCode::InvalidParameter) << TEXT("Invalid content handle");
@@ -1191,8 +1209,6 @@ FOnDemandInstallResult FOnDemandIoStore::TickInstallRequest(const FInstallReques
 		return OutResult;
 	}
 
-	//TODO: Implement cancellation
-
 	TAnsiStringBuilder<512> ChunkUrl;
 	FStringView Host, TocRelUrl;
 	Private::SplitHostUrl(InstallRequest.Args.Url, Host, TocRelUrl);
@@ -1235,6 +1251,11 @@ FOnDemandInstallResult FOnDemandIoStore::TickInstallRequest(const FInstallReques
 		}
 	}
 
+	if (CheckAndSetCancelled())
+	{
+		return OutResult;
+	}
+
 	// Purge
 	{
 		TMap<FIoHash, uint64> ChunksToInstall;
@@ -1258,6 +1279,11 @@ FOnDemandInstallResult FOnDemandIoStore::TickInstallRequest(const FInstallReques
 		}
 	}
 
+	if (CheckAndSetCancelled())
+	{
+		return OutResult;
+	}
+
 	// Download all chunks
 	const int32 MaxConcurrentRequests = 16;
 	int32		ConcurrentRequests = 0;
@@ -1377,7 +1403,7 @@ FOnDemandInstallResult FOnDemandIoStore::TickInstallRequest(const FInstallReques
 				HttpClient->Tick();
 			}
 
-			if (Status.IsOk() == false)
+			if ((Status.IsOk() == false) || CheckAndSetCancelled())
 			{
 				return OutResult;
 			}
@@ -1397,6 +1423,11 @@ FOnDemandInstallResult FOnDemandIoStore::TickInstallRequest(const FInstallReques
 		return OutResult;
 	}
 
+	if (CheckAndSetCancelled())
+	{
+		return OutResult;
+	}
+
 	const TSharedPtr<FOnDemandInternalContentHandle, ESPMode::ThreadSafe>& ContentHandle = InstallRequest.Args.ContentHandle.Handle;
 	if (ContentHandle->IoStore.IsValid() == false)
 	{
diff --git a/Engine/Source/Runtime/Experimental/IoStore/OnDemand/Private/OnDemandIoStore.h b/Engine/Source/Runtime/Experimental/IoStore/OnDemand/Private/OnDemandIoStore.h
--- a/Engine/Source/Runtime/Experimental/IoStore/OnDemand/Private/OnDemandIoStore.h
+++ b/Engine/Source/Runtime/Experimental/IoStore/OnDemand/Private/OnDemandIoStore.h
@@ -252,18 +252,19 @@ class FOnDemandIoStore
 {
 	struct FMountRequest
 	{
-		FOnDemandMountArgs			Args;
-		FOnDemandMountCompleted		OnCompleted;
-		double						DurationInSeconds = 0.0;
+		FOnDemandMountArgs		Args;
+		FOnDemandMountCompleted	OnCompleted;
+		double					DurationInSeconds = 0.0;
 	};
 
 	using FSharedMountRequest	= TSharedPtr<FMountRequest>;
 
 	struct FInstallRequest
 	{
-		FOnDemandInstallArgs		Args;
-		FOnDemandInstallCompleted	OnCompleted;
-		FOnDemandInstallProgressed	OnProgressed;
+		FOnDemandInstallArgs				Args;
+		FOnDemandInstallCompleted			OnCompleted;
+		FOnDemandInstallProgressed			OnProgressed;
+		const FOnDemandCancellationToken*	CancellationToken = nullptr;
 	};
 
 	using FSharedInstallRequest	= TSharedPtr<FInstallRequest>;
@@ -278,7 +279,11 @@ public:
 
 	FIoStatus				Initialize();
 	void					Mount(FOnDemandMountArgs&& Args, FOnDemandMountCompleted&& OnCompleted);
-	void					Install(FOnDemandInstallArgs&& Args, FOnDemandInstallCompleted&& OnCompleted, FOnDemandInstallProgressed&& OnProgress = nullptr);
+	void					Install(
+								FOnDemandInstallArgs&& Args,
+								FOnDemandInstallCompleted&& OnCompleted,
+								FOnDemandInstallProgressed&& OnProgress = nullptr,
+								const FOnDemandCancellationToken* CancellationToken = nullptr);
 	FIoStatus				Unmount(FStringView MountId);
 	TIoStatusOr<uint64>		GetInstallSize(const FOnDemandGetInstallSizeArgs& Args) const;
 	FIoStatus				GetInstallSizesByMountId(const FOnDemandGetInstallSizeArgs& Args, TMap<FString, uint64>& OutSizesByMountId) const;
diff --git a/Engine/Source/Runtime/Experimental/IoStore/OnDemand/Public/IO/IoStoreOnDemand.h b/Engine/Source/Runtime/Experimental/IoStore/OnDemand/Public/IO/IoStoreOnDemand.h
--- a/Engine/Source/Runtime/Experimental/IoStore/OnDemand/Public/IO/IoStoreOnDemand.h
+++ b/Engine/Source/Runtime/Experimental/IoStore/OnDemand/Public/IO/IoStoreOnDemand.h
@@ -16,6 +16,8 @@
 #include "Templates/SharedPointer.h"
 #include "UObject/NameTypes.h"
 
+#include <atomic>
+
 #if (IS_PROGRAM || WITH_EDITOR)
 #include "Containers/Map.h"
 #include "Misc/AES.h"
@@ -377,6 +379,29 @@ struct FOnDemandGetInstallSizeArgs
 	TArray<FPackageId> PackageIds;
 };
 
+/** Token used for signalling an operation to be cancelled. */
+class FOnDemandCancellationToken
+{
+public:
+	UE_NONCOPYABLE(FOnDemandCancellationToken);
+	FOnDemandCancellationToken() = default;
+
+	/** Signal the operation to be cancelled. */
+	void Cancel()
+	{
+		bCanceled = true;
+	}
+
+	/** Returns whether an operation should be cancelled. */
+	bool IsCanceled() const
+	{
+		return bCanceled;
+	}
+
+private:
+	std::atomic<bool> bCanceled{ false };
+};
+
 class FIoStoreOnDemandModule
 	: public IModuleInterface
 {
@@ -402,9 +427,14 @@ public:
 	UE_API void ReportAnalytics(TArray<FAnalyticsEventAttribute>& OutAnalyticsArray) const;
 
 	UE_API void Mount(FOnDemandMountArgs&& Args, FOnDemandMountCompleted&& OnCompleted);
-	UE_API void Install(FOnDemandInstallArgs&& Args, FOnDemandInstallCompleted&& OnCompleted, FOnDemandInstallProgressed&& OnProgress = nullptr);
 	UE_API FIoStatus Unmount(FStringView MountId);
 
+	UE_API void Install(
+		FOnDemandInstallArgs&& Args,
+		FOnDemandInstallCompleted&& OnCompleted,
+		FOnDemandInstallProgressed&& OnProgress = nullptr,
+		const FOnDemandCancellationToken* CancellationToken = nullptr);
+
 	UE_API TIoStatusOr<uint64> GetInstallSize(const FOnDemandGetInstallSizeArgs& Args) const;
 
 	UE_API FIoStatus GetInstallSizesByMountId(const FOnDemandGetInstallSizeArgs& Args, TMap<FString, uint64>& OutSizesByMountId) const;
