Experimental routine to stringify a list of 'root objects'

#rnx

diff --git a/Engine/Source/Runtime/Experimental/JsonObjectGraph/JsonObjectGraph.Build.cs b/Engine/Source/Runtime/Experimental/JsonObjectGraph/JsonObjectGraph.Build.cs
new file mode 100644
--- /dev/null
+++ b/Engine/Source/Runtime/Experimental/JsonObjectGraph/JsonObjectGraph.Build.cs
@@ -0,0 +1,27 @@
+// Copyright Epic Games, Inc. All Rights Reserved.
+
+using UnrealBuildTool;
+
+public class JsonObjectGraph : ModuleRules
+{
+	public JsonObjectGraph(ReadOnlyTargetRules Target) : base(Target)
+	{
+		PCHUsage = ModuleRules.PCHUsageMode.UseExplicitOrSharedPCHs;
+		
+		PublicDependencyModuleNames.AddRange(
+			new string[]
+			{
+				"Core",
+			}
+		);
+			
+		
+		PrivateDependencyModuleNames.AddRange(
+			new string[]
+			{
+				"CoreUObject",
+				"Json"
+			}
+		);
+	}
+}
diff --git a/Engine/Source/Runtime/Experimental/JsonObjectGraph/Private/JsonObjectGraph.cpp b/Engine/Source/Runtime/Experimental/JsonObjectGraph/Private/JsonObjectGraph.cpp
new file mode 100644
--- /dev/null
+++ b/Engine/Source/Runtime/Experimental/JsonObjectGraph/Private/JsonObjectGraph.cpp
@@ -0,0 +1,29 @@
+// Copyright Epic Games, Inc. All Rights Reserved.
+
+#include "Modules/ModuleManager.h"
+#include "Modules/ModuleInterface.h"
+
+#define LOCTEXT_NAMESPACE "JsonOjbectGraphModule"
+
+namespace UE::Private
+{
+
+class JsonObjectGraph : public IModuleInterface
+{
+	virtual void StartupModule() override;
+	virtual void ShutdownModule() override;
+};
+
+void JsonObjectGraph::StartupModule()
+{
+}
+
+void JsonObjectGraph::ShutdownModule()
+{
+}
+
+}
+
+#undef LOCTEXT_NAMESPACE
+
+IMPLEMENT_MODULE(UE::Private::JsonObjectGraph, JsonObjectGraph)
diff --git a/Engine/Source/Runtime/Experimental/JsonObjectGraph/Private/JsonObjectGraph/Stringify.cpp b/Engine/Source/Runtime/Experimental/JsonObjectGraph/Private/JsonObjectGraph/Stringify.cpp
new file mode 100644
--- /dev/null
+++ b/Engine/Source/Runtime/Experimental/JsonObjectGraph/Private/JsonObjectGraph/Stringify.cpp
@@ -0,0 +1,10 @@
+// Copyright Epic Games, Inc. All Rights Reserved.
+
+#include "JsonObjectGraph/Stringify.h"
+#include "JsonStringifyImpl.h"
+
+FUtf8String UE::JsonObjectGraph::Stringify(TConstArrayView<const UObject*> RootObjects, const FJsonStringifyOptions& Options)
+{
+	UE::Private::FJsonStringifyImpl Impl(RootObjects, Options);
+	return Impl.ToJson();
+}
diff --git a/Engine/Source/Runtime/Experimental/JsonObjectGraph/Private/JsonObjectGraphConventions.h b/Engine/Source/Runtime/Experimental/JsonObjectGraph/Private/JsonObjectGraphConventions.h
new file mode 100644
--- /dev/null
+++ b/Engine/Source/Runtime/Experimental/JsonObjectGraph/Private/JsonObjectGraphConventions.h
@@ -0,0 +1,37 @@
+// Copyright Epic Games, Inc. All Rights Reserved.
+
+#pragma once
+
+// Private conventions of the JsonObjectGraph object representation - these 
+// would be shared between 'from' and 'to' implementations - if you update
+// these make sure to update the tchar version as well if present. TCHAR
+// versions exist because the underlying JSON writer/reader I'm using 
+// speaks it 'natively', long term I'm hoping to remove all TCHAR.
+
+// 'Package' level conventions, these are used at the root level:
+#define UE_JSON_ROOT_OBJECTS_KEY_TCHAR TEXT("__RootObjects")
+
+#define UE_JSON_CUSTOM_VERSIONS_KEY_TCHAR TEXT("__CustomVersions")
+
+// UObject instance conventions:
+#define UE_JSON_OBJECT_INSTANCE_KEY_TCHAR TEXT("__UObject")
+#define UE_JSON_OBJECT_NAME_KEY_TCHAR TEXT("Name")
+#define UE_JSON_OBJECT_CLASS_KEY_TCHAR TEXT("Class")
+#define UE_JSON_OBJECT_FLAGS_KEY_TCHAR TEXT("Flags")
+#define UE_JSON_OBJECT_STRUCTURED_DATA_KEY_TCHAR TEXT("__StructuredData")
+#define UE_JSON_OBJECT_SERIAL_DATA_KEY_TCHAR TEXT("__SerialData")
+#define UE_JSON_OBJECT_SPARSE_CLASS_DATA_KEY_TCHAR TEXT("__SparseClassData")
+#define UE_JSON_OBJECT_INDIRECTLY_REFERENCED_KEY_TCHAR TEXT("__IndirectlyReferenced")
+
+// Reference encoding related conventions:
+#define UE_JSON_OBJECT_REF_PREFIX "uobject:"
+#define UE_JSON_FIELD_REF_PREFIX "ffield:"
+#define UE_JSON_REF_NONE "None"
+#define UE_JSON_DOMAIN_IDENTIFIER_NATIVE "__native__"
+#define UE_JSON_DOMAIN_IDENTIFIER_FILE "__file__"
+#define UE_JSON_REF_POSTFIX ".uasset"
+
+// Conventions related to encoding generic intrinsic types (TMap and TOptional):
+#define UE_JSON_TMAP_KEY_KEY_TCHAR TEXT("Key")
+#define UE_JSON_TMAP_VALUE_KEY_TCHAR TEXT("Value")
+#define UE_JSON_OPTIONAL_VALUE_KEY_TCHAR TEXT("OptionalValue")
\ No newline at end of file
diff --git a/Engine/Source/Runtime/Experimental/JsonObjectGraph/Private/JsonStringifyArchive.cpp b/Engine/Source/Runtime/Experimental/JsonObjectGraph/Private/JsonStringifyArchive.cpp
new file mode 100644
--- /dev/null
+++ b/Engine/Source/Runtime/Experimental/JsonObjectGraph/Private/JsonStringifyArchive.cpp
@@ -0,0 +1,162 @@
+// Copyright Epic Games, Inc. All Rights Reserved.
+
+#include "JsonStringifyArchive.h"
+#include "Misc/Base64.h"
+#include "JsonStringifyImpl.h"
+#include "UObject/Object.h"
+#include "UObject/LazyObjectPtr.h"
+#include "UObject/SoftObjectPtr.h"
+#include "UObject/WeakObjectPtr.h"
+#include "Serialization/CustomVersion.h"
+
+namespace UE::Private
+{
+
+FJsonStringifyArchive::FJsonStringifyArchive(
+	const UObject* InObject
+	, int32 InInitialIndentLevel
+	, FJsonStringifyImpl* InRootImpl
+	, TArray<FCustomVersion>& InVersionsToHarvest
+	, bool bFilterEditorOnly)
+	: ObjectBeingStreamSerialized(InObject)
+	, RootImpl(InRootImpl)
+	, MemoryWriter(Result)
+	, Writer(UE::Private::FPrettyJsonWriter::Create(&MemoryWriter, InInitialIndentLevel))
+	, InitialIndentLevel(InInitialIndentLevel)
+	, VersionsToHarvest(InVersionsToHarvest)
+{
+	// we don't want to write any properties, only 
+	// stuff from user serializes:
+	ArUseCustomPropertyList = true;
+	ArCustomPropertyList = nullptr;
+	SetIsSaving(true);
+	SetFilterEditorOnly(bFilterEditorOnly);
+	SetWantBinaryPropertySerialization(true); // Setting this to disable SerializeScriptProperties
+	MemoryWriter.SetIsPersistent(true);
+	MemoryWriter.SetIsTextFormat(true);
+}
+
+TArray<uint8> FJsonStringifyArchive::ToJson()
+{
+	if (ObjectBeingStreamSerialized->HasAnyFlags(RF_ClassDefaultObject))
+	{
+		// CDOs are never/have never been natively serialized
+		return TArray<uint8>();
+	}
+
+	Writer->WriteArrayStartInline();
+	Writer->WriteLineTerminator();
+	const_cast<UObject*>(ObjectBeingStreamSerialized)->Serialize(*this);
+	const TArray<uint8> NullStream = GetNullStream();
+
+	const bool bDidWriteAnything = Result != NullStream;
+	if (bDidWriteAnything)
+	{
+		Writer->WriteNewlineAndArrayEnd();
+	}
+
+	// record version regardless of whether anything
+	// was written, lots of people just use Serialize as
+	// a lifecycle function:
+	VersionsToHarvest.Append(GetCustomVersions().GetAllVersions());
+
+	return bDidWriteAnything ? MoveTemp(Result) : TArray<uint8>();
+}
+
+TArray<uint8> FJsonStringifyArchive::GetNullStream() const
+{
+	TArray<uint8> NullStream;
+	NullStream.Add('[');
+	NullStream.Add('\n');
+	for (int32 I = 0; I < InitialIndentLevel + 1; ++I)
+	{
+		NullStream.Add('\t');
+	}
+	NullStream.Append("false");
+	NullStream.Pop(); // should probably just look for the single false token.. but.. this works
+
+	return NullStream;
+}
+
+void FJsonStringifyArchive::Serialize(void* V, int64 Length)
+{
+	// The stream Serializers manage their own branches, and ultimately just write a stream of bytes. 
+	// We can encode those bytes any way we choose, but some information is effectively lost (e.g. 
+	// did the caller have a float or an int? would have been nice to know!). I've played
+	// around with various representations but didn't find any particularly helpful, so i've chosen
+	// the simplest representation that can handle all data: Base64 encoding - this is very much
+	// a 'get a byte, get a byte, get another byte' kind of implementation, and is not fast. The
+	// text representation is also not good. I think if you're having problems with this you
+	// should move to a structured archive or, even better, get the UPROPERTY based declarative system 
+	// working for you in some way. That may require some ingenuity, but you read and understood 
+	// this entire so you're definitely the cream of the crop.
+	uint32 ExpectedLength = FBase64::GetEncodedDataSize(Length);
+	TArray<uint8> EncodedData;
+	EncodedData.SetNum(ExpectedLength + 1);
+	FBase64::Encode<ANSICHAR>(static_cast<uint8*>(V), Length, (ANSICHAR*)EncodedData.GetData());
+	Writer->WriteValueInline(FAnsiStringView((const ANSICHAR*)EncodedData.GetData(), EncodedData.Num() - 1));
+}
+
+#if WITH_EDITOR
+void FJsonStringifyArchive::SerializeBool(bool& D)
+{
+	bool bBool = D;
+	Writer->WriteValueInline(bBool);
+}
+#endif
+
+FArchive& FJsonStringifyArchive::operator<<(UObject*& Value)
+{
+	RootImpl->WriteObjectAsJsonToWriter(ObjectBeingStreamSerialized, Value, Writer);
+	return *this;
+}
+
+FArchive& FJsonStringifyArchive::operator<<(FField*& Value)
+{
+	RootImpl->WriteFieldReferenceTo(ObjectBeingStreamSerialized, Value, Writer);
+	return *this;
+}
+
+FArchive& FJsonStringifyArchive::operator<<(struct FLazyObjectPtr& Value)
+{
+	UObject* Resolved = Value.Get();
+	return *this << Resolved;
+}
+
+FArchive& FJsonStringifyArchive::operator<<(struct FObjectPtr& Value)
+{
+	UObject* Resolved = Value.Get();
+	return *this << Resolved;
+}
+
+FArchive& FJsonStringifyArchive::operator<<(struct FSoftObjectPtr& Value)
+{
+	FSoftObjectPath Path = Value.ToSoftObjectPath();
+	return *this << Path;
+}
+
+FArchive& FJsonStringifyArchive::operator<<(struct FSoftObjectPath& Value)
+{
+	Writer->WriteValueInline(Value.ToString());
+	return *this;
+}
+
+FArchive& FJsonStringifyArchive::operator<<(struct FWeakObjectPtr& Value)
+{
+	UObject* Resolved = Value.Get();
+	return *this << Resolved;
+}
+
+FArchive& FJsonStringifyArchive::operator<<(FName& Value)
+{
+	Writer->WriteValueInline(Value.ToString());
+	return *this;
+}
+
+FArchive& FJsonStringifyArchive::operator<<(FText& Value)
+{
+	Writer->WriteValueInline(Value);
+	return *this;
+}
+
+}
diff --git a/Engine/Source/Runtime/Experimental/JsonObjectGraph/Private/JsonStringifyArchive.h b/Engine/Source/Runtime/Experimental/JsonObjectGraph/Private/JsonStringifyArchive.h
new file mode 100644
--- /dev/null
+++ b/Engine/Source/Runtime/Experimental/JsonObjectGraph/Private/JsonStringifyArchive.h
@@ -0,0 +1,58 @@
+// Copyright Epic Games, Inc. All Rights Reserved.
+
+#pragma once
+
+#include "Serialization/ArchiveUObject.h"
+#include "Serialization/MemoryWriter.h"
+#include "PrettyJsonWriter.h"
+
+struct FCustomVersion;
+
+namespace UE::Private
+{
+
+struct FJsonStringifyImpl;
+
+// Serial in the name of this type is meant to refer to 'UObject::Serialize(FArchive&)'. 
+// This writer writes the resulting byte stream safely to json.
+struct FJsonStringifyArchive final : private FArchiveUObject
+{
+	FJsonStringifyArchive(
+		const UObject* InObject, 
+		int32 InitialIndentLevel, 
+		FJsonStringifyImpl* InRootImpl,
+		TArray<FCustomVersion>& InVersionsToHarvest,
+		bool bFilterEditorOnly);
+
+	// FMemoryWriter wants to work with TArray<uint8>
+	// so that's what we're returning:
+	TArray<uint8> ToJson();
+private:
+	TArray<uint8> GetNullStream() const;
+
+	virtual void Serialize(void* V, int64 Length) override;
+	#if WITH_EDITOR
+	virtual void SerializeBool(bool& D) override;
+	#endif
+	virtual FArchive& operator<<(UObject*& Value) override;
+	virtual FArchive& operator<<(FField*& Value) override;
+	virtual FArchive& operator<<(struct FLazyObjectPtr& Value) override;
+	virtual FArchive& operator<<(struct FObjectPtr& Value) override;
+	virtual FArchive& operator<<(struct FSoftObjectPtr& Value) override;
+	virtual FArchive& operator<<(struct FSoftObjectPath& Value) override;
+	virtual FArchive& operator<<(struct FWeakObjectPtr& Value) override;
+	// overrides for strings:
+	virtual FArchive& operator<<(FName& Value) override;
+	virtual FArchive& operator<<(FText& Value) override;
+
+	const UObject* const ObjectBeingStreamSerialized;
+	FJsonStringifyImpl* RootImpl;
+
+	TArray<uint8> Result;
+	FMemoryWriter MemoryWriter;
+	TSharedRef<UE::Private::FPrettyJsonWriter> Writer;
+	int32 InitialIndentLevel;
+	TArray<FCustomVersion>& VersionsToHarvest;
+};
+
+}
diff --git a/Engine/Source/Runtime/Experimental/JsonObjectGraph/Private/JsonStringifyImpl.cpp b/Engine/Source/Runtime/Experimental/JsonObjectGraph/Private/JsonStringifyImpl.cpp
new file mode 100644
--- /dev/null
+++ b/Engine/Source/Runtime/Experimental/JsonObjectGraph/Private/JsonStringifyImpl.cpp
@@ -0,0 +1,988 @@
+// Copyright Epic Games, Inc. All Rights Reserved.
+
+#include "JsonStringifyImpl.h"
+#include "PrettyJsonWriter.h"
+#include "Algo/ForEach.h"
+#include "Algo/RemoveIf.h"
+#include "Misc/App.h"
+#include "JsonObjectGraphConventions.h"
+#include "JsonStringifyArchive.h"
+#include "JsonStringifyStructuredArchive.h"
+#include "Serialization/StructuredArchive.h"
+#include "Serialization/StructuredArchiveSlots.h"
+#include "Serialization/Formatters/JsonArchiveOutputFormatter.h"
+#include "Serialization/ArchiveUObjectFromStructuredArchive.h"
+#include "UObject/EnumProperty.h"
+#include "UObject/PropertyOptional.h"
+#include "UObject/Package.h"
+#include "UObject/TextProperty.h"
+
+namespace UE::Private
+{
+
+FUtf8String GetPackageModuleName(FUtf8StringView PackageName)
+{
+	if (PackageName.StartsWith("/Script/"))
+	{
+		// native package, module name will immediately follow script
+		// and object name will immediately follow that module name:
+		return FUtf8String(UE_JSON_DOMAIN_IDENTIFIER_NATIVE);
+	}
+	else if (PackageName.StartsWith("/Game/"))
+	{
+		return TCHAR_TO_UTF8(FApp::GetProjectName());
+	}
+	else
+	{
+		check(PackageName[0] == '/');
+		FUtf8StringView PackageNameNoSlash(PackageName.GetData() + 1, PackageName.Len() - 1);
+		// normal asset, it should be located within a plugin or
+		// project, that will serve as its 'module' name:
+		int32 FolderIndex = INDEX_NONE;
+		check(PackageNameNoSlash.FindChar(static_cast<UTF8CHAR>('/'), FolderIndex));
+
+		return FUtf8String(PackageName.Mid(1, FolderIndex));
+	}
+}
+
+FUtf8String GetPackageRelativePath(FUtf8StringView PackageName)
+{
+	check(PackageName[0] == '/');
+	FUtf8StringView PackageNameNoSlash(PackageName.GetData() + 1, PackageName.Len() - 1);
+	int32 FolderIndex = INDEX_NONE;
+	check(PackageNameNoSlash.FindChar(static_cast<UTF8CHAR>('/'), FolderIndex));
+	++FolderIndex;
+	return FUtf8String(PackageName.RightChop(FolderIndex));
+}
+
+// gathering utils
+struct FPackageReferenceFinder : public FArchiveUObject
+{
+	FPackageReferenceFinder(const UObject* Obj, TArray<const UObject*>& InReferences, bool bFilterEditorOnly)
+		: FArchiveUObject()
+		, References(InReferences)
+		, OwningPackage(Obj->GetOutermost())
+	{
+		// Copying FPackageHarvester:
+		SetIsPersistent(true);
+		SetIsSaving(true);
+		SetFilterEditorOnly(bFilterEditorOnly);
+		ArIsObjectReferenceCollector = true;
+		ArShouldSkipBulkData = true;
+		if (Obj->HasAnyFlags(RF_ClassDefaultObject))
+		{
+			Obj->GetClass()->SerializeDefaultObject((UObject*)Obj, *this);
+		}
+		else
+		{
+			((UObject*)Obj)->Serialize(*this);
+		}
+	}
+
+private:
+	virtual FArchive& operator<<(UObject*& ObjRef) override
+	{
+		if (ObjRef != nullptr &&
+			(!ObjRef->HasAnyFlags(RF_Transient) || ObjRef->IsNative()) &&
+			!ObjRef->IsIn(GetTransientPackage()) &&
+			ObjRef != OwningPackage)
+		{
+			// Set to null any pointer to an external asset
+			References.Add(ObjRef);
+		}
+
+		return *this;
+	}
+
+	TArray<const UObject*>& References;
+	const UPackage* const OwningPackage;
+};
+
+static const UClass* GetFirstNativeClass(const UClass* Class)
+{
+	const UClass* Iter = Class;
+	while (Iter && !Iter->HasAnyClassFlags(CLASS_Native))
+	{
+		Iter = Iter->GetSuperClass();
+	}
+	return Iter;
+}
+
+static void FilterEditorOnlyObjects(
+	TArray<UObject*>& Objects
+)
+{
+	#if WITH_EDITOR
+	Objects.SetNum(Algo::StableRemoveIf(Objects, [](const UObject* Obj) {
+		const UClass* NativeClass = GetFirstNativeClass(Obj->GetClass());
+		return IsEditorOnlyObject(NativeClass) || !(NativeClass->ClassDefaultObject->NeedsLoadForClient() || NativeClass->ClassDefaultObject->NeedsLoadForServer());
+		}));
+	#endif
+}
+
+static void GatherExports(
+	TConstArrayView<const UObject*> Roots,
+	TConstArrayView<const UObject*> DisallowList,
+	bool bFilterEditorOnly,
+	TArray<const UObject*>& OutRoots,
+	TArray<const UObject*>& OutExports)
+{
+	// find everything roots references that is within roots and put it into
+	// OutExports, unless it is in the disallow list:
+
+	TSet<const UObject*> DisallowSet;
+	DisallowSet.Append(DisallowList);
+	for (const UObject* Obj : DisallowList)
+	{
+		// const hacks and GetObjects expressiveness junk
+		TArray<UObject*> DisallowedSubObjects;
+		GetObjectsWithOuter(Obj, DisallowedSubObjects, true, RF_Transient);
+		for (UObject* DisallowedObj : DisallowedSubObjects)
+		{
+			DisallowSet.Add(DisallowedObj);
+		}
+	}
+	TSet<const UObject*> AllowedRoots;
+	AllowedRoots.Append(Roots);
+	for (const UObject* Obj : Roots)
+	{
+		if (!Obj->IsA<UPackage>())
+		{
+			OutRoots.Add(Obj);
+		}
+	}
+
+	TArray<const UObject*> PendingRefs = TArray<const UObject*>(Roots);
+
+	TSet<const UObject*> RefsProcessed;
+	RefsProcessed.Append(PendingRefs);
+
+	TArray<const UObject*> ScratchRefs; // just keeping this allocation alive across iterations
+
+	// Roots is invalid after this point!
+
+	UObject* TransientPackage = GetTransientPackage();
+	while (PendingRefs.Num())
+	{
+		const UObject* Iter = PendingRefs.Pop();
+		if (Iter->GetOutermost() == TransientPackage)
+		{
+			continue;
+		}
+
+		if (AllowedRoots.Contains(Iter) && Iter->IsA<UPackage>())
+		{
+			// add all immediate roots of a package, if filter 'editor only' objects
+			// exclude any object that are of an editor only native type:
+			TArray<UObject*> PackageInners;
+			GetObjectsWithOuter(Iter, PackageInners, false, RF_Transient);
+			if (bFilterEditorOnly)
+			{
+				FilterEditorOnlyObjects(PackageInners);
+			}
+			PackageInners.Sort([](const UObject& A, const UObject& B) { return A.GetFName().LexicalLess(B.GetFName()); });
+			OutRoots.Append(PackageInners);
+			PendingRefs.Append(PackageInners);
+			continue;
+		}
+
+		OutExports.Add(Iter);
+
+		ScratchRefs.Add(Iter->GetClass());
+		FPackageReferenceFinder ReferencedObjects(Iter, ScratchRefs, bFilterEditorOnly);
+		for (const UObject* Obj : ScratchRefs)
+		{
+			if (RefsProcessed.Contains(Obj))
+			{
+				continue;
+			}
+			RefsProcessed.Add(Obj);
+
+			if (DisallowSet.Contains(Obj))
+			{
+				continue;
+			}
+
+			const UObject* OuterIter = Obj->GetOuter();
+			bool bIsInRoot = false;
+			while (OuterIter)
+			{
+				if (AllowedRoots.Contains(OuterIter))
+				{
+					bIsInRoot = true;
+					break;
+				}
+				OuterIter = OuterIter->GetOuter();
+			}
+			if (!bIsInRoot)
+			{
+				continue;
+			}
+
+			PendingRefs.Add(Obj);
+		}
+		ScratchRefs.Reset();
+	}
+}
+
+static void WriteObjectPath(FUtf8StringBuilderBase& OutPath, const UObject* ForObject, const UObject* OuterLimit)
+{
+	const auto WriteObjectPathImpl = [&OutPath, OuterLimit](const UObject* ForObject, auto Self)
+		{
+			if (ForObject == OuterLimit)
+			{
+				return;
+			}
+			Self(ForObject->GetOuter(), Self);
+			OutPath << "/";
+			OutPath << ForObject->GetName();
+		};
+	WriteObjectPathImpl(ForObject, WriteObjectPathImpl);
+}
+
+static void WriteFieldPath(FUtf8StringBuilderBase& OutPath, const FField* ForField, const UObject* OuterLimit)
+{
+	const auto WriteWriteFieldPathImpl = [&OutPath, OuterLimit](const UObject* ForObject, const FField* ForField, auto Self)
+		{
+			if (ForObject == OuterLimit)
+			{
+				return;
+			}
+			if (ForField)
+			{
+				FFieldVariant Owner = ForField->GetOwnerVariant();
+				Self(Owner.ToUObject(), Owner.ToField(), Self);
+			}
+			else if (ForObject)
+			{
+				Self(ForObject->GetOuter(), nullptr, Self);
+			}
+			OutPath << "/";
+			if (ForObject)
+			{
+				OutPath << ForObject->GetName();
+			}
+			else
+			{
+				check(ForField);
+				OutPath << ForField->GetName();
+			}
+		};
+	WriteWriteFieldPathImpl(nullptr, ForField, WriteWriteFieldPathImpl);
+}
+
+FJsonStringifyImpl::FJsonStringifyImpl(TConstArrayView<const UObject*> Roots, const FJsonStringifyOptions& Options)
+	: WriteOptions(Options)
+	, MemoryWriter(Result)
+	, Writer(FJsonWriter::Create(&MemoryWriter))
+	, CurrentObject(nullptr)
+{
+	const bool bFilterEditorOnly = (Options.Flags & EJsonStringifyFlags::FilterEditorOnlyData) != EJsonStringifyFlags::Default;
+	TArray<const UObject*> Exports;
+	GatherExports(Roots, {}, bFilterEditorOnly, RootObjects, Exports);
+
+	ObjectsToExport = TSet<const UObject*>(Exports);
+
+	MemoryWriter.SetIsPersistent(true);
+	MemoryWriter.SetFilterEditorOnly(bFilterEditorOnly);
+	MemoryWriter.SetIsTextFormat(true);
+}
+
+FUtf8String FJsonStringifyImpl::ToJson()
+{
+	ToJsonBytes();
+	// Hacking around TJsonWriter and FMemoryWriter, is an array of bytes by
+	// any other name not exceedingly dangerous? Everyone should be working
+	// with utf8:
+	TArray<ANSICHAR>& ResultTyped = reinterpret_cast<TArray<ANSICHAR>&>(Result);
+	return FUtf8String(MoveTemp(ResultTyped));
+}
+
+void FJsonStringifyImpl::WriteObjectAsJsonToWriter(const UObject* OwningObject, const UObject* InObject, TSharedRef<FJsonWriter> WriterToUse)
+{
+	if (CurrentScope)
+	{
+		CurrentScope->Apply();
+	}
+
+	if (InObject && InObject->IsIn(OwningObject) &&
+		ObjectsToExport.Contains(InObject) &&
+		!ObjectsExported.Contains(InObject))
+	{
+		ObjectsExported.Add(InObject);
+		TSharedRef<FJsonWriter> RootWriter = Writer;
+		Writer = WriterToUse;
+		WriteObjectToJson(InObject);
+		Writer = RootWriter;
+	}
+	else
+	{
+		// we just need to write a path
+		WriterToUse->WriteValueInline(FUtf8StringView(WriteObjectReference(InObject)));
+	}
+}
+
+void FJsonStringifyImpl::WriteFieldReferenceTo(const UObject* OwningObject, const FField* Value, TSharedRef<FJsonWriter> WriterToUse)
+{
+	WriterToUse->WriteValueInline(FUtf8StringView(WriteFieldReference(Value)));
+}
+
+void FJsonStringifyImpl::WriteObjectAsJsonToArchive(const UObject* OwningObject, const UObject* InObject, FArchive* ArchiveToUse, int32 InitialIndentLevel)
+{
+	TSharedRef<FJsonWriter> JsonWriter = FJsonWriter::Create(ArchiveToUse, InitialIndentLevel);
+	JsonWriter->HACK_SetPreviousTokenWritten();
+	WriteObjectAsJsonToWriter(OwningObject, InObject, JsonWriter);
+}
+
+void FJsonStringifyImpl::ToJsonBytes()
+{
+	if (RootObjects.Num() == 0)
+	{
+		return;
+	}
+
+	// main entry point to our writer state machine:
+	Writer->WriteObjectStart();
+	Writer->WriteIdentifierPrefix(UE_JSON_ROOT_OBJECTS_KEY_TCHAR);
+	Writer->WriteArrayStartInline();
+	Writer->WriteLineTerminator();
+	for (const UObject* Object : RootObjects)
+	{
+		ObjectsExported.Add(Object);
+		WriteObjectToJson(Object);
+	}
+	Writer->WriteArrayEnd();
+	
+#if WITH_TEXT_ARCHIVE_SUPPORT
+	if (Versions.Num() != 0)
+	{
+		// filter dupes, preserving order of first encounter - we need some stable order:
+		TSet<FGuid> Encountered;
+		Versions.SetNum(Algo::StableRemoveIf(Versions,
+			[&Encountered](const FCustomVersion& V)
+			{
+				if (Encountered.Contains(V.Key))
+				{
+					return true;
+				}
+				else
+				{
+					Encountered.Add(V.Key);
+					return false;
+				}
+			}));
+
+		Writer->WriteIdentifierPrefix(UE_JSON_CUSTOM_VERSIONS_KEY_TCHAR);
+		FJsonStringifyStructuredArchive::WriteCustomVersionValueInline(Versions, Writer->GetIndentLevel(), MemoryWriter);
+	}
+#endif // WITH_TEXT_ARCHIVE_SUPPORT
+	Writer->WriteObjectEnd();
+	// Result is now populated with Json representation
+}
+
+void FJsonStringifyImpl::WriteObjectToJson(const UObject* Object)
+{
+	CurrentObject = Object;
+
+	// loop properties, writing any that have changed:
+	const UObject* const Archetype = Object->GetArchetype();
+	const UClass* const ArchetypeClass = Archetype->GetClass();
+	Writer->WriteObjectStartInline();
+	// Write native UObject data - name, type, flags, native user serialize, etc
+	WriteNativeObjectData();
+
+	for (TFieldIterator<FProperty> FieldIt(Object->GetClass()); FieldIt; ++FieldIt)
+	{
+		FProperty* Property = *FieldIt;
+		const UObject* ValidatedArchtype = ArchetypeClass->IsChildOf(Property->GetOwnerClass()) ?
+			Archetype : nullptr;
+
+		WriteIdentifierAndValueToJson(Object, ValidatedArchtype, Property);
+	}
+	// We may have inner objects that were not referenced directly by this object, 
+	// but will be referenced by other objects in the graph. We must write them here, or we
+	// would have to encode them at the root level, which would disrupt locality.
+	// The draw back of recording them here is that we must order them ourselves,
+	// which will degrade the stability of the serialized buffer. It feels like 'the best
+	// we can do' is to write these objects in alphabetical order:
+	WriteIndirectlyReferencedContainedObjects(Object);
+	Writer->WriteObjectEnd();
+
+	// emit stream serializer, omitting tagged properties:
+	CurrentObject = CurrentObject->GetOuter();
+}
+
+void FJsonStringifyImpl::WriteNativeObjectData()
+{
+	const UObject* Object = CurrentObject;
+	check(Object);
+
+	// Could use a structured archive to write this more declaratively...
+	Writer->WriteIdentifierPrefix(UE_JSON_OBJECT_INSTANCE_KEY_TCHAR);
+	Writer->WriteObjectStartInline();
+	Writer->WriteValue(UE_JSON_OBJECT_NAME_KEY_TCHAR, Object->GetName());
+	Writer->WriteUtf8Value(UE_JSON_OBJECT_CLASS_KEY_TCHAR, WriteObjectReference(Object->GetClass()));
+	Writer->WriteValue(UE_JSON_OBJECT_FLAGS_KEY_TCHAR, Object->GetFlags() & RF_Load);
+	Writer->WriteObjectEnd();
+
+	// This is tricky, we have no good mechanism for detecting whether an object wants a structured
+	// serialization or a traditional stream serialization. So I've decided to 'try' the structured
+	// serializer and if it writes nothing we fall back to the stream (FArchive) serializer. Hopefully
+	// some day we get rid of FStructuredArchive as it does not spark joy.
+#if WITH_TEXT_ARCHIVE_SUPPORT
+	TArray<uint8> StructuredData = StructuredDataToJson(Object, Writer->GetIndentLevel());
+	if (!StructuredData.IsEmpty())
+	{
+		Writer->WriteIdentifierPrefix(UE_JSON_OBJECT_STRUCTURED_DATA_KEY_TCHAR);
+		Writer->WriteJsonRaw(FAnsiStringView((ANSICHAR*)StructuredData.GetData(), StructuredData.Num()));
+	}
+	else
+#endif
+	{
+		// no useful native structured data, write the native serial data:
+		TArray<uint8> SerialData = SerialDataToJson(Object, Writer->GetIndentLevel());
+		if (!SerialData.IsEmpty())
+		{
+			Writer->WriteIdentifierPrefix(UE_JSON_OBJECT_SERIAL_DATA_KEY_TCHAR);
+			Writer->WriteJsonRaw(FAnsiStringView((ANSICHAR*)SerialData.GetData(), SerialData.Num()));
+		}
+	}
+
+	// Sparse class data is ambiguously serialized as part of SerializeDefaultObject,
+	// which our text serializer does not use (except for reference gathering). Lets write
+	// the special SparseClassData member here:
+	if (const UClass* AsClass = Cast<UClass>(Object))
+	{
+		if (const void* SCD = const_cast<UClass*>(AsClass)->GetSparseClassData(EGetSparseClassDataMethod::ReturnIfNull))
+		{
+			const UScriptStruct* SCDStruct = AsClass->GetSparseClassDataStruct();
+			const void* DefaultSCD = AsClass->GetArchetypeForSparseClassData();
+			const UScriptStruct* DefaultSCDStruct = AsClass->GetSparseClassDataArchetypeStruct();
+			WriteStructToJsonWithIdentifier(UE_JSON_OBJECT_SPARSE_CLASS_DATA_KEY_TCHAR, SCD, DefaultSCD, SCDStruct, DefaultSCDStruct);
+		}
+	}
+}
+
+void FJsonStringifyImpl::WriteIndirectlyReferencedContainedObjects(const UObject* ForObject)
+{
+	TArray<const UObject*> UnwrittenInners;
+	// @todo: naive implementation - is looping over every object here too slow?
+	for (const UObject* ObjectToExport : ObjectsToExport)
+	{
+		if (ObjectToExport->GetOuter() == ForObject && !ObjectsExported.Contains(ObjectToExport))
+		{
+			UnwrittenInners.Add(ObjectToExport);
+		}
+	}
+
+	if (UnwrittenInners.Num() == 0)
+	{
+		return;
+	}
+
+	UnwrittenInners.Sort([](const UObject& A, const UObject& B) { return A.GetName().Compare(B.GetName()) < 0; });
+
+	// write all UnwrittenInners into the __IndirectlyReferenced member, for simplicity we 
+	// will always encode this as an array. We know that these unwritten inners will not be outered
+	// to the written objects here, so we don't need to tag them as exported ahead of time:
+	Writer->WriteIdentifierPrefix(UE_JSON_OBJECT_INDIRECTLY_REFERENCED_KEY_TCHAR);
+	Writer->WriteArrayStartInline();
+	Writer->WriteLineTerminator();
+	for (const UObject* Object : UnwrittenInners)
+	{
+		WriteObjectToJson(Object);
+	}
+	Writer->WriteArrayEnd();
+}
+
+FUtf8String FJsonStringifyImpl::WriteObjectReference(const UObject* ForObject) const
+{
+	// matching FSaveContext::GetSaveableStatusNoOuter, if we integrate into
+	// save package we can reuse the logic from there:
+	if (!ForObject || 
+		!IsValid(ForObject) || 
+		(ForObject->HasAnyFlags(RF_Transient) && !ForObject->IsNative()))
+	{
+		return UE_JSON_OBJECT_REF_PREFIX UE_JSON_REF_NONE;
+	}
+
+	// We want users to be able to move these files around the filesystem with
+	// minimal fixup, so intra artifact references are relative to the file itself,
+	// intra module references are relative to the module, but cross module references
+	// are unambiguous and fully qualified. Here are some examples:
+	//
+	// uobject:__file__/a/path/to/a/child - this is a reference to an object within our file
+	// uobject:__native__/native_module/native_object - this is a reference to a natively defined object in native_module
+	// uobject:real_module/a/path/to/a/child.uasset/a/path/to/a/child/object - this is a reference to an object in another file
+	//
+	// file extensions are included in the path to disambiguate between folders and assets
+	// that share a name. The special __file__ identifier exists so that assets internal
+	// reference remain valid when copied, the special __native__ identifier exists so that
+	// natively defined objects can be unambiguously referenced
+	//
+	// real_module in this example will refer to either a uplugin or uproject.
+
+	FUtf8StringBuilderBase Reference;
+	Reference << UE_JSON_OBJECT_REF_PREFIX;
+	if (ObjectsToExport.Contains(ForObject))
+	{
+		Reference << UE_JSON_DOMAIN_IDENTIFIER_FILE;
+		const UObject* OuterLimit = ForObject;
+		while (ObjectsToExport.Contains(OuterLimit))
+		{
+			OuterLimit = OuterLimit->GetOuter();
+		}
+		WriteObjectPath(Reference, ForObject, OuterLimit);
+	}
+	else if (const UPackage* Package = ForObject->GetPackage())
+	{
+		const FUtf8String PackageName = TCHAR_TO_UTF8(*Package->GetName());
+		const FUtf8String ModuleName = GetPackageModuleName(PackageName);
+
+		Reference << ModuleName;
+
+		FUtf8String PackagePath = GetPackageRelativePath(PackageName);
+		Reference << PackagePath;
+		if (ModuleName != UE_JSON_DOMAIN_IDENTIFIER_NATIVE)
+		{
+			Reference << UE_JSON_REF_POSTFIX;
+		}
+		WriteObjectPath(Reference, ForObject, Package);
+	}
+	else
+	{
+		checkf(false, TEXT("Found object with no package: %s"), *ForObject->GetPathName());
+	}
+
+	return Reference.ToString();
+}
+
+FUtf8String FJsonStringifyImpl::WriteFieldReference(const FField* Value) const
+{
+	if (!Value)
+	{
+		return UE_JSON_FIELD_REF_PREFIX UE_JSON_REF_NONE;
+	}
+
+	// identical convention for FField as above, albeit with a fieldref prefix
+	// instead of a uobject:
+	FUtf8StringBuilderBase Reference;
+
+	Reference << UE_JSON_FIELD_REF_PREFIX;
+	const UObject* OwningObject = Value->GetOwnerUObject();
+	check(OwningObject);
+	if (ObjectsToExport.Contains(OwningObject))
+	{
+		// object is within this file:
+		Reference << UE_JSON_DOMAIN_IDENTIFIER_FILE;
+		const UObject* OuterLimit = OwningObject;
+		while (ObjectsToExport.Contains(OuterLimit))
+		{
+			OuterLimit = OuterLimit->GetOuter();
+		}
+
+		WriteFieldPath(Reference, Value, OuterLimit);
+	}
+	else if (const UPackage* Package = OwningObject->GetPackage())
+	{
+		// external object:
+		const FUtf8String PackageName = TCHAR_TO_UTF8(*Package->GetName());
+		const FUtf8String ModuleName = GetPackageModuleName(PackageName);
+
+		Reference << ModuleName;
+
+		FUtf8String PackagePath = GetPackageRelativePath(PackageName);
+		Reference << PackagePath;
+		if (ModuleName != UE_JSON_DOMAIN_IDENTIFIER_NATIVE)
+		{
+			Reference << UE_JSON_REF_POSTFIX;
+		}
+		WriteFieldPath(Reference, Value, Package);
+	}
+	else
+	{
+		checkf(false, TEXT("Found field with no package: %s"), *OwningObject->GetPathName());
+	}
+
+	return Reference.ToString();
+}
+
+void FJsonStringifyImpl::WriteIdentifierAndValueToJson(const void* Container, const void* DefaultContainer, const FProperty* Property)
+{
+	if (!Property->ShouldSerializeValue(MemoryWriter))
+	{
+		return;
+	}
+
+	// this is inefficient for structs... but it papers over some
+	// problematic identical implementations. We will improve it
+	bool bMatchesDefault = DefaultContainer != nullptr;
+	for (int32 Idx = 0; Idx < Property->ArrayDim && bMatchesDefault; ++Idx)
+	{
+		const void* Value = Property->ContainerPtrToValuePtr<void>(Container, Idx);
+		const void* ArchetypeValue = Property->ContainerPtrToValuePtr<void>(DefaultContainer, Idx);
+
+		FString StringValue;
+		Property->ExportText_Direct(StringValue, Value, Value, nullptr, PPF_ForDiff);
+		if (ArchetypeValue)
+		{
+			FString DefaultStringValue;
+			Property->ExportText_Direct(DefaultStringValue, ArchetypeValue, ArchetypeValue, nullptr, PPF_ForDiff);
+			bMatchesDefault = DefaultStringValue.Equals(StringValue);
+		}
+	}
+
+	if (bMatchesDefault)
+	{
+		return;
+	}
+
+	FPendingScope PropertyIdentifier(this, [Impl = this, Property]()
+		{
+			Impl->Writer->WriteIdentifierPrefix(Property->GetName());
+		});
+	
+	if (Property->ArrayDim > 1)
+	{
+		// encode as a static array:
+		FPendingScope ArrayIdentifier(this, [Impl = this]()
+			{
+				Impl->Writer->WriteArrayStartInline();
+				Impl->Writer->WriteLineTerminator();
+			},
+			[Impl = this]()
+			{
+				Impl->Writer->WriteArrayEnd();
+			}
+		);
+		for (int32 Index = 0; Index < Property->ArrayDim; ++Index)
+		{
+			WriteValueToJson(
+				Property->ContainerPtrToValuePtr<void>(Container, Index),
+				DefaultContainer ? Property->ContainerPtrToValuePtr<void>(DefaultContainer, Index) : nullptr,
+				Property);
+		}
+	}
+	else
+	{
+		WriteValueToJson(
+			Property->ContainerPtrToValuePtr<void>(Container),
+			DefaultContainer ? Property->ContainerPtrToValuePtr<void>(DefaultContainer) : nullptr,
+			Property);
+	}
+}
+
+void FJsonStringifyImpl::WriteValueToJson(const void* Value, const void* DefaultValue, const FProperty* Property)
+{
+	bool bHandledAsAggregate = true;
+	if (const FStructProperty* StructProperty = CastField<FStructProperty>(Property))
+	{
+		// write any tagged data for the struct - but be sure to delta serialize
+		WriteStructToJson(Value, DefaultValue, StructProperty->Struct, StructProperty->Struct);
+	}
+	else if (CastField<FObjectProperty>(Property) || CastField<FClassProperty>(Property))
+	{
+		const UObject* Object = *(UObject**)Value;
+		WriteObjectAsJsonToWriter(CurrentObject, Object, Writer);
+	}
+	// containers:
+	else if (const FArrayProperty* ArrayProperty = CastField<FArrayProperty>(Property))
+	{
+		WriteArrayToJson(Value, ArrayProperty);
+	}
+	else if (const FSetProperty* SetProperty = CastField<FSetProperty>(Property))
+	{
+		WriteSetToJson(Value, SetProperty);
+	}
+	else if (const FMapProperty* MapProperty = CastField<FMapProperty>(Property))
+	{
+		WriteMapToJson(Value, MapProperty);
+	}
+	else if (const FOptionalProperty* OptionalProperty = CastField<FOptionalProperty>(Property))
+	{
+		WriteOptionalToJson(Value, OptionalProperty);
+	}
+	else
+	{
+		bHandledAsAggregate = false; // fall back to value serialization
+	}
+
+	if (!bHandledAsAggregate)
+	{
+		WriteIntrinsicToJson(Value, Property);
+	}
+}
+
+void FJsonStringifyImpl::WriteIntrinsicToJson(const void* Value, const FProperty* Property)
+{
+	if (CurrentScope)
+	{
+		CurrentScope->Apply();
+	}
+
+	// bools are special, because of GetPropertyValue, which can apply a mask
+	if (const FBoolProperty* BoolProp = CastField<FBoolProperty>(Property))
+	{
+		bool bValue = BoolProp->GetPropertyValue(Value);
+		Writer->WriteValueInline(bValue);
+	}
+	else if (const FByteProperty* ByteProperty = CastField<FByteProperty>(Property))
+	{
+		// byte property is special because of the enum/byte duality:
+		uint8 ValueAsIntrinsic = *(uint8*)Value;
+		if (const UEnum* EnumDef = ByteProperty->Enum)
+		{
+			FString StringValue = EnumDef->GetAuthoredNameStringByValue(ValueAsIntrinsic);
+			Writer->WriteValueInline(StringValue);
+		}
+		else
+		{
+			Writer->WriteValueInline(ValueAsIntrinsic);
+		}
+	}
+	else if (const FEnumProperty* EnumProperty = CastField<FEnumProperty>(Property))
+	{
+		// export enums as strings
+		const UEnum* EnumDef = EnumProperty->GetEnum();
+		FString StringValue = EnumDef->GetAuthoredNameStringByValue(
+			EnumProperty->GetUnderlyingProperty()->GetSignedIntPropertyValue(Value));
+		Writer->WriteValueInline(StringValue);
+	}
+	else if( const FTextProperty* TextProperty = CastField<FTextProperty>(Property))
+	{
+		FText ValueAsText = *(FText*)Value;
+		Writer->WriteValueInline(ValueAsText);
+	}
+#define INTRINSIC_TABLE \
+INTRINSIC_ENTRY(FInt8Property, int8) \
+INTRINSIC_ENTRY(FInt16Property, int16) \
+INTRINSIC_ENTRY(FIntProperty, int32) \
+INTRINSIC_ENTRY(FInt64Property, int64) \
+INTRINSIC_ENTRY(FUInt16Property, uint16) \
+INTRINSIC_ENTRY(FUInt32Property, uint32) \
+INTRINSIC_ENTRY(FUInt64Property, uint64) \
+INTRINSIC_ENTRY(FFloatProperty, float) \
+INTRINSIC_ENTRY(FDoubleProperty, double)
+
+#define INTRINSIC_ENTRY(FPropertyType, intrinsic_type) \
+	else if (const FPropertyType* As##intrinsic_type = CastField<FPropertyType>(Property)) \
+	{ \
+		intrinsic_type ValueAsIntrinsic = *(intrinsic_type*)Value; \
+		Writer->WriteValueInline(ValueAsIntrinsic); \
+	}
+	INTRINSIC_TABLE
+#undef INTRINSIC_ENTRY
+#undef INTRINSIC_TABLE
+	else
+	{
+		FString StringValue;
+		Property->ExportText_Direct(StringValue, Value, Value, nullptr, PPF_None);
+		Writer->WriteValueInline(StringValue);
+	}
+}
+
+void FJsonStringifyImpl::WriteStructToJsonWithIdentifier(const TCHAR* Identifier, const void* StructInstance, const void* DefaultInstance, const UScriptStruct* Struct, const UScriptStruct* DefaultStruct)
+{
+	check(StructInstance && Struct);
+	FPendingScope StructIdentifier(this, [Impl = this, Identifier]()
+		{
+			Impl->Writer->WriteIdentifierPrefix(Identifier);
+		});
+
+	WriteStructToJson(
+		StructInstance,
+		DefaultInstance,
+		Struct,
+		DefaultStruct);
+}
+
+void FJsonStringifyImpl::WriteStructToJson(const void* StructInstance, const void* DefaultInstance, const UScriptStruct* Struct, const UScriptStruct* DefaultStruct)
+{
+	FPendingScope ObjectStartIdentifier( 
+		this, 
+		[Impl = this]() { Impl->Writer->WriteObjectStartInline(); }, 
+		[Impl = this]() { Impl->Writer->WriteObjectEnd(); }
+		);
+
+	for (TFieldIterator<FProperty> PropertyIt(Struct); PropertyIt; ++PropertyIt)
+	{
+		FProperty* Property = *PropertyIt;
+		WriteIdentifierAndValueToJson(
+			StructInstance,
+			Property->IsInContainer(const_cast<UScriptStruct*>(DefaultStruct)) ? DefaultInstance : nullptr,
+			Property);
+	}
+}
+
+void FJsonStringifyImpl::WriteArrayToJson(const void* ArrayInstance, const FArrayProperty* Array)
+{
+	if (CurrentScope)
+	{
+		CurrentScope->Apply();
+	}
+
+	Writer->WriteArrayStartInline();
+	FScriptArrayHelper ArrayHelper(Array, ArrayInstance);
+	const int32 ArrayNum = ArrayHelper.Num();
+	if (ArrayNum != 0)
+	{
+		Writer->WriteLineTerminator();
+	}
+	for (int32 Index = 0; Index < ArrayNum; ++Index)
+	{
+		const void* Element = ArrayHelper.GetRawPtr(Index);
+		const FProperty* Property = Array->Inner;
+		WriteValueToJson(Element, nullptr, Property);
+	}
+
+	Writer->WriteArrayEnd();
+}
+
+void FJsonStringifyImpl::WriteSetToJson(const void* SetInstance, const FSetProperty* SetProperty)
+{
+	if (CurrentScope)
+	{
+		CurrentScope->Apply();
+	}
+
+	Writer->WriteArrayStartInline();
+	FScriptSetHelper SetHelper(SetProperty, SetInstance);
+	// Num() is number of elements, max index expensive to calculate, hence the odd loop:
+	// @todo: use set iterator
+	int32 SetNum = SetHelper.Num();
+	if (SetNum != 0)
+	{
+		Writer->WriteLineTerminator();
+	}
+	for (int32 Index = 0; SetNum; ++Index)
+	{
+		if (!SetHelper.IsValidIndex(Index))
+		{
+			continue;
+		}
+
+		--SetNum;
+		const void* Element = SetHelper.GetElementPtr(Index);
+		const FProperty* Property = SetProperty->ElementProp;
+		WriteValueToJson(Element, nullptr, Property);
+	}
+
+	Writer->WriteArrayEnd();
+}
+
+void FJsonStringifyImpl::WriteMapToJson(const void* MapInstance, const FMapProperty* MapProperty)
+{
+	if (CurrentScope)
+	{
+		CurrentScope->Apply();
+	}
+
+	Writer->WriteArrayStartInline();
+	FScriptMapHelper MapHelper(MapProperty, MapInstance);
+	// Num() is number of elements, max index expensive to calculate, hence the odd loop:
+	// @todo: use map iterator
+	int32 MapNum = MapHelper.Num();
+	if (MapNum != 0)
+	{
+		Writer->WriteLineTerminator();
+	}
+	for (int32 Index = 0; MapNum; ++Index)
+	{
+		if (!MapHelper.IsValidIndex(Index))
+		{
+			continue;
+		}
+
+		--MapNum;
+		const void* Key = MapHelper.GetKeyPtr(Index);
+		const FProperty* KeyProperty = MapProperty->KeyProp;
+		const void* Value = MapHelper.GetValuePtr(Index);
+		const FProperty* ValueProperty = MapProperty->ValueProp;
+
+		// age old question - how do you want to encode a tuple in json? I've
+		// chosen named tuple:
+		Writer->WriteObjectStartInline();
+		Writer->WriteIdentifierPrefix(UE_JSON_TMAP_KEY_KEY_TCHAR);
+		WriteValueToJson(Key, nullptr, KeyProperty);
+		Writer->WriteIdentifierPrefix(UE_JSON_TMAP_VALUE_KEY_TCHAR);
+		WriteValueToJson(Value, nullptr, ValueProperty);
+		Writer->WriteObjectEnd();
+	}
+
+	Writer->WriteArrayEnd();
+}
+
+void FJsonStringifyImpl::WriteOptionalToJson(const void* OptionalInstnace, const FOptionalProperty* OptionalProperty)
+{
+	if (CurrentScope)
+	{
+		CurrentScope->Apply();
+	}
+	Writer->WriteObjectStartInline();
+	if (const void* ValueAddress = static_cast<const void*>(OptionalProperty->GetValuePointerForReadOrReplaceIfSet(OptionalInstnace)))
+	{
+		Writer->WriteIdentifierPrefix(UE_JSON_OPTIONAL_VALUE_KEY_TCHAR);
+		WriteValueToJson(ValueAddress, nullptr, OptionalProperty->GetValueProperty());
+	}
+	Writer->WriteObjectEnd();
+}
+
+TArray<uint8> FJsonStringifyImpl::SerialDataToJson(const UObject* Object, int32 InitialIndentLevel)
+{
+	FJsonStringifyArchive SerialData(Object, InitialIndentLevel, this, Versions, MemoryWriter.IsFilterEditorOnly());
+	return SerialData.ToJson();
+}
+
+#if WITH_TEXT_ARCHIVE_SUPPORT
+TArray<uint8> FJsonStringifyImpl::StructuredDataToJson(const UObject* Object, int32 InitialIndentLevel)
+{
+	FJsonStringifyStructuredArchive StructuredData(Object, InitialIndentLevel, this, Versions, MemoryWriter.IsFilterEditorOnly());
+	return StructuredData.ToJson();
+}
+#endif // WITH_TEXT_ARCHIVE_SUPPORT
+
+FJsonStringifyImpl::FPendingScope::FPendingScope(FJsonStringifyImpl* To, TFunctionRef<void()> Prefix)
+	: Owner(To)
+	, Outer(To->CurrentScope)
+	, PendingPrefix(Prefix)
+{
+	To->CurrentScope = this;
+}
+
+FJsonStringifyImpl::FPendingScope::FPendingScope(FJsonStringifyImpl* To, TFunctionRef<void()> Prefix, TFunctionRef<void()> Postfix)
+	: Owner(To)
+	, Outer(To->CurrentScope)
+	, PendingPrefix(Prefix)
+	, PendingPostfix({Postfix})
+{
+	To->CurrentScope = this;
+}
+
+FJsonStringifyImpl::FPendingScope::~FPendingScope()
+{
+	if (bHasBeenApplied && PendingPostfix)
+	{
+		(*PendingPostfix)();
+	}
+
+	// scope left, if we've applied the prefix then we need to apply the postfix:
+	Owner->CurrentScope = Outer;
+}
+
+void FJsonStringifyImpl::FPendingScope::Apply()
+{
+	// writer has decided to apply the scope, write all pending prefixes:
+	if (bHasBeenApplied)
+	{
+		return;
+	}
+
+	bHasBeenApplied = true;
+	if (Outer)
+	{
+		Outer->Apply();
+	}
+	PendingPrefix();
+}
+
+}
diff --git a/Engine/Source/Runtime/Experimental/JsonObjectGraph/Private/JsonStringifyImpl.h b/Engine/Source/Runtime/Experimental/JsonObjectGraph/Private/JsonStringifyImpl.h
new file mode 100644
--- /dev/null
+++ b/Engine/Source/Runtime/Experimental/JsonObjectGraph/Private/JsonStringifyImpl.h
@@ -0,0 +1,108 @@
+// Copyright Epic Games, Inc. All Rights Reserved.
+
+#pragma once
+
+#include "Containers/Utf8String.h"
+#include "Containers/Array.h"
+#include "Serialization/CustomVersion.h"
+#include "Serialization/MemoryWriter.h"
+#include "JsonObjectGraph/Stringify.h"
+
+class FArrayProperty;
+class FSetProperty;
+class FMapProperty;
+class FOptionalProperty;
+
+namespace UE::Private 
+{
+
+struct FPrettyJsonWriter; 
+using FJsonWriter = FPrettyJsonWriter;
+
+struct FJsonStringifyImpl
+{
+	FJsonStringifyImpl(TConstArrayView<const UObject*> Roots, const FJsonStringifyOptions& Options);
+
+	// Writes and returns the provided root objects as a UTF8 encode Json string, honoring
+	// options specified in our constructor:
+	FUtf8String ToJson();
+
+	// Helpers for writing to an object to a specific FJsonWriter, or specific FArchive. These
+	// are used by the native stream serializers and native structured serializers respectively -
+	// WriteObjectAsJsonToWriter is also used as objects are encountered in reflected properties.
+	// The owning object must be provided so that we can determine whether to write InObject as a 
+	// peer reference, or inline. For root level serialization use ToJson():
+	void WriteObjectAsJsonToWriter(const UObject* OwningObject, const UObject* InObject, TSharedRef<FJsonWriter> WriterToUse);
+	void WriteObjectAsJsonToArchive(const UObject* OwningObject, const UObject* InObject, FArchive* ArchiveToUse, int32 InitialIndentLevel);
+
+	// Helper to improve support for FField references, which are not uncommon:
+	void WriteFieldReferenceTo(const UObject* OwningObject, const FField* Value, TSharedRef<FJsonWriter> WriterToUse);
+
+private:
+	void ToJsonBytes();
+	void WriteObjectToJson(const UObject* Object);
+
+	// Utility for unreflected UObject data - name, class, and flags - this 
+	// data is placed into a special key for a loader that needs to order
+	// UObject construction (find archetypes, UClasses, UStructs, etc):
+	void WriteNativeObjectData();
+	void WriteIndirectlyReferencedContainedObjects(const UObject* ForObject);
+	FUtf8String WriteObjectReference(const UObject* Object) const;
+	FUtf8String WriteFieldReference(const FField* Value) const;
+
+	void WriteIdentifierAndValueToJson(const void* Container, const void* DefaultContainer, const FProperty* Property);
+	void WriteValueToJson(const void* Value, const void* DefaultValue, const FProperty* Property);
+	void WriteIntrinsicToJson(const void* Value, const FProperty* Property);
+
+	// May want to expose WriteStructToJsonWithIdentifier for users to manually write USTRUCT instances into json:
+	void WriteStructToJsonWithIdentifier(const TCHAR* Identifier, const void* StructInstance, const void* DefaultInstance, const UScriptStruct* Struct, const UScriptStruct* DefaultStruct);
+	void WriteStructToJson(const void* StructInstance, const void* DefaultInstance, const UScriptStruct* Struct, const UScriptStruct* DefaultStruct);
+	void WriteArrayToJson(const void* ArrayInstance, const FArrayProperty* Array);
+	void WriteSetToJson(const void* SetInstance, const FSetProperty* SetProperty);
+	void WriteMapToJson(const void* MapInstance, const FMapProperty* MapProperty);
+	void WriteOptionalToJson(const void* OptionalInstnace, const FOptionalProperty* OptionalProperty);
+
+	TArray<uint8> SerialDataToJson(const UObject* Object, int32 InitialIndentLevel);
+#if WITH_TEXT_ARCHIVE_SUPPORT
+	TArray<uint8> StructuredDataToJson(const UObject* Object, int32 InitialIndentLevel);
+#endif// WITH_TEXT_ARCHIVE_SUPPORT
+
+	const FJsonStringifyOptions& WriteOptions;
+
+	TArray<uint8> Result;
+	FMemoryWriter MemoryWriter;
+
+	// helper for enqueuing writes that may not need to be applied, e.g. if there
+	// are values that match the default value we won't write out the enclosing scope:
+	struct FPendingScope
+	{
+		FPendingScope(FJsonStringifyImpl* To, TFunctionRef<void()> Prefix);
+		FPendingScope(FJsonStringifyImpl* To, TFunctionRef<void()> Prefix, TFunctionRef<void()> Postfix);
+		~FPendingScope();
+		void Apply();
+	private:
+		FJsonStringifyImpl* Owner = nullptr;
+		FPendingScope* Outer = nullptr;
+		TFunctionRef<void()> PendingPrefix;
+		TOptional<TFunctionRef<void()>> PendingPostfix;
+
+		bool bHasBeenApplied = false;
+
+		FPendingScope(const FPendingScope& RHS) = delete;
+		FPendingScope(FPendingScope&& RHS) = delete;
+		FPendingScope& operator=(const FPendingScope& RHS) = delete;
+		FPendingScope& operator=(FPendingScope&& RHS) = delete;
+	};
+	FPendingScope* CurrentScope = nullptr;
+
+	TSharedRef<FJsonWriter> Writer;
+
+	const UObject* CurrentObject;
+	TArray<const UObject*> RootObjects;
+	TSet<const UObject*> ObjectsToExport;
+	TSet<const UObject*> ObjectsExported;
+
+	TArray<FCustomVersion> Versions;
+};
+
+}
diff --git a/Engine/Source/Runtime/Experimental/JsonObjectGraph/Private/JsonStringifyStructuredArchive.cpp b/Engine/Source/Runtime/Experimental/JsonObjectGraph/Private/JsonStringifyStructuredArchive.cpp
new file mode 100644
--- /dev/null
+++ b/Engine/Source/Runtime/Experimental/JsonObjectGraph/Private/JsonStringifyStructuredArchive.cpp
@@ -0,0 +1,777 @@
+// Copyright Epic Games, Inc. All Rights Reserved.
+
+#include "JsonStringifyStructuredArchive.h"
+#include "Misc/Base64.h"
+#include "Misc/SecureHash.h"
+#include "JsonObjectGraphConventions.h"
+#include "JsonStringifyImpl.h"
+#include "UObject/Object.h"
+#include "UObject/LazyObjectPtr.h"
+#include "UObject/SoftObjectPtr.h"
+#include "UObject/WeakObjectPtr.h"
+#include "Serialization/Formatters/JsonArchiveOutputFormatter.h"
+
+#include <cinttypes>
+
+#if WITH_TEXT_ARCHIVE_SUPPORT
+
+namespace UE::Private
+{
+
+FJsonStringifyStructuredArchive::FJsonStringifyStructuredArchive(
+	const UObject* JsonObjectGraph
+	, int32 InitialIndentLevel
+	, FJsonStringifyImpl* InRootImpl
+	, TArray<FCustomVersion>& InVersionsToHarvest
+	, bool bFilterEditorOnly)
+	: VersionsToHarvest(&InVersionsToHarvest)
+	, Object(JsonObjectGraph)
+	, RootImpl(InRootImpl)
+	, Inner(ResultBuff)
+{
+	Newline.Add('\n');
+	IndentLevel = InitialIndentLevel;
+	for (int32 I = 0; I < InitialIndentLevel; ++I)
+	{
+		Newline.Add('\t');
+	}
+
+	Inner.SetIsPersistent(true);
+	Inner.SetFilterEditorOnly(bFilterEditorOnly);
+	Inner.SetIsTextFormat(true);
+}
+
+FJsonStringifyStructuredArchive::FJsonStringifyStructuredArchive(FArchive& ToWriter, int32 InitialIndentLevel)
+	: VersionsToHarvest(nullptr)
+	, Object(nullptr)
+	, Inner(ResultBuff)
+	, Override(&ToWriter)
+{
+	Newline.Add('\n');
+	IndentLevel = InitialIndentLevel;
+	for (int32 I = 0; I < InitialIndentLevel; ++I)
+	{
+		Newline.Add('\t');
+	}
+}
+
+TArray<uint8> FJsonStringifyStructuredArchive::ToJson()
+{
+	{
+		FStructuredArchive StructuredArchive(*this);
+		FStructuredArchive::FRecord ExportRecord = StructuredArchive.Open().EnterRecord();
+		const_cast<UObject*>(Object)->Serialize(ExportRecord);
+		// RAII will close the json block.. so if you remove these containing {} you'll lose your trailing '}'
+		// in the Json - I do not think this is a good design for FStructuredArchive.. alternatively 
+		// we can make this a long one liner, but that would also require an explanatory comment
+	}
+	if (ResultBuff.Num() > 2)// Num() > 2 here eliminates the default object {}, empty string, default container etc
+	{
+		if (VersionsToHarvest)
+		{
+			VersionsToHarvest->Append(GetUnderlyingArchive().GetCustomVersions().GetAllVersions());
+		}
+		return ResultBuff;
+	}
+	return TArray<uint8>(); 
+}
+
+void FJsonStringifyStructuredArchive::WriteTextValueInline(FText Value, int32 IndentLevel, FArchive& ToWriter)
+{
+	FJsonStringifyStructuredArchive Formatter(ToWriter, IndentLevel);
+	Formatter.Serialize(Value);
+}
+
+void FJsonStringifyStructuredArchive::WriteCustomVersionValueInline(const TArray<FCustomVersion>& Version, int32 IndentLevel, FArchive& ToWriter)
+{
+	FJsonStringifyStructuredArchive Formatter(ToWriter, IndentLevel);
+	FStructuredArchive ChildArchive(Formatter);
+	ChildArchive.Open() << (TArray<FCustomVersion>&)Version;
+	ChildArchive.Close();
+}
+
+FArchive& FJsonStringifyStructuredArchive::GetUnderlyingArchive()
+{
+	return Override ? *Override : Inner;
+}
+
+bool FJsonStringifyStructuredArchive::HasDocumentTree() const
+{
+	return true;
+}
+
+void FJsonStringifyStructuredArchive::EnterRecord()
+{
+	if (ScopeSkipCount > 0)
+	{
+		return;
+	}
+
+	WriteOptionalComma();
+	WriteOptionalNewline();
+	Write("{");
+	Newline.Add('\t');
+	++IndentLevel;
+	bNeedsNewline = true;
+
+	TextStartPosStack.Push(GetUnderlyingArchive().Tell());
+}
+
+void FJsonStringifyStructuredArchive::LeaveRecord()
+{
+	if (ScopeSkipCount > 0)
+	{
+		return;
+	}
+
+	--IndentLevel;
+	Newline.Pop(EAllowShrinking::No);
+	if (TextStartPosStack.Pop() == GetUnderlyingArchive().Tell())
+	{
+		bNeedsNewline = false;
+	}
+	WriteOptionalNewline();
+	Write("}");
+	bNeedsComma = true;
+	bNeedsNewline = true;
+}
+
+void FJsonStringifyStructuredArchive::EnterField(FArchiveFieldName Name)
+{
+	// The base UObject serializer for structured archives is badly flawed,
+	// so I have disabled it. The reflected properties are handled by FJsonStringifyImpl
+	// so we filter them here. The macro generated BaseClassAutoGen is also
+	// useless, and is handled by the FSerialDataJsonWriter which calls the 
+	// natively provided stream serializer.
+	if (FCString::Strcmp(Name.Name, TEXT("Properties")) == 0 ||
+		FCString::Strcmp(Name.Name, TEXT("BaseClassAutoGen")) == 0 ||
+		ScopeSkipCount != 0)
+	{
+		++ScopeSkipCount;
+	}
+
+	check(ScopeSkipCount >= 0); // scope overflow, halt
+
+	WriteOptionalComma();
+	WriteOptionalNewline();
+	WriteFieldName(Name.Name);
+}
+
+void FJsonStringifyStructuredArchive::LeaveField()
+{
+	const bool bWasSkipping = ScopeSkipCount > 0;
+	if (ScopeSkipCount)
+	{
+		--ScopeSkipCount;
+	}
+
+	if (bWasSkipping)
+	{
+		return;
+	}
+
+	bNeedsComma = true;
+	bNeedsNewline = true;
+}
+
+bool FJsonStringifyStructuredArchive::TryEnterField(FArchiveFieldName Name, bool bEnterWhenSaving)
+{
+	if (bEnterWhenSaving)
+	{
+		EnterField(Name);
+	}
+	return bEnterWhenSaving;
+}
+
+void FJsonStringifyStructuredArchive::EnterArray(int32& NumElements)
+{
+	EnterStream();
+}
+
+void FJsonStringifyStructuredArchive::LeaveArray()
+{
+	LeaveStream();
+}
+
+void FJsonStringifyStructuredArchive::EnterArrayElement()
+{
+	EnterStreamElement();
+}
+
+void FJsonStringifyStructuredArchive::LeaveArrayElement()
+{
+	LeaveStreamElement();
+}
+
+void FJsonStringifyStructuredArchive::EnterStream()
+{
+	if (ScopeSkipCount > 0)
+	{
+		return;
+	}
+
+	WriteOptionalComma();
+	WriteOptionalNewline();
+	Write("[");
+	Newline.Add('\t');
+	++IndentLevel;
+	bNeedsNewline = true;
+	TextStartPosStack.Push(GetUnderlyingArchive().Tell());
+}
+
+void FJsonStringifyStructuredArchive::LeaveStream()
+{
+	if (ScopeSkipCount > 0)
+	{
+		return;
+	}
+
+	--IndentLevel;
+	Newline.Pop(EAllowShrinking::No);
+	if (TextStartPosStack.Pop() == GetUnderlyingArchive().Tell())
+	{
+		bNeedsNewline = false;
+	}
+	WriteOptionalNewline();
+	Write("]");
+	bNeedsComma = true;
+	bNeedsNewline = true;
+}
+
+void FJsonStringifyStructuredArchive::EnterStreamElement()
+{
+	if (ScopeSkipCount > 0)
+	{
+		return;
+	}
+
+	WriteOptionalComma();
+	WriteOptionalNewline();
+}
+
+void FJsonStringifyStructuredArchive::LeaveStreamElement()
+{
+	if (ScopeSkipCount > 0)
+	{
+		return;
+	}
+
+	bNeedsComma = true;
+	bNeedsNewline = true;
+}
+
+void FJsonStringifyStructuredArchive::EnterMap(int32& NumElements)
+{
+	EnterRecord();
+}
+
+void FJsonStringifyStructuredArchive::LeaveMap()
+{
+	LeaveRecord();
+}
+
+void FJsonStringifyStructuredArchive::EnterMapElement(FString& Name)
+{
+	EnterField(FArchiveFieldName(*Name));
+}
+
+void FJsonStringifyStructuredArchive::LeaveMapElement()
+{
+	if (ScopeSkipCount > 0)
+	{
+		return;
+	}
+
+	LeaveField();
+}
+
+void FJsonStringifyStructuredArchive::EnterAttributedValue()
+{
+	if (ScopeSkipCount > 0)
+	{
+		return;
+	}
+
+	NumAttributesStack.Push(0);
+}
+
+void FJsonStringifyStructuredArchive::EnterAttribute(FArchiveFieldName AttributeName)
+{
+	if (ScopeSkipCount > 0)
+	{
+		return;
+	}
+
+	WriteOptionalComma();
+	WriteOptionalNewline();
+	WriteOptionalAttributedBlockOpening();
+	WriteOptionalComma();
+	WriteOptionalNewline();
+	checkf(FCString::Strcmp(AttributeName.Name, TEXT("Value")) != 0, TEXT("Attributes called 'Value' are reserved by the implementation"));
+	WriteFieldName(*FString::Printf(TEXT("_%s"), AttributeName.Name));
+	++NumAttributesStack.Top();
+}
+
+void FJsonStringifyStructuredArchive::LeaveAttribute()
+{
+	if (ScopeSkipCount > 0)
+	{
+		return;
+	}
+
+	bNeedsComma = true;
+	bNeedsNewline = true;
+}
+
+void FJsonStringifyStructuredArchive::LeaveAttributedValue()
+{
+	if (ScopeSkipCount > 0)
+	{
+		return;
+	}
+
+	WriteOptionalAttributedBlockClosing();
+	NumAttributesStack.Pop();
+	bNeedsComma = true;
+	bNeedsNewline = true;
+}
+
+void FJsonStringifyStructuredArchive::EnterAttributedValueValue()
+{
+	WriteOptionalComma();
+	WriteOptionalNewline();
+	WriteOptionalAttributedBlockValue();
+}
+
+bool FJsonStringifyStructuredArchive::TryEnterAttributedValueValue()
+{
+	return false;
+}
+
+bool FJsonStringifyStructuredArchive::TryEnterAttribute(FArchiveFieldName AttributeName, bool bEnterWhenSaving)
+{
+	if (bEnterWhenSaving)
+	{
+		EnterAttribute(AttributeName);
+	}
+	return bEnterWhenSaving;
+}
+
+void FJsonStringifyStructuredArchive::Serialize(uint8& Value)
+{
+	WriteValue(LexToString(Value));
+}
+
+void FJsonStringifyStructuredArchive::Serialize(uint16& Value)
+{
+	WriteValue(LexToString(Value));
+}
+
+void FJsonStringifyStructuredArchive::Serialize(uint32& Value)
+{
+	WriteValue(LexToString(Value));
+}
+
+void FJsonStringifyStructuredArchive::Serialize(uint64& Value)
+{
+	WriteValue(LexToString(Value));
+}
+
+void FJsonStringifyStructuredArchive::Serialize(int8& Value)
+{
+	WriteValue(LexToString(Value));
+}
+
+void FJsonStringifyStructuredArchive::Serialize(int16& Value)
+{
+	WriteValue(LexToString(Value));
+}
+
+void FJsonStringifyStructuredArchive::Serialize(int32& Value)
+{
+	WriteValue(LexToString(Value));
+}
+
+void FJsonStringifyStructuredArchive::Serialize(int64& Value)
+{
+	WriteValue(LexToString(Value));
+}
+
+void FJsonStringifyStructuredArchive::Serialize(float& Value)
+{
+	if (FPlatformMath::IsFinite(Value))
+	{
+		FString String = FString::Printf(TEXT("%.17g"), Value);
+#if DO_GUARD_SLOW
+		float RoundTripped;
+		LexFromString(RoundTripped, *String);
+		check(RoundTripped == Value);
+#endif
+		WriteValue(String);
+	}
+	else if (FPlatformMath::IsNaN(Value))
+	{
+		const uint32 ValueAsInt = BitCast<uint32>(Value);
+		const bool bIsNegative = !!(ValueAsInt & 0x80000000);
+		const uint32 Significand = ValueAsInt & 0x007fffff;
+		WriteValue(FString::Printf(TEXT("\"Number:%snan:0x%" PRIx32 "\""), bIsNegative ? TEXT("-") : TEXT("+"), Significand));
+	}
+	else
+	{
+		WriteValue(Value < 0.0f ? TEXT("\"Number:-inf\"") : TEXT("\"Number:+inf\""));
+	}
+}
+
+void FJsonStringifyStructuredArchive::Serialize(double& Value)
+{
+	if (FPlatformMath::IsFinite(Value))
+	{
+		FString String = FString::Printf(TEXT("%.17g"), Value);
+#if DO_GUARD_SLOW
+		double RoundTripped;
+		LexFromString(RoundTripped, *String);
+		check(RoundTripped == Value);
+#endif
+		WriteValue(String);
+	}
+	else if (FPlatformMath::IsNaN(Value))
+	{
+		const uint64 ValueAsInt = BitCast<uint64>(Value);
+		const bool bIsNegative = !!(ValueAsInt & 0x8000000000000000);
+		const uint64 Significand = ValueAsInt & 0x000fffffffffffff;
+		WriteValue(FString::Printf(TEXT("\"Number:%snan:0x%" PRIx64 "\""), bIsNegative ? TEXT("-") : TEXT("+"), Significand));
+	}
+	else
+	{
+		WriteValue(Value < 0.0 ? TEXT("\"Number:-inf\"") : TEXT("\"Number:+inf\""));
+	}
+}
+
+void FJsonStringifyStructuredArchive::Serialize(bool& Value)
+{
+	WriteValue(LexToString(Value));
+}
+
+void FJsonStringifyStructuredArchive::Serialize(FString& Value)
+{
+	// Insert a "String:" prefix to prevent incorrect interpretation as another explicit type
+	if (Value.StartsWith(TEXT("Object:")) || Value.StartsWith(TEXT("String:")) || Value.StartsWith(TEXT("Base64:")))
+	{
+		SerializeStringInternal(FString::Printf(TEXT("String:%s"), *Value));
+	}
+	else
+	{
+		SerializeStringInternal(Value);
+	}
+}
+
+void FJsonStringifyStructuredArchive::Serialize(FName& Value)
+{
+	SerializeStringInternal(*Value.ToString());
+}
+
+void FJsonStringifyStructuredArchive::Serialize(UObject*& Value)
+{
+	if (ScopeSkipCount > 0)
+	{
+		return;
+	}
+
+	RootImpl->WriteObjectAsJsonToArchive(Object, Value, &Inner, IndentLevel);
+}
+
+void FJsonStringifyStructuredArchive::Serialize(FText& Value)
+{
+	/* 
+		We could write using the structured serializer defined by FText but
+		it is not as widely used as FTextStringHelper::WriteToBuffer and 
+		is currently private:
+	FStructuredArchive ChildArchive(*this);
+	FText::SerializeText(ChildArchive.Open(), Value);
+	ChildArchive.Close();
+	*/
+
+	FString AsString;
+	FTextStringHelper::WriteToBuffer(AsString, Value);
+	Serialize(AsString);
+}
+
+void FJsonStringifyStructuredArchive::Serialize(FWeakObjectPtr& Value)
+{
+	UObject* Ptr = Value.IsValid() ? Value.Get() : nullptr;
+	Serialize(Ptr);
+}
+
+void FJsonStringifyStructuredArchive::Serialize(FSoftObjectPtr& Value)
+{
+	FSoftObjectPath Path = Value.ToSoftObjectPath();
+	Serialize(Path);
+}
+
+void FJsonStringifyStructuredArchive::Serialize(FSoftObjectPath& Value)
+{
+	FString ValueStr;
+	Value.ExportTextItem(ValueStr, FSoftObjectPath(), nullptr, 0, nullptr);
+	Serialize(ValueStr);
+}
+
+void FJsonStringifyStructuredArchive::Serialize(FLazyObjectPtr& Value)
+{
+	UObject* ObjectResolved = Value.Get();
+	Serialize(ObjectResolved);
+}
+
+void FJsonStringifyStructuredArchive::Serialize(FObjectPtr& Value)
+{
+	UObject* ResolvedObject = Value.Get();
+	Serialize(ResolvedObject);
+}
+
+void FJsonStringifyStructuredArchive::Serialize(TArray<uint8>& Data)
+{
+	Serialize(Data.GetData(), Data.Num());
+}
+
+void FJsonStringifyStructuredArchive::Serialize(void* Data, uint64 DataSize)
+{
+	if (ScopeSkipCount > 0)
+	{
+		return;
+	}
+
+	static const int32 MaxLineChars = 120;
+	static const int32 MaxLineBytes = FBase64::GetMaxDecodedDataSize(MaxLineChars);
+
+	if (DataSize < MaxLineBytes)
+	{
+		// Encode the data on a single line. No need for hashing; intra-line merge conflicts are rare.
+		WriteValue(FString::Printf(TEXT("\"Base64:%s\""), *FBase64::Encode((const uint8*)Data, static_cast<uint32>(DataSize))));
+	}
+	else
+	{
+		// Encode the data as a record containing a digest and array of base-64 encoded lines
+		EnterRecord();
+		GetUnderlyingArchive().Serialize(Newline.GetData(), Newline.Num());
+
+		// Compute a SHA digest for the raw data, so we can check if it's corrupted
+		uint8 Digest[FSHA1::DigestSize];
+		FSHA1::HashBuffer(Data, DataSize, Digest);
+
+		// Convert the hash to a string
+		ANSICHAR DigestString[(FSHA1::DigestSize * 2) + 1];
+		for (int32 Idx = 0; Idx < UE_ARRAY_COUNT(Digest); Idx++)
+		{
+			static const ANSICHAR HexDigits[] = "0123456789abcdef";
+			DigestString[(Idx * 2) + 0] = HexDigits[Digest[Idx] >> 4];
+			DigestString[(Idx * 2) + 1] = HexDigits[Digest[Idx] & 15];
+		}
+		DigestString[UE_ARRAY_COUNT(DigestString) - 1] = 0;
+
+		FArchive& Writer = GetUnderlyingArchive();
+		// Write the digest
+		Write("\"Digest\": \"");
+		Write(DigestString);
+		Write("\",");
+		Writer.Serialize(Newline.GetData(), Newline.Num());
+
+		// Write the base64 data
+		Write("\"Base64\": ");
+		for (uint64 DataPos = 0; DataPos < DataSize; DataPos += MaxLineBytes)
+		{
+			Write((DataPos > 0) ? ',' : '[');
+			Writer.Serialize(Newline.GetData(), Newline.Num());
+			Write("\t\"");
+
+			ANSICHAR LineData[MaxLineChars + 1];
+			uint64 NumLineChars = FBase64::Encode((const uint8*)Data + DataPos, FMath::Min<uint32>(IntCastChecked<uint32>(DataSize - DataPos), MaxLineBytes), LineData);
+			Writer.Serialize(LineData, NumLineChars);
+
+			Write("\"");
+		}
+
+		// Close the array
+		Writer.Serialize(Newline.GetData(), Newline.Num());
+		Write(']');
+		bNeedsNewline = true;
+
+		// Close the record
+		LeaveRecord();
+	}
+}
+
+void FJsonStringifyStructuredArchive::Write(ANSICHAR Character)
+{
+	if (ScopeSkipCount > 0)
+	{
+		return;
+	}
+
+	GetUnderlyingArchive().Serialize((void*)&Character, 1);
+}
+
+void FJsonStringifyStructuredArchive::Write(const ANSICHAR* Text)
+{
+	if (ScopeSkipCount > 0)
+	{
+		return;
+	}
+
+	GetUnderlyingArchive().Serialize((void*)Text, TCString<ANSICHAR>::Strlen(Text));
+}
+
+void FJsonStringifyStructuredArchive::Write(const FString& Text)
+{
+	Write(TCHAR_TO_UTF8(*Text));
+}
+
+void FJsonStringifyStructuredArchive::WriteFieldName(const TCHAR* Name)
+{
+	if (FCString::Stricmp(Name, TEXT("Base64")) == 0 || FCString::Stricmp(Name, TEXT("Digest")) == 0)
+	{
+		Write(FString::Printf(TEXT("\"_%s\": "), Name));
+	}
+	else if (Name[0] == '_')
+	{
+		Write(FString::Printf(TEXT("\"_%s\": "), Name));
+	}
+	else
+	{
+		Write(FString::Printf(TEXT("\"%s\": "), Name));
+	}
+}
+
+void FJsonStringifyStructuredArchive::WriteValue(const FString& Text)
+{
+	Write(Text);
+}
+
+void FJsonStringifyStructuredArchive::WriteOptionalComma()
+{
+	if (ScopeSkipCount > 0)
+	{
+		return;
+	}
+
+	if (bNeedsComma)
+	{
+		Write(',');
+		bNeedsComma = false;
+	}
+}
+
+void FJsonStringifyStructuredArchive::WriteOptionalNewline()
+{
+	if (ScopeSkipCount > 0)
+	{
+		return;
+	}
+
+	if (bNeedsNewline)
+	{
+		GetUnderlyingArchive().Serialize(Newline.GetData(), Newline.Num());
+		bNeedsNewline = false;
+	}
+}
+
+void FJsonStringifyStructuredArchive::WriteOptionalAttributedBlockOpening()
+{
+	if (ScopeSkipCount > 0)
+	{
+		return;
+	}
+
+	if (NumAttributesStack.Top() == 0)
+	{
+		Write('{');
+		Newline.Add('\t');
+		++IndentLevel;
+		bNeedsNewline = true;
+	}
+}
+
+void FJsonStringifyStructuredArchive::WriteOptionalAttributedBlockValue()
+{
+	if (ScopeSkipCount > 0)
+	{
+		return;
+	}
+
+	if (NumAttributesStack.Top() != 0)
+	{
+		WriteFieldName(TEXT("_Value"));
+	}
+}
+
+void FJsonStringifyStructuredArchive::WriteOptionalAttributedBlockClosing()
+{
+	if (ScopeSkipCount > 0)
+	{
+		return;
+	}
+
+	if (NumAttributesStack.Top() != 0)
+	{
+		--IndentLevel;
+		Newline.Pop(EAllowShrinking::No);
+		WriteOptionalNewline();
+		Write("}");
+		bNeedsComma = true;
+		bNeedsNewline = true;
+	}
+}
+
+void FJsonStringifyStructuredArchive::SerializeStringInternal(const FString& String)
+{
+	if (ScopeSkipCount > 0)
+	{
+		return;
+	}
+
+	FString Result = TEXT("\"");
+
+	// Escape the string characters
+	for (int32 Idx = 0; Idx < String.Len(); Idx++)
+	{
+		switch (String[Idx])
+		{
+		case '\"':
+			Result += "\\\"";
+			break;
+		case '\\':
+			Result += "\\\\";
+			break;
+		case '\b':
+			Result += "\\b";
+			break;
+		case '\f':
+			Result += "\\f";
+			break;
+		case '\n':
+			Result += "\\n";
+			break;
+		case '\r':
+			Result += "\\r";
+			break;
+		case '\t':
+			Result += "\\t";
+			break;
+		default:
+			if (String[Idx] <= 0x1f || String[Idx] >= 0x7f)
+			{
+				Result += FString::Printf(TEXT("\\u%04x"), String[Idx]);
+			}
+			else
+			{
+				Result.AppendChar(String[Idx]);
+			}
+			break;
+		}
+	}
+	Result += TEXT("\"");
+
+	WriteValue(Result);
+}
+
+}
+
+#endif // WITH_TEXT_ARCHIVE_SUPPORT
diff --git a/Engine/Source/Runtime/Experimental/JsonObjectGraph/Private/JsonStringifyStructuredArchive.h b/Engine/Source/Runtime/Experimental/JsonObjectGraph/Private/JsonStringifyStructuredArchive.h
new file mode 100644
--- /dev/null
+++ b/Engine/Source/Runtime/Experimental/JsonObjectGraph/Private/JsonStringifyStructuredArchive.h
@@ -0,0 +1,136 @@
+// Copyright Epic Games, Inc. All Rights Reserved.
+
+#pragma once
+
+#include "Serialization/StructuredArchiveFormatter.h"
+#include "Serialization/CustomVersion.h"
+#include "Serialization/MemoryWriter.h"
+
+#if WITH_TEXT_ARCHIVE_SUPPORT
+
+namespace UE::Private
+{
+
+struct FJsonStringifyImpl;
+
+/*
+	FJsonStringifyStructuredArchive is an implementation detail of JsonObjectGraph. It provides support
+	for saving the relatively rare data that provides a native Serialize(FStructuredArchive::FRecord) 
+	but explicitly disables serialization of reflected properties and versioning data, which are handled
+	by the root implementation
+*/
+struct FJsonStringifyStructuredArchive : private FStructuredArchiveFormatter
+{
+	FJsonStringifyStructuredArchive(
+		const UObject* InObject, 
+		int32 InitialIndentLevel, 
+		FJsonStringifyImpl* InRootImpl,
+		TArray<FCustomVersion>& InVersionsToHarvest,
+		bool bFilterEditorOnly);
+	
+	// FMemoryWriter wants a byte array, so that's what we're working with at this level:
+	TArray<uint8> ToJson();
+
+	// Unsure if we want to use the structured archive for FText or the string encoding
+	// long term, this provides us the option:
+	static void WriteTextValueInline(FText Value, int32 IndentLevel, FArchive& ToWriter);
+	// FCustomVersion is best encoded via operator<<(FStructuredArchive::FSlot Slot, FCustomVersion& Version):
+	static void WriteCustomVersionValueInline(const TArray<FCustomVersion>& Version, int32 IndentLevel, FArchive& ToWriter);
+private:
+	FJsonStringifyStructuredArchive(FArchive& ToWriter, int32 InitialIndentLevel);
+
+	virtual FArchive& GetUnderlyingArchive() override;
+
+	virtual bool HasDocumentTree() const override;
+
+	virtual void EnterRecord() override;
+	virtual void LeaveRecord() override;
+	virtual void EnterField(FArchiveFieldName Name) override;
+	virtual void LeaveField() override;
+	virtual bool TryEnterField(FArchiveFieldName Name, bool bEnterWhenSaving) override;
+
+	virtual void EnterArray(int32& NumElements) override;
+	virtual void LeaveArray() override;
+	virtual void EnterArrayElement() override;
+	virtual void LeaveArrayElement() override;
+
+	virtual void EnterStream() override;
+	virtual void LeaveStream() override;
+	virtual void EnterStreamElement() override;
+	virtual void LeaveStreamElement() override;
+
+	virtual void EnterMap(int32& NumElements) override;
+	virtual void LeaveMap() override;
+	virtual void EnterMapElement(FString& Name) override;
+	virtual void LeaveMapElement() override;
+
+	virtual void EnterAttributedValue() override;
+	virtual void EnterAttribute(FArchiveFieldName AttributeName) override;
+	virtual void EnterAttributedValueValue() override;
+	virtual void LeaveAttribute() override;
+	virtual void LeaveAttributedValue() override;
+	virtual bool TryEnterAttribute(FArchiveFieldName AttributeName, bool bEnterWhenSaving) override;
+
+	virtual bool TryEnterAttributedValueValue() override;
+
+	virtual void Serialize(uint8& Value) override;
+	virtual void Serialize(uint16& Value) override;
+	virtual void Serialize(uint32& Value) override;
+	virtual void Serialize(uint64& Value) override;
+	virtual void Serialize(int8& Value) override;
+	virtual void Serialize(int16& Value) override;
+	virtual void Serialize(int32& Value) override;
+	virtual void Serialize(int64& Value) override;
+	virtual void Serialize(float& Value) override;
+	virtual void Serialize(double& Value) override;
+	virtual void Serialize(bool& Value) override;
+	virtual void Serialize(FString& Value) override;
+	virtual void Serialize(FName& Value) override;
+	virtual void Serialize(UObject*& Value) override;
+	virtual void Serialize(FText& Value) override;
+	virtual void Serialize(FWeakObjectPtr& Value) override;
+	virtual void Serialize(FSoftObjectPtr& Value) override;
+	virtual void Serialize(FSoftObjectPath& Value) override;
+	virtual void Serialize(FLazyObjectPtr& Value) override;
+	virtual void Serialize(FObjectPtr& Value) override;
+	virtual void Serialize(TArray<uint8>& Value) override;
+	virtual void Serialize(void* Data, uint64 DataSize) override;
+
+	void Write(ANSICHAR Character);
+
+	void Write(const ANSICHAR* Text);
+	void Write(const FString& Text);
+
+	void WriteFieldName(const TCHAR* Name);
+	void WriteValue(const FString& Value);
+
+	void WriteOptionalComma();
+	void WriteOptionalNewline();
+	void WriteOptionalAttributedBlockOpening();
+	void WriteOptionalAttributedBlockValue();
+	void WriteOptionalAttributedBlockClosing();
+
+	void SerializeStringInternal(const FString& String);
+
+	TArray<FCustomVersion>* VersionsToHarvest;
+	TArray<ANSICHAR> Newline;
+
+	TArray<int32> NumAttributesStack;
+	TArray<int64> TextStartPosStack;
+
+	const UObject* const Object;
+	FJsonStringifyImpl* RootImpl;
+
+	TArray<uint8> ResultBuff;
+	FMemoryWriter Inner;
+	FArchive* Override = nullptr;
+
+	int64_t ScopeSkipCount = 0;
+	int32_t IndentLevel = 0;
+	bool bNeedsComma = false;
+	bool bNeedsNewline = false;
+};
+
+}
+
+#endif
diff --git a/Engine/Source/Runtime/Experimental/JsonObjectGraph/Private/PrettyJsonWriter.cpp b/Engine/Source/Runtime/Experimental/JsonObjectGraph/Private/PrettyJsonWriter.cpp
new file mode 100644
--- /dev/null
+++ b/Engine/Source/Runtime/Experimental/JsonObjectGraph/Private/PrettyJsonWriter.cpp
@@ -0,0 +1,164 @@
+// Copyright Epic Games, Inc. All Rights Reserved.
+
+#include "PrettyJsonWriter.h"
+#include "JsonStringifyStructuredArchive.h" // for inline structured data, e.g. FText
+
+namespace UE::Private
+{
+
+FPrettyJsonWriter::FPrettyJsonWriter(FArchive* const InStream, int32 InitialIndentLevel)
+	: TJsonWriter<ANSICHAR, TPrettyJsonPrintPolicySingleNewLine>(InStream, InitialIndentLevel)
+{
+}
+
+TSharedRef<FPrettyJsonWriter> FPrettyJsonWriter::Create(FArchive* const InStream, int32 InitialIndent)
+{
+	return MakeShareable(new FPrettyJsonWriter(InStream, InitialIndent));
+}
+
+void FPrettyJsonWriter::WriteJsonRaw(FAnsiStringView Value)
+{
+	check(CanWriteValueWithoutIdentifier());
+	WriteCommaAndNewlineIfNeeded();
+	PrintPolicy::WriteString(Stream, Value);
+	PreviousTokenWritten = EJsonToken::String;
+}
+
+void FPrettyJsonWriter::WriteValueInline(FText Value)
+{
+	// use a structured archive to write FText nicely:
+	#if WITH_TEXT_ARCHIVE_SUPPORT
+	WriteCommaAndNewlineIfNeeded();
+	FJsonStringifyStructuredArchive::WriteTextValueInline(Value, GetIndentLevel(), *Stream);
+	PreviousTokenWritten = EJsonToken::String;
+	#else
+	FString AsString;
+	FTextStringHelper::WriteToBuffer(AsString, Value);
+	WriteValueInline(AsString);
+	#endif
+}
+
+void FPrettyJsonWriter::WriteValueInline(const FString& Value)
+{
+	check(CanWriteValueWithoutIdentifier());
+	WriteCommaAndNewlineIfNeeded();
+	PreviousTokenWritten = Super::WriteValueOnly(Value);
+}
+
+void FPrettyJsonWriter::WriteValueInline(FAnsiStringView UTF8Value)
+{
+	check(CanWriteValueWithoutIdentifier());
+	WriteCommaAndNewlineIfNeeded();
+	Super::WriteStringValue(UTF8Value);
+	PreviousTokenWritten = EJsonToken::String; 
+}
+
+void FPrettyJsonWriter::WriteValueInline(FUtf8StringView UTF8Value)
+{
+	check(CanWriteValueWithoutIdentifier());
+	WriteCommaAndNewlineIfNeeded();
+	// @todo: I think super should be utf8 aware?
+	FAnsiStringView Hack((char*)UTF8Value.GetData(), UTF8Value.Len());
+	Super::WriteStringValue(Hack);
+	PreviousTokenWritten = EJsonToken::String;
+}
+
+void FPrettyJsonWriter::WriteValueInline(int16 Value)
+{
+	WriteCommaAndNewlineIfNeeded();
+	PreviousTokenWritten = Super::WriteValueOnly((int64)Value);
+}
+
+void FPrettyJsonWriter::WriteValueInline(uint16 Value)
+{
+	WriteCommaAndNewlineIfNeeded();
+	PreviousTokenWritten = Super::WriteValueOnly((uint64)Value);
+}
+
+void FPrettyJsonWriter::WriteValueInline(uint32 Value)
+{
+	WriteCommaAndNewlineIfNeeded();
+	PreviousTokenWritten = Super::WriteValueOnly((uint64)Value);
+}
+
+void FPrettyJsonWriter::WriteUtf8Value(FStringView Identifier, FUtf8StringView UTF8Value)
+{
+	check(Stack.Top() == EJson::Object);
+	WriteIdentifier(Forward<FStringView>(Identifier));
+
+	PrintPolicy::WriteSpace(Stream);
+
+	// @todo: I think super should be utf8 aware?
+	FAnsiStringView Hack((char*)UTF8Value.GetData(), UTF8Value.Len());
+	Super::WriteStringValue(Hack);
+	PreviousTokenWritten = EJsonToken::String;
+}
+
+void FPrettyJsonWriter::WriteObjectStartInline()
+{
+	check(CanWriteObjectWithoutIdentifier());
+	if (PreviousTokenWritten != EJsonToken::None)
+	{
+		WriteCommaAndNewlineIfNeeded();
+	}
+
+	PrintPolicy::WriteChar(Stream, CharType('{'));
+	++IndentLevel;
+	Stack.Push(EJson::Object);
+	PreviousTokenWritten = EJsonToken::CurlyOpen;
+}
+
+void FPrettyJsonWriter::WriteArrayStartInline()
+{
+	check(CanWriteObjectWithoutIdentifier());
+	if (PreviousTokenWritten != EJsonToken::None)
+	{
+		WriteCommaAndNewlineIfNeeded();
+	}
+
+	PrintPolicy::WriteChar(Stream, CharType('['));
+	++IndentLevel;
+	Stack.Push(EJson::Array);
+	PreviousTokenWritten = EJsonToken::SquareOpen;
+}
+
+void FPrettyJsonWriter::WriteNewlineAndArrayEnd()
+{
+	check(Stack.Top() == EJson::Array);
+
+	--IndentLevel;
+
+	PrintPolicy::WriteLineTerminator(Stream);
+	PrintPolicy::WriteTabs(Stream, IndentLevel);
+	PrintPolicy::WriteChar(Stream, CharType(']'));
+	Stack.Pop();
+	PreviousTokenWritten = EJsonToken::SquareClose;
+}
+
+void FPrettyJsonWriter::WriteLineTerminator()
+{
+	PrintPolicy::WriteLineTerminator(Stream);
+	PrintPolicy::WriteTabs(Stream, IndentLevel);
+}
+
+void FPrettyJsonWriter::HACK_SetPreviousTokenWritten()
+{
+	// so that we can use this to write inline object references in structured archives..
+	// we don't want any commas written... currently relied upon by UMetaData, 
+	// which will write object references out using structured archive
+	PreviousTokenWritten = EJsonToken::CurlyOpen; 
+}
+
+void FPrettyJsonWriter::WriteCommaAndNewlineIfNeeded()
+{
+	if (PreviousTokenWritten != EJsonToken::CurlyOpen && 
+		PreviousTokenWritten != EJsonToken::SquareOpen && 
+		PreviousTokenWritten != EJsonToken::Identifier)
+	{
+		PrintPolicy::WriteChar(Stream, CharType(','));
+		PrintPolicy::WriteLineTerminator(Stream);
+		PrintPolicy::WriteTabs(Stream, IndentLevel);
+	}
+}
+
+}
diff --git a/Engine/Source/Runtime/Experimental/JsonObjectGraph/Private/PrettyJsonWriter.h b/Engine/Source/Runtime/Experimental/JsonObjectGraph/Private/PrettyJsonWriter.h
new file mode 100644
--- /dev/null
+++ b/Engine/Source/Runtime/Experimental/JsonObjectGraph/Private/PrettyJsonWriter.h
@@ -0,0 +1,58 @@
+// Copyright Epic Games, Inc. All Rights Reserved.
+
+#pragma once
+
+#include "Serialization/JsonWriter.h"
+
+namespace UE::Private
+{
+
+struct TPrettyJsonPrintPolicySingleNewLine
+	: public TPrettyJsonPrintPolicy<ANSICHAR>
+{
+	static inline void WriteLineTerminator(FArchive* Stream)
+	{
+		TJsonPrintPolicy<ANSICHAR>::WriteString(Stream, FStringView(TEXT("\n")));
+	}
+};
+
+// Extending TJsonWriter to get some control over whitespace and tabbing:
+struct FPrettyJsonWriter : public TJsonWriter<ANSICHAR, TPrettyJsonPrintPolicySingleNewLine>
+{
+	using Super = TJsonWriter<ANSICHAR, TPrettyJsonPrintPolicySingleNewLine>;
+	using PrintPolicy = TPrettyJsonPrintPolicySingleNewLine;
+	using CharType = ANSICHAR;
+
+	FPrettyJsonWriter(FArchive* const InStream, int32 InitialIndentLevel);
+
+	static TSharedRef<FPrettyJsonWriter> Create(FArchive* const InStream, int32 InitialIndent = 0);
+
+	// useful for composing blocks of json
+	void WriteJsonRaw(FAnsiStringView Value);
+
+	void WriteValueInline(FText Value);
+	void WriteValueInline(const FString& Value);
+	void WriteValueInline(FAnsiStringView UTF8Value);
+	void WriteValueInline(FUtf8StringView UTF8Value);
+	void WriteValueInline(int16 Value);
+	void WriteValueInline(uint16 Value);
+	void WriteValueInline(uint32 Value);
+	template<typename T>
+	void WriteValueInline(T Value)
+	{
+		check(CanWriteValueWithoutIdentifier());
+		WriteCommaAndNewlineIfNeeded();
+		PreviousTokenWritten = WriteValueOnly(Value);
+	}
+	void WriteUtf8Value(FStringView Identifier, FUtf8StringView UTF8Value);
+
+	void WriteObjectStartInline();
+	void WriteArrayStartInline();
+	void WriteNewlineAndArrayEnd();
+	void WriteLineTerminator();
+	void HACK_SetPreviousTokenWritten();
+
+	void WriteCommaAndNewlineIfNeeded();
+};
+
+}
diff --git a/Engine/Source/Runtime/Experimental/JsonObjectGraph/Public/JsonObjectGraph/Stringify.h b/Engine/Source/Runtime/Experimental/JsonObjectGraph/Public/JsonObjectGraph/Stringify.h
new file mode 100644
--- /dev/null
+++ b/Engine/Source/Runtime/Experimental/JsonObjectGraph/Public/JsonObjectGraph/Stringify.h
@@ -0,0 +1,53 @@
+// Copyright Epic Games, Inc. All Rights Reserved.
+
+#pragma once
+
+#include "Containers/Utf8String.h"
+#include "UObject/ObjectMacros.h"
+
+#include "Stringify.generated.h"
+
+class UObject;
+
+UENUM(BlueprintType, Flags, meta = (Bitflags, UseEnumValuesAsMaskValuesInEditor = "true"))
+enum class EJsonStringifyFlags : uint8
+{
+	/** The default representation of an object attempts to be complete and stable across time */
+	Default = 0,
+	/** Filters editor only data such that it is not written to the Json */
+	FilterEditorOnlyData = 1 << 1,
+};
+ENUM_CLASS_FLAGS(EJsonStringifyFlags);
+
+USTRUCT(BlueprintType)
+struct FJsonStringifyOptions
+{
+	GENERATED_BODY()
+	
+	FJsonStringifyOptions(EJsonStringifyFlags InFlags)
+		: Flags(InFlags)
+	{
+	}
+
+	FJsonStringifyOptions() = default;
+	~FJsonStringifyOptions() = default;
+	FJsonStringifyOptions(const FJsonStringifyOptions&) = default;
+	FJsonStringifyOptions(FJsonStringifyOptions&&) = default;
+	FJsonStringifyOptions& operator=(const FJsonStringifyOptions&) = default;
+	FJsonStringifyOptions& operator=(FJsonStringifyOptions&&) = default;
+
+	UPROPERTY(BlueprintReadWrite, Category = "Options")
+	EJsonStringifyFlags Flags = EJsonStringifyFlags::Default;
+};
+
+namespace UE::JsonObjectGraph
+{
+	/** 
+	 * ! EXPERIMENTAL ! contents of return string will change. Currently this is used
+	 * as a debugging faculty.
+	 * 
+	 * @return: A string containing single JSON object with serialized representations of the 
+	 * provided Objects in that single object's __RootObjects field.
+	 */
+	JSONOBJECTGRAPH_API FUtf8String Stringify(TConstArrayView<const UObject*> RootObjects, const FJsonStringifyOptions& Options = FJsonStringifyOptions());
+}
