Added cvar "-iostore.TreatMissingOnDemandChunksAsError" for toggling missing chunks as errors or warnings when installing on-demand content

#rnx

diff --git a/Engine/Source/Runtime/Experimental/IoStore/OnDemand/Private/OnDemandIoStore.cpp b/Engine/Source/Runtime/Experimental/IoStore/OnDemand/Private/OnDemandIoStore.cpp
--- a/Engine/Source/Runtime/Experimental/IoStore/OnDemand/Private/OnDemandIoStore.cpp
+++ b/Engine/Source/Runtime/Experimental/IoStore/OnDemand/Private/OnDemandIoStore.cpp
@@ -39,13 +39,22 @@
 
 ///////////////////////////////////////////////////////////////////////////////
 bool GIoStoreOnDemandInstallCacheEnabled = true;
-static FAutoConsoleVariableRef CVar_IoStoreOnDemandInstallCacheEnabled (
+static FAutoConsoleVariableRef CVar_IoStoreOnDemandInstallCacheEnabled(
 	TEXT("iostore.OnDemandInstallCacheEnabled"),
 	GIoStoreOnDemandInstallCacheEnabled,
 	TEXT("Whether the on-demand install cache is enabled."),
 	ECVF_ReadOnly
 );
 
+bool GIoStoreOnDemandTreatMissingChunksAsError = true;
+static FAutoConsoleVariableRef CVar_IoStoreOnDemandTreatMissingChunksAsError(
+	TEXT("iostore.TreatMissingOnDemandChunksAsError"),
+	GIoStoreOnDemandTreatMissingChunksAsError,
+	TEXT("Whether to treat missing chunks as error when installing on-demand content."),
+	ECVF_ReadOnly
+);
+
+///////////////////////////////////////////////////////////////////////////////
 namespace UE::IoStore
 {
 
@@ -1239,18 +1248,28 @@ FOnDemandInstallResult FOnDemandIoStore::TickInstallRequest(const FInstallReques
 	}
 
 	// Check the other I/O backends for missing package chunks
+	uint32 MissingCount = 0;
 	for (const FPackageId& PackageId : Missing)
 	{
 		const FIoChunkId ChunkId = CreatePackageDataChunkId(PackageId);
 		if (FIoDispatcher::Get().DoesChunkExist(ChunkId) == false)
 		{
-			UE_LOG(LogIoStoreOnDemand, Error, TEXT("Missing package chunk '%s'"), *LexToString(ChunkId));
-			Status = FIoStatusBuilder(EIoErrorCode::UnknownChunkID) <<
-				TEXT("Missing package chunk '") << LexToString(ChunkId) << TEXT("'");
-			return OutResult;
+			UE_CLOG(MissingCount == 0, LogIoStoreOnDemand, Warning, TEXT("Failed to resolve the following chunk(s) for content handle '%s':"),
+				*LexToString(InstallRequest.Args.ContentHandle));
+
+			UE_LOG(LogIoStoreOnDemand, Warning, TEXT("ChunkId='%s'"), *LexToString(ChunkId));
+			MissingCount++;
 		}
 	}
 
+	if (MissingCount > 0 && GIoStoreOnDemandTreatMissingChunksAsError)
+	{
+		UE_LOG(LogIoStoreOnDemand, Error, TEXT("Install request failed to due missing chunk(s), ContentHandle='%s'"),
+			*LexToString(InstallRequest.Args.ContentHandle));
+		Status = FIoStatus(EIoErrorCode::UnknownChunkID);
+		return OutResult;
+	}
+
 	if (CheckAndSetCancelled())
 	{
 		return OutResult;
