Dataflow node fixes
- Fixed SelectByIntAttribute so it does not test for float type instead of int32
- Also simplified the SelectByIntAttribute and SelectByFloatAttribute by using FindAttributeTyped instead
- Fixed code reading the "TransformIndex" attribute in the Transform group instead of the GeometryGroup , resulting in failure in the selection convert node

diff --git a/Engine/Plugins/Experimental/Fracture/Source/FractureEngine/Private/FractureEngineSelection.cpp b/Engine/Plugins/Experimental/Fracture/Source/FractureEngine/Private/FractureEngineSelection.cpp
--- a/Engine/Plugins/Experimental/Fracture/Source/FractureEngine/Private/FractureEngineSelection.cpp
+++ b/Engine/Plugins/Experimental/Fracture/Source/FractureEngine/Private/FractureEngineSelection.cpp
@@ -209,7 +209,7 @@ void FFractureEngineSelection::SelectLevel(const FManagedArrayCollection& Collec
 void FFractureEngineSelection::SelectContact(FGeometryCollection& GeometryCollection, TArray<int32>& SelectedBones)
 {
 	if (GeometryCollection.HasGroup(FGeometryCollection::TransformGroup) &&
-		GeometryCollection.HasAttribute("TransformIndex", FGeometryCollection::TransformGroup) &&
+		GeometryCollection.HasAttribute("TransformIndex", FGeometryCollection::GeometryGroup) &&
 		GeometryCollection.HasAttribute("TransformToGeometryIndex", FGeometryCollection::TransformGroup) &&
 		GeometryCollection.HasGroup(FGeometryCollection::GeometryGroup) &&
 		GeometryCollection.HasAttribute("Proximity", FGeometryCollection::GeometryGroup))
@@ -217,7 +217,7 @@ void FFractureEngineSelection::SelectContact(FGeometryCollection& GeometryCollec
 		FGeometryCollectionProximityUtility ProximityUtility(&GeometryCollection);
 		ProximityUtility.RequireProximity();
 
-		const TManagedArray<int32>& TransformIndex = GeometryCollection.GetAttribute<int32>("TransformIndex", FGeometryCollection::TransformGroup);
+		const TManagedArray<int32>& TransformIndex = GeometryCollection.GetAttribute<int32>("TransformIndex", FGeometryCollection::GeometryGroup);
 		const TManagedArray<int32>& TransformToGeometryIndex = GeometryCollection.GetAttribute<int32>("TransformToGeometryIndex", FGeometryCollection::TransformGroup);
 		const TManagedArray<TSet<int32>>& Proximity = GeometryCollection.GetAttribute<TSet<int32>>("Proximity", FGeometryCollection::GeometryGroup);
 
diff --git a/Engine/Source/Runtime/Experimental/Chaos/Private/GeometryCollection/Facades/CollectionTransformSelectionFacade.cpp b/Engine/Source/Runtime/Experimental/Chaos/Private/GeometryCollection/Facades/CollectionTransformSelectionFacade.cpp
--- a/Engine/Source/Runtime/Experimental/Chaos/Private/GeometryCollection/Facades/CollectionTransformSelectionFacade.cpp
+++ b/Engine/Source/Runtime/Experimental/Chaos/Private/GeometryCollection/Facades/CollectionTransformSelectionFacade.cpp
@@ -16,8 +16,8 @@ namespace GeometryCollection::Facades
 		, ChildrenAttribute(InCollection, "Children", FTransformCollection::TransformGroup)
 		, LevelAttribute(InCollection, "Level", FTransformCollection::TransformGroup)
 		, SimulationTypeAttribute(InCollection, "SimulationType", FTransformCollection::TransformGroup)
-		, TransformIndexAttribute(InCollection, "TransformIndex", FTransformCollection::TransformGroup)
 		, TransformToGeometryIndexAttribute(InCollection, "TransformToGeometryIndex", FTransformCollection::TransformGroup)
+		, TransformIndexAttribute(InCollection, "TransformIndex", FGeometryCollection::GeometryGroup)
 		, VertexStartAttribute(InCollection, "VertexStart", FGeometryCollection::GeometryGroup)
 		, VertexCountAttribute(InCollection, "VertexCount", FGeometryCollection::GeometryGroup)
 		, FaceStartAttribute(InCollection, "FaceStart", FGeometryCollection::GeometryGroup)
@@ -1029,54 +1029,45 @@ namespace GeometryCollection::Facades
 	{
 		TArray<int32> OutSelection;
 
-		if (ConstCollection.HasGroup(FName(*GroupName)))
+		if (const TManagedArray<float>* FloatArray = ConstCollection.FindAttributeTyped<float>(FName(*AttrName), FName(*GroupName)))
 		{
-			if (ConstCollection.HasAttribute(FName(*AttrName), FName(*GroupName)))
+			for (int32 Idx = 0; Idx < FloatArray->Num(); ++Idx)
 			{
-				if (ConstCollection.GetAttributeType(FName(*AttrName), FName(*GroupName)) == EManagedArrayType::FFloatType)
-				{
-					const TManagedArray<float>& FloatArray = (const TManagedArray<float>&)ConstCollection.GetAttribute<float>(FName(*AttrName), FName(*GroupName));
+				const float FloatValue = (*FloatArray)[Idx];
 
-					for (int32 Idx = 0; Idx < FloatArray.Num(); ++Idx)
+				if (bInsideRange)
+				{
+					if (bInclusive)
 					{
-						const float FloatValue = FloatArray[Idx];
-
-						if (bInsideRange)
+						if (FloatValue >= Min && FloatValue <= Max)
 						{
-							if (bInclusive)
-							{
-								if (FloatValue >= Min && FloatValue <= Max)
-								{
-									OutSelection.Add(Idx);
-								}
-							}
-							else
-							{
-								if (FloatValue > Min && FloatValue < Max)
-								{
-									OutSelection.Add(Idx);
-								}
-							}
+							OutSelection.Add(Idx);
 						}
-						else
+					}
+					else
+					{
+						if (FloatValue > Min && FloatValue < Max)
 						{
-							if (bInclusive)
-							{
-								if (FloatValue <= Min || FloatValue >= Max)
-								{
-									OutSelection.Add(Idx);
-								}
-							}
-							else
-							{
-								if (FloatValue < Min || FloatValue > Max)
-								{
-									OutSelection.Add(Idx);
-								}
-							}
+							OutSelection.Add(Idx);
+						}
+					}
+				}
+				else
+				{
+					if (bInclusive)
+					{
+						if (FloatValue <= Min || FloatValue >= Max)
+						{
+							OutSelection.Add(Idx);
+						}
+					}
+					else
+					{
+						if (FloatValue < Min || FloatValue > Max)
+						{
+							OutSelection.Add(Idx);
 						}
 					}
-
 				}
 			}
 		}
@@ -1088,54 +1079,45 @@ namespace GeometryCollection::Facades
 	{
 		TArray<int32> OutSelection;
 
-		if (ConstCollection.HasGroup(FName(*GroupName)))
+		if (const TManagedArray<int32>* IntArray = ConstCollection.FindAttributeTyped<int32>(FName(*AttrName), FName(*GroupName)))
 		{
-			if (ConstCollection.HasAttribute(FName(*AttrName), FName(*GroupName)))
+			for (int32 Idx = 0; Idx < IntArray->Num(); ++Idx)
 			{
-				if (ConstCollection.GetAttributeType(FName(*AttrName), FName(*GroupName)) == EManagedArrayType::FFloatType)
-				{
-					const TManagedArray<int32>& IntArray = (const TManagedArray<int32>&)ConstCollection.GetAttribute<int>(FName(*AttrName), FName(*GroupName));
+				const int32 IntValue = (*IntArray)[Idx];
 
-					for (int32 Idx = 0; Idx < IntArray.Num(); ++Idx)
+				if (bInsideRange)
+				{
+					if (bInclusive)
 					{
-						const int32 IntValue = IntArray[Idx];
-
-						if (bInsideRange)
+						if (IntValue >= Min && IntValue <= Max)
 						{
-							if (bInclusive)
-							{
-								if (IntValue >= Min && IntValue <= Max)
-								{
-									OutSelection.Add(Idx);
-								}
-							}
-							else
-							{
-								if (IntValue > Min && IntValue < Max)
-								{
-									OutSelection.Add(Idx);
-								}
-							}
+							OutSelection.Add(Idx);
 						}
-						else
+					}
+					else
+					{
+						if (IntValue > Min && IntValue < Max)
 						{
-							if (bInclusive)
-							{
-								if (IntValue <= Min || IntValue >= Max)
-								{
-									OutSelection.Add(Idx);
-								}
-							}
-							else
-							{
-								if (IntValue < Min || IntValue > Max)
-								{
-									OutSelection.Add(Idx);
-								}
-							}
+							OutSelection.Add(Idx);
+						}
+					}
+				}
+				else
+				{
+					if (bInclusive)
+					{
+						if (IntValue <= Min || IntValue >= Max)
+						{
+							OutSelection.Add(Idx);
+						}
+					}
+					else
+					{
+						if (IntValue < Min || IntValue > Max)
+						{
+							OutSelection.Add(Idx);
 						}
 					}
-
 				}
 			}
 		}
