Geometry Collection : Fix crash when adding a geometry collection to a blueprint that contains embedded geometry
- crash was caused by the fact that the owner of a component is null when added to a blueprint
- added extra checks to avoid crashing when this is the case

diff --git a/Engine/Source/Runtime/Experimental/GeometryCollectionEngine/Private/GeometryCollection/GeometryCollectionComponent.cpp b/Engine/Source/Runtime/Experimental/GeometryCollectionEngine/Private/GeometryCollection/GeometryCollectionComponent.cpp
--- a/Engine/Source/Runtime/Experimental/GeometryCollectionEngine/Private/GeometryCollection/GeometryCollectionComponent.cpp
+++ b/Engine/Source/Runtime/Experimental/GeometryCollectionEngine/Private/GeometryCollection/GeometryCollectionComponent.cpp
@@ -6215,22 +6215,23 @@ bool UGeometryCollectionComponent::IsEmbeddedGeometryValid() const
 
 void UGeometryCollectionComponent::ClearEmbeddedGeometry()
 {
-	AActor* OwningActor = GetOwner();
-	TArray<UActorComponent*> TargetComponents;
-	OwningActor->GetComponents(TargetComponents, false);
-
-	for (UActorComponent* TargetComponent : TargetComponents)
+	if (AActor* OwningActor = GetOwner())
 	{
-		if ((TargetComponent->GetOuter() == this) || !IsValidChecked(TargetComponent->GetOuter()))
+		TArray<UActorComponent*> TargetComponents;
+		OwningActor->GetComponents(TargetComponents, false);
+
+		for (UActorComponent* TargetComponent : TargetComponents)
 		{
-			if (UInstancedStaticMeshComponent* ISMComponent = Cast<UInstancedStaticMeshComponent>(TargetComponent))
+			if ((TargetComponent->GetOuter() == this) || !IsValidChecked(TargetComponent->GetOuter()))
 			{
-				ISMComponent->ClearInstances();
-				ISMComponent->DestroyComponent();
+				if (UInstancedStaticMeshComponent* ISMComponent = Cast<UInstancedStaticMeshComponent>(TargetComponent))
+				{
+					ISMComponent->ClearInstances();
+					ISMComponent->DestroyComponent();
+				}
 			}
 		}
 	}
-
 	EmbeddedGeometryComponents.Empty();
 }
 
@@ -6240,35 +6241,36 @@ void UGeometryCollectionComponent::InitializeEmbeddedGeometry()
 	{
 		ClearEmbeddedGeometry();
 		
-		AActor* ActorOwner = GetOwner();
-		check(ActorOwner);
-
-		// Construct an InstancedStaticMeshComponent for each exemplar
-		for (const FGeometryCollectionEmbeddedExemplar& Exemplar : RestCollection->EmbeddedGeometryExemplar)
+		if (AActor* ActorOwner = GetOwner())
 		{
-			if (UStaticMesh* ExemplarStaticMesh = Cast<UStaticMesh>(Exemplar.StaticMeshExemplar.TryLoad()))
+
+			// Construct an InstancedStaticMeshComponent for each exemplar
+			for (const FGeometryCollectionEmbeddedExemplar& Exemplar : RestCollection->EmbeddedGeometryExemplar)
 			{
-				if (UInstancedStaticMeshComponent* ISMC = NewObject<UInstancedStaticMeshComponent>(this))
+				if (UStaticMesh* ExemplarStaticMesh = Cast<UStaticMesh>(Exemplar.StaticMeshExemplar.TryLoad()))
 				{
-					ISMC->SetStaticMesh(ExemplarStaticMesh);
-					ISMC->SetCullDistances(Exemplar.StartCullDistance, Exemplar.EndCullDistance);
-					ISMC->SetCanEverAffectNavigation(false);
-					ISMC->SetCollisionProfileName(UCollisionProfile::NoCollision_ProfileName);
-					ISMC->SetCastShadow(false);
-					ISMC->SetMobility(EComponentMobility::Stationary);
-					ISMC->SetupAttachment(this);
-					ActorOwner->AddInstanceComponent(ISMC);
-					ISMC->RegisterComponent();
-
-					EmbeddedGeometryComponents.Add(ISMC);
+					if (UInstancedStaticMeshComponent* ISMC = NewObject<UInstancedStaticMeshComponent>(this))
+					{
+						ISMC->SetStaticMesh(ExemplarStaticMesh);
+						ISMC->SetCullDistances(Exemplar.StartCullDistance, Exemplar.EndCullDistance);
+						ISMC->SetCanEverAffectNavigation(false);
+						ISMC->SetCollisionProfileName(UCollisionProfile::NoCollision_ProfileName);
+						ISMC->SetCastShadow(false);
+						ISMC->SetMobility(EComponentMobility::Stationary);
+						ISMC->SetupAttachment(this);
+						ActorOwner->AddInstanceComponent(ISMC);
+						ISMC->RegisterComponent();
+
+						EmbeddedGeometryComponents.Add(ISMC);
+					}
 				}
 			}
-		}
 
 #if WITH_EDITOR
-		EmbeddedBoneMaps.SetNum(RestCollection->EmbeddedGeometryExemplar.Num());
-		EmbeddedInstanceIndex.Init(INDEX_NONE,RestCollection->GetGeometryCollection()->NumElements(FGeometryCollection::TransformGroup));
+			EmbeddedBoneMaps.SetNum(RestCollection->EmbeddedGeometryExemplar.Num());
+			EmbeddedInstanceIndex.Init(INDEX_NONE, RestCollection->GetGeometryCollection()->NumElements(FGeometryCollection::TransformGroup));
 #endif
+		}
 	}
 }
 
