Dataflow : fix template type upgrade path to properly work with some older graph

diff --git a/Engine/Source/Runtime/Experimental/Dataflow/Core/Private/Dataflow/DataflowConnection.cpp b/Engine/Source/Runtime/Experimental/Dataflow/Core/Private/Dataflow/DataflowConnection.cpp
--- a/Engine/Source/Runtime/Experimental/Dataflow/Core/Private/Dataflow/DataflowConnection.cpp
+++ b/Engine/Source/Runtime/Experimental/Dataflow/Core/Private/Dataflow/DataflowConnection.cpp
@@ -97,16 +97,25 @@ bool FDataflowConnection::SupportsType(FName InType) const
 	return (InType == GetType());
 }
 
+bool FDataflowConnection::IsExtendedType(FName InType) const
+{
+	return (InType.ToString().StartsWith(Type.ToString() + "<"));
+}
+
 bool FDataflowConnection::SetConcreteType(FName InType)
 {
 	// Can only change from AnyType to a concrete type
-	if (Type != InType && ensure(!bHasConcreteType))
+	if (Type != InType)
 	{
-		if (ensure(SupportsType(InType)))
+		const bool bExtendedType = IsExtendedType(InType);
+		if (ensure(!bHasConcreteType || bExtendedType))
 		{
-			Type = InType;
-			bHasConcreteType = true;
-			return true;
+			if (ensure(bExtendedType || SupportsType(InType)))
+			{
+				Type = InType;
+				bHasConcreteType = true;
+				return true;
+			}
 		}
 	}
 	return false;
diff --git a/Engine/Source/Runtime/Experimental/Dataflow/Core/Private/Dataflow/DataflowInputOutput.cpp b/Engine/Source/Runtime/Experimental/Dataflow/Core/Private/Dataflow/DataflowInputOutput.cpp
--- a/Engine/Source/Runtime/Experimental/Dataflow/Core/Private/Dataflow/DataflowInputOutput.cpp
+++ b/Engine/Source/Runtime/Experimental/Dataflow/Core/Private/Dataflow/DataflowInputOutput.cpp
@@ -89,23 +89,33 @@ void FDataflowInput::PullValue(Dataflow::FContext& Context) const
 void FDataflowInput::FixAndPropagateType(FName InType)
 {
 	check(InType.ToString().StartsWith(Type.ToString()));
-	// first try to fix as a legal change to give a chance for the node to proapagte the type changes if necessary
-	const bool bNeedToPropagate = OwningNode->TrySetConnectionType(this, InType);
-	if (bNeedToPropagate)
+	check(!FDataflowConnection::IsAnyType(InType));
+	
+	if (GetType() != InType)
 	{
-		if (FDataflowReRouteNode* ReRouteNode = OwningNode->AsType<FDataflowReRouteNode>())
+		Type = InType;
+		bHasConcreteType = true;
+
+		// if we have a reroute node propagate through to make sure each reroute segment is properly handled
+		// IMPORTANT : this needs to be done before we propagate through the input connections
+		if (const FDataflowReRouteNode* const ReRouteNode = OwningNode->AsType<FDataflowReRouteNode>())
 		{
-			if (FDataflowOutput* Output = GetConnection())
+			for (FDataflowOutput* const ReRouteOutput : ReRouteNode->GetOutputs())
 			{
-				Output->FixAndPropagateType(InType);
+				if (ReRouteOutput)
+				{
+					ReRouteOutput->FixAndPropagateType(InType);
+				}
 			}
 		}
-	}
-	else
-	{
-		// Simply fix the type  ( at this point we should not be a anytype )
-		check(!bIsAnyType);
-		Type = InType;
+
+		OwningNode->NotifyConnectionTypeChanged(this);
+
+		// Now propagate to the connected output
+		if (FDataflowOutput* const Output = GetConnection())
+		{
+			Output->FixAndPropagateType(InType);
+		}
 	}
 }
 
@@ -299,25 +309,35 @@ void FDataflowOutput::ForwardInput(const FDataflowInput* Input, Dataflow::FConte
 void FDataflowOutput::FixAndPropagateType(FName InType)
 {
 	check(InType.ToString().StartsWith(Type.ToString()));
-	// first try to fix as a legal change to give a chance for the node to proapagte the type changes if necessary
-	const bool bNeedToPropagate = OwningNode->TrySetConnectionType(this, InType);
-	if (bNeedToPropagate)
+	check(!FDataflowConnection::IsAnyType(InType));
+
+	if (Type != InType)
 	{
-		if (FDataflowReRouteNode* ReRouteNode = OwningNode->AsType<FDataflowReRouteNode>())
+		Type = InType;
+		bHasConcreteType = true;
+
+		// if we have a reroute node propagate through to make sure each reroute segment is properly handled
+		// IMPORTANT : this needs to be done before we propagate through the output connections
+		if (const FDataflowReRouteNode* const ReRouteNode = OwningNode->AsType<FDataflowReRouteNode>())
 		{
-			for (FDataflowInput* Input : Connections)
+			for (FDataflowInput* const RerouteInput : ReRouteNode->GetInputs())
 			{
-				if (Input)
+				if (RerouteInput)
 				{
-					Input->FixAndPropagateType(InType);
+					RerouteInput->FixAndPropagateType(InType);
 				}
 			}
 		}
-	}
-	else
-	{
-		// Simply fix the type  ( at this point we should not be a anytype )
-		check(!bIsAnyType);
-		Type = InType;
+
+		OwningNode->NotifyConnectionTypeChanged(this);
+
+		// Now propagate through the connected inputs
+		for (FDataflowInput* const Input : Connections)
+		{
+			if (Input)
+			{
+				Input->FixAndPropagateType(InType);
+			}
+		}
 	}
 }
\ No newline at end of file
diff --git a/Engine/Source/Runtime/Experimental/Dataflow/Core/Private/Dataflow/DataflowNode.cpp b/Engine/Source/Runtime/Experimental/Dataflow/Core/Private/Dataflow/DataflowNode.cpp
--- a/Engine/Source/Runtime/Experimental/Dataflow/Core/Private/Dataflow/DataflowNode.cpp
+++ b/Engine/Source/Runtime/Experimental/Dataflow/Core/Private/Dataflow/DataflowNode.cpp
@@ -1165,20 +1165,28 @@ bool FDataflowNode::TrySetConnectionType(FDataflowConnection* Connection, FName
 		if (Connection->IsAnyType() && Connection->GetType() != NewType && !FDataflowConnection::IsAnyType(NewType))
 		{
 			Connection->SetConcreteType(NewType);
-			if (Connection->GetDirection() == Dataflow::FPin::EDirection::INPUT)
-			{
-				OnInputTypeChanged((FDataflowInput*)Connection);
-			}
-			if (Connection->GetDirection() == Dataflow::FPin::EDirection::OUTPUT)
-			{
-				OnOutputTypeChanged((FDataflowOutput*)Connection);
-			}
+			NotifyConnectionTypeChanged(Connection);
 			return true;
 		}
 	}
 	return false;
 }
 
+void FDataflowNode::NotifyConnectionTypeChanged(FDataflowConnection* Connection)
+{
+	if (Connection->IsAnyType())
+	{
+		if (Connection->GetDirection() == Dataflow::FPin::EDirection::INPUT)
+		{
+			OnInputTypeChanged((FDataflowInput*)Connection);
+		}
+		if (Connection->GetDirection() == Dataflow::FPin::EDirection::OUTPUT)
+		{
+			OnOutputTypeChanged((FDataflowOutput*)Connection);
+		}
+	}
+}
+
 bool FDataflowNode::SetInputConcreteType(const Dataflow::FConnectionReference& InputReference, FName NewType)
 {
 	if (FDataflowInput* Input = FindInput(InputReference))
diff --git a/Engine/Source/Runtime/Experimental/Dataflow/Core/Public/Dataflow/DataflowConnection.h b/Engine/Source/Runtime/Experimental/Dataflow/Core/Public/Dataflow/DataflowConnection.h
--- a/Engine/Source/Runtime/Experimental/Dataflow/Core/Public/Dataflow/DataflowConnection.h
+++ b/Engine/Source/Runtime/Experimental/Dataflow/Core/Public/Dataflow/DataflowConnection.h
@@ -159,6 +159,12 @@ protected:
 	DATAFLOWCORE_API void FixAndPropagateType();
 	DATAFLOWCORE_API virtual void FixAndPropagateType(FName InType) { ensure(false); }
 
+	/** 
+	* returns true is the parameter type is an extension of the current type 
+	* TArray<int> from TArray for example 
+	*/
+	DATAFLOWCORE_API bool IsExtendedType(FName InType) const;
+
 	/** 
 	* Set the concrete type of an anytype connection 
 	* @return true if the typewas effectivelly changed
diff --git a/Engine/Source/Runtime/Experimental/Dataflow/Core/Public/Dataflow/DataflowNode.h b/Engine/Source/Runtime/Experimental/Dataflow/Core/Public/Dataflow/DataflowNode.h
--- a/Engine/Source/Runtime/Experimental/Dataflow/Core/Public/Dataflow/DataflowNode.h
+++ b/Engine/Source/Runtime/Experimental/Dataflow/Core/Public/Dataflow/DataflowNode.h
@@ -550,6 +550,10 @@ PRAGMA_ENABLE_DEPRECATION_WARNINGS
 	// returns true if the type was changed successfully
 	// only unset datatype connection will be set a new type 
 	DATAFLOWCORE_API bool TrySetConnectionType(FDataflowConnection* Connection, FName NewType);
+
+	// Only used when forcing types on connection in order to make sure the node properly refreshes the rest of its connection accordingly if there's any dependencies between their types
+	DATAFLOWCORE_API void NotifyConnectionTypeChanged(FDataflowConnection* Connection);
+
 protected:
 	virtual bool OnInputTypeChanged(const FDataflowInput* Input) { return false; };
 	virtual bool OnOutputTypeChanged(const FDataflowOutput* Output) { return false; }
