Fix a few more 14.30 warnings

#rnx

diff --git a/Engine/Source/Runtime/Experimental/ChaosVisualDebugger/Public/DataWrappers/ChaosVDAccelerationStructureDataWrappers.h b/Engine/Source/Runtime/Experimental/ChaosVisualDebugger/Public/DataWrappers/ChaosVDAccelerationStructureDataWrappers.h
--- a/Engine/Source/Runtime/Experimental/ChaosVisualDebugger/Public/DataWrappers/ChaosVDAccelerationStructureDataWrappers.h
+++ b/Engine/Source/Runtime/Experimental/ChaosVisualDebugger/Public/DataWrappers/ChaosVDAccelerationStructureDataWrappers.h
@@ -75,7 +75,10 @@ struct FChaosVDBoundingVolumeDataWrapper : public FChaosVDAccelerationStructureB
 	CHAOSVDRUNTIME_API bool Serialize(FArchive& Ar);
 };
 
+// TODO: Disabling C4996 due to MSVC 14.30 bug
+#pragma warning(disable : 4996)
 CVD_IMPLEMENT_SERIALIZER(FChaosVDBoundingVolumeDataWrapper)
+#pragma warning(default : 4996)
 
 USTRUCT(DisplayName="AABB Tree Node")
 struct FChaosVDAABBTreeNodeDataWrapper : public FChaosVDWrapperDataBase
@@ -96,7 +99,10 @@ struct FChaosVDAABBTreeNodeDataWrapper : public FChaosVDWrapperDataBase
 	CHAOSVDRUNTIME_API bool Serialize(FArchive& Ar);
 };
 
+// TODO: Disabling C4996 due to MSVC 14.30 bug
+#pragma warning(disable : 4996)
 CVD_IMPLEMENT_SERIALIZER(FChaosVDAABBTreeNodeDataWrapper)
+#pragma warning(default : 4996)
 
 USTRUCT()
 struct FChaosVDAABBTreePayloadBoundsElement : public FChaosVDWrapperDataBase
@@ -117,7 +123,10 @@ struct FChaosVDAABBTreePayloadBoundsElement : public FChaosVDWrapperDataBase
 	CHAOSVDRUNTIME_API bool Serialize(FArchive& Ar);
 };
 
+// TODO: Disabling C4996 due to MSVC 14.30 bug
+#pragma warning(disable : 4996)
 CVD_IMPLEMENT_SERIALIZER(FChaosVDAABBTreePayloadBoundsElement)
+#pragma warning(default : 4996)
 
 USTRUCT(DisplayName="AABB Tree Leaf")
 struct FChaosVDAABBTreeLeafDataWrapper : public FChaosVDWrapperDataBase
@@ -133,7 +142,10 @@ struct FChaosVDAABBTreeLeafDataWrapper : public FChaosVDWrapperDataBase
 	CHAOSVDRUNTIME_API bool Serialize(FArchive& Ar);
 };
 
+// TODO: Disabling C4996 due to MSVC 14.30 bug
+#pragma warning(disable : 4996)
 CVD_IMPLEMENT_SERIALIZER(FChaosVDAABBTreeLeafDataWrapper)
+#pragma warning(default : 4996)
 
 USTRUCT(DisplayName="AABB Tree Data")
 struct FChaosVDAABBTreeDataWrapper : public FChaosVDAccelerationStructureBase
@@ -178,4 +190,7 @@ struct FChaosVDAABBTreeDataWrapper : public FChaosVDAccelerationStructureBase
 	int32 GetCorrectedRootNodeIndex() const { return bDynamicTree ? RootNodeIndex : 0;};
 };
 
+// TODO: Disabling C4996 due to MSVC 14.30 bug
+#pragma warning(disable : 4996)
 CVD_IMPLEMENT_SERIALIZER(FChaosVDAABBTreeDataWrapper)
+#pragma warning(default : 4996)
