Ignore nullptr rootset objects in JsonObjectGraph::Stringify

#rnx

diff --git a/Engine/Source/Runtime/Experimental/JsonObjectGraph/Private/JsonStringifyImpl.cpp b/Engine/Source/Runtime/Experimental/JsonObjectGraph/Private/JsonStringifyImpl.cpp
--- a/Engine/Source/Runtime/Experimental/JsonObjectGraph/Private/JsonStringifyImpl.cpp
+++ b/Engine/Source/Runtime/Experimental/JsonObjectGraph/Private/JsonStringifyImpl.cpp
@@ -2,6 +2,7 @@
 
 #include "JsonStringifyImpl.h"
 #include "PrettyJsonWriter.h"
+#include "Algo/Copy.h"
 #include "Algo/ForEach.h"
 #include "Algo/RemoveIf.h"
 #include "Misc/App.h"
@@ -146,13 +147,14 @@ static void GatherExports(
 	AllowedRoots.Append(Roots);
 	for (const UObject* Obj : Roots)
 	{
-		if (!Obj->IsA<UPackage>())
+		if (Obj && !Obj->IsA<UPackage>())
 		{
 			OutRoots.Add(Obj);
 		}
 	}
 
-	TArray<const UObject*> PendingRefs = TArray<const UObject*>(Roots);
+	TArray<const UObject*> PendingRefs;
+	Algo::CopyIf(Roots, PendingRefs, [](const UObject* Obj) { return IsValid(Obj); });
 
 	TSet<const UObject*> RefsProcessed;
 	RefsProcessed.Append(PendingRefs);
