Altered dirty particle checking to make sure we always add dirty kinematics to be processed. We can control per-body how the position is synchronised with the physics results but we must always sync the velocity and dynamics information. The proxy pull method will handle when and when not to sync positions but if the particle is not marked as dirty the pull method will never be called.

diff --git a/Engine/Source/Runtime/Experimental/Chaos/Private/PBDRigidsSolver.cpp b/Engine/Source/Runtime/Experimental/Chaos/Private/PBDRigidsSolver.cpp
--- a/Engine/Source/Runtime/Experimental/Chaos/Private/PBDRigidsSolver.cpp
+++ b/Engine/Source/Runtime/Experimental/Chaos/Private/PBDRigidsSolver.cpp
@@ -2092,7 +2092,13 @@ namespace Chaos
 						{
 							if(!bIsResim || DirtyParticle.SyncState() == ESyncState::HardDesync)
 							{
-								if (ShouldUpdateFromSimulation(DirtyParticle))
+								// Although per-particle we can control the syncing of target positions (see ShouldUpdateFromSimulation)
+								// we cannot avoid marking kinematics as dirty (unless the global config Chaos::SyncKinematicOnGameThread
+								// forces it) because we always need the correct velocities/dynamics to be synced back from the kinematic
+								// target.
+								// FSingleParticlePhysicsProxy::PullFromPhysicsState will use the proper checks to see whether it needs
+								// to sync the particle positions when the dirty particle is processed.
+								if(!(Chaos::SyncKinematicOnGameThread == 0 && DirtyParticle.ObjectState() == EObjectStateType::Kinematic))
 								{
 									ActiveRigid.AddUnique((FSingleParticlePhysicsProxy*)Proxy);
 								}
