Iris - Fixed subobjects could be part of scope even though root object was filtered out if hysteresis was enabled.

diff --git a/Engine/Source/Runtime/Experimental/Iris/Core/Private/Iris/ReplicationSystem/Filtering/ReplicationFiltering.cpp b/Engine/Source/Runtime/Experimental/Iris/Core/Private/Iris/ReplicationSystem/Filtering/ReplicationFiltering.cpp
--- a/Engine/Source/Runtime/Experimental/Iris/Core/Private/Iris/ReplicationSystem/Filtering/ReplicationFiltering.cpp
+++ b/Engine/Source/Runtime/Experimental/Iris/Core/Private/Iris/ReplicationSystem/Filtering/ReplicationFiltering.cpp
@@ -29,6 +29,9 @@ DEFINE_LOG_CATEGORY_STATIC(LogIrisFiltering, Log, All);
 bool bCVarRepFilterCullNonRelevant = true;
 static FAutoConsoleVariableRef CVarRepFilterCullNonRelevant(TEXT("Net.Iris.CullNonRelevant"), bCVarRepFilterCullNonRelevant, TEXT("When enabled will cull replicated actors that are not relevant to any client."), ECVF_Default);
 
+bool bCVarRepFilterValidateNoSubObjectInScopeWithFilteredOutRootObject = false;
+static FAutoConsoleVariableRef CVarRepFilterValidateNoSubObjectInScopeWithFilteredOutRootObject(TEXT("Net.Iris.Filtering.ValidateNobSubObjectInScopeWithFilteredOutRootObject"), bCVarRepFilterValidateNoSubObjectInScopeWithFilteredOutRootObject, TEXT("Validate there are no subobjects in scope with a filtered out root object."), ECVF_Default);
+
 FName GetStaticFilterName(FNetObjectFilterHandle Filter)
 {
 	switch (Filter)
@@ -358,6 +361,14 @@ void FReplicationFiltering::FilterNonRelevantObjects()
 {
 	IRIS_PROFILER_SCOPE(FReplicationFiltering_FilterNonRelevantObjects);
 
+	if (bCVarRepFilterValidateNoSubObjectInScopeWithFilteredOutRootObject)
+	{
+		ValidConnections.ForAllSetBits([this](uint32 ConnectionIndex)
+			{
+				ensureMsgf(!HasSubObjectInScopeWithFilteredOutRootObject(ConnectionIndex), TEXT("Connection %u has orphaned subobjects."), ConnectionIndex);
+			});
+	}
+
 	if (!bCVarRepFilterCullNonRelevant)
 	{
 		// Make every object in the global scope part of the relevant list
@@ -1226,6 +1237,8 @@ void FReplicationFiltering::UpdateGroupExclusionFiltering()
 					}
 				}
 			}
+
+			ensureMsgf(!bCVarRepFilterValidateNoSubObjectInScopeWithFilteredOutRootObject || !HasSubObjectInScopeWithFilteredOutRootObject(MakeNetBitArrayView(ConnectionInfos[ConnectionId].GroupExcludedObjects)), TEXT("UpdateGroupExclusionFiltering GroupExcludedObjects"));
 		};
 
 		this->ValidConnections.ForAllSetBits(UpdateGroupFilterForConnection);
@@ -1498,7 +1511,9 @@ void FReplicationFiltering::UpdateDynamicFiltering()
 			for (const uint32 ObjectIndex : FilteredOutObjects)
 			{
 				const uint32 HysteresisFrameCount = ObjectScopeHysteresisFrameCounts[ObjectIndex];
-				if (HysteresisFrameCount && DynamicFilterEnabledObjects.GetBit(ObjectIndex) && !HysteresisState.ObjectsExemptFromHysteresis.GetBit(ObjectIndex))
+				const bool bAlreadyFilteredOut = DynamicFilteredOutObjectsHysteresisAdjusted.GetBit(ObjectIndex);
+				// If the object is already filtered out we cannot add to hysteresis. It may be that a subobject was added and the root object was marked for processing, in which case we should immediately filter out the subobject.
+				if (!bAlreadyFilteredOut && HysteresisFrameCount && DynamicFilterEnabledObjects.GetBit(ObjectIndex) && !HysteresisState.ObjectsExemptFromHysteresis.GetBit(ObjectIndex))
 				{
 					// We need to adjust the hysteresis frame count to account for when it will be updated. The -1 stems from the fact that updating won't happen until next frame at the earliest.
 					const uint16 TotalHysteresisFrameCount = static_cast<uint16>(HysteresisFrameCount - 1 + AdjustHysteresisForUpdateThrottling);
@@ -1558,6 +1573,8 @@ void FReplicationFiltering::UpdateDynamicFiltering()
 			}
 		}
 
+		ensure(!bCVarRepFilterValidateNoSubObjectInScopeWithFilteredOutRootObject || !HasSubObjectInScopeWithFilteredOutRootObject(DynamicFilteredOutObjectsHysteresisAdjusted));
+
 		// Update the entire scope for the connection.
 		{
 			IRIS_PROFILER_SCOPE(FReplicationFiltering_UpdateConnectionScope);
@@ -2779,6 +2796,49 @@ uint8 FReplicationFiltering::GetObjectScopeHysteresisFrameCount(FName ProfileNam
 	return Config->GetDefaultHysteresisFrameCount();
 }
 
+bool FReplicationFiltering::HasSubObjectInScopeWithFilteredOutRootObject(FNetBitArrayView Objects) const
+{
+	bool bReturnValue = false;
+	Objects.ForAllSetBits([this, &Objects, &bReturnValue](uint32 ObjectIndex)
+		{
+			const FNetRefHandleManager::FReplicatedObjectData& ReplicatedObjectData = this->NetRefHandleManager->GetReplicatedObjectDataNoCheck(ObjectIndex);
+			if (ReplicatedObjectData.SubObjectRootIndex != FNetRefHandleManager::InvalidInternalIndex)
+			{
+				if (!Objects.GetBit(ReplicatedObjectData.SubObjectRootIndex))
+				{
+					bReturnValue = true;
+					ensureMsgf(Objects.GetBit(ReplicatedObjectData.SubObjectRootIndex), TEXT("Root index %u is not in scope for subobject %u"), ReplicatedObjectData.SubObjectRootIndex, ObjectIndex);
+				}
+			}
+			if (ReplicatedObjectData.SubObjectParentIndex != FNetRefHandleManager::InvalidInternalIndex)
+			{
+				if (!Objects.GetBit(ReplicatedObjectData.SubObjectParentIndex))
+				{
+					bReturnValue = true;
+					ensureMsgf(Objects.GetBit(ReplicatedObjectData.SubObjectParentIndex), TEXT("Parent index %u is not in scope for subobject %u"), ReplicatedObjectData.SubObjectParentIndex, ObjectIndex);
+				}
+			}
+		});
+
+	return bReturnValue;
+}
+
+bool FReplicationFiltering::HasSubObjectInScopeWithFilteredOutRootObject(uint32 ConnectionId) const
+{
+	const FPerConnectionInfo& ConnectionInfo = ConnectionInfos[ConnectionId];
+	if (!ensureMsgf(!HasSubObjectInScopeWithFilteredOutRootObject(MakeNetBitArrayView(ConnectionInfo.ObjectsInScope)), TEXT("ObjectsInScope")))
+	{
+		ensureMsgf(!HasSubObjectInScopeWithFilteredOutRootObject(MakeNetBitArrayView(ConnectionInfo.GroupIncludedObjects)), TEXT("IncludedObjects"));
+		ensureMsgf(!HasSubObjectInScopeWithFilteredOutRootObject(MakeNetBitArrayView(ConnectionInfo.GroupExcludedObjects)), TEXT("ExcludedObjects"));
+		ensureMsgf(!HasSubObjectInScopeWithFilteredOutRootObject(MakeNetBitArrayView(ConnectionInfo.ObjectsInScopeBeforeDynamicFiltering)), TEXT("BeforeDynamic"));
+		ensureMsgf(!HasSubObjectInScopeWithFilteredOutRootObject(MakeNetBitArrayView(ConnectionInfo.ConnectionFilteredObjects)), TEXT("ConnectionFiltered"));
+		ensureMsgf(!HasSubObjectInScopeWithFilteredOutRootObject(MakeNetBitArrayView(ConnectionInfo.DynamicFilteredOutObjectsHysteresisAdjusted)), TEXT("DynamicFilteredOutObjectsHysteresisAdjusted"));
+		return true;
+	}
+
+	return false;
+}
+
 void FReplicationFiltering::FPerConnectionInfo::Deinit()
 {
 	ConnectionFilteredObjects.Empty();
diff --git a/Engine/Source/Runtime/Experimental/Iris/Core/Private/Iris/ReplicationSystem/Filtering/ReplicationFiltering.h b/Engine/Source/Runtime/Experimental/Iris/Core/Private/Iris/ReplicationSystem/Filtering/ReplicationFiltering.h
--- a/Engine/Source/Runtime/Experimental/Iris/Core/Private/Iris/ReplicationSystem/Filtering/ReplicationFiltering.h
+++ b/Engine/Source/Runtime/Experimental/Iris/Core/Private/Iris/ReplicationSystem/Filtering/ReplicationFiltering.h
@@ -270,6 +270,9 @@ private:
 
 	uint8 GetObjectScopeHysteresisFrameCount(FName Profile) const;
 
+	bool HasSubObjectInScopeWithFilteredOutRootObject(FNetBitArrayView Objects) const;
+	bool HasSubObjectInScopeWithFilteredOutRootObject(uint32 connectionId) const;
+
 private:
 	enum EHysteresisProcessingMode : uint32
 	{
