Iris
* Deleted use of static DefaultCreateNetRefHandleParams. The FCreateNetRefHandleParams default parameter values should be used instead.
* Deleted bCanReceive param from FCreateNetRefHandleParams. This was only used by some unit tests.
* Removed unused cache option param in InternalFlushStateData
* Misc API cleanup

#rnx

diff --git a/Engine/Plugins/Runtime/ReplicationSystemTestPlugin/Source/Private/Tests/ReplicationSystem/Filtering/TestGridFilter.cpp b/Engine/Plugins/Runtime/ReplicationSystemTestPlugin/Source/Private/Tests/ReplicationSystem/Filtering/TestGridFilter.cpp
--- a/Engine/Plugins/Runtime/ReplicationSystemTestPlugin/Source/Private/Tests/ReplicationSystem/Filtering/TestGridFilter.cpp
+++ b/Engine/Plugins/Runtime/ReplicationSystemTestPlugin/Source/Private/Tests/ReplicationSystem/Filtering/TestGridFilter.cpp
@@ -86,7 +86,6 @@ UE_NET_TEST_FIXTURE(FTestGridFilterFixture, TestWorldLocGridFilter)
 {
 	// Spawn object with WorldLocation's on server
 	UObjectReplicationBridge::FCreateNetRefHandleParams Params;
-	Params.bCanReceive = true;
 	Params.bNeedsWorldLocationUpdate = true;
 	Params.bUseClassConfigDynamicFilter = true;
 
@@ -267,7 +266,6 @@ UE_NET_TEST_FIXTURE(FTestGridFilterFixture, TestFilterPerformance)
 	for (int32 i = 0; i < NUM_OBJECTS; i++)
 	{
 		UObjectReplicationBridge::FCreateNetRefHandleParams Params;
-		Params.bCanReceive = true;
 		Params.bNeedsWorldLocationUpdate = true;
 		Params.bUseClassConfigDynamicFilter = true;
 
diff --git a/Engine/Plugins/Runtime/ReplicationSystemTestPlugin/Source/Private/Tests/ReplicationSystem/ReplicatedTestObject.cpp b/Engine/Plugins/Runtime/ReplicationSystemTestPlugin/Source/Private/Tests/ReplicationSystem/ReplicatedTestObject.cpp
--- a/Engine/Plugins/Runtime/ReplicationSystemTestPlugin/Source/Private/Tests/ReplicationSystem/ReplicatedTestObject.cpp
+++ b/Engine/Plugins/Runtime/ReplicationSystemTestPlugin/Source/Private/Tests/ReplicationSystem/ReplicatedTestObject.cpp
@@ -50,9 +50,7 @@ const UE::Net::FReplicationInstanceProtocol* UReplicatedTestObjectBridge::GetRep
 UE::Net::FNetRefHandle UReplicatedTestObjectBridge::BeginReplication(UReplicatedTestObject* Instance)
 {
 	// Create NetRefHandle for the registered fragments
-	Super::FCreateNetRefHandleParams Params = Super::DefaultCreateNetRefHandleParams;
-	Params.bCanReceive = true;
-
+	Super::FCreateNetRefHandleParams Params;
 	FNetRefHandle Handle = Super::BeginReplication(Instance, Params);
 
 	// This is optional but typically we want to cache at least the NetRefHandle in the game instance to avoid doing map lookups to find it
@@ -83,11 +81,8 @@ UE::Net::FNetRefHandle UReplicatedTestObjectBridge::BeginReplication(FNetRefHand
 	check(OwnerHandle.IsValid());
 
 	// Create NetRefHandle for the registered fragments
-	Super::FCreateNetRefHandleParams Params = Super::DefaultCreateNetRefHandleParams;
-	Params.bCanReceive = true;
-
-	// Create NetRefHandle for the registered fragments
-	FNetRefHandle Handle = Super::BeginReplication(OwnerHandle, SubObjectInstance, InsertRelativeToSubObjectHandle, InsertionOrder, Params);
+	Super::FCreateNetRefHandleParams Params;
+	FNetRefHandle Handle = Super::BeginReplication(OwnerHandle, SubObjectInstance, InsertRelativeToSubObjectHandle, Params, InsertionOrder);
 
 	if (Handle.IsValid())
 	{
diff --git a/Engine/Plugins/Runtime/ReplicationSystemTestPlugin/Source/Private/Tests/ReplicationSystem/TestDependentObjects.cpp b/Engine/Plugins/Runtime/ReplicationSystemTestPlugin/Source/Private/Tests/ReplicationSystem/TestDependentObjects.cpp
--- a/Engine/Plugins/Runtime/ReplicationSystemTestPlugin/Source/Private/Tests/ReplicationSystem/TestDependentObjects.cpp
+++ b/Engine/Plugins/Runtime/ReplicationSystemTestPlugin/Source/Private/Tests/ReplicationSystem/TestDependentObjects.cpp
@@ -269,7 +269,7 @@ UE_NET_TEST_FIXTURE(FReplicationSystemServerClientTestFixture, TestDependentObje
 
 	// Spawn second object on server that later will be added as a dependent object
 	// With high PollFramePeriod so that it will not replicate in a while unless it is a dependent
-	UObjectReplicationBridge::FCreateNetRefHandleParams Params = Server->GetReplicationBridge()->DefaultCreateNetRefHandleParams;
+	UObjectReplicationBridge::FCreateNetRefHandleParams Params;
 	Params.PollFrequency = Server->ConvertPollPeriodIntoFrequency(255U);
 	UTestReplicatedIrisObject* ServerDependentObject = Server->CreateObject(Params);
 
@@ -336,7 +336,7 @@ UE_NET_TEST_FIXTURE(FReplicationSystemServerClientTestFixture, TestDependentObje
 	FReplicationSystemTestClient* Client = CreateClient();
 
 	// Spawn object on server
-	UObjectReplicationBridge::FCreateNetRefHandleParams Params = Server->GetReplicationBridge()->DefaultCreateNetRefHandleParams;
+	UObjectReplicationBridge::FCreateNetRefHandleParams Params;
 
 	// Setup different poll frequencies for the objects
 	Params.PollFrequency = Server->ConvertPollPeriodIntoFrequency(10U);
@@ -415,7 +415,7 @@ UE_NET_TEST_FIXTURE(FReplicationSystemServerClientTestFixture, TestDependentObje
 	FReplicationSystemTestClient* Client = CreateClient();
 
 	// Spawn object on server
-	UObjectReplicationBridge::FCreateNetRefHandleParams Params = Server->GetReplicationBridge()->DefaultCreateNetRefHandleParams;
+	UObjectReplicationBridge::FCreateNetRefHandleParams Params;
 	Params.PollFrequency = Server->ConvertPollPeriodIntoFrequency(14U);
 	UTestReplicatedIrisObject* ServerObject = Server->CreateObject(Params);
 
diff --git a/Engine/Plugins/Runtime/ReplicationSystemTestPlugin/Source/Private/Tests/ReplicationSystem/TestDirtyNetObjectTracker.cpp b/Engine/Plugins/Runtime/ReplicationSystemTestPlugin/Source/Private/Tests/ReplicationSystem/TestDirtyNetObjectTracker.cpp
--- a/Engine/Plugins/Runtime/ReplicationSystemTestPlugin/Source/Private/Tests/ReplicationSystem/TestDirtyNetObjectTracker.cpp
+++ b/Engine/Plugins/Runtime/ReplicationSystemTestPlugin/Source/Private/Tests/ReplicationSystem/TestDirtyNetObjectTracker.cpp
@@ -146,7 +146,6 @@ UE_NET_TEST_FIXTURE(FReplicationSystemServerClientTestFixture, GlobalDirtyTracke
 	const uint32 PollPeriod = 2;
 	const float PollFrequency = Server->ConvertPollPeriodIntoFrequency(PollPeriod);
 	Params.PollFrequency = PollFrequency;
-	Params.bCanReceive = true;
 	Params.bUseClassConfigDynamicFilter = true;
 	Params.bNeedsPreUpdate = true;
 
@@ -191,7 +190,6 @@ UE_NET_TEST_FIXTURE(FReplicationSystemServerClientTestFixture, NetForceUpdateOth
 	const uint32 PollPeriod = 100;
 	const float PollFrequency = Server->ConvertPollPeriodIntoFrequency(PollPeriod);
 	Params.PollFrequency = PollFrequency;
-	Params.bCanReceive = true;
 	Params.bUseClassConfigDynamicFilter = true;
 	Params.bNeedsPreUpdate = true;
 	UTestReplicatedIrisObject* ServerObjectA = Server->CreateObject(Params);
@@ -309,7 +307,6 @@ UE_NET_TEST_FIXTURE(FReplicationSystemServerClientTestFixture, DirtyOtherObjectI
 
 	// Spawn object on server that is polled every frame.
 	UObjectReplicationBridge::FCreateNetRefHandleParams Params;
-	Params.bCanReceive = true;
 	Params.bUseClassConfigDynamicFilter = true;
 	Params.bNeedsPreUpdate = true;
 
diff --git a/Engine/Plugins/Runtime/ReplicationSystemTestPlugin/Source/Private/Tests/ReplicationSystem/TestObjectPollFrequencyLimiter.cpp b/Engine/Plugins/Runtime/ReplicationSystemTestPlugin/Source/Private/Tests/ReplicationSystem/TestObjectPollFrequencyLimiter.cpp
--- a/Engine/Plugins/Runtime/ReplicationSystemTestPlugin/Source/Private/Tests/ReplicationSystem/TestObjectPollFrequencyLimiter.cpp
+++ b/Engine/Plugins/Runtime/ReplicationSystemTestPlugin/Source/Private/Tests/ReplicationSystem/TestObjectPollFrequencyLimiter.cpp
@@ -123,7 +123,6 @@ UE_NET_TEST_FIXTURE(FReplicationSystemServerClientTestFixture, TestSubObjectsPol
 	const float PollFrequency = Server->ConvertPollPeriodIntoFrequency(PollPeriod);
 	UObjectReplicationBridge::FCreateNetRefHandleParams Params;
 	Params.PollFrequency = PollFrequency;
-	Params.bCanReceive = true;
 	Params.bUseClassConfigDynamicFilter = true;
 	Params.bNeedsPreUpdate = true;
 	UTestReplicatedIrisObject* ServerRootObject = Server->CreateObject(Params);
diff --git a/Engine/Plugins/Runtime/ReplicationSystemTestPlugin/Source/Private/Tests/ReplicationSystem/TestPreUpdateCallback.cpp b/Engine/Plugins/Runtime/ReplicationSystemTestPlugin/Source/Private/Tests/ReplicationSystem/TestPreUpdateCallback.cpp
--- a/Engine/Plugins/Runtime/ReplicationSystemTestPlugin/Source/Private/Tests/ReplicationSystem/TestPreUpdateCallback.cpp
+++ b/Engine/Plugins/Runtime/ReplicationSystemTestPlugin/Source/Private/Tests/ReplicationSystem/TestPreUpdateCallback.cpp
@@ -25,7 +25,6 @@ UE_NET_TEST_FIXTURE(FReplicationSystemServerClientTestFixture, CreateObjectsInsi
 	const uint32 PollPeriod = 100;
 	const float PollFrequency = Server->ConvertPollPeriodIntoFrequency(PollPeriod);
 	Params.PollFrequency = PollFrequency;
-	Params.bCanReceive = true;
 	Params.bUseClassConfigDynamicFilter = true;
 	Params.bNeedsPreUpdate = true;
 	UTestReplicatedIrisObject* ServerObject = Server->CreateObject(Params);
diff --git a/Engine/Plugins/Runtime/ReplicationSystemTestPlugin/Source/Private/Tests/ReplicationSystem/TestReplicationOperations.cpp b/Engine/Plugins/Runtime/ReplicationSystemTestPlugin/Source/Private/Tests/ReplicationSystem/TestReplicationOperations.cpp
--- a/Engine/Plugins/Runtime/ReplicationSystemTestPlugin/Source/Private/Tests/ReplicationSystem/TestReplicationOperations.cpp
+++ b/Engine/Plugins/Runtime/ReplicationSystemTestPlugin/Source/Private/Tests/ReplicationSystem/TestReplicationOperations.cpp
@@ -1,6 +1,8 @@
 // Copyright Epic Games, Inc. All Rights Reserved.
 
+#include "ReplicationSystemServerClientTestFixture.h"
 #include "ReplicationSystemTestFixture.h"
+
 #include "Iris/ReplicationSystem/ChangeMaskUtil.h"
 #include "Iris/ReplicationSystem/ReplicationOperations.h"
 #include "Iris/ReplicationSystem/ReplicationOperationsInternal.h"
@@ -347,90 +349,55 @@ UE_NET_TEST_FIXTURE(FTestReplicationOperationsFixture, CanSerializeAndDeserializ
 	ReplicationBridge->EndReplication(TestObject);
 }
 
-UE_NET_TEST_FIXTURE(FTestReplicationOperationsFixture, PreAndPostNetReceivedOnlyCalledOnce)
+UE_NET_TEST_FIXTURE(FReplicationSystemServerClientTestFixture, PreAndPostNetReceivedOnlyCalledOnce)
 {
-	UTestReplicatedIrisObject* TestObject = CreateObject(1, 0);
-	
-	// Create NetRefHandle for the CreatedHandle,
-	FNetRefHandle CreatedHandle = ReplicationBridge->BeginReplication(TestObject);
-	const FReplicationProtocol* ProtocolA = ReplicationSystem->GetReplicationProtocol(CreatedHandle);
-
-	UE_NET_ASSERT_EQ(3U, ProtocolA->ReplicationStateCount);
+	// Create replicated object 
+	UTestReplicatedIrisObject* ServerTestObject = Server->CreateObject(1, 0);
+	Server->ReplicationBridge->BeginReplication(ServerTestObject);
 
 	// Set some values
-	TestObject->Components[0]->IntA = 3;
-	TestObject->Components[0]->IntB = 3;
-
-	// Poll data to update property replication state
-	FReplicationInstanceOperations::PollAndCopyPropertyData(ReplicationBridge->GetReplicationInstanceProtocol(TestObject->NetRefHandle));
-
-	// Quantize to buffer
-	uint8* AlignedBuffer = (uint8*)(GetTempAllocator()).Alloc(ProtocolA->InternalTotalSize, ProtocolA->InternalTotalAlignment);
-
-	// Copy data from test object to buffer
-	FNetSerializationContext NetContext;
-	NetContext.SetIsInitState(true);
-	FReplicationInstanceOperations::Quantize(NetContext, AlignedBuffer, GetChangeMaskWriter(ProtocolA->ChangeMaskBitCount), ReplicationBridge->GetReplicationInstanceProtocol(TestObject->NetRefHandle), ProtocolA);
-
-	// Create second object 
-	UTestReplicatedIrisObject* TestObjectB = CreateObject(1, 0);
-	ReplicationBridge->BeginReplication(TestObjectB);
-	const FReplicationProtocol* ProtocolB = ReplicationSystem->GetReplicationProtocol(TestObjectB->NetRefHandle);
-	UE_NET_ASSERT_EQ(ProtocolA, ProtocolB);
-
-	// Check values before we try to apply data
-	UE_NET_ASSERT_EQ(0U, TestObjectB->Components[0]->CallCounts.PreNetReceiveCounter);
-	UE_NET_ASSERT_EQ(0U, TestObjectB->Components[0]->CallCounts.PostNetReceiveCounter);
+	ServerTestObject->Components[0]->IntA = 3;
+	ServerTestObject->Components[0]->IntB = 3;
 
-	// Push state data to ObjectB
-	FReplicationInstanceOperations::DequantizeAndApply(NetContext, GetTempAllocator(), GetChangeMaskStorage(), ReplicationBridge->GetReplicationInstanceProtocol(TestObjectB->NetRefHandle), AlignedBuffer, ProtocolB);
+	// Add a client
+	FReplicationSystemTestClient* Client = CreateClient();
 
-	// Check value after apply
-	UE_NET_ASSERT_EQ(1U, TestObjectB->Components[0]->CallCounts.PreNetReceiveCounter);
-	UE_NET_ASSERT_EQ(1U, TestObjectB->Components[0]->CallCounts.PostNetReceiveCounter);
+	// Spawn the replicated instance
+	Server->UpdateAndSend({ Client });
 
-	ReplicationBridge->EndReplication(TestObject);
-	ReplicationBridge->EndReplication(TestObjectB);
+	// Find the client instance and make sure it replicated
+	UTestReplicatedIrisObject* ClientTestObject = Cast<UTestReplicatedIrisObject>(Client->ReplicationBridge->GetReplicatedObject(ServerTestObject->NetRefHandle));
+	UE_NET_ASSERT_NE(ClientTestObject, nullptr);
+	UE_NET_ASSERT_EQ(3, ClientTestObject->Components[0]->IntA);
+	UE_NET_ASSERT_EQ(3, ClientTestObject->Components[0]->IntB);
+	
+	// Check NetReceive callbacks
+	UE_NET_ASSERT_EQ(1U, ClientTestObject->Components[0]->CallCounts.PreNetReceiveCounter);
+	UE_NET_ASSERT_EQ(1U, ClientTestObject->Components[0]->CallCounts.PostNetReceiveCounter);
 }
 
-UE_NET_TEST_FIXTURE(FTestReplicationOperationsFixture, RepNotifyIsCalledForSimpleProperty)
+UE_NET_TEST_FIXTURE(FReplicationSystemServerClientTestFixture, RepNotifyIsCalledForSimpleProperty)
 {
-	constexpr uint32 PropertyComponentCount = 1;
-	UTestReplicatedIrisObject* TestObject = CreateObject(PropertyComponentCount, 0);
-	
-	// Create NetRefHandle for the CreatedHandle,
-	FNetRefHandle CreatedHandle = ReplicationBridge->BeginReplication(TestObject);
-	const FReplicationProtocol* Protocol = ReplicationSystem->GetReplicationProtocol(CreatedHandle);
+	// Create replicated object 
+	UTestReplicatedIrisObject* ServerTestObject = Server->CreateObject(1, 0);
+	Server->ReplicationBridge->BeginReplication(ServerTestObject);
 
 	// Set some values
-	TestObject->Components[0]->IntB ^= 1;
+	ServerTestObject->Components[0]->IntB = 3;
 
-	// Poll data to update property replication state
-	FReplicationInstanceOperations::PollAndCopyPropertyData(ReplicationBridge->GetReplicationInstanceProtocol(TestObject->NetRefHandle));
+	// Add a client
+	FReplicationSystemTestClient* Client = CreateClient();
 
-	// Quantize to buffer
-	uint8* AlignedBuffer = (uint8*)(GetTempAllocator()).Alloc(Protocol->InternalTotalSize, Protocol->InternalTotalAlignment);
+	// Spawn the replicated instance
+	Server->UpdateAndSend({ Client });
 
-	// Copy data from test object to buffer
-	FNetSerializationContext NetContext;
-	NetContext.SetIsInitState(true);
-	FReplicationInstanceOperations::Quantize(NetContext, AlignedBuffer, GetChangeMaskWriter(Protocol->ChangeMaskBitCount), ReplicationBridge->GetReplicationInstanceProtocol(TestObject->NetRefHandle), Protocol);
-
-	// Create second object 
-	UTestReplicatedIrisObject* TestObjectB = CreateObject(PropertyComponentCount, 0);
-	ReplicationBridge->BeginReplication(TestObjectB);
-
-	// Check that we've not yet gotten any RepNotify calls
-	UE_NET_ASSERT_EQ(TestObjectB->Components[0]->CallCounts.IntBRepNotifyCounter, 0U);
-
-	// Push state data to ObjectB
-	FReplicationInstanceOperations::DequantizeAndApply(NetContext, GetTempAllocator(), GetChangeMaskStorage(), ReplicationBridge->GetReplicationInstanceProtocol(TestObjectB->NetRefHandle), AlignedBuffer, Protocol);
+	// Find the client instance and make sure it replicated
+	UTestReplicatedIrisObject* ClientTestObject = Cast<UTestReplicatedIrisObject>(Client->ReplicationBridge->GetReplicatedObject(ServerTestObject->NetRefHandle));
+	UE_NET_ASSERT_NE(ClientTestObject, nullptr);
+	UE_NET_ASSERT_EQ(3, ClientTestObject->Components[0]->IntB);
 
 	// Check that we got exactly one RepNotify call
-	UE_NET_ASSERT_EQ(TestObjectB->Components[0]->CallCounts.IntBRepNotifyCounter, 1U);
-
-	ReplicationBridge->EndReplication(TestObject);
-	ReplicationBridge->EndReplication(TestObjectB);
+	UE_NET_ASSERT_EQ(ClientTestObject->Components[0]->CallCounts.IntBRepNotifyCounter, 1U);
 }
 
 UE_NET_TEST_FIXTURE(FTestReplicationOperationsFixture, InitPropertyIsNotUpdatedWhenInitStateIsExcluded)
@@ -476,58 +443,36 @@ UE_NET_TEST_FIXTURE(FTestReplicationOperationsFixture, InitPropertyIsNotUpdatedW
 	ReplicationBridge->EndReplication(TestObjectB);
 }
 
-UE_NET_TEST_FIXTURE(FTestReplicationOperationsFixture, RepNotifyIsCalledExactlyOnceForArrays)
+UE_NET_TEST_FIXTURE(FReplicationSystemServerClientTestFixture, RepNotifyIsCalledExactlyOnceForArrays)
 {
-	constexpr uint32 DynamicStateComponentCount = 1;
-	UTestReplicatedIrisObject* TestObject = CreateObjectWithDynamicState(0, 0, DynamicStateComponentCount);
-	
-	// Create NetRefHandle for the TestObject
-	FNetRefHandle CreatedHandle = ReplicationBridge->BeginReplication(TestObject);
-	const FReplicationProtocol* Protocol = ReplicationSystem->GetReplicationProtocol(CreatedHandle);
+	// Create replicated object 
+	UTestReplicatedIrisObject* ServerTestObject = Server->CreateObjectWithDynamicState(0, 0, 1);
+	Server->ReplicationBridge->BeginReplication(ServerTestObject);
 
 	// Set some values
-	TestObject->DynamicStateComponents[0]->IntArray.SetNum(4);
-	TestObject->DynamicStateComponents[0]->IntArray[0] = 1;
-	TestObject->DynamicStateComponents[0]->IntArray[1] = 1;
-	TestObject->DynamicStateComponents[0]->IntArray[2] = 7;
-	TestObject->DynamicStateComponents[0]->IntArray[3] = 4;
+	ServerTestObject->DynamicStateComponents[0]->IntArray.SetNum(4);
+	ServerTestObject->DynamicStateComponents[0]->IntArray[0] = 1;
+	ServerTestObject->DynamicStateComponents[0]->IntArray[1] = 1;
+	ServerTestObject->DynamicStateComponents[0]->IntArray[2] = 7;
+	ServerTestObject->DynamicStateComponents[0]->IntArray[3] = 4;
 
-	TestObject->DynamicStateComponents[0]->IntStaticArray[1] = 5;
-	TestObject->DynamicStateComponents[0]->IntStaticArray[5] = 1;
+	ServerTestObject->DynamicStateComponents[0]->IntStaticArray[1] = 5;
+	ServerTestObject->DynamicStateComponents[0]->IntStaticArray[5] = 1;
 
-	// Poll data to update property replication state
-	FReplicationInstanceOperations::PollAndCopyPropertyData(ReplicationBridge->GetReplicationInstanceProtocol(TestObject->NetRefHandle));
+	// Add a client
+	FReplicationSystemTestClient* Client = CreateClient();
 
-	// Quantize to buffer
-	uint8* AlignedBuffer = (uint8*)(GetTempAllocator()).Alloc(Protocol->InternalTotalSize, Protocol->InternalTotalAlignment);
-	FMemory::Memzero(AlignedBuffer, Protocol->InternalTotalSize);
-
-	// Copy data from test object to buffer
-	FNetSerializationContext SerializationContext;
-	FInternalNetSerializationContext InternalContext;
-	SerializationContext.SetInternalContext(&InternalContext);
-	SerializationContext.SetIsInitState(true);
-	FReplicationInstanceOperations::Quantize(SerializationContext, AlignedBuffer, GetChangeMaskWriter(Protocol->ChangeMaskBitCount), ReplicationBridge->GetReplicationInstanceProtocol(TestObject->NetRefHandle), Protocol);
-
-	// Create second object 
-	UTestReplicatedIrisObject* TestObjectB = CreateObjectWithDynamicState(0, 0, DynamicStateComponentCount);
-	ReplicationBridge->BeginReplication(TestObjectB);
+	// Spawn the replicated instance
+	Server->UpdateAndSend({ Client });
 
-	// Check that we've not yet gotten any RepNotify calls
-	UE_NET_ASSERT_EQ(TestObjectB->DynamicStateComponents[0]->CallCounts.IntArrayRepNotifyCounter, 0U);
-	UE_NET_ASSERT_EQ(TestObjectB->DynamicStateComponents[0]->CallCounts.IntStaticArrayRepNotifyCounter, 0U);
-
-	// Push state data to ObjectB
-	FReplicationInstanceOperations::DequantizeAndApply(SerializationContext, GetTempAllocator(), GetChangeMaskStorage(), ReplicationBridge->GetReplicationInstanceProtocol(TestObjectB->NetRefHandle), AlignedBuffer, Protocol);
+	// Find the client instance and make sure it replicated
+	UTestReplicatedIrisObject* ClientTestObject = Cast<UTestReplicatedIrisObject>(Client->ReplicationBridge->GetReplicatedObject(ServerTestObject->NetRefHandle));
+	UE_NET_ASSERT_NE(ClientTestObject, nullptr);
+	UE_NET_ASSERT_EQ(4, ClientTestObject->DynamicStateComponents[0]->IntArray.Num());
 
 	// Check that we got exactly one RepNotify call per array
-	UE_NET_ASSERT_EQ(TestObjectB->DynamicStateComponents[0]->CallCounts.IntArrayRepNotifyCounter, 1U);
-	UE_NET_ASSERT_EQ(TestObjectB->DynamicStateComponents[0]->CallCounts.IntStaticArrayRepNotifyCounter, 1U);
-
-	FReplicationProtocolOperations::FreeDynamicState(SerializationContext, AlignedBuffer, Protocol);
-
-	ReplicationBridge->EndReplication(TestObject);
-	ReplicationBridge->EndReplication(TestObjectB);
+	UE_NET_ASSERT_EQ(ClientTestObject->DynamicStateComponents[0]->CallCounts.IntArrayRepNotifyCounter, 1U);
+	UE_NET_ASSERT_EQ(ClientTestObject->DynamicStateComponents[0]->CallCounts.IntStaticArrayRepNotifyCounter, 1U);
 }
 
 class FTestReplicationOperationsForObjectsFixture : public FTestReplicationOperationsFixture
diff --git a/Engine/Plugins/Runtime/ReplicationSystemTestPlugin/Source/Private/Tests/ReplicationSystem/TestReplicationSystem.cpp b/Engine/Plugins/Runtime/ReplicationSystemTestPlugin/Source/Private/Tests/ReplicationSystem/TestReplicationSystem.cpp
--- a/Engine/Plugins/Runtime/ReplicationSystemTestPlugin/Source/Private/Tests/ReplicationSystem/TestReplicationSystem.cpp
+++ b/Engine/Plugins/Runtime/ReplicationSystemTestPlugin/Source/Private/Tests/ReplicationSystem/TestReplicationSystem.cpp
@@ -1702,7 +1702,7 @@ UE_NET_TEST_FIXTURE(FReplicationSystemServerClientTestFixture, TestObjectPollFra
 	UTestReplicatedIrisObject* ServerObject = Server->CreateObject(0, 0);
 
 	// Spawn second object on server that later will be added as a dependent object
-	UObjectReplicationBridge::FCreateNetRefHandleParams Params = Server->GetReplicationBridge()->DefaultCreateNetRefHandleParams;
+	UObjectReplicationBridge::FCreateNetRefHandleParams Params;
 	Params.PollFrequency = Server->ConvertPollPeriodIntoFrequency(1U);
 	UTestReplicatedIrisObject* ServerObjectPolledEveryOtherFrame = Server->CreateObject(Params);
 
diff --git a/Engine/Source/Runtime/Engine/Private/Net/Iris/ReplicationSystem/ActorReplicationBridge.cpp b/Engine/Source/Runtime/Engine/Private/Net/Iris/ReplicationSystem/ActorReplicationBridge.cpp
--- a/Engine/Source/Runtime/Engine/Private/Net/Iris/ReplicationSystem/ActorReplicationBridge.cpp
+++ b/Engine/Source/Runtime/Engine/Private/Net/Iris/ReplicationSystem/ActorReplicationBridge.cpp
@@ -267,12 +267,13 @@ UE::Net::FNetRefHandle UActorReplicationBridge::BeginReplication(AActor* Actor,
 	}
 
 	// Create handles for the registered fragments
-	Super::FCreateNetRefHandleParams CreateNetRefHandleParams = UObjectReplicationBridge::DefaultCreateNetRefHandleParams;
-	CreateNetRefHandleParams.bNeedsPreUpdate = 1U;
-	CreateNetRefHandleParams.bNeedsWorldLocationUpdate = 1U;
-	CreateNetRefHandleParams.bIsDormant = Actor->NetDormancy > DORM_Awake;
-	CreateNetRefHandleParams.StaticPriority = (Actor->bAlwaysRelevant || Actor->bOnlyRelevantToOwner) ? Actor->NetPriority : 0.0f;
-	CreateNetRefHandleParams.PollFrequency = Actor->GetNetUpdateFrequency();
+	Super::FCreateNetRefHandleParams CreateNetRefHandleParams = {
+		.bNeedsPreUpdate = 1U,
+		.bNeedsWorldLocationUpdate = 1U,
+		.bIsDormant = Actor->NetDormancy > DORM_Awake,
+		.StaticPriority = (Actor->bAlwaysRelevant || Actor->bOnlyRelevantToOwner) ? Actor->NetPriority : 0.0f,
+		.PollFrequency = Actor->GetNetUpdateFrequency()
+	};
 
 #if !UE_BUILD_SHIPPING
 	ensureMsgf(!(Actor->bAlwaysRelevant || Actor->bOnlyRelevantToOwner) || CreateNetRefHandleParams.StaticPriority >= 1.0f, TEXT("Very low NetPriority %.02f for always relevant or owner relevant Actor %s. Set it to 1.0f or higher."), Actor->NetPriority, ToCStr(Actor->GetName()));
@@ -335,7 +336,8 @@ UE::Net::FNetRefHandle UActorReplicationBridge::BeginReplication(AActor* Actor,
 			UObject* SubObjectToReplicate = SubObjectInfo.GetSubObject();
 			if (IsValid(SubObjectToReplicate) && SubObjectInfo.NetCondition != ELifetimeCondition::COND_Never)
 			{
-				FNetRefHandle SubObjectRefHandle = UObjectReplicationBridge::BeginReplication(ActorRefHandle, SubObjectToReplicate);
+				const UObjectReplicationBridge::FCreateNetRefHandleParams CreateNetRefParams;
+				FNetRefHandle SubObjectRefHandle = UObjectReplicationBridge::BeginReplication(ActorRefHandle, SubObjectToReplicate, CreateNetRefParams);
 				if (SubObjectRefHandle.IsValid() && SubObjectInfo.NetCondition != ELifetimeCondition::COND_None)
 				{
 					UObjectReplicationBridge::SetSubObjectNetCondition(SubObjectRefHandle, SubObjectInfo.NetCondition);
@@ -394,7 +396,7 @@ UE::Net::FNetRefHandle UActorReplicationBridge::BeginReplication(FNetRefHandle O
 		}
 
 		// Start replicating the subobject with its owner.
-		ReplicatedComponentHandle = Super::BeginReplication(OwnerHandle, SubObject);
+		ReplicatedComponentHandle = Super::BeginReplication(OwnerHandle, SubObject, UObjectReplicationBridge::FCreateNetRefHandleParams());
 	}
 
 	if (!ReplicatedComponentHandle.IsValid())
@@ -410,12 +412,13 @@ UE::Net::FNetRefHandle UActorReplicationBridge::BeginReplication(FNetRefHandle O
 	}
 
 	// Begin replication for any SubObjects registered by the component
+	const UObjectReplicationBridge::FCreateNetRefHandleParams CreateNetRefParams;
 	for (const FSubObjectRegistry::FEntry& SubObjectInfo : RepComponentInfo->SubObjects.GetRegistryList())
 	{
 		UObject* SubObjectToReplicate = SubObjectInfo.GetSubObject();
 		if (IsValid(SubObjectToReplicate) && SubObjectInfo.NetCondition != ELifetimeCondition::COND_Never)
 		{
-			FNetRefHandle SubObjectHandle = UObjectReplicationBridge::BeginReplication(OwnerHandle, SubObjectToReplicate, ReplicatedComponentHandle, UReplicationBridge::ESubObjectInsertionOrder::ReplicateWith);
+			FNetRefHandle SubObjectHandle = UObjectReplicationBridge::BeginReplication(OwnerHandle, SubObjectToReplicate, ReplicatedComponentHandle, CreateNetRefParams, UReplicationBridge::ESubObjectInsertionOrder::ReplicateWith);
 			if (SubObjectHandle.IsValid() && SubObjectInfo.NetCondition != ELifetimeCondition::COND_None)
 			{
 				SetSubObjectNetCondition(SubObjectHandle, SubObjectInfo.NetCondition);
@@ -995,8 +998,6 @@ void UActorReplicationBridge::GetInitialDependencies(FNetRefHandle Handle, FNetD
 
 void UActorReplicationBridge::SetNetDriver(UNetDriver* const InNetDriver)
 {
-	Super::SetNetDriver(InNetDriver);
-
 	if (NetDriver)
 	{
 		NetDriver->OnNetServerMaxTickRateChanged.RemoveAll(this);
diff --git a/Engine/Source/Runtime/Engine/Private/Net/Iris/ReplicationSystem/ReplicationSystemUtil.cpp b/Engine/Source/Runtime/Engine/Private/Net/Iris/ReplicationSystem/ReplicationSystemUtil.cpp
--- a/Engine/Source/Runtime/Engine/Private/Net/Iris/ReplicationSystem/ReplicationSystemUtil.cpp
+++ b/Engine/Source/Runtime/Engine/Private/Net/Iris/ReplicationSystem/ReplicationSystemUtil.cpp
@@ -160,13 +160,13 @@ void FReplicationSystemUtil::BeginReplicationForActorComponent(FNetHandle ActorH
 		return;
 	}
 
-	if (!ensureAlways(ActorComp != nullptr))
+	if (!ensure(ActorComp != nullptr))
 	{
 		return;
 	}
 
 	AActor* Actor = Cast<AActor>(ActorComp->GetOwner());
-	ensureAlways(FNetHandleManager::GetNetHandle(Actor) == ActorHandle);
+	ensureMsgf(FNetHandleManager::GetNetHandle(Actor) == ActorHandle, TEXT("BeginReplicationForActorComponent received invalid owner handle %s for actual owner %s"), *ActorHandle.ToString(), *GetNameSafe(Actor));
 	if (Actor)
 	{
 		if (const UWorld* World = Actor->GetWorld())
@@ -242,7 +242,8 @@ void FReplicationSystemUtil::BeginReplicationForActorSubObject(const AActor* Act
 					const FNetRefHandle ActorRefHandle = Bridge->GetReplicatedRefHandle(ActorHandle);
 					if (ActorRefHandle.IsValid())
 					{
-						const FNetRefHandle SubObjectRefHandle = Bridge->BeginReplication(ActorRefHandle, ActorSubObject);
+						const UObjectReplicationBridge::FCreateNetRefHandleParams CreateNetRefParams;
+						const FNetRefHandle SubObjectRefHandle = Bridge->BeginReplication(ActorRefHandle, ActorSubObject, CreateNetRefParams);
 						if (SubObjectRefHandle.IsValid() && NetCondition != ELifetimeCondition::COND_None)
 						{
 							Bridge->SetSubObjectNetCondition(SubObjectRefHandle, NetCondition);
@@ -277,7 +278,9 @@ void FReplicationSystemUtil::BeginReplicationForActorComponentSubObject(UActorCo
 						const FNetRefHandle ActorComponentRefHandle = Bridge->GetReplicatedRefHandle(ActorComponent);
 						if (ActorRefHandle.IsValid() && ActorComponentRefHandle.IsValid())
 						{
-							const FNetRefHandle SubObjectRefHandle = Bridge->BeginReplication(ActorRefHandle, SubObject, ActorComponentRefHandle, UReplicationBridge::ESubObjectInsertionOrder::ReplicateWith);
+							const UObjectReplicationBridge::FCreateNetRefHandleParams CreateNetRefParams;
+
+							const FNetRefHandle SubObjectRefHandle = Bridge->BeginReplication(ActorRefHandle, SubObject, ActorComponentRefHandle, CreateNetRefParams, UReplicationBridge::ESubObjectInsertionOrder::ReplicateWith);
 							if (SubObjectRefHandle.IsValid() && NetCondition != ELifetimeCondition::COND_None)
 							{
 								Bridge->SetSubObjectNetCondition(SubObjectRefHandle, NetCondition);
diff --git a/Engine/Source/Runtime/Engine/Private/NetDriver.cpp b/Engine/Source/Runtime/Engine/Private/NetDriver.cpp
--- a/Engine/Source/Runtime/Engine/Private/NetDriver.cpp
+++ b/Engine/Source/Runtime/Engine/Private/NetDriver.cpp
@@ -7096,7 +7096,7 @@ void UNetDriver::CreateReplicationSystem(bool bInitAsClient)
 		return;
 	}
 
-	UReplicationBridge* ReplicationBridge = NewObject<UActorReplicationBridge>(GetTransientPackage(), ReplicationBridgeClass);
+	UActorReplicationBridge* ReplicationBridge = NewObject<UActorReplicationBridge>(GetTransientPackage(), ReplicationBridgeClass);
 	if (ReplicationBridge)
 	{
 		ReplicationBridge->SetNetDriver(this);
diff --git a/Engine/Source/Runtime/Engine/Public/Net/Iris/ReplicationSystem/ActorReplicationBridge.h b/Engine/Source/Runtime/Engine/Public/Net/Iris/ReplicationSystem/ActorReplicationBridge.h
--- a/Engine/Source/Runtime/Engine/Public/Net/Iris/ReplicationSystem/ActorReplicationBridge.h
+++ b/Engine/Source/Runtime/Engine/Public/Net/Iris/ReplicationSystem/ActorReplicationBridge.h
@@ -62,7 +62,7 @@ public:
 	ENGINE_API static UActorReplicationBridge* Create(UNetDriver* NetDriver);
 
 	/** Sets the net driver for the bridge. */
-	ENGINE_API virtual void SetNetDriver(UNetDriver* const InNetDriver) override;
+	ENGINE_API void SetNetDriver(UNetDriver* const InNetDriver);
 	
 	/** Get net driver used by the bridge .*/
 	inline UNetDriver* GetNetDriver() const { return NetDriver; }
diff --git a/Engine/Source/Runtime/Experimental/Iris/Core/Private/Iris/ReplicationSystem/ObjectReplicationBridge.cpp b/Engine/Source/Runtime/Experimental/Iris/Core/Private/Iris/ReplicationSystem/ObjectReplicationBridge.cpp
--- a/Engine/Source/Runtime/Experimental/Iris/Core/Private/Iris/ReplicationSystem/ObjectReplicationBridge.cpp
+++ b/Engine/Source/Runtime/Experimental/Iris/Core/Private/Iris/ReplicationSystem/ObjectReplicationBridge.cpp
@@ -108,18 +108,6 @@ static FAutoConsoleVariableRef CVarUseVerboseCsvStats(
 	TEXT("Whether to use verbose per-class csv stats. Default is false in Shipping, otherwise True.")
 );
 
-UObjectReplicationBridge::FCreateNetRefHandleParams UObjectReplicationBridge::DefaultCreateNetRefHandleParams =
-{
-	.bCanReceive=false, 
-	.bNeedsPreUpdate=false,
-	.bNeedsWorldLocationUpdate=false,
-	.bIsDormant=false,
-	.bUseClassConfigDynamicFilter=true,
-	.bUseExplicitDynamicFilter=false,
-	.StaticPriority=0.0f, 
-	.PollFrequency=0.0f,
-};
-
 namespace UE::Net::Private
 {
 	void CallRegisterReplicationFragments(UObject* Object, FFragmentRegistrationContext& Context, EFragmentRegistrationFlags RegistrationFlags)
@@ -287,7 +275,6 @@ UE::Net::FNetRefHandle UObjectReplicationBridge::BeginReplication(UObject* Insta
 	
 	// Register fragments
 	EReplicationFragmentTraits Traits = EReplicationFragmentTraits::CanReplicate;
-	Traits |= Params.bCanReceive ? EReplicationFragmentTraits::CanReceive : EReplicationFragmentTraits::None;
 	Traits |= Params.bNeedsPreUpdate ? EReplicationFragmentTraits::NeedsPreSendUpdate : EReplicationFragmentTraits::None;
 	Traits |= Params.bNeedsWorldLocationUpdate ? EReplicationFragmentTraits::NeedsWorldLocationUpdate : EReplicationFragmentTraits::None;
 	
@@ -470,7 +457,7 @@ void UObjectReplicationBridge::AssignDynamicFilter(UObject* Instance, const FCre
 	
 }
 
-UE::Net::FNetRefHandle UObjectReplicationBridge::BeginReplication(FNetRefHandle OwnerRefHandle, UObject* Instance, FNetRefHandle InsertRelativeToSubObjectRefHandle, ESubObjectInsertionOrder InsertionOrder, const FCreateNetRefHandleParams& Params)
+UE::Net::FNetRefHandle UObjectReplicationBridge::BeginReplication(FNetRefHandle OwnerRefHandle, UObject* Instance, FNetRefHandle InsertRelativeToSubObjectRefHandle, const FCreateNetRefHandleParams& Params, ESubObjectInsertionOrder InsertionOrder)
 {
 	LLM_SCOPE_BYTAG(IrisState);
 
@@ -775,7 +762,7 @@ FReplicationBridgeCreateNetRefHandleResult UObjectReplicationBridge::CreateNetRe
 		{
 			if (RootObjectOfSubObject.IsValid())
 			{
-				UE_LOG(LogIrisBridge, Warning, TEXT("CreateNetRefHandleFromRemote: Failed to instantiate SubObject NetHandle: %s of RootObject: %s (%s)"), *WantedNetHandle.ToString(), *RootObjectOfSubObject.ToString(), *GetNameSafe(GetReplicatedObject(RootObjectOfSubObject)));
+				UE_LOG(LogIrisBridge, Warning, TEXT("CreateNetRefHandleFromRemote: Failed to instantiate SubObject NetHandle: %s of %s"), *WantedNetHandle.ToString(), *PrintObjectFromNetRefHandle(RootObjectOfSubObject));
 			}
 			else
 			{
diff --git a/Engine/Source/Runtime/Experimental/Iris/Core/Private/Iris/ReplicationSystem/ReplicationBridge.cpp b/Engine/Source/Runtime/Experimental/Iris/Core/Private/Iris/ReplicationSystem/ReplicationBridge.cpp
--- a/Engine/Source/Runtime/Experimental/Iris/Core/Private/Iris/ReplicationSystem/ReplicationBridge.cpp
+++ b/Engine/Source/Runtime/Experimental/Iris/Core/Private/Iris/ReplicationSystem/ReplicationBridge.cpp
@@ -113,10 +113,6 @@ void UReplicationBridge::PruneStaleObjects()
 {
 }
 
-void UReplicationBridge::SetNetDriver(UNetDriver* NetDriver)
-{
-}
-
 void UReplicationBridge::GetInitialDependencies(FNetRefHandle Handle, FNetDependencyInfoArray& OutDependencies) const
 {
 	return;
@@ -613,7 +609,7 @@ void UReplicationBridge::EndReplication(FNetRefHandle Handle, EEndReplicationFla
 			AddPendingEndReplication(Handle, EndReplicationFlags);
 
 			// We do however copy the final state data and mark object to stop propagating state changes
-			InternalTearOff(Handle, ECacheCreationInfo::Yes);
+			InternalTearOff(Handle);
 
 			// Detach instance as we must assume that we should not access the object after this call.
 			ObjectData.bPendingEndReplication = 1U;
@@ -643,7 +639,7 @@ void UReplicationBridge::EndReplication(FNetRefHandle Handle, EEndReplicationFla
 				AddPendingEndReplication(Handle, EndReplicationFlags);
 			
 				// Capture final state
-				InternalFlushStateData(Handle, ECacheCreationInfo::Yes);
+				InternalFlushStateData(Handle);
 
 				// Detach instance as we must assume that we should not access the object after this call.
 				ObjectData.bPendingEndReplication = 1U;
@@ -731,7 +727,7 @@ void UReplicationBridge::TearOffHandlesPendingTearOff()
 	{
 		if (EnumHasAnyFlags(Info.DestroyFlags, EEndReplicationFlags::TearOff))
 		{
-			InternalTearOff(Info.Handle, ECacheCreationInfo::Yes);
+			InternalTearOff(Info.Handle);
 		}
 	}
 }
@@ -789,7 +785,7 @@ void UReplicationBridge::AddPendingEndReplication(FNetRefHandle Handle, EEndRepl
 	}
 }
 
-void UReplicationBridge::InternalFlushStateData(UE::Net::FNetSerializationContext& SerializationContext, UE::Net::Private::FChangeMaskCache& ChangeMaskCache, UE::Net::FNetBitStreamWriter& ChangeMaskWriter, uint32 InternalObjectIndex, ECacheCreationInfo CacheCreationInfo)
+void UReplicationBridge::InternalFlushStateData(UE::Net::FNetSerializationContext& SerializationContext, UE::Net::Private::FChangeMaskCache& ChangeMaskCache, UE::Net::FNetBitStreamWriter& ChangeMaskWriter, uint32 InternalObjectIndex)
 {
 	using namespace UE::Net;
 	using namespace UE::Net::Private;
@@ -817,10 +813,7 @@ void UReplicationBridge::InternalFlushStateData(UE::Net::FNetSerializationContex
 		}
 
 		// Cache creation info
-		if (CacheCreationInfo == ECacheCreationInfo::Yes)
-		{
-			ObjectData.bHasCachedCreationInfo =  CallCacheNetRefHandleCreationInfo(ObjectData.RefHandle) ? 1U : 0U;
-		}
+		ObjectData.bHasCachedCreationInfo =  CallCacheNetRefHandleCreationInfo(ObjectData.RefHandle) ? 1U : 0U;
 
 		FReplicationInstanceOperationsInternal::QuantizeObjectStateData(ChangeMaskWriter, ChangeMaskCache, *NetRefHandleManager, SerializationContext, InternalObjectIndex);
 
@@ -830,13 +823,13 @@ void UReplicationBridge::InternalFlushStateData(UE::Net::FNetSerializationContex
 
 	for (FInternalNetRefIndex SubObjectInternalIndex : NetRefHandleManager->GetChildSubObjects(InternalObjectIndex))
 	{
-		InternalFlushStateData(SerializationContext, ChangeMaskCache, ChangeMaskWriter, SubObjectInternalIndex, CacheCreationInfo);
+		InternalFlushStateData(SerializationContext, ChangeMaskCache, ChangeMaskWriter, SubObjectInternalIndex);
 	}
 
 	// $IRIS TODO:  Should we also clear the DirtyTracker flags for this flushed object ?
 }
 
-void UReplicationBridge::InternalFlushStateData(FNetRefHandle Handle, ECacheCreationInfo CacheCreationInfo)
+void UReplicationBridge::InternalFlushStateData(FNetRefHandle Handle)
 {
 	using namespace UE::Net;
 	using namespace UE::Net::Private;
@@ -860,7 +853,7 @@ void UReplicationBridge::InternalFlushStateData(FNetRefHandle Handle, ECacheCrea
 	SerializationContext.SetInternalContext(&InternalContext);
 	SerializationContext.SetNetStatsContext(ReplicationSystem->GetReplicationSystemInternal()->GetNetTypeStats().GetNetStatsContext());
 
-	InternalFlushStateData(SerializationContext, ChangeMaskCache, ChangeMaskWriter, InternalObjectIndex, CacheCreationInfo);
+	InternalFlushStateData(SerializationContext, ChangeMaskCache, ChangeMaskWriter, InternalObjectIndex);
 
 	// Iterate over connections and propagate dirty changemasks to all connections already scoping this object
 	if (ChangeMaskCache.Indices.Num() > 0)
@@ -879,7 +872,7 @@ void UReplicationBridge::InternalFlushStateData(FNetRefHandle Handle, ECacheCrea
 	}
 }
 
-void UReplicationBridge::InternalTearOff(FNetRefHandle Handle, ECacheCreationInfo CacheCreationInfo)
+void UReplicationBridge::InternalTearOff(FNetRefHandle Handle)
 {
 	using namespace UE::Net;
 	using namespace UE::Net::Private;
@@ -887,7 +880,6 @@ void UReplicationBridge::InternalTearOff(FNetRefHandle Handle, ECacheCreationInf
 	IRIS_PROFILER_SCOPE(InternalTearOff);
 
 	const uint32 InternalObjectIndex = NetRefHandleManager->GetInternalIndex(Handle);
-
 	if (InternalObjectIndex == FNetRefHandleManager::InvalidInternalIndex)
 	{
 		return;
@@ -895,77 +887,77 @@ void UReplicationBridge::InternalTearOff(FNetRefHandle Handle, ECacheCreationInf
 
 	FNetRefHandleManager::FReplicatedObjectData& ObjectData = NetRefHandleManager->GetReplicatedObjectDataNoCheck(InternalObjectIndex);
 
-	// Copy state data and tear off now
-	if (!ObjectData.bTearOff)
+	if (ObjectData.bTearOff)
 	{
-		UE_LOG_REPLICATIONBRIDGE(Verbose, TEXT("TearOff %s ( InteralIndex: %u )"), *Handle.ToString(), InternalObjectIndex);
-
-		// Force copy of final state data as we will detach the object after scope update
-		FChangeMaskCache ChangeMaskCache;
-		FNetBitStreamWriter ChangeMaskWriter;
+		// Already torn off
+		return;
+	}
 
-		// Setup context
-		FNetSerializationContext SerializationContext;
-		FInternalNetSerializationContext InternalContext(ReplicationSystem);
-		SerializationContext.SetInternalContext(&InternalContext);
-		SerializationContext.SetNetStatsContext(ReplicationSystem->GetReplicationSystemInternal()->GetNetTypeStats().GetNetStatsContext());
+	// Copy state data and tear off now
+	UE_LOG_REPLICATIONBRIDGE(Verbose, TEXT("TearOff: %s"), *PrintObjectFromNetRefHandle(Handle));
 
-		if (ObjectData.InstanceProtocol)
-		{
-			if (EnumHasAnyFlags(ObjectData.InstanceProtocol->InstanceTraits, EReplicationInstanceProtocolTraits::NeedsPoll | EReplicationInstanceProtocolTraits::NeedsPreSendUpdate))
-			{
-				CallPreSendUpdateSingleHandle(Handle);
-			} 
+	// Force copy of final state data as we will detach the object after scope update
+	FChangeMaskCache ChangeMaskCache;
+	FNetBitStreamWriter ChangeMaskWriter;
 
-			// Cache creation info
-			if (CacheCreationInfo == ECacheCreationInfo::Yes)
-			{
-				ObjectData.bHasCachedCreationInfo =  CallCacheNetRefHandleCreationInfo(Handle) ? 1U : 0U;
-			}
-		}
+	// Setup context
+	FNetSerializationContext SerializationContext;
+	FInternalNetSerializationContext InternalContext(ReplicationSystem);
+	SerializationContext.SetInternalContext(&InternalContext);
+	SerializationContext.SetNetStatsContext(ReplicationSystem->GetReplicationSystemInternal()->GetNetTypeStats().GetNetStatsContext());
 
-		if (ObjectData.InstanceProtocol && ObjectData.Protocol->InternalTotalSize > 0U)
+	if (ObjectData.InstanceProtocol)
+	{
+		if (EnumHasAnyFlags(ObjectData.InstanceProtocol->InstanceTraits, EReplicationInstanceProtocolTraits::NeedsPoll | EReplicationInstanceProtocolTraits::NeedsPreSendUpdate))
 		{
-			FReplicationInstanceOperationsInternal::QuantizeObjectStateData(ChangeMaskWriter, ChangeMaskCache, *NetRefHandleManager, SerializationContext, InternalObjectIndex);
+			CallPreSendUpdateSingleHandle(Handle);
+		} 
 
-			// Clear the quantize flag since it was done directly here.
-			NetRefHandleManager->GetDirtyObjectsToQuantize().ClearBit(InternalObjectIndex);
-		}
-		else
-		{
-			// Nothing to copy, but we must still propagate the tear-off state.
-			FChangeMaskCache::FCachedInfo& Info = ChangeMaskCache.AddEmptyChangeMaskForObject(InternalObjectIndex);
-			// If we are a subobject we must also mark owner as dirty.
-			const uint32 SubObjectOwnerIndex = ObjectData.SubObjectRootIndex;
-			if (SubObjectOwnerIndex != FNetRefHandleManager::InvalidInternalIndex) 
-			{
-				ChangeMaskCache.AddSubObjectOwnerDirty(SubObjectOwnerIndex);
-			}			
-		}
+		// Cache creation info
+		ObjectData.bHasCachedCreationInfo =  CallCacheNetRefHandleCreationInfo(Handle) ? 1U : 0U;
+	}
 
-		// Propagate changes to all connections that we currently have in scope
-		FReplicationConnections& Connections = ReplicationSystem->GetReplicationSystemInternal()->GetConnections();
+	if (ObjectData.InstanceProtocol && ObjectData.Protocol->InternalTotalSize > 0U)
+	{
+		FReplicationInstanceOperationsInternal::QuantizeObjectStateData(ChangeMaskWriter, ChangeMaskCache, *NetRefHandleManager, SerializationContext, InternalObjectIndex);
 
-		// Iterate over connections and propagate dirty changemasks to all connections already scoping this object
-		auto UpdateDirtyChangeMasks = [&Connections, &ChangeMaskCache](uint32 ConnectionId)
+		// Clear the quantize flag since it was done directly here.
+		NetRefHandleManager->GetDirtyObjectsToQuantize().ClearBit(InternalObjectIndex);
+	}
+	else
+	{
+		// Nothing to copy, but we must still propagate the tear-off state.
+		FChangeMaskCache::FCachedInfo& Info = ChangeMaskCache.AddEmptyChangeMaskForObject(InternalObjectIndex);
+		// If we are a subobject we must also mark owner as dirty.
+		const uint32 SubObjectOwnerIndex = ObjectData.SubObjectRootIndex;
+		if (SubObjectOwnerIndex != FNetRefHandleManager::InvalidInternalIndex) 
 		{
-			FReplicationConnection* Conn = Connections.GetConnection(ConnectionId);
-			const bool bMarkForTearOff = true;
-			Conn->ReplicationWriter->ForceUpdateDirtyChangeMasks(ChangeMaskCache, FReplicationWriter::FlushFlags_None, bMarkForTearOff);
-		};
-		const FNetBitArray& ValidConnections = Connections.GetValidConnections();
-		ValidConnections.ForAllSetBits(UpdateDirtyChangeMasks);		
+			ChangeMaskCache.AddSubObjectOwnerDirty(SubObjectOwnerIndex);
+		}			
+	}
 
-		// TearOff subobjects as well.
-		for (FInternalNetRefIndex SubObjectInternalIndex : NetRefHandleManager->GetChildSubObjects(InternalObjectIndex))
-		{
-			InternalTearOff(NetRefHandleManager->GetNetRefHandleFromInternalIndex(SubObjectInternalIndex), CacheCreationInfo);
-		}	
+	// Propagate changes to all connections that we currently have in scope
+	FReplicationConnections& Connections = ReplicationSystem->GetReplicationSystemInternal()->GetConnections();
 
-		// Mark object as being torn-off and that we should no longer propagate state changes
-		ObjectData.bTearOff = 1U;
-		ObjectData.bShouldPropagateChangedStates = 0U;
-	}
+	// Iterate over connections and propagate dirty changemasks to all connections already scoping this object
+	auto UpdateDirtyChangeMasks = [&Connections, &ChangeMaskCache](uint32 ConnectionId)
+	{
+		FReplicationConnection* Conn = Connections.GetConnection(ConnectionId);
+		const bool bMarkForTearOff = true;
+		Conn->ReplicationWriter->ForceUpdateDirtyChangeMasks(ChangeMaskCache, FReplicationWriter::FlushFlags_None, bMarkForTearOff);
+	};
+	const FNetBitArray& ValidConnections = Connections.GetValidConnections();
+	ValidConnections.ForAllSetBits(UpdateDirtyChangeMasks);		
+
+	// TearOff subobjects as well.
+	for (FInternalNetRefIndex SubObjectInternalIndex : NetRefHandleManager->GetChildSubObjects(InternalObjectIndex))
+	{
+		InternalTearOff(NetRefHandleManager->GetNetRefHandleFromInternalIndex(SubObjectInternalIndex));
+	}	
+
+	// Mark object as being torn-off and that we should no longer propagate state changes
+	ObjectData.bTearOff = 1U;
+	ObjectData.bShouldPropagateChangedStates = 0U;
 }
 
 UE::Net::FNetRefHandle UReplicationBridge::InternalAddDestructionInfo(FNetRefHandle Handle, const FEndReplicationParameters& Parameters)
diff --git a/Engine/Source/Runtime/Experimental/Iris/Core/Private/Iris/ReplicationSystem/ReplicationReader.cpp b/Engine/Source/Runtime/Experimental/Iris/Core/Private/Iris/ReplicationSystem/ReplicationReader.cpp
--- a/Engine/Source/Runtime/Experimental/Iris/Core/Private/Iris/ReplicationSystem/ReplicationReader.cpp
+++ b/Engine/Source/Runtime/Experimental/Iris/Core/Private/Iris/ReplicationSystem/ReplicationReader.cpp
@@ -1852,7 +1852,7 @@ void FReplicationReader::DispatchStateData(FNetSerializationContext& Context)
 		PostDispatchObjectInfo.DequantizeAndApplyContext = nullptr;
 		PostDispatchObjectInfo.AttachmentDispatchedFlags = ENetObjectAttachmentDispatchFlags::None;
 
-		// For SubObjects we call must call this method after applying state data for the owner, in order to remain backwards compatible.
+		// For SubObjects we call this method after applying state data for the owner, in order to remain backwards compatible.
 		if (Info.bShouldCallSubObjectCreatedFromReplication)
 		{
 			if (ObjectData.SubObjectRootIndex != FNetRefHandleManager::InvalidInternalIndex)
diff --git a/Engine/Source/Runtime/Experimental/Iris/Core/Private/Iris/ReplicationSystem/ReplicationWriter.cpp b/Engine/Source/Runtime/Experimental/Iris/Core/Private/Iris/ReplicationSystem/ReplicationWriter.cpp
--- a/Engine/Source/Runtime/Experimental/Iris/Core/Private/Iris/ReplicationSystem/ReplicationWriter.cpp
+++ b/Engine/Source/Runtime/Experimental/Iris/Core/Private/Iris/ReplicationSystem/ReplicationWriter.cpp
@@ -364,7 +364,7 @@ void FReplicationWriter::Init(const FReplicationParameters& InParameters)
 	// Store copy of parameters
 	Parameters = InParameters;
 
-	UE_LOG(LogIris, Log, TEXT("ReplicationWriter: Configured with MaxInternalNetRefIndex=%d and MaxReplicationWriterObjectCount=%d."), 
+	UE_LOG(LogIris, Verbose, TEXT("ReplicationWriter: Configured with MaxInternalNetRefIndex=%d and MaxReplicationWriterObjectCount=%d."), 
 		Parameters.MaxInternalNetRefIndex, Parameters.MaxReplicationWriterObjectCount);
 
 	// Cache internal systems
diff --git a/Engine/Source/Runtime/Experimental/Iris/Core/Public/Iris/ReplicationSystem/ObjectReplicationBridge.h b/Engine/Source/Runtime/Experimental/Iris/Core/Public/Iris/ReplicationSystem/ObjectReplicationBridge.h
--- a/Engine/Source/Runtime/Experimental/Iris/Core/Public/Iris/ReplicationSystem/ObjectReplicationBridge.h
+++ b/Engine/Source/Runtime/Experimental/Iris/Core/Public/Iris/ReplicationSystem/ObjectReplicationBridge.h
@@ -51,22 +51,23 @@ public:
 
 	struct FCreateNetRefHandleParams
 	{
-		//$IRIS TODO: These need documentation
-		bool bCanReceive = false;
+		/** When true it means the object wants to receive a PreUpdate callback just before it gets polled. */
 		bool bNeedsPreUpdate = false;
+
+		/** When true the object has a dynamic world location and we should ask the bridge to update its current location everytime it is polled. */
 		bool bNeedsWorldLocationUpdate = false;
 
 		/** Whether the object is dormant or not */
 		bool bIsDormant = false;
 
-		/** When true we ask the class config if a dynamic filter was assigned to this class or one of it's parent inherited class. */
-		bool bUseClassConfigDynamicFilter = false;
+		/** Ask the class config for a dynamic filter assigned to this class or one of it's parent class. Default is true. */
+		bool bUseClassConfigDynamicFilter = true;
 
-		/** When enabled we ignore the class config for this object and instead use the one specified in ExplicitDynamicFilter */
+		/** When enabled we ignore the class config for this object and instead use the one specified by ExplicitDynamicFilter */
 		bool bUseExplicitDynamicFilter = false;
 
 		/** 
-		 * The name of the dynamic filter to use for this object (instead of asking the class config). 
+		 * The name of the dynamic filter to use for this object (instead of asking the class config).
 		 * Can be none so that no dynamic filter is assigned to the object.
 		 * Only used when bUseExplicitDynamicFilter is true.
 		 */
@@ -85,8 +86,6 @@ public:
 		float PollFrequency = 0.0f;
 	};
 
-	IRISCORE_API static FCreateNetRefHandleParams DefaultCreateNetRefHandleParams;
-
 	IRISCORE_API UObjectReplicationBridge();
 
 	/** Get the Object from a replicated handle, if the handle is invalid or not is a replicated handle the function will return nullptr */
@@ -111,18 +110,18 @@ public:
 	IRISCORE_API UE::Net::FNetObjectReference GetOrCreateObjectReference(const FString& Path, const UObject* Outer) const;
 
 	/** Begin replicating the Instance and return a valid NetRefHandle for the Instance if successful. */
-	IRISCORE_API FNetRefHandle BeginReplication(UObject* Instance, const FCreateNetRefHandleParams& Params = DefaultCreateNetRefHandleParams);
+	IRISCORE_API FNetRefHandle BeginReplication(UObject* Instance, const FCreateNetRefHandleParams& Params);
 
 	/**
 	 * Begin replicating the Instance as a subobject of the OwnerHandle. If InsertRelativeSubObjectHandle is valid
 	 * the new subobject will be inserted in the subobject replication list next to the specified handle and the wanted insertion order.
 	 * Default behavior is to always add new subobjects at the end of the list. Returns a valid NetRefHandle for the Instance if successful.
 	 */
-	IRISCORE_API FNetRefHandle BeginReplication(FNetRefHandle OwnerHandle, UObject* Instance, FNetRefHandle InsertRelativeToSubObjectHandle, ESubObjectInsertionOrder InsertionOrder = UReplicationBridge::ESubObjectInsertionOrder::None, const FCreateNetRefHandleParams& Params = DefaultCreateNetRefHandleParams);
+	IRISCORE_API FNetRefHandle BeginReplication(FNetRefHandle OwnerHandle, UObject* Instance, FNetRefHandle InsertRelativeToSubObjectHandle, const FCreateNetRefHandleParams& Params, ESubObjectInsertionOrder InsertionOrder = UReplicationBridge::ESubObjectInsertionOrder::None);
 
 	/** Create handle and start replicating the Instance as a SubObject of the OwnerHandle. */
 	/** Begin replicating the Instance as a subobject of the OwnerHandle and return a valid NetRefHandle for the Instance if successful. */
-	IRISCORE_API FNetRefHandle BeginReplication(FNetRefHandle OwnerHandle, UObject* Instance, const FCreateNetRefHandleParams& Params = DefaultCreateNetRefHandleParams);
+	IRISCORE_API FNetRefHandle BeginReplication(FNetRefHandle OwnerHandle, UObject* Instance, const FCreateNetRefHandleParams& Params);
 
 	/** 
 	 * Set NetCondition for a subobject, the condition is used to determine if the SubObject should replicate or not.
@@ -494,5 +493,5 @@ protected:
 
 inline UE::Net::FNetRefHandle UObjectReplicationBridge::BeginReplication(UE::Net::FNetRefHandle OwnerHandle, UObject* Instance, const FCreateNetRefHandleParams& Params)
 {
-	return BeginReplication(OwnerHandle, Instance, FNetRefHandle::GetInvalid(), ESubObjectInsertionOrder::None, Params);
+	return BeginReplication(OwnerHandle, Instance, FNetRefHandle::GetInvalid(), Params, ESubObjectInsertionOrder::None);
 }
diff --git a/Engine/Source/Runtime/Experimental/Iris/Core/Public/Iris/ReplicationSystem/ReplicationBridge.h b/Engine/Source/Runtime/Experimental/Iris/Core/Public/Iris/ReplicationSystem/ReplicationBridge.h
--- a/Engine/Source/Runtime/Experimental/Iris/Core/Public/Iris/ReplicationSystem/ReplicationBridge.h
+++ b/Engine/Source/Runtime/Experimental/Iris/Core/Public/Iris/ReplicationSystem/ReplicationBridge.h
@@ -160,9 +160,6 @@ public:
 	/** Returns true if the handle is replicated. */
 	IRISCORE_API bool IsReplicatedHandle(FNetRefHandle Handle) const;
 
-	/** Set the NetDriver used by the bridge. Called during creation if the NetDriver is recreated. */
-	IRISCORE_API virtual void SetNetDriver(UNetDriver* NetDriver);
-
 	/** Get the group associated with the level in order to control connection filtering for it. */
 	IRISCORE_API UE::Net::FNetObjectGroupHandle GetLevelGroup(const UObject* Level) const;
 
@@ -315,19 +312,13 @@ private:
 	void CallDetachInstanceFromRemote(FNetRefHandle Handle, EReplicationBridgeDestroyInstanceReason DestroyReason, EReplicationBridgeDestroyInstanceFlags DestroyFlags);
 
 private:
-	// Specifies if we should capture and cache CreationHeader or not for flush and tearoff.
-	enum class ECacheCreationInfo : uint8
-	{
-		No,
-		Yes
-	};
 
-	void InternalFlushStateData(UE::Net::FNetSerializationContext& SerializationContext, UE::Net::Private::FChangeMaskCache& ChangeMaskCache, UE::Net::FNetBitStreamWriter& ChangeMaskWriter, uint32 InternalObjectIndex, ECacheCreationInfo CacheCreationInfo);
+	void InternalFlushStateData(UE::Net::FNetSerializationContext& SerializationContext, UE::Net::Private::FChangeMaskCache& ChangeMaskCache, UE::Net::FNetBitStreamWriter& ChangeMaskWriter, uint32 InternalObjectIndex);
 	// Internal method to copy state data for Handle
-	void InternalFlushStateData(FNetRefHandle Handle, ECacheCreationInfo CacheCreationInfo);
+	void InternalFlushStateData(FNetRefHandle Handle);
 
 	// Internal method to copy state data for Handle and any SubObjects and mark them as being torn-off
-	void InternalTearOff(FNetRefHandle OwnerHandle, ECacheCreationInfo CacheCreationInfo);
+	void InternalTearOff(FNetRefHandle OwnerHandle);
 
 	// Destroy all SubObjects owned by provided handle
 	void InternalDestroySubObjects(FNetRefHandle OwnerHandle, EEndReplicationFlags Flags);
diff --git a/Engine/Source/Runtime/Net/Core/Public/Net/Core/NetHandle/NetHandle.h b/Engine/Source/Runtime/Net/Core/Public/Net/Core/NetHandle/NetHandle.h
--- a/Engine/Source/Runtime/Net/Core/Public/Net/Core/NetHandle/NetHandle.h
+++ b/Engine/Source/Runtime/Net/Core/Public/Net/Core/NetHandle/NetHandle.h
@@ -47,10 +47,8 @@ public:
 	bool operator<(FNetHandle Other) const;
 	bool operator!=(FNetHandle Other) const;
 
-	FString ToString() const;
+	NETCORE_API FString ToString() const;
 
-	static bool FullCompare(FNetHandle A, FNetHandle B);
-	
 private:
 	friend FNetHandleManager;
 	friend Private::FNetPushObjectHandle;
