Fix HTTP client being destroyed with pending requests when cancelling IOStoreOnDemand requests.

diff --git a/Engine/Source/Runtime/Experimental/IoStore/OnDemand/Private/OnDemandIoStore.cpp b/Engine/Source/Runtime/Experimental/IoStore/OnDemand/Private/OnDemandIoStore.cpp
--- a/Engine/Source/Runtime/Experimental/IoStore/OnDemand/Private/OnDemandIoStore.cpp
+++ b/Engine/Source/Runtime/Experimental/IoStore/OnDemand/Private/OnDemandIoStore.cpp
@@ -1424,6 +1424,10 @@ FOnDemandInstallResult FOnDemandIoStore::TickInstallRequest(const FInstallReques
 
 			if ((Status.IsOk() == false) || CheckAndSetCancelled())
 			{
+				// TODO: Should there be a way to flush/cancel the HTTPClient?
+				while (HttpClient->Tick())
+					;
+
 				return OutResult;
 			}
 		}
