Added missing include

diff --git a/Engine/Source/Runtime/Experimental/IoStore/OnDemand/Private/Tool/Socks.cpp b/Engine/Source/Runtime/Experimental/IoStore/OnDemand/Private/Tool/Socks.cpp
--- a/Engine/Source/Runtime/Experimental/IoStore/OnDemand/Private/Tool/Socks.cpp
+++ b/Engine/Source/Runtime/Experimental/IoStore/OnDemand/Private/Tool/Socks.cpp
@@ -2,6 +2,7 @@
 
 #if defined(UE_WITH_IAS_TOOL) && PLATFORM_WINDOWS
 
+#include "CoreGlobals.h"
 #include "Command.h"
 
 #include <Misc/ScopeExit.h>
