ISMPool: disallow nanite on managed ISMs with translucent materials.
Removes cases of log spam where we have translucent materials that disable nanite later in the pipeline when creating scene proxies.

diff --git a/Engine/Source/Runtime/Experimental/GeometryCollectionEngine/Private/GeometryCollection/GeometryCollectionISMPoolComponent.cpp b/Engine/Source/Runtime/Experimental/GeometryCollectionEngine/Private/GeometryCollection/GeometryCollectionISMPoolComponent.cpp
--- a/Engine/Source/Runtime/Experimental/GeometryCollectionEngine/Private/GeometryCollection/GeometryCollectionISMPoolComponent.cpp
+++ b/Engine/Source/Runtime/Experimental/GeometryCollectionEngine/Private/GeometryCollection/GeometryCollectionISMPoolComponent.cpp
@@ -147,6 +147,8 @@ void FGeometryCollectionISM::InitISM(const FGeometryCollectionStaticMeshInstance
 	ISMComponent->bUseAttachParentBound = bOverrideTransformUpdates;
 	ISMComponent->SetAbsolute(bOverrideTransformUpdates, bOverrideTransformUpdates, bOverrideTransformUpdates);
 
+	bool bDisallowNanite = false;
+
 	ISMComponent->EmptyOverrideMaterials();
 	for (int32 MaterialIndex = 0; MaterialIndex < MeshInstance.MaterialsOverrides.Num(); MaterialIndex++)
 	{
@@ -154,6 +156,9 @@ void FGeometryCollectionISM::InitISM(const FGeometryCollectionStaticMeshInstance
 		// We should only get here for valid material objects.
 		check(Material != nullptr);
 		ISMComponent->SetMaterial(MaterialIndex, Material);
+
+		// Nanite doesn't support translucent materials.
+		bDisallowNanite |= Material->GetBlendMode() == BLEND_Translucent;
 	}
 
 	ISMComponent->SetStaticMesh(StaticMesh);
@@ -212,6 +217,7 @@ void FGeometryCollectionISM::InitISM(const FGeometryCollectionStaticMeshInstance
 	ISMComponent->SetLODDistanceScale(MeshInstance.Desc.LodScale);
 	ISMComponent->SetUseConservativeBounds(true);
 	ISMComponent->bComputeFastLocalBounds = true;
+	ISMComponent->bDisallowNanite = bDisallowNanite;
 	ISMComponent->SetMeshDrawCommandStatsCategory(MeshInstance.Desc.StatsCategory);
 	ISMComponent->ComponentTags = MeshInstance.Desc.Tags;
 }
