Dataflow : Fix select clusters node that was selecting leaves instead

diff --git a/Engine/Source/Runtime/Experimental/Chaos/Private/GeometryCollection/Facades/CollectionTransformSelectionFacade.cpp b/Engine/Source/Runtime/Experimental/Chaos/Private/GeometryCollection/Facades/CollectionTransformSelectionFacade.cpp
--- a/Engine/Source/Runtime/Experimental/Chaos/Private/GeometryCollection/Facades/CollectionTransformSelectionFacade.cpp
+++ b/Engine/Source/Runtime/Experimental/Chaos/Private/GeometryCollection/Facades/CollectionTransformSelectionFacade.cpp
@@ -456,7 +456,7 @@ namespace GeometryCollection::Facades
 
 			OutSelection.SetNum(Algo::RemoveIf(OutSelection, [&](int32 BoneIdx)
 				{
-					return SimType[BoneIdx] != FGeometryCollection::ESimulationTypes::FST_Rigid
+					return SimType[BoneIdx] != FGeometryCollection::ESimulationTypes::FST_Clustered
 						|| (ViewLevel != INDEX_NONE && Levels[BoneIdx] != ViewLevel);
 				}));
 		}
