Chaos Cloth: Aerodynamics updates

* Cloth Asset and Legacy Cloth: Added concept of Inner and Outer Lift and Drag. This allows setting different lift and drag coefficients for when air is hitting the frontface vs backface of the cloth.
* Cloth Asset only: Added ability to specify WindVelocity in different spaces rather than just world space.

diff --git a/Engine/Plugins/ChaosCloth/Source/ChaosCloth/Private/ChaosCloth/ChaosClothConstraints.cpp b/Engine/Plugins/ChaosCloth/Source/ChaosCloth/Private/ChaosCloth/ChaosClothConstraints.cpp
--- a/Engine/Plugins/ChaosCloth/Source/ChaosCloth/Private/ChaosCloth/ChaosClothConstraints.cpp
+++ b/Engine/Plugins/ChaosCloth/Source/ChaosCloth/Private/ChaosCloth/ChaosClothConstraints.cpp
@@ -2187,7 +2187,9 @@ void FClothConstraints::Update(
 	const TMap<FString, const TSet<int32>*>& FaceSets,
 	const TMap<FString, TConstArrayView<int32>>& FaceIntMaps,
 	Softs::FSolverReal MeshScale,
-	Softs::FSolverReal MaxDistancesScale)
+	Softs::FSolverReal MaxDistancesScale,
+	const FRotation3& LocalSpaceRotation,
+	const FRotation3& ReferenceSpaceRotation)
 {
 	if (EdgeConstraints)
 	{
@@ -2282,7 +2284,9 @@ void FClothConstraints::Update(
 			WeightMaps,
 			WorldScale,
 			bEnableAerodynamics,
-			SolverWindVelocity
+			SolverWindVelocity,
+			LocalSpaceRotation,
+			ReferenceSpaceRotation
 		);
 	}
 	if (CollisionConstraint)
diff --git a/Engine/Plugins/ChaosCloth/Source/ChaosCloth/Private/ChaosCloth/ChaosClothingSimulationConfig.cpp b/Engine/Plugins/ChaosCloth/Source/ChaosCloth/Private/ChaosCloth/ChaosClothingSimulationConfig.cpp
--- a/Engine/Plugins/ChaosCloth/Source/ChaosCloth/Private/ChaosCloth/ChaosClothingSimulationConfig.cpp
+++ b/Engine/Plugins/ChaosCloth/Source/ChaosCloth/Private/ChaosCloth/ChaosClothingSimulationConfig.cpp
@@ -163,10 +163,24 @@ namespace Chaos
 				Property->SetWeightedValue(DragIndex, ClothConfig->Drag.Low, ClothConfig->Drag.High);
 				Property->SetStringValue(DragIndex, TEXT("Drag"));
 
+				if (ClothConfig->bEnableOuterDrag)
+				{
+					const int32 OuterDragIndex = Property->AddProperty(TEXT("OuterDrag"), AnimatablePropertyFlags);
+					Property->SetWeightedValue(OuterDragIndex, ClothConfig->OuterDrag.Low, ClothConfig->OuterDrag.High);
+					Property->SetStringValue(OuterDragIndex, TEXT("OuterDrag"));
+				}
+
 				const int32 LiftIndex = Property->AddProperty(TEXT("Lift"), AnimatablePropertyFlags);
 				Property->SetWeightedValue(LiftIndex, ClothConfig->Lift.Low, ClothConfig->Lift.High);
 				Property->SetStringValue(LiftIndex, TEXT("Lift"));
 
+				if (ClothConfig->bEnableOuterLift)
+				{
+					const int32 OuterLiftIndex = Property->AddProperty(TEXT("OuterLift"), AnimatablePropertyFlags);
+					Property->SetWeightedValue(OuterLiftIndex, ClothConfig->OuterLift.Low, ClothConfig->OuterLift.High);
+					Property->SetStringValue(OuterLiftIndex, TEXT("OuterLift"));
+				}
+
 				constexpr float AirDensity = 1.225f;  // Air density in kg/m^3
 				Property->AddValue(TEXT("FluidDensity"), AirDensity, AnimatablePropertyFlags);
 
diff --git a/Engine/Plugins/ChaosCloth/Source/ChaosCloth/Private/ChaosCloth/ChaosClothingSimulationInteractor.cpp b/Engine/Plugins/ChaosCloth/Source/ChaosCloth/Private/ChaosCloth/ChaosClothingSimulationInteractor.cpp
--- a/Engine/Plugins/ChaosCloth/Source/ChaosCloth/Private/ChaosCloth/ChaosClothingSimulationInteractor.cpp
+++ b/Engine/Plugins/ChaosCloth/Source/ChaosCloth/Private/ChaosCloth/ChaosClothingSimulationInteractor.cpp
@@ -129,13 +129,15 @@ void UChaosClothingInteractor::SetAerodynamics(float DragCoefficient, float Lift
 	}));
 }
 
-void UChaosClothingInteractor::SetWind(FVector2D Drag, FVector2D Lift, float AirDensity, FVector WindVelocity)
+void UChaosClothingInteractor::SetWind(FVector2D Drag, FVector2D Lift, float AirDensity, FVector WindVelocity, FVector2D OuterDrag, FVector2D OuterLift)
 {
-	ConfigCommands.Add(FChaosClothingInteractorConfigCommand::CreateLambda([Drag, Lift, AirDensity, WindVelocity](Chaos::FClothingSimulationConfig* Config, int32 LODIndex)
+	ConfigCommands.Add(FChaosClothingInteractorConfigCommand::CreateLambda([Drag, Lift, AirDensity, WindVelocity, OuterDrag, OuterLift](Chaos::FClothingSimulationConfig* Config, int32 LODIndex)
 	{
 		constexpr float WorldScale = 100.f;  // Unreal's world unit is the cm
 		Config->GetProperties(LODIndex).SetWeightedFloatValue(TEXT("Drag"), FVector2f(Drag));
+		Config->GetProperties(LODIndex).SetWeightedFloatValue(TEXT("OuterDrag"), FVector2f(OuterDrag));
 		Config->GetProperties(LODIndex).SetWeightedFloatValue(TEXT("Lift"), FVector2f(Lift));
+		Config->GetProperties(LODIndex).SetWeightedFloatValue(TEXT("OuterLift"), FVector2f(OuterLift));
 		Config->GetProperties(LODIndex).SetValue(TEXT("FluidDensity"), (float)AirDensity * FMath::Cube(WorldScale));  // AirDensity is here in kg/cm^3 for legacy reason but must be in kg/m^3 in the config UI
 		Config->GetProperties(LODIndex).SetValue(TEXT("WindVelocity"), FVector3f(WindVelocity) / WorldScale);
 	}));
diff --git a/Engine/Plugins/ChaosCloth/Source/ChaosCloth/Private/ChaosCloth/ChaosWeightMapTarget.h b/Engine/Plugins/ChaosCloth/Source/ChaosCloth/Private/ChaosCloth/ChaosWeightMapTarget.h
--- a/Engine/Plugins/ChaosCloth/Source/ChaosCloth/Private/ChaosCloth/ChaosWeightMapTarget.h
+++ b/Engine/Plugins/ChaosCloth/Source/ChaosCloth/Private/ChaosCloth/ChaosWeightMapTarget.h
@@ -25,5 +25,7 @@ enum class EChaosWeightMapTarget : uint8
 	AreaStiffness,
 	BucklingStiffness,
 	Pressure,
-	FlatnessRatio
+	FlatnessRatio,
+	OuterDrag,
+	OuterLift
 };
diff --git a/Engine/Plugins/ChaosCloth/Source/ChaosCloth/Public/ChaosCloth/ChaosClothConfig.h b/Engine/Plugins/ChaosCloth/Source/ChaosCloth/Public/ChaosCloth/ChaosClothConfig.h
--- a/Engine/Plugins/ChaosCloth/Source/ChaosCloth/Public/ChaosCloth/ChaosClothConfig.h
+++ b/Engine/Plugins/ChaosCloth/Source/ChaosCloth/Public/ChaosCloth/ChaosClothConfig.h
@@ -304,7 +304,9 @@ public:
 	bool bUsePointBasedWindModel = false;
 
 	/**
-	 * The aerodynamic coefficient of drag applying on each particle.
+	 * The aerodynamic coefficient of drag applying on each particle. 
+	 * When "Outer Drag" is enabled, this acts as the "Inner Drag", i.e., drag applied when the air velocity is
+	 * moving in the mesh normal direction.
 	 * If an enabled Weight Map (Mask with values in the range [0;1]) targeting the "Drag" is added to the cloth, 
 	 * then both the Low and High values will be used in conjunction with the per particle Weight stored
 	 * in the Weight Map to interpolate the final value from them.
@@ -313,8 +315,24 @@ public:
 	UPROPERTY(EditAnywhere, Category = "Environmental Properties", meta = (UIMin = "0", UIMax = "1", ClampMin = "0", ClampMax = "10", EditCondition = "!bUsePointBasedWindModel"))
 	FChaosClothWeightedValue Drag = { 0.035f, 1.f };
 
+	UPROPERTY(EditAnywhere, Category = "Environmental Properties", meta = (InlineEditConditionToggle))
+	bool bEnableOuterDrag = false;
+
+	/**
+	 * The aerodynamic coefficient of drag applying on each particle when the air velocity is moving
+	 * against the mesh normal direction.
+	 * If an enabled Weight Map (Mask with values in the range [0;1]) targeting the "Outer Drag" is added to the cloth,
+	 * then both the Low and High values will be used in conjunction with the per particle Weight stored
+	 * in the Weight Map to interpolate the final value from them.
+	 * Otherwise only the Low value is meaningful and sufficient to set the aerodynamic drag.
+	 */
+	UPROPERTY(EditAnywhere, Category = "Environmental Properties", meta = (UIMin = "0", UIMax = "1", ClampMin = "0", ClampMax = "10", EditCondition = "bEnableOuterDrag"))
+	FChaosClothWeightedValue OuterDrag = { 0.035f, 1.f };
+
 	/**
 	 * The aerodynamic coefficient of lift applying on each particle.
+	 * When "Outer Lift" is enabled, this acts as the "Inner lift", i.e., lift applied when the air velocity is
+	 * moving in the mesh normal direction.
 	 * If an enabled Weight Map (Mask with values in the range [0;1]) targeting the "Lift" is added to the cloth, 
 	 * then both the Low and High values will be used in conjunction with the per particle Weight stored
 	 * in the Weight Map to interpolate the final value from them.
@@ -323,6 +341,20 @@ public:
 	UPROPERTY(EditAnywhere, Category = "Environmental Properties", meta = (UIMin = "0", UIMax = "1", ClampMin = "0", ClampMax = "10", EditCondition = "!bUsePointBasedWindModel"))
 	FChaosClothWeightedValue Lift = { 0.035f, 1.f };
 
+	UPROPERTY(EditAnywhere, Category = "Environmental Properties", meta = (InlineEditConditionToggle))
+	bool bEnableOuterLift = false;
+
+	/**
+	 * The aerodynamic coefficient of lift applying on each particle when the air velocity is moving
+	 * against the mesh normal direction.
+	 * If an enabled Weight Map (Mask with values in the range [0;1]) targeting the "Outer Lift" is added to the cloth,
+	 * then both the Low and High values will be used in conjunction with the per particle Weight stored
+	 * in the Weight Map to interpolate the final value from them.
+	 * Otherwise only the Low value is meaningful and sufficient to set the aerodynamic lift.
+	 */
+	UPROPERTY(EditAnywhere, Category = "Environmental Properties", meta = (UIMin = "0", UIMax = "1", ClampMin = "0", ClampMax = "10", EditCondition = "bEnableOuterLift"))
+	FChaosClothWeightedValue OuterLift = { 0.035f, 1.f };
+
 	// Use the config gravity value instead of world gravity.
 	UPROPERTY(EditAnywhere, Category = "Environmental Properties", meta = (InlineEditConditionToggle))
 	bool bUseGravityOverride = false;
diff --git a/Engine/Plugins/ChaosCloth/Source/ChaosCloth/Public/ChaosCloth/ChaosClothConstraints.h b/Engine/Plugins/ChaosCloth/Source/ChaosCloth/Public/ChaosCloth/ChaosClothConstraints.h
--- a/Engine/Plugins/ChaosCloth/Source/ChaosCloth/Public/ChaosCloth/ChaosClothConstraints.h
+++ b/Engine/Plugins/ChaosCloth/Source/ChaosCloth/Public/ChaosCloth/ChaosClothConstraints.h
@@ -88,7 +88,9 @@ namespace Chaos
 			const TMap<FString, const TSet<int32>*>& FaceSets,
 			const TMap<FString, TConstArrayView<int32>>& FaceIntMaps,
 			Softs::FSolverReal MeshScale,
-			Softs::FSolverReal MaxDistancesScale = (Softs::FSolverReal)1.);
+			Softs::FSolverReal MaxDistancesScale = (Softs::FSolverReal)1.,
+			const FRotation3& LocalSpaceRotation = FRotation3::Identity,
+			const FRotation3& ReferenceSpaceRotation = FRotation3::Identity);
 
 		UE_DEPRECATED(5.4, "Use Update() with WeightMaps, VertexSets, FaceSets, and FaceIntMaps instead.")
 		void Update(
diff --git a/Engine/Plugins/ChaosCloth/Source/ChaosCloth/Public/ChaosCloth/ChaosClothingSimulationInteractor.h b/Engine/Plugins/ChaosCloth/Source/ChaosCloth/Public/ChaosCloth/ChaosClothingSimulationInteractor.h
--- a/Engine/Plugins/ChaosCloth/Source/ChaosCloth/Public/ChaosCloth/ChaosClothingSimulationInteractor.h
+++ b/Engine/Plugins/ChaosCloth/Source/ChaosCloth/Public/ChaosCloth/ChaosClothingSimulationInteractor.h
@@ -55,7 +55,7 @@ public:
 	CHAOSCLOTH_API void SetAerodynamics(float DragCoefficient = 0.07f, float LiftCoefficient = 0.035f, FVector WindVelocity = FVector(0.f, 0.f, 0.f));
 
 	UFUNCTION(BlueprintCallable, Category = ClothingSimulation, Meta = (Keywords = "Chaos Wind Drag Lift Coefficient Air Density Velocity"))
-	CHAOSCLOTH_API void SetWind(FVector2D Drag = FVector2D(0.07f, 0.5f), FVector2D Lift = FVector2D(0.07f, 0.5f), float AirDensity = 1.225e-6f, FVector WindVelocity = FVector(0.f, 0.f, 0.f));
+	CHAOSCLOTH_API void SetWind(FVector2D Drag = FVector2D(0.07f, 0.5f), FVector2D Lift = FVector2D(0.07f, 0.5f), float AirDensity = 1.225e-6f, FVector WindVelocity = FVector(0.f, 0.f, 0.f), FVector2D OuterDrag = FVector2D(0.07f, 0.5f), FVector2D OuterLift = FVector2D(0.07, 0.5f));
 
 	UFUNCTION(BlueprintCallable, Category = ClothingSimulation, Meta = (Keywords = "Chaos Pressure"))
 	CHAOSCLOTH_API void SetPressure(FVector2D Pressure = FVector2D(0.f, 1.f));
diff --git a/Engine/Plugins/ChaosClothAssetEditor/Source/ChaosClothAssetDataflowNodes/Private/ChaosClothAsset/SimulationAerodynamicsConfigNode.cpp b/Engine/Plugins/ChaosClothAssetEditor/Source/ChaosClothAssetDataflowNodes/Private/ChaosClothAsset/SimulationAerodynamicsConfigNode.cpp
--- a/Engine/Plugins/ChaosClothAssetEditor/Source/ChaosClothAssetDataflowNodes/Private/ChaosClothAsset/SimulationAerodynamicsConfigNode.cpp
+++ b/Engine/Plugins/ChaosClothAssetEditor/Source/ChaosClothAssetDataflowNodes/Private/ChaosClothAsset/SimulationAerodynamicsConfigNode.cpp
@@ -13,14 +13,21 @@ FChaosClothAssetSimulationAerodynamicsConfigNode::FChaosClothAssetSimulationAero
 	RegisterInputConnection(&Drag.WeightMap)
 		.SetCanHidePin(true)
 		.SetPinIsHidden(true);
+	RegisterInputConnection(&OuterDrag.WeightMap)
+		.SetCanHidePin(true)
+		.SetPinIsHidden(true);
 	RegisterInputConnection(&Lift.WeightMap)
 		.SetCanHidePin(true)
 		.SetPinIsHidden(true);
+	RegisterInputConnection(&OuterLift.WeightMap)
+		.SetCanHidePin(true)
+		.SetPinIsHidden(true);
 }
 
 void FChaosClothAssetSimulationAerodynamicsConfigNode::AddProperties(FPropertyHelper& PropertyHelper) const
 {
 	PropertyHelper.SetProperty(this, &FluidDensity);
+	PropertyHelper.SetPropertyEnum(this, &WindVelocitySpace, {}, ECollectionPropertyFlags::None);
 	PropertyHelper.SetProperty(this, &WindVelocity);
 	
 	PropertyHelper.SetSolverPropertyWeighted(FName(TEXT("Drag")), Drag, [](
@@ -29,9 +36,19 @@ void FChaosClothAssetSimulationAerodynamicsConfigNode::AddProperties(FPropertyHe
 		return ClothFacade.GetSolverAirDamping();
 	},{});
 
+	if (bEnableOuterDrag)
+	{
+		PropertyHelper.SetPropertyWeighted(this, &OuterDrag);
+	}
+
 	PropertyHelper.SetSolverPropertyWeighted(FName(TEXT("Lift")), Lift, [](
 				const UE::Chaos::ClothAsset::FCollectionClothFacade& ClothFacade)-> float
 	{
 		return ClothFacade.GetSolverAirDamping();
 	},{});
+
+	if (bEnableOuterLift)
+	{
+		PropertyHelper.SetPropertyWeighted(this, &OuterLift);
+	}
 }
diff --git a/Engine/Plugins/ChaosClothAssetEditor/Source/ChaosClothAssetDataflowNodes/Public/ChaosClothAsset/SimulationAerodynamicsConfigNode.h b/Engine/Plugins/ChaosClothAssetEditor/Source/ChaosClothAssetDataflowNodes/Public/ChaosClothAsset/SimulationAerodynamicsConfigNode.h
--- a/Engine/Plugins/ChaosClothAssetEditor/Source/ChaosClothAssetDataflowNodes/Public/ChaosClothAsset/SimulationAerodynamicsConfigNode.h
+++ b/Engine/Plugins/ChaosClothAssetEditor/Source/ChaosClothAssetDataflowNodes/Public/ChaosClothAsset/SimulationAerodynamicsConfigNode.h
@@ -4,6 +4,7 @@
 
 #include "ChaosClothAsset/SimulationBaseConfigNode.h"
 #include "ChaosClothAsset/WeightedValue.h"
+#include "Chaos/SoftsSimulationSpace.h"
 #include "SimulationAerodynamicsConfigNode.generated.h"
 
 /** Aerodynamics properties configuration node. */
@@ -22,6 +23,12 @@ public:
 	 */
 	UPROPERTY(EditAnywhere, Category = "Aerodynamics", Meta = (UIMin = "0", ClampMin = "0", ClampMax = "10"))
 	float FluidDensity = 1.225f;
+
+	/**
+	 * Wind velocity is specified in this space.
+	 */
+	UPROPERTY(EditAnywhere, Category = "Aerodynamics")
+	EChaosSoftsSimulationSpace WindVelocitySpace = EChaosSoftsSimulationSpace::WorldSpace;
 	
 	/**
 	 * The fixed wind velocity [m/s] for this asset.
@@ -32,6 +39,8 @@ public:
 
 	/**
 	 * The aerodynamic coefficient of drag applying on each particle.
+	 * When "Outer Drag" is enabled, this acts as the "Inner Drag", i.e., drag applied when the air velocity is
+	 * moving in the mesh normal direction.
 	 * If a valid weight map is found with the given Weight Map name, then both Low and High values
 	 * are interpolated with the per particle weight to make the final value used for the simulation.
 	 * Otherwise all particles are considered to have a zero weight, and only the Low value is meaningful.
@@ -39,14 +48,43 @@ public:
 	UPROPERTY(EditAnywhere, Category = "Aerodynamics", Meta = (UIMin = "0", UIMax = "1", ClampMin = "0", ClampMax = "10"))
 	FChaosClothAssetWeightedValue Drag = { true, 0.035f, 1.f, TEXT("Drag"), true };
 
+	UPROPERTY(EditAnywhere, Category = "Aerodynamics", Meta = (InlineEditConditionToggle))
+	bool bEnableOuterDrag = false;
+
+	/**
+	 * The aerodynamic coefficient of drag applying on each particle when the air velocity is moving
+	 * against the mesh normal direction.
+	 * If a valid weight map is found with the given Weight Map name, then both Low and High values
+	 * are interpolated with the per particle weight to make the final value used for the simulation.
+	 * Otherwise all particles are considered to have a zero weight, and only the Low value is meaningful.
+	 */
+	UPROPERTY(EditAnywhere, Category = "Aerodynamics", Meta = (UIMin = "0", UIMax = "1", ClampMin = "0", ClampMax = "10", EditCondition = "bEnableOuterDrag"))
+	FChaosClothAssetWeightedValue OuterDrag = { true, 0.035f, 1.f, TEXT("OuterDrag"), true };
+
+
 	/**
 	 * The aerodynamic coefficient of lift applying on each particle.
+	 * When "Outer Lift" is enabled, this acts as the "Inner Lift", i.e., lift applied when the air velocity is
+	 * moving in the mesh normal direction.
 	 * If a valid weight map is found with the given Weight Map name, then both Low and High values
 	 * are interpolated with the per particle weight to make the final value used for the simulation.
 	 * Otherwise all particles are considered to have a zero weight, and only the Low value is meaningful.
 	 */
 	UPROPERTY(EditAnywhere, Category = "Aerodynamics", Meta = (UIMin = "0", UIMax = "1", ClampMin = "0", ClampMax = "10"))
 	FChaosClothAssetWeightedValue Lift = { true, 0.035f, 1.f, TEXT("Lift"), true };
+
+	UPROPERTY(EditAnywhere, Category = "Aerodynamics", Meta = (InlineEditConditionToggle))
+	bool bEnableOuterLift = false;
+
+	/**
+	 * The aerodynamic coefficient of lift applying on each particle when the air velocity is moving
+	 * against the mesh normal direction.
+	 * If a valid weight map is found with the given Weight Map name, then both Low and High values
+	 * are interpolated with the per particle weight to make the final value used for the simulation.
+	 * Otherwise all particles are considered to have a zero weight, and only the Low value is meaningful.
+	 */
+	UPROPERTY(EditAnywhere, Category = "Aerodynamics", Meta = (UIMin = "0", UIMax = "1", ClampMin = "0", ClampMax = "10", EditCondition = "bEnableOuterLift"))
+	FChaosClothAssetWeightedValue OuterLift = { true, 0.035f, 1.f, TEXT("OuterLift"), true };
 	
 	FChaosClothAssetSimulationAerodynamicsConfigNode(const Dataflow::FNodeParameters& InParam, FGuid InGuid = FGuid::NewGuid());
 
diff --git a/Engine/Source/Runtime/Experimental/Chaos/Private/Chaos/VelocityField.cpp b/Engine/Source/Runtime/Experimental/Chaos/Private/Chaos/VelocityField.cpp
--- a/Engine/Source/Runtime/Experimental/Chaos/Private/Chaos/VelocityField.cpp
+++ b/Engine/Source/Runtime/Experimental/Chaos/Private/Chaos/VelocityField.cpp
@@ -51,6 +51,25 @@ void FVelocityAndPressureField::SetProperties(
 		}
 	}
 
+	if (IsOuterDragMutable(PropertyCollection))
+	{
+		const FSolverVec2 WeightedValue(GetWeightedFloatOuterDrag(PropertyCollection));
+		if (IsOuterDragStringDirty(PropertyCollection))
+		{
+			const FString& WeightMapName = GetOuterDragString(PropertyCollection);
+			OuterDrag = FPBDFlatWeightMap(
+				WeightedValue.ClampAxes(MinCoefficient, MaxCoefficient),
+				Weightmaps.FindRef(WeightMapName),
+				TConstArrayView<TVec3<int32>>(Elements),
+				Offset,
+				NumParticles);
+		}
+		else
+		{
+			OuterDrag.SetWeightedValue(WeightedValue.ClampAxes(MinCoefficient, MaxCoefficient));
+		}
+	}
+
 	if (IsLiftMutable(PropertyCollection))
 	{
 		const FSolverVec2 WeightedValue(GetWeightedFloatLift(PropertyCollection));
@@ -71,6 +90,26 @@ void FVelocityAndPressureField::SetProperties(
 		}
 	}
 
+	if (IsOuterLiftMutable(PropertyCollection))
+	{
+		const FSolverVec2 WeightedValue(GetWeightedFloatOuterLift(PropertyCollection));
+
+		if (IsOuterLiftStringDirty(PropertyCollection))
+		{
+			const FString& WeightMapName = GetOuterLiftString(PropertyCollection);
+			OuterLift = FPBDFlatWeightMap(
+				WeightedValue.ClampAxes(MinCoefficient, MaxCoefficient),
+				Weightmaps.FindRef(WeightMapName),
+				TConstArrayView<TVec3<int32>>(Elements),
+				Offset,
+				NumParticles);
+		}
+		else
+		{
+			OuterLift.SetWeightedValue(WeightedValue.ClampAxes(MinCoefficient, MaxCoefficient));
+		}
+	}
+
 	if (IsFluidDensityMutable(PropertyCollection))
 	{
 		Rho = (FSolverReal)FMath::Max(GetFluidDensity(PropertyCollection), 0.f) / FMath::Cube(WorldScale);
@@ -105,22 +144,42 @@ void FVelocityAndPressureField::SetPropertiesAndWind(
 	const TMap<FString, TConstArrayView<FRealSingle>>& WeightMaps,
 	FSolverReal WorldScale,
 	bool bEnableAerodynamics,
-	const FSolverVec3& SolverWind)
+	const FSolverVec3& SolverWind,
+	const FRotation3& LocalSpaceRotation,
+	const FRotation3& ReferenceSpaceRotation)
 {
 	SetProperties(PropertyCollection, WeightMaps, WorldScale, bEnableAerodynamics);
-	const FSolverVec3 WindVelocity = WindVelocityIndex != INDEX_NONE ? WorldScale * FSolverVec3(GetWindVelocity(PropertyCollection)) : FSolverVec3(0.f);
+
+	FSolverVec3 WindVelocity = WindVelocityIndex != INDEX_NONE ? WorldScale * FSolverVec3(GetWindVelocity(PropertyCollection)) : FSolverVec3(0.f);
+	const EChaosSoftsSimulationSpace WindVelocitySpace = WindVelocitySpaceIndex != INDEX_NONE ? (EChaosSoftsSimulationSpace)GetWindVelocitySpace(PropertyCollection) : DefaultWindVelocitySpace;
+	switch (WindVelocitySpace)
+	{
+	case EChaosSoftsSimulationSpace::WorldSpace:
+	default:
+		break;
+	case EChaosSoftsSimulationSpace::ComponentSpace:
+		WindVelocity = ((Softs::FSolverRotation3)LocalSpaceRotation).RotateVector(WindVelocity);
+		break;
+	case EChaosSoftsSimulationSpace::ReferenceBoneSpace:
+		WindVelocity = ((Softs::FSolverRotation3)ReferenceSpaceRotation).RotateVector(WindVelocity);
+		break;
+	}
 	SetVelocity(WindVelocity + SolverWind);
 }
 
 void FVelocityAndPressureField::SetProperties(
 	const FSolverVec2& InDrag,
+	const FSolverVec2& InOuterDrag,
 	const FSolverVec2& InLift,
+	const FSolverVec2& InOuterLift,
 	const FSolverReal FluidDensity,
 	const FSolverVec2& InPressure,
 	FSolverReal WorldScale)
 {
 	Drag.SetWeightedValue(InDrag.ClampAxes(MinCoefficient, MaxCoefficient));
+	OuterDrag.SetWeightedValue(InOuterDrag.ClampAxes(MinCoefficient, MaxCoefficient));
 	Lift.SetWeightedValue(InLift.ClampAxes(MinCoefficient, MaxCoefficient));
+	OuterLift.SetWeightedValue(InOuterLift.ClampAxes(MinCoefficient, MaxCoefficient));
 	Pressure.SetWeightedValue(InPressure / WorldScale);
 	Rho = FMath::Max(FluidDensity / FMath::Cube(WorldScale), (FSolverReal)0.);
 
@@ -136,17 +195,27 @@ void FVelocityAndPressureField::SetGeometry(
 {
 	// Reinit indices
 	DragIndex = FDragIndex(PropertyCollection);
+	OuterDragIndex = FOuterDragIndex(PropertyCollection);
 	LiftIndex = FLiftIndex(PropertyCollection);
+	OuterLiftIndex = FOuterLiftIndex(PropertyCollection);
 	FluidDensityIndex = FFluidDensityIndex(PropertyCollection);
 	PressureIndex = FPressureIndex(PropertyCollection);
 
 	// Reset geometry, properties, and weight maps
 	SetGeometry(TriangleMesh);
+	// If these properties don't exist, set their values to 0, not to DefaultCoefficients!
+	const FSolverVec2 DragValue = FSolverVec2(GetWeightedFloatDrag(PropertyCollection, 0.f));
+	const FSolverVec2 LiftValue = FSolverVec2(GetWeightedFloatLift(PropertyCollection, 0.f));
+	// OuterDrag and OuterLift default to Drag and Lift respectively.
+	const FSolverVec2 OuterDragValue = FSolverVec2(GetWeightedFloatOuterDrag(PropertyCollection, DragValue));
+	const FSolverVec2 OuterLiftValue = FSolverVec2(GetWeightedFloatOuterLift(PropertyCollection, LiftValue));
 	SetProperties(
-		FSolverVec2(GetWeightedFloatDrag(PropertyCollection, 0.f)),  // If these properties don't exist, set their values to 0, not to DefaultCoefficients!
-		FSolverVec2(GetWeightedFloatLift(PropertyCollection, 0.f)),
+		DragValue,
+		OuterDragValue,
+		LiftValue,
+		OuterLiftValue,
 		(FSolverReal)GetFluidDensity(PropertyCollection, 0.f),
-		FSolverVec2(GetWeightedFloatPressure(PropertyCollection, 0.f)),  // These getters also initialize the property indices, so keep before SetMultipliers
+		FSolverVec2(GetWeightedFloatPressure(PropertyCollection, 0.f)),
 		WorldScale);
 	SetMultipliers(PropertyCollection, Weightmaps);
 	ResetColor();
@@ -155,11 +224,13 @@ void FVelocityAndPressureField::SetGeometry(
 void FVelocityAndPressureField::SetGeometry(
 	const FTriangleMesh* TriangleMesh,
 	const TConstArrayView<FRealSingle>& DragMultipliers,
+	const TConstArrayView<FRealSingle>& OuterDragMultipliers,
 	const TConstArrayView<FRealSingle>& LiftMultipliers,
+	const TConstArrayView<FRealSingle>& OuterLiftMultipliers,
 	const TConstArrayView<FRealSingle>& PressureMultipliers)
 {
 	SetGeometry(TriangleMesh);
-	SetMultipliers(DragMultipliers, LiftMultipliers, PressureMultipliers);
+	SetMultipliers(DragMultipliers, OuterDragMultipliers, LiftMultipliers, OuterLiftMultipliers, PressureMultipliers);
 	ResetColor();
 }
 
@@ -280,26 +351,35 @@ void FVelocityAndPressureField::SetMultipliers(
 {
 	const TConstArrayView<FRealSingle> DragMultipliers = (DragIndex != INDEX_NONE) ?
 		Weightmaps.FindRef(GetDragString(PropertyCollection)) : TConstArrayView<FRealSingle>();
-
+	const TConstArrayView<FRealSingle> OuterDragMultipliers = (OuterDragIndex != INDEX_NONE) ?
+		Weightmaps.FindRef(GetOuterDragString(PropertyCollection)) : TConstArrayView<FRealSingle>();
 	const TConstArrayView<FRealSingle> LiftMultipliers = (LiftIndex != INDEX_NONE) ?
 		Weightmaps.FindRef(GetLiftString(PropertyCollection)) : TConstArrayView<FRealSingle>();
-
+	const TConstArrayView<FRealSingle> OuterLiftMultipliers = (OuterLiftIndex != INDEX_NONE) ?
+		Weightmaps.FindRef(GetOuterLiftString(PropertyCollection)) : TConstArrayView<FRealSingle>();
 	const TConstArrayView<FRealSingle> PressureMultipliers = (PressureIndex != INDEX_NONE) ?
 		Weightmaps.FindRef(GetPressureString(PropertyCollection)) : TConstArrayView<FRealSingle>();
 
-	SetMultipliers(DragMultipliers, LiftMultipliers, PressureMultipliers);
+	SetMultipliers(DragMultipliers, OuterDragMultipliers, LiftMultipliers, OuterLiftMultipliers, PressureMultipliers);
 }
 
 void FVelocityAndPressureField::SetMultipliers(
 	const TConstArrayView<FRealSingle>& DragMultipliers,
+	const TConstArrayView<FRealSingle>& OuterDragMultipliers,
 	const TConstArrayView<FRealSingle>& LiftMultipliers,
+	const TConstArrayView<FRealSingle>& OuterLiftMultipliers,
 	const TConstArrayView<FRealSingle>& PressureMultipliers)
 {	
 	const FSolverVec2 DragValues(Drag.GetLow(), Drag.GetHigh());
+	const FSolverVec2 OuterDragValues(OuterDrag.GetLow(), OuterDrag.GetHigh());
 	const FSolverVec2 LiftValues(Lift.GetLow(), Lift.GetHigh());
+	const FSolverVec2 OuterLiftValues(OuterLift.GetLow(), OuterLift.GetHigh());
 	const FSolverVec2 PressureValues(Pressure.GetLow(), Pressure.GetHigh());
 	Drag = FPBDFlatWeightMap(DragValues, DragMultipliers, TConstArrayView<TVec3<int32>>(Elements), Offset, NumParticles);
+	OuterDrag = FPBDFlatWeightMap(OuterDragValues, OuterDragMultipliers, TConstArrayView<TVec3<int32>>(Elements), Offset, NumParticles);
 	Lift = FPBDFlatWeightMap(LiftValues, LiftMultipliers, TConstArrayView<TVec3<int32>>(Elements), Offset, NumParticles);
+	OuterLift = FPBDFlatWeightMap(OuterLiftValues, OuterLiftMultipliers, TConstArrayView<TVec3<int32>>(Elements), Offset, NumParticles);
+
 	Pressure = FPBDFlatWeightMap(PressureValues, PressureMultipliers, TConstArrayView<TVec3<int32>>(Elements), Offset, NumParticles);
 }
 
@@ -309,9 +389,11 @@ void FVelocityAndPressureField::UpdateForces(const FSolverParticles& InParticles
 	const FSolverReal MaxVelocitySquared = (Private::VelocityFieldMaxVelocity > 0.f) ? FMath::Square((FSolverReal)Private::VelocityFieldMaxVelocity) : TNumericLimits<FSolverReal>::Max();
 
 	const bool bDragHasMap = Drag.HasWeightMap();
+	const bool bOuterDragHasMap = OuterDrag.HasWeightMap();
 	const bool bLiftHasMap = Lift.HasWeightMap();
+	const bool bOuterLiftHasMap = OuterLift.HasWeightMap();
 	const bool bPressureHasMap = Pressure.HasWeightMap();
-	if (!bDragHasMap && !bLiftHasMap && !bPressureHasMap)
+	if (!bDragHasMap && !bOuterDragHasMap && !bLiftHasMap && !bOuterLiftHasMap && !bPressureHasMap)
 	{
 #if INTEL_ISPC
 		if (bRealTypeCompatibleWithISPC && bChaos_VelocityField_ISPC_Enabled)
@@ -326,7 +408,9 @@ void FVelocityAndPressureField::UpdateForces(const FSolverParticles& InParticles
 					(const ispc::FVector3f&)Velocity,
 					QuarterRho,
 					(FSolverReal)Drag,
+					(FSolverReal)OuterDrag,
 					(FSolverReal)Lift,
+					(FSolverReal)OuterLift,
 					(FSolverReal)Pressure,
 					Elements.Num());
 			}
@@ -340,7 +424,9 @@ void FVelocityAndPressureField::UpdateForces(const FSolverParticles& InParticles
 					(const ispc::FVector3f&)Velocity,
 					QuarterRho,
 					(FSolverReal)Drag,
+					(FSolverReal)OuterDrag,
 					(FSolverReal)Lift,
+					(FSolverReal)OuterLift,
 					(FSolverReal)Pressure,
 					Elements.Num(),
 					MaxVelocitySquared);
@@ -355,7 +441,9 @@ void FVelocityAndPressureField::UpdateForces(const FSolverParticles& InParticles
 				{
 					UpdateField(InParticles, ElementIndex, Velocity,
 						(FSolverReal)Drag,
+						(FSolverReal)OuterDrag,
 						(FSolverReal)Lift,
+						(FSolverReal)OuterLift,
 						(FSolverReal)Pressure);
 				}
 			}
@@ -365,7 +453,9 @@ void FVelocityAndPressureField::UpdateForces(const FSolverParticles& InParticles
 				{
 					UpdateField(InParticles, ElementIndex, Velocity,
 						(FSolverReal)Drag,
+						(FSolverReal)OuterDrag,
 						(FSolverReal)Lift,
+						(FSolverReal)OuterLift,
 						(FSolverReal)Pressure, MaxVelocitySquared);
 				}
 			}
@@ -388,9 +478,15 @@ void FVelocityAndPressureField::UpdateForces(const FSolverParticles& InParticles
 					bDragHasMap,
 					reinterpret_cast<const ispc::FVector2f&>(Drag.GetOffsetRange()),
 					bDragHasMap ? Drag.GetMapValues().GetData() : nullptr,
+					bOuterDragHasMap,
+					reinterpret_cast<const ispc::FVector2f&>(OuterDrag.GetOffsetRange()),
+					bOuterDragHasMap ? OuterDrag.GetMapValues().GetData() : nullptr,
 					bLiftHasMap,
 					reinterpret_cast<const ispc::FVector2f&>(Lift.GetOffsetRange()),
 					bLiftHasMap ? Lift.GetMapValues().GetData() : nullptr,
+					bOuterLiftHasMap,
+					reinterpret_cast<const ispc::FVector2f&>(OuterLift.GetOffsetRange()),
+					bOuterLiftHasMap ? OuterLift.GetMapValues().GetData() : nullptr,
 					bPressureHasMap,
 					reinterpret_cast<const ispc::FVector2f&>(Pressure.GetOffsetRange()),
 					bPressureHasMap ? Pressure.GetMapValues().GetData() : nullptr,
@@ -408,9 +504,15 @@ void FVelocityAndPressureField::UpdateForces(const FSolverParticles& InParticles
 					bDragHasMap,
 					reinterpret_cast<const ispc::FVector2f&>(Drag.GetOffsetRange()),
 					bDragHasMap ? Drag.GetMapValues().GetData() : nullptr,
+					bOuterDragHasMap,
+					reinterpret_cast<const ispc::FVector2f&>(OuterDrag.GetOffsetRange()),
+					bOuterDragHasMap ? OuterDrag.GetMapValues().GetData() : nullptr,
 					bLiftHasMap,
 					reinterpret_cast<const ispc::FVector2f&>(Lift.GetOffsetRange()),
 					bLiftHasMap ? Lift.GetMapValues().GetData() : nullptr,
+					bOuterLiftHasMap,
+					reinterpret_cast<const ispc::FVector2f&>(OuterLift.GetOffsetRange()),
+					bOuterLiftHasMap ? OuterLift.GetMapValues().GetData() : nullptr,
 					bPressureHasMap,
 					reinterpret_cast<const ispc::FVector2f&>(Pressure.GetOffsetRange()),
 					bPressureHasMap ? Pressure.GetMapValues().GetData() : nullptr,
@@ -425,22 +527,26 @@ void FVelocityAndPressureField::UpdateForces(const FSolverParticles& InParticles
 			{
 				for (int32 ElementIndex = 0; ElementIndex < Elements.Num(); ++ElementIndex)
 				{
-					const FSolverReal Cd = Drag.GetValue(ElementIndex);
-					const FSolverReal Cl = Lift.GetValue(ElementIndex);
+					const FSolverReal CdI = Drag.GetValue(ElementIndex);
+					const FSolverReal CdO = OuterDrag.GetValue(ElementIndex);
+					const FSolverReal ClI = Lift.GetValue(ElementIndex);
+					const FSolverReal ClO = OuterLift.GetValue(ElementIndex);
 					const FSolverReal Cp = Pressure.GetValue(ElementIndex);
 
-					UpdateField(InParticles, ElementIndex, Velocity, Cd, Cl, Cp);
+					UpdateField(InParticles, ElementIndex, Velocity, CdI, CdO, ClI, ClO, Cp);
 				}
 			}
 			else
 			{
 				for (int32 ElementIndex = 0; ElementIndex < Elements.Num(); ++ElementIndex)
 				{
-					const FSolverReal Cd = Drag.GetValue(ElementIndex);
-					const FSolverReal Cl = Lift.GetValue(ElementIndex);
+					const FSolverReal CdI = Drag.GetValue(ElementIndex);
+					const FSolverReal CdO = OuterDrag.GetValue(ElementIndex);
+					const FSolverReal ClI = Lift.GetValue(ElementIndex);
+					const FSolverReal ClO = OuterLift.GetValue(ElementIndex);
 					const FSolverReal Cp = Pressure.GetValue(ElementIndex);
 
-					UpdateField(InParticles, ElementIndex, Velocity, Cd, Cl, Cp, MaxVelocitySquared);
+					UpdateField(InParticles, ElementIndex, Velocity, CdI, CdO, ClI, ClO, Cp, MaxVelocitySquared);
 				}
 			}
 		}
@@ -453,9 +559,11 @@ void FVelocityAndPressureField::Apply(FSolverParticlesRange& InParticles, const
 	const FSolverReal MaxVelocitySquared = (Private::VelocityFieldMaxVelocity > 0.f) ? FMath::Square((FSolverReal)Private::VelocityFieldMaxVelocity) : TNumericLimits<FSolverReal>::Max();
 
 	const bool bDragHasMap = Drag.HasWeightMap();
+	const bool bOuterDragHasMap = OuterDrag.HasWeightMap();
 	const bool bLiftHasMap = Lift.HasWeightMap();
+	const bool bOuterLiftHasMap = OuterLift.HasWeightMap();
 	const bool bPressureHasMap = Pressure.HasWeightMap();
-	if (!bDragHasMap && !bLiftHasMap && !bPressureHasMap)
+	if (!bDragHasMap && !bOuterDragHasMap && !bLiftHasMap && !bOuterLiftHasMap && !bPressureHasMap)
 	{
 #if INTEL_ISPC
 		if (bRealTypeCompatibleWithISPC && bChaos_VelocityField_ISPC_Enabled && ConstraintsPerColorStartIndex.Num() > 1)
@@ -476,7 +584,9 @@ void FVelocityAndPressureField::Apply(FSolverParticlesRange& InParticles, const
 						(const ispc::FVector3f&)Velocity,
 						QuarterRho,
 						(FSolverReal)Drag,
+						(FSolverReal)OuterDrag,
 						(FSolverReal)Lift,
+						(FSolverReal)OuterLift,
 						(FSolverReal)Pressure,
 						ColorSize);
 				}
@@ -496,7 +606,9 @@ void FVelocityAndPressureField::Apply(FSolverParticlesRange& InParticles, const
 						(const ispc::FVector3f&)Velocity,
 						QuarterRho,
 						(FSolverReal)Drag,
+						(FSolverReal)OuterDrag,
 						(FSolverReal)Lift,
+						(FSolverReal)OuterLift,
 						(FSolverReal)Pressure,
 						ColorSize,
 						MaxVelocitySquared);
@@ -512,7 +624,9 @@ void FVelocityAndPressureField::Apply(FSolverParticlesRange& InParticles, const
 				{
 					const FSolverVec3 Force = CalculateForce(InParticles, ElementIndex, Velocity,
 						(FSolverReal)Drag,
+						(FSolverReal)OuterDrag,
 						(FSolverReal)Lift,
+						(FSolverReal)OuterLift,
 						(FSolverReal)Pressure);
 					InParticles.Acceleration(Elements[ElementIndex][0]) += InParticles.InvM(Elements[ElementIndex][0]) * Force;
 					InParticles.Acceleration(Elements[ElementIndex][1]) += InParticles.InvM(Elements[ElementIndex][1]) * Force;
@@ -525,7 +639,9 @@ void FVelocityAndPressureField::Apply(FSolverParticlesRange& InParticles, const
 				{
 					const FSolverVec3 Force = CalculateForce(InParticles, ElementIndex, Velocity,
 						(FSolverReal)Drag,
+						(FSolverReal)OuterDrag,
 						(FSolverReal)Lift,
+						(FSolverReal)OuterLift,
 						(FSolverReal)Pressure, MaxVelocitySquared);
 					InParticles.Acceleration(Elements[ElementIndex][0]) += InParticles.InvM(Elements[ElementIndex][0]) * Force;
 					InParticles.Acceleration(Elements[ElementIndex][1]) += InParticles.InvM(Elements[ElementIndex][1]) * Force;
@@ -557,9 +673,15 @@ void FVelocityAndPressureField::Apply(FSolverParticlesRange& InParticles, const
 						bDragHasMap,
 						reinterpret_cast<const ispc::FVector2f&>(Drag.GetOffsetRange()),
 						bDragHasMap ? Drag.GetMapValues().GetData() : nullptr,
+						bOuterDragHasMap,
+						reinterpret_cast<const ispc::FVector2f&>(OuterDrag.GetOffsetRange()),
+						bOuterDragHasMap ? OuterDrag.GetMapValues().GetData() : nullptr,
 						bLiftHasMap,
 						reinterpret_cast<const ispc::FVector2f&>(Lift.GetOffsetRange()),
 						bLiftHasMap ? Lift.GetMapValues().GetData() : nullptr,
+						bOuterLiftHasMap,
+						reinterpret_cast<const ispc::FVector2f&>(OuterLift.GetOffsetRange()),
+						bOuterLiftHasMap ? OuterLift.GetMapValues().GetData() : nullptr,
 						bPressureHasMap,
 						reinterpret_cast<const ispc::FVector2f&>(Pressure.GetOffsetRange()),
 						bPressureHasMap ? Pressure.GetMapValues().GetData() : nullptr,
@@ -583,9 +705,15 @@ void FVelocityAndPressureField::Apply(FSolverParticlesRange& InParticles, const
 						bDragHasMap,
 						reinterpret_cast<const ispc::FVector2f&>(Drag.GetOffsetRange()),
 						bDragHasMap ? Drag.GetMapValues().GetData() : nullptr,
+						bOuterDragHasMap,
+						reinterpret_cast<const ispc::FVector2f&>(OuterDrag.GetOffsetRange()),
+						bOuterDragHasMap ? OuterDrag.GetMapValues().GetData() : nullptr,
 						bLiftHasMap,
 						reinterpret_cast<const ispc::FVector2f&>(Lift.GetOffsetRange()),
 						bLiftHasMap ? Lift.GetMapValues().GetData() : nullptr,
+						bOuterLiftHasMap,
+						reinterpret_cast<const ispc::FVector2f&>(OuterLift.GetOffsetRange()),
+						bOuterLiftHasMap ? OuterLift.GetMapValues().GetData() : nullptr,
 						bPressureHasMap,
 						reinterpret_cast<const ispc::FVector2f&>(Pressure.GetOffsetRange()),
 						bPressureHasMap ? Pressure.GetMapValues().GetData() : nullptr,
@@ -601,11 +729,13 @@ void FVelocityAndPressureField::Apply(FSolverParticlesRange& InParticles, const
 			{
 				for (int32 ElementIndex = 0; ElementIndex < Elements.Num(); ++ElementIndex)
 				{
-					const FSolverReal Cd = Drag.GetValue(ElementIndex);
-					const FSolverReal Cl = Lift.GetValue(ElementIndex);
+					const FSolverReal CdI = Drag.GetValue(ElementIndex);
+					const FSolverReal CdO = OuterDrag.GetValue(ElementIndex);
+					const FSolverReal ClI = Lift.GetValue(ElementIndex);
+					const FSolverReal ClO = OuterLift.GetValue(ElementIndex);
 					const FSolverReal Cp = Pressure.GetValue(ElementIndex);
 
-					const FSolverVec3 Force = CalculateForce(InParticles, ElementIndex, Velocity, Cd, Cl, Cp);
+					const FSolverVec3 Force = CalculateForce(InParticles, ElementIndex, Velocity, CdI, CdO, ClI, ClO, Cp);
 					InParticles.Acceleration(Elements[ElementIndex][0]) += InParticles.InvM(Elements[ElementIndex][0]) * Force;
 					InParticles.Acceleration(Elements[ElementIndex][1]) += InParticles.InvM(Elements[ElementIndex][1]) * Force;
 					InParticles.Acceleration(Elements[ElementIndex][2]) += InParticles.InvM(Elements[ElementIndex][2]) * Force;
@@ -615,11 +745,13 @@ void FVelocityAndPressureField::Apply(FSolverParticlesRange& InParticles, const
 			{
 				for (int32 ElementIndex = 0; ElementIndex < Elements.Num(); ++ElementIndex)
 				{
-					const FSolverReal Cd = Drag.GetValue(ElementIndex);
-					const FSolverReal Cl = Lift.GetValue(ElementIndex);
+					const FSolverReal CdI = Drag.GetValue(ElementIndex);
+					const FSolverReal CdO = OuterDrag.GetValue(ElementIndex);
+					const FSolverReal ClI = Lift.GetValue(ElementIndex);
+					const FSolverReal ClO = OuterLift.GetValue(ElementIndex);
 					const FSolverReal Cp = Pressure.GetValue(ElementIndex);
 
-					const FSolverVec3 Force = CalculateForce(InParticles, ElementIndex, Velocity, Cd, Cl, Cp, MaxVelocitySquared);
+					const FSolverVec3 Force = CalculateForce(InParticles, ElementIndex, Velocity, CdI, CdO, ClI, ClO, Cp, MaxVelocitySquared);
 					InParticles.Acceleration(Elements[ElementIndex][0]) += InParticles.InvM(Elements[ElementIndex][0]) * Force;
 					InParticles.Acceleration(Elements[ElementIndex][1]) += InParticles.InvM(Elements[ElementIndex][1]) * Force;
 					InParticles.Acceleration(Elements[ElementIndex][2]) += InParticles.InvM(Elements[ElementIndex][2]) * Force;
diff --git a/Engine/Source/Runtime/Experimental/Chaos/Private/Chaos/VelocityField.ispc b/Engine/Source/Runtime/Experimental/Chaos/Private/Chaos/VelocityField.ispc
--- a/Engine/Source/Runtime/Experimental/Chaos/Private/Chaos/VelocityField.ispc
+++ b/Engine/Source/Runtime/Experimental/Chaos/Private/Chaos/VelocityField.ispc
@@ -19,8 +19,10 @@ export void UpdateField(uniform FVector3f Forces[],
 								const uniform FVector3f X[],
 								const uniform FVector3f& Velocity,
 								const uniform float QuarterRho,
-								const uniform float Cd,
-								const uniform float Cl,
+								const uniform float CdI,
+								const uniform float CdO,
+								const uniform float ClI,
+								const uniform float ClO,
 								const uniform float Cp,
 								const uniform int32 NumElements)
 {
@@ -45,7 +47,7 @@ export void UpdateField(uniform FVector3f Forces[],
 		const FVector3f X2 = VectorGather(&X[Element.V[2]]);
 
 		// Calculate the normal and the area of the surface exposed to the flow
-		FVector3f N = VectorCross(X1 - X0, X2 - X0);
+		FVector3f N = VectorCross(X2 - X0, X1 - X0);
 		const float DoubleArea = SafeNormalize(N);
 
 		const FVector3f V0 = VectorGather(&V[Element.V[0]]);
@@ -67,10 +69,10 @@ export void UpdateField(uniform FVector3f Forces[],
 		}
 #endif
 
-		const FVector3f EvenForce = (Cd - Cl) * VDotN * RelVelocity + Cl * VSquare * N;
-		const FVector3f OddForce = (Cl - Cd) * VDotN * RelVelocity - Cl * VSquare * N;
+		const FVector3f EvenForce = (CdI - ClI) * VDotN * RelVelocity + ClI * VSquare * N;
+		const FVector3f OddForce = (ClO - CdO) * VDotN * RelVelocity - ClO * VSquare * N;
 
-		const FVector3f Force = QuarterRho * DoubleArea * VectorSelect(VDotN >= 0, EvenForce, OddForce) - DoubleArea * 0.5 * Cp * N; // N points in the opposite direction of the actual mesh normals
+		const FVector3f Force = QuarterRho * DoubleArea * VectorSelect(VDotN >= 0, EvenForce, OddForce) + DoubleArea * 0.5 * Cp * N;
 
 		VectorStore(&Forces[Offset], Force);
 
@@ -87,9 +89,15 @@ export void UpdateFieldWithWeightMaps(uniform FVector3f Forces[],
 								const uniform bool bDragHasMap,
 								const uniform FVector2f& DragOffsetRange,
 								const uniform float DragMapValues[],
+								const uniform bool bOuterDragHasMap,
+								const uniform FVector2f& OuterDragOffsetRange,
+								const uniform float OuterDragMapValues[],
 								const uniform bool bLiftHasMap,
 								const uniform FVector2f& LiftOffsetRange,
 								const uniform float LiftMapValues[],
+								const uniform bool bOuterLiftHasMap,
+								const uniform FVector2f& OuterLiftOffsetRange,
+								const uniform float OuterLiftMapValues[],
 								const uniform bool bPressureHasMap,
 								const uniform FVector2f& PressureOffsetRange,
 								const uniform float PressureMapValues[],
@@ -115,7 +123,7 @@ export void UpdateFieldWithWeightMaps(uniform FVector3f Forces[],
 		const FVector3f X2 = VectorGather(&X[Element.V[2]]);
 
 		// Calculate the normal and the area of the surface exposed to the flow
-		FVector3f N = VectorCross(X1 - X0, X2 - X0);
+		FVector3f N = VectorCross(X2 - X0, X1 - X0);
 		const float DoubleArea = SafeNormalize(N);
 
 		const FVector3f V0 = VectorGather(&V[Element.V[0]]);
@@ -130,8 +138,10 @@ export void UpdateFieldWithWeightMaps(uniform FVector3f Forces[],
 		const float VDotN = VectorDot(RelVelocity, N);
 		const float VSquare = VectorDot(RelVelocity, RelVelocity);
 
-		const float Cd = bDragHasMap ? DragOffsetRange.V[0] + DragMapValues[ElementIndex] * DragOffsetRange.V[1] : DragOffsetRange.V[0];
-		const float Cl = bLiftHasMap ? LiftOffsetRange.V[0] + LiftMapValues[ElementIndex] * LiftOffsetRange.V[1] : LiftOffsetRange.V[0];
+		const float CdI = bDragHasMap ? DragOffsetRange.V[0] + DragMapValues[ElementIndex] * DragOffsetRange.V[1] : DragOffsetRange.V[0];
+		const float CdO = bOuterDragHasMap ? OuterDragOffsetRange.V[0] + OuterDragMapValues[ElementIndex] * OuterDragOffsetRange.V[1] : OuterDragOffsetRange.V[0];
+		const float ClI = bLiftHasMap ? LiftOffsetRange.V[0] + LiftMapValues[ElementIndex] * LiftOffsetRange.V[1] : LiftOffsetRange.V[0];
+		const float ClO = bOuterLiftHasMap ? OuterLiftOffsetRange.V[0] + OuterLiftMapValues[ElementIndex] * OuterLiftOffsetRange.V[1] : OuterLiftOffsetRange.V[0];
 		const float Cp = bPressureHasMap ? PressureOffsetRange.V[0] + PressureMapValues[ElementIndex] * PressureOffsetRange.V[1] : PressureOffsetRange.V[0];
 
 #if HW_GATHER_SUPPORTED == 1
@@ -141,10 +151,10 @@ export void UpdateFieldWithWeightMaps(uniform FVector3f Forces[],
 		}
 #endif
 
-		const FVector3f EvenForce = (Cd - Cl) * VDotN * RelVelocity + Cl * VSquare * N;
-		const FVector3f OddForce = (Cl - Cd) * VDotN * RelVelocity - Cl * VSquare * N;
+		const FVector3f EvenForce = (CdI - ClI) * VDotN * RelVelocity + ClI * VSquare * N;
+		const FVector3f OddForce = (ClO - CdO) * VDotN * RelVelocity - ClO * VSquare * N;
 
-		const FVector3f Force = QuarterRho * DoubleArea * VectorSelect(VDotN >= 0, EvenForce, OddForce) - DoubleArea * 0.5 * Cp * N; // N points in the opposite direction of the actual mesh normals
+		const FVector3f Force = QuarterRho * DoubleArea * VectorSelect(VDotN >= 0, EvenForce, OddForce) + DoubleArea * 0.5 * Cp * N;
 
 		VectorStore(&Forces[Offset], Force);
 
@@ -158,8 +168,10 @@ export void UpdateFieldAndClampVelocity(uniform FVector3f Forces[],
 								const uniform FVector3f X[],
 								const uniform FVector3f& Velocity,
 								const uniform float QuarterRho,
-								const uniform float Cd,
-								const uniform float Cl,
+								const uniform float CdI,
+								const uniform float CdO,
+								const uniform float ClI,
+								const uniform float ClO,
 								const uniform float Cp,
 								const uniform int32 NumElements,
 								const uniform float MaxVelocitySquared)
@@ -185,7 +197,7 @@ export void UpdateFieldAndClampVelocity(uniform FVector3f Forces[],
 		const FVector3f X2 = VectorGather(&X[Element.V[2]]);
 
 		// Calculate the normal and the area of the surface exposed to the flow
-		FVector3f N = VectorCross(X1 - X0, X2 - X0);
+		FVector3f N = VectorCross(X2 - X0, X1 - X0);
 		const float DoubleArea = SafeNormalize(N);
 
 		const FVector3f V0 = VectorGather(&V[Element.V[0]]);
@@ -214,10 +226,10 @@ export void UpdateFieldAndClampVelocity(uniform FVector3f Forces[],
 		}
 #endif
 
-		const FVector3f EvenForce = (Cd - Cl) * VDotN * RelVelocity + Cl * VSquare * N;
-		const FVector3f OddForce = (Cl - Cd) * VDotN * RelVelocity - Cl * VSquare * N;
-
-		const FVector3f Force = QuarterRho * DoubleArea * VectorSelect(VDotN >= 0, EvenForce, OddForce) - DoubleArea * 0.5 * Cp * N; // N points in the opposite direction of the actual mesh normals
+		const FVector3f EvenForce = (CdI - ClI) * VDotN * RelVelocity + ClI * VSquare * N;
+		const FVector3f OddForce = (ClO - CdO) * VDotN * RelVelocity - ClO * VSquare * N;
+		
+		const FVector3f Force = QuarterRho * DoubleArea * VectorSelect(VDotN >= 0, EvenForce, OddForce) + DoubleArea * 0.5 * Cp * N;
 
 		VectorStore(&Forces[Offset], Force);
 
@@ -234,9 +246,15 @@ export void UpdateFieldWithWeightMapsAndClampVelocity(uniform FVector3f Forces[]
 								const uniform bool bDragHasMap,
 								const uniform FVector2f& DragOffsetRange,
 								const uniform float DragMapValues[],
+								const uniform bool bOuterDragHasMap,
+								const uniform FVector2f& OuterDragOffsetRange,
+								const uniform float OuterDragMapValues[],
 								const uniform bool bLiftHasMap,
 								const uniform FVector2f& LiftOffsetRange,
 								const uniform float LiftMapValues[],
+								const uniform bool bOuterLiftHasMap,
+								const uniform FVector2f& OuterLiftOffsetRange,
+								const uniform float OuterLiftMapValues[],
 								const uniform bool bPressureHasMap,
 								const uniform FVector2f& PressureOffsetRange,
 								const uniform float PressureMapValues[],
@@ -263,7 +281,7 @@ export void UpdateFieldWithWeightMapsAndClampVelocity(uniform FVector3f Forces[]
 		const FVector3f X2 = VectorGather(&X[Element.V[2]]);
 
 		// Calculate the normal and the area of the surface exposed to the flow
-		FVector3f N = VectorCross(X1 - X0, X2 - X0);
+		FVector3f N = VectorCross(X2 - X0, X1 - X0);
 		const float DoubleArea = SafeNormalize(N);
 
 		const FVector3f V0 = VectorGather(&V[Element.V[0]]);
@@ -285,8 +303,10 @@ export void UpdateFieldWithWeightMapsAndClampVelocity(uniform FVector3f Forces[]
 		const float VDotN = VectorDot(RelVelocity, N);
 		const float VSquare = VectorDot(RelVelocity, RelVelocity);
 		
-		const float Cd = bDragHasMap ? DragOffsetRange.V[0] + DragMapValues[ElementIndex] * DragOffsetRange.V[1] : DragOffsetRange.V[0];
-		const float Cl = bLiftHasMap ? LiftOffsetRange.V[0] + LiftMapValues[ElementIndex] * LiftOffsetRange.V[1] : LiftOffsetRange.V[0];
+		const float CdI = bDragHasMap ? DragOffsetRange.V[0] + DragMapValues[ElementIndex] * DragOffsetRange.V[1] : DragOffsetRange.V[0];
+		const float CdO = bOuterDragHasMap ? OuterDragOffsetRange.V[0] + OuterDragMapValues[ElementIndex] * OuterDragOffsetRange.V[1] : OuterDragOffsetRange.V[0];
+		const float ClI = bLiftHasMap ? LiftOffsetRange.V[0] + LiftMapValues[ElementIndex] * LiftOffsetRange.V[1] : LiftOffsetRange.V[0];
+		const float ClO = bOuterLiftHasMap ? OuterLiftOffsetRange.V[0] + OuterLiftMapValues[ElementIndex] * OuterLiftOffsetRange.V[1] : OuterLiftOffsetRange.V[0];
 		const float Cp = bPressureHasMap ? PressureOffsetRange.V[0] + PressureMapValues[ElementIndex] * PressureOffsetRange.V[1] : PressureOffsetRange.V[0];
 
 #if HW_GATHER_SUPPORTED == 1
@@ -296,10 +316,10 @@ export void UpdateFieldWithWeightMapsAndClampVelocity(uniform FVector3f Forces[]
 		}
 #endif
 
-		const FVector3f EvenForce = (Cd - Cl) * VDotN * RelVelocity + Cl * VSquare * N;
-		const FVector3f OddForce = (Cl - Cd) * VDotN * RelVelocity - Cl * VSquare * N;
+		const FVector3f EvenForce = (CdI - ClI) * VDotN * RelVelocity + ClI * VSquare * N;
+		const FVector3f OddForce = (ClO - CdO) * VDotN * RelVelocity - ClO * VSquare * N;
 
-		const FVector3f Force = QuarterRho * DoubleArea * VectorSelect(VDotN >= 0, EvenForce, OddForce) - DoubleArea * 0.5 * Cp * N; // N points in the opposite direction of the actual mesh normals
+		const FVector3f Force = QuarterRho * DoubleArea * VectorSelect(VDotN >= 0, EvenForce, OddForce) + DoubleArea * 0.5 * Cp * N;
 
 		VectorStore(&Forces[Offset], Force);
 
@@ -314,8 +334,10 @@ export void UpdateAndApplyVelocityField(uniform FVector3f Acceleration[],
 								const uniform FVector3f X[],
 								const uniform FVector3f& Velocity,
 								const uniform float QuarterRho,
-								const uniform float Cd,
-								const uniform float Cl,
+								const uniform float CdI,
+								const uniform float CdO,
+								const uniform float ClI,
+								const uniform float ClO,
 								const uniform float Cp,
 								const uniform int32 NumElements)
 {
@@ -332,7 +354,7 @@ export void UpdateAndApplyVelocityField(uniform FVector3f Acceleration[],
 		const FVector3f X2 = VectorGather(&X[Index2]);
 
 		// Calculate the normal and the area of the surface exposed to the flow
-		FVector3f N = VectorCross(X1 - X0, X2 - X0);
+		FVector3f N = VectorCross(X2 - X0, X1 - X0);
 		const float DoubleArea = SafeNormalize(N);
 
 		const FVector3f V0 = VectorGather(&V[Index0]);
@@ -346,11 +368,11 @@ export void UpdateAndApplyVelocityField(uniform FVector3f Acceleration[],
 		// Set the aerodynamic forces
 		const float VDotN = VectorDot(RelVelocity, N);
 		const float VSquare = VectorDot(RelVelocity, RelVelocity);
-
-		const FVector3f EvenForce = (Cd - Cl) * VDotN * RelVelocity + Cl * VSquare * N;
-		const FVector3f OddForce = (Cl - Cd) * VDotN * RelVelocity - Cl * VSquare * N;
-
-		const FVector3f Force = QuarterRho * DoubleArea * VectorSelect(VDotN >= 0, EvenForce, OddForce) - DoubleArea * 0.5 * Cp * N; // N points in the opposite direction of the actual mesh normals
+		
+		const FVector3f EvenForce = (CdI - ClI) * VDotN * RelVelocity + ClI * VSquare * N;
+		const FVector3f OddForce = (ClO - CdO) * VDotN * RelVelocity - ClO * VSquare * N;
+		
+		const FVector3f Force = QuarterRho * DoubleArea * VectorSelect(VDotN >= 0, EvenForce, OddForce) + DoubleArea * 0.5 * Cp * N;
 		
 		#pragma ignore warning(perf)
 		const float InvM0 = InvM[Index0];
@@ -376,8 +398,10 @@ export void UpdateAndClampVelocityAndApplyVelocityField(uniform FVector3f Accele
 								const uniform FVector3f X[],
 								const uniform FVector3f& Velocity,
 								const uniform float QuarterRho,
-								const uniform float Cd,
-								const uniform float Cl,
+								const uniform float CdI,
+								const uniform float CdO,
+								const uniform float ClI,
+								const uniform float ClO,
 								const uniform float Cp,
 								const uniform int32 NumElements,
 								const uniform float MaxVelocitySquared)
@@ -395,7 +419,7 @@ export void UpdateAndClampVelocityAndApplyVelocityField(uniform FVector3f Accele
 		const FVector3f X2 = VectorGather(&X[Index2]);
 
 		// Calculate the normal and the area of the surface exposed to the flow
-		FVector3f N = VectorCross(X1 - X0, X2 - X0);
+		FVector3f N = VectorCross(X2 - X0, X1 - X0);
 		const float DoubleArea = SafeNormalize(N);
 
 		const FVector3f V0 = VectorGather(&V[Index0]);
@@ -416,11 +440,11 @@ export void UpdateAndClampVelocityAndApplyVelocityField(uniform FVector3f Accele
 		// Set the aerodynamic forces
 		const float VDotN = VectorDot(RelVelocity, N);
 		const float VSquare = VectorDot(RelVelocity, RelVelocity);
-
-		const FVector3f EvenForce = (Cd - Cl) * VDotN * RelVelocity + Cl * VSquare * N;
-		const FVector3f OddForce = (Cl - Cd) * VDotN * RelVelocity - Cl * VSquare * N;
-
-		const FVector3f Force = QuarterRho * DoubleArea * VectorSelect(VDotN >= 0, EvenForce, OddForce) - DoubleArea * 0.5 * Cp * N; // N points in the opposite direction of the actual mesh normals
+		
+		const FVector3f EvenForce = (CdI - ClI) * VDotN * RelVelocity + ClI * VSquare * N;
+		const FVector3f OddForce = (ClO - CdO) * VDotN * RelVelocity - ClO * VSquare * N;
+		
+		const FVector3f Force = QuarterRho * DoubleArea * VectorSelect(VDotN >= 0, EvenForce, OddForce) + DoubleArea * 0.5 * Cp * N;
 		
 		#pragma ignore warning(perf)
 		const float InvM0 = InvM[Index0];
@@ -450,9 +474,15 @@ export void UpdateAndApplyVelocityFieldWithWeightMaps(uniform FVector3f Accelera
 								const uniform bool bDragHasMap,
 								const uniform FVector2f& DragOffsetRange,
 								const uniform float DragMapValues[],
+								const uniform bool bOuterDragHasMap,
+								const uniform FVector2f& OuterDragOffsetRange,
+								const uniform float OuterDragMapValues[],
 								const uniform bool bLiftHasMap,
 								const uniform FVector2f& LiftOffsetRange,
 								const uniform float LiftMapValues[],
+								const uniform bool bOuterLiftHasMap,
+								const uniform FVector2f& OuterLiftOffsetRange,
+								const uniform float OuterLiftMapValues[],
 								const uniform bool bPressureHasMap,
 								const uniform FVector2f& PressureOffsetRange,
 								const uniform float PressureMapValues[],
@@ -471,7 +501,7 @@ export void UpdateAndApplyVelocityFieldWithWeightMaps(uniform FVector3f Accelera
 		const FVector3f X2 = VectorGather(&X[Index2]);
 
 		// Calculate the normal and the area of the surface exposed to the flow
-		FVector3f N = VectorCross(X1 - X0, X2 - X0);
+		FVector3f N = VectorCross(X2 - X0, X1 - X0);
 		const float DoubleArea = SafeNormalize(N);
 
 		const FVector3f V0 = VectorGather(&V[Index0]);
@@ -486,14 +516,16 @@ export void UpdateAndApplyVelocityFieldWithWeightMaps(uniform FVector3f Accelera
 		const float VDotN = VectorDot(RelVelocity, N);
 		const float VSquare = VectorDot(RelVelocity, RelVelocity);
 		
-		const float Cd = bDragHasMap ? DragOffsetRange.V[0] + DragMapValues[ElementIndex] * DragOffsetRange.V[1] : DragOffsetRange.V[0];
-		const float Cl = bLiftHasMap ? LiftOffsetRange.V[0] + LiftMapValues[ElementIndex] * LiftOffsetRange.V[1] : LiftOffsetRange.V[0];
+		const float CdI = bDragHasMap ? DragOffsetRange.V[0] + DragMapValues[ElementIndex] * DragOffsetRange.V[1] : DragOffsetRange.V[0];
+		const float CdO = bOuterDragHasMap ? OuterDragOffsetRange.V[0] + OuterDragMapValues[ElementIndex] * OuterDragOffsetRange.V[1] : OuterDragOffsetRange.V[0];
+		const float ClI = bLiftHasMap ? LiftOffsetRange.V[0] + LiftMapValues[ElementIndex] * LiftOffsetRange.V[1] : LiftOffsetRange.V[0];
+		const float ClO = bOuterLiftHasMap ? OuterLiftOffsetRange.V[0] + OuterLiftMapValues[ElementIndex] * OuterLiftOffsetRange.V[1] : OuterLiftOffsetRange.V[0];
 		const float Cp = bPressureHasMap ? PressureOffsetRange.V[0] + PressureMapValues[ElementIndex] * PressureOffsetRange.V[1] : PressureOffsetRange.V[0];
+		
+		const FVector3f EvenForce = (CdI - ClI) * VDotN * RelVelocity + ClI * VSquare * N;
+		const FVector3f OddForce = (ClO - CdO) * VDotN * RelVelocity - ClO * VSquare * N;
 
-		const FVector3f EvenForce = (Cd - Cl) * VDotN * RelVelocity + Cl * VSquare * N;
-		const FVector3f OddForce = (Cl - Cd) * VDotN * RelVelocity - Cl * VSquare * N;
-
-		const FVector3f Force = QuarterRho * DoubleArea * VectorSelect(VDotN >= 0, EvenForce, OddForce) - DoubleArea * 0.5 * Cp * N; // N points in the opposite direction of the actual mesh normals
+		const FVector3f Force = QuarterRho * DoubleArea * VectorSelect(VDotN >= 0, EvenForce, OddForce) + DoubleArea * 0.5 * Cp * N;
 		
 		#pragma ignore warning(perf)
 		const float InvM0 = InvM[Index0];
@@ -522,9 +554,15 @@ export void UpdateAndClampVelocityAndApplyVelocityFieldWithWeightMaps(uniform FV
 								const uniform bool bDragHasMap,
 								const uniform FVector2f& DragOffsetRange,
 								const uniform float DragMapValues[],
+								const uniform bool bOuterDragHasMap,
+								const uniform FVector2f& OuterDragOffsetRange,
+								const uniform float OuterDragMapValues[],
 								const uniform bool bLiftHasMap,
 								const uniform FVector2f& LiftOffsetRange,
 								const uniform float LiftMapValues[],
+								const uniform bool bOuterLiftHasMap,
+								const uniform FVector2f& OuterLiftOffsetRange,
+								const uniform float OuterLiftMapValues[],
 								const uniform bool bPressureHasMap,
 								const uniform FVector2f& PressureOffsetRange,
 								const uniform float PressureMapValues[],
@@ -544,7 +582,7 @@ export void UpdateAndClampVelocityAndApplyVelocityFieldWithWeightMaps(uniform FV
 		const FVector3f X2 = VectorGather(&X[Index2]);
 
 		// Calculate the normal and the area of the surface exposed to the flow
-		FVector3f N = VectorCross(X1 - X0, X2 - X0);
+		FVector3f N = VectorCross(X2 - X0, X1 - X0);
 		const float DoubleArea = SafeNormalize(N);
 
 		const FVector3f V0 = VectorGather(&V[Index0]);
@@ -566,14 +604,16 @@ export void UpdateAndClampVelocityAndApplyVelocityFieldWithWeightMaps(uniform FV
 		const float VDotN = VectorDot(RelVelocity, N);
 		const float VSquare = VectorDot(RelVelocity, RelVelocity);
 		
-		const float Cd = bDragHasMap ? DragOffsetRange.V[0] + DragMapValues[ElementIndex] * DragOffsetRange.V[1] : DragOffsetRange.V[0];
-		const float Cl = bLiftHasMap ? LiftOffsetRange.V[0] + LiftMapValues[ElementIndex] * LiftOffsetRange.V[1] : LiftOffsetRange.V[0];
+		const float CdI = bDragHasMap ? DragOffsetRange.V[0] + DragMapValues[ElementIndex] * DragOffsetRange.V[1] : DragOffsetRange.V[0];
+		const float CdO = bOuterDragHasMap ? OuterDragOffsetRange.V[0] + OuterDragMapValues[ElementIndex] * OuterDragOffsetRange.V[1] : OuterDragOffsetRange.V[0];
+		const float ClI = bLiftHasMap ? LiftOffsetRange.V[0] + LiftMapValues[ElementIndex] * LiftOffsetRange.V[1] : LiftOffsetRange.V[0];
+		const float ClO = bOuterLiftHasMap ? OuterLiftOffsetRange.V[0] + OuterLiftMapValues[ElementIndex] * OuterLiftOffsetRange.V[1] : OuterLiftOffsetRange.V[0];
 		const float Cp = bPressureHasMap ? PressureOffsetRange.V[0] + PressureMapValues[ElementIndex] * PressureOffsetRange.V[1] : PressureOffsetRange.V[0];
+		
+		const FVector3f EvenForce = (CdI - ClI) * VDotN * RelVelocity + ClI * VSquare * N;
+		const FVector3f OddForce = (ClO - CdO) * VDotN * RelVelocity - ClO * VSquare * N;
 
-		const FVector3f EvenForce = (Cd - Cl) * VDotN * RelVelocity + Cl * VSquare * N;
-		const FVector3f OddForce = (Cl - Cd) * VDotN * RelVelocity - Cl * VSquare * N;
-
-		const FVector3f Force = QuarterRho * DoubleArea * VectorSelect(VDotN >= 0, EvenForce, OddForce) - DoubleArea * 0.5 * Cp * N; // N points in the opposite direction of the actual mesh normals
+		const FVector3f Force = QuarterRho * DoubleArea * VectorSelect(VDotN >= 0, EvenForce, OddForce) + DoubleArea * 0.5 * Cp * N;
 		
 		#pragma ignore warning(perf)
 		const float InvM0 = InvM[Index0];
diff --git a/Engine/Source/Runtime/Experimental/Chaos/Public/Chaos/SoftsSimulationSpace.h b/Engine/Source/Runtime/Experimental/Chaos/Public/Chaos/SoftsSimulationSpace.h
new file mode 100644
--- /dev/null
+++ b/Engine/Source/Runtime/Experimental/Chaos/Public/Chaos/SoftsSimulationSpace.h
@@ -0,0 +1,18 @@
+// Copyright Epic Games, Inc. All Rights Reserved.
+#pragma once
+#include "HAL/Platform.h"
+#include "UObject/ObjectMacros.h"
+#include "SoftsSimulationSpace.generated.h"
+
+UENUM()
+enum struct EChaosSoftsSimulationSpace : uint8
+{
+	// World space
+	WorldSpace = 0,
+
+	// Component space
+	ComponentSpace,
+
+	// Top level bone for each cloth.
+	ReferenceBoneSpace,
+};
\ No newline at end of file
diff --git a/Engine/Source/Runtime/Experimental/Chaos/Public/Chaos/VelocityField.h b/Engine/Source/Runtime/Experimental/Chaos/Public/Chaos/VelocityField.h
--- a/Engine/Source/Runtime/Experimental/Chaos/Public/Chaos/VelocityField.h
+++ b/Engine/Source/Runtime/Experimental/Chaos/Public/Chaos/VelocityField.h
@@ -4,6 +4,7 @@
 #include "Chaos/PBDSoftsEvolutionFwd.h"
 #include "Chaos/PBDSoftsSolverParticles.h"
 #include "Chaos/SoftsSolverParticlesRange.h"
+#include "Chaos/SoftsSimulationSpace.h"
 #include "Chaos/TriangleMesh.h"
 #include "Chaos/CollectionPropertyFacade.h"
 #include "Chaos/PBDFlatWeightMap.h"
@@ -20,6 +21,7 @@ public:
 	static constexpr FSolverReal DefaultFluidDensity = (FSolverReal)1.225;
 	static constexpr FSolverReal MinCoefficient = (FSolverReal)0.;   // Applies to both drag and lift
 	static constexpr FSolverReal MaxCoefficient = (FSolverReal)10.;  //
+	static constexpr EChaosSoftsSimulationSpace DefaultWindVelocitySpace = EChaosSoftsSimulationSpace::WorldSpace;
 
 	static bool IsEnabled(const FCollectionPropertyConstFacade& PropertyCollection)
 	{
@@ -33,15 +35,20 @@ public:
 		: Offset(INDEX_NONE)
 		, NumParticles(0)
 		, Lift(GetWeightedFloatLift(PropertyCollection, (FSolverReal)0.).ClampAxes(MinCoefficient, MaxCoefficient))
+		, OuterLift(GetWeightedFloatOuterLift(PropertyCollection, FSolverVec2(Lift.GetLow(), Lift.GetHigh())).ClampAxes(MinCoefficient, MaxCoefficient))
 		, Drag(GetWeightedFloatDrag(PropertyCollection, (FSolverReal)0.).ClampAxes(MinCoefficient, MaxCoefficient))
+		, OuterDrag(GetWeightedFloatOuterDrag(PropertyCollection, FSolverVec2(Drag.GetLow(), Drag.GetHigh())).ClampAxes(MinCoefficient, MaxCoefficient))
 		, Pressure(GetWeightedFloatPressure(PropertyCollection, (FSolverReal)0.))
 		, Rho(FMath::Max(GetFluidDensity(PropertyCollection, (FSolverReal)0.), (FSolverReal)0.))
 		, QuarterRho(Rho * (FSolverReal)0.25f)
 		, DragIndex(PropertyCollection)
+		, OuterDragIndex(PropertyCollection)
 		, LiftIndex(PropertyCollection)
+		, OuterLiftIndex(PropertyCollection)
 		, FluidDensityIndex(PropertyCollection)
 		, PressureIndex(PropertyCollection)
 		, WindVelocityIndex(PropertyCollection)
+		, WindVelocitySpaceIndex(PropertyCollection)
 	{
 	}
 
@@ -52,21 +59,22 @@ public:
 		const TMap<FString, TConstArrayView<FRealSingle>>& Weightmaps,
 		FSolverReal WorldScale)
 		: Lift(GetWeightedFloatLift(PropertyCollection, (FSolverReal)0.).ClampAxes(MinCoefficient, MaxCoefficient))
+		, OuterLift(GetWeightedFloatOuterLift(PropertyCollection, FSolverVec2(Lift.GetLow(), Lift.GetHigh())).ClampAxes(MinCoefficient, MaxCoefficient))
 		, Drag(GetWeightedFloatDrag(PropertyCollection, (FSolverReal)0.).ClampAxes(MinCoefficient, MaxCoefficient))
+		, OuterDrag(GetWeightedFloatOuterDrag(PropertyCollection, FSolverVec2(Drag.GetLow(), Drag.GetHigh())).ClampAxes(MinCoefficient, MaxCoefficient))
 		, Pressure(GetWeightedFloatPressure(PropertyCollection, (FSolverReal)0.)/WorldScale)
+		, Rho(FMath::Max(GetFluidDensity(PropertyCollection, (FSolverReal)0.), (FSolverReal)0.))
+		, QuarterRho(Rho* (FSolverReal)0.25f)
 		, DragIndex(PropertyCollection)
+		, OuterDragIndex(PropertyCollection)
 		, LiftIndex(PropertyCollection)
+		, OuterLiftIndex(PropertyCollection)
 		, FluidDensityIndex(PropertyCollection)
 		, PressureIndex(PropertyCollection)
 		, WindVelocityIndex(PropertyCollection)
+		, WindVelocitySpaceIndex(PropertyCollection)
 	{
 		SetGeometry(Particles, TriangleMesh);
-		SetProperties(
-			FSolverVec2(GetWeightedFloatDrag(PropertyCollection, 0.f)),  // If these properties don't exist, set their values to 0, not to DefaultCoefficients!
-			FSolverVec2(GetWeightedFloatLift(PropertyCollection, 0.f)),
-			(FSolverReal)GetFluidDensity(PropertyCollection, 0.f),
-			FSolverVec2(GetWeightedFloatPressure(PropertyCollection, 0.f)),  // These getters also initialize the property indices, so keep before SetMultipliers
-			WorldScale);
 		SetMultipliers(PropertyCollection, Weightmaps);
 		InitColor(Particles);
 	}
@@ -76,15 +84,21 @@ public:
 		: Offset(INDEX_NONE)
 		, NumParticles(0)
 		, Lift(FSolverVec2(0.))
+		, OuterLift(FSolverVec2(0.))
 		, Drag(FSolverVec2(0.))
+		, OuterDrag(FSolverVec2(0.))
 		, Pressure(FSolverVec2(0.))
+		, Rho((FSolverReal)0.)
+		, QuarterRho(Rho* (FSolverReal)0.25f)
 		, DragIndex(ForceInit)
+		, OuterDragIndex(ForceInit)
 		, LiftIndex(ForceInit)
+		, OuterLiftIndex(ForceInit)
 		, FluidDensityIndex(ForceInit)
 		, PressureIndex(ForceInit)
 		, WindVelocityIndex(ForceInit)
+		, WindVelocitySpaceIndex(ForceInit)
 	{
-		SetProperties(FSolverVec2(0.), FSolverVec2(0.), (FSolverReal)0., FSolverVec2(0.));
 	}
 
 	~FVelocityAndPressureField() {}
@@ -112,26 +126,44 @@ public:
 		bool bEnableAerodynamics);
 
 	// This version will load WindVelocity from the config
+	// Provide LocalSpaceRotation and/or ReferenceSpaceRotation to convert WindVelocity to solver space based on WindVelocitySpace
 	CHAOS_API void SetPropertiesAndWind(
 		const FCollectionPropertyConstFacade& PropertyCollection,
 		const TMap<FString, TConstArrayView<FRealSingle>>& Weightmaps,
 		FSolverReal WorldScale,
 		bool bEnableAerodynamics,
-		const FSolverVec3& SolverWind);
+		const FSolverVec3& SolverWind,
+		const FRotation3& LocalSpaceRotation = FRotation3::Identity,
+		const FRotation3& ReferenceSpaceRotation = FRotation3::Identity);
 
 	CHAOS_API void SetProperties(
 		const FSolverVec2& Drag,
+		const FSolverVec2& OuterDrag,
 		const FSolverVec2& Lift,
+		const FSolverVec2& OuterLift,
 		const FSolverReal FluidDensity,
 		const FSolverVec2& Pressure = FSolverVec2::ZeroVector,
 		FSolverReal WorldScale = 1.f);
 
+	UE_DEPRECATED(5.5, "Use the SetProperties that includes OuterDrag and OuterLift")
+	void SetProperties(
+		const FSolverVec2& InDrag,
+		const FSolverVec2& InLift,
+		const FSolverReal InFluidDensity,
+		const FSolverVec2& InPressure = FSolverVec2::ZeroVector,
+		FSolverReal WorldScale = 1.f)
+	{
+		SetProperties(InDrag, InDrag, InLift, InLift, InFluidDensity, InPressure, WorldScale);
+	}
+
 	bool IsActive() const 
 	{ 
 		return Pressure.GetLow() != (FSolverReal)0. || Pressure.GetHigh() != (FSolverReal)0. ||
 			(AreAerodynamicsEnabled() && (
 				Drag.GetLow() > (FSolverReal)0. || Drag.GetOffsetRange()[1] != (FSolverReal)0. ||  // Note: range can be a negative value (although not when Lift or Drag base is zero)
-				Lift.GetLow() > (FSolverReal)0. || Lift.GetOffsetRange()[1] != (FSolverReal)0.));
+				OuterDrag.GetLow() > (FSolverReal)0. || OuterDrag.GetOffsetRange()[1] != (FSolverReal)0. ||
+				Lift.GetLow() > (FSolverReal)0. || Lift.GetOffsetRange()[1] != (FSolverReal)0. ||
+				OuterLift.GetLow() > (FSolverReal)0. || OuterLift.GetOffsetRange()[1] != (FSolverReal)0.));
 	}
 
 	CHAOS_API void SetGeometry(
@@ -143,8 +175,20 @@ public:
 	CHAOS_API void SetGeometry(
 		const FTriangleMesh* TriangleMesh,
 		const TConstArrayView<FRealSingle>& DragMultipliers,
+		const TConstArrayView<FRealSingle>& OuterDragMultipliers,
 		const TConstArrayView<FRealSingle>& LiftMultipliers,
-		const TConstArrayView<FRealSingle>& PressureMultipliers = TConstArrayView<FRealSingle>());
+		const TConstArrayView<FRealSingle>& OuterLiftMultipliers,
+		const TConstArrayView<FRealSingle>& PressureMultipliers);
+
+	UE_DEPRECATED(5.5, "Use SetGeometry with OuterDrag and OuterLift multipliers")
+	void SetGeometry(
+		const FTriangleMesh* TriangleMesh,
+		const TConstArrayView<FRealSingle>& DragMultipliers,
+		const TConstArrayView<FRealSingle>& LiftMultipliers,
+		const TConstArrayView<FRealSingle>& PressureMultipliers = TConstArrayView<FRealSingle>())
+	{
+		SetGeometry(TriangleMesh, DragMultipliers, DragMultipliers, LiftMultipliers, LiftMultipliers, PressureMultipliers);
+	}
 
 	void SetVelocity(const FSolverVec3& InVelocity) { Velocity = InVelocity; }
 
@@ -164,19 +208,21 @@ private:
 
 	CHAOS_API void SetMultipliers(
 		const TConstArrayView<FRealSingle>& DragMultipliers,
+		const TConstArrayView<FRealSingle>& OuterDragMultipliers,
 		const TConstArrayView<FRealSingle>& LiftMultipliers,
+		const TConstArrayView<FRealSingle>& OuterLiftMultipliers,
 		const TConstArrayView<FRealSingle>& PressureMultipliers);
 
 
 	template<typename SolverParticlesOrRange>
-	FSolverVec3 CalculateForce(const SolverParticlesOrRange& InParticles, int32 ElementIndex, const FSolverVec3& InVelocity, const FSolverReal Cd, const FSolverReal Cl, const FSolverReal Cp) const 
+	FSolverVec3 CalculateForce(const SolverParticlesOrRange& InParticles, int32 ElementIndex, const FSolverVec3& InVelocity, const FSolverReal CdI, const FSolverReal CdO, const FSolverReal ClI, const FSolverReal ClO, const FSolverReal Cp) const 
 	{
 		const TVec3<int32>& Element = Elements[ElementIndex];
 
 		// Calculate the normal and the area of the surface exposed to the flow
 		FSolverVec3 N = FSolverVec3::CrossProduct(
-			InParticles.GetX(Element[1]) - InParticles.GetX(Element[0]),
-			InParticles.GetX(Element[2]) - InParticles.GetX(Element[0]));
+			InParticles.GetX(Element[2]) - InParticles.GetX(Element[0]),
+			InParticles.GetX(Element[1]) - InParticles.GetX(Element[0]));
 		const FSolverReal DoubleArea = N.SafeNormalize();
 
 		// Calculate the direction and the relative velocity of the triangle to the flow
@@ -191,17 +237,17 @@ private:
 		const FSolverReal VSquare = FSolverVec3::DotProduct(V, V);
 
 		return QuarterRho * DoubleArea * (VDotN >= (FSolverReal)0. ?  // The flow can hit either side of the triangle, so the normal might need to be reversed
-			(Cd - Cl) * VDotN * V + Cl * VSquare * N :
-			(Cl - Cd) * VDotN * V - Cl * VSquare * N) - DoubleArea * (FSolverReal)0.5 * Cp * N; // N points in the opposite direction of the actual mesh normals
+			(CdI - ClI) * VDotN * V + ClI * VSquare * N :
+			(ClO - CdO) * VDotN * V - ClO * VSquare * N) + DoubleArea * (FSolverReal)0.5 * Cp * N;
 	}
 
-	void UpdateField(const FSolverParticles& InParticles, int32 ElementIndex, const FSolverVec3& InVelocity, const FSolverReal Cd, const FSolverReal Cl, const FSolverReal Cp)
+	void UpdateField(const FSolverParticles& InParticles, int32 ElementIndex, const FSolverVec3& InVelocity, const FSolverReal CdI, const FSolverReal CdO, const FSolverReal ClI, const FSolverReal ClO, const FSolverReal Cp)
 	{
-		Forces[ElementIndex] = CalculateForce(InParticles, ElementIndex, InVelocity, Cd, Cl, Cp);
+		Forces[ElementIndex] = CalculateForce(InParticles, ElementIndex, InVelocity, CdI, CdO, ClI, ClO, Cp);
 	}
 
 	template<typename SolverParticlesOrRange>
-	FSolverVec3 CalculateForce(const SolverParticlesOrRange& InParticles, int32 ElementIndex, const FSolverVec3& InVelocity, const FSolverReal Cd, const FSolverReal Cl, const FSolverReal Cp, const FSolverReal MaxVelocitySquared) const
+	FSolverVec3 CalculateForce(const SolverParticlesOrRange& InParticles, int32 ElementIndex, const FSolverVec3& InVelocity, const FSolverReal CdI, const FSolverReal CdO, const FSolverReal ClI, const FSolverReal ClO, const FSolverReal Cp, const FSolverReal MaxVelocitySquared) const
 	{
 		checkSlow(MaxVelocitySquared > (FSolverReal)0);
 
@@ -209,8 +255,8 @@ private:
 
 		// Calculate the normal and the area of the surface exposed to the flow
 		FSolverVec3 N = FSolverVec3::CrossProduct(
-			InParticles.GetX(Element[1]) - InParticles.GetX(Element[0]),
-			InParticles.GetX(Element[2]) - InParticles.GetX(Element[0]));
+			InParticles.GetX(Element[2]) - InParticles.GetX(Element[0]),
+			InParticles.GetX(Element[1]) - InParticles.GetX(Element[0]));
 		const FSolverReal DoubleArea = N.SafeNormalize();
 
 		// Calculate the direction and the relative velocity of the triangle to the flow
@@ -232,13 +278,13 @@ private:
 		const FSolverReal VSquare = FSolverVec3::DotProduct(V, V);
 
 		return QuarterRho * DoubleArea * (VDotN >= (FSolverReal)0. ?  // The flow can hit either side of the triangle, so the normal might need to be reversed
-			(Cd - Cl) * VDotN * V + Cl * VSquare * N :
-			(Cl - Cd) * VDotN * V - Cl * VSquare * N) - DoubleArea * (FSolverReal)0.5 * Cp * N; // N points in the opposite direction of the actual mesh normals
+			(CdI - ClI) * VDotN * V + ClI * VSquare * N :
+			(ClO - CdO) * VDotN * V - ClO * VSquare * N) + DoubleArea * (FSolverReal)0.5 * Cp * N;
 	}
 
-	void UpdateField(const FSolverParticles& InParticles, int32 ElementIndex, const FSolverVec3& InVelocity, const FSolverReal Cd, const FSolverReal Cl, const FSolverReal Cp, const FSolverReal MaxVelocitySquared)
+	void UpdateField(const FSolverParticles& InParticles, int32 ElementIndex, const FSolverVec3& InVelocity, const FSolverReal CdI, const FSolverReal CdO, const FSolverReal ClI, const FSolverReal ClO, const FSolverReal Cp, const FSolverReal MaxVelocitySquared)
 	{
-		Forces[ElementIndex] = CalculateForce(InParticles, ElementIndex, InVelocity, Cd, Cl, Cp, MaxVelocitySquared);
+		Forces[ElementIndex] = CalculateForce(InParticles, ElementIndex, InVelocity, CdI, CdO, ClI, ClO, Cp, MaxVelocitySquared);
 	}
 
 private:
@@ -250,7 +296,9 @@ private:
 	TArray<TVec3<int32>> ElementsLocal; // Local copy of the triangle mesh's elements. Kinematic faces have been removed, and may be reordered by coloring.
 	TArray<int32> ConstraintsPerColorStartIndex; // Constraints are ordered so each batch is contiguous. This is ColorNum + 1 length so it can be used as start and end.
 	FPBDFlatWeightMap Lift;
+	FPBDFlatWeightMap OuterLift;
 	FPBDFlatWeightMap Drag;
+	FPBDFlatWeightMap OuterDrag;
 	FPBDFlatWeightMap Pressure;
 
 	TArray<FSolverVec3> Forces;
@@ -259,10 +307,13 @@ private:
 	FSolverReal QuarterRho;
 
 	UE_CHAOS_DECLARE_PROPERTYCOLLECTION_NAME(Drag, float);
+	UE_CHAOS_DECLARE_PROPERTYCOLLECTION_NAME(OuterDrag, float);
 	UE_CHAOS_DECLARE_PROPERTYCOLLECTION_NAME(Lift, float);
+	UE_CHAOS_DECLARE_PROPERTYCOLLECTION_NAME(OuterLift, float);
 	UE_CHAOS_DECLARE_PROPERTYCOLLECTION_NAME(FluidDensity, float);
 	UE_CHAOS_DECLARE_PROPERTYCOLLECTION_NAME(Pressure, float);
 	UE_CHAOS_DECLARE_PROPERTYCOLLECTION_NAME(WindVelocity, FVector3f);
+	UE_CHAOS_DECLARE_PROPERTYCOLLECTION_NAME(WindVelocitySpace, int32);
 };
 
 }  // End namespace Chaos::Softs
