Moved ownership of filehandles from callback to internal chunk request to prevent handles being destroyed before callback completes

#rnx

diff --git a/Engine/Source/Runtime/Experimental/IoStore/OnDemand/Private/OnDemandInstallCache.cpp b/Engine/Source/Runtime/Experimental/IoStore/OnDemand/Private/OnDemandInstallCache.cpp
--- a/Engine/Source/Runtime/Experimental/IoStore/OnDemand/Private/OnDemandInstallCache.cpp
+++ b/Engine/Source/Runtime/Experimental/IoStore/OnDemand/Private/OnDemandInstallCache.cpp
@@ -799,11 +799,13 @@ class FOnDemandInstallCache final
 	struct FChunkRequest
 	{
 		explicit FChunkRequest(
+			FSharedAsyncFileHandle FileHandle,
 			FIoRequestImpl* Request,
 			FOnDemandChunkInfo&& Info,
 			FIoOffsetAndLength Range,
 			uint64 RequestedRawSize)
-				: DispatcherRequest(Request)
+				: SharedFileHandle(FileHandle)
+				, DispatcherRequest(Request)
 				, ChunkInfo(MoveTemp(Info))
 				, ChunkRange(Range)
 				, EncodedChunk(ChunkRange.GetLength())
@@ -813,15 +815,41 @@ class FOnDemandInstallCache final
 			check(ChunkInfo.IsValid());
 			check(Request->NextRequest == nullptr);
 			check(Request->BackendData == nullptr);
+		}
+
+		static FChunkRequest* Get(FIoRequestImpl& Request)
+		{
+			return reinterpret_cast<FChunkRequest*>(Request.BackendData);
+		}
 
-			DispatcherRequest->BackendData = this;
+		static FChunkRequest& GetRef(FIoRequestImpl& Request)
+		{
+			check(Request.BackendData);
+			return *reinterpret_cast<FChunkRequest*>(Request.BackendData);
+		}
+
+		static FChunkRequest& Attach(FIoRequestImpl& Request, FChunkRequest* ChunkRequest)
+		{
+			check(Request.BackendData == nullptr);
+			check(ChunkRequest != nullptr);
+			Request.BackendData = ChunkRequest;
+			return *ChunkRequest;
 		}
 
-		FIoRequestImpl*		DispatcherRequest;
-		FOnDemandChunkInfo	ChunkInfo;
-		FIoOffsetAndLength	ChunkRange;
-		FIoBuffer			EncodedChunk;
-		uint64				RawSize;
+		static TUniquePtr<FChunkRequest> Detach(FIoRequestImpl& Request)
+		{
+			void* ChunkRequest = nullptr;
+			Swap(ChunkRequest, Request.BackendData);
+			return TUniquePtr<FChunkRequest>(reinterpret_cast<FChunkRequest*>(ChunkRequest));
+		}
+
+		FSharedAsyncFileHandle			SharedFileHandle;
+		TUniquePtr<IAsyncReadRequest>	FileReadRequest;
+		FIoRequestImpl*					DispatcherRequest;
+		FOnDemandChunkInfo				ChunkInfo;
+		FIoOffsetAndLength				ChunkRange;
+		FIoBuffer						EncodedChunk;
+		uint64							RawSize;
 	};
 
 	struct FPendingChunks
@@ -889,7 +917,7 @@ public:
 
 private:
 	bool						Resolve(FIoRequestImpl* Request);
-	void						CompleteRequest(FChunkRequest& ChunkRequest);
+	void						CompleteRequest(FIoRequestImpl* Request, bool bFileReadWasCancelled);
 	FIoStatus					FlushPendingChunks(FPendingChunks& Block);
 	FString						GetJournalFilename() const { return CacheDirectory / TEXT("cas.jrn"); }
 
@@ -1050,6 +1078,10 @@ FIoRequestImpl* FOnDemandInstallCache::GetCompletedIoRequests()
 	FIoRequestImpl* FirstCompleted = nullptr;
 	{
 		UE::TUniqueLock Lock(Mutex);
+		for (FIoRequestImpl& Completed : CompletedRequests)
+		{
+			TUniquePtr<FChunkRequest> Detached = FChunkRequest::Detach(Completed);
+		}
 		FirstCompleted = CompletedRequests.GetHead();
 		CompletedRequests = FIoRequestList();
 	}
@@ -1059,6 +1091,15 @@ FIoRequestImpl* FOnDemandInstallCache::GetCompletedIoRequests()
 
 void FOnDemandInstallCache::CancelIoRequest(FIoRequestImpl* Request)
 {
+	check(Request != nullptr);
+	UE::TUniqueLock Lock(Mutex);
+	if (FChunkRequest* ChunkRequest = FChunkRequest::Get(*Request))
+	{
+		if (ChunkRequest->FileReadRequest.IsValid())
+		{
+			ChunkRequest->FileReadRequest->Cancel();
+		}
+	}
 }
 
 void FOnDemandInstallCache::UpdatePriorityForIoRequest(FIoRequestImpl* Request)
@@ -1123,36 +1164,38 @@ bool FOnDemandInstallCache::Resolve(FIoRequestImpl* Request)
 		UE_LOG(LogIoStoreOnDemand, Error, TEXT("Failed to CAS block file for reading"));
 	}
 
-	TSharedPtr<FChunkRequest> ChunkRequest = MakeShared<FChunkRequest>(
+	// The internal request parameters are attached/owned by the I/O request via
+	// the backend data parameter. The chunk request is deleted in GetCompletedRequests
+	FChunkRequest& ChunkRequest = FChunkRequest::Attach(*Request, new FChunkRequest(
+		FileHandle,
 		Request,
 		MoveTemp(ChunkInfo),
 		ChunkRange.ConsumeValueOrDie(),
-		RequestSize);
+		RequestSize));
 
-	FIoBuffer EncodedChunk		= ChunkRequest->EncodedChunk;
-	FAsyncFileCallBack Callback = [this, ChunkRequest, FileHandle](bool bWasCanceled, IAsyncReadRequest* ReadRequest) 
+	FAsyncFileCallBack Callback = [this, Request](bool bWasCancelled, IAsyncReadRequest* ReadRequest) 
 	{
-		UE::Tasks::Launch(
-			UE_SOURCE_LOCATION,
-			[this, ChunkRequest, FileHandle, ReadRequest = TUniquePtr<IAsyncReadRequest>(ReadRequest), bWasCanceled]
-			{
-				if (bWasCanceled)
-				{
-					ChunkRequest->EncodedChunk = FIoBuffer();
-				}
-				CompleteRequest(*ChunkRequest);
-			});
+		UE::Tasks::Launch(UE_SOURCE_LOCATION, [this, Request, bWasCancelled]
+		{
+			CompleteRequest(Request, bWasCancelled);
+		});
 	};
 
 	Cas.TrackAccess(CasLoc.BlockId);
-	IAsyncReadRequest* ReadRequest = FileHandle->ReadRequest(
-		CasLoc.BlockOffset + ChunkRequest->ChunkRange.GetOffset(),
-		ChunkRequest->ChunkRange.GetLength(),
+	ChunkRequest.FileReadRequest.Reset(FileHandle->ReadRequest(
+		CasLoc.BlockOffset + ChunkRequest.ChunkRange.GetOffset(),
+		ChunkRequest.ChunkRange.GetLength(),
 		EAsyncIOPriorityAndFlags::AIOP_BelowNormal,
 		&Callback,
-		EncodedChunk.GetData());
+		ChunkRequest.EncodedChunk.GetData()));
 
-	return ReadRequest != nullptr;
+	if (ChunkRequest.FileReadRequest.IsValid() == false)
+	{
+		TUniquePtr<FChunkRequest> Detached = FChunkRequest::Detach(*Request);
+		return false;
+	}
+
+	return true;
 }
 
 bool FOnDemandInstallCache::IsChunkCached(const FIoHash& ChunkHash)
@@ -1399,12 +1442,12 @@ FIoStatus FOnDemandInstallCache::FlushPendingChunks(FPendingChunks& Chunks)
 	return FIoStatus::Ok;
 }
 
-void FOnDemandInstallCache::CompleteRequest(FChunkRequest& ChunkRequest)
+void FOnDemandInstallCache::CompleteRequest(FIoRequestImpl* Request, bool bFileReadWasCancelled)
 {
-	FIoRequestImpl* Request				= ChunkRequest.DispatcherRequest;
-	const FOnDemandChunkInfo& ChunkInfo = ChunkRequest.ChunkInfo;
-	FMemoryView EncodedChunk			= ChunkRequest.EncodedChunk.GetView();
-	bool bSucceeded						= EncodedChunk.IsEmpty() == false;
+	FChunkRequest& ChunkRequest			= FChunkRequest::GetRef(*Request);
+	const FOnDemandChunkInfo& ChunkInfo	= ChunkRequest.ChunkInfo;
+	FIoBuffer EncodedChunk				= MoveTemp(ChunkRequest.EncodedChunk);
+	bool bSucceeded						= EncodedChunk.GetSize() > 0 && !bFileReadWasCancelled && !Request->IsCancelled();
 
 	if (bSucceeded)
 	{
@@ -1421,7 +1464,7 @@ void FOnDemandInstallCache::CompleteRequest(FChunkRequest& ChunkRequest)
 		Request->CreateBuffer(ChunkRequest.RawSize);
 		FMutableMemoryView RawChunk = Request->GetBuffer().GetMutableView();
 
-		bSucceeded = FIoChunkEncoding::Decode(Params, EncodedChunk, RawChunk);
+		bSucceeded = FIoChunkEncoding::Decode(Params, EncodedChunk.GetView(), RawChunk);
 		UE_CLOG(!bSucceeded, LogIoStoreOnDemand, Error, TEXT("Failed to decode chunk, ChunkId='%s'"), *LexToString(Request->ChunkId));
 	}
 
@@ -1433,7 +1476,6 @@ void FOnDemandInstallCache::CompleteRequest(FChunkRequest& ChunkRequest)
 
 	{
 		UE::TUniqueLock Lock(Mutex);
-		Request->BackendData = nullptr;
 		CompletedRequests.AddTail(Request);
 	}
 
