Iris
- Refactored NetTokens a bit to allow export of data associated with NetTokens during serialization to be able to support exporting arbitrary data.
- Now passes full context when quantizing default state for future support of assymetric quantization.
- Support capturing NetToken exports when preserializing NetBlobs/RPCs

diff --git a/Engine/Plugins/Runtime/ReplicationSystemTestPlugin/Source/Private/Tests/ReplicationSystem/NetBlob/TestPartialNetBlob.cpp b/Engine/Plugins/Runtime/ReplicationSystemTestPlugin/Source/Private/Tests/ReplicationSystem/NetBlob/TestPartialNetBlob.cpp
--- a/Engine/Plugins/Runtime/ReplicationSystemTestPlugin/Source/Private/Tests/ReplicationSystem/NetBlob/TestPartialNetBlob.cpp
+++ b/Engine/Plugins/Runtime/ReplicationSystemTestPlugin/Source/Private/Tests/ReplicationSystem/NetBlob/TestPartialNetBlob.cpp
@@ -15,7 +15,7 @@ UE_NET_TEST_FIXTURE(FPartialNetBlobTestFixture, CanSplitNetBlobNotInNeedOfSplitt
 	const TRefCountPtr<FNetBlob>& Blob = CreateUnreliableMockNetBlob(PayloadBitCount);
 
 	TArray<TRefCountPtr<FNetBlob>> PartialNetBlobs;
-	MockSequentialPartialNetBlobHandler->SplitNetBlob(Blob, PartialNetBlobs);
+	MockSequentialPartialNetBlobHandler->SplitNetBlob(ServerContext.SerializationContext, Blob, PartialNetBlobs);
 	UE_NET_ASSERT_EQ(PartialNetBlobs.Num(), 1);
 
 	for (const TRefCountPtr<FNetBlob>& PartialNetBlob : PartialNetBlobs)
@@ -39,7 +39,7 @@ UE_NET_TEST_FIXTURE(FPartialNetBlobTestFixture, CanSplitNetBlobInNeedOfSplitting
 	const TRefCountPtr<FNetBlob>& Blob = CreateUnreliableMockNetBlob(PayloadBitCount);
 
 	TArray<TRefCountPtr<FNetBlob>> PartialNetBlobs;
-	MockSequentialPartialNetBlobHandler->SplitNetBlob(Blob, PartialNetBlobs);
+	MockSequentialPartialNetBlobHandler->SplitNetBlob(ServerContext.SerializationContext, Blob, PartialNetBlobs);
 	UE_NET_ASSERT_GE(PartialNetBlobs.Num(), 2);
 
 	for (const TRefCountPtr<FNetBlob>& PartialNetBlob : PartialNetBlobs)
@@ -62,7 +62,7 @@ UE_NET_TEST_FIXTURE(FPartialNetBlobTestFixture, CanSplitHugeNetBlob)
 	const TRefCountPtr<FNetBlob>& Blob = CreateUnreliableMockNetBlob(PayloadBitCount);
 
 	TArray<TRefCountPtr<FNetBlob>> PartialNetBlobs;
-	MockSequentialPartialNetBlobHandler->SplitNetBlob(Blob, PartialNetBlobs);
+	MockSequentialPartialNetBlobHandler->SplitNetBlob(ServerContext.SerializationContext, Blob, PartialNetBlobs);
 	UE_NET_ASSERT_GE(PartialNetBlobs.Num(), 2);
 
 	for (const TRefCountPtr<FNetBlob>& PartialNetBlob : PartialNetBlobs)
@@ -85,7 +85,7 @@ UE_NET_TEST_FIXTURE(FPartialNetBlobTestFixture, MissingFirstPartialNetBlobCauses
 	const TRefCountPtr<FNetBlob>& Blob = CreateReliableMockNetBlob(PayloadBitCount);
 
 	TArray<TRefCountPtr<FNetBlob>> PartialNetBlobs;
-	MockSequentialPartialNetBlobHandler->SplitNetBlob(Blob, PartialNetBlobs);
+	MockSequentialPartialNetBlobHandler->SplitNetBlob(ServerContext.SerializationContext, Blob, PartialNetBlobs);
 	UE_NET_ASSERT_GE(PartialNetBlobs.Num(), 2);
 
 	ServerContext.SerializationContext.GetNetBlobReceiver()->OnNetBlobReceived(ServerContext.SerializationContext, PartialNetBlobs[1]);
@@ -101,7 +101,7 @@ UE_NET_TEST_FIXTURE(FPartialNetBlobTestFixture, MissingArbitraryPartialNetBlobCa
 	const TRefCountPtr<FNetBlob>& Blob = CreateReliableMockNetBlob(PayloadBitCount);
 
 	TArray<TRefCountPtr<FNetBlob>> PartialNetBlobs;
-	MockSequentialPartialNetBlobHandler->SplitNetBlob(Blob, PartialNetBlobs);
+	MockSequentialPartialNetBlobHandler->SplitNetBlob(ServerContext.SerializationContext, Blob, PartialNetBlobs);
 	UE_NET_ASSERT_GE(PartialNetBlobs.Num(), 3);
 
 	ServerContext.SerializationContext.GetNetBlobReceiver()->OnNetBlobReceived(ServerContext.SerializationContext, PartialNetBlobs[0]);
@@ -215,12 +215,12 @@ UE_NET_TEST_FIXTURE(FPartialNetBlobTestFixture, UnexpectedPartSequenceCausesErro
 		constexpr int32 PartCount = 3U;
 		uint32 TotalPayloadBitCount = PartCount * Config->GetMaxPartBitCount() - BlobOverheadBitCount - 1U;
 		const TRefCountPtr<FNetBlob>& Blob = CreateReliableMockNetBlob(TotalPayloadBitCount);
-		MockSequentialPartialNetBlobHandler->SplitNetBlob(Blob, PartialNetBlobs);
+		MockSequentialPartialNetBlobHandler->SplitNetBlob(ServerContext.SerializationContext, Blob, PartialNetBlobs);
 		UE_NET_ASSERT_EQ(PartialNetBlobs.Num(), PartCount);
 
 		TArray<TRefCountPtr<FNetBlob>> OtherPartialNetBlobs;
 		const TRefCountPtr<FNetBlob>& OtherBlob = CreateReliableMockNetBlob(TotalPayloadBitCount);
-		MockSequentialPartialNetBlobHandler->SplitNetBlob(OtherBlob, OtherPartialNetBlobs);
+		MockSequentialPartialNetBlobHandler->SplitNetBlob(ServerContext.SerializationContext, OtherBlob, OtherPartialNetBlobs);
 	
 		WrongSecondBlob = MoveTemp(OtherPartialNetBlobs[1]);
 	}
diff --git a/Engine/Plugins/Runtime/ReplicationSystemTestPlugin/Source/Private/Tests/ReplicationSystem/RPC/RPCTestFixture.cpp b/Engine/Plugins/Runtime/ReplicationSystemTestPlugin/Source/Private/Tests/ReplicationSystem/RPC/RPCTestFixture.cpp
--- a/Engine/Plugins/Runtime/ReplicationSystemTestPlugin/Source/Private/Tests/ReplicationSystem/RPC/RPCTestFixture.cpp
+++ b/Engine/Plugins/Runtime/ReplicationSystemTestPlugin/Source/Private/Tests/ReplicationSystem/RPC/RPCTestFixture.cpp
@@ -11,6 +11,7 @@ void FRPCTestFixture::SetUp()
 	OriginalHandlerDefinitions = BlobHandlerDefinitions->ReadWriteHandlerDefinitions();
 
 	HandlerDefinitions.Add({ TEXT("NetRPCHandler") });
+	HandlerDefinitions.Add({ TEXT("PartialNetObjectAttachmentHandler") });
 
 	BlobHandlerDefinitions->ReadWriteHandlerDefinitions() = HandlerDefinitions;
 
diff --git a/Engine/Source/Runtime/Experimental/Iris/Core/Private/Iris/ReplicationState/ReplicationStateDescriptorBuilder.cpp b/Engine/Source/Runtime/Experimental/Iris/Core/Private/Iris/ReplicationState/ReplicationStateDescriptorBuilder.cpp
--- a/Engine/Source/Runtime/Experimental/Iris/Core/Private/Iris/ReplicationState/ReplicationStateDescriptorBuilder.cpp
+++ b/Engine/Source/Runtime/Experimental/Iris/Core/Private/Iris/ReplicationState/ReplicationStateDescriptorBuilder.cpp
@@ -208,6 +208,8 @@ public:
 		const FStringBuilderBase* PathName = nullptr;
 		// Optionally pass in pointer to source data buffer from which we should create the default state
 		const uint8* DefaultStateSourceData = nullptr;
+		// When building default state we need the ReplicationSystem in order to capture exports.
+		UReplicationSystem* ReplicationSystem = nullptr;
 		EDescriptorType DescriptorType;
 		uint32 bIsInitState : 1;
 		uint32 bAllMembersAreReplicated : 1;
@@ -347,8 +349,8 @@ private:
 	void BuildMemberSerializerConfigs(uint8* SerializerConfigBuffer, FReplicationStateDescriptor* Descriptor, const FMemberCache& MemberCache) const;
 	void FixupNetRoleNetSerializerConfigs(FReplicationStateDescriptor* Descriptor);
 	void FixupDescriptorForNativeFastArray(FReplicationStateDescriptor* Descriptor) const;
-	void AllocateAndInitializeDefaultInternalStateBuffer(const uint8* RESTRICT SrcBuffer, const FReplicationStateDescriptor* Descriptor, const uint8*& OutDefaultStateBuffer) const;
-	bool CalculateDefaultStateChecksum(const FReplicationStateDescriptor* Descriptor, const uint8* OutDefaultStateBuffer, uint64& OutHashValue) const;
+	void AllocateAndInitializeDefaultInternalStateBuffer(const uint8* RESTRICT SrcBuffer, const FReplicationStateDescriptor* Descriptor, UReplicationSystem* ReplicationSystem, const uint8*& OutDefaultStateBuffer) const;
+	bool CalculateDefaultStateChecksum(const FReplicationStateDescriptor* Descriptor, UReplicationSystem* ReplicationSystem, const uint8* OutDefaultStateBuffer, uint64& OutHashValue) const;
 
 	EReplicationStateTraits BuildReplicationStateTraits(const FBuilderContext& Context) const;
 
@@ -756,7 +758,7 @@ void FPropertyReplicationStateDescriptorBuilder::BuildMemberFunctionCache(FBuild
 	}
 }
 
-void FPropertyReplicationStateDescriptorBuilder::AllocateAndInitializeDefaultInternalStateBuffer(const uint8* RESTRICT SrcBuffer, const FReplicationStateDescriptor* Descriptor, const uint8*& OutDefaultStateBuffer) const
+void FPropertyReplicationStateDescriptorBuilder::AllocateAndInitializeDefaultInternalStateBuffer(const uint8* RESTRICT SrcBuffer, const FReplicationStateDescriptor* Descriptor, UReplicationSystem* ReplicationSystem, const uint8*& OutDefaultStateBuffer) const
 {
 	// Allocate storage for incoming data, it will be freed when we destroy the descriptor
 	uint8* DstStateBuffer = static_cast<uint8*>(FMemory::MallocZeroed(Descriptor->InternalSize, Descriptor->InternalAlignment));
@@ -765,7 +767,14 @@ void FPropertyReplicationStateDescriptorBuilder::AllocateAndInitializeDefaultInt
 	// on both server and client
 	FNetSerializationContext Context;
 	FInternalNetSerializationContext InternalContext;
-	InternalContext.PackageMap = GetOrCreateIrisObjectReferencePackageMap();
+	if (ReplicationSystem)
+	{
+		InternalContext = FInternalNetSerializationContext(ReplicationSystem);
+	}
+	else
+	{
+		InternalContext.PackageMap = GetOrCreateIrisObjectReferencePackageMap();
+	}
 	Context.SetInternalContext(&InternalContext);
 	Context.SetIsInitializingDefaultState(true);
 
@@ -794,7 +803,7 @@ void FPropertyReplicationStateDescriptorBuilder::AllocateAndInitializeDefaultInt
 	OutDefaultStateBuffer = DstStateBuffer;
 }
 
-bool FPropertyReplicationStateDescriptorBuilder::CalculateDefaultStateChecksum(const FReplicationStateDescriptor* Descriptor, const uint8* OutDefaultStateBuffer, uint64& OutHashValue) const
+bool FPropertyReplicationStateDescriptorBuilder::CalculateDefaultStateChecksum(const FReplicationStateDescriptor* Descriptor, UReplicationSystem* ReplicationSystem, const uint8* OutDefaultStateBuffer, uint64& OutHashValue) const
 {
 	const UPartialNetObjectAttachmentHandlerConfig* PartialNetObjectAttachmentHandlerConfig = GetDefault<UPartialNetObjectAttachmentHandlerConfig>();
 
@@ -810,7 +819,14 @@ bool FPropertyReplicationStateDescriptorBuilder::CalculateDefaultStateChecksum(c
 	FNetBitStreamWriter Writer;
 	FNetSerializationContext Context(&Writer);
 	FInternalNetSerializationContext InternalContext;
-	InternalContext.PackageMap = GetOrCreateIrisObjectReferencePackageMap();
+	if (ReplicationSystem)
+	{
+		InternalContext = FInternalNetSerializationContext(ReplicationSystem);
+	}
+	else
+	{
+		InternalContext.PackageMap = GetOrCreateIrisObjectReferencePackageMap();
+	}
 	Context.SetInternalContext(&InternalContext);
 
 	// Tell serializers we are serializing default state. It allows serializers to opt out of being part of the checksum by simply not serializing any data.
@@ -1804,10 +1820,10 @@ FPropertyReplicationStateDescriptorBuilder::Build(const FString& StateName, FRep
 		if (Descriptor->InternalSize && BuildParams.DefaultStateSourceData)
 		{
 			// Initialize default state, we need to do this after finalizing the descriptor as we use the descriptor to iterate over the members
-			AllocateAndInitializeDefaultInternalStateBuffer(BuildParams.DefaultStateSourceData, Descriptor, Descriptor->DefaultStateBuffer);
+			AllocateAndInitializeDefaultInternalStateBuffer(BuildParams.DefaultStateSourceData, Descriptor, BuildParams.ReplicationSystem, Descriptor->DefaultStateBuffer);
 
 			uint64 DefaultStateHash = uint64(0);
-			if (CalculateDefaultStateChecksum(Descriptor, Descriptor->DefaultStateBuffer, DefaultStateHash))
+			if (CalculateDefaultStateChecksum(Descriptor, BuildParams.ReplicationSystem, Descriptor->DefaultStateBuffer, DefaultStateHash))
 			{
 				// Store the default state checksum for later use when verifying protocol
 				Descriptor->DescriptorIdentifier.DefaultStateHash = DefaultStateHash;
@@ -2309,6 +2325,7 @@ const IConsoleVariable* FReplicationStateDescriptorBuilder::CVarReplicateCustomD
 
 FReplicationStateDescriptorBuilder::FParameters::FParameters()
 : DescriptorRegistry(nullptr),
+  ReplicationSystem(nullptr),
   DefaultStateSource(nullptr),
   IncludeSuper(1U),
   GetLifeTimeProperties(1U),
@@ -2781,6 +2798,7 @@ SIZE_T FReplicationStateDescriptorBuilder::CreateDescriptorsForClass(FResult& Cr
 			FPropertyReplicationStateDescriptorBuilder::FBuildParameters BuildParameters = {};
 			BuildParameters.PathName = LazyPathName.GetPathName(InObjectClass);
 			BuildParameters.DefaultStateSourceData = DefaultStateSourceData;
+			BuildParameters.ReplicationSystem = Parameters.ReplicationSystem;
 			BuildParameters.DescriptorType = FPropertyReplicationStateDescriptorBuilder::EDescriptorType::Class;
 			BuildParameters.bIsInitState = (BuilderTypeIndex == InitPropertyReplicationStateBuilderIndex);
 			// We set this to false for the time being for all replication states
@@ -2824,6 +2842,7 @@ SIZE_T FReplicationStateDescriptorBuilder::CreateDescriptorsForClass(FResult& Cr
 		FPropertyReplicationStateDescriptorBuilder::FBuildParameters BuildParameters = {};
 		BuildParameters.PathName = LazyPathName.GetPathName(InObjectClass);
 		BuildParameters.DefaultStateSourceData = DefaultStateSourceData;
+		BuildParameters.ReplicationSystem = Parameters.ReplicationSystem;
 		BuildParameters.DescriptorType = FPropertyReplicationStateDescriptorBuilder::EDescriptorType::Class;
 		BuildParameters.bIsInitState = EnumHasAnyFlags(MemberProperty.Traits, EMemberPropertyTraits::InitOnly);
 		
diff --git a/Engine/Source/Runtime/Experimental/Iris/Core/Private/Iris/ReplicationSystem/AttachmentReplication.cpp b/Engine/Source/Runtime/Experimental/Iris/Core/Private/Iris/ReplicationSystem/AttachmentReplication.cpp
--- a/Engine/Source/Runtime/Experimental/Iris/Core/Private/Iris/ReplicationSystem/AttachmentReplication.cpp
+++ b/Engine/Source/Runtime/Experimental/Iris/Core/Private/Iris/ReplicationSystem/AttachmentReplication.cpp
@@ -383,7 +383,12 @@ uint32 FNetObjectAttachmentSendQueue::SerializeUnreliable(FNetSerializationConte
 		const TRefCountPtr<FNetBlob>& Attachment = UnreliableQueue.PeekAtOffsetNoCheck(AttachmentIt);
 
 		// If we have exports, append them, if attachment is rolled back we will roll back any appended exports as well.
-		ObjectReferenceCache->AddPendingExports(Context, Attachment->CallGetExports());
+		FNetExportContext* ExportContext = Attachment->HasExports() ? Context.GetExportContext() : nullptr;
+		if (ExportContext)
+		{
+			ObjectReferenceCache->AddPendingExports(Context, Attachment->CallGetNetObjectReferenceExports());
+			ExportContext->AddPendingExports(Attachment->CallGetNetTokenExports());
+		}
 
 		Attachment->SerializeCreationInfo(Context, Attachment->GetCreationInfo());
 		if (bSerializeWithObject)
diff --git a/Engine/Source/Runtime/Experimental/Iris/Core/Private/Iris/ReplicationSystem/NetBlob/NetBlob.cpp b/Engine/Source/Runtime/Experimental/Iris/Core/Private/Iris/ReplicationSystem/NetBlob/NetBlob.cpp
--- a/Engine/Source/Runtime/Experimental/Iris/Core/Private/Iris/ReplicationSystem/NetBlob/NetBlob.cpp
+++ b/Engine/Source/Runtime/Experimental/Iris/Core/Private/Iris/ReplicationSystem/NetBlob/NetBlob.cpp
@@ -126,11 +126,16 @@ void FNetBlob::Release() const
 	}
 }
 
-TArrayView<const FNetObjectReference> FNetBlob::GetExports() const
+TArrayView<const FNetObjectReference> FNetBlob::GetNetObjectReferenceExports() const
 {
 	return MakeArrayView<const FNetObjectReference>(nullptr, 0);
 };
 
+TArrayView<const FNetToken> FNetBlob::GetNetTokenExports() const
+{
+	return MakeArrayView<const FNetToken>(nullptr, 0);
+};
+
 // NetObjectAttachment
 FNetObjectAttachment::FNetObjectAttachment(const FNetBlobCreationInfo& CreationInfo)
 : FNetBlob(CreationInfo)
diff --git a/Engine/Source/Runtime/Experimental/Iris/Core/Private/Iris/ReplicationSystem/NetBlob/NetRPC.cpp b/Engine/Source/Runtime/Experimental/Iris/Core/Private/Iris/ReplicationSystem/NetBlob/NetRPC.cpp
--- a/Engine/Source/Runtime/Experimental/Iris/Core/Private/Iris/ReplicationSystem/NetBlob/NetRPC.cpp
+++ b/Engine/Source/Runtime/Experimental/Iris/Core/Private/Iris/ReplicationSystem/NetBlob/NetRPC.cpp
@@ -57,7 +57,7 @@ FNetRPC::~FNetRPC()
 {
 }
 
-TArrayView<const FNetObjectReference> FNetRPC::GetExports() const
+TArrayView<const FNetObjectReference> FNetRPC::GetNetObjectReferenceExports() const
 {
 	return ReferencesToExport.IsValid() ? MakeArrayView(*ReferencesToExport) : MakeArrayView<const FNetObjectReference>(nullptr, 0);
 }
diff --git a/Engine/Source/Runtime/Experimental/Iris/Core/Private/Iris/ReplicationSystem/NetBlob/NetRPC.h b/Engine/Source/Runtime/Experimental/Iris/Core/Private/Iris/ReplicationSystem/NetBlob/NetRPC.h
--- a/Engine/Source/Runtime/Experimental/Iris/Core/Private/Iris/ReplicationSystem/NetBlob/NetRPC.h
+++ b/Engine/Source/Runtime/Experimental/Iris/Core/Private/Iris/ReplicationSystem/NetBlob/NetRPC.h
@@ -35,7 +35,7 @@ public:
 private:
 	virtual ~FNetRPC();
 
-	virtual TArrayView<const FNetObjectReference> GetExports() const override final;
+	virtual TArrayView<const FNetObjectReference> GetNetObjectReferenceExports() const override final;
 
 	virtual void SerializeWithObject(FNetSerializationContext& Context, FNetRefHandle RefHandle) const override;
 	virtual void DeserializeWithObject(FNetSerializationContext& Context, FNetRefHandle RefHandle) override;
diff --git a/Engine/Source/Runtime/Experimental/Iris/Core/Private/Iris/ReplicationSystem/NetBlob/PartialNetBlob.cpp b/Engine/Source/Runtime/Experimental/Iris/Core/Private/Iris/ReplicationSystem/NetBlob/PartialNetBlob.cpp
--- a/Engine/Source/Runtime/Experimental/Iris/Core/Private/Iris/ReplicationSystem/NetBlob/PartialNetBlob.cpp
+++ b/Engine/Source/Runtime/Experimental/Iris/Core/Private/Iris/ReplicationSystem/NetBlob/PartialNetBlob.cpp
@@ -4,6 +4,7 @@
 #include "Iris/Serialization/NetBitStreamReader.h"
 #include "Iris/Serialization/NetBitStreamWriter.h"
 #include "Iris/Serialization/NetBitStreamUtil.h"
+#include "Iris/Serialization/NetExportContext.h"
 #include "Iris/Serialization/InternalNetSerializationContext.h"
 #include "Iris/Serialization/NetSerializationContext.h"
 #include "Net/Core/Trace/NetTrace.h"
@@ -29,11 +30,11 @@ FPartialNetBlob::FPartialNetBlob(const FNetBlobCreationInfo& CreationInfo)
 {
 }
 
-TArrayView<const FNetObjectReference> FPartialNetBlob::GetExports() const
+TArrayView<const FNetObjectReference> FPartialNetBlob::GetNetObjectReferenceExports() const
 {
 	if (IsFirstPart() && OriginalBlob.IsValid())
 	{
-		return OriginalBlob->CallGetExports();
+		return OriginalBlob->CallGetNetObjectReferenceExports();
 	}
 	else
 	{
@@ -41,6 +42,11 @@ TArrayView<const FNetObjectReference> FPartialNetBlob::GetExports() const
 	}
 }
 
+TArrayView<const FNetToken> FPartialNetBlob::GetNetTokenExports() const
+{
+	return MakeConstArrayView<>(NetTokenExportsArray);
+}
+
 void FPartialNetBlob::SerializeWithObject(FNetSerializationContext& Context, FNetRefHandle RefHandle) const
 {
 	InternalSerialize(Context);
@@ -156,7 +162,7 @@ void FPartialNetBlob::InternalDeserializeBlob(FNetSerializationContext& Context)
 	Reader->ReadBitStream(Payload.GetData(), PayloadBitCount);
 }
 
-bool FPartialNetBlob::SplitNetBlob(const FNetSerializationContext& Context, const FNetBlobCreationInfo& CreationInfo, const FPartialNetBlob::FSplitParams& SplitParams, const TRefCountPtr<FNetBlob>& Blob, TArray<TRefCountPtr<FNetBlob>>& OutPartialBlobs)
+bool FPartialNetBlob::SplitNetBlob(FNetSerializationContext& Context, const FNetBlobCreationInfo& CreationInfo, const FPartialNetBlob::FSplitParams& SplitParams, const TRefCountPtr<FNetBlob>& Blob, TArray<TRefCountPtr<FNetBlob>>& OutPartialBlobs)
 {
 	check(SplitParams.MaxPartBitCount > 31U && SplitParams.MaxPartBitCount < 65536U && SplitParams.MaxPartCount > 0 && SplitParams.MaxPartCount < 65536U);
 	if (!Blob.IsValid())
@@ -186,6 +192,8 @@ bool FPartialNetBlob::SplitNetBlob(const FNetSerializationContext& Context, cons
 		Writer.InitBytes(Payload.GetData(), static_cast<uint32>(Payload.Num())*4U);
 
 		FNetSerializationContext SubContext = Context.MakeSubContext(&Writer);
+		Private::FNetExportRollbackScope ExportRollBack(SubContext);
+
 		if (SplitParams.bSerializeWithObject)
 		{
 			Blob->SerializeWithObject(SubContext, SplitParams.NetObjectReference.GetRefHandle());
@@ -232,6 +240,10 @@ bool FPartialNetBlob::SplitNetBlob(const FNetSerializationContext& Context, cons
 		PayloadSplitParams.Payload = Payload.GetData();
 		PayloadSplitParams.PayloadBitCount = CurrentPayloadBitCount;
 		PayloadSplitParams.PartBitCount = MaxPartBitCount;
+		if (Private::FNetExportContext* ExportContext = Context.GetExportContext())
+		{			
+			PayloadSplitParams.NetTokensPendingExport = MakeConstArrayView<>(ExportContext->GetBatchExports().NetTokensPendingExportInCurrentBatch);
+		}
 
 		SplitPayload(PayloadSplitParams, OutPartialBlobs);
 	}
@@ -247,7 +259,7 @@ bool FPartialNetBlob::SplitNetBlob(const FNetBlobCreationInfo& CreationInfo, con
 		return false;
 	}
 
-	FPayloadSplitParams PayloadSplitParams;
+	FPayloadSplitParams PayloadSplitParams = {};
 	PayloadSplitParams.DebugName = SplitParams.DebugName;
 	PayloadSplitParams.CreationInfo = CreationInfo;
 	PayloadSplitParams.OriginalBlob = Blob.GetReference();
@@ -269,6 +281,8 @@ void FPartialNetBlob::SplitPayload(const FPartialNetBlob::FPayloadSplitParams& S
 	uint32 SequenceNumber = Private::PartialNetBlobGlobalSequenceNumber.fetch_add(PartialBlobCount, std::memory_order_relaxed);
 
 	uint32 PayloadBitOffset = 0U;
+
+	// Do we have exports	
 	for (uint32 PartIt = 0, PartEndIt = PartialBlobCount; PartIt != PartEndIt; ++PartIt)
 	{
 		const bool bIsFirstPart = PartIt == 0U;
@@ -279,10 +293,11 @@ void FPartialNetBlob::SplitPayload(const FPartialNetBlob::FPayloadSplitParams& S
 		PartialBlob->PartCount = static_cast<uint16>(PartialBlobCount);
 		PartialBlob->SequenceFlags = (bIsFirstPart ? ESequenceFlags::IsFirstPart : ESequenceFlags::None);
 		PartialBlob->SequenceNumber = SequenceNumber++;
-		if (bIsFirstPart && EnumHasAnyFlags(SplitParams.OriginalCreationInfo.Flags, ENetBlobFlags::HasExports))
+		if (bIsFirstPart && (SplitParams.NetTokensPendingExport.Num() || EnumHasAnyFlags(SplitParams.OriginalCreationInfo.Flags, ENetBlobFlags::HasExports)))
 		{
 			PartialBlob->OriginalBlob = SplitParams.OriginalBlob;
 			PartialBlob->CreationInfo.Flags |= ENetBlobFlags::HasExports;
+			PartialBlob->NetTokenExportsArray.Append(SplitParams.NetTokensPendingExport);
 		}
 
 		// Copy relevant data from our temporary buffer.
diff --git a/Engine/Source/Runtime/Experimental/Iris/Core/Private/Iris/ReplicationSystem/NetBlob/PartialNetObjectAttachmentHandler.cpp b/Engine/Source/Runtime/Experimental/Iris/Core/Private/Iris/ReplicationSystem/NetBlob/PartialNetObjectAttachmentHandler.cpp
--- a/Engine/Source/Runtime/Experimental/Iris/Core/Private/Iris/ReplicationSystem/NetBlob/PartialNetObjectAttachmentHandler.cpp
+++ b/Engine/Source/Runtime/Experimental/Iris/Core/Private/Iris/ReplicationSystem/NetBlob/PartialNetObjectAttachmentHandler.cpp
@@ -7,6 +7,7 @@
 #include "Iris/Serialization/NetSerializationContext.h"
 #include "Iris/Serialization/InternalNetSerializationContext.h"
 #include "Iris/Serialization/NetBitStreamWriter.h"
+#include "Iris/Serialization/NetExportContext.h"
 
 UPartialNetObjectAttachmentHandler::UPartialNetObjectAttachmentHandler()
 : USequentialPartialNetBlobHandler()
@@ -27,7 +28,7 @@ void UPartialNetObjectAttachmentHandler::Init(const FPartialNetObjectAttachmentH
 	Super::Init(ParentInitParams);
 }
 
-bool UPartialNetObjectAttachmentHandler::PreSerializeAndSplitNetBlob(uint32 ConnectionId, const TRefCountPtr<UE::Net::FNetObjectAttachment>& Blob, TArray<TRefCountPtr<FNetBlob>>& OutPartialBlobs, bool bSerializeWithObject) const
+bool UPartialNetObjectAttachmentHandler::PreSerializeAndSplitNetBlob(uint32 ConnectionId, const TRefCountPtr<UE::Net::FNetObjectAttachment>& Blob, TArray<TRefCountPtr<FNetBlob>>& OutPartialBlobs, bool bSerializeWithObject)
 {
 	using namespace UE::Net;
 	
@@ -52,19 +53,27 @@ bool UPartialNetObjectAttachmentHandler::PreSerializeAndSplitNetBlob(uint32 Conn
 	SerializationContext.SetLocalConnectionId(ConnectionId);
 	SerializationContext.SetInternalContext(&InternalContext);
 
-	if (bSerializeWithObject)
-	{
-		Blob->SerializeWithObject(SerializationContext, Blob->GetNetObjectReference().GetRefHandle());
-	}
-	else
-	{
-		Blob->Serialize(SerializationContext);
-	}
+	// Setup ExportScope to capture non object reference exports.
+	Private::FNetExportContext::FBatchExports BatchExports;
+	Private::FNetExports::FExportScope ExportScope = NetExports.MakeExportScope(SerializationContext, BatchExports);
 
-	// Errors are bad. We cannot recover from this.
-	if (SerializationContext.HasError())
 	{
-		return false;
+		Private::FNetExportRollbackScope ExportRollBack(SerializationContext);
+
+		if (bSerializeWithObject)
+		{
+			Blob->SerializeWithObject(SerializationContext, Blob->GetNetObjectReference().GetRefHandle());
+		}
+		else
+		{
+			Blob->Serialize(SerializationContext);
+		}
+
+		// Errors are bad. We cannot recover from this.
+		if (SerializationContext.HasError())
+		{
+			return false;
+		}
 	}
 
 	Writer.CommitWrites();
@@ -75,16 +84,16 @@ bool UPartialNetObjectAttachmentHandler::PreSerializeAndSplitNetBlob(uint32 Conn
 		// This will redo all the serialization work, but it should rarely happen.
 		if (bSerializeWithObject)
 		{
-			return Super::SplitNetBlob(Blob->GetNetObjectReference(), reinterpret_cast<const TRefCountPtr<FNetBlob>&>(Blob), OutPartialBlobs);
+			return Super::SplitNetBlob(SerializationContext, Blob->GetNetObjectReference(), reinterpret_cast<const TRefCountPtr<FNetBlob>&>(Blob), OutPartialBlobs);
 		}
 		else
 		{
-			return Super::SplitNetBlob(reinterpret_cast<const TRefCountPtr<FNetBlob>&>(Blob), OutPartialBlobs);
+			return Super::SplitNetBlob(SerializationContext, reinterpret_cast<const TRefCountPtr<FNetBlob>&>(Blob), OutPartialBlobs);
 		}
 	}
 	else
 	{
-		FShrinkWrapNetObjectAttachment* ShrinkWrapNetBlob = new FShrinkWrapNetObjectAttachment(Blob, MoveTemp(Payload), Writer.GetPosBits());
+		FShrinkWrapNetObjectAttachment* ShrinkWrapNetBlob = new FShrinkWrapNetObjectAttachment(SerializationContext, Blob, MoveTemp(Payload), Writer.GetPosBits());
 		OutPartialBlobs.AddDefaulted_GetRef() = ShrinkWrapNetBlob;
 		return true;
 	}
diff --git a/Engine/Source/Runtime/Experimental/Iris/Core/Private/Iris/ReplicationSystem/NetBlob/PartialNetObjectAttachmentHandler.h b/Engine/Source/Runtime/Experimental/Iris/Core/Private/Iris/ReplicationSystem/NetBlob/PartialNetObjectAttachmentHandler.h
--- a/Engine/Source/Runtime/Experimental/Iris/Core/Private/Iris/ReplicationSystem/NetBlob/PartialNetObjectAttachmentHandler.h
+++ b/Engine/Source/Runtime/Experimental/Iris/Core/Private/Iris/ReplicationSystem/NetBlob/PartialNetObjectAttachmentHandler.h
@@ -3,6 +3,7 @@
 #pragma once
 #include "Iris/ReplicationSystem/NetBlob/SequentialPartialNetBlobHandler.h"
 #include "Iris/ReplicationSystem/NetBlob/RawDataNetBlob.h"
+#include "Iris/ReplicationSystem/NetExports.h"
 #include "PartialNetObjectAttachmentHandler.generated.h"
 
 UCLASS()
@@ -50,12 +51,14 @@ public:
 	void Init(const FPartialNetObjectAttachmentHandlerInitParams& InitParams);
 
 	/** Serializes the NetBlob and either store the serialized version in a new NetBlob or splits into multiple partial NetBlobs. */
-	bool PreSerializeAndSplitNetBlob(uint32 ConnectionId, const TRefCountPtr<UE::Net::FNetObjectAttachment>& Blob, TArray<TRefCountPtr<FNetBlob>>& OutPartialBlobs, bool bSerializeWithObject) const;
+	bool PreSerializeAndSplitNetBlob(uint32 ConnectionId, const TRefCountPtr<UE::Net::FNetObjectAttachment>& Blob, TArray<TRefCountPtr<FNetBlob>>& OutPartialBlobs, bool bSerializeWithObject);
 
 	/** Splits a RawDataNetBlob. The blob must have been created by a registered NetBlobHandler in order to be reconstructed on the receiving side. */
 	bool SplitRawDataNetBlob(const TRefCountPtr<UE::Net::FRawDataNetBlob>& Blob, TArray<TRefCountPtr<FNetBlob>>& OutPartialBlobs, const UE::Net::FNetDebugName* InDebugName) const;
 
 	const UPartialNetObjectAttachmentHandlerConfig* GetConfig() const;
+private:
+	UE::Net::Private::FNetExports NetExports;
 };
 
 inline const UPartialNetObjectAttachmentHandlerConfig* UPartialNetObjectAttachmentHandler::GetConfig() const
diff --git a/Engine/Source/Runtime/Experimental/Iris/Core/Private/Iris/ReplicationSystem/NetBlob/ReliableNetBlobQueue.cpp b/Engine/Source/Runtime/Experimental/Iris/Core/Private/Iris/ReplicationSystem/NetBlob/ReliableNetBlobQueue.cpp
--- a/Engine/Source/Runtime/Experimental/Iris/Core/Private/Iris/ReplicationSystem/NetBlob/ReliableNetBlobQueue.cpp
+++ b/Engine/Source/Runtime/Experimental/Iris/Core/Private/Iris/ReplicationSystem/NetBlob/ReliableNetBlobQueue.cpp
@@ -137,7 +137,12 @@ uint32 FReliableNetBlobQueue::SerializeInternal(FNetSerializationContext& Contex
 			UE_NET_TRACE_NAMED_DYNAMIC_NAME_SCOPE(ReliabilityScope, Private::ReliabilityNetDebugNames[Attachment->IsReliable()], *Writer, Context.GetTraceCollector(), ENetTraceVerbosity::Verbose);
 
 			// If we have exports, append them, if attachment is rolled back we will roll back any appended exports as well.
-			ObjectReferenceCache->AddPendingExports(Context, Attachment->CallGetExports());
+			FNetExportContext* ExportContext = Attachment->HasExports() ? Context.GetExportContext() : nullptr;
+			if (ExportContext)
+			{
+				ObjectReferenceCache->AddPendingExports(Context, Attachment->CallGetNetObjectReferenceExports());
+				ExportContext->AddPendingExports(Attachment->CallGetNetTokenExports());
+			}
 
 			Attachment->SerializeCreationInfo(Context, Attachment->GetCreationInfo());
 			if (bSerializeWithObject)
diff --git a/Engine/Source/Runtime/Experimental/Iris/Core/Private/Iris/ReplicationSystem/NetBlob/SequentialPartialNetBlobHandler.cpp b/Engine/Source/Runtime/Experimental/Iris/Core/Private/Iris/ReplicationSystem/NetBlob/SequentialPartialNetBlobHandler.cpp
--- a/Engine/Source/Runtime/Experimental/Iris/Core/Private/Iris/ReplicationSystem/NetBlob/SequentialPartialNetBlobHandler.cpp
+++ b/Engine/Source/Runtime/Experimental/Iris/Core/Private/Iris/ReplicationSystem/NetBlob/SequentialPartialNetBlobHandler.cpp
@@ -19,15 +19,10 @@ void USequentialPartialNetBlobHandler::Init(const FSequentialPartialNetBlobHandl
 	Config = InitParams.Config;
 }
 
-bool USequentialPartialNetBlobHandler::SplitNetBlob(const TRefCountPtr<FNetBlob>& Blob, TArray<TRefCountPtr<FNetBlob>>& OutPartialBlobs, const UE::Net::FNetDebugName* InDebugName) const
+bool USequentialPartialNetBlobHandler::SplitNetBlob(UE::Net::FNetSerializationContext& Context, const TRefCountPtr<FNetBlob>& Blob, TArray<TRefCountPtr<FNetBlob>>& OutPartialBlobs, const UE::Net::FNetDebugName* InDebugName) const
 {
 	using namespace UE::Net;
 
-	FNetSerializationContext Context;
-	Private::FInternalNetSerializationContext InternalContext(ReplicationSystem);
-	Context.SetLocalConnectionId(0U);
-	Context.SetInternalContext(&InternalContext);
-
 	FNetBlobCreationInfo CreationInfo = {};
 	CreationInfo.Type = GetNetBlobType();
 	CreationInfo.Flags = Blob->GetCreationInfo().Flags & (ENetBlobFlags::Ordered | ENetBlobFlags::Reliable);
@@ -50,15 +45,10 @@ bool USequentialPartialNetBlobHandler::SplitNetBlob(const TRefCountPtr<FNetBlob>
 	return FPartialNetBlob::SplitNetBlob(Context, CreationInfo, SplitParams, Blob, OutPartialBlobs);
 }
 
-bool USequentialPartialNetBlobHandler::SplitNetBlob(const UE::Net::FNetObjectReference& NetObjectReference, const TRefCountPtr<FNetBlob>& Blob, TArray<TRefCountPtr<FNetBlob>>& OutPartialBlobs, const UE::Net::FNetDebugName* InDebugName) const
+bool USequentialPartialNetBlobHandler::SplitNetBlob(UE::Net::FNetSerializationContext& Context, const UE::Net::FNetObjectReference& NetObjectReference, const TRefCountPtr<FNetBlob>& Blob, TArray<TRefCountPtr<FNetBlob>>& OutPartialBlobs, const UE::Net::FNetDebugName* InDebugName) const
 {
 	using namespace UE::Net;
 
-	FNetSerializationContext Context;
-	Private::FInternalNetSerializationContext InternalContext(ReplicationSystem);
-	Context.SetLocalConnectionId(0U);
-	Context.SetInternalContext(&InternalContext);
-
 	FNetBlobCreationInfo CreationInfo = {};
 	CreationInfo.Type = GetNetBlobType();
 	CreationInfo.Flags = Blob->GetCreationInfo().Flags & (ENetBlobFlags::Ordered | ENetBlobFlags::Reliable);
diff --git a/Engine/Source/Runtime/Experimental/Iris/Core/Private/Iris/ReplicationSystem/NetBlob/ShrinkWrapNetBlob.cpp b/Engine/Source/Runtime/Experimental/Iris/Core/Private/Iris/ReplicationSystem/NetBlob/ShrinkWrapNetBlob.cpp
--- a/Engine/Source/Runtime/Experimental/Iris/Core/Private/Iris/ReplicationSystem/NetBlob/ShrinkWrapNetBlob.cpp
+++ b/Engine/Source/Runtime/Experimental/Iris/Core/Private/Iris/ReplicationSystem/NetBlob/ShrinkWrapNetBlob.cpp
@@ -11,17 +11,24 @@
 namespace UE::Net
 {
 
-FShrinkWrapNetBlob::FShrinkWrapNetBlob(const TRefCountPtr<FNetBlob>& InOriginalBlob, TArray<uint32>&& Payload, uint32 PayloadBitCount)
+FShrinkWrapNetBlob::FShrinkWrapNetBlob(FNetSerializationContext& Context, const TRefCountPtr<FNetBlob>& InOriginalBlob, TArray<uint32>&& Payload, uint32 PayloadBitCount)
 : FNetBlob(InOriginalBlob->GetCreationInfo())
 , OriginalBlob(InOriginalBlob)
 , SerializedBlob(MoveTemp(Payload))
 , SerializedBlobBitCount(PayloadBitCount)
 {
+	NetTokenExportsArray = Context.GetExportContext()->GetBatchExports().NetTokensPendingExportInCurrentBatch;
+	CreationInfo.Flags |= NetTokenExportsArray.Num() ? ENetBlobFlags::HasExports : ENetBlobFlags::None;
 }
 
-TArrayView<const FNetObjectReference> FShrinkWrapNetBlob::GetExports() const
+TArrayView<const FNetObjectReference> FShrinkWrapNetBlob::GetNetObjectReferenceExports() const
 {
-	return OriginalBlob->CallGetExports();
+	return OriginalBlob->CallGetNetObjectReferenceExports();
+}
+
+TArrayView<const FNetToken> FShrinkWrapNetBlob::GetNetTokenExports() const
+{
+	return MakeArrayView<const FNetToken>(NetTokenExportsArray.GetData(), NetTokenExportsArray.Num());
 }
 
 void FShrinkWrapNetBlob::SerializeWithObject(FNetSerializationContext& Context, FNetRefHandle RefHandle) const
@@ -69,17 +76,24 @@ void FShrinkWrapNetBlob::InternalSerialize(FNetSerializationContext& Context) co
 }
 
 
-FShrinkWrapNetObjectAttachment::FShrinkWrapNetObjectAttachment(const TRefCountPtr<FNetObjectAttachment>& InOriginalBlob, TArray<uint32>&& Payload, uint32 PayloadBitCount)
+FShrinkWrapNetObjectAttachment::FShrinkWrapNetObjectAttachment(FNetSerializationContext& Context, const TRefCountPtr<FNetObjectAttachment>& InOriginalBlob, TArray<uint32>&& Payload, uint32 PayloadBitCount)
 : FNetBlob(InOriginalBlob->GetCreationInfo())
 , OriginalBlob(InOriginalBlob)
 , SerializedBlob(MoveTemp(Payload))
 , SerializedBlobBitCount(PayloadBitCount)
 {
+	NetTokenExportsArray = Context.GetExportContext()->GetBatchExports().NetTokensPendingExportInCurrentBatch;
+	CreationInfo.Flags |= NetTokenExportsArray.Num() ? ENetBlobFlags::HasExports : ENetBlobFlags::None;
+}
+
+TArrayView<const FNetObjectReference> FShrinkWrapNetObjectAttachment::GetNetObjectReferenceExports() const
+{
+	return OriginalBlob->CallGetNetObjectReferenceExports();
 }
 
-TArrayView<const FNetObjectReference> FShrinkWrapNetObjectAttachment::GetExports() const
+TArrayView<const FNetToken> FShrinkWrapNetObjectAttachment::GetNetTokenExports() const
 {
-	return OriginalBlob->CallGetExports();
+	return MakeArrayView<const FNetToken>(NetTokenExportsArray.GetData(), NetTokenExportsArray.Num());
 }
 
 void FShrinkWrapNetObjectAttachment::SerializeWithObject(FNetSerializationContext& Context, FNetRefHandle RefHandle) const
diff --git a/Engine/Source/Runtime/Experimental/Iris/Core/Private/Iris/ReplicationSystem/NetTokenStore.cpp b/Engine/Source/Runtime/Experimental/Iris/Core/Private/Iris/ReplicationSystem/NetTokenStore.cpp
--- a/Engine/Source/Runtime/Experimental/Iris/Core/Private/Iris/ReplicationSystem/NetTokenStore.cpp
+++ b/Engine/Source/Runtime/Experimental/Iris/Core/Private/Iris/ReplicationSystem/NetTokenStore.cpp
@@ -2,10 +2,13 @@
 
 #include "Iris/ReplicationSystem/NetTokenStore.h"
 #include "Iris/ReplicationSystem/NetTokenStoreState.h"
+#include "Iris/ReplicationSystem/ObjectReferenceCacheFwd.h"
 #include "Iris/Serialization/NetBitStreamReader.h"
 #include "Iris/Serialization/NetBitStreamWriter.h"
 #include "Iris/Serialization/NetBitStreamUtil.h"
 #include "Iris/Serialization/NetSerializationContext.h"
+#include "Iris/Serialization/InternalNetSerializationContext.h"
+#include "Iris/Serialization/NetExportContext.h"
 
 namespace UE::Net
 {
@@ -172,4 +175,41 @@ void FNetTokenStore::ReadTokenDataForIndex(FNetSerializationContext& Context, ui
 	TokenStoreState.TokenInfos[TokenIndex] = StoreKey;
 }
 
+void FNetTokenStore::ConditionalWriteNetTokenData(FNetSerializationContext& Context, Private::FNetExportContext* ExportContext, const FNetToken& NetToken) const
+{
+	FNetBitStreamWriter* Writer = Context.GetBitStreamWriter();
+
+	if (ExportContext)
+	{
+		if (Writer->WriteBool(!ExportContext->IsExported(NetToken)))
+		{
+			WriteTokenData(Context, NetToken);
+			ExportContext->AddExported(NetToken);			
+		}
+	}
+	else
+	{
+		Writer->WriteBool(true);
+		WriteTokenData(Context, NetToken);
+	}
+}
+
+void FNetTokenStore::ConditionalReadNetTokenData(FNetSerializationContext& Context, const FNetToken& NetToken)
+{
+	FNetBitStreamReader* Reader = Context.GetBitStreamReader();
+
+	const bool bIsExportToken = Reader->ReadBool();
+	if (bIsExportToken)
+	{
+		if (Reader->IsOverflown())
+		{
+			return;
+		}
+
+		FNetObjectResolveContext& ResolveContext = Context.GetInternalContext()->ResolveContext;
+	
+		ReadTokenData(Context, NetToken, *ResolveContext.RemoteNetTokenStoreState);
+	}
+}
+
 }
diff --git a/Engine/Source/Runtime/Experimental/Iris/Core/Private/Iris/ReplicationSystem/ObjectReferenceCache.cpp b/Engine/Source/Runtime/Experimental/Iris/Core/Private/Iris/ReplicationSystem/ObjectReferenceCache.cpp
--- a/Engine/Source/Runtime/Experimental/Iris/Core/Private/Iris/ReplicationSystem/ObjectReferenceCache.cpp
+++ b/Engine/Source/Runtime/Experimental/Iris/Core/Private/Iris/ReplicationSystem/ObjectReferenceCache.cpp
@@ -1196,43 +1196,6 @@ FNetObjectReference FObjectReferenceCache::GetOrCreateObjectReference(const FStr
 	return OutReference;
 }
 
-void FObjectReferenceCache::ConditionalWriteNetTokenData(FNetSerializationContext& Context, FNetExportContext* ExportContext, const FNetToken& NetToken) const
-{
-	FNetBitStreamWriter* Writer = Context.GetBitStreamWriter();
-
-	if (ExportContext)
-	{
-		if (Writer->WriteBool(!ExportContext->IsExported(NetToken)))
-		{
-			NetTokenStore->WriteTokenData(Context, NetToken);
-			ExportContext->AddExported(NetToken);			
-		}
-	}
-	else
-	{
-		Writer->WriteBool(true);
-		NetTokenStore->WriteTokenData(Context, NetToken);
-	}
-}
-
-void FObjectReferenceCache::ConditionalReadNetTokenData(FNetSerializationContext& Context, const FNetToken& NetToken) const
-{
-	FNetBitStreamReader* Reader = Context.GetBitStreamReader();
-
-	const bool bIsExportToken = Reader->ReadBool();
-	if (bIsExportToken)
-	{
-		if (Reader->IsOverflown())
-		{
-			return;
-		}
-
-		FNetObjectResolveContext& ResolveContext = Context.GetInternalContext()->ResolveContext;
-	
-		NetTokenStore->ReadTokenData(Context, NetToken, *ResolveContext.RemoteNetTokenStoreState);
-	}
-}
-
 void FObjectReferenceCache::WriteFullReferenceInternal(FNetSerializationContext& Context, const FNetObjectReference& Ref) const
 {
 	FNetBitStreamWriter* Writer = Context.GetBitStreamWriter();
@@ -1271,7 +1234,7 @@ void FObjectReferenceCache::WriteFullReferenceInternal(FNetSerializationContext&
 		if (Writer->WriteBool(bHasPath))
 		{
 			WriteNetToken(Writer, CachedObject->RelativePath);
-			ConditionalWriteNetTokenData(Context, ExportContext, CachedObject->RelativePath);
+			NetTokenStore->ConditionalWriteNetTokenData(Context, ExportContext, CachedObject->RelativePath);
 			WriteFullReferenceInternal(Context, FNetObjectReference(CachedObject->OuterNetRefHandle));
 		}
 	
@@ -1338,7 +1301,7 @@ void FObjectReferenceCache::ReadFullReferenceInternal(FNetSerializationContext&
 		if (Reader->ReadBool())
 		{
 			RelativePath = ReadNetToken(Reader);
-			ConditionalReadNetTokenData(Context, RelativePath);
+			NetTokenStore->ConditionalReadNetTokenData(Context, RelativePath);
 			ReadFullReferenceInternal(Context, OuterRef, RecursionCount + 1U);
 		}
 	}
@@ -1424,7 +1387,7 @@ void FObjectReferenceCache::ReadFullReference(FNetSerializationContext& Context,
 
 		const FNetRefHandle NetRefHandle = ReadNetRefHandle(Context);
 		FNetToken RelativePath = ReadNetToken(Reader);
-		ConditionalReadNetTokenData(Context, RelativePath);
+		NetTokenStore->ConditionalReadNetTokenData(Context, RelativePath);
 
 		OutRef.RefHandle = NetRefHandle;
 		OutRef.PathToken = RelativePath;
@@ -1463,7 +1426,7 @@ void FObjectReferenceCache::WriteFullReference(FNetSerializationContext& Context
 
 		WriteNetRefHandle(Context, Ref.GetRefHandle());
 		WriteNetToken(Writer, Ref.PathToken);
-		ConditionalWriteNetTokenData(Context, Context.GetExportContext(), Ref.PathToken);
+		NetTokenStore->ConditionalWriteNetTokenData(Context, Context.GetExportContext(), Ref.PathToken);
 		return;
 	}
 
@@ -1595,7 +1558,7 @@ FObjectReferenceCache::EWriteExportsResult FObjectReferenceCache::WritePendingEx
 	FNetBitStreamWriter& Writer = *Context.GetBitStreamWriter();
 
 	FNetExportContext* ExportContext = Context.GetExportContext();
-	if (!ExportContext || ExportContext->GetBatchExports().ReferencesPendingExportInCurrentBatch.IsEmpty())
+	if (!ExportContext || !ExportContext->GetBatchExports().HasPendingExports())
 	{
 		return EWriteExportsResult::NoExports;
 	}
@@ -1604,17 +1567,39 @@ FObjectReferenceCache::EWriteExportsResult FObjectReferenceCache::WritePendingEx
 	FNetBitStreamRollbackScope Rollback(*Context.GetBitStreamWriter());
 	const uint32 ExportsStartBitPos = Writer.GetPosBits();
 
-	TArrayView<const FNetObjectReference> ExportsView = MakeArrayView(ExportContext->GetBatchExports().ReferencesPendingExportInCurrentBatch);
+#if UE_NET_IRIS_CSV_STATS
+		uint32 CountExports = 0;
+#endif
+
+	// Export NetObjectReference pending export
+	TArrayView<const FNetObjectReference> NetObjectReferenceExportsView = MakeArrayView(ExportContext->GetBatchExports().ReferencesPendingExportInCurrentBatch);
 
 	// Force exports to be written
 	{
 		FForceInlineExportScope ForceInlineExportScope(Context.GetInternalContext());
 
-#if UE_NET_IRIS_CSV_STATS
-		uint32 CountExports = 0;
-#endif
+		// For now we export NetTokens here
+		{
+			UE_NET_TRACE_SCOPE(NetTokenExports, Writer, Context.GetTraceCollector(), ENetTraceVerbosity::Verbose);
 
-		for (const FNetObjectReference& Reference : ExportsView)
+			TArrayView<const FNetToken> NetTokenExportsView = MakeArrayView(ExportContext->GetBatchExports().NetTokensPendingExportInCurrentBatch);
+
+			for (const FNetToken& NetToken : NetTokenExportsView)
+			{
+				if (!ExportContext->IsExported(NetToken))
+				{
+					//UE_NET_TRACE_DYNAMIC_NAME_SCOPE(*FString::Printf(TEXT("%u"), NetToken.GetIndex()), *Context.GetBitStreamWriter(), Context.GetTraceCollector(), ENetTraceVerbosity::VeryVerbose);
+					Writer.WriteBool(true);
+					WriteNetToken(&Writer, NetToken);
+					NetTokenStore->WriteTokenData(Context, NetToken);
+					ExportContext->AddExported(NetToken);			
+				}
+			}
+			// Write stop bit
+			Writer.WriteBool(false);
+		}
+
+		for (const FNetObjectReference& Reference : NetObjectReferenceExportsView)
 		{
 			const bool bIsClientAssigned = Reference.PathToken.IsValid();
 			const FNetRefHandle Handle = Reference.GetRefHandle();
@@ -1636,7 +1621,7 @@ FObjectReferenceCache::EWriteExportsResult FObjectReferenceCache::WritePendingEx
 	}
 
 	// We also write any must be mapped exports
-	WriteMustBeMappedExports(Context, ObjectIndex, ExportsView);
+	WriteMustBeMappedExports(Context, ObjectIndex, NetObjectReferenceExportsView);
 
 	// Reset state of pending exports
 	ExportContext->ClearPendingExports();
@@ -1659,17 +1644,39 @@ bool FObjectReferenceCache::ReadExports(FNetSerializationContext& Context, TArra
 {
 	FNetBitStreamReader& Reader = *Context.GetBitStreamReader();
 
-	bool bHasExportsToRead = Reader.ReadBool(); 
-	UE_NET_TRACE_SCOPE(Exports, Reader, Context.GetTraceCollector(), ENetTraceVerbosity::Verbose);
+	// Read NetToken exports
+	{
+		bool bHasExportsToRead = Reader.ReadBool(); 
+		UE_NET_TRACE_SCOPE(NetTokenExports, Reader, Context.GetTraceCollector(), ENetTraceVerbosity::Verbose);
+
+		// Force inlined exports
+		{
+			FNetObjectResolveContext& ResolveContext = Context.GetInternalContext()->ResolveContext;
 
-	// Force inlined exports
+			FForceInlineExportScope ForceInlineExportScope(Context.GetInternalContext());
+			while (bHasExportsToRead && !Context.HasErrorOrOverflow())
+			{
+				FNetToken ImportedNetToken = ReadNetToken(&Reader);
+				NetTokenStore->ReadTokenData(Context, ImportedNetToken, *ResolveContext.RemoteNetTokenStoreState);
+				bHasExportsToRead = Reader.ReadBool();
+			}
+		}
+	}
+
+	// Read NetObjectReference exports
 	{
-		FForceInlineExportScope ForceInlineExportScope(Context.GetInternalContext());
-		while (bHasExportsToRead && !Context.HasErrorOrOverflow())
+		bool bHasExportsToRead = Reader.ReadBool(); 
+		UE_NET_TRACE_SCOPE(NetObjectReferenceExports, Reader, Context.GetTraceCollector(), ENetTraceVerbosity::Verbose);
+
+		// Force inlined exports
 		{
-			FNetObjectReference Import;
-			ReadFullReference(Context, Import);
-			bHasExportsToRead = Reader.ReadBool();
+			FForceInlineExportScope ForceInlineExportScope(Context.GetInternalContext());
+			while (bHasExportsToRead && !Context.HasErrorOrOverflow())
+			{
+				FNetObjectReference Import;
+				ReadFullReference(Context, Import);
+				bHasExportsToRead = Reader.ReadBool();
+			}
 		}
 	}
 
@@ -1710,6 +1717,7 @@ bool FObjectReferenceCache::WriteMustBeMappedExports(FNetSerializationContext& C
 			}
 		}
 
+		// $TODO: This should be a separate state, as MustBeMapped is not really an export, just a list of important exports
 		UE_NET_IRIS_STATS_ADD_COUNT_FOR_OBJECT(Context.GetNetStatsContext(), WriteExports, ObjectIndex, CountExports);
 	}
 
diff --git a/Engine/Source/Runtime/Experimental/Iris/Core/Private/Iris/ReplicationSystem/ObjectReferenceCache.h b/Engine/Source/Runtime/Experimental/Iris/Core/Private/Iris/ReplicationSystem/ObjectReferenceCache.h
--- a/Engine/Source/Runtime/Experimental/Iris/Core/Private/Iris/ReplicationSystem/ObjectReferenceCache.h
+++ b/Engine/Source/Runtime/Experimental/Iris/Core/Private/Iris/ReplicationSystem/ObjectReferenceCache.h
@@ -182,8 +182,6 @@ private:
 
 	bool CreateObjectReferenceInternal(const UObject* Object, FNetObjectReference& OutReference);
 
-	void ConditionalWriteNetTokenData(FNetSerializationContext& Context, Private::FNetExportContext* ExportContext, const FNetToken& NetToken) const;
-	void ConditionalReadNetTokenData(FNetSerializationContext& Context, const FNetToken& NetToken) const;
 
 	void ReadFullReferenceInternal(FNetSerializationContext& Context, FNetObjectReference& OutRef, uint32 RecursionCount);
 	void WriteFullReferenceInternal(FNetSerializationContext& Context, const FNetObjectReference& Ref) const;
diff --git a/Engine/Source/Runtime/Experimental/Iris/Core/Private/Iris/ReplicationSystem/ObjectReplicationBridge.cpp b/Engine/Source/Runtime/Experimental/Iris/Core/Private/Iris/ReplicationSystem/ObjectReplicationBridge.cpp
--- a/Engine/Source/Runtime/Experimental/Iris/Core/Private/Iris/ReplicationSystem/ObjectReplicationBridge.cpp
+++ b/Engine/Source/Runtime/Experimental/Iris/Core/Private/Iris/ReplicationSystem/ObjectReplicationBridge.cpp
@@ -278,7 +278,7 @@ UE::Net::FNetRefHandle UObjectReplicationBridge::BeginReplication(UObject* Insta
 	Traits |= Params.bNeedsPreUpdate ? EReplicationFragmentTraits::NeedsPreSendUpdate : EReplicationFragmentTraits::None;
 	Traits |= Params.bNeedsWorldLocationUpdate ? EReplicationFragmentTraits::NeedsWorldLocationUpdate : EReplicationFragmentTraits::None;
 	
-	FFragmentRegistrationContext FragmentRegistrationContext(GetReplicationStateDescriptorRegistry(), Traits);
+	FFragmentRegistrationContext FragmentRegistrationContext(GetReplicationStateDescriptorRegistry(), GetReplicationSystem(), Traits);
 
 	// For everything derived from UObject we can call the virtual function RegisterReplicationFragments	
 	CallRegisterReplicationFragments(Instance, FragmentRegistrationContext, EFragmentRegistrationFlags::None);
@@ -732,7 +732,7 @@ FReplicationBridgeCreateNetRefHandleResult UObjectReplicationBridge::CreateNetRe
 	}
 
 	// Currently remote objects can only receive replicated data
-	FFragmentRegistrationContext FragmentRegistrationContext(GetReplicationStateDescriptorRegistry(), EReplicationFragmentTraits::CanReceive);
+	FFragmentRegistrationContext FragmentRegistrationContext(GetReplicationStateDescriptorRegistry(), GetReplicationSystem(), EReplicationFragmentTraits::CanReceive);
 	FReplicationProtocolManager* ProtocolManager = GetReplicationProtocolManager();
 
 	FReplicationBridgeCreateNetRefHandleResult CreateResult;
diff --git a/Engine/Source/Runtime/Experimental/Iris/Core/Private/Iris/ReplicationSystem/ReplicationFragmentInternal.h b/Engine/Source/Runtime/Experimental/Iris/Core/Private/Iris/ReplicationSystem/ReplicationFragmentInternal.h
--- a/Engine/Source/Runtime/Experimental/Iris/Core/Private/Iris/ReplicationSystem/ReplicationFragmentInternal.h
+++ b/Engine/Source/Runtime/Experimental/Iris/Core/Private/Iris/ReplicationSystem/ReplicationFragmentInternal.h
@@ -9,5 +9,8 @@ namespace UE::Net::Private
 	struct FFragmentRegistrationContextPrivateAccessor
 	{
 		static const FReplicationFragments& GetReplicationFragments(const FFragmentRegistrationContext& Context) { return Context.Fragments; }
+		static const Private::FReplicationStateDescriptorRegistry* GetReplicationStateRegistry(const FFragmentRegistrationContext& Context) { return Context.ReplicationStateRegistry; }
+		static Private::FReplicationStateDescriptorRegistry* GetReplicationStateRegistry(FFragmentRegistrationContext& Context) { return Context.ReplicationStateRegistry; }
+		static UReplicationSystem* GetReplicationSystem(const FFragmentRegistrationContext& Context) { return Context.ReplicationSystem; }
 	};
 }
diff --git a/Engine/Source/Runtime/Experimental/Iris/Core/Private/Iris/ReplicationSystem/ReplicationFragmentUtil.cpp b/Engine/Source/Runtime/Experimental/Iris/Core/Private/Iris/ReplicationSystem/ReplicationFragmentUtil.cpp
--- a/Engine/Source/Runtime/Experimental/Iris/Core/Private/Iris/ReplicationSystem/ReplicationFragmentUtil.cpp
+++ b/Engine/Source/Runtime/Experimental/Iris/Core/Private/Iris/ReplicationSystem/ReplicationFragmentUtil.cpp
@@ -3,6 +3,7 @@
 #include "Iris/ReplicationSystem/ReplicationFragmentUtil.h"
 #include "Iris/ReplicationSystem/ReplicationFragment.h"
 #include "Iris/ReplicationSystem/PropertyReplicationFragment.h"
+#include "Iris/ReplicationSystem/ReplicationFragmentInternal.h"
 #include "Iris/ReplicationSystem/FastArrayReplicationFragment.h"
 #include "Iris/ReplicationState/ReplicationStateDescriptorBuilder.h"
 #include "Iris/Core/IrisProfiler.h"
@@ -16,7 +17,8 @@ uint32 FReplicationFragmentUtil::CreateAndRegisterFragmentsForObject(UObject* Ob
 	IRIS_PROFILER_SCOPE(FReplicationFragmentUtil_CreateAndRegisterFragmentsForObject);
 
 	FReplicationStateDescriptorBuilder::FParameters BuilderParameters;
-	BuilderParameters.DescriptorRegistry = Context.GetReplicationStateRegistry();
+	BuilderParameters.DescriptorRegistry = UE::Net::Private::FFragmentRegistrationContextPrivateAccessor::GetReplicationStateRegistry(Context);
+	BuilderParameters.ReplicationSystem = UE::Net::Private::FFragmentRegistrationContextPrivateAccessor::GetReplicationSystem(Context);
 
 	// By default we use the archetype for our default values, but if InitializeDefaultStateFromClassDefaults is set we 
 	// use the class default instead
diff --git a/Engine/Source/Runtime/Experimental/Iris/Core/Private/Iris/Serialization/NetExportContext.h b/Engine/Source/Runtime/Experimental/Iris/Core/Private/Iris/Serialization/NetExportContext.h
--- a/Engine/Source/Runtime/Experimental/Iris/Core/Private/Iris/Serialization/NetExportContext.h
+++ b/Engine/Source/Runtime/Experimental/Iris/Core/Private/Iris/Serialization/NetExportContext.h
@@ -19,6 +19,7 @@ public:
 	typedef TArray<FNetRefHandle, TInlineAllocator<32>> FExportsArray;
 	typedef TArray<FNetToken, TInlineAllocator<32>> FNetTokenExportsArray;
 	typedef TArray<FNetObjectReference, TInlineAllocator<32>> FPendingExportArray;
+	typedef TArray<FNetToken, TInlineAllocator<32>> FNetTokenPendingExportArray;
 
 	struct FAcknowledgedExports
 	{
@@ -33,12 +34,16 @@ public:
 			HandlesExportedInCurrentBatch.Empty();
 			NetTokensExportedInCurrentBatch.Empty();
 			ReferencesPendingExportInCurrentBatch.Empty();
+			NetTokensPendingExportInCurrentBatch.Empty();
 		}
 
+		bool HasPendingExports() const { return (ReferencesPendingExportInCurrentBatch.Num() > 0) || (NetTokensPendingExportInCurrentBatch.Num() > 0); }
+
 		// Exports in the current batch
 		FExportsArray HandlesExportedInCurrentBatch;
 		FNetTokenExportsArray NetTokensExportedInCurrentBatch;
 		FPendingExportArray ReferencesPendingExportInCurrentBatch;
+		FNetTokenPendingExportArray NetTokensPendingExportInCurrentBatch;
 	};
 
 public:
@@ -57,14 +62,24 @@ public:
 	// Add a Handle to the current export batch
 	void AddExported(FNetToken Token);
 
-	// Add a reference to the current pending exports list.
+	// Add a reference to the current pending exports arr.
 	void AddPendingExport(const FNetObjectReference& Ref);
 
+	// Add a NetToken to the current pending export array
+	void AddPendingExport(FNetToken Token);
+
+	// Add NetTokens to the current pending export array
+	void AddPendingExports(TArrayView<const FNetToken> NetTokens);
+
 	// Returns true if the Reference is in PendingExports array
 	bool IsPendingExport(const FNetObjectReference& Ref) const;
 
-	// Clear the list of reference pending exports
-	void ClearPendingExports() { BatchExports.ReferencesPendingExportInCurrentBatch.Empty(); }
+	// Clear the list of pending exports
+	void ClearPendingExports()
+	{ 
+		BatchExports.ReferencesPendingExportInCurrentBatch.Empty(); 
+		BatchExports.NetTokensPendingExportInCurrentBatch.Empty();
+	}
 
 	// Get current batch exports
 	const FNetExportContext::FBatchExports& GetBatchExports() const { return BatchExports; }
@@ -93,6 +108,7 @@ private:
 	int32 StartNumNetHandleExports;
 	int32 StartNumNetTokenExports;
 	int32 StartNumPendingExports;
+	int32 StartNumNetTokensPendingExports;
 };
 
 inline FNetExportRollbackScope::FNetExportRollbackScope(FNetSerializationContext& InContext)
@@ -103,6 +119,7 @@ inline FNetExportRollbackScope::FNetExportRollbackScope(FNetSerializationContext
 	StartNumNetHandleExports = ExportContext ? ExportContext->BatchExports.HandlesExportedInCurrentBatch.Num() : 0;
 	StartNumNetTokenExports = ExportContext ? ExportContext->BatchExports.NetTokensExportedInCurrentBatch.Num() : 0;
 	StartNumPendingExports = ExportContext ? ExportContext->BatchExports.ReferencesPendingExportInCurrentBatch.Num() : 0;
+	StartNumNetTokensPendingExports = ExportContext ? ExportContext->BatchExports.NetTokensPendingExportInCurrentBatch.Num() : 0;
 }
 
 inline void FNetExportRollbackScope::Rollback()
@@ -111,6 +128,7 @@ inline void FNetExportRollbackScope::Rollback()
 	{ 
 		ExportContext->BatchExports.HandlesExportedInCurrentBatch.SetNum(StartNumNetHandleExports);
 		ExportContext->BatchExports.NetTokensExportedInCurrentBatch.SetNum(StartNumNetTokenExports);
+		ExportContext->BatchExports.NetTokensPendingExportInCurrentBatch.SetNum(StartNumNetTokensPendingExports);
 		ExportContext->BatchExports.ReferencesPendingExportInCurrentBatch.SetNum(StartNumPendingExports);
 	}
 }
@@ -155,6 +173,19 @@ inline void FNetExportContext::AddPendingExport(const FNetObjectReference& Ref)
 	BatchExports.ReferencesPendingExportInCurrentBatch.AddUnique(Ref);
 }
 
+inline void FNetExportContext::AddPendingExport(FNetToken NetToken)
+{
+	BatchExports.NetTokensPendingExportInCurrentBatch.AddUnique(NetToken);
+}
+
+inline void FNetExportContext::AddPendingExports(TArrayView<const FNetToken> NetTokens)
+{
+	for (const FNetToken& NetToken : NetTokens)
+	{
+		AddPendingExport(NetToken);
+	}
+}
+
 inline bool FNetExportContext::IsPendingExport(const FNetObjectReference& Ref) const
 {
 	return BatchExports.ReferencesPendingExportInCurrentBatch.Contains(Ref);
diff --git a/Engine/Source/Runtime/Experimental/Iris/Core/Private/Iris/Serialization/NetSerializationContext.cpp b/Engine/Source/Runtime/Experimental/Iris/Core/Private/Iris/Serialization/NetSerializationContext.cpp
--- a/Engine/Source/Runtime/Experimental/Iris/Core/Private/Iris/Serialization/NetSerializationContext.cpp
+++ b/Engine/Source/Runtime/Experimental/Iris/Core/Private/Iris/Serialization/NetSerializationContext.cpp
@@ -71,6 +71,11 @@ UObject* FNetSerializationContext::GetLocalConnectionUserData(uint32 ConnectionI
 		return nullptr;
 	}
 
+	if (ConnectionId == UE::Net::InvalidConnectionId)
+	{
+		return nullptr;
+	}
+
 	UObject* UserData = ReplicationSystem->GetConnectionUserData(ConnectionId);
 	return UserData;
 }
diff --git a/Engine/Source/Runtime/Experimental/Iris/Core/Private/Iris/Serialization/ObjectNetSerializer.cpp b/Engine/Source/Runtime/Experimental/Iris/Core/Private/Iris/Serialization/ObjectNetSerializer.cpp
--- a/Engine/Source/Runtime/Experimental/Iris/Core/Private/Iris/Serialization/ObjectNetSerializer.cpp
+++ b/Engine/Source/Runtime/Experimental/Iris/Core/Private/Iris/Serialization/ObjectNetSerializer.cpp
@@ -182,8 +182,12 @@ void FObjectNetSerializerBase<T>::Quantize(FNetSerializationContext& Context, co
 	QuantizedType& Target = *reinterpret_cast<QuantizedType*>(Args.Target);
 	UObject* Source = GetValue(*reinterpret_cast<T*>(Args.Source));
 
+	// Notice that we quantize to default here if initializing for default state, this is due to the fact that
+	// object references cannot be quantized for default state as they will use locally assigned ids which will differ between server and client
+	// $TODO: Jira: UE-221750 It should now be possible to implement support for this, need to track if stored quantized data is local or not, as long as dequantized
+	// state is correct it should work.
 	const FInternalNetSerializationContext* InternalContext = Context.GetInternalContext();
-	Target = InternalContext->ObjectReferenceCache ? InternalContext->ObjectReferenceCache->GetOrCreateObjectReference(Source) : QuantizedType();
+	Target = Context.IsInitializingDefaultState() ? QuantizedType() : InternalContext->ObjectReferenceCache->GetOrCreateObjectReference(Source);
 }
 
 template<typename T>
diff --git a/Engine/Source/Runtime/Experimental/Iris/Core/Public/Iris/ReplicationState/ReplicationStateDescriptorBuilder.h b/Engine/Source/Runtime/Experimental/Iris/Core/Public/Iris/ReplicationState/ReplicationStateDescriptorBuilder.h
--- a/Engine/Source/Runtime/Experimental/Iris/Core/Public/Iris/ReplicationState/ReplicationStateDescriptorBuilder.h
+++ b/Engine/Source/Runtime/Experimental/Iris/Core/Public/Iris/ReplicationState/ReplicationStateDescriptorBuilder.h
@@ -8,6 +8,7 @@
 #include "UObject/CoreNet.h"
 
 class IConsoleVariable;
+class UReplicationSystem;
 namespace UE::Net::Private
 {
 	class FReplicationStateDescriptorRegistry;
@@ -24,6 +25,7 @@ public:
 		IRISCORE_API FParameters();
 
 		Private::FReplicationStateDescriptorRegistry* DescriptorRegistry;
+		UReplicationSystem* ReplicationSystem;
 		// If provided we use this to source our default state, if not provided we source default state from the CDO instead
 		UObject* DefaultStateSource;
 		 // Include super class when building descriptor.
diff --git a/Engine/Source/Runtime/Experimental/Iris/Core/Public/Iris/ReplicationSystem/NetBlob/NetBlob.h b/Engine/Source/Runtime/Experimental/Iris/Core/Public/Iris/ReplicationSystem/NetBlob/NetBlob.h
--- a/Engine/Source/Runtime/Experimental/Iris/Core/Public/Iris/ReplicationSystem/NetBlob/NetBlob.h
+++ b/Engine/Source/Runtime/Experimental/Iris/Core/Public/Iris/ReplicationSystem/NetBlob/NetBlob.h
@@ -124,8 +124,14 @@ public:
 	/** Returns the reference count.  A blob is created with reference count zero. */
 	int32 GetRefCount() const { return RefCount; }
 
-	/** Retrieve the object references that need to be exported.  */
-	TArrayView<const FNetObjectReference> CallGetExports() const { return EnumHasAnyFlags(CreationInfo.Flags, ENetBlobFlags::HasExports) ? GetExports() : MakeArrayView<const FNetObjectReference>(nullptr, 0); };
+	/* Returns true if the blob has additional exports to add. */
+	bool HasExports() const { return EnumHasAnyFlags(GetCreationInfo().Flags, ENetBlobFlags::HasExports); }
+
+	/** Retrieve the object references that need to be exported. */
+	TArrayView<const FNetObjectReference> CallGetNetObjectReferenceExports() const { return GetNetObjectReferenceExports(); };
+
+	/** Retrieve NetToken that needs to be exported. */
+	TArrayView<const FNetToken> CallGetNetTokenExports() const { return GetNetTokenExports(); };
 
 protected:
 	FNetBlob(const FNetBlob&) = delete;
@@ -135,7 +141,10 @@ protected:
 	IRISCORE_API virtual ~FNetBlob();
 
 	/** Override to return the object references that need to be exported. */
-	virtual TArrayView<const FNetObjectReference> GetExports() const;
+	virtual TArrayView<const FNetObjectReference> GetNetObjectReferenceExports() const;
+
+	/** Override to return NetTokenExports that need to be exported, mostly relevant for pre-serialized blobs. */
+	virtual TArrayView<const FNetToken> GetNetTokenExports() const;
 
 	/** Serializes the state if there's a valid descriptor and state buffer. */
 	IRISCORE_API void SerializeBlob(FNetSerializationContext& Context) const;
diff --git a/Engine/Source/Runtime/Experimental/Iris/Core/Public/Iris/ReplicationSystem/NetBlob/PartialNetBlob.h b/Engine/Source/Runtime/Experimental/Iris/Core/Public/Iris/ReplicationSystem/NetBlob/PartialNetBlob.h
--- a/Engine/Source/Runtime/Experimental/Iris/Core/Public/Iris/ReplicationSystem/NetBlob/PartialNetBlob.h
+++ b/Engine/Source/Runtime/Experimental/Iris/Core/Public/Iris/ReplicationSystem/NetBlob/PartialNetBlob.h
@@ -22,7 +22,7 @@ public:
 	};
 
 	// Split a NetBlob into multiple PartialNetBlobs. The blob will be split even if the original one didn't need it.
-	IRISCORE_API static bool SplitNetBlob(const FNetSerializationContext& Context, const FNetBlobCreationInfo& CreationInfo, const FSplitParams& SplitParams, const TRefCountPtr<FNetBlob>& Blob, TArray<TRefCountPtr<FNetBlob>>& OutPartialBlobs);
+	IRISCORE_API static bool SplitNetBlob(FNetSerializationContext& Context, const FNetBlobCreationInfo& CreationInfo, const FSplitParams& SplitParams, const TRefCountPtr<FNetBlob>& Blob, TArray<TRefCountPtr<FNetBlob>>& OutPartialBlobs);
 
 	// Split a RawDataNetBlob into multiple PartialNetBlobs. The blob will be split even if the original one didn't need it.
 	IRISCORE_API static bool SplitNetBlob(const FNetBlobCreationInfo& CreationInfo, const FSplitParams& SplitParams, const TRefCountPtr<FRawDataNetBlob>& Blob, TArray<TRefCountPtr<FNetBlob>>& OutPartialBlobs);
@@ -53,7 +53,9 @@ public:
 
 private:
 
-	virtual TArrayView<const FNetObjectReference> GetExports() const override final;
+	virtual TArrayView<const FNetObjectReference> GetNetObjectReferenceExports() const override final;
+	virtual TArrayView<const FNetToken> GetNetTokenExports() const override final;
+
 	virtual void SerializeWithObject(FNetSerializationContext& Context, FNetRefHandle RefHandle) const override;
 	virtual void DeserializeWithObject(FNetSerializationContext& Context, FNetRefHandle RefHandle) override;
 
@@ -73,6 +75,7 @@ private:
 		FNetBlobCreationInfo CreationInfo;
 		FNetBlobCreationInfo OriginalCreationInfo;
 		FNetBlob* OriginalBlob;
+		TArrayView<const FNetToken> NetTokensPendingExport;
 		const uint32* Payload;
 		uint32 PayloadBitCount;
 		uint32 PartBitCount;
@@ -90,6 +93,7 @@ private:
 	uint16 PayloadBitCount = 0;
 	// Use uint32 for guaranteed FNetBitStreamReader/Writer compatibility
 	TArray<uint32> Payload;
+	TArray<FNetToken, TInlineAllocator<4>> NetTokenExportsArray;
 
 	FNetDebugName DebugName;
 };
diff --git a/Engine/Source/Runtime/Experimental/Iris/Core/Public/Iris/ReplicationSystem/NetBlob/SequentialPartialNetBlobHandler.h b/Engine/Source/Runtime/Experimental/Iris/Core/Public/Iris/ReplicationSystem/NetBlob/SequentialPartialNetBlobHandler.h
--- a/Engine/Source/Runtime/Experimental/Iris/Core/Public/Iris/ReplicationSystem/NetBlob/SequentialPartialNetBlobHandler.h
+++ b/Engine/Source/Runtime/Experimental/Iris/Core/Public/Iris/ReplicationSystem/NetBlob/SequentialPartialNetBlobHandler.h
@@ -46,10 +46,10 @@ class USequentialPartialNetBlobHandler : public UNetBlobHandler
 
 public:
 	/** Unconditionally splits a NetBlob into a sequence of PartialNetBlobs which are small in size. Calls FNetBlob::Serialize(). */
-	IRISCORE_API bool SplitNetBlob(const TRefCountPtr<FNetBlob>& Blob, TArray<TRefCountPtr<FNetBlob>>& OutPartialBlobs, const UE::Net::FNetDebugName* InDebugName = nullptr) const;
+	IRISCORE_API bool SplitNetBlob(UE::Net::FNetSerializationContext& Context, const TRefCountPtr<FNetBlob>& Blob, TArray<TRefCountPtr<FNetBlob>>& OutPartialBlobs, const UE::Net::FNetDebugName* InDebugName = nullptr) const;
 
 	/** Unconditionally splits a NetBlob into a sequence of PartialNetBlobs which are small in size. Calls FNetBlob::SerializeWithObject(). */
-	IRISCORE_API bool SplitNetBlob(const UE::Net::FNetObjectReference& NetObjectReference, const TRefCountPtr<FNetBlob>& Blob, TArray<TRefCountPtr<FNetBlob>>& OutPartialBlobs, const UE::Net::FNetDebugName* InDebugName = nullptr) const;
+	IRISCORE_API bool SplitNetBlob(UE::Net::FNetSerializationContext& Context, const UE::Net::FNetObjectReference& NetObjectReference, const TRefCountPtr<FNetBlob>& Blob, TArray<TRefCountPtr<FNetBlob>>& OutPartialBlobs, const UE::Net::FNetDebugName* InDebugName = nullptr) const;
 
 protected:
 	IRISCORE_API USequentialPartialNetBlobHandler();
diff --git a/Engine/Source/Runtime/Experimental/Iris/Core/Public/Iris/ReplicationSystem/NetBlob/ShrinkWrapNetBlob.h b/Engine/Source/Runtime/Experimental/Iris/Core/Public/Iris/ReplicationSystem/NetBlob/ShrinkWrapNetBlob.h
--- a/Engine/Source/Runtime/Experimental/Iris/Core/Public/Iris/ReplicationSystem/NetBlob/ShrinkWrapNetBlob.h
+++ b/Engine/Source/Runtime/Experimental/Iris/Core/Public/Iris/ReplicationSystem/NetBlob/ShrinkWrapNetBlob.h
@@ -3,6 +3,7 @@
 #pragma once
 #include "CoreTypes.h"
 #include "Iris/ReplicationSystem/NetBlob/NetBlob.h"
+#include "Iris/ReplicationSystem/NetToken.h"
 
 namespace UE::Net
 {
@@ -19,10 +20,11 @@ namespace UE::Net
 class FShrinkWrapNetBlob final : public FNetBlob
 {
 public:
-	IRISCORE_API FShrinkWrapNetBlob(const TRefCountPtr<FNetBlob>& OriginalBlob, TArray<uint32>&& Payload, uint32 PayloadBitCount);
+	IRISCORE_API FShrinkWrapNetBlob(FNetSerializationContext& Context, const TRefCountPtr<FNetBlob>& OriginalBlob, TArray<uint32>&& Payload, uint32 PayloadBitCount);
 
 private:
-	virtual TArrayView<const FNetObjectReference> GetExports() const override final;
+	virtual TArrayView<const FNetObjectReference> GetNetObjectReferenceExports() const override final;
+	virtual TArrayView<const FNetToken> GetNetTokenExports() const override final;
 	virtual void SerializeWithObject(FNetSerializationContext& Context, FNetRefHandle RefHandle) const override final;
 	virtual void DeserializeWithObject(FNetSerializationContext& Context, FNetRefHandle RefHandle) override final;
 
@@ -32,6 +34,7 @@ private:
 	void InternalSerialize(FNetSerializationContext& Context) const;
 
 	TRefCountPtr<FNetBlob> OriginalBlob;
+	TArray<FNetToken, TInlineAllocator<4>> NetTokenExportsArray;
 	TArray<uint32> SerializedBlob;
 	uint32 SerializedBlobBitCount;
 };
@@ -39,10 +42,11 @@ private:
 class FShrinkWrapNetObjectAttachment final : public FNetBlob
 {
 public:
-	IRISCORE_API FShrinkWrapNetObjectAttachment(const TRefCountPtr<FNetObjectAttachment>& OriginalBlob, TArray<uint32>&& Payload, uint32 PayloadBitCount);
+	IRISCORE_API FShrinkWrapNetObjectAttachment(FNetSerializationContext& Context, const TRefCountPtr<FNetObjectAttachment>& OriginalBlob, TArray<uint32>&& Payload, uint32 PayloadBitCount);
 
 private:
-	virtual TArrayView<const FNetObjectReference> GetExports() const override final;
+	virtual TArrayView<const FNetObjectReference> GetNetObjectReferenceExports() const override final;
+	virtual TArrayView<const FNetToken> GetNetTokenExports() const override final;
 	virtual void SerializeWithObject(FNetSerializationContext& Context, FNetRefHandle RefHandle) const override final;
 	virtual void DeserializeWithObject(FNetSerializationContext& Context, FNetRefHandle RefHandle) override final;
 
@@ -52,6 +56,7 @@ private:
 	void InternalSerialize(FNetSerializationContext& Context) const;
 
 	TRefCountPtr<FNetObjectAttachment> OriginalBlob;
+	TArray<FNetToken, TInlineAllocator<4>> NetTokenExportsArray;
 	TArray<uint32> SerializedBlob;
 	uint32 SerializedBlobBitCount;
 };
diff --git a/Engine/Source/Runtime/Experimental/Iris/Core/Public/Iris/ReplicationSystem/NetToken.h b/Engine/Source/Runtime/Experimental/Iris/Core/Public/Iris/ReplicationSystem/NetToken.h
--- a/Engine/Source/Runtime/Experimental/Iris/Core/Public/Iris/ReplicationSystem/NetToken.h
+++ b/Engine/Source/Runtime/Experimental/Iris/Core/Public/Iris/ReplicationSystem/NetToken.h
@@ -90,3 +90,5 @@ IRISCORE_API FNetToken ReadNetToken(FNetBitStreamReader* Reader);
 IRISCORE_API void WriteNetToken(FNetBitStreamWriter* Writer, FNetToken Token);
 
 }
+
+template <> struct TIsPODType<UE::Net::FNetToken> { enum { Value = true }; };
diff --git a/Engine/Source/Runtime/Experimental/Iris/Core/Public/Iris/ReplicationSystem/NetTokenStore.h b/Engine/Source/Runtime/Experimental/Iris/Core/Public/Iris/ReplicationSystem/NetTokenStore.h
--- a/Engine/Source/Runtime/Experimental/Iris/Core/Public/Iris/ReplicationSystem/NetTokenStore.h
+++ b/Engine/Source/Runtime/Experimental/Iris/Core/Public/Iris/ReplicationSystem/NetTokenStore.h
@@ -12,6 +12,11 @@ namespace UE::Net
 	class FNetTokenStore;
 	class FNetTokenStoreState;
 	class FNetSerializationContext;
+
+	namespace Private
+	{
+		class FNetExportContext;
+	}
 }
 
 namespace UE::Net
@@ -64,6 +69,12 @@ public:
 	// Read data associated with the NetToken
 	void ReadTokenData(FNetSerializationContext& Context, const FNetToken& NetToken, FNetTokenStoreState& RemoteNetTokenStoreState);
 
+	// Conditionally write NetTokenData unless already exported
+	void ConditionalWriteNetTokenData(FNetSerializationContext& Context, Private::FNetExportContext* ExportContext, const FNetToken& NetToken) const;
+
+	// Conditionally read NetTokenData if exported
+	void ConditionalReadNetTokenData(FNetSerializationContext& Context, const FNetToken& NetToken);
+
 private:
 
 	// Write token data for the TokenIndex
diff --git a/Engine/Source/Runtime/Experimental/Iris/Core/Public/Iris/ReplicationSystem/ReplicationFragment.h b/Engine/Source/Runtime/Experimental/Iris/Core/Public/Iris/ReplicationSystem/ReplicationFragment.h
--- a/Engine/Source/Runtime/Experimental/Iris/Core/Public/Iris/ReplicationSystem/ReplicationFragment.h
+++ b/Engine/Source/Runtime/Experimental/Iris/Core/Public/Iris/ReplicationSystem/ReplicationFragment.h
@@ -8,6 +8,8 @@
 #include "Misc/EnumClassFlags.h"
 
 // Forward declarations
+
+class UReplicationSystem;
 namespace UE::Net
 {
 	class FFragmentRegistrationContext;
@@ -230,8 +232,9 @@ typedef TArray<FReplicationFragmentInfo, TInlineAllocator<32>> FReplicationFragm
 class FFragmentRegistrationContext
 {
 public:
-	explicit FFragmentRegistrationContext(Private::FReplicationStateDescriptorRegistry* InReplicationStateRegistry, EReplicationFragmentTraits InFragmentTraits)
+	explicit FFragmentRegistrationContext(Private::FReplicationStateDescriptorRegistry* InReplicationStateRegistry, UReplicationSystem* InReplicationSystem, const EReplicationFragmentTraits InFragmentTraits)
 		: ReplicationStateRegistry(InReplicationStateRegistry)
+		, ReplicationSystem(InReplicationSystem)
 		, FragmentTraits(InFragmentTraits)
 	{
 	}
@@ -251,13 +254,12 @@ public:
 	/** Returns the number of fragments registered */
 	int32 NumFragments() const { return Fragments.Num(); }
 
-	Private::FReplicationStateDescriptorRegistry* GetReplicationStateRegistry() const { return ReplicationStateRegistry; }
-
 	friend Private::FFragmentRegistrationContextPrivateAccessor;
 
 private:
 	FReplicationFragments Fragments;
 	Private::FReplicationStateDescriptorRegistry* ReplicationStateRegistry;
+	UReplicationSystem* ReplicationSystem;
 	const EReplicationFragmentTraits FragmentTraits;
 	bool bIsAFragmentlessNetObject = false;
 };
