Iris
- Implemented support for exporting FNames using NetTokens to avoid constantly sending them as strings

diff --git a/Engine/Plugins/Runtime/ReplicationSystemTestPlugin/Source/Private/Tests/Serialization/TestNameNetSerializer.cpp b/Engine/Plugins/Runtime/ReplicationSystemTestPlugin/Source/Private/Tests/Serialization/TestNameNetSerializer.cpp
--- a/Engine/Plugins/Runtime/ReplicationSystemTestPlugin/Source/Private/Tests/Serialization/TestNameNetSerializer.cpp
+++ b/Engine/Plugins/Runtime/ReplicationSystemTestPlugin/Source/Private/Tests/Serialization/TestNameNetSerializer.cpp
@@ -1,8 +1,13 @@
 // Copyright Epic Games, Inc. All Rights Reserved.
 
+#include "TestNameNetSerializer.h"
 #include "TestNetSerializerFixture.h"
 #include "Iris/Serialization/StringNetSerializers.h"
+#include "Tests/ReplicationSystem/ReplicationSystemServerClientTestFixture.h"
 #include "Containers/StringConv.h"
+#include "Net/UnrealNetwork.h"
+#include "Iris/ReplicationSystem/ReplicationFragmentUtil.h"
+#include "Tests/ReplicationSystem/RPC/RPCTestFixture.h"
 
 namespace UE::Net::Private
 {
@@ -23,12 +28,12 @@ public:
 	void TestSerialize();
 	void TestCloneDynamicState();
 
-protected:
-	typedef TTestNetSerializerFixture<PrintNameNetSerializerConfig, FName> Super;
-
 	static const FName TestNames[];
 	static const SIZE_T TestNameCount;
 
+protected:
+	typedef TTestNetSerializerFixture<PrintNameNetSerializerConfig, FName> Super;
+
 	// Serializer
 	static FNameNetSerializerConfig SerializerConfig;
 };
@@ -55,6 +60,7 @@ const SIZE_T FTestNameNetSerializer::TestNameCount = sizeof(TestNames)/sizeof(Te
 
 FNameNetSerializerConfig FTestNameNetSerializer::SerializerConfig;
 
+// does not work as we need to setup export system as well.
 UE_NET_TEST_FIXTURE(FTestNameNetSerializer, TestValidate)
 {
 	TestValidate();
@@ -153,4 +159,252 @@ void FTestNameNetSerializer::TestCloneDynamicState()
 	}
 }
 
+UE_NET_TEST_FIXTURE(FReplicationSystemServerClientTestFixture, TestFName)
+{
+	// Add a client
+	FReplicationSystemTestClient* Client = CreateClient();
+
+	// Spawn objects on server
+	UTestNameNetSerializer_TestObject* ServerObject = Server->CreateObject<UTestNameNetSerializer_TestObject>();
+
+	Server->UpdateAndSend({Client});
+
+	// Verify that object has been spawned on client
+	auto ClientObject = Client->GetObjectAs<UTestNameNetSerializer_TestObject>(ServerObject->NetRefHandle);
+
+	// Verify that created server handle now also exists on client
+	UE_NET_ASSERT_NE(ClientObject, nullptr);	
+
+	// Verify that created server handle now also exists on client
+	UE_NET_ASSERT_EQ(ServerObject->NameProperty, ClientObject->NameProperty);
+
+	FName Modified("ModifiedName");
+
+	ServerObject->NameProperty = Modified;
+
+	Server->UpdateAndSend({Client});
+
+	// Verify that we managed to replicate the expected name
+	UE_NET_ASSERT_EQ(ServerObject->NameProperty, ClientObject->NameProperty);	
+}
+
+UE_NET_TEST_FIXTURE(FReplicationSystemServerClientTestFixture, TestFName_ReplicateCommonNames)
+{
+	// Add a client
+	FReplicationSystemTestClient* Client = CreateClient();
+
+	// Spawn objects on server
+	UTestNameNetSerializer_TestObject* ServerObject = Server->CreateObject<UTestNameNetSerializer_TestObject>();
+
+	Server->UpdateAndSend({Client});
+
+	// Verify that object has been spawned on client
+	auto ClientObject = Client->GetObjectAs<UTestNameNetSerializer_TestObject>(ServerObject->NetRefHandle);
+
+	// Verify that created server handle now also exists on client
+	UE_NET_ASSERT_NE(ClientObject, nullptr);	
+
+	// Send all testnames
+	for (const FName& Name : MakeArrayView<const FName>(FTestNameNetSerializer::TestNames, FTestNameNetSerializer::TestNameCount))
+	{
+		ServerObject->NameArrayProperty.Add(Name);
+	}
+
+	// Send a bunch of packets
+	for (uint32 Packets = 10; Packets > 0; Packets--)
+	{
+		Server->UpdateAndSend({Client});
+	}
+
+	UE_NET_ASSERT_TRUE(ServerObject->NameArrayProperty == ClientObject->NameArrayProperty);
+}
+
+
+UE_NET_TEST_FIXTURE(FRPCTestFixture, TestFNameClientRPCCanExportsName)
+{
+	// Add a client
+	FReplicationSystemTestClient* Client = CreateClient();
+
+	// Spawn object on server
+	UTestNameNetSerializer_TestObjectWithRPC* ServerObject = Server->CreateObject<UTestNameNetSerializer_TestObjectWithRPC>();
+
+	ServerObject->bIsServerObject = true;
+	ServerObject->ReplicationSystem = Server->GetReplicationSystem();
+	Server->ReplicationSystem->SetOwningNetConnection(ServerObject->NetRefHandle, 0x01);
+
+	// Send and deliver packet
+	Server->UpdateAndSend({Client});
+
+	UTestNameNetSerializer_TestObjectWithRPC* ClientObject = Cast<UTestNameNetSerializer_TestObjectWithRPC>(Client->GetReplicationBridge()->GetReplicatedObject(ServerObject->NetRefHandle));
+		
+	// Verify that created server handle now also exists on client
+	UE_NET_ASSERT_NE(ClientObject, nullptr);
+
+	FName ExpectedName("ExpectedClientName");
+	ServerObject->ClientRPCWithName(ExpectedName);
+	
+	// Send and deliver packet
+	Server->UpdateAndSend({Client});
+
+	// Verify that created server handle now also exists on client
+	UE_NET_ASSERT_EQ(ExpectedName, ClientObject->NameFromClientRPC);
+}
+
+UE_NET_TEST_FIXTURE(FRPCTestFixture, TestFNameServerRPCCanExportsName)
+{
+	// Add a client
+	FReplicationSystemTestClient* Client = CreateClient();
+
+	// Spawn object on server
+	UTestNameNetSerializer_TestObjectWithRPC* ServerObject = Server->CreateObject<UTestNameNetSerializer_TestObjectWithRPC>();
+
+	ServerObject->bIsServerObject = true;
+	ServerObject->ReplicationSystem = Server->GetReplicationSystem();
+	Server->ReplicationSystem->SetOwningNetConnection(ServerObject->NetRefHandle, 0x01);
+
+	// Send and deliver packet
+	Server->UpdateAndSend({Client});
+
+	UTestNameNetSerializer_TestObjectWithRPC* ClientObject = Cast<UTestNameNetSerializer_TestObjectWithRPC>(Client->GetReplicationBridge()->GetReplicatedObject(ServerObject->NetRefHandle));
+		
+	// Verify that created server handle now also exists on client
+	UE_NET_ASSERT_NE(ClientObject, nullptr);
+
+	ClientObject->ReplicationSystem = Client->GetReplicationSystem();
+
+	FName ExpectedName("ExpectedClientName");
+	ClientObject->ServerRPCWithName(ExpectedName);
+
+	// Send and deliver client packet
+	Client->UpdateAndSend(Server);
+	
+	// Verify that created server handle now also exists on client
+	UE_NET_ASSERT_EQ(ExpectedName, ServerObject->NameFromServerRPC);
+}
+
+}
+
+UTestNameNetSerializer_TestObject::UTestNameNetSerializer_TestObject()
+: UReplicatedTestObject()
+{
+}
+
+void UTestNameNetSerializer_TestObject::GetLifetimeReplicatedProps(TArray<class FLifetimeProperty>& OutLifetimeProps) const
+{
+	FDoRepLifetimeParams Params;
+	Params.bIsPushBased = false;
+
+	DOREPLIFETIME_WITH_PARAMS_FAST(UTestNameNetSerializer_TestObject, NameProperty, Params);
+	DOREPLIFETIME_WITH_PARAMS_FAST(UTestNameNetSerializer_TestObject, NameArrayProperty, Params);
+}
+
+void UTestNameNetSerializer_TestObject::RegisterReplicationFragments(UE::Net::FFragmentRegistrationContext& Context, UE::Net::EFragmentRegistrationFlags RegistrationFlags)
+{
+	UE::Net::FReplicationFragmentUtil::CreateAndRegisterFragmentsForObject(this, Context, RegistrationFlags);
+}
+
+//////////////////////////////////////////////////////////////////////////
+// Implementation for UTestNameNetSerializer_TestObjectWithRPC
+//////////////////////////////////////////////////////////////////////////
+UTestNameNetSerializer_TestObjectWithRPC::UTestNameNetSerializer_TestObjectWithRPC()
+	: UReplicatedTestObject()
+{
+}
+
+void UTestNameNetSerializer_TestObjectWithRPC::RegisterReplicationFragments(UE::Net::FFragmentRegistrationContext& Context, UE::Net::EFragmentRegistrationFlags RegistrationFlags)
+{
+	this->ReplicationFragments.Reset();
+	UE::Net::FReplicationFragmentUtil::CreateAndRegisterFragmentsForObject(this, Context, RegistrationFlags, &this->ReplicationFragments);
+}
+
+void UTestNameNetSerializer_TestObjectWithRPC::Init(UReplicationSystem* InRepSystem)
+{
+	bIsServerObject = InRepSystem->IsServer();
+	ReplicationSystem = InRepSystem;
+}
+
+void UTestNameNetSerializer_TestObjectWithRPC::SetRootObject(UTestNameNetSerializer_TestObjectWithRPC* InRootObject)
+{
+	check(InRootObject);
+	RootObject = InRootObject;
+}
+
+int32 UTestNameNetSerializer_TestObjectWithRPC::GetFunctionCallspace(UFunction* Function, FFrame* Stack)
+{
+	check(!(Function->FunctionFlags & FUNC_Static));
+	check(Function->FunctionFlags & FUNC_Net);
+
+	const bool bIsOnServer = bIsServerObject;
+
+	// get the top most function
+	while (Function->GetSuperFunction() != nullptr)
+	{
+		Function = Function->GetSuperFunction();
+	}
+
+	// Multicast RPCs
+	if ((Function->FunctionFlags & FUNC_NetMulticast))
+	{
+		if (bIsOnServer)
+		{
+			// Server should execute locally and call remotely
+			return (FunctionCallspace::Local | FunctionCallspace::Remote);
+		}
+		else
+		{
+			return FunctionCallspace::Local;
+		}
+	}
+
+	// if we are the authority
+	if (bIsOnServer)
+	{
+		if (Function->FunctionFlags & FUNC_NetClient)
+		{
+			return FunctionCallspace::Remote;
+		}
+		else
+		{
+			return FunctionCallspace::Local;
+		}
+
+	}
+	// if we are not the authority
+	else
+	{
+		if (Function->FunctionFlags & FUNC_NetServer)
+		{
+			return FunctionCallspace::Remote;
+		}
+		else
+		{
+			// don't replicate
+			return FunctionCallspace::Local;
+		}
+	}
+}
+
+bool UTestNameNetSerializer_TestObjectWithRPC::CallRemoteFunction(UFunction* Function, void* Parameters, FOutParmRec* OutParms, FFrame* Stack)
+{
+	if (bIsSubObject)
+	{
+		return ReplicationSystem->SendRPC(RootObject, this, Function, Parameters);
+	}
+	else
+	{
+		return ReplicationSystem->SendRPC(this, nullptr, Function, Parameters);
+	}
+}
+
+void UTestNameNetSerializer_TestObjectWithRPC::ClientRPCWithName_Implementation(FName Name)
+{
+	NameFromClientRPC = Name;
 }
+
+
+void UTestNameNetSerializer_TestObjectWithRPC::ServerRPCWithName_Implementation(FName Name)
+{
+	NameFromServerRPC = Name;
+}
+
+
diff --git a/Engine/Plugins/Runtime/ReplicationSystemTestPlugin/Source/Private/Tests/Serialization/TestNameNetSerializer.h b/Engine/Plugins/Runtime/ReplicationSystemTestPlugin/Source/Private/Tests/Serialization/TestNameNetSerializer.h
new file mode 100644
--- /dev/null
+++ b/Engine/Plugins/Runtime/ReplicationSystemTestPlugin/Source/Private/Tests/Serialization/TestNameNetSerializer.h
@@ -0,0 +1,62 @@
+// Copyright Epic Games, Inc. All Rights Reserved.
+
+#pragma once
+
+#include "UObject/ObjectMacros.h"
+#include "Tests/ReplicationSystem/ReplicatedTestObject.h"
+#include "TestNameNetSerializer.generated.h"
+
+UCLASS()
+class UTestNameNetSerializer_TestObject : public UReplicatedTestObject
+{
+	GENERATED_BODY()
+public:
+	UTestNameNetSerializer_TestObject();
+
+	void RegisterReplicationFragments(UE::Net::FFragmentRegistrationContext& Fragments, UE::Net::EFragmentRegistrationFlags RegistrationFlags) override;
+
+	UPROPERTY(Replicated)
+	FName NameProperty = FName("Initial");
+
+	UPROPERTY(Replicated)
+	TArray<FName> NameArrayProperty;
+};
+
+UCLASS()
+class UTestNameNetSerializer_TestObjectWithRPC : public UReplicatedTestObject
+{
+	GENERATED_BODY()
+
+public:
+	UTestNameNetSerializer_TestObjectWithRPC();
+
+	void Init(UReplicationSystem* InRepSystem);
+	void SetRootObject(UTestNameNetSerializer_TestObjectWithRPC* InRootObject);
+
+	virtual void RegisterReplicationFragments(UE::Net::FFragmentRegistrationContext& Context, UE::Net::EFragmentRegistrationFlags RegistrationFlags) override;
+
+	virtual int32 GetFunctionCallspace(UFunction* Function, FFrame* Stack) override;
+	virtual bool CallRemoteFunction(UFunction* Function, void* Parameters, FOutParmRec* OutParms, FFrame* Stack) override;
+
+	UReplicationSystem* ReplicationSystem = nullptr;
+
+	// Network data only for test
+	TArray<UE::Net::FReplicationFragment*> ReplicationFragments;
+
+	// To determine if this object is located on the server or client
+	bool bIsServerObject = false;
+
+	// Our owner when the object is a subobject
+	UTestNameNetSerializer_TestObjectWithRPC* RootObject = nullptr;
+
+	// RPC test functions
+	UFUNCTION(Reliable, Client)
+	void ClientRPCWithName(FName Name);
+	FName NameFromClientRPC;
+
+	UFUNCTION(Reliable, Server)
+	void ServerRPCWithName(FName Name);
+	FName NameFromServerRPC;
+};
+
+
diff --git a/Engine/Source/Runtime/Experimental/Iris/Core/Private/Iris/ReplicationSystem/NameTokenStore.cpp b/Engine/Source/Runtime/Experimental/Iris/Core/Private/Iris/ReplicationSystem/NameTokenStore.cpp
new file mode 100644
--- /dev/null
+++ b/Engine/Source/Runtime/Experimental/Iris/Core/Private/Iris/ReplicationSystem/NameTokenStore.cpp
@@ -0,0 +1,126 @@
+// Copyright Epic Games, Inc. All Rights Reserved.
+
+#include "Iris/ReplicationSystem/NameTokenStore.h"
+#include "Iris/Core/IrisLog.h"
+#include "Iris/Serialization/NetBitStreamReader.h"
+#include "Iris/Serialization/NetBitStreamWriter.h"
+#include "Iris/Serialization/NetBitStreamUtil.h"
+#include "Iris/Serialization/NetSerializationContext.h"
+#include "Hash/CityHash.h"
+#include "Net/Core/Trace/NetTrace.h"
+
+#define UE_NET_ENABLE_FNAME_TOKEN_LOG 1
+
+#if UE_NET_ENABLE_FNAME_TOKEN_LOG
+#	define UE_LOG_FNAMETOKEN(Format, ...)  UE_LOG(LogIris, Verbose, Format, ##__VA_ARGS__)
+#else
+#	define UE_LOG_FNAMETOKEN(...)
+#endif
+
+#define UE_LOG_FNAMETOKEN_WARNING(Format, ...)  UE_LOG(LogIris, Warning, Format, ##__VA_ARGS__)
+
+namespace UE::Net
+{
+
+FNetToken FNameTokenStore::GetOrCreateToken(FName Name)
+{
+	FNetTokenStoreKey Key = GetOrCreateTokenStoreKey(Name);
+	if (Key.IsValid())
+	{
+		const uint32 NameIndex = Key.GetKeyValue();
+		if (StoredTokens[NameIndex].IsValid())
+		{
+			return StoredTokens[NameIndex];
+		}
+		else
+		{
+			const FNetToken NewToken = CreateToken(Key, *TokenStore.GetLocalNetTokenStoreState());
+			StoredTokens[NameIndex] = NewToken;
+
+			UE_LOG_FNAMETOKEN(TEXT("FNameTokenStore::GetOrCreateToken - Created new token %u for %s"), NewToken.GetIndex(), *Name.ToString());
+
+			return NewToken;
+		}
+	}
+
+	return FNetToken();
+}
+
+FNetTokenStoreKey FNameTokenStore::GetOrCreateTokenStoreKey(FName Name)
+{
+	// Lock if we have to.. make that a policy
+	//FScopeLock Lock(&CriticalSection);
+	if (const FNetTokenStoreKey* ExistingKey = FNameToKey.Find(Name))
+	{
+		return *ExistingKey;
+	}
+	else if (FNameToKey.Num() < FNetToken::MaxNetTokenCount)
+	{
+		// Store bookkeeping data
+		const FNetTokenStoreKey NewKey = MakeNetTokenStoreKey(GetTypeId(), StoredFNames.Num());
+
+		FNameToKey.Add(Name, NewKey);
+		StoredFNames.Add(Name);
+		StoredTokens.Add(FNetToken());
+
+		return NewKey;
+	}
+
+	return FNetTokenStoreKey();
+}
+
+FNameTokenStore::FNameTokenStore(FNetTokenStore& InTokenStore)
+: TokenStore(InTokenStore)
+{
+	// Reserve 0
+	StoredFNames.Add(FName());
+	StoredTokens.Add(FNetToken());
+
+	// Register
+	TokenStore.RegisterDataStore(this);
+}
+
+FName FNameTokenStore::ResolveToken(FNetToken Token) const
+{
+	return ResolveRemoteToken(Token, *TokenStore.GetLocalNetTokenStoreState());
+}
+
+FName FNameTokenStore::ResolveRemoteToken(FNetToken Token, const FNetTokenStoreState& NetTokenStoreState) const
+{
+	const FNetTokenStoreKey StoreKey = GetTokenKey(Token, NetTokenStoreState);
+
+	if (StoreKey.IsValid() && StoreKey.GetKeyValue() < (uint32)StoredFNames.Num())
+	{
+		return StoredFNames[StoreKey.GetKeyValue()];
+	}
+
+	return FName();
+}
+
+void FNameTokenStore::WriteTokenData(FNetSerializationContext& Context, FNetTokenStoreKey TokenStoreKey) const
+{
+	//UE_NET_TRACE_DYNAMIC_NAME_SCOPE(StoredFNames[TokenStoreKey.GetKeyValue()], *Context.GetBitStreamWriter(), Context.GetTraceCollector(), ENetTraceVerbosity::VeryVerbose);
+	// $TODO: $IRIS: We can be a bit smarter here and utilize the string-number split of FNames to export less data.. JIRA: UE-221753
+	WriteString(Context.GetBitStreamWriter(), StoredFNames[TokenStoreKey.GetKeyValue()].ToString());
+}
+
+FNetTokenStoreKey FNameTokenStore::ReadTokenData(FNetSerializationContext& Context)
+{
+	FNetBitStreamReader* Reader = Context.GetBitStreamReader();
+
+	// Read the token data and add it to the store without assigning LocalToken
+	FString Temp;
+	ReadString(Reader, Temp);
+
+	if (!Reader->IsOverflown())
+	{
+		FName Name(Temp);
+		return GetOrCreateTokenStoreKey(Name);
+	}
+	else
+	{
+		return FNetTokenStoreKey();
+	}
+}
+
+}
diff --git a/Engine/Source/Runtime/Experimental/Iris/Core/Private/Iris/ReplicationSystem/ReplicationSystem.cpp b/Engine/Source/Runtime/Experimental/Iris/Core/Private/Iris/ReplicationSystem/ReplicationSystem.cpp
--- a/Engine/Source/Runtime/Experimental/Iris/Core/Private/Iris/ReplicationSystem/ReplicationSystem.cpp
+++ b/Engine/Source/Runtime/Experimental/Iris/Core/Private/Iris/ReplicationSystem/ReplicationSystem.cpp
@@ -1026,6 +1026,16 @@ UE::Net::FStringTokenStore* UReplicationSystem::GetStringTokenStore()
 	return &Impl->ReplicationSystemInternal.GetStringTokenStore();
 }
 
+const UE::Net::FNameTokenStore* UReplicationSystem::GetNameTokenStore() const
+{
+	return &Impl->ReplicationSystemInternal.GetNameTokenStore();
+}
+
+UE::Net::FNameTokenStore* UReplicationSystem::GetNameTokenStore()
+{
+	return &Impl->ReplicationSystemInternal.GetNameTokenStore();
+}
+
 bool UReplicationSystem::RegisterNetBlobHandler(UNetBlobHandler* Handler)
 {
 	UE::Net::Private::FNetBlobManager& NetBlobManager = Impl->ReplicationSystemInternal.GetNetBlobManager();
diff --git a/Engine/Source/Runtime/Experimental/Iris/Core/Private/Iris/ReplicationSystem/ReplicationSystemInternal.h b/Engine/Source/Runtime/Experimental/Iris/Core/Private/Iris/ReplicationSystem/ReplicationSystemInternal.h
--- a/Engine/Source/Runtime/Experimental/Iris/Core/Private/Iris/ReplicationSystem/ReplicationSystemInternal.h
+++ b/Engine/Source/Runtime/Experimental/Iris/Core/Private/Iris/ReplicationSystem/ReplicationSystemInternal.h
@@ -22,6 +22,7 @@
 #include "Iris/ReplicationSystem/NetBlob/NetBlobManager.h"
 #include "Iris/ReplicationSystem/NetTokenStore.h"
 #include "Iris/ReplicationSystem/StringTokenStore.h"
+#include "Iris/ReplicationSystem/NameTokenStore.h"
 #include "Iris/ReplicationSystem/WorldLocations.h"
 #include "Iris/ReplicationState/ReplicationStateDescriptorRegistry.h"
 #include "Iris/Stats/NetStats.h"
@@ -49,6 +50,7 @@ public:
 	, ReplicationBridge(nullptr)
 	, IrisObjectReferencePackageMap(nullptr)
 	, StringTokenStore(NetTokenStore)
+	, NameTokenStore(NetTokenStore)
 	, Id(Params.ReplicationSystemId)
 	{}
 
@@ -94,6 +96,9 @@ public:
 	const FStringTokenStore& GetStringTokenStore() const { return StringTokenStore; }
 	FStringTokenStore& GetStringTokenStore() { return StringTokenStore; }
 
+	const FNameTokenStore& GetNameTokenStore() const { return NameTokenStore; }
+	FNameTokenStore& GetNameTokenStore() { return NameTokenStore; }
+
 	FNetTokenStore& GetNetTokenStore() { return NetTokenStore; }
 
 	FNetCullDistanceOverrides& GetNetCullDistanceOverrides() { return NetCullDistanceOverrides; }
@@ -140,6 +145,7 @@ private:
 	FNetBlobManager NetBlobManager;
 	FNetTokenStore NetTokenStore;
 	FStringTokenStore StringTokenStore;
+	FNameTokenStore NameTokenStore;
 	FNetCullDistanceOverrides NetCullDistanceOverrides;
 	FWorldLocations WorldLocations;
 	FDeltaCompressionBaselineManager DeltaCompressionBaselineManager;
diff --git a/Engine/Source/Runtime/Experimental/Iris/Core/Private/Iris/ReplicationSystem/StringTokenStore.cpp b/Engine/Source/Runtime/Experimental/Iris/Core/Private/Iris/ReplicationSystem/StringTokenStore.cpp
--- a/Engine/Source/Runtime/Experimental/Iris/Core/Private/Iris/ReplicationSystem/StringTokenStore.cpp
+++ b/Engine/Source/Runtime/Experimental/Iris/Core/Private/Iris/ReplicationSystem/StringTokenStore.cpp
@@ -37,7 +37,7 @@ FNetToken FStringTokenStore::GetOrCreateToken(const TCHAR* Name, uint32 Length)
 		}
 		else
 		{
-			const FNetToken NewToken = CreateToken(MakeNetTokenStoreKey(GetTypeId(), StringIndex), *TokenStore.GetLocalNetTokenStoreState());
+			const FNetToken NewToken = CreateToken(Key, *TokenStore.GetLocalNetTokenStoreState());
 			StoredTokens[StringIndex] = NewToken;
 
 			UE_LOG_STRINGTOKEN(TEXT("FStringTokenStore::GetOrCreateToken - Created new token %u for %s"), NewToken.GetIndex(), Name);
diff --git a/Engine/Source/Runtime/Experimental/Iris/Core/Private/Iris/Serialization/StringNetSerializers.cpp b/Engine/Source/Runtime/Experimental/Iris/Core/Private/Iris/Serialization/StringNetSerializers.cpp
--- a/Engine/Source/Runtime/Experimental/Iris/Core/Private/Iris/Serialization/StringNetSerializers.cpp
+++ b/Engine/Source/Runtime/Experimental/Iris/Core/Private/Iris/Serialization/StringNetSerializers.cpp
@@ -9,6 +9,12 @@
 #include "Iris/Core/BitTwiddling.h"
 #include "Containers/StringConv.h"
 #include "GenericPlatform/GenericPlatformMemory.h"
+#include "Iris/ReplicationSystem/NameTokenStore.h"
+#include "Iris/ReplicationSystem/ReplicationSystem.h"
+#include "Iris/ReplicationSystem/ReplicationSystemInternal.h"
+#include "Iris/Serialization/InternalNetSerializationContext.h"
+#include "Iris/Serialization/NetExportContext.h"
+#include "Net/Core/Trace/NetTrace.h"
 
 static_assert(sizeof(ANSICHAR) == sizeof(uint8), "ANSICHAR is expected to be one byte.");
 
@@ -414,4 +420,300 @@ bool FStringNetSerializer::Validate(FNetSerializationContext& Context, const FNe
 	return FStringNetSerializerBase::Validate(Context, Args, Source);
 }
 
+// FNameAsNetTokenNetSerializer
+// Naive implementation ignoring numbers so it will reexport name part for each number.
+// Good thing is that it works for all different configurations of FNames
+struct FNameAsNetTokenNetSerializerQuantizedType
+{
+	// NetToken
+	FNetToken NetToken;
+
+	// If bIsString this is the Number, otherwise it's the EName
+	int32 ENameOrNumber;
+
+	// If the FName is hardcoded or not
+	uint32 bIsString : 1U;
+	// When serializing/deserializing the number is expressed as (MAX_int32 - number)
+	uint32 bEncodeNumberFromIntMax : 1U;
+
+	// if this is a local (obtained from quantize)
+	uint32 bIsLocal : 1U;
+};
+
+} // End of namespace UE::Net
+
+template<> struct TIsPODType<UE::Net::FNameAsNetTokenNetSerializerQuantizedType> { enum { Value = true }; };
+
+namespace UE::Net
+{
+
+struct FNameAsNetTokenNetSerializer
+{
+	// Version
+	static const uint32 Version = 0;
+
+	// Types
+
+	// Crafted such that zeroed memory will represent FName(NAME_None)
+	typedef FNameAsNetTokenNetSerializerQuantizedType FQuantizedType;
+	typedef FName SourceType;
+	typedef FQuantizedType QuantizedType;
+	typedef FNameAsNetTokenNetSerializerConfig ConfigType;
+
+	static const ConfigType DefaultConfig;
+
+	//
+	static void Serialize(FNetSerializationContext&, const FNetSerializeArgs& Args);
+	static void Deserialize(FNetSerializationContext&, const FNetDeserializeArgs& Args);
+
+	static void Quantize(FNetSerializationContext&, const FNetQuantizeArgs& Args);
+	static void Dequantize(FNetSerializationContext&, const FNetDequantizeArgs& Args);
+
+	static bool IsEqual(FNetSerializationContext&, const FNetIsEqualArgs& Args);
+	static bool Validate(FNetSerializationContext&, const FNetValidateArgs& Args);
+
+private:
+
+	// Utility methods, consolidate with other changes to NetTokenStore as next step.
+	static FName ResolveNetToken(FNetSerializationContext&, FNetToken NetToken, bool bIsLocal);
+
+	static void AppendExportOrWriteInlinedExportData(FNetSerializationContext&, FNetToken NetToken);
+	static void ReadInlinedExportData(FNetSerializationContext&, FNetToken NetToken);
+
+	static const uint32 BitCountNeededForEName;
+};
+
+UE_NET_IMPLEMENT_SERIALIZER_INTERNAL(FNameAsNetTokenNetSerializer);
+const FNameAsNetTokenNetSerializer::ConfigType FNameAsNetTokenNetSerializer::DefaultConfig;
+const uint32 FNameAsNetTokenNetSerializer::BitCountNeededForEName = UE::Net::GetBitsNeeded(MAX_NETWORKED_HARDCODED_NAME);
+
+void FNameAsNetTokenNetSerializer::Serialize(FNetSerializationContext& Context, const FNetSerializeArgs& Args)
+{
+	using namespace UE::Net::Private;
+
+	// For now we ignore this in default state hash
+	// We could probably output lowercase hash to allow validation of defaultstate.
+	if (Context.IsInitializingDefaultState())
+	{
+		return;
+	}
+
+	const QuantizedType& Value = *reinterpret_cast<QuantizedType*>(Args.Source);
+	
+	FNetBitStreamWriter* Writer = Context.GetBitStreamWriter();
+	if (Writer->WriteBool(Value.bIsString))
+	{
+		//UE_NET_TRACE_DYNAMIC_NAME_SCOPE(*FString::Printf(TEXT("%u"), Value.NetToken.GetIndex()), *Context.GetBitStreamWriter(), Context.GetTraceCollector(), ENetTraceVerbosity::VeryVerbose);
+
+		// Always write the token
+		WriteNetToken(Writer, Value.NetToken);
+
+		// Export or add to pending exports for later export
+		FNameAsNetTokenNetSerializer::AppendExportOrWriteInlinedExportData(Context, Value.NetToken);
+	}
+	else
+	{
+		Writer->WriteBits(static_cast<uint32>(Value.ENameOrNumber), BitCountNeededForEName);
+	}
+}
+
+void FNameAsNetTokenNetSerializer::Deserialize(FNetSerializationContext& Context, const FNetDeserializeArgs& Args)
+{
+	using namespace Private;
+
+	// Unexpected, but consistent with Serialize.
+	if (Context.IsInitializingDefaultState())
+	{
+		return;
+	}
+
+	QuantizedType& Target = *reinterpret_cast<QuantizedType*>(Args.Target);
+
+	FNetBitStreamReader* Reader = Context.GetBitStreamReader();
+	if (const bool bIsString = Reader->ReadBool())
+	{
+		Target.bIsString = 1;
+
+		// Always Read the token
+		FNetToken NetToken = ReadNetToken(Reader);
+
+		if (Reader->IsOverflown())
+		{
+			return;
+		}
+
+		// Read inlined exports if there are any
+		ReadInlinedExportData(Context, NetToken);
+
+		Target.NetToken = NetToken;
+		Target.bIsLocal = 0U;
+	}
+	else
+	{
+		const uint32 ENameNumber = Reader->ReadBits(BitCountNeededForEName);
+		if (!ShouldReplicateAsInteger(EName(ENameNumber), FName(EName(ENameNumber))))
+		{
+			Context.SetError(GNetError_BitStreamError);
+			return;
+		}
+
+		Target.bIsString = 0U;
+		Target.bEncodeNumberFromIntMax = 0;
+		Target.ENameOrNumber = ENameNumber;
+	
+		Target.NetToken = FNetToken();
+		Target.bIsLocal = 0U;
+	}
+}
+
+void FNameAsNetTokenNetSerializer::Quantize(FNetSerializationContext& Context, const FNetQuantizeArgs& Args)
+{
+	using namespace Private;
+
+	const SourceType SourceName = *reinterpret_cast<const SourceType*>(Args.Source);
+	QuantizedType& TargetName = *reinterpret_cast<QuantizedType*>(Args.Target);
+
+	const EName* AsEName = (SourceName.GetNumber() == NAME_NO_NUMBER_INTERNAL ? SourceName.ToEName() : nullptr);
+	const bool bIsString = (AsEName == nullptr || !ShouldReplicateAsInteger(*AsEName, SourceName));
+	if (bIsString)
+	{
+		TargetName.bIsString = 1;
+		TargetName.ENameOrNumber = 0;
+		TargetName.bEncodeNumberFromIntMax = 0;
+
+		// Store as NetToken
+		FInternalNetSerializationContext* InternalContext = Context.GetInternalContext();
+		using namespace UE::Net::Private;
+			
+		TargetName.NetToken = InternalContext->ReplicationSystem->GetNameTokenStore()->GetOrCreateToken(SourceName);
+		TargetName.bIsLocal = 1U;
+	}
+	else
+	{
+		TargetName.bIsString = 0;
+		// Again, this value doesn't matter for hardcoded names as we know they will start from 0
+		TargetName.bEncodeNumberFromIntMax = 0;
+		// The EName we do care about!
+		TargetName.ENameOrNumber = static_cast<int32>(static_cast<uint32>(*AsEName));
+
+		TargetName.bIsLocal = 0;
+		TargetName.NetToken = FNetToken();
+	}
+}
+
+FName FNameAsNetTokenNetSerializer::ResolveNetToken(FNetSerializationContext& Context, FNetToken NetToken, bool bIsLocal)
+{
+	using namespace UE::Net::Private;
+
+	FName Result;
+	FInternalNetSerializationContext* InternalContext = Context.GetInternalContext();
+	FNameTokenStore* NameTokenStore = InternalContext->ReplicationSystem->GetNameTokenStore();
+
+	return bIsLocal ? NameTokenStore->ResolveToken(NetToken) : NameTokenStore->ResolveRemoteToken(NetToken, *InternalContext->ResolveContext.RemoteNetTokenStoreState);
+}
+
+void FNameAsNetTokenNetSerializer::AppendExportOrWriteInlinedExportData(FNetSerializationContext& Context, FNetToken NetToken)
+{
+	using namespace UE::Net::Private;
+
+	FNetExportContext* ExportContext = Context.GetExportContext();
+	const FInternalNetSerializationContext* InternalContext = Context.GetInternalContext();
+	
+	if (InternalContext->bInlineObjectReferenceExports == 0U)
+	{
+		ExportContext->AddPendingExport(NetToken);
+	}
+	else
+	{
+		FNetTokenStore& NetTokenStore = InternalContext->ReplicationSystem->GetReplicationSystemInternal()->GetNetTokenStore();
+
+		NetTokenStore.ConditionalWriteNetTokenData(Context, ExportContext, NetToken);
+	}
+}
+
+void FNameAsNetTokenNetSerializer::ReadInlinedExportData(FNetSerializationContext& Context, FNetToken NetToken)
+{
+	using namespace UE::Net::Private;
+
+	const FInternalNetSerializationContext* InternalContext = Context.GetInternalContext();
+
+	// Conditionally import if using inlined exports
+	if (InternalContext->bInlineObjectReferenceExports == 1U)
+	{
+		FNetTokenStore& NetTokenStore = InternalContext->ReplicationSystem->GetReplicationSystemInternal()->GetNetTokenStore();
+		NetTokenStore.ConditionalReadNetTokenData(Context, NetToken);
+	}
+}
+
+void FNameAsNetTokenNetSerializer::Dequantize(FNetSerializationContext& Context, const FNetDequantizeArgs& Args)
+{
+	using namespace Private;
+
+	const QuantizedType& Source = *reinterpret_cast<const QuantizedType*>(Args.Source);
+	SourceType& Target = *reinterpret_cast<SourceType*>(Args.Target);
+
+	if (Source.bIsString)
+	{
+		// Resolve from RemoteNetToken
+		Target = FNameAsNetTokenNetSerializer::ResolveNetToken(Context, Source.NetToken, Source.bIsLocal);
+	}
+	else
+	{
+		Target = EName(static_cast<uint32>(Source.ENameOrNumber));
+	}
+}
+
+bool FNameAsNetTokenNetSerializer::IsEqual(FNetSerializationContext& Context, const FNetIsEqualArgs& Args)
+{
+	using namespace Private;
+
+	if (Args.bStateIsQuantized)
+	{
+		const QuantizedType& Value0 = *reinterpret_cast<const QuantizedType*>(Args.Source0);
+		const QuantizedType& Value1 = *reinterpret_cast<const QuantizedType*>(Args.Source1);
+
+		// Bits will cancel out when values are equal.
+		const bool bIsEqual = ((Value0.bIsString ^ Value1.bIsString) | (Value0.ENameOrNumber ^ Value1.ENameOrNumber)) == 0;
+		if (!bIsEqual)
+		{
+			return false;
+		}
+
+		// We are pessimistic here
+		if (Value0.bIsString)
+		{
+			if (Value0.bIsLocal != Value1.bIsLocal)
+			{
+				const FName Name0 = FNameAsNetTokenNetSerializer::ResolveNetToken(Context, Value0.NetToken, Value0.bIsLocal);
+				const FName Name1 = FNameAsNetTokenNetSerializer::ResolveNetToken(Context, Value1.NetToken, Value1.bIsLocal);;
+				// Need to compare actual FNames to handle compare of local and remote data.
+				if (Name0 != Name1)
+				{
+					return false;
+				}
+			}
+			else if (Value0.NetToken != Value1.NetToken)
+			{
+				return false;
+			}
+		}
+
+		return true;
+	}
+	else
+	{
+		const SourceType& Value0 = *reinterpret_cast<const SourceType*>(Args.Source0);
+		const SourceType& Value1 = *reinterpret_cast<const SourceType*>(Args.Source1);
+		const bool bIsEqual = (Value0 == Value1);
+		return bIsEqual;
+	}
+}
+
+bool FNameAsNetTokenNetSerializer::Validate(FNetSerializationContext& Context, const FNetValidateArgs& Args)
+{
+	const SourceType& Value = *reinterpret_cast<const SourceType*>(Args.Source);
+	const bool bIsValid = Value.IsValid();
+	return bIsValid;
+}
+
 }
diff --git a/Engine/Source/Runtime/Experimental/Iris/Core/Public/Iris/ReplicationSystem/NameTokenStore.h b/Engine/Source/Runtime/Experimental/Iris/Core/Public/Iris/ReplicationSystem/NameTokenStore.h
new file mode 100644
--- /dev/null
+++ b/Engine/Source/Runtime/Experimental/Iris/Core/Public/Iris/ReplicationSystem/NameTokenStore.h
@@ -0,0 +1,47 @@
+// Copyright Epic Games, Inc. All Rights Reserved.
+
+#pragma once
+
+#include "CoreTypes.h"
+#include "Containers/Map.h"
+#include "Iris/ReplicationSystem/NetTokenStore.h"
+#include "Misc/MemStack.h"
+
+namespace UE::Net
+{
+
+// Simple token store used to store string tokens
+// When the PackageMapRefactor is complete we will most likely rely on NetTagManager for persistent storage
+class FNameTokenStore : public FNetTokenDataStore
+{
+	UE_NONCOPYABLE(FNameTokenStore);
+public:
+	explicit FNameTokenStore(FNetTokenStore& TokenStore);
+
+	// Create a string token for the provided string
+	IRISCORE_API FNetToken GetOrCreateToken(FName Name);
+
+	// Resolve a local token
+	IRISCORE_API FName ResolveToken(FNetToken Token) const;
+
+	// Resolve a token received from remote
+	IRISCORE_API FName ResolveRemoteToken(FNetToken Token, const FNetTokenStoreState& NetTokenStoreState) const;
+
+protected:
+	// Serialize data for a token, note there is not validation in this function
+	virtual void WriteTokenData(FNetSerializationContext& Context, FNetTokenStoreKey TokenStoreKey) const override;
+
+	// Read data for a token, returns a valid StoreKey if successful read
+	virtual FNetTokenStoreKey ReadTokenData(FNetSerializationContext& Context) override;
+
+	// Create a persistent string
+	FNetTokenStoreKey GetOrCreateTokenStoreKey(FName Name);
+
+private:
+	FNetTokenStore& TokenStore;
+	TMap<FName, FNetTokenStoreKey> FNameToKey;
+	TArray<FName> StoredFNames;
+	TArray<FNetToken> StoredTokens;
+};
+
+}
diff --git a/Engine/Source/Runtime/Experimental/Iris/Core/Public/Iris/ReplicationSystem/ReplicationSystem.h b/Engine/Source/Runtime/Experimental/Iris/Core/Public/Iris/ReplicationSystem/ReplicationSystem.h
--- a/Engine/Source/Runtime/Experimental/Iris/Core/Public/Iris/ReplicationSystem/ReplicationSystem.h
+++ b/Engine/Source/Runtime/Experimental/Iris/Core/Public/Iris/ReplicationSystem/ReplicationSystem.h
@@ -33,6 +33,7 @@ namespace UE::Net
 	class FNetObjectReference;
 	class FReplicationSystemFactory;
 	class FStringTokenStore;
+	class FNameTokenStore;
 	class FWorldLocations;
 
 	struct FReplicationSystemUtil;
@@ -340,6 +341,19 @@ public:
 	 */
 	IRISCORE_API UE::Net::FStringTokenStore* GetStringTokenStore();
 
+	/**
+	 * @return A const version of the FName token store.
+	 * @see UE::Net::FFNameTokenStore
+	 */
+	IRISCORE_API const UE::Net::FNameTokenStore* GetNameTokenStore() const;
+
+	/**
+	 * @return The string token store.
+	 * @see UE::Net::FNameTokenStore
+	 */
+	IRISCORE_API UE::Net::FNameTokenStore* GetNameTokenStore();
+
+
 	/**
 	 * Check whether a FNetRefHandle is still associated with a replicated object.
 	 * @param Handle The handle check.
diff --git a/Engine/Source/Runtime/Experimental/Iris/Core/Public/Iris/Serialization/StringNetSerializers.h b/Engine/Source/Runtime/Experimental/Iris/Core/Public/Iris/Serialization/StringNetSerializers.h
--- a/Engine/Source/Runtime/Experimental/Iris/Core/Public/Iris/Serialization/StringNetSerializers.h
+++ b/Engine/Source/Runtime/Experimental/Iris/Core/Public/Iris/Serialization/StringNetSerializers.h
@@ -12,6 +12,12 @@ struct FNameNetSerializerConfig : public FNetSerializerConfig
 	GENERATED_BODY()
 };
 
+USTRUCT()
+struct FNameAsNetTokenNetSerializerConfig : public FNetSerializerConfig
+{
+	GENERATED_BODY()
+};
+
 USTRUCT()
 struct FStringNetSerializerConfig : public FNetSerializerConfig
 {
@@ -22,6 +28,7 @@ namespace UE::Net
 {
 
 UE_NET_DECLARE_SERIALIZER(FNameNetSerializer, IRISCORE_API)
+UE_NET_DECLARE_SERIALIZER(FNameAsNetTokenNetSerializer, IRISCORE_API)
 UE_NET_DECLARE_SERIALIZER(FStringNetSerializer, IRISCORE_API)
 
 }
