Iris
-Enabled FNameAsNetTokenNetSerializer for fnames

diff --git a/Engine/Source/Runtime/Experimental/Iris/Core/Private/Iris/ReplicationState/DefaultPropertyNetSerializerInfos.cpp b/Engine/Source/Runtime/Experimental/Iris/Core/Private/Iris/ReplicationState/DefaultPropertyNetSerializerInfos.cpp
--- a/Engine/Source/Runtime/Experimental/Iris/Core/Private/Iris/ReplicationState/DefaultPropertyNetSerializerInfos.cpp
+++ b/Engine/Source/Runtime/Experimental/Iris/Core/Private/Iris/ReplicationState/DefaultPropertyNetSerializerInfos.cpp
@@ -41,7 +41,9 @@ UE_NET_IMPLEMENT_NAMED_STRUCT_NETSERIALIZER_INFO(PropertyNetSerializerRegistry_N
 
 // Strings
 UE_NET_IMPLEMENT_SIMPLE_NETSERIALIZER_INFO(FStrProperty, FStringNetSerializer);
-UE_NET_IMPLEMENT_SIMPLE_NETSERIALIZER_INFO(FNameProperty, FNameNetSerializer);
+//UE_NET_IMPLEMENT_SIMPLE_NETSERIALIZER_INFO(FNameProperty, FNameNetSerializer);
+// Use NetTokens instead of strings when serializing FNames
+UE_NET_IMPLEMENT_SIMPLE_NETSERIALIZER_INFO(FNameProperty, FNameAsNetTokenNetSerializer);
 
 // Named structs with specific serializers
 static const FName PropertyNetSerializerRegistry_NAME_Guid("Guid");
