Iris
* Downgraded FNetProfilerObjectInstance to a 32bit hash to save bandwidth when sending the FCreationHeader of new replicated objects.

#rnx

diff --git a/Engine/Source/Runtime/Experimental/Iris/Core/Private/Iris/Core/IrisDebugging.cpp b/Engine/Source/Runtime/Experimental/Iris/Core/Private/Iris/Core/IrisDebugging.cpp
--- a/Engine/Source/Runtime/Experimental/Iris/Core/Private/Iris/Core/IrisDebugging.cpp
+++ b/Engine/Source/Runtime/Experimental/Iris/Core/Private/Iris/Core/IrisDebugging.cpp
@@ -494,7 +494,7 @@ FNetReplicatedObjectDebugInfo DebugInternalNetRefIndex(uint32 InternalIndex, uin
 	return Info;
 }
 
-void NetObjectProtocolReferencesToString(FStringBuilderBase& StringBuilder, uint64 ProtocolId, uint32 ReplicationSystemId)
+void NetObjectProtocolReferencesToString(FStringBuilderBase& StringBuilder, FReplicationProtocolIdentifier ProtocolId, uint32 ReplicationSystemId)
 {
 	using namespace UE::Net::Private;
 
@@ -512,7 +512,7 @@ void NetObjectProtocolReferencesToString(FStringBuilderBase& StringBuilder, uint
 	ProtocolManager.ForEachProtocol(ProtocolId, [&StringBuilder, &NetRefHandleManager](const FReplicationProtocol* Protocol, const UObject* ArchetypeOrCDOUsedAsKey)
 		{
 			StringBuilder << TEXT("Protocol: ") << ToCStr(Protocol->DebugName);
-			StringBuilder.Appendf(TEXT("Id: 0x%" UINT64_x_FMT " Created From : 0x%p"), Protocol->ProtocolIdentifier, ArchetypeOrCDOUsedAsKey) << TEXT(" Used by : \n");
+			StringBuilder.Appendf(TEXT("Id: 0x%x Created From : 0x%p"), Protocol->ProtocolIdentifier, ArchetypeOrCDOUsedAsKey) << TEXT(" Used by : \n");
 
 			auto FindMatchingProtocol = [&StringBuilder, &Protocol, &NetRefHandleManager](uint32 InternalObjectIndex)
 			{
@@ -528,7 +528,7 @@ void NetObjectProtocolReferencesToString(FStringBuilderBase& StringBuilder, uint
 	);
 }
 
-void DebugOutputNetObjectProtocolReferences(uint64 ProtocolId, uint32 ReplicationSystemId)
+void DebugOutputNetObjectProtocolReferences(FReplicationProtocolIdentifier ProtocolId, uint32 ReplicationSystemId)
 {
 	TStringBuilder<4096> StringBuilder;
 	StringBuilder.Reset();
@@ -538,7 +538,7 @@ void DebugOutputNetObjectProtocolReferences(uint64 ProtocolId, uint32 Replicatio
 	FPlatformMisc::LowLevelOutputDebugString(StringBuilder.ToString());
 }
 
-const TCHAR* DebugNetObjectProtocolReferencesToString(uint64 ProtocolId, uint32 ReplicationSystemId)
+const TCHAR* DebugNetObjectProtocolReferencesToString(FReplicationProtocolIdentifier ProtocolId, uint32 ReplicationSystemId)
 {
 	static TStringBuilder<4096> StringBuilder;
 	StringBuilder.Reset();
diff --git a/Engine/Source/Runtime/Experimental/Iris/Core/Private/Iris/ReplicationSystem/ObjectReplicationBridge.cpp b/Engine/Source/Runtime/Experimental/Iris/Core/Private/Iris/ReplicationSystem/ObjectReplicationBridge.cpp
--- a/Engine/Source/Runtime/Experimental/Iris/Core/Private/Iris/ReplicationSystem/ObjectReplicationBridge.cpp
+++ b/Engine/Source/Runtime/Experimental/Iris/Core/Private/Iris/ReplicationSystem/ObjectReplicationBridge.cpp
@@ -325,7 +325,7 @@ UE::Net::FNetRefHandle UObjectReplicationBridge::BeginReplication(UObject* Insta
 		const bool bIsValidProtocol = FReplicationProtocolManager::ValidateReplicationProtocol(ReplicationProtocol, RegisteredFragments);
 		if (!bIsValidProtocol)
 		{
-			UE_LOG_OBJECTREPLICATIONBRIDGE(Error, TEXT("BeginReplication Found invalid protocol ProtocolId:0x%" UINT64_x_FMT " for Object named %s"), ReplicationProtocol->ProtocolIdentifier, *Instance->GetName());
+			UE_LOG_OBJECTREPLICATIONBRIDGE(Error, TEXT("BeginReplication Found invalid protocol ProtocolId:0x%x for Object named %s"), ReplicationProtocol->ProtocolIdentifier, *Instance->GetName());
 			return FNetRefHandle::GetInvalid();
 		}
 	}
@@ -363,7 +363,7 @@ UE::Net::FNetRefHandle UObjectReplicationBridge::BeginReplication(UObject* Insta
 			uint8 PollFramePeriod = ConvertPollFrequencyIntoFrames(PollFrequency);
 			PollFrequencyLimiter->SetPollFramePeriod(InternalReplicationIndex, PollFramePeriod);
 
-			UE_LOG_OBJECTREPLICATIONBRIDGE(Verbose, TEXT("BeginReplication Created %s with ProtocolId:0x%" UINT64_x_FMT " for Object named %s"), *RefHandle.ToString(), ReplicationProtocol->ProtocolIdentifier, *Instance->GetName());
+			UE_LOG_OBJECTREPLICATIONBRIDGE(Verbose, TEXT("BeginReplication Created %s with ProtocolId:0x%x for Object named %s"), *RefHandle.ToString(), ReplicationProtocol->ProtocolIdentifier, *Instance->GetName());
 
 			{
 				FWorldLocations& WorldLocations = ReplicationSystem->GetReplicationSystemInternal()->GetWorldLocations();
@@ -412,7 +412,7 @@ UE::Net::FNetRefHandle UObjectReplicationBridge::BeginReplication(UObject* Insta
 
 			return RefHandle;
 		}
-		UE_LOG_OBJECTREPLICATIONBRIDGE(Warning, TEXT("BeginReplication Failed to create NetRefHandle with ProtocolId:0x%" UINT64_x_FMT " for Object named %s"), (ReplicationProtocol != nullptr ? ReplicationProtocol->ProtocolIdentifier : FReplicationProtocolIdentifier(0)), *Instance->GetName());
+		UE_LOG_OBJECTREPLICATIONBRIDGE(Warning, TEXT("BeginReplication Failed to create NetRefHandle with ProtocolId:0x%x for Object named %s"), (ReplicationProtocol != nullptr ? ReplicationProtocol->ProtocolIdentifier : FReplicationProtocolIdentifier(0)), *Instance->GetName());
 	}
 
 	// If we get here, it means that we failed to assign an internal handle for the object. We've probably run out of handles which currently is a fatal error.
@@ -602,7 +602,7 @@ bool UObjectReplicationBridge::WriteNetRefHandleCreationInfo(FReplicationBridgeS
 	// Write Type header, if there is a cached one, use it!
 	if (TUniquePtr<const FCreationHeader>* CachedHeader = CachedCreationHeaders.Find(Handle))
 	{
-		WriteUint64(Context.SerializationContext.GetBitStreamWriter(), (*CachedHeader)->ProtocolIdentifier);
+		Context.SerializationContext.GetBitStreamWriter()->WriteBits((*CachedHeader)->ProtocolIdentifier, 32);
 		return WriteCreationHeader(Context.SerializationContext, (*CachedHeader).Get());
 	}
 	else
@@ -615,7 +615,7 @@ bool UObjectReplicationBridge::WriteNetRefHandleCreationInfo(FReplicationBridgeS
 			ensureMsgf(Protocol, TEXT("WriteNetRefHandleCreationInfo: Cannot write creationinfo for %s, since protocol has been detached"), *Handle.ToString());
 			return false;
 		}
-		WriteUint64(Context.SerializationContext.GetBitStreamWriter(), Protocol->ProtocolIdentifier);
+		Context.SerializationContext.GetBitStreamWriter()->WriteBits(Protocol->ProtocolIdentifier, 32);
 		return WriteCreationHeader(Context.SerializationContext, Handle);
 	}
 }
@@ -722,7 +722,7 @@ FReplicationBridgeCreateNetRefHandleResult UObjectReplicationBridge::CreateNetRe
 
 	FNetBitStreamReader* Reader = Context.SerializationContext.GetBitStreamReader();
 
-	FReplicationProtocolIdentifier ReceivedProtocolId = ReadUint64(Context.SerializationContext.GetBitStreamReader());
+	FReplicationProtocolIdentifier ReceivedProtocolId = Context.SerializationContext.GetBitStreamReader()->ReadBits(32);
 
 	// Read creation header
 	TUniquePtr<FCreationHeader> Header(ReadCreationHeader(Context.SerializationContext));
diff --git a/Engine/Source/Runtime/Experimental/Iris/Core/Private/Iris/ReplicationSystem/ReplicationBridge.cpp b/Engine/Source/Runtime/Experimental/Iris/Core/Private/Iris/ReplicationSystem/ReplicationBridge.cpp
--- a/Engine/Source/Runtime/Experimental/Iris/Core/Private/Iris/ReplicationSystem/ReplicationBridge.cpp
+++ b/Engine/Source/Runtime/Experimental/Iris/Core/Private/Iris/ReplicationSystem/ReplicationBridge.cpp
@@ -402,7 +402,7 @@ UE::Net::FNetRefHandle UReplicationBridge::InternalCreateNetObject(FNetRefHandle
 
 	if (Handle.IsValid())
 	{
-		UE_NET_TRACE_NETHANDLE_CREATED(Handle, ReplicationProtocol->DebugName, ReplicationProtocol->ProtocolIdentifier, 0/*Local*/);
+		UE_NET_TRACE_NETHANDLE_CREATED(Handle, ReplicationProtocol->DebugName, (uint64)ReplicationProtocol->ProtocolIdentifier, 0/*Local*/);
 	}
 
 	return Handle;
@@ -419,7 +419,7 @@ UE::Net::FNetRefHandle UReplicationBridge::InternalCreateNetObjectFromRemote(FNe
 
 	if (Handle.IsValid())
 	{
-		UE_NET_TRACE_NETHANDLE_CREATED(Handle, ReplicationProtocol->DebugName, ReplicationProtocol->ProtocolIdentifier, 1/*Remote*/);
+		UE_NET_TRACE_NETHANDLE_CREATED(Handle, ReplicationProtocol->DebugName, (uint64)ReplicationProtocol->ProtocolIdentifier, 1/*Remote*/);
 	}
 
 	return Handle;
diff --git a/Engine/Source/Runtime/Experimental/Iris/Core/Private/Iris/ReplicationSystem/ReplicationProtocol.h b/Engine/Source/Runtime/Experimental/Iris/Core/Private/Iris/ReplicationSystem/ReplicationProtocol.h
--- a/Engine/Source/Runtime/Experimental/Iris/Core/Private/Iris/ReplicationSystem/ReplicationProtocol.h
+++ b/Engine/Source/Runtime/Experimental/Iris/Core/Private/Iris/ReplicationSystem/ReplicationProtocol.h
@@ -9,7 +9,7 @@
 namespace UE::Net
 {
 	struct FNetDebugName;
-	typedef uint64 FReplicationProtocolIdentifier;
+	typedef uint32 FReplicationProtocolIdentifier;
 	class FReplicationFragment;
 	struct FReplicationStateDescriptor;
 }
diff --git a/Engine/Source/Runtime/Experimental/Iris/Core/Private/Iris/ReplicationSystem/ReplicationProtocolManager.cpp b/Engine/Source/Runtime/Experimental/Iris/Core/Private/Iris/ReplicationSystem/ReplicationProtocolManager.cpp
--- a/Engine/Source/Runtime/Experimental/Iris/Core/Private/Iris/ReplicationSystem/ReplicationProtocolManager.cpp
+++ b/Engine/Source/Runtime/Experimental/Iris/Core/Private/Iris/ReplicationSystem/ReplicationProtocolManager.cpp
@@ -144,7 +144,7 @@ FReplicationProtocolIdentifier FReplicationProtocolManager::CalculateProtocolIde
 	}
 
 	FReplicationProtocolIdentifier ProtocolIdentifier;
-	ProtocolIdentifier = CityHash64(reinterpret_cast<const char*>(IdBuffer.GetData()), sizeof(uint64) * IdBuffer.Num());
+	ProtocolIdentifier = CityHash32(reinterpret_cast<const char*>(IdBuffer.GetData()), sizeof(uint64) * IdBuffer.Num());
 
 	return ProtocolIdentifier;
 }
@@ -225,7 +225,7 @@ const FReplicationProtocol* FReplicationProtocolManager::CreateReplicationProtoc
 		const FReplicationProtocolIdentifier NewProtocolId = CalculateProtocolIdentifier(Fragments);
 		if (NewProtocolId != ProtocolId)
 		{
-			UE_LOG(LogIris, Warning, TEXT("FReplicationProtocolManager::CreateReplicationProtocol Id mismatch when creating protocol named %s with ProtocolId:0x%" UINT64_x_FMT), DebugName, ProtocolId);
+			UE_LOG(LogIris, Warning, TEXT("FReplicationProtocolManager::CreateReplicationProtocol Id mismatch when creating protocol named %s with ProtocolId:0x%x"), DebugName, ProtocolId);
  #if UE_NET_ENABLE_PROTOCOLMANAGER_LOG
  			if (UE_LOG_ACTIVE(LogIris, Warning))
  			{
@@ -235,7 +235,7 @@ const FReplicationProtocol* FReplicationProtocolManager::CreateReplicationProtoc
  				UE_LOG(LogIris, Warning, TEXT("%s"), StringBuilder.ToString());
  			}
  #endif
-			ensureMsgf(NewProtocolId == ProtocolId, TEXT("FReplicationProtocolManager::CreateReplicationProtocol Id mismatch when creating protocol named %s with in ProtocolId:0x%" UINT64_x_FMT "Calculated ProtocolId:0x$%" UINT64_x_FMT), DebugName, ProtocolId, NewProtocolId);
+			ensureMsgf(NewProtocolId == ProtocolId, TEXT("FReplicationProtocolManager::CreateReplicationProtocol Id mismatch when creating protocol named %s with in ProtocolId:0x%x Calculated ProtocolId:0x%x"), DebugName, ProtocolId, NewProtocolId);
 			return nullptr;
 		}
 	}
@@ -385,12 +385,12 @@ const FReplicationProtocol* FReplicationProtocolManager::CreateReplicationProtoc
 		TStringBuilder<4096> StringBuilder;
 		FragmentListToString(StringBuilder, Fragments);
 		LOG_SCOPE_VERBOSITY_OVERRIDE(LogIris, ELogVerbosity::Log);
-		UE_LOG_PROTOCOLMANAGER(Log, TEXT("FReplicationProtocolManager::CreateReplicationProtocol Created new protocol %s with ProtocolId:0x%" UINT64_x_FMT), ToCStr(Protocol->DebugName), ProtocolId);
+		UE_LOG_PROTOCOLMANAGER(Log, TEXT("FReplicationProtocolManager::CreateReplicationProtocol Created new protocol %s with ProtocolId:0x%x"), ToCStr(Protocol->DebugName), ProtocolId);
 		UE_LOG_PROTOCOLMANAGER(Log, TEXT("%s"), StringBuilder.ToString());
 	}
 	else
 	{
-		UE_LOG_PROTOCOLMANAGER(Verbose, TEXT("FReplicationProtocolManager::CreateReplicationProtocol Created new protocol %s with ProtocolId:0x%" UINT64_x_FMT), ToCStr(Protocol->DebugName), ProtocolId);	
+		UE_LOG_PROTOCOLMANAGER(Verbose, TEXT("FReplicationProtocolManager::CreateReplicationProtocol Created new protocol %s with ProtocolId:0x%x"), ToCStr(Protocol->DebugName), ProtocolId);	
 	}
 #endif
 
diff --git a/Engine/Source/Runtime/Experimental/Iris/Core/Public/Iris/Core/IrisDebugging.h b/Engine/Source/Runtime/Experimental/Iris/Core/Public/Iris/Core/IrisDebugging.h
--- a/Engine/Source/Runtime/Experimental/Iris/Core/Public/Iris/Core/IrisDebugging.h
+++ b/Engine/Source/Runtime/Experimental/Iris/Core/Public/Iris/Core/IrisDebugging.h
@@ -9,8 +9,10 @@
 class UReplicationSystem;
 namespace UE::Net
 {
-struct FReplicationProtocol;
-struct FReplicationInstanceProtocol;
+	struct FReplicationProtocol;
+	struct FReplicationInstanceProtocol;
+
+	typedef uint32 FReplicationProtocolIdentifier;
 }
 
 namespace UE::Net::Private
@@ -51,7 +53,7 @@ IRISCORE_API FNetRefHandle GetDebugNetRefHandle();
 void NetObjectStateToString(FStringBuilderBase& StringBuilder, FNetRefHandle RefHandle);
 
 /** Find all handles references registered for a protocol and output to StringBuilder */
-void NetObjectProtocolReferencesToString(FStringBuilderBase& StringBuilder, uint64 ProtocolId, uint32 ReplicationSystemId);
+void NetObjectProtocolReferencesToString(FStringBuilderBase& StringBuilder, FReplicationProtocolIdentifier ProtocolId, uint32 ReplicationSystemId);
 
 // Helper functions for debugging state data exposed as extern "C" that are callable from immediate and watch window in the debugger
 
@@ -69,7 +71,7 @@ extern "C" IRISCORE_API void DebugOutputNetObjectState(uint64 NetRefHandleId, ui
 extern "C" IRISCORE_API const TCHAR* DebugNetObjectStateToString(uint32 NetRefHandleId, uint32 ReplicationSystemId);
 
 /** Find all handles references registered for a protocol and output to DebugOutput in debugger */
-extern "C" IRISCORE_API void DebugOutputNetObjectProtocolReferences(uint64 ProtocolId, uint32 ReplicationSystemId);
+extern "C" IRISCORE_API void DebugOutputNetObjectProtocolReferences(FReplicationProtocolIdentifier ProtocolId, uint32 ReplicationSystemId);
 
 /** Get info about replicated object */
 struct FNetReplicatedObjectDebugInfo
@@ -100,7 +102,7 @@ extern "C" IRISCORE_API FNetReplicatedObjectDebugInfo DebugInternalNetRefIndex(U
 /** Variant of DebugOutputNetObjectProtocolReferences that can be used from breakpoints or in watch window to find all handles references registered for a protocol and output to DebugOutput in debugger
  * NOTE: Use only for debugging as this variant uses a static buffer which is not thread safe	
 */
-extern "C" IRISCORE_API const TCHAR* DebugNetObjectProtocolReferencesToString(uint64 ProtocolId, uint32 ReplicationSystemId);
+extern "C" IRISCORE_API const TCHAR* DebugNetObjectProtocolReferencesToString(FReplicationProtocolIdentifier ProtocolId, uint32 ReplicationSystemId);
 
 /** Set the Object Name to break on */
 extern "C" IRISCORE_API void SetIrisDebugObjectName(const ANSICHAR* NameBuffer);
diff --git a/Engine/Source/Runtime/Experimental/Iris/Core/Public/Iris/ReplicationSystem/ObjectReplicationBridge.h b/Engine/Source/Runtime/Experimental/Iris/Core/Public/Iris/ReplicationSystem/ObjectReplicationBridge.h
--- a/Engine/Source/Runtime/Experimental/Iris/Core/Public/Iris/ReplicationSystem/ObjectReplicationBridge.h
+++ b/Engine/Source/Runtime/Experimental/Iris/Core/Public/Iris/ReplicationSystem/ObjectReplicationBridge.h
@@ -19,9 +19,13 @@ namespace UE::Net
 	enum class ENetRefHandleError : uint32;
 
 	struct FNetObjectResolveContext;
+
 	typedef uint32 FNetObjectFilterHandle;
 	typedef uint32 FNetObjectPrioritizerHandle;
+	typedef uint32 FReplicationProtocolIdentifier;
+
 	class FNetObjectReference;
+
 	namespace Private
 	{
 		typedef uint32 FInternalNetRefIndex;
@@ -151,7 +155,7 @@ public:
 
 	private:
 		friend UObjectReplicationBridge;
-		uint64 ProtocolIdentifier = uint64(0);
+		UE::Net::FReplicationProtocolIdentifier ProtocolIdentifier = 0U;
 	};
 	
 	/**
diff --git a/Engine/Source/Runtime/Net/Core/Private/Net/Core/Trace/NetTraceInternal.cpp b/Engine/Source/Runtime/Net/Core/Private/Net/Core/Trace/NetTraceInternal.cpp
--- a/Engine/Source/Runtime/Net/Core/Private/Net/Core/Trace/NetTraceInternal.cpp
+++ b/Engine/Source/Runtime/Net/Core/Private/Net/Core/Trace/NetTraceInternal.cpp
@@ -13,13 +13,6 @@
 #define UE_NET_TRACE_VALIDATE 1
 
 struct FNetTraceReporter;
-namespace UE::Net
-{
-	namespace Private
-	{
-		typedef uint64 FReplicationProtocolIdentifier;
-	}
-}
 
 #if UE_NET_TRACE_ENABLED
 
