Geometry collection : Fix removal timer starting at -1 and causing remove on break to be longer by a second

diff --git a/Engine/Source/Runtime/Experimental/GeometryCollectionEngine/Private/GeometryCollection/GeometryCollectionEngineRemoval.cpp b/Engine/Source/Runtime/Experimental/GeometryCollectionEngine/Private/GeometryCollection/GeometryCollectionEngineRemoval.cpp
--- a/Engine/Source/Runtime/Experimental/GeometryCollectionEngine/Private/GeometryCollection/GeometryCollectionEngineRemoval.cpp
+++ b/Engine/Source/Runtime/Experimental/GeometryCollectionEngine/Private/GeometryCollection/GeometryCollectionEngineRemoval.cpp
@@ -44,7 +44,7 @@ bool FGeometryCollectionRemoveOnBreakDynamicFacade::IsConst() const
 void FGeometryCollectionRemoveOnBreakDynamicFacade::DefineSchema()
 {
 	check(!IsConst());
-	BreakTimerAttribute.AddAndFill(DisabledBreakTimer);
+	BreakTimerAttribute.AddAndFill(BreakTimerStartValue);
 	PostBreakDurationAttribute.Add();
 	BreakRemovalDurationAttribute.Add();
 }
@@ -54,7 +54,7 @@ void FGeometryCollectionRemoveOnBreakDynamicFacade::SetAttributeValues(const Geo
 	check(!IsConst());
 	if (IsValid())
 	{
-		BreakTimerAttribute.Fill(DisabledBreakTimer);
+		BreakTimerAttribute.Fill(BreakTimerStartValue);
 
 		// make sure we generate random value consistently between client and server 
 		// we can use the length of the transform group for that
diff --git a/Engine/Source/Runtime/Experimental/GeometryCollectionEngine/Public/GeometryCollection/GeometryCollectionEngineRemoval.h b/Engine/Source/Runtime/Experimental/GeometryCollectionEngine/Public/GeometryCollection/GeometryCollectionEngineRemoval.h
--- a/Engine/Source/Runtime/Experimental/GeometryCollectionEngine/Public/GeometryCollection/GeometryCollectionEngineRemoval.h
+++ b/Engine/Source/Runtime/Experimental/GeometryCollectionEngine/Public/GeometryCollection/GeometryCollectionEngineRemoval.h
@@ -21,6 +21,7 @@ class FGeometryCollectionRemoveOnBreakDynamicFacade
 {
 public:
 	static constexpr float DisabledBreakTimer = -1;
+	static constexpr float BreakTimerStartValue = 0;
 	static constexpr float DisabledPostBreakDuration = -1;
 	static constexpr float CrumblingRemovalTimer = -1;
 	
