Iris
* FPropertyReplicationState::ToString will now print all members of replicated structs

#rnx

diff --git a/Engine/Source/Runtime/Experimental/Iris/Core/Private/Iris/ReplicationState/PropertyReplicationState.cpp b/Engine/Source/Runtime/Experimental/Iris/Core/Private/Iris/ReplicationState/PropertyReplicationState.cpp
--- a/Engine/Source/Runtime/Experimental/Iris/Core/Private/Iris/ReplicationState/PropertyReplicationState.cpp
+++ b/Engine/Source/Runtime/Experimental/Iris/Core/Private/Iris/ReplicationState/PropertyReplicationState.cpp
@@ -524,7 +524,8 @@ const TCHAR* FPropertyReplicationState::ToString(FStringBuilderBase& StringBuild
 
 			if (bIncludeAll || IsDirty(MemberIt))
 			{
-				Property->ExportTextItem_Direct(TempString, StateBuffer + Descriptor->MemberDescriptors[MemberIt].ExternalMemberOffset, nullptr, nullptr, PPF_SimpleObjectText);
+				void* PropertyData = StateBuffer + Descriptor->MemberDescriptors[MemberIt].ExternalMemberOffset;
+				Property->ExportTextItem_Direct(TempString, PropertyData, PropertyData, nullptr, PPF_SimpleObjectText|PPF_IncludeTransient|PPF_UseDeprecatedProperties);
 				StringBuilder.Appendf(TEXT("%u - %s : %s\n"), MemberIt, *Property->GetName(), ToCStr(TempString));
 				TempString.Reset();
 			}
