Added CVar to control when UniqueIdxs are cleared on geometry collection physics proxies
(On behalf of benn.gallagher)

diff --git a/Engine/Source/Runtime/Experimental/Chaos/Private/PhysicsProxy/GeometryCollectionPhysicsProxy.cpp b/Engine/Source/Runtime/Experimental/Chaos/Private/PhysicsProxy/GeometryCollectionPhysicsProxy.cpp
--- a/Engine/Source/Runtime/Experimental/Chaos/Private/PhysicsProxy/GeometryCollectionPhysicsProxy.cpp
+++ b/Engine/Source/Runtime/Experimental/Chaos/Private/PhysicsProxy/GeometryCollectionPhysicsProxy.cpp
@@ -144,6 +144,12 @@ FAutoConsoleVariableRef CVarGeometryCollectionScaleClusterGeometry(
 	bGeometryCollectionScaleClusterGeometry,
 	TEXT("If enabled, update the cluster geometry if the scale has changed"));
 
+bool bEmptyUniqueIndicesOnGT = false;
+FAutoConsoleVariableRef CVarEmptyUniqueIndicesOnGT(
+	TEXT("p.GeometryCollection.EmptyUniqueIndicesOnGT"),
+	bEmptyUniqueIndicesOnGT,
+	TEXT("If enabled, immediately free particle unique indices after particle initialisation"));
+
 enum EOverrideGCCollisionSetupForTraces
 {
 	GCCSFT_Property   = -1,  // Default: do what property says
@@ -671,7 +677,7 @@ void FGeometryCollectionPhysicsProxy::Initialize(Chaos::FPBDRigidsEvolutionBase
 		GameThreadCollection.RemoveAttribute(FGeometryDynamicCollection::ImplicitsAttribute, FTransformCollection::TransformGroup);
 	}
 	// If GT particle have been created, we can use directly the UniqueIdx from the GTParticle, then free up some space.
-	if (bBuildGeometryForChildrenOnGT || bCreateGTParticleForChildren)
+	if ((bBuildGeometryForChildrenOnGT || bCreateGTParticleForChildren) && bEmptyUniqueIndicesOnGT)
 	{
 		UniqueIdxs.Empty();
 	}
@@ -903,7 +909,11 @@ void FGeometryCollectionPhysicsProxy::CreateChildrenGeometry_External()
 			TManagedArray<Chaos::FImplicitObjectPtr>& Implicits = GameThreadCollection.ModifyAttribute<Chaos::FImplicitObjectPtr>(FGeometryDynamicCollection::ImplicitsAttribute, FTransformCollection::TransformGroup);
 			CreateGTParticles(Implicits, RBDSolver->GetEvolution(), /*bInitializationTime*/false);
 			SyncParticles_External();
-			UniqueIdxs.Empty();
+
+			if (bEmptyUniqueIndicesOnGT)
+			{
+				UniqueIdxs.Empty();
+			}
 
 			// The Implicits attributes from the Dynamic Collection are just used for initialization, after they can be removed and so free some memory.
 			GameThreadCollection.RemoveAttribute(FGeometryDynamicCollection::ImplicitsAttribute, FTransformCollection::TransformGroup);
