Remove the 'ias.SuspendSystem' cvar which allowed IAS to be disabled remotely via hot fixes

#rnx

- This feature was only added while the system was being tested in production and by now we are confident enough in the system not to need this.
- Suspending IAS only worked safely if the game provided backup data on disk which would defeat the point of the system in the first place which is why this cvar was never designed for general use.

diff --git a/Engine/Source/Runtime/Experimental/IoStore/OnDemand/Private/IoStoreOnDemand.cpp b/Engine/Source/Runtime/Experimental/IoStore/OnDemand/Private/IoStoreOnDemand.cpp
--- a/Engine/Source/Runtime/Experimental/IoStore/OnDemand/Private/IoStoreOnDemand.cpp
+++ b/Engine/Source/Runtime/Experimental/IoStore/OnDemand/Private/IoStoreOnDemand.cpp
@@ -49,14 +49,6 @@ FString GIasOnDemandTocExt = TEXT(".uondemandtoc");
 
 static const TCHAR* NotInitializedError = TEXT("I/O store on-demand not initialized");
 
-bool GIasSuspendSystem = false;
-static FAutoConsoleVariableRef CVar_SuspendSystemEnabled(
-	TEXT("ias.SuspendSystem"),
-	GIasSuspendSystem,
-	TEXT("Suspends the use of the OnDemand system"),
-	ECVF_ReadOnly
-);
-
 /** Temp cvar to allow the fallback url to be hotfixed in case of problems */
 static FString GDistributedEndpointFallbackUrl;
 static FAutoConsoleVariableRef CVar_DistributedEndpointFallbackUrl(
@@ -1319,16 +1311,7 @@ void FIoStoreOnDemandModule::StartupModule()
 #endif //PLATFORM_WINDOWS
 
 #if !UE_IAS_CUSTOM_INITIALIZATION
-
-	if (!GIasSuspendSystem)
-	{
-		InitializeInternal();
-	}
-	else
-	{
-		UE_LOG(LogIas, Display, TEXT("The IoStoreOnDemand module has been remotely disabled by the 'ias.SuspendSystemEnabled' cvar"));
-	}
-
+	InitializeInternal();
 #endif // !UE_IAS_CUSTOM_INITIALIZATION
 }
 
@@ -1356,12 +1339,6 @@ void FIoStoreOnDemandModule::ShutdownModule()
 
 EOnDemandInitResult FIoStoreOnDemandModule::Initialize()
 {
-	if (GIasSuspendSystem)
-	{
-		UE_LOG(LogIas, Display, TEXT("The IoStoreOnDemand module has been remotely disabled by the 'ias.SuspendSystemEnabled' cvar"));
-		return EOnDemandInitResult::Suspended;
-	}
-
 	InitializeInternal();
 
 	return HttpIoDispatcherBackend.IsValid() ? EOnDemandInitResult::Success : EOnDemandInitResult::Disabled;
diff --git a/Engine/Source/Runtime/Experimental/IoStore/OnDemand/Public/IO/IoStoreOnDemand.h b/Engine/Source/Runtime/Experimental/IoStore/OnDemand/Public/IO/IoStoreOnDemand.h
--- a/Engine/Source/Runtime/Experimental/IoStore/OnDemand/Public/IO/IoStoreOnDemand.h
+++ b/Engine/Source/Runtime/Experimental/IoStore/OnDemand/Public/IO/IoStoreOnDemand.h
@@ -234,12 +234,6 @@ enum class EOnDemandInitResult
 	Disabled,
 	/** The module was unable to start up correctly due to an unexpected error */
 	Error,
-
-	/**
-	 * The use of the module has been suspended, if possible calling systems should activate alternative ways
-	 * to access the OnDemand data. This option is temporary and not intended for general use.
-	 */
-	Suspended
 };
 
 #endif // UE_IAS_CUSTOM_INITIALIZATION
