Dataflow - Deprecated FDataflowNode's public member variables with the view to make them private in 5.7.
This allows to avoid accidentaly shadowing variables in the node declarations (i.e. Name is used a lot in cloth nodes) and accidentally access the wrong variable.

#rnx

diff --git a/Engine/Source/Runtime/Experimental/Dataflow/Core/Private/Dataflow/DataflowInputOutput.cpp b/Engine/Source/Runtime/Experimental/Dataflow/Core/Private/Dataflow/DataflowInputOutput.cpp
--- a/Engine/Source/Runtime/Experimental/Dataflow/Core/Private/Dataflow/DataflowInputOutput.cpp
+++ b/Engine/Source/Runtime/Experimental/Dataflow/Core/Private/Dataflow/DataflowInputOutput.cpp
@@ -260,13 +260,13 @@ bool FDataflowOutput::EvaluateImpl(Dataflow::FContext& Context) const
 	}
 
 	// check if the cache has a valid version
-	if(Context.HasData(CacheKey(), OwningNode->LastModifiedTimestamp))
+	if(Context.HasData(CacheKey(), OwningNode->GetTimestamp()))
 	{
-		UE_LOG(LogChaosDataflow, Verbose, TEXT("FDataflowOutput::EvaluateImpl(): Context has data, NodeTimestamp [%ul], CacheTimestamp [%lu]"), OwningNode->LastModifiedTimestamp.Value, Context.GetTimestamp(CacheKey()).Value);
+		UE_LOG(LogChaosDataflow, Verbose, TEXT("FDataflowOutput::EvaluateImpl(): Context has data, NodeTimestamp [%ul], CacheTimestamp [%lu]"), OwningNode->GetTimestamp().Value, Context.GetTimestamp(CacheKey()).Value);
 		return true;
 	}
 	// if not, evaluate
-	UE_LOG(LogChaosDataflow, Verbose, TEXT("FDataflowNode::Evaluate(): Node [%s], Output [%s], NodeTimestamp [%lu]"), *GetOwningNode()->GetName().ToString(), *GetName().ToString(), OwningNode->LastModifiedTimestamp.Value);
+	UE_LOG(LogChaosDataflow, Verbose, TEXT("FDataflowNode::Evaluate(): Node [%s], Output [%s], NodeTimestamp [%lu]"), *GetOwningNode()->GetName().ToString(), *GetName().ToString(), OwningNode->GetTimestamp().Value);
 	OwningNode->Evaluate(Context, this);
 	// Validation
 	if (!Context.HasData(CacheKey()))
diff --git a/Engine/Source/Runtime/Experimental/Dataflow/Core/Private/Dataflow/DataflowNode.cpp b/Engine/Source/Runtime/Experimental/Dataflow/Core/Private/Dataflow/DataflowNode.cpp
--- a/Engine/Source/Runtime/Experimental/Dataflow/Core/Private/Dataflow/DataflowNode.cpp
+++ b/Engine/Source/Runtime/Experimental/Dataflow/Core/Private/Dataflow/DataflowNode.cpp
@@ -309,7 +309,9 @@ void FDataflowNode::AddOutput(FDataflowOutput* InPtr)
 {
 	if (InPtr)
 	{
+PRAGMA_DISABLE_DEPRECATION_WARNINGS  // Until Outputs becomes private
 		for (const TPair<int32, FDataflowOutput*>& Elem : Outputs)
+PRAGMA_ENABLE_DEPRECATION_WARNINGS
 		{
 			const FDataflowOutput* const Out = Elem.Value;
 			ensureMsgf(!Out->GetName().IsEqual(InPtr->GetName()), TEXT("Add Output Failed: Existing Node output already defined with name (%s)"), *InPtr->GetName().ToString());
@@ -318,16 +320,20 @@ void FDataflowNode::AddOutput(FDataflowOutput* InPtr)
 		check(InPtr->GetOwningNode() == this);
 
 		const uint32 PropertyOffset = InPtr->GetOffset();
+PRAGMA_DISABLE_DEPRECATION_WARNINGS  // Until Outputs becomes private
 		if (ensure(!Outputs.Contains(PropertyOffset)))
 		{
 			Outputs.Add(PropertyOffset, InPtr);
 		}
+PRAGMA_ENABLE_DEPRECATION_WARNINGS
 	}
 }
 
 FDataflowOutput* FDataflowNode::FindOutput(uint32 InGuidHash)
 {
+PRAGMA_DISABLE_DEPRECATION_WARNINGS  // Until Outputs becomes private
 	for (TPair<int32, FDataflowOutput*>& Elem : Outputs)
+PRAGMA_ENABLE_DEPRECATION_WARNINGS
 	{
 		FDataflowOutput* const Con = Elem.Value;
 		if (GetTypeHash(Con->GetGuid()) == InGuidHash)
@@ -341,7 +347,9 @@ FDataflowOutput* FDataflowNode::FindOutput(uint32 InGuidHash)
 
 FDataflowOutput* FDataflowNode::FindOutput(FName InName)
 {
+PRAGMA_DISABLE_DEPRECATION_WARNINGS  // Until Outputs becomes private
 	for (TPair<int32, FDataflowOutput*>& Elem : Outputs)
+PRAGMA_ENABLE_DEPRECATION_WARNINGS
 	{
 		FDataflowOutput* const Con = Elem.Value;
 		if (Con->GetName().IsEqual(InName))
@@ -354,7 +362,9 @@ FDataflowOutput* FDataflowNode::FindOutput(FName InName)
 
 const FDataflowOutput* FDataflowNode::FindOutput(FName InName) const
 {
+PRAGMA_DISABLE_DEPRECATION_WARNINGS  // Until Outputs becomes private
 	for (const TPair<int32, FDataflowOutput*>& Elem : Outputs)
+PRAGMA_ENABLE_DEPRECATION_WARNINGS
 	{
 		const FDataflowOutput* const Con = Elem.Value;
 		if (Con->GetName().IsEqual(InName))
@@ -367,7 +377,9 @@ const FDataflowOutput* FDataflowNode::FindOutput(FName InName) const
 
 const FDataflowOutput* FDataflowNode::FindOutput(uint32 InGuidHash) const
 {
+PRAGMA_DISABLE_DEPRECATION_WARNINGS  // Until Outputs becomes private
 	for (const TPair<int32, FDataflowOutput*>& Elem : Outputs)
+PRAGMA_ENABLE_DEPRECATION_WARNINGS
 	{
 		const FDataflowOutput* const Con = Elem.Value;
 		if (GetTypeHash(Con->GetGuid()) == InGuidHash)
@@ -380,7 +392,9 @@ const FDataflowOutput* FDataflowNode::FindOutput(uint32 InGuidHash) const
 
 const FDataflowOutput* FDataflowNode::FindOutput(const Dataflow::FConnectionKey& Key) const
 {
+PRAGMA_DISABLE_DEPRECATION_WARNINGS  // Until Outputs becomes private
 	if (const FDataflowOutput* const* Con = Outputs.Find(Key.Offset))
+PRAGMA_ENABLE_DEPRECATION_WARNINGS
 	{
 		check(*Con);
 		return *Con;
@@ -401,7 +415,9 @@ const FDataflowOutput* FDataflowNode::FindOutput(const Dataflow::FConnectionRefe
 
 FDataflowOutput* FDataflowNode::FindOutput(const Dataflow::FConnectionKey& Key)
 {
+PRAGMA_DISABLE_DEPRECATION_WARNINGS  // Until Outputs becomes private
 	if (FDataflowOutput* const* Con = Outputs.Find(Key.Offset))
+PRAGMA_ENABLE_DEPRECATION_WARNINGS
 	{
 		check(*Con);
 		return *Con;
@@ -422,7 +438,9 @@ FDataflowOutput* FDataflowNode::FindOutput(const Dataflow::FConnectionReference&
 
 const FDataflowOutput* FDataflowNode::FindOutput(const FGuid& InGuid) const
 {
+PRAGMA_DISABLE_DEPRECATION_WARNINGS  // Until Outputs becomes private
 	for (const TPair<int32, FDataflowOutput*>& Elem : Outputs)
+PRAGMA_ENABLE_DEPRECATION_WARNINGS
 	{
 		const FDataflowOutput* const Con = Elem.Value;
 		if (Con->GetGuid() == InGuid)
@@ -435,31 +453,41 @@ const FDataflowOutput* FDataflowNode::FindOutput(const FGuid& InGuid) const
 
 int32 FDataflowNode::NumOutputs() const
 {
+PRAGMA_DISABLE_DEPRECATION_WARNINGS  // Until Outputs becomes private
 	return Outputs.Num();
+PRAGMA_ENABLE_DEPRECATION_WARNINGS
 }
 
 
 TArray< FDataflowOutput* > FDataflowNode::GetOutputs() const
 {
 	TArray< FDataflowOutput* > Result;
+PRAGMA_DISABLE_DEPRECATION_WARNINGS  // Until Outputs becomes private
 	Outputs.GenerateValueArray(Result);
+PRAGMA_ENABLE_DEPRECATION_WARNINGS
 	return Result;
 }
 
 
 void FDataflowNode::ClearOutputs()
 {
+PRAGMA_DISABLE_DEPRECATION_WARNINGS  // Until Outputs becomes private
 	for (TPair<int32, FDataflowOutput*>& Elem : Outputs)
+PRAGMA_ENABLE_DEPRECATION_WARNINGS
 	{
 		FDataflowOutput* const Con = Elem.Value;
 		delete Con;
 	}
+PRAGMA_DISABLE_DEPRECATION_WARNINGS  // Until Outputs becomes private
 	Outputs.Reset();
+PRAGMA_ENABLE_DEPRECATION_WARNINGS
 }
 
 bool FDataflowNode::HasHideableOutputs() const
 {
+PRAGMA_DISABLE_DEPRECATION_WARNINGS  // Until Outputs becomes private
 	for (const TPair<int32, FDataflowOutput*>& Elem : Outputs)
+PRAGMA_ENABLE_DEPRECATION_WARNINGS
 	{
 		const FDataflowOutput* const Con = Elem.Value;
 		if (Con->GetCanHidePin())
@@ -472,7 +500,9 @@ bool FDataflowNode::HasHideableOutputs() const
 
 bool FDataflowNode::HasHiddenOutputs() const
 {
+PRAGMA_DISABLE_DEPRECATION_WARNINGS  // Until Outputs becomes private
 	for (const TPair<int32, FDataflowOutput*>& Elem : Outputs)
+PRAGMA_ENABLE_DEPRECATION_WARNINGS
 	{
 		const FDataflowOutput* const Con = Elem.Value;
 		if (Con->GetPinIsHidden())
@@ -491,7 +521,9 @@ TArray<Dataflow::FPin> FDataflowNode::GetPins() const
 		const FDataflowInput* const Con = Elem.Value;
 		RetVal.Add({ Dataflow::FPin::EDirection::INPUT,Con->GetType(), Con->GetName(), Con->GetPinIsHidden()});
 	}
+PRAGMA_DISABLE_DEPRECATION_WARNINGS  // Until Outputs becomes private
 	for (const TPair<int32, FDataflowOutput*>& Elem : Outputs)
+PRAGMA_ENABLE_DEPRECATION_WARNINGS
 	{
 		const FDataflowOutput* const Con = Elem.Value;
 		RetVal.Add({ Dataflow::FPin::EDirection::OUTPUT,Con->GetType(), Con->GetName(), Con->GetPinIsHidden() });
@@ -519,7 +551,9 @@ void FDataflowNode::UnregisterPinConnection(const Dataflow::FPin& Pin)
 	}
 	else if (Pin.Direction == Dataflow::FPin::EDirection::OUTPUT)
 	{
+PRAGMA_DISABLE_DEPRECATION_WARNINGS  // Until Outputs becomes private
 		for (TMap<int32, FDataflowOutput*>::TIterator Iter = Outputs.CreateIterator(); Iter; ++Iter)
+PRAGMA_ENABLE_DEPRECATION_WARNINGS
 		{
 			FDataflowOutput* Con = Iter.Value();
 			if (Con->GetName().IsEqual(Pin.Name) && Con->GetType().IsEqual(Pin.Type))
@@ -545,12 +579,12 @@ void FDataflowNode::Invalidate(const Dataflow::FTimestamp& InModifiedTimestamp)
 		}
 		return;
 	}
-
+PRAGMA_DISABLE_DEPRECATION_WARNINGS  // Until LastModifiedTimestamp becomes private
 	if (LastModifiedTimestamp < InModifiedTimestamp)
 	{
 		LastModifiedTimestamp = InModifiedTimestamp;
-
 		for (TPair<int32, FDataflowOutput*>& Elem : Outputs)
+PRAGMA_ENABLE_DEPRECATION_WARNINGS
 		{
 			FDataflowOutput* const Con = Elem.Value;
 			Con->Invalidate(InModifiedTimestamp);
diff --git a/Engine/Source/Runtime/Experimental/Dataflow/Core/Private/Dataflow/DataflowNodeParameters.cpp b/Engine/Source/Runtime/Experimental/Dataflow/Core/Private/Dataflow/DataflowNodeParameters.cpp
--- a/Engine/Source/Runtime/Experimental/Dataflow/Core/Private/Dataflow/DataflowNodeParameters.cpp
+++ b/Engine/Source/Runtime/Experimental/Dataflow/Core/Private/Dataflow/DataflowNodeParameters.cpp
@@ -80,7 +80,7 @@ namespace Dataflow
 				// TODO: When no outputs are specified, this call to Evaluate should really be removed.
 				//       The purpose of the node evaluation function is to evaluate outputs.
 				//       Therefore if a node has no outputs, then it shouldn't need any evaluation.
-				UE_LOG(LogChaosDataflow, Verbose, TEXT("FDataflowNode::Evaluate(): Node [%s], Output [nullptr], NodeTimestamp [%lu]"), *Node->GetName().ToString(), Node->LastModifiedTimestamp.Value);
+				UE_LOG(LogChaosDataflow, Verbose, TEXT("FDataflowNode::Evaluate(): Node [%s], Output [nullptr], NodeTimestamp [%lu]"), *Node->GetName().ToString(), Node->GetTimestamp().Value);
 				Node->Evaluate(Context, nullptr);
 			}
 		}
diff --git a/Engine/Source/Runtime/Experimental/Dataflow/Core/Public/Dataflow/DataflowNode.h b/Engine/Source/Runtime/Experimental/Dataflow/Core/Public/Dataflow/DataflowNode.h
--- a/Engine/Source/Runtime/Experimental/Dataflow/Core/Public/Dataflow/DataflowNode.h
+++ b/Engine/Source/Runtime/Experimental/Dataflow/Core/Public/Dataflow/DataflowNode.h
@@ -43,12 +43,16 @@ PRAGMA_ENABLE_DEPRECATION_WARNINGS
 	friend class Dataflow::FGraph;
 	friend struct FDataflowConnection;
 
+	UE_DEPRECATED(5.5, "Will be made private in 5.7, use Get/Set methods instead.")
 	FGuid Guid;
+	UE_DEPRECATED(5.5, "Will be made private in 5.7, use Get/Set methods instead.")
 	FName Name;
+	UE_DEPRECATED(5.5, "Will be made private in 5.7, use Get/Set methods instead.")
 	Dataflow::FTimestamp LastModifiedTimestamp;
 
 	UE_DEPRECATED(5.5, "Inputs type has changed and has been made private")
 	TMap< int, FDataflowInput* > Inputs;
+	UE_DEPRECATED(5.5, "Will be made private in 5.7, use Get/Set methods instead.")
 	TMap< int, FDataflowOutput* > Outputs;
 
 	UPROPERTY(EditAnywhere, Category = "Dataflow")
@@ -73,12 +77,12 @@ PRAGMA_ENABLE_DEPRECATION_WARNINGS
 	// Warning: FDataflowNodes aren't actually safe to copy/move yet. These are here to disable deprecation warnings from the implicit operators that were getting created anyways.
 	// (Deleting the operators would require tagging all derived classes since this is a USTRUCT, so also not doing that here).
 	FDataflowNode& operator=(const FDataflowNode&) = default;
-	PRAGMA_ENABLE_DEPRECATION_WARNINGS
 
 	FGuid GetGuid() const { return Guid; }
 	FName GetName() const { return Name; }
 	void SetName(FName InName) { Name = InName; }
 	Dataflow::FTimestamp GetTimestamp() const { return LastModifiedTimestamp;  }
+	PRAGMA_ENABLE_DEPRECATION_WARNINGS
 	DATAFLOWCORE_API uint32 GetValueHash();
 
 	static FName StaticType() { return FName("FDataflowNode"); }
