Migrate FChaosRefCountedObject to a dedicated header.

This will make it easier to add AutoRTFM tests without pulling in all of ImplicitObject and its transitive dependencies.

diff --git a/Engine/Source/Runtime/Experimental/Chaos/Public/Chaos/ImplicitObject.h b/Engine/Source/Runtime/Experimental/Chaos/Public/Chaos/ImplicitObject.h
--- a/Engine/Source/Runtime/Experimental/Chaos/Public/Chaos/ImplicitObject.h
+++ b/Engine/Source/Runtime/Experimental/Chaos/Public/Chaos/ImplicitObject.h
@@ -8,6 +8,7 @@
 #include "Chaos/ImplicitFwd.h"
 #include "Chaos/ImplicitObjectType.h"
 #include "Chaos/AABB.h"
+#include "Chaos/RefCountedObject.h"
 #include "Templates/RefCounting.h"
 #include "AutoRTFM/AutoRTFM.h"
 
@@ -105,68 +106,6 @@ struct TImplicitTypeInfo
 // you should be OK.
 #define DISALLOW_FIMPLICIT_OBJECT_TAIL_PADDING INTEL_ISPC
 
-// Chaos ref counted object
-//  * @note AutoRTFM means that the return value of AddRef/Release is nonsense (as the ref-count doesn't change until the
-//  *       transaction is committed), but this is fine for use with TRefCountPtr (as it doesn't use those return values).
-class FChaosRefCountedObject
-{
-public:
-	FChaosRefCountedObject() : NumRefs(0) {}
-	virtual ~FChaosRefCountedObject() { check(NumRefs.GetValue() == 0); }
-	FChaosRefCountedObject(const FChaosRefCountedObject& Rhs) = delete;
-	FChaosRefCountedObject& operator=(const FChaosRefCountedObject& Rhs) = delete;
-	uint32 AddRef() const
-	{
-		UE_AUTORTFM_ONCOMMIT2(this)
-		{
-			NumRefs.Increment();
-		};
-
-		// Note: TRefCountPtr doesn't use the return value
-		return 0;
-	}
-	uint32 Release() const
-	{
-		UE_AUTORTFM_ONCOMMIT2(this)
-		{
-			uint32 Refs = uint32(NumRefs.Decrement());
-			if (Refs == 0)
-			{
-				if(bTransientFlag)
-				{
-					delete this;
-				}
-			}
-		};
-
-		// Note: TRefCountPtr doesn't use the return value
-		return 0;
-	}
-	uint32 GetRefCount() const
-	{
-		uint32 Ret = 0;
-		UE_AUTORTFM_OPEN2
-		{
-			Ret = uint32(NumRefs.GetValue());
-		};
-
-		return Ret;
-	}
-
-	void MakePersistent() const
-	{
-		bTransientFlag = false;
-	}
-	
-private:
-	// Number of refs onto the object
-	mutable FThreadSafeCounter NumRefs;
-
-	// Transient flag to trigger or not the automatic deletion
-	mutable std::atomic<bool> bTransientFlag = true;
-};
-	
-
 class FImplicitObject : public FChaosRefCountedObject
 {
 public:
diff --git a/Engine/Source/Runtime/Experimental/Chaos/Public/Chaos/RefCountedObject.h b/Engine/Source/Runtime/Experimental/Chaos/Public/Chaos/RefCountedObject.h
new file mode 100644
--- /dev/null
+++ b/Engine/Source/Runtime/Experimental/Chaos/Public/Chaos/RefCountedObject.h
@@ -0,0 +1,75 @@
+// Copyright Epic Games, Inc. All Rights Reserved.
+#pragma once
+
+#include <atomic>
+#include "AutoRTFM/AutoRTFM.h"
+#include "Runtime/Core/Public/HAL/ThreadSafeCounter.h"
+
+namespace Chaos
+{
+
+// Chaos ref counted object
+//  * @note AutoRTFM means that the return value of AddRef/Release is nonsense (as the ref-count doesn't change until the
+//  *       transaction is committed), but this is fine for use with TRefCountPtr (as it doesn't use those return values).
+class FChaosRefCountedObject
+{
+public:
+	FChaosRefCountedObject() : NumRefs(0) {}
+	virtual ~FChaosRefCountedObject()
+	{
+		check(NumRefs.GetValue() == 0);
+	}
+	FChaosRefCountedObject(const FChaosRefCountedObject& Rhs) = delete;
+	FChaosRefCountedObject& operator=(const FChaosRefCountedObject& Rhs) = delete;
+	uint32 AddRef() const
+	{
+		UE_AUTORTFM_ONCOMMIT2(this)
+		{
+			NumRefs.Increment();
+		};
+
+		// Note: TRefCountPtr doesn't use the return value
+		return 0;
+	}
+	uint32 Release() const
+	{
+		UE_AUTORTFM_ONCOMMIT2(this)
+		{
+			uint32 Refs = uint32(NumRefs.Decrement());
+			if (Refs == 0)
+			{
+				if (bTransientFlag)
+				{
+					delete this;
+				}
+			}
+		};
+
+		// Note: TRefCountPtr doesn't use the return value
+		return 0;
+	}
+	uint32 GetRefCount() const
+	{
+		uint32 Ret = 0;
+		UE_AUTORTFM_OPEN2
+		{
+			Ret = uint32(NumRefs.GetValue());
+		};
+
+		return Ret;
+	}
+
+	void MakePersistent() const
+	{
+		bTransientFlag = false;
+	}
+
+private:
+	// Number of refs onto the object
+	mutable FThreadSafeCounter NumRefs;
+
+	// Transient flag to trigger the automatic deletion, or not
+	mutable std::atomic<bool> bTransientFlag = true;
+};
+
+}  // namespace Chaos
\ No newline at end of file
