[AutoRTFM] Avoid a crash when a Chaos ref-counted object immediately goes out of scope.

Previously, ~FChaosRefCountedObject would access NumRefs directly without being wrapped in ONCOMMIT2. This was only safe if we we already inside in ONCOMMIT2 block (in this case, via Release).

This CL adds an ONCOMMIT2 block to the destructor, and includes unit tests for the class to verify the behavior.

diff --git a/Engine/Source/Runtime/Experimental/Chaos/Private/Tests/AutoRTFMRefCountedObjectTests.cpp b/Engine/Source/Runtime/Experimental/Chaos/Private/Tests/AutoRTFMRefCountedObjectTests.cpp
new file mode 100644
--- /dev/null
+++ b/Engine/Source/Runtime/Experimental/Chaos/Private/Tests/AutoRTFMRefCountedObjectTests.cpp
@@ -0,0 +1,55 @@
+// Copyright Epic Games, Inc. All Rights Reserved.
+
+#include "Misc/AutomationTest.h"
+#include "Misc/App.h"
+#include "AutoRTFM/AutoRTFM.h"
+#include "Chaos/RefCountedObject.h"
+#include "Templates/UniquePtr.h"
+
+#if WITH_DEV_AUTOMATION_TESTS
+
+IMPLEMENT_SIMPLE_AUTOMATION_TEST(FAutoRTFMChaosRefCountedObject, "AutoRTFM + ChaosRefCountedObject", \
+	                             EAutomationTestFlags::EngineFilter | EAutomationTestFlags::ClientContext | \
+								 EAutomationTestFlags::ServerContext | EAutomationTestFlags::CommandletContext)
+
+bool FAutoRTFMChaosRefCountedObject::RunTest(const FString& Parameters)
+{
+	if (!AutoRTFM::ForTheRuntime::IsAutoRTFMRuntimeEnabled())
+	{
+		ExecutionInfo.AddEvent(FAutomationEvent(EAutomationEventType::Info, 
+								                TEXT("SKIPPED 'FAutoRTFMChaosRefCountedObject' test. AutoRTFM disabled.")));
+		return true;
+	}
+
+	// It should be safe to declare a ref-counted object which is not used.
+	AutoRTFM::Transact([&]
+		{
+			Chaos::FChaosRefCountedObject DoNothingObject;
+		});
+
+	// Adding and releasing a reference to a transient object will cause it to delete itself.
+	{
+		auto* TransientObject = new Chaos::FChaosRefCountedObject();
+		AutoRTFM::Transact([&]
+			{
+				TransientObject->AddRef();
+				TransientObject->Release();
+			});
+	}
+
+	// Adding and releasing a reference to a persistent object will not delete it.
+	// This tests uses TUniquePtr to perform the deletion when the object falls out of scope.
+	{
+		auto PersistentObject = MakeUnique<Chaos::FChaosRefCountedObject>();
+		PersistentObject->MakePersistent();
+		AutoRTFM::Transact([&]
+			{
+				PersistentObject->AddRef();
+				PersistentObject->Release();
+			});
+	}
+
+	return true;
+}
+
+#endif  // WITH_DEV_AUTOMATION_TESTS
diff --git a/Engine/Source/Runtime/Experimental/Chaos/Public/Chaos/RefCountedObject.h b/Engine/Source/Runtime/Experimental/Chaos/Public/Chaos/RefCountedObject.h
--- a/Engine/Source/Runtime/Experimental/Chaos/Public/Chaos/RefCountedObject.h
+++ b/Engine/Source/Runtime/Experimental/Chaos/Public/Chaos/RefCountedObject.h
@@ -17,7 +17,10 @@ public:
 	FChaosRefCountedObject() : NumRefs(0) {}
 	virtual ~FChaosRefCountedObject()
 	{
-		check(NumRefs.GetValue() == 0);
+		UE_AUTORTFM_ONCOMMIT2(this)
+		{
+			check(NumRefs.GetValue() == 0);
+		};
 	}
 	FChaosRefCountedObject(const FChaosRefCountedObject& Rhs) = delete;
 	FChaosRefCountedObject& operator=(const FChaosRefCountedObject& Rhs) = delete;
