[AutoRTFM] Fix RTFM error when calling MakePersistent outside of an RTFM transaction.

Because bTransientFlag uses std::atomic, we need to use ONCOMMIT2 to update its value. The value can't be read back directly and is only used by Release, which is already wrapped in ONCOMMIT2, so this reordering should be invisible to callers.

diff --git a/Engine/Source/Runtime/Experimental/Chaos/Private/Tests/AutoRTFMRefCountedObjectTests.cpp b/Engine/Source/Runtime/Experimental/Chaos/Private/Tests/AutoRTFMRefCountedObjectTests.cpp
--- a/Engine/Source/Runtime/Experimental/Chaos/Private/Tests/AutoRTFMRefCountedObjectTests.cpp
+++ b/Engine/Source/Runtime/Experimental/Chaos/Private/Tests/AutoRTFMRefCountedObjectTests.cpp
@@ -49,6 +49,17 @@ bool FAutoRTFMChaosRefCountedObject::RunTest(const FString& Parameters)
 			});
 	}
 
+	// It is safe to make an object persistent inside of a transaction.
+	{
+		auto PersistentObject = MakeUnique<Chaos::FChaosRefCountedObject>();
+		AutoRTFM::Transact([&]
+			{
+				PersistentObject->MakePersistent();
+				PersistentObject->AddRef();
+				PersistentObject->Release();
+			});
+	}
+
 	return true;
 }
 
diff --git a/Engine/Source/Runtime/Experimental/Chaos/Public/Chaos/RefCountedObject.h b/Engine/Source/Runtime/Experimental/Chaos/Public/Chaos/RefCountedObject.h
--- a/Engine/Source/Runtime/Experimental/Chaos/Public/Chaos/RefCountedObject.h
+++ b/Engine/Source/Runtime/Experimental/Chaos/Public/Chaos/RefCountedObject.h
@@ -64,7 +64,10 @@ public:
 
 	void MakePersistent() const
 	{
-		bTransientFlag = false;
+		UE_AUTORTFM_ONCOMMIT2(this)
+		{
+			bTransientFlag = false;
+		};
 	}
 
 private:
