Fixed crash when attempting to call an object interface function on an empty view of physics bodies. This is a valid case and the object interface was accessing the first element of an array view unconditionally causing a range check to fail.

diff --git a/Engine/Source/Runtime/Experimental/Chaos/Private/Chaos/PhysicsObjectInterface.cpp b/Engine/Source/Runtime/Experimental/Chaos/Private/Chaos/PhysicsObjectInterface.cpp
--- a/Engine/Source/Runtime/Experimental/Chaos/Private/Chaos/PhysicsObjectInterface.cpp
+++ b/Engine/Source/Runtime/Experimental/Chaos/Private/Chaos/PhysicsObjectInterface.cpp
@@ -31,6 +31,11 @@ namespace PhysicsObjectInterfaceCVars
 
 namespace
 {
+	Chaos::FPhysicsSolverBase* GetSolverForObjects(TArrayView<const Chaos::FPhysicsObjectHandle> InView)
+	{
+		return InView.IsEmpty() ? nullptr : Chaos::FPhysicsObjectInterface::GetSolver(InView[0]);
+	}
+
 	template<Chaos::EThreadContext Id>
 	void SetParticleStateHelper(const Chaos::FPhysicsObjectHandle PhysicsObject, Chaos::EObjectStateType State)
 	{
@@ -989,7 +994,7 @@ namespace Chaos
 
 		if constexpr (Id == EThreadContext::External)
 		{
-			if (Chaos::FPhysicsSolverBase* Solver = FPhysicsObjectInterface::GetSolver(InObjects[0]))
+			if (Chaos::FPhysicsSolverBase* Solver = GetSolverForObjects(InObjects))
 			{
 				Solver->EnqueueCommandImmediate(
 					[AllObjects = TArray<FPhysicsObjectHandle>{ InObjects }, Force, bInvalidate]() {
@@ -1049,7 +1054,7 @@ namespace Chaos
 
 		if constexpr (Id == EThreadContext::External)
 		{
-			if (Chaos::FPhysicsSolverBase* Solver = FPhysicsObjectInterface::GetSolver(InObjects[0]))
+			if (Chaos::FPhysicsSolverBase* Solver = GetSolverForObjects(InObjects))
 			{
 				Solver->EnqueueCommandImmediate(
 					[AllObjects = TArray<FPhysicsObjectHandle>{ InObjects }, Torque, bInvalidate, bAccelChange]() {
@@ -1094,7 +1099,7 @@ namespace Chaos
 
 		if constexpr (Id == EThreadContext::External)
 		{
-			if (Chaos::FPhysicsSolverBase* Solver = FPhysicsObjectInterface::GetSolver(InObjects[0]))
+			if (Chaos::FPhysicsSolverBase* Solver = GetSolverForObjects(InObjects))
 			{
 				Solver->EnqueueCommandImmediate(
 					[AllObjects = TArray<FPhysicsObjectHandle>{InObjects}, Impulse, bVelChange]() {
