Added NetRefHandleManager as a protected member of NetObjectFilter for other internal filters to use and not having to create their own member.
Added protected GetObjectIndex function to NetObjectFilter.  It returns the object index for the given NetRefHandle.
Added IRISCORE_API to UNetObjectGridWorldLocFilter, so that other game modules can derive from this class and extend it.

#tests Tested in a local Test build as well as in PIE

diff --git a/Engine/Source/Runtime/Experimental/Iris/Core/Private/Iris/ReplicationSystem/Filtering/NetObjectFilter.cpp b/Engine/Source/Runtime/Experimental/Iris/Core/Private/Iris/ReplicationSystem/Filtering/NetObjectFilter.cpp
--- a/Engine/Source/Runtime/Experimental/Iris/Core/Private/Iris/ReplicationSystem/Filtering/NetObjectFilter.cpp
+++ b/Engine/Source/Runtime/Experimental/Iris/Core/Private/Iris/ReplicationSystem/Filtering/NetObjectFilter.cpp
@@ -2,6 +2,8 @@
 
 #include "Iris/ReplicationSystem/Filtering/NetObjectFilter.h"
 #include "Iris/ReplicationSystem/Filtering/ReplicationFiltering.h"
+#include "Iris/ReplicationSystem/ReplicationSystem.h"
+#include "Iris/ReplicationSystem/ReplicationSystemInternal.h"
 
 UNetObjectFilter::UNetObjectFilter()
 {
@@ -15,7 +17,7 @@ void UNetObjectFilter::Init(const FNetObjectFilterInitParams& Params)
 		UE::Net::Private::FNetObjectFilteringInfoAccessor FilteringInfoAccessor;
 		FilteringInfos = FilteringInfoAccessor.GetNetObjectFilteringInfos(Params.ReplicationSystem);
 	}
-
+	NetRefHandleManager = &Params.ReplicationSystem->GetReplicationSystemInternal()->GetNetRefHandleManager();
 	OnInit(Params);
 }
 
@@ -24,6 +26,7 @@ void UNetObjectFilter::Deinit()
 	OnDeinit();
 
 	FilteringInfos = TArrayView<FNetObjectFilteringInfo>();
+	NetRefHandleManager = nullptr;
 }
 
 void UNetObjectFilter::MaxInternalNetRefIndexIncreased(UE::Net::Private::FInternalNetRefIndex MaxInternalIndex, TArrayView<FNetObjectFilteringInfo> NewFilterInfoView)
@@ -72,3 +75,7 @@ FNetObjectFilteringInfo* UNetObjectFilter::GetFilteringInfo(uint32 ObjectIndex)
 	return &FilteringInfos[ObjectIndex];
 }
 
+uint32 UNetObjectFilter::GetObjectIndex(UE::Net::FNetRefHandle NetRefHandle) const
+{
+	return NetRefHandleManager->GetInternalIndex(NetRefHandle);
+}
diff --git a/Engine/Source/Runtime/Experimental/Iris/Core/Private/Iris/ReplicationSystem/Filtering/NetObjectGridFilter.cpp b/Engine/Source/Runtime/Experimental/Iris/Core/Private/Iris/ReplicationSystem/Filtering/NetObjectGridFilter.cpp
--- a/Engine/Source/Runtime/Experimental/Iris/Core/Private/Iris/ReplicationSystem/Filtering/NetObjectGridFilter.cpp
+++ b/Engine/Source/Runtime/Experimental/Iris/Core/Private/Iris/ReplicationSystem/Filtering/NetObjectGridFilter.cpp
@@ -5,7 +5,6 @@
 #include "Iris/ReplicationSystem/RepTag.h"
 #include "Iris/ReplicationSystem/ReplicationProtocol.h"
 #include "Iris/ReplicationSystem/ReplicationSystem.h"
-#include "Iris/ReplicationSystem/ReplicationSystemInternal.h"
 #include "Iris/ReplicationSystem/NetRefHandleManager.h"
 #include "Iris/ReplicationSystem/WorldLocations.h"
 #include "Iris/Core/IrisCsv.h"
@@ -23,14 +22,13 @@ void UNetObjectGridFilter::OnInit(const FNetObjectFilterInitParams& Params)
 
 	PerConnectionInfos.SetNum(Params.MaxConnectionCount + 1);
 
-	NetRefHandleManager = &Params.ReplicationSystem->GetReplicationSystemInternal()->GetNetRefHandleManager();
 	NetCullDistanceOverrides = &Params.ReplicationSystem->GetNetCullDistanceOverrides();
 }
 
 void UNetObjectGridFilter::OnDeinit()
 {
 	Config = nullptr;
-	NetRefHandleManager = nullptr;
+
 	NetCullDistanceOverrides = nullptr;
 
 	PerConnectionInfos.Empty();
@@ -583,6 +581,8 @@ void UNetObjectGridWorldLocFilter::OnInit(const FNetObjectFilterInitParams& Para
 void UNetObjectGridWorldLocFilter::OnDeinit()
 {
 	WorldLocations = nullptr;
+
+	Super::OnDeinit();
 }
 
 void UNetObjectGridWorldLocFilter::PreFilter(FNetObjectPreFilteringParams& Params)
diff --git a/Engine/Source/Runtime/Experimental/Iris/Core/Private/Iris/ReplicationSystem/Prioritization/NetObjectConnectionFilter.cpp b/Engine/Source/Runtime/Experimental/Iris/Core/Private/Iris/ReplicationSystem/Prioritization/NetObjectConnectionFilter.cpp
--- a/Engine/Source/Runtime/Experimental/Iris/Core/Private/Iris/ReplicationSystem/Prioritization/NetObjectConnectionFilter.cpp
+++ b/Engine/Source/Runtime/Experimental/Iris/Core/Private/Iris/ReplicationSystem/Prioritization/NetObjectConnectionFilter.cpp
@@ -15,7 +15,8 @@ void UNetObjectConnectionFilter::SetReplicateToConnection(UE::Net::FNetRefHandle
 		return;
 	}
 
-	const UE::Net::Private::FInternalNetRefIndex ObjectIndex = NetRefHandleManager->GetInternalIndex(RefHandle);
+	const UE::Net::Private::FInternalNetRefIndex ObjectIndex = GetObjectIndex(RefHandle);
+
 	if (ObjectIndex == UE::Net::Private::FNetRefHandleManager::InvalidInternalIndex)
 	{
 		return;
@@ -33,8 +34,6 @@ void UNetObjectConnectionFilter::SetReplicateToConnection(UE::Net::FNetRefHandle
 
 void UNetObjectConnectionFilter::OnInit(const FNetObjectFilterInitParams& Params)
 {
-	NetRefHandleManager = &Params.ReplicationSystem->GetReplicationSystemInternal()->GetNetRefHandleManager();
-
 	MaxInternalIndex = Params.CurrentMaxInternalIndex; 
 
 	Config = TStrongObjectPtr<UNetObjectConnectionFilterConfig>(CastChecked<UNetObjectConnectionFilterConfig>(Params.Config));
@@ -48,7 +47,6 @@ void UNetObjectConnectionFilter::OnInit(const FNetObjectFilterInitParams& Params
 
 void UNetObjectConnectionFilter::OnDeinit()
 {
-	NetRefHandleManager = nullptr;
 	Config = nullptr;
 
 	UsedLocalInfoIndices.Empty();
diff --git a/Engine/Source/Runtime/Experimental/Iris/Core/Public/Iris/ReplicationSystem/Filtering/NetObjectConnectionFilter.h b/Engine/Source/Runtime/Experimental/Iris/Core/Public/Iris/ReplicationSystem/Filtering/NetObjectConnectionFilter.h
--- a/Engine/Source/Runtime/Experimental/Iris/Core/Public/Iris/ReplicationSystem/Filtering/NetObjectConnectionFilter.h
+++ b/Engine/Source/Runtime/Experimental/Iris/Core/Public/Iris/ReplicationSystem/Filtering/NetObjectConnectionFilter.h
@@ -60,7 +60,6 @@ protected:
 	IRISCORE_API virtual void PreFilter(FNetObjectPreFilteringParams&) override;
 	IRISCORE_API virtual void Filter(FNetObjectFilteringParams&) override;
 
-	const UE::Net::Private::FNetRefHandleManager* NetRefHandleManager = nullptr;
 	TStrongObjectPtr<UNetObjectConnectionFilterConfig> Config;
 	TArray<uint32> LocalToNetRefIndex;
 	TArray<FPerConnectionInfo> PerConnectionInfos;
diff --git a/Engine/Source/Runtime/Experimental/Iris/Core/Public/Iris/ReplicationSystem/Filtering/NetObjectFilter.h b/Engine/Source/Runtime/Experimental/Iris/Core/Public/Iris/ReplicationSystem/Filtering/NetObjectFilter.h
--- a/Engine/Source/Runtime/Experimental/Iris/Core/Public/Iris/ReplicationSystem/Filtering/NetObjectFilter.h
+++ b/Engine/Source/Runtime/Experimental/Iris/Core/Public/Iris/ReplicationSystem/Filtering/NetObjectFilter.h
@@ -15,10 +15,12 @@ namespace UE::Net
 	typedef uint32 FNetObjectFilterHandle;
 	struct FReplicationInstanceProtocol;
 	struct FReplicationProtocol;
+	class FNetRefHandle;
 
 	namespace Private
 	{
 		typedef uint32 FInternalNetRefIndex;
+		class FNetRefHandleManager;
 	}
 }
 
@@ -180,7 +182,7 @@ public:
 	/** A new connection has been added. An opportunity for the filter to allocate per connection info. */
 	IRISCORE_API virtual void AddConnection(uint32 ConnectionId);
 
-	/** A new connection has been added. An opportunity for the filter to deallocate per connection info. */
+	/** A new connection has been removed. An opportunity for the filter to deallocate per connection info. */
 	IRISCORE_API virtual void RemoveConnection(uint32 ConnectionId);
 
 	/** A new object want to use this filter. Opportunity to cache some information for it. The info struct passed has been zeroed. Must be overriden. */
@@ -245,6 +247,9 @@ protected:
 	/* Returns the filtering info for this object if it's handled by this filter, nullptr otherwise. */
 	IRISCORE_API FNetObjectFilteringInfo* GetFilteringInfo(uint32 ObjectIndex);
 
+	/* Returns the object index for the given NetRefHandle */
+	IRISCORE_API uint32 GetObjectIndex(UE::Net::FNetRefHandle NetRefHandle) const;
+
 	/* Returns true if the object is assigned to be filtered by this filter.*/
 	inline bool IsObjectFiltered(uint32 ObjectIndex) const;
 
@@ -259,6 +264,8 @@ protected:
 	/** The indices of the objects that have this filter set. The indices of set bits correspond to the object indices. */
 	UE::Net::FNetBitArray FilteredObjects;
 
+	const UE::Net::Private::FNetRefHandleManager* NetRefHandleManager = nullptr;
+
 private:
 	ENetFilterTraits FilterTraits = ENetFilterTraits::None;
 	
diff --git a/Engine/Source/Runtime/Experimental/Iris/Core/Public/Iris/ReplicationSystem/Filtering/NetObjectGridFilter.h b/Engine/Source/Runtime/Experimental/Iris/Core/Public/Iris/ReplicationSystem/Filtering/NetObjectGridFilter.h
--- a/Engine/Source/Runtime/Experimental/Iris/Core/Public/Iris/ReplicationSystem/Filtering/NetObjectGridFilter.h
+++ b/Engine/Source/Runtime/Experimental/Iris/Core/Public/Iris/ReplicationSystem/Filtering/NetObjectGridFilter.h
@@ -15,11 +15,6 @@ namespace UE::Net
 	class FNetCullDistanceOverrides;
 	class FWorldLocations;
 	struct FRepTagFindInfo;
-
-	namespace Private
-	{
-		class FNetRefHandleManager;
-	}
 }
 
 /**
@@ -293,7 +288,6 @@ private:
 	// Cached state from UpdateCellInfoForObject().
 	FDebugUpdateCellInfo DebugUpdateCellInfo;
 
-	const UE::Net::Private::FNetRefHandleManager*  NetRefHandleManager = nullptr;
 	const UE::Net::FNetCullDistanceOverrides* NetCullDistanceOverrides = nullptr;
 };
 
@@ -302,18 +296,17 @@ private:
  * 
  * This filter is more efficient since it's run before Polling and culls out objects that are not relevant to any connection.
  */
-UCLASS()
+UCLASS(transient, MinimalAPI)
 class UNetObjectGridWorldLocFilter : public UNetObjectGridFilter
 {
 	GENERATED_BODY()
 
 protected:
-
-	virtual void OnInit(const FNetObjectFilterInitParams&) override;
-	virtual void OnDeinit() override;
-	virtual void PreFilter(FNetObjectPreFilteringParams&) override;
-	virtual void UpdateObjectInfo(FPerObjectInfo& PerObjectInfo, const UNetObjectGridFilter::FObjectLocationInfo& ObjectLocationInfo, const UE::Net::FReplicationInstanceProtocol* InstanceProtocol) override;
-	virtual bool BuildObjectInfo(uint32 ObjectIndex, FNetObjectFilterAddObjectParams& Params) override;
+	IRISCORE_API virtual void OnInit(const FNetObjectFilterInitParams&) override;
+	IRISCORE_API virtual void OnDeinit() override;
+	IRISCORE_API virtual void PreFilter(FNetObjectPreFilteringParams&) override;
+	IRISCORE_API virtual void UpdateObjectInfo(FPerObjectInfo& PerObjectInfo, const UNetObjectGridFilter::FObjectLocationInfo& ObjectLocationInfo, const UE::Net::FReplicationInstanceProtocol* InstanceProtocol) override;
+	IRISCORE_API virtual bool BuildObjectInfo(uint32 ObjectIndex, FNetObjectFilterAddObjectParams& Params) override;
 
 private:
 
