Improve error message if we are unable to read from a .utoc file while attempting to upload it for OnDemand.

#rnx

diff --git a/Engine/Source/Runtime/Experimental/IoStore/OnDemand/Private/Tool/Upload.cpp b/Engine/Source/Runtime/Experimental/IoStore/OnDemand/Private/Tool/Upload.cpp
--- a/Engine/Source/Runtime/Experimental/IoStore/OnDemand/Private/Tool/Upload.cpp
+++ b/Engine/Source/Runtime/Experimental/IoStore/OnDemand/Private/Tool/Upload.cpp
@@ -507,7 +507,7 @@ static TIoStatusOr<FUploadResult> UploadContainerFiles(
 			FIoStatus Status = ContainerFileReader.Initialize(*FPaths::ChangeExtension(Path, TEXT("")), EncryptionKeys);
 			if (!Status.IsOk())
 			{
-				UE_LOG(LogIas, Error, TEXT("Failed to open container '%s' for reading"), *Path);
+				UE_LOG(LogIas, Error, TEXT("Failed to open container '%s' for reading due to: '%s'"), *Path, *Status.ToString());
 				continue;
 			}
 		}
