New Geometry Collection option to always update active transforms - so vehicle parts can fall off and remain on the ground rather than following the root actors movement

diff --git a/Engine/Source/Runtime/Experimental/Chaos/Private/PhysicsProxy/GeometryCollectionPhysicsProxy.cpp b/Engine/Source/Runtime/Experimental/Chaos/Private/PhysicsProxy/GeometryCollectionPhysicsProxy.cpp
--- a/Engine/Source/Runtime/Experimental/Chaos/Private/PhysicsProxy/GeometryCollectionPhysicsProxy.cpp
+++ b/Engine/Source/Runtime/Experimental/Chaos/Private/PhysicsProxy/GeometryCollectionPhysicsProxy.cpp
@@ -3971,19 +3971,19 @@ void FGeometryCollectionPhysicsProxy::FlipBuffer()
 }
 
 // Update game thread particle X and R properties and retun true if they changed from previous values
-static inline bool UpdateGTParticleXR(Chaos::FPBDRigidParticle& GTParticle, const Chaos::FVec3& NewX, const Chaos::FRotation3& NewR)
+static inline bool UpdateGTParticleXR(Chaos::FPBDRigidParticle& GTParticle, const Chaos::FVec3& NewX, const Chaos::FRotation3& NewR, bool bForceUpdateActiveTransforms)
 {
 	GC_PHYSICSPROXY_CHECK_FOR_NAN(NewX);
 
 	const Chaos::FVec3 OldX = GTParticle.GetX();
-	const bool bNeedUpdateX = (!NewX.Equals(OldX, GeometryCollectionPositionUpdateTolerance));
+	const bool bNeedUpdateX = bForceUpdateActiveTransforms || (!NewX.Equals(OldX, GeometryCollectionPositionUpdateTolerance));
 	if (bNeedUpdateX)
 	{
 		GTParticle.SetX(NewX, false);
 	}
 
 	const Chaos::FRotation3 OldR = GTParticle.R();
-	const bool bNeedUpdateR = (!NewR.Equals(OldR, GeometryCollectionRotationUpdateTolerance));
+	const bool bNeedUpdateR = bForceUpdateActiveTransforms || (!NewR.Equals(OldR, GeometryCollectionRotationUpdateTolerance));
 	if (bNeedUpdateR)
 	{
 		GTParticle.SetR(NewR, false);
@@ -4223,7 +4223,7 @@ bool FGeometryCollectionPhysicsProxy::PullNonInterpolatableDataFromSinglePhysics
 			{
 				const FGeometryCollectionResults::FPositionData& PositionData = CurrentResults.GetPositions(EntryIndex);
 
-				const bool XRModified = UpdateGTParticleXR(GTParticle, PositionData.ParticleX, PositionData.ParticleR);
+				const bool XRModified = UpdateGTParticleXR(GTParticle, PositionData.ParticleX, PositionData.ParticleR, Parameters.bForceUpdateActiveTransforms);
 				bIsCollectionDirty |= XRModified;
 
 				if (LinearVelocities && AngularVelocities)
@@ -4383,7 +4383,7 @@ bool FGeometryCollectionPhysicsProxy::PullFromPhysicsState(const Chaos::FDirtyGe
 						NewR = InterpData->GetErrorR(*Alpha) * NewR;
 					}
 
-					const bool XRModified = UpdateGTParticleXR(GTParticle, NewX, NewR);
+					const bool XRModified = UpdateGTParticleXR(GTParticle, NewX, NewR, Parameters.bForceUpdateActiveTransforms);
 					bIsCollectionDirty |= XRModified;
 
 					if (LinearVelocities && AngularVelocities)
diff --git a/Engine/Source/Runtime/Experimental/Chaos/Public/GeometryCollection/GeometryCollectionSimulationCoreTypes.h b/Engine/Source/Runtime/Experimental/Chaos/Public/GeometryCollection/GeometryCollectionSimulationCoreTypes.h
--- a/Engine/Source/Runtime/Experimental/Chaos/Public/GeometryCollection/GeometryCollectionSimulationCoreTypes.h
+++ b/Engine/Source/Runtime/Experimental/Chaos/Public/GeometryCollection/GeometryCollectionSimulationCoreTypes.h
@@ -198,6 +198,7 @@ PRAGMA_ENABLE_DEPRECATION_WARNINGS
 		, bUseMaterialDamageModifiers(false)
 		, bUsePerClusterOnlyDamageThreshold(false)
 		, StartAwake(true)
+		, bForceUpdateActiveTransforms(false)
 		, bGenerateBreakingData(false)
 		, bGenerateCollisionData(false)
 		, bGenerateTrailingData(false)
@@ -271,6 +272,7 @@ PRAGMA_ENABLE_DEPRECATION_WARNINGS
 		, bUseMaterialDamageModifiers(Other.bUseMaterialDamageModifiers)
 		, bUsePerClusterOnlyDamageThreshold(Other.bUsePerClusterOnlyDamageThreshold)
 		, StartAwake(Other.StartAwake)
+		, bForceUpdateActiveTransforms(Other.bForceUpdateActiveTransforms)
 		, bGenerateBreakingData(Other.bGenerateBreakingData)
 		, bGenerateCollisionData(Other.bGenerateCollisionData)
 		, bGenerateTrailingData(Other.bGenerateTrailingData)
@@ -375,6 +377,7 @@ PRAGMA_ENABLE_DEPRECATION_WARNINGS
 	bool bUseMaterialDamageModifiers : 1;
 	bool bUsePerClusterOnlyDamageThreshold : 1;
 	bool StartAwake : 1;
+	bool bForceUpdateActiveTransforms : 1;
 
 	bool bGenerateBreakingData : 1;
 	bool bGenerateCollisionData : 1;
diff --git a/Engine/Source/Runtime/Experimental/GeometryCollectionEngine/Private/GeometryCollection/GeometryCollectionComponent.cpp b/Engine/Source/Runtime/Experimental/GeometryCollectionEngine/Private/GeometryCollection/GeometryCollectionComponent.cpp
--- a/Engine/Source/Runtime/Experimental/GeometryCollectionEngine/Private/GeometryCollection/GeometryCollectionComponent.cpp
+++ b/Engine/Source/Runtime/Experimental/GeometryCollectionEngine/Private/GeometryCollection/GeometryCollectionComponent.cpp
@@ -3914,6 +3914,7 @@ PRAGMA_ENABLE_DEPRECATION_WARNINGS
 		SimulationParameters.bUsePerClusterOnlyDamageThreshold = RestCollection? RestCollection->PerClusterOnlyDamageThreshold: false; 
 		SimulationParameters.ClusterConnectionMethod = (Chaos::FClusterCreationParameters::EConnectionMethod)ClusterCollectionType;
 		SimulationParameters.ConnectionGraphBoundsFilteringMargin = ConnectionGraphBoundsFilteringMargin; 
+		SimulationParameters.bForceUpdateActiveTransforms = bForceUpdateActiveTransforms;
 		SimulationParameters.CollisionGroup = CollisionGroup;
 		SimulationParameters.CollisionSampleFraction = CollisionSampleFraction;
 		SimulationParameters.InitialVelocityType = InitialVelocityType;
diff --git a/Engine/Source/Runtime/Experimental/GeometryCollectionEngine/Public/GeometryCollection/GeometryCollectionComponent.h b/Engine/Source/Runtime/Experimental/GeometryCollectionEngine/Public/GeometryCollection/GeometryCollectionComponent.h
--- a/Engine/Source/Runtime/Experimental/GeometryCollectionEngine/Public/GeometryCollection/GeometryCollectionComponent.h
+++ b/Engine/Source/Runtime/Experimental/GeometryCollectionEngine/Public/GeometryCollection/GeometryCollectionComponent.h
@@ -1044,6 +1044,10 @@ public:
 	UPROPERTY(EditAnywhere, BlueprintReadWrite, Category = "ChaosPhysics|Removal")
 	bool bAllowRemovalOnBreak;
 	
+	/** Update transforms of active particles even when they are not moving. Has performance implications. Use only when GC is a child of a moving actor, to prevent released particle 'following the actor around' */
+	UPROPERTY(EditAnywhere, BlueprintReadWrite, Category = "ChaosPhysics|Clustering")
+	bool bForceUpdateActiveTransforms;
+
 	/** */
 	UPROPERTY(meta = (DeprecatedProperty, DeprecationMessage = "Connection types are defined on the asset now."))
 	EClusterConnectionTypeEnum ClusterConnectionType_DEPRECATED;
