Chaos::add CVar for joint simd in scene simulation.

diff --git a/Engine/Source/Runtime/Experimental/Chaos/Private/Chaos/PBDRigidsEvolutionGBF.cpp b/Engine/Source/Runtime/Experimental/Chaos/Private/Chaos/PBDRigidsEvolutionGBF.cpp
--- a/Engine/Source/Runtime/Experimental/Chaos/Private/Chaos/PBDRigidsEvolutionGBF.cpp
+++ b/Engine/Source/Runtime/Experimental/Chaos/Private/Chaos/PBDRigidsEvolutionGBF.cpp
@@ -137,6 +137,8 @@ namespace Chaos
 		int32 ChaosOneWayInteractionPairCollisionMode = (int32)EOneWayInteractionPairCollisionMode::SphereCollision;
 		FAutoConsoleVariableRef CVarChaosIgnoreOneWayPairCollisions(TEXT("p.Chaos.Solver.OneWayPairCollisionMode"), ChaosOneWayInteractionPairCollisionMode, TEXT("How to treat collisions between two one-way interaction particles. See EOneWayInteractionPairCollisionMode (0: Ignore collisions; 1: Collide as normal; 2: Collide as spheres)"));
 
+		bool bChaosRigids_UseSimdForJointsSolver = true;
+		FAutoConsoleVariableRef  CVarChaosImmPhysUseSimdForLinearSolver(TEXT("p.Chaos.Solver.Joint.UseSimd"), bChaosRigids_UseSimdForJointsSolver, TEXT("Enable/Disable SIMD on the linear joint solver"));
 
 		DECLARE_CYCLE_STAT(TEXT("FPBDRigidsEvolutionGBF::AdvanceOneTimeStep"), STAT_Evolution_AdvanceOneTimeStep, STATGROUP_Chaos);
 		DECLARE_CYCLE_STAT(TEXT("FPBDRigidsEvolutionGBF::UnclusterUnions"), STAT_Evolution_UnclusterUnions, STATGROUP_Chaos);
@@ -1118,6 +1120,8 @@ FPBDRigidsEvolutionGBF::FPBDRigidsEvolutionGBF(
 
 	CollisionConstraints.SetCullDistance(DefaultCollisionCullDistance);
 
+	JointConstraints.SetUseSimd(bChaosRigids_UseSimdForJointsSolver);
+
 	GetIslandManager().SetMaterialContainers(&PhysicsMaterials, &PerParticlePhysicsMaterials, &SolverPhysicsMaterials);
 	GetIslandManager().SetGravityForces(&GravityForces);
 	GetIslandManager().SetDisableCounterThreshold(DisableThreshold);
