Iris - Increased send window size on ReliableNetBlobQueue from 256 to 1024.

diff --git a/Engine/Source/Runtime/Experimental/Iris/Core/Private/Iris/ReplicationSystem/AttachmentReplication.cpp b/Engine/Source/Runtime/Experimental/Iris/Core/Private/Iris/ReplicationSystem/AttachmentReplication.cpp
--- a/Engine/Source/Runtime/Experimental/Iris/Core/Private/Iris/ReplicationSystem/AttachmentReplication.cpp
+++ b/Engine/Source/Runtime/Experimental/Iris/Core/Private/Iris/ReplicationSystem/AttachmentReplication.cpp
@@ -24,7 +24,7 @@ namespace AttachmentReplicationCVars
 	FAutoConsoleVariableRef CVarUnreliableRPCQueueSize(TEXT("net.UnreliableRPCQueueSize"), UnreliableRPCQueueSize, TEXT("Maximum number of unreliable RPCs queued per object. If more RPCs are queued then older ones will be dropped."));
 
 	static int32 ReliableRPCQueueSize = 4096;
-	FAutoConsoleVariableRef CVarReliableRPCQueueSize(TEXT("net.ReliableRPCQueueSize"), ReliableRPCQueueSize, TEXT("Maximum number of reliable RPCs queued per object. This is in addition to the 256 that are in the send window. This is to support very large RPCs that are split into smaller pieces."));
+	FAutoConsoleVariableRef CVarReliableRPCQueueSize(TEXT("net.ReliableRPCQueueSize"), ReliableRPCQueueSize, TEXT("Maximum number of reliable RPCs queued per object. This is in addition to the send window size. This is to support very large RPCs that are split into smaller pieces."));
 
 	static int32 ClientToServerUnreliableRPCQueueSize = 16;
 	FAutoConsoleVariableRef CVarClientToServerUnreliableRPCQueueSize(TEXT("net.ClientToServerUnreliableRPCQueueSize"), ClientToServerUnreliableRPCQueueSize, TEXT( "Maximum number of unreliable RPCs queued for sending from the client to the server. If more RPCs are queued then older ones will be dropped."));
diff --git a/Engine/Source/Runtime/Experimental/Iris/Core/Private/Iris/ReplicationSystem/NetBlob/ReliableNetBlobQueue.cpp b/Engine/Source/Runtime/Experimental/Iris/Core/Private/Iris/ReplicationSystem/NetBlob/ReliableNetBlobQueue.cpp
--- a/Engine/Source/Runtime/Experimental/Iris/Core/Private/Iris/ReplicationSystem/NetBlob/ReliableNetBlobQueue.cpp
+++ b/Engine/Source/Runtime/Experimental/Iris/Core/Private/Iris/ReplicationSystem/NetBlob/ReliableNetBlobQueue.cpp
@@ -175,7 +175,7 @@ uint32 FReliableNetBlobQueue::SerializeInternal(FNetSerializationContext& Contex
 			++SerializedCount;
 
 			++WrittenCount[WrittenIndex];
-			if (WrittenCount[WrittenIndex] == 255)
+			if (WrittenCount[WrittenIndex] == MaxSequenceLength)
 			{
 				++WrittenIndex;
 				if (WrittenIndex >= UE_ARRAY_COUNT(WrittenCount))
@@ -193,8 +193,8 @@ uint32 FReliableNetBlobQueue::SerializeInternal(FNetSerializationContext& Contex
 	FReplicationRecord Record;
 	for (uint32 Index = 0; Index < MaxWriteSequenceCount; ++Index)
 	{
-		Record.Sequences[Index] = WrittenSeq[Index] & 255U;
-		Record.Counts[Index] = WrittenCount[Index] & 255U;
+		Record.Sequences[Index].Number = WrittenSeq[Index] & (MaxUnackedBlobCount - 1U);
+		Record.Sequences[Index].Count = WrittenCount[Index] & MaxSequenceLength;
 	}
 	OutRecord = Record;
 
@@ -205,9 +205,9 @@ void FReliableNetBlobQueue::CommitReplicationRecord(const FReliableNetBlobQueue:
 {
 	for (uint32 Index = 0, EndIndex = MaxWriteSequenceCount; Index != EndIndex; ++Index)
 	{
-		const uint32 Count = Record.Counts[Index];
+		const uint32 Count = Record.Sequences[Index].Count;
 		UnsentBlobCount -= Count;
-		for (uint32 Seq = Record.Sequences[Index], EndSeq = Seq + Count; Seq != EndSeq; ++Seq)
+		for (uint32 Seq = Record.Sequences[Index].Number, EndSeq = Seq + Count; Seq != EndSeq; ++Seq)
 		{
 			const uint32 BlobIndex = SequenceToIndex(Seq);
 			SetIndexIsSent(BlobIndex);
@@ -443,8 +443,8 @@ void FReliableNetBlobQueue::OnPacketDelivered(const FReliableNetBlobQueue::FRepl
 	// Mark blobs as acked
 	for (uint32 SeqIt = 0, EndSeqIt = MaxWriteSequenceCount; SeqIt != EndSeqIt; ++SeqIt)
 	{
-		const uint32 Count = Record.Counts[SeqIt];
-		for (uint32 Seq = Record.Sequences[SeqIt], EndSeq = Seq + Count; Seq != EndSeq; ++Seq)
+		const uint32 Count = Record.Sequences[SeqIt].Count;
+		for (uint32 Seq = Record.Sequences[SeqIt].Number, EndSeq = Seq + Count; Seq != EndSeq; ++Seq)
 		{
 			const uint32 Index = SequenceToIndex(Seq);
 
@@ -464,9 +464,9 @@ void FReliableNetBlobQueue::OnPacketDropped(const FReliableNetBlobQueue::FReplic
 	// Mark blobs as unsent
 	for (uint32 SeqIt = 0, EndSeqIt = MaxWriteSequenceCount; SeqIt != EndSeqIt; ++SeqIt)
 	{
-		const uint32 Count = Record.Counts[SeqIt];
+		const uint32 Count = Record.Sequences[SeqIt].Count;
 		UnsentBlobCount += Count;
-		for (uint32 Seq = Record.Sequences[SeqIt], EndSeq = Seq + Count; Seq != EndSeq; ++Seq)
+		for (uint32 Seq = Record.Sequences[SeqIt].Number, EndSeq = Seq + Count; Seq != EndSeq; ++Seq)
 		{
 			ClearSequenceIsSent(Seq);
 		}
diff --git a/Engine/Source/Runtime/Experimental/Iris/Core/Public/Iris/ReplicationSystem/NetBlob/ReliableNetBlobQueue.h b/Engine/Source/Runtime/Experimental/Iris/Core/Public/Iris/ReplicationSystem/NetBlob/ReliableNetBlobQueue.h
--- a/Engine/Source/Runtime/Experimental/Iris/Core/Public/Iris/ReplicationSystem/NetBlob/ReliableNetBlobQueue.h
+++ b/Engine/Source/Runtime/Experimental/Iris/Core/Public/Iris/ReplicationSystem/NetBlob/ReliableNetBlobQueue.h
@@ -24,15 +24,39 @@ namespace UE::Net::Private
 class FReliableNetBlobQueue
 {
 public:
+	// A single ReplicationRecord supports up to four disjoint sequences to be serialized in a packet. This allows for resends of dropped data while potentially also writing new blobs for the first time.
 	struct FReplicationRecord
 	{
+		struct FSequence
+		{
+			enum : unsigned
+			{
+				NumberBitCount = 10U,
+				CountBitCount = 6U,
+			};
+
+			FSequence() = default;
+			explicit FSequence(uint16 Value)
+			{
+				Number = Value & ((1U << NumberBitCount) - 1U);
+				Count = Value >> NumberBitCount;
+			}
+
+			uint16 ToUint16() const
+			{
+				return static_cast<uint16>((Count << NumberBitCount) | Number);
+			}
+
+			uint16 Number : NumberBitCount = 0;
+			uint16 Count : CountBitCount = 0;
+		};
+
 		FReplicationRecord() = default;
 		explicit FReplicationRecord(uint64 Value)
 		{
 			for (unsigned Index : {3U, 2U, 1U, 0U})
 			{
-				Counts[Index] = Value & 255U;
-				Sequences[Index] = (Value >> 8U) & 255U;
+				Sequences[Index] = FSequence(Value & 0xFFFFU);
 				Value >>= 16U;
 			}
 		}
@@ -42,22 +66,21 @@ public:
 			uint64 Value = 0;
 			for (unsigned Index : {0U, 1U, 2U, 3U})
 			{
-				Value = (Value << 16U) | (uint64(Sequences[Index]) << 8U) | Counts[Index];
+				Value = (Value << 16U) | Sequences[Index].ToUint16();
 			}
 
 			return Value;
 		}
 
-		bool IsValid() const { return Counts[0] | Counts[1] | Counts[2] | Counts[3]; }
+		bool IsValid() const { return (Sequences[0].Count | Sequences[1].Count | Sequences[2].Count | Sequences[3].Count) != 0U; }
 
-		uint8 Sequences[4] = {};
-		uint8 Counts[4] = {};
+		FSequence Sequences[4];
 	};
 
 	/** This represents a ReplicationRecord where nothing was serialized. */
 
-	/** How many blobs can be sent before an ACK/NAK is required to continue sending. */
-	static constexpr uint32 MaxUnackedBlobCount = 256U;
+	/** How many blobs can be sent before an ACK/NAK is required to continue sending. Changing this might require changing FReplicationRecord too. */
+	static constexpr uint32 MaxUnackedBlobCount = 1024U;
 
 	FReliableNetBlobQueue();
 	~FReliableNetBlobQueue();
@@ -140,8 +163,9 @@ public:
 private:
 	enum Constants : uint32
 	{
-		IndexBitCount = 8U,
+		IndexBitCount = 10U,
 		MaxWriteSequenceCount = UE_ARRAY_COUNT(FReplicationRecord::Sequences),
+		MaxSequenceLength = (1U << FReplicationRecord::FSequence::CountBitCount) - 1U,
 	};
 
 	uint32 SequenceToIndex(uint32 Seq) const;
diff --git a/Engine/Source/Runtime/Experimental/Iris/Core/Public/Iris/ReplicationSystem/NetBlob/SequentialPartialNetBlobHandler.h b/Engine/Source/Runtime/Experimental/Iris/Core/Public/Iris/ReplicationSystem/NetBlob/SequentialPartialNetBlobHandler.h
--- a/Engine/Source/Runtime/Experimental/Iris/Core/Public/Iris/ReplicationSystem/NetBlob/SequentialPartialNetBlobHandler.h
+++ b/Engine/Source/Runtime/Experimental/Iris/Core/Public/Iris/ReplicationSystem/NetBlob/SequentialPartialNetBlobHandler.h
@@ -36,7 +36,7 @@ protected:
 
 	/** How many parts a NetBlob can be split into at most. If more parts are required the splitting will fail. Cannot exceed 65535. */
 	UPROPERTY(Config)
-	uint32 MaxPartCount = 1024;
+	uint32 MaxPartCount = 4096;
 };
 
 UCLASS(abstract, MinimalApi, transient)
