Fix assert about CachedDensityFromPhysicsMaterialInGCm3 when saving an empty geometry collection
- make sure the cache data is computed before saving of the asset happens

diff --git a/Engine/Source/Runtime/Experimental/GeometryCollectionEngine/Private/GeometryCollection/GeometryCollectionObject.cpp b/Engine/Source/Runtime/Experimental/GeometryCollectionEngine/Private/GeometryCollection/GeometryCollectionObject.cpp
--- a/Engine/Source/Runtime/Experimental/GeometryCollectionEngine/Private/GeometryCollection/GeometryCollectionObject.cpp
+++ b/Engine/Source/Runtime/Experimental/GeometryCollectionEngine/Private/GeometryCollection/GeometryCollectionObject.cpp
@@ -1927,7 +1927,10 @@ void UGeometryCollection::PreSave(FObjectPreSaveContext SaveContext)
 {
 #if WITH_EDITOR
 	constexpr bool bInitResources = false; 
-	EnsureRenderDataIsCooked(bInitResources);
+	constexpr bool bIsTransacting = false;
+	constexpr bool bIsPersistant = false; // note that this has no effect on the call below  
+	constexpr bool bAllowCopyFromDDC = false;
+	EnsureDataIsCooked(bInitResources, bIsTransacting, bIsPersistant, bAllowCopyFromDDC);
 #endif
 
 	Super::PreSave(SaveContext);
