fix crashes related to momentuim restoring when particle is destroyed / disabled by the clustering code
- wrap particle disabling for clustering code in a method
- change the structure for momentum restoring , so it can be indexed by particle ( remove dependency on the collision constraint )

diff --git a/Engine/Source/Runtime/Experimental/Chaos/Private/Chaos/PBDRigidClustering.cpp b/Engine/Source/Runtime/Experimental/Chaos/Private/Chaos/PBDRigidClustering.cpp
--- a/Engine/Source/Runtime/Experimental/Chaos/Private/Chaos/PBDRigidClustering.cpp
+++ b/Engine/Source/Runtime/Experimental/Chaos/Private/Chaos/PBDRigidClustering.cpp
@@ -159,14 +159,14 @@ namespace Chaos
 		}
 
 		template<typename TParticleContainer>
-		void GenericThrottleReleasedParticlesIfNecessary(TParticleContainer& Container, typename FRigidClustering::FRigidEvolution& MEvolution)
+		void GenericThrottleReleasedParticlesIfNecessary(TParticleContainer& Container, FRigidClustering& Clustering)
 		{
 			if (!CVarShouldThrottleParticleRelease())
 			{
 				return;
 			}
 
-			const float RatioOfParticlesToDisable = GetRatioOfReleasedParticlesToDisable(MEvolution);
+			const float RatioOfParticlesToDisable = GetRatioOfReleasedParticlesToDisable(Clustering.GetEvolution());
 			const int32 NumberOfParticlesToDisable = (int32)((float)Container.Num() * RatioOfParticlesToDisable);
 			if (NumberOfParticlesToDisable > 0)
 			{ 
@@ -176,8 +176,8 @@ namespace Chaos
 					if (FPBDRigidParticleHandle* Child = *ChildIt)
 					{
 						DisabledParticleCount++;
-						MEvolution.DisableParticle(Child);
-						MEvolution.GetParticles().MarkTransientDirtyParticle(Child);
+						Clustering.DisableParticle(Child);
+						Clustering.GetEvolution().GetParticles().MarkTransientDirtyParticle(Child);
 						ChildIt.RemoveCurrent();
 					}
 					if (DisabledParticleCount >= NumberOfParticlesToDisable)
@@ -708,8 +708,29 @@ namespace Chaos
 					return ECollisionVisitorResult::Continue;
 				}
 
-				// Track this collision
-				BreakingCollisions.Add(TPair<FPBDCollisionConstraint*, FPBDRigidParticleHandle*>(&Collision, OtherRigid));
+				// Flip the impulse if we're restoring particle 0's momentum.
+				// This is because by convention constraint impulses point from 1 to 0.
+				const FVec3 CollisionImpulse = 
+					(Collision.GetParticle0() == OtherRigid)
+					? -Collision.AccumulatedImpulse
+					: Collision.AccumulatedImpulse;
+
+				FConstGenericParticleHandle Generic(OtherRigid);
+				// Compute the angular impulse based on distance from the contact point to the CoM
+				const FVec3 Location = Collision.CalculateWorldContactLocation();
+				const Chaos::FVec3 AngularImpulse = Chaos::FVec3::CrossProduct(Location - Generic->PCom(), CollisionImpulse);
+
+				// Compute impulse velocities
+				const FVec3 ImpulseVelocity = Generic->InvM() * CollisionImpulse;
+
+				const FMatrix33 OtherInvI = Utilities::ComputeWorldSpaceInertia(Generic->QCom(), Generic->ConditionedInvI());
+				const FVec3 AngularImpulseVelocity = OtherInvI * AngularImpulse;
+
+				// Update linear and angular impulses for the body, to be integrated next solve
+
+				FMomentumRestoringData& MomentumRestoringData = MomentumRestoringDataByParticle.FindOrAdd(OtherRigid);
+				MomentumRestoringData.V += ImpulseVelocity * RestoreBreakingMomentumPercent;
+				MomentumRestoringData.W += AngularImpulseVelocity * RestoreBreakingMomentumPercent;
 
 				return ECollisionVisitorResult::Continue;
 			});
@@ -718,37 +739,22 @@ namespace Chaos
 
 	void FRigidClustering::RestoreBreakingMomentum()
 	{
-		for (TPair<FPBDCollisionConstraint*, FPBDRigidParticleHandle*>& Pair : BreakingCollisions)
+		for (TPair<FPBDRigidParticleHandle*, FMomentumRestoringData> Pair : MomentumRestoringDataByParticle)
 		{
-			FPBDCollisionConstraint& Collision = *Pair.Key;
-			FPBDRigidParticleHandle& Rigid = *Pair.Value;
-			FConstGenericParticleHandle Generic(&Rigid);
-
-			// Flip the impulse if we're restoring particle 0's momentum.
-			// This is because by convention constraint impulses point from 1 to 0.
-			uint8 OtherIdx = Collision.GetParticle0() == &Rigid ? 1 : 0;
-			const FVec3 Impulse
-				= OtherIdx == 0
-				? Collision.AccumulatedImpulse
-				: -Collision.AccumulatedImpulse;
-
-			// Compute the angular impulse based on distance from the contact point to the CoM
-			const FVec3 Location = Collision.CalculateWorldContactLocation();
-			const Chaos::FVec3 AngularImpulse = Chaos::FVec3::CrossProduct(Location - Generic->PCom(), Impulse);
-
-			// Compute impulse velocities
-			const FVec3 ImpulseVelocity = Generic->InvM() * Impulse;
-
-			const FMatrix33 OtherInvI = Utilities::ComputeWorldSpaceInertia(Generic->QCom(), Generic->ConditionedInvI());
-			const FVec3 AngularImpulseVelocity = OtherInvI * AngularImpulse;
-
-			// Update linear and angular impulses for the body, to be integrated next solve
-			const float RestorationPercent = RestoreBreakingMomentumPercent;
-			Rigid.SetV(Rigid.GetV() + ImpulseVelocity * RestorationPercent);
-			Rigid.SetW(Rigid.GetW() + AngularImpulseVelocity * RestorationPercent);
+			if (FPBDRigidParticleHandle* Particle = Pair.Key)
+			{
+				const FMomentumRestoringData& Data = Pair.Value;
+				Particle->SetV(Particle->GetV() + Data.V);
+				Particle->SetW(Particle->GetW() + Data.W);
+			}
 		}
 	}
 
+	void FRigidClustering::RemoveFromMomentumRestoringStructures(const FPBDRigidParticleHandle* ParticleToRemove)
+	{
+		MomentumRestoringDataByParticle.Remove(ParticleToRemove);
+	}
+
 	void FRigidClustering::SendBreakingEvent(FPBDRigidClusteredParticleHandle* ClusteredParticle, bool bFromCrumble)
 	{
 		// only emit break event if the proxy needs it 
@@ -919,7 +925,7 @@ namespace Chaos
 				if (ClusteredParent->InternalCluster() && Children->IsEmpty() && ClusteredParent->PhysicsProxy() && ClusteredParent->PhysicsProxy()->GetType() == EPhysicsProxyType::GeometryCollectionType)
 				{
 					// It's safe to disable the particle until we get to the point where we want to destroy the particle.
-					MEvolution.DisableParticle(ClusteredParent);
+					DisableParticle(ClusteredParent);
 
 					// We shouldn't ever need to do an AddUnique here since when we remove a child from a parent, the parent should only ever turn empty once.
 					EmptyInternalClustersPerProxy.FindOrAdd(ClusteredParent->PhysicsProxy()).Add(ClusteredParent);
@@ -1476,12 +1482,7 @@ namespace Chaos
 							// just got re-enabled in RemoveParticlesFromCluster.
 							for (FPBDRigidParticleHandle* Particle : Island)
 							{
-								if (FPBDRigidClusteredParticleHandle* ClusterParticle = Particle->CastToClustered())
-								{
-									TopLevelClusterParentsStrained.Remove(Particle->CastToClustered());
-									TopLevelClusterParents.Remove(Particle->CastToClustered());
-								}
-								MEvolution.DisableParticle(Particle);
+								DisableParticle(Particle);
 							}
 						}
 					}
@@ -1879,9 +1880,6 @@ namespace Chaos
 	{
 		SCOPE_CYCLE_COUNTER(STAT_BreakingModel_AllParticles);
 		
-		// Clear the set tracking breaking collisions
-		BreakingCollisions.Empty();
-
 		//make copy because release cluster modifies active indices. We want to iterate over original active indices
 		TArray<FPBDRigidClusteredParticleHandle*> ClusteredParticlesToProcess;
 		for(FTransientPBDRigidParticleHandle& Particle : MEvolution.GetNonDisabledClusteredView())
@@ -1939,8 +1937,7 @@ namespace Chaos
 
 		FrameProcessedClusters += InParticles.Num();
 
-		// Clear the set tracking breaking collisions
-		BreakingCollisions.Empty();
+		MomentumRestoringDataByParticle.Empty();
 
 		bool bHasReleasedParticles = false;
 		for(FPBDRigidClusteredParticleHandle* ClusteredParticle : InParticles)
@@ -1963,6 +1960,7 @@ namespace Chaos
 				RestoreBreakingMomentum();
 			}
 		}
+		MomentumRestoringDataByParticle.Empty();
 		return bHasReleasedParticles;
 	}
 
@@ -2406,32 +2404,30 @@ namespace Chaos
 
 	void FRigidClustering::DisableCluster(FPBDRigidClusteredParticleHandle* ClusteredParticle)
 	{
-		// Before disabling the particle, we need to make sure to remove the recorded breaking collision from BreakingCollisions
-		// to avoid it crashing later when processing it 
-		ClusteredParticle->ParticleCollisions().VisitCollisions(
-			[this, ClusteredParticle](FPBDCollisionConstraint& Collision)
-			{
-				const uint8 OtherIdx = (Collision.GetParticle0() == ClusteredParticle) ? 1 : 0;
-				if (FGeometryParticleHandle* OtherGeometry = Collision.GetParticle(OtherIdx))
-				{
-					if (FPBDRigidParticleHandle* OtherRigid = OtherGeometry->CastToRigidParticle())
-					{
-						BreakingCollisions.Remove({ &Collision, ClusteredParticle });
-						BreakingCollisions.Remove({ &Collision, OtherRigid });
-					}
-				}
-				return ECollisionVisitorResult::Continue;
-			});
-
 		// #note: we don't recursively descend to the children
-		MEvolution.DisableParticle(ClusteredParticle);
-		TopLevelClusterParents.Remove(ClusteredParticle);
-		TopLevelClusterParentsStrained.Remove(ClusteredParticle);
-		GetChildrenMap().Remove(ClusteredParticle);
+		DisableParticle(ClusteredParticle);
 		ClusteredParticle->ClusterIds() = ClusterId();
 		ClusteredParticle->ClusterGroupIndex() = 0;
 	}
 
+	void FRigidClustering::DisableParticle(FPBDRigidParticleHandle* ParticleToDisable)
+	{
+		if (ParticleToDisable)
+		{
+			// Before disabling the particle, we need to make sure to remove the recorded breaking collision from BreakingCollisions
+			// to avoid it crashing later when processing it 
+			RemoveFromMomentumRestoringStructures(ParticleToDisable);
+
+			MEvolution.DisableParticle(ParticleToDisable);
+			if (FPBDRigidClusteredParticleHandle* ClusteredParticle = ParticleToDisable->CastToClustered())
+			{
+				TopLevelClusterParents.Remove(ClusteredParticle);
+				TopLevelClusterParentsStrained.Remove(ClusteredParticle);
+				MChildren.Remove(ClusteredParticle);
+			}
+		}
+	}
+
 	void FRigidClustering::ApplyStrainModifiers(const TArray<FPBDRigidClusteredParticleHandle*>& StrainedParticles)
 	{
 		if (StrainModifiers)
@@ -2451,6 +2447,10 @@ namespace Chaos
 	{
 		FClusterHandle ParentParticle = nullptr;
 
+		// Before derstroying the particle, we need to make sure to remove the recorded breaking collision from BreakingCollisions
+		// to avoid it crashing later when processing it 
+		RemoveFromMomentumRestoringStructures(ClusteredParticle);
+
 		// detach connections to thie parent from the children
 		if (MChildren.Contains(ClusteredParticle))
 		{
@@ -2469,7 +2469,7 @@ namespace Chaos
 		// disable within the solver
 		if (!ClusteredParticle->Disabled())
 		{
-			MEvolution.DisableParticle(ClusteredParticle);
+			DisableParticle(ClusteredParticle);
 			ensure(ClusteredParticle->ClusterIds().Id == nullptr);
 		}
 
@@ -3057,14 +3057,14 @@ namespace Chaos
 		return CVarShouldThrottleParticleRelease();
 	}
 
-	void FRigidClustering::ThrottleReleasedParticlesIfNecessary(TSet<FPBDRigidParticleHandle*>& Particles) const
+	void FRigidClustering::ThrottleReleasedParticlesIfNecessary(TSet<FPBDRigidParticleHandle*>& Particles)
 	{
-		GenericThrottleReleasedParticlesIfNecessary(Particles, MEvolution);
+		GenericThrottleReleasedParticlesIfNecessary(Particles, *this);
 	}
 
-	void FRigidClustering::ThrottleReleasedParticlesIfNecessary(TArray<FPBDRigidParticleHandle*>& Particles) const
+	void FRigidClustering::ThrottleReleasedParticlesIfNecessary(TArray<FPBDRigidParticleHandle*>& Particles)
 	{
-		GenericThrottleReleasedParticlesIfNecessary(Particles, MEvolution);
+		GenericThrottleReleasedParticlesIfNecessary(Particles, *this);
 	}
 
 } // namespace Chaos
diff --git a/Engine/Source/Runtime/Experimental/Chaos/Public/Chaos/PBDRigidClustering.h b/Engine/Source/Runtime/Experimental/Chaos/Public/Chaos/PBDRigidClustering.h
--- a/Engine/Source/Runtime/Experimental/Chaos/Public/Chaos/PBDRigidClustering.h
+++ b/Engine/Source/Runtime/Experimental/Chaos/Public/Chaos/PBDRigidClustering.h
@@ -493,8 +493,11 @@ public:
 	CHAOS_API void DisableCluster(FPBDRigidClusteredParticleHandle* ClusteredParticle);
 
 	bool ShouldThrottleParticleRelease() const;
-	void ThrottleReleasedParticlesIfNecessary(TSet<FPBDRigidParticleHandle*>& Particles) const;
-	void ThrottleReleasedParticlesIfNecessary(TArray<FPBDRigidParticleHandle*>& Particles) const;
+	void ThrottleReleasedParticlesIfNecessary(TSet<FPBDRigidParticleHandle*>& Particles);
+	void ThrottleReleasedParticlesIfNecessary(TArray<FPBDRigidParticleHandle*>& Particles);
+
+	/** disable a particle and make sure the related internal structure are up to date */
+	void DisableParticle(FPBDRigidParticleHandle* ParticleToDisable);
 
  protected:
 
@@ -551,6 +554,7 @@ public:
 	// Restore some percentage of momenta for objects which were involved in collisions
 	// with destroyed GCs
 	CHAOS_API void RestoreBreakingMomentum();
+	CHAOS_API void RemoveFromMomentumRestoringStructures(const FPBDRigidParticleHandle* ParticleToRemove);
 
 	CHAOS_API void SendBreakingEvent(FPBDRigidClusteredParticleHandle* ClusteredParticle, bool bFromCrumble);
 	CHAOS_API void SendCrumblingEvent(FPBDRigidClusteredParticleHandle* ClusteredParticle);
@@ -614,10 +618,14 @@ private:
 	TSet<FPBDRigidClusteredParticleHandle*> CrumbledSinceLastUpdate;
 	TMap<IPhysicsProxyBase*, TArray<FPBDRigidClusteredParticleHandle*>> EmptyInternalClustersPerProxy;
 
-	// Pairs of collision constraints and rigid particle handles of particles which collided with
-	// rigid clusters which broken. Some portion of the momentum change due to the constraint
-	// will be restored to each of the corresponding particles.
-	TSet<TPair<FPBDCollisionConstraint*, FPBDRigidParticleHandle*>> BreakingCollisions;
+	struct FMomentumRestoringData
+	{
+		FVector V = FVector::ZeroVector;
+		FVector W = FVector::ZeroVector;
+	};
+	// this map stores the momentum information a particle had when colliding with a breaking clustered particle
+	// this allow the colliding particle to keep going if the clustered particle break depsite the fact that the collision constraint resolution stopped it 
+	TMap<FPBDRigidParticleHandle*, FMomentumRestoringData> MomentumRestoringDataByParticle;
 
 	FReal MClusterConnectionFactor;
 	FClusterCreationParameters::EConnectionMethod MClusterUnionConnectionType;
