Add linear and angular joint violation callbacks. Use them in swivel plate to break the plate if the constraint is too violated for too long.

diff --git a/Engine/Source/Runtime/Engine/Classes/Engine/EngineTypes.h b/Engine/Source/Runtime/Engine/Classes/Engine/EngineTypes.h
--- a/Engine/Source/Runtime/Engine/Classes/Engine/EngineTypes.h
+++ b/Engine/Source/Runtime/Engine/Classes/Engine/EngineTypes.h
@@ -1034,6 +1034,7 @@ enum ECollisionChannel : int
 };
 
 DECLARE_DELEGATE_OneParam(FOnConstraintBroken, int32 /*ConstraintIndex*/);
+DECLARE_DELEGATE_ThreeParams(FOnConstraintViolated, int32 /*ConstraintIndex*/, float /*LinearViolation*/, float /*AngularViolation*/);
 DECLARE_DELEGATE_OneParam(FOnPlasticDeformation, int32 /*ConstraintIndex*/);
 
 #define COLLISION_GIZMO ECC_EngineTraceChannel1
diff --git a/Engine/Source/Runtime/Engine/Classes/PhysicsEngine/ConstraintInstance.h b/Engine/Source/Runtime/Engine/Classes/PhysicsEngine/ConstraintInstance.h
--- a/Engine/Source/Runtime/Engine/Classes/PhysicsEngine/ConstraintInstance.h
+++ b/Engine/Source/Runtime/Engine/Classes/PhysicsEngine/ConstraintInstance.h
@@ -221,16 +221,21 @@ struct FConstraintInstanceBase
 	/** Set the constraint broken delegate. */
 	ENGINE_API void SetConstraintBrokenDelegate(FOnConstraintBroken InConstraintBrokenDelegate);
 
+	/** Set the constraint violation delegate */
+	ENGINE_API void SetConstraintViolatedDelegate(FOnConstraintViolated InConstraintViolatedDelegate);
+
 	/** Set the plastic deformation delegate. */
 	ENGINE_API void SetPlasticDeformationDelegate(FOnPlasticDeformation InPlasticDeformationDelegate);
 
 	protected:
 
 		FOnConstraintBroken OnConstraintBrokenDelegate;
+		FOnConstraintViolated OnConstraintViolatedDelegate;
 		FOnPlasticDeformation OnPlasticDeformationDelegate;
 
 		friend struct FConstraintBrokenDelegateData;
 		friend struct FConstraintBrokenDelegateWrapper;
+		friend struct FConstraintViolatedDelegateWrapper;
 		friend struct FPlasticDeformationDelegateWrapper;
 
 };
diff --git a/Engine/Source/Runtime/Engine/Private/PhysicsEngine/ConstraintInstance.cpp b/Engine/Source/Runtime/Engine/Private/PhysicsEngine/ConstraintInstance.cpp
--- a/Engine/Source/Runtime/Engine/Private/PhysicsEngine/ConstraintInstance.cpp
+++ b/Engine/Source/Runtime/Engine/Private/PhysicsEngine/ConstraintInstance.cpp
@@ -333,6 +333,11 @@ void FConstraintInstanceBase::SetConstraintBrokenDelegate(FOnConstraintBroken In
 	OnConstraintBrokenDelegate = InConstraintBrokenDelegate;
 }
 
+void FConstraintInstanceBase::SetConstraintViolatedDelegate(FOnConstraintViolated InConstraintViolatedDelegate)
+{
+	OnConstraintViolatedDelegate = InConstraintViolatedDelegate;
+}
+
 
 /** Constructor **/
 FConstraintInstance::FConstraintInstance()
diff --git a/Engine/Source/Runtime/Engine/Private/PhysicsEngine/Experimental/PhysScene_Chaos.cpp b/Engine/Source/Runtime/Engine/Private/PhysicsEngine/Experimental/PhysScene_Chaos.cpp
--- a/Engine/Source/Runtime/Engine/Private/PhysicsEngine/Experimental/PhysScene_Chaos.cpp
+++ b/Engine/Source/Runtime/Engine/Private/PhysicsEngine/Experimental/PhysScene_Chaos.cpp
@@ -2254,6 +2254,17 @@ void FPhysScene_Chaos::OnSyncBodies(Chaos::FPhysicsSolverBase* Solver)
 				Constraint->GetOutputData().bIsBreaking = false;
 			}
 
+			if (Constraint->GetOutputData().bIsViolating)
+			{
+				if (FConstraintInstanceBase* ConstraintInstance = (Constraint) ? FPhysicsUserData_Chaos::Get<FConstraintInstanceBase>(Constraint->GetUserData()) : nullptr)
+				{
+					FConstraintViolatedDelegateWrapper CVD(ConstraintInstance);
+					CVD.DispatchOnViolated(Constraint->GetOutputData().LinearViolation, Constraint->GetOutputData().AngularViolation);
+				}
+
+				Constraint->GetOutputData().bIsViolating = false;
+			}
+
 			if (Constraint->GetOutputData().bDriveTargetChanged)
 			{
 				if (FConstraintInstanceBase* ConstraintInstance = (Constraint) ? FPhysicsUserData_Chaos::Get<FConstraintInstanceBase>(Constraint->GetUserData()) : nullptr)
@@ -2412,6 +2423,16 @@ void FConstraintBrokenDelegateWrapper::DispatchOnBroken()
 	OnConstraintBrokenDelegate.ExecuteIfBound(ConstraintIndex);
 }
 
+FConstraintViolatedDelegateWrapper::FConstraintViolatedDelegateWrapper(FConstraintInstanceBase* ConstraintInstance)
+	: OnConstraintViolatedDelegate(ConstraintInstance->OnConstraintViolatedDelegate)
+	, ConstraintIndex(ConstraintInstance->ConstraintIndex)
+{ }
+
+void FConstraintViolatedDelegateWrapper::DispatchOnViolated(const float LinearViolation, const float AngularViolation)
+{
+	OnConstraintViolatedDelegate.ExecuteIfBound(ConstraintIndex, LinearViolation, AngularViolation);
+}
+
 FPlasticDeformationDelegateWrapper::FPlasticDeformationDelegateWrapper(FConstraintInstanceBase* ConstraintInstance)
 	: OnPlasticDeformationDelegate(ConstraintInstance->OnPlasticDeformationDelegate)
 	, ConstraintIndex(ConstraintInstance->ConstraintIndex)
diff --git a/Engine/Source/Runtime/Engine/Public/Physics/Experimental/PhysScene_Chaos.h b/Engine/Source/Runtime/Engine/Public/Physics/Experimental/PhysScene_Chaos.h
--- a/Engine/Source/Runtime/Engine/Public/Physics/Experimental/PhysScene_Chaos.h
+++ b/Engine/Source/Runtime/Engine/Public/Physics/Experimental/PhysScene_Chaos.h
@@ -83,6 +83,16 @@ struct FConstraintBrokenDelegateWrapper
 	int32 ConstraintIndex;
 };
 
+struct FConstraintViolatedDelegateWrapper
+{
+	FConstraintViolatedDelegateWrapper(FConstraintInstanceBase* ConstraintInstance);
+
+	void DispatchOnViolated(float LinearViolation, float AngularViolation);
+
+	FOnConstraintViolated OnConstraintViolatedDelegate;
+	int32 ConstraintIndex;
+};
+
 struct FPlasticDeformationDelegateWrapper
 {
 	FPlasticDeformationDelegateWrapper(FConstraintInstanceBase* ConstraintInstance);
diff --git a/Engine/Source/Runtime/Experimental/Chaos/Private/Chaos/Joint/PBDJointCachedSolverGaussSeidel.cpp b/Engine/Source/Runtime/Experimental/Chaos/Private/Chaos/Joint/PBDJointCachedSolverGaussSeidel.cpp
--- a/Engine/Source/Runtime/Experimental/Chaos/Private/Chaos/Joint/PBDJointCachedSolverGaussSeidel.cpp
+++ b/Engine/Source/Runtime/Experimental/Chaos/Private/Chaos/Joint/PBDJointCachedSolverGaussSeidel.cpp
@@ -197,6 +197,7 @@ void FPBDJointCachedSolver::Init(
 	NumActiveConstraints = -1;
 	bIsActive = true;
 	bIsBroken = false;
+	bIsViolating = false;
 	bUseSimd = SolverSettings.bUseSimd;
 	bUsePositionBasedDrives = SolverSettings.bUsePositionBasedDrives;
 
diff --git a/Engine/Source/Runtime/Experimental/Chaos/Private/Chaos/Joint/PBDJointContainerSolver.cpp b/Engine/Source/Runtime/Experimental/Chaos/Private/Chaos/Joint/PBDJointContainerSolver.cpp
--- a/Engine/Source/Runtime/Experimental/Chaos/Private/Chaos/Joint/PBDJointContainerSolver.cpp
+++ b/Engine/Source/Runtime/Experimental/Chaos/Private/Chaos/Joint/PBDJointContainerSolver.cpp
@@ -2,6 +2,7 @@
 #include "Chaos/Joint/PBDJointContainerSolver.h"
 #include "Chaos/Joint/ChaosJointLog.h"
 #include "Chaos/Island/IslandManager.h"
+#include "Chaos/DebugDrawQueue.h"
 
 namespace Chaos
 {
@@ -82,6 +83,28 @@ namespace Chaos
 			return bBreak;
 		}
 
+		bool GetJointIsViolating(const FPBDJointSettings& JointSettings, const FReal LinearViolationSq, const FReal AngularViolation, const int32 It, const int32 NumIts)
+		{
+			bool bViolating = false;
+			if (It == NumIts - 1)
+			{
+				const FReal LinearViolationCallbackThresholdSq
+					= JointSettings.LinearViolationCallbackThreshold
+					* JointSettings.LinearViolationCallbackThreshold;
+				if (!bViolating && JointSettings.LinearViolationCallbackThreshold != FLT_MAX)
+				{
+					bViolating = LinearViolationSq > LinearViolationCallbackThresholdSq;
+				}
+
+				if (!bViolating && JointSettings.AngularViolationCallbackThreshold != FLT_MAX)
+				{
+					bViolating = FMath::RadiansToDegrees(AngularViolation) > JointSettings.AngularViolationCallbackThreshold;
+				}
+			}
+
+			return bViolating;
+		}
+
 		///////////////////////////////////////////////////////////////////////////////////////////
 		///////////////////////////////////////////////////////////////////////////////////////////
 		///////////////////////////////////////////////////////////////////////////////////////////
@@ -241,14 +264,17 @@ namespace Chaos
 					const FSolverBody* SolverBody0 = &Solver.Body0().SolverBody();
 					const FSolverBody* SolverBody1 = &Solver.Body1().SolverBody();
 					const bool bIsBroken = Solver.IsBroken();
+					const bool bIsViolating = Solver.IsViolating();
+					const float LinearViolation = bIsViolating ? FMath::Sqrt((float)Solver.GetLinearViolationSq()) : 0.f;
+					const float AngularViolation = bIsViolating ? (float)Solver.GetAngularViolation() : 0.f;
 
-					Container.GetContainer().SetSolverResults(ContainerConstraintIndex, LinearImpulse, AngularImpulse, bIsBroken, SolverBody0, SolverBody1);
+					Container.GetContainer().SetSolverResults(ContainerConstraintIndex, LinearImpulse, AngularImpulse, LinearViolation, AngularViolation, bIsBroken, bIsViolating, SolverBody0, SolverBody1);
 
 					Solver.Deinit();
 				}
 				else
 				{
-					Container.GetContainer().SetSolverResults(ContainerConstraintIndex, FVec3(0), FVec3(0), false, nullptr, nullptr);
+					Container.GetContainer().SetSolverResults(ContainerConstraintIndex, FVec3(0), FVec3(0), 0.f, 0.f, false, false, nullptr, nullptr);
 				}
 			}
 		}
@@ -306,11 +332,12 @@ namespace Chaos
 				Solver.ApplyConstraints(Dt, IterationStiffness, Settings, JointSettings);
 
 				// @todo(ccaulfield): We should be clamping the impulse at this point. Maybe move breaking to the solver
-				if ((JointSettings.LinearBreakForce != FLT_MAX || JointSettings.AngularBreakTorque != FLT_MAX) &&
-					GetJointShouldBreak(JointSettings, Dt, Solver.GetNetLinearImpulse(), Solver.GetNetAngularImpulse()))
+				if (GetJointShouldBreak(JointSettings, Dt, Solver.GetNetLinearImpulse(), Solver.GetNetAngularImpulse()))
 				{
 					Solver.SetIsBroken(true);
 				}
+
+				Solver.SetIsViolating(GetJointIsViolating(JointSettings, Solver.GetLinearViolationSq(), Solver.GetAngularViolation(), It, NumIts));
 			}
 		}
 
diff --git a/Engine/Source/Runtime/Experimental/Chaos/Private/Chaos/Joint/PBDJointSolverGaussSeidel.cpp b/Engine/Source/Runtime/Experimental/Chaos/Private/Chaos/Joint/PBDJointSolverGaussSeidel.cpp
--- a/Engine/Source/Runtime/Experimental/Chaos/Private/Chaos/Joint/PBDJointSolverGaussSeidel.cpp
+++ b/Engine/Source/Runtime/Experimental/Chaos/Private/Chaos/Joint/PBDJointSolverGaussSeidel.cpp
@@ -113,6 +113,9 @@ namespace Chaos
 		NetLinearImpulse = FVec3(0);
 		NetAngularImpulse = FVec3(0);
 
+		LinearViolationSq = 0;
+		AngularViolation = 0;
+
 		LinearSoftLambda = 0;
 		TwistSoftLambda = 0;
 		SwingSoftLambda = 0;
@@ -223,6 +226,9 @@ namespace Chaos
 	{
 		SolverStiffness = InSolverStiffness;
 
+		LinearViolationSq = 0.f;
+		AngularViolation = 0.f;
+
 		if (SolverSettings.bSolvePositionLast)
 		{
 			ApplyRotationConstraints(Dt, SolverSettings, JointSettings);
@@ -1560,6 +1566,11 @@ namespace Chaos
 		{
 			FReal TwistStiffness = FPBDJointUtilities::GetTwistStiffness(SolverSettings, JointSettings);
 			ApplyRotationConstraint(TwistStiffness, Axis0, R01.X, (int32)EJointAngularConstraintIndex::Twist);
+
+			// Update the angular violation
+			//
+			// TODO: Use post-correction error instead of pre-correction error?
+			AngularViolation = FMath::Max(AngularViolation, FMath::Abs(R01.X));
 		}
 
 		if (bApplySwing)
@@ -1567,6 +1578,12 @@ namespace Chaos
 			FReal SwingStiffness = FPBDJointUtilities::GetSwingStiffness(SolverSettings, JointSettings);
 			ApplyRotationConstraint(SwingStiffness, Axis1, R01.Y, (int32)EJointAngularConstraintIndex::Swing2);
 			ApplyRotationConstraint(SwingStiffness, Axis2, R01.Z, (int32)EJointAngularConstraintIndex::Swing1);
+
+			// Update the angular violation
+			//
+			// TODO: Use post-correction error instead of pre-correction error?
+			AngularViolation = FMath::Max(AngularViolation, FMath::Abs(R01.Y));
+			AngularViolation = FMath::Max(AngularViolation, FMath::Abs(R01.Z));
 		}
 	}
 
@@ -1611,6 +1628,11 @@ namespace Chaos
 				FReal TwistStiffness = FPBDJointUtilities::GetTwistStiffness(SolverSettings, JointSettings);
 				ApplyRotationConstraint(TwistStiffness, TwistAxis, DTwistAngle, (int32)EJointAngularConstraintIndex::Twist);
 			}
+
+			// Update the angular violation
+			//
+			// TODO: Use post-correction error instead of pre-correction error?
+			AngularViolation = FMath::Max(AngularViolation, DTwistAngle);
 		}
 	}
 
@@ -1647,10 +1669,14 @@ namespace Chaos
 				// For cone constraints, the lambda are all accumulated in Swing2
 				ApplyRotationConstraint(SwingStiffness, SwingAxis, DSwingAngle, (int32)EJointAngularConstraintIndex::Swing2);
 			}
+
+			// Update the angular violation
+			//
+			// TODO: Use post-correction error instead of pre-correction error?
+			AngularViolation = FMath::Max(AngularViolation, DSwingAngle);
 		}
 	}
 
-
 	void FPBDJointSolver::ApplySingleLockedSwingConstraint(
 		const FReal Dt,
 		const FPBDJointSolverSettings& SolverSettings,
@@ -1682,6 +1708,11 @@ namespace Chaos
 				const FReal SwingStiffness = FPBDJointUtilities::GetSwingStiffness(SolverSettings, JointSettings);
 				ApplyRotationConstraint(SwingStiffness, SwingAxis, DSwingAngle, (int32)SwingConstraintIndex);
 			}
+
+			// Update the angular violation
+			//
+			// TODO: Use post-correction error instead of pre-correction error?
+			AngularViolation = FMath::Max(AngularViolation, FMath::Abs(DSwingAngle));
 		}
 	}
 
@@ -1714,7 +1745,7 @@ namespace Chaos
 		UE_LOG(LogChaosJoint, VeryVerbose, TEXT("    DualConeSwing%d Angle %f [Limit %f]"), (SwingConstraintIndex == EJointAngularConstraintIndex::Swing1) ? 1 : 2, FMath::RadiansToDegrees(SwingAngle), FMath::RadiansToDegrees(SwingAngleMax));
 
 		// Apply swing correction
-		if (DSwingAngle > SolverSettings.AngleTolerance)
+		if (DSwingAngle > AngleTolerance)
 		{
 			if (bUseSoftLimit)
 			{
@@ -1728,6 +1759,11 @@ namespace Chaos
 				const FReal SwingStiffness = FPBDJointUtilities::GetSwingStiffness(SolverSettings, JointSettings);
 				ApplyRotationConstraint(SwingStiffness, SwingAxis, DSwingAngle, (int32)SwingConstraintIndex);
 			}
+
+			// Update the angular violation
+			//
+			// TODO: Use post-correction error instead of pre-correction error?
+			AngularViolation = FMath::Max(AngularViolation, DSwingAngle);
 		}
 	}
 
@@ -1774,6 +1810,11 @@ namespace Chaos
 				const FReal SwingStiffness = FPBDJointUtilities::GetSwingStiffness(SolverSettings, JointSettings);
 				ApplyRotationConstraint(SwingStiffness, SwingAxis, DSwingAngle, (int32)SwingConstraintIndex);
 			}
+
+			// Update the angular violation
+			//
+			// TODO: Use post-correction error instead of pre-correction error?
+			AngularViolation = FMath::Max(AngularViolation, DSwingAngle);
 		}
 	}
 
@@ -1889,7 +1930,9 @@ namespace Chaos
 
 		UE_LOG(LogChaosJoint, VeryVerbose, TEXT("    PointKD Delta %f [Limit %f]"), CX.Size(), PositionTolerance);
 
-		if (CX.SizeSquared() > PositionTolerance * PositionTolerance)
+		const FReal CXSquared = CX.SizeSquared();
+		const FReal PositionToleranceSquared = PositionTolerance * PositionTolerance;
+		if (CXSquared > PositionToleranceSquared)
 		{
 			if (bRealTypeCompatibleWithISPC && bChaos_Joint_ISPC_Enabled)
 			{
@@ -1912,6 +1955,11 @@ namespace Chaos
 
 				NetLinearImpulse += (KIndex == 0) ? DX : -DX;
 			}
+
+			// Update the linear violation
+			//
+			// TODO: Use post-correction violation
+			LinearViolationSq = FMath::Max(LinearViolationSq, CXSquared);
 		}
 	}
 
@@ -1930,7 +1978,9 @@ namespace Chaos
 
 		UE_LOG(LogChaosJoint, VeryVerbose, TEXT("    PointDD Delta %f [Limit %f]"), CX.Size(), PositionTolerance);
 
-		if (CX.SizeSquared() > PositionTolerance * PositionTolerance)
+		const FReal CXSquared = CX.SizeSquared();
+		const FReal PositionToleranceSquared = PositionTolerance * PositionTolerance;
+		if (CXSquared > PositionToleranceSquared)
 		{
 			if (bRealTypeCompatibleWithISPC && bChaos_Joint_ISPC_Enabled)
 			{
@@ -1957,6 +2007,11 @@ namespace Chaos
 
 				NetLinearImpulse += DX;
 			}
+
+			// Update the linear violation
+			//
+			// TODO: Use post-correction violation
+			LinearViolationSq = FMath::Max(LinearViolationSq, CXSquared);
 		}
 	}
 
@@ -1988,6 +2043,11 @@ namespace Chaos
 				const bool bAccelerationMode = FPBDJointUtilities::GetLinearSoftAccelerationMode(SolverSettings, JointSettings);
 				ApplyPositionConstraintSoft(Dt, JointStiffness, JointDamping, bAccelerationMode, Axis, Error, 0.0f, LinearSoftLambda);
 			}
+
+			// Update the linear violation
+			//
+			// TODO: Use post-correction violation
+			LinearViolationSq = FMath::Max(LinearViolationSq, Error * Error);
 		}
 	}
 
@@ -2033,6 +2093,11 @@ namespace Chaos
 				const FVec3 Connector0Correction = ConnectorXs[1] - ConnectorXs[0] - AxialError * Axis;
 				ApplyPositionConstraint(JointStiffness, Axis, AxialError, Connector0Correction, (int32)EJointCylindricalPositionConstraintType::Axial);
 			}
+
+			// Update the linear violation
+			//
+			// TODO: Use post-correction violation
+			LinearViolationSq = FMath::Max(LinearViolationSq, AxialError * AxialError);
 		}
 
 		const FReal RadialLimit = (RadialMotion == EJointMotionType::Locked) ? 0.0f : JointSettings.LinearLimit;
@@ -2055,6 +2120,11 @@ namespace Chaos
 				const FVec3 Connector0Correction = ConnectorXs[1] - ConnectorXs[0] - RadialError * RadialAxis;
 				ApplyPositionConstraint(JointStiffness, RadialAxis, RadialError, Connector0Correction, (int32)EJointCylindricalPositionConstraintType::Radial);
 			}
+
+			// Update the linear violation
+			//
+			// TODO: Use post-correction violation
+			LinearViolationSq = FMath::Max(LinearViolationSq, RadialError * RadialError);
 		}
 	}
 
@@ -2093,6 +2163,9 @@ namespace Chaos
 				const FVec3 Connector0Correction = ConnectorXs[1] - ConnectorXs[0] - Axis * Error;
 				ApplyPositionConstraint(JointStiffness, Axis, Error, Connector0Correction);
 			}
+
+			// Update the linear violation
+			LinearViolationSq = FMath::Max(LinearViolationSq, Error * Error);
 		}
 	}
 
diff --git a/Engine/Source/Runtime/Experimental/Chaos/Private/Chaos/PBDJointConstraints.cpp b/Engine/Source/Runtime/Experimental/Chaos/Private/Chaos/PBDJointConstraints.cpp
--- a/Engine/Source/Runtime/Experimental/Chaos/Private/Chaos/PBDJointConstraints.cpp
+++ b/Engine/Source/Runtime/Experimental/Chaos/Private/Chaos/PBDJointConstraints.cpp
@@ -87,6 +87,16 @@ namespace Chaos
 		return ConcreteContainer()->ClearConstraintBreaking(ConstraintIndex);
 	}
 
+	bool FPBDJointConstraintHandle::IsConstraintViolating() const
+	{
+		return ConcreteContainer()->IsConstraintViolating(ConstraintIndex);
+	}
+
+	void FPBDJointConstraintHandle::ClearConstraintViolating()
+	{
+		return ConcreteContainer()->ClearConstraintViolating(ConstraintIndex);
+	}
+
 	bool FPBDJointConstraintHandle::IsDriveTargetChanged() const
 	{
 		return ConcreteContainer()->IsDriveTargetChanged(ConstraintIndex);
@@ -112,6 +122,16 @@ namespace Chaos
 		return ConcreteContainer()->GetConstraintAngularImpulse(ConstraintIndex);
 	}
 
+	float FPBDJointConstraintHandle::GetLinearViolation() const
+	{
+		return ConcreteContainer()->GetConstraintLinearViolation(ConstraintIndex);
+	}
+
+	float FPBDJointConstraintHandle::GetAngularViolation() const
+	{
+		return ConcreteContainer()->GetConstraintAngularViolation(ConstraintIndex);
+	}
+
 	ESyncState FPBDJointConstraintHandle::SyncState() const
 	{
 		return ConcreteContainer()->GetConstraintSyncState(ConstraintIndex);
@@ -363,10 +383,12 @@ namespace Chaos
 		, AngularDriveDamping(FVec3(0))
 		, AngularDriveMaxTorque(FVec3(UE_MAX_FLT))
 		, LinearBreakForce(UE_MAX_FLT)
+		, LinearViolationCallbackThreshold(UE_MAX_FLT)
 		, LinearPlasticityLimit(UE_MAX_FLT)
 		, LinearPlasticityType(EPlasticityType::Free)
 		, LinearPlasticityInitialDistanceSquared(UE_MAX_FLT)
 		, AngularBreakTorque(UE_MAX_FLT)
+		, AngularViolationCallbackThreshold(UE_MAX_FLT)
 		, AngularPlasticityLimit(UE_MAX_FLT)
 		, ContactTransferScale(0.f)
 		, UserData(nullptr)
@@ -749,6 +771,16 @@ namespace Chaos
 		ConstraintStates[ConstraintIndex].bBreaking = false;
 	}
 
+	bool FPBDJointConstraints::IsConstraintViolating(int32 ConstraintIndex) const
+	{
+		return ConstraintStates[ConstraintIndex].bViolating;
+	}
+
+	void FPBDJointConstraints::ClearConstraintViolating(int32 ConstraintIndex)
+	{
+		ConstraintStates[ConstraintIndex].bViolating = false;
+	}
+
 	bool FPBDJointConstraints::IsDriveTargetChanged(int32 ConstraintIndex) const
 	{
 		return ConstraintStates[ConstraintIndex].bDriveTargetChanged;
@@ -897,6 +929,16 @@ namespace Chaos
 		return ConstraintStates[ConstraintIndex].AngularImpulse;
 	}
 
+	float FPBDJointConstraints::GetConstraintLinearViolation(int32 ConstraintIndex) const
+	{
+		return ConstraintStates[ConstraintIndex].LinearViolation;
+	}
+
+	float FPBDJointConstraints::GetConstraintAngularViolation(int32 ConstraintIndex) const
+	{
+		return ConstraintStates[ConstraintIndex].AngularViolation;
+	}
+
 	ESyncState FPBDJointConstraints::GetConstraintSyncState(int32 ConstraintIndex) const
 	{
 		return ConstraintStates[ConstraintIndex].SyncState;
@@ -983,7 +1025,7 @@ namespace Chaos
 		}
 	}
 
-	void FPBDJointConstraints::SetSolverResults(const int32 ConstraintIndex, const FVec3& LinearImpulse, const FVec3& AngularImpulse, bool bIsBroken, const FSolverBody* SolverBody0, const FSolverBody* SolverBody1)
+	void FPBDJointConstraints::SetSolverResults(const int32 ConstraintIndex, const FVec3& LinearImpulse, const FVec3& AngularImpulse, const float LinearViolation, const float AngularViolation, const bool bIsBroken, const bool bIsViolating, const FSolverBody* SolverBody0, const FSolverBody* SolverBody1)
 	{
 		ConstraintStates[ConstraintIndex].LinearImpulse = LinearImpulse;
 		ConstraintStates[ConstraintIndex].AngularImpulse = AngularImpulse;
@@ -993,6 +1035,10 @@ namespace Chaos
 			BreakConstraint(ConstraintIndex);
 		}
 
+		ConstraintStates[ConstraintIndex].bViolating = bIsViolating;
+		ConstraintStates[ConstraintIndex].LinearViolation = LinearViolation;
+		ConstraintStates[ConstraintIndex].AngularViolation = AngularViolation;
+
 		if ((SolverBody0 != nullptr) && (SolverBody1 != nullptr))
 		{
 			ApplyPlasticityLimits(ConstraintIndex, *SolverBody0, *SolverBody1);
diff --git a/Engine/Source/Runtime/Experimental/Chaos/Private/PhysicsProxy/JointConstraintProxy.cpp b/Engine/Source/Runtime/Experimental/Chaos/Private/PhysicsProxy/JointConstraintProxy.cpp
--- a/Engine/Source/Runtime/Experimental/Chaos/Private/PhysicsProxy/JointConstraintProxy.cpp
+++ b/Engine/Source/Runtime/Experimental/Chaos/Private/PhysicsProxy/JointConstraintProxy.cpp
@@ -50,11 +50,15 @@ void FJointConstraintPhysicsProxy::BufferPhysicsResults(FDirtyJointConstraintDat
 	{
 		Buffer.OutputData.bIsBreaking = Constraint_PT->IsConstraintBreaking();
 		Buffer.OutputData.bIsBroken = !Constraint_PT->IsConstraintEnabled();
+		Buffer.OutputData.bIsViolating = Constraint_PT->IsConstraintViolating();
 		Buffer.OutputData.bDriveTargetChanged = Constraint_PT->IsDriveTargetChanged();
 		Buffer.OutputData.Force = Constraint_PT->GetLinearImpulse();
 		Buffer.OutputData.Torque = Constraint_PT->GetAngularImpulse();
+		Buffer.OutputData.LinearViolation = Constraint_PT->GetLinearViolation();
+		Buffer.OutputData.AngularViolation = Constraint_PT->GetAngularViolation();
 
 		Constraint_PT->ClearConstraintBreaking(); // it's a single frame event, so reset
+		Constraint_PT->ClearConstraintViolating();
 		Constraint_PT->ClearDriveTargetChanged(); // it's a single frame event, so reset
 	}
 }
@@ -72,6 +76,10 @@ bool FJointConstraintPhysicsProxy::PullFromPhysicsState(const FDirtyJointConstra
 		}
 		Constraint_GT->GetOutputData().Force = Buffer.OutputData.Force;
 		Constraint_GT->GetOutputData().Torque = Buffer.OutputData.Torque;
+
+		Constraint_GT->GetOutputData().bIsViolating = Buffer.OutputData.bIsViolating;
+		Constraint_GT->GetOutputData().LinearViolation = Buffer.OutputData.LinearViolation;
+		Constraint_GT->GetOutputData().AngularViolation = Buffer.OutputData.AngularViolation;
 	}
 
 	return true;
diff --git a/Engine/Source/Runtime/Experimental/Chaos/Public/Chaos/Joint/PBDJointCachedSolverGaussSeidel.h b/Engine/Source/Runtime/Experimental/Chaos/Public/Chaos/Joint/PBDJointCachedSolverGaussSeidel.h
--- a/Engine/Source/Runtime/Experimental/Chaos/Public/Chaos/Joint/PBDJointCachedSolverGaussSeidel.h
+++ b/Engine/Source/Runtime/Experimental/Chaos/Public/Chaos/Joint/PBDJointCachedSolverGaussSeidel.h
@@ -417,6 +417,16 @@ struct FAxisConstraintDatas
 			return Impulse;
 		}
 
+		inline FReal GetLinearViolationSq() const
+		{
+			return -1.f;
+		}
+
+		inline FReal GetAngularViolation() const
+		{
+			return -1.f;
+		}
+
 		inline int32 GetNumActiveConstraints() const
 		{
 			// We use -1 as unitialized, but that should not be exposed outside the solver
@@ -544,6 +554,16 @@ struct FAxisConstraintDatas
 			return bIsBroken;
 		}
 
+		void SetIsViolating(const bool bInIsViolating)
+		{
+			bIsViolating = bInIsViolating;
+		}
+
+		bool IsViolating() const
+		{
+			return bIsViolating;
+		}
+
 		bool RequiresSolve() const
 		{
 			return !IsBroken() && (IsDynamic(0) || IsDynamic(1));
@@ -909,6 +929,7 @@ struct FAxisConstraintDatas
 		FAxisConstraintDatas RotationDrives;
 
 		bool bIsBroken;
+		bool bIsViolating;
 		bool bUseSimd;
 
 		// dummy indices
diff --git a/Engine/Source/Runtime/Experimental/Chaos/Public/Chaos/Joint/PBDJointSolverGaussSeidel.h b/Engine/Source/Runtime/Experimental/Chaos/Public/Chaos/Joint/PBDJointSolverGaussSeidel.h
--- a/Engine/Source/Runtime/Experimental/Chaos/Public/Chaos/Joint/PBDJointSolverGaussSeidel.h
+++ b/Engine/Source/Runtime/Experimental/Chaos/Public/Chaos/Joint/PBDJointSolverGaussSeidel.h
@@ -127,6 +127,16 @@ namespace Chaos
 			return NetAngularImpulse;
 		}
 
+		inline FReal GetLinearViolationSq() const
+		{
+			return LinearViolationSq;
+		}
+
+		inline FReal GetAngularViolation() const
+		{
+			return AngularViolation;
+		}
+
 		FPBDJointSolver()
 		{
 		}
@@ -194,6 +204,16 @@ namespace Chaos
 			return bIsBroken;
 		}
 
+		void SetIsViolating(const bool bInIsViolating)
+		{
+			bIsViolating = bInIsViolating;
+		}
+
+		bool IsViolating() const
+		{
+			return bIsViolating;
+		}
+
 		bool RequiresSolve() const
 		{
 			return !IsBroken() && (IsDynamic(0) || IsDynamic(1));
@@ -720,6 +740,9 @@ namespace Chaos
 		FVec3 NetLinearImpulse;
 		FVec3 NetAngularImpulse;
 
+		FReal LinearViolationSq;
+		FReal AngularViolation;
+
 		// Lagrange multipliers of the position constraints.
 		// Currently these are only used in ApplyCylindricalVelocityConstraints
 		FVec3 LinearHardLambda;
@@ -749,6 +772,7 @@ namespace Chaos
 		FVec3 InitConstraintAxisAngularVelocities; // Angular velocities along the constraint axes at the begining of the frame, used by restitution
 
 		bool bIsBroken;
+		bool bIsViolating;
 	};
 
 }
\ No newline at end of file
diff --git a/Engine/Source/Runtime/Experimental/Chaos/Public/Chaos/JointProperties.inl b/Engine/Source/Runtime/Experimental/Chaos/Public/Chaos/JointProperties.inl
--- a/Engine/Source/Runtime/Experimental/Chaos/Public/Chaos/JointProperties.inl
+++ b/Engine/Source/Runtime/Experimental/Chaos/Public/Chaos/JointProperties.inl
@@ -13,9 +13,11 @@ CHAOS_INNER_JOINT_PROPERTY(JointSettings, ProjectionAngularTolerance, TeleportAn
 CHAOS_INNER_JOINT_PROPERTY(JointSettings, ShockPropagationAlpha, ShockPropagation, FReal)
 CHAOS_INNER_JOINT_PROPERTY(JointSettings, ParentInvMassScale, ParentInvMassScale, FReal)
 CHAOS_INNER_JOINT_PROPERTY(JointSettings, LinearBreakForce, LinearBreakForce, FReal)
+CHAOS_INNER_JOINT_PROPERTY(JointSettings, LinearViolationCallbackThreshold, LinearViolationCallbackThreshold, FReal)
 CHAOS_INNER_JOINT_PROPERTY(JointSettings, LinearPlasticityLimit, LinearPlasticityLimit, FReal)
 CHAOS_INNER_JOINT_PROPERTY(JointSettings, LinearPlasticityType, LinearPlasticityType, EPlasticityType)
 CHAOS_INNER_JOINT_PROPERTY(JointSettings, AngularBreakTorque, AngularBreakTorque, FReal)
+CHAOS_INNER_JOINT_PROPERTY(JointSettings, AngularViolationCallbackThreshold, AngularViolationCallbackThreshold, FReal)
 CHAOS_INNER_JOINT_PROPERTY(JointSettings, AngularPlasticityLimit, AngularPlasticityLimit, FReal)
 CHAOS_INNER_JOINT_PROPERTY(JointSettings, UserData, UserData, void*)
 CHAOS_INNER_JOINT_PROPERTY(JointSettings, LinearPositionDriveXEnabled, bLinearPositionDriveEnabled[0], bool)
diff --git a/Engine/Source/Runtime/Experimental/Chaos/Public/Chaos/PBDJointConstraintData.h b/Engine/Source/Runtime/Experimental/Chaos/Public/Chaos/PBDJointConstraintData.h
--- a/Engine/Source/Runtime/Experimental/Chaos/Public/Chaos/PBDJointConstraintData.h
+++ b/Engine/Source/Runtime/Experimental/Chaos/Public/Chaos/PBDJointConstraintData.h
@@ -14,17 +14,19 @@ namespace Chaos
 
 	enum class EJointConstraintFlags : uint64_t
 	{
-		JointTransforms             = static_cast<uint64_t>(1) << 0,
-		CollisionEnabled            = static_cast<uint64_t>(1) << 1,
-		Projection                  = static_cast<uint64_t>(1) << 2,
-		ParentInvMassScale          = static_cast<uint64_t>(1) << 3,
-		LinearBreakForce            = static_cast<uint64_t>(1) << 4,
-		AngularBreakTorque          = static_cast<uint64_t>(1) << 5,
-		UserData                    = static_cast<uint64_t>(1) << 6,
-		LinearDrive                 = static_cast<uint64_t>(1) << 7,
-		AngularDrive                = static_cast<uint64_t>(1) << 8,
-		Stiffness                   = static_cast<uint64_t>(1) << 9,
-		Limits                      = static_cast<uint64_t>(1) << 10,
+		JointTransforms                     = static_cast<uint64_t>(1) << 0,
+		CollisionEnabled                    = static_cast<uint64_t>(1) << 1,
+		Projection                          = static_cast<uint64_t>(1) << 2,
+		ParentInvMassScale                  = static_cast<uint64_t>(1) << 3,
+		LinearBreakForce                    = static_cast<uint64_t>(1) << 4,
+		AngularBreakTorque                  = static_cast<uint64_t>(1) << 5,
+		UserData                            = static_cast<uint64_t>(1) << 6,
+		LinearDrive                         = static_cast<uint64_t>(1) << 7,
+		AngularDrive                        = static_cast<uint64_t>(1) << 8,
+		Stiffness                           = static_cast<uint64_t>(1) << 9,
+		Limits                              = static_cast<uint64_t>(1) << 10,
+		LinearViolationCallbackThreshold    = static_cast<uint64_t>(1) << 11,
+		AngularViolationCallbackThreshold   = static_cast<uint64_t>(1) << 12,
 
 		DummyFlag
 	};
@@ -66,6 +68,9 @@ namespace Chaos
 			bool bDriveTargetChanged = false;
 			FVector Force = FVector(0);
 			FVector Torque = FVector(0);
+			bool bIsViolating = false;
+			float LinearViolation = 0.f;
+			float AngularViolation = 0.f;
 		};
 		FOutputData& GetOutputData() { return Output; }
 
diff --git a/Engine/Source/Runtime/Experimental/Chaos/Public/Chaos/PBDJointConstraintTypes.h b/Engine/Source/Runtime/Experimental/Chaos/Public/Chaos/PBDJointConstraintTypes.h
--- a/Engine/Source/Runtime/Experimental/Chaos/Public/Chaos/PBDJointConstraintTypes.h
+++ b/Engine/Source/Runtime/Experimental/Chaos/Public/Chaos/PBDJointConstraintTypes.h
@@ -189,10 +189,12 @@ namespace Chaos
 		FVec3 AngularDriveMaxTorque;
 
 		FReal LinearBreakForce;
+		FReal LinearViolationCallbackThreshold;
 		FReal LinearPlasticityLimit;
 		EPlasticityType LinearPlasticityType;
 		FReal LinearPlasticityInitialDistanceSquared;
 		FReal AngularBreakTorque;
+		FReal AngularViolationCallbackThreshold;
 		FReal AngularPlasticityLimit;
 
 		FReal ContactTransferScale;
diff --git a/Engine/Source/Runtime/Experimental/Chaos/Public/Chaos/PBDJointConstraints.h b/Engine/Source/Runtime/Experimental/Chaos/Public/Chaos/PBDJointConstraints.h
--- a/Engine/Source/Runtime/Experimental/Chaos/Public/Chaos/PBDJointConstraints.h
+++ b/Engine/Source/Runtime/Experimental/Chaos/Public/Chaos/PBDJointConstraints.h
@@ -39,10 +39,14 @@ namespace Chaos
 		CHAOS_API bool IsConstraintBroken() const;
 		CHAOS_API bool IsConstraintBreaking() const;
 		CHAOS_API void ClearConstraintBreaking();
+		CHAOS_API bool IsConstraintViolating() const;
+		CHAOS_API void ClearConstraintViolating();
 		CHAOS_API bool IsDriveTargetChanged() const;
 		CHAOS_API void ClearDriveTargetChanged();
 		CHAOS_API FVec3 GetLinearImpulse() const;
 		CHAOS_API FVec3 GetAngularImpulse() const;
+		CHAOS_API float GetLinearViolation() const;
+		CHAOS_API float GetAngularViolation() const;
 
 		CHAOS_API const FPBDJointSettings& GetSettings() const;
 		const FPBDJointSettings& GetJointSettings() const { return GetSettings(); }	//needed for property macros
@@ -137,9 +141,12 @@ namespace Chaos
 		bool bDisabled;
 		bool bBroken;
 		bool bBreaking;
+		bool bViolating;
 		bool bDriveTargetChanged;
 		FVec3 LinearImpulse;
 		FVec3 AngularImpulse;
+		float LinearViolation;
+		float AngularViolation;
 		EResimType ResimType = EResimType::FullResim;
 		ESyncState SyncState = ESyncState::InSync;
 		bool bEnabledDuringResim = true;
@@ -236,6 +243,16 @@ namespace Chaos
 		 */
 		CHAOS_API void ClearConstraintBreaking(int32 ConstraintIndex);
 
+		/*
+		 * Whether the constraint was violated beyond a threshold this frame (transient flag for use by event system)
+		 */
+		CHAOS_API bool IsConstraintViolating(int32 ConstraintIndex) const;
+
+		/*
+		 * Clear the transient constraint violation state (called by event system when it has used the flag)
+		 */
+		CHAOS_API void ClearConstraintViolating(int32 ConstraintIndex);
+
 		/*
 		 * Whether the drive target has changed
 		 */
@@ -339,6 +356,9 @@ namespace Chaos
 		*/
 		CHAOS_API FVec3 GetConstraintAngularImpulse(int32 ConstraintIndex) const;
 
+		CHAOS_API float GetConstraintLinearViolation(int32 ConstraintIndex) const;
+		CHAOS_API float GetConstraintAngularViolation(int32 ConstraintIndex) const;
+
 		CHAOS_API ESyncState GetConstraintSyncState(int32 ConstraintIndex) const;
 		CHAOS_API void SetConstraintSyncState(int32 ConstraintIndex, ESyncState SyncState);
 		
@@ -360,7 +380,7 @@ namespace Chaos
 		CHAOS_API virtual void UnprepareTick() override final;
 
 		// Called by the joint solver at the end of the constraint solver phase
-		CHAOS_API void SetSolverResults(const int32 ConstraintIndex, const FVec3& LinearImpulse, const FVec3& AngularImpulse, const bool bIsBroken, const FSolverBody* SolverBody0, const FSolverBody* SolverBody1);
+		CHAOS_API void SetSolverResults(const int32 ConstraintIndex, const FVec3& LinearImpulse, const FVec3& AngularImpulse, const float LinearViolation, const float AngularViolation, const bool bIsBroken, const bool bIsViolating, const FSolverBody* SolverBody0, const FSolverBody* SolverBody1);
 
 		// @todo(chaos): only needed for RBAN, and should be private or moved to the solver
 		CHAOS_API int32 GetConstraintIsland(int32 ConstraintIndex) const;
diff --git a/Engine/Source/Runtime/Experimental/Chaos/Public/Chaos/PullPhysicsDataImp.h b/Engine/Source/Runtime/Experimental/Chaos/Public/Chaos/PullPhysicsDataImp.h
--- a/Engine/Source/Runtime/Experimental/Chaos/Public/Chaos/PullPhysicsDataImp.h
+++ b/Engine/Source/Runtime/Experimental/Chaos/Public/Chaos/PullPhysicsDataImp.h
@@ -134,9 +134,12 @@ struct FDirtyClusterUnionData : public TBasePullData<FClusterUnionPhysicsProxy,
 struct FJointConstraintOutputData {
 	bool bIsBreaking = false;
 	bool bIsBroken = false;
+	bool bIsViolating = false;
 	bool bDriveTargetChanged = false;
 	FVector Force = FVector(0);
 	FVector Torque = FVector(0);
+	float LinearViolation = 0.f;
+	float AngularViolation = 0.f;
 };
 
 class FJointConstraint;
diff --git a/Engine/Source/Runtime/PhysicsCore/Private/ChaosEngineInterface.cpp b/Engine/Source/Runtime/PhysicsCore/Private/ChaosEngineInterface.cpp
--- a/Engine/Source/Runtime/PhysicsCore/Private/ChaosEngineInterface.cpp
+++ b/Engine/Source/Runtime/PhysicsCore/Private/ChaosEngineInterface.cpp
@@ -1718,6 +1718,18 @@ void FChaosEngineInterface::SetBreakForces_AssumesLocked(const FPhysicsConstrain
 	}
 }
 
+void FChaosEngineInterface::SetViolationCallbackThresholds_AssumesLocked(const FPhysicsConstraintHandle& InConstraintRef,float InLinearViolationCallbackThreshold,float InAngularViolationCallbackThreshold)
+{
+	if (InConstraintRef.IsValid() && InConstraintRef.Constraint->IsType(Chaos::EConstraintType::JointConstraintType))
+	{
+		if (Chaos::FJointConstraint* Constraint = static_cast<Chaos::FJointConstraint*>(InConstraintRef.Constraint))
+		{
+			Constraint->SetLinearViolationCallbackThreshold(InLinearViolationCallbackThreshold);
+			Constraint->SetAngularViolationCallbackThreshold(InAngularViolationCallbackThreshold);
+		}
+	}
+}
+
 void FChaosEngineInterface::SetPlasticityLimits_AssumesLocked(const FPhysicsConstraintHandle& InConstraintRef, float InLinearPlasticityLimit, float InAngularPlasticityLimit, EConstraintPlasticityType InLinearPlasticityType)
 {
 	if (InConstraintRef.IsValid() && InConstraintRef.Constraint->IsType(Chaos::EConstraintType::JointConstraintType))
diff --git a/Engine/Source/Runtime/PhysicsCore/Public/Chaos/ChaosEngineInterface.h b/Engine/Source/Runtime/PhysicsCore/Public/Chaos/ChaosEngineInterface.h
--- a/Engine/Source/Runtime/PhysicsCore/Public/Chaos/ChaosEngineInterface.h
+++ b/Engine/Source/Runtime/PhysicsCore/Public/Chaos/ChaosEngineInterface.h
@@ -517,6 +517,7 @@ public:
 	static PHYSICSCORE_API void SetParentDominates_AssumesLocked(const FPhysicsConstraintHandle& InConstraintRef, bool bInParentDominates);
 	static PHYSICSCORE_API void SetMassConditioningEnabled_AssumesLocked(const FPhysicsConstraintHandle& InConstraintRef, bool bInMassConditioningEnabled);
 	static PHYSICSCORE_API void SetBreakForces_AssumesLocked(const FPhysicsConstraintHandle& InConstraintRef,float InLinearBreakForce,float InAngularBreakForce);
+	static PHYSICSCORE_API void SetViolationCallbackThresholds_AssumesLocked(const FPhysicsConstraintHandle& InConstraintRef,float InLinearBreakViolation,float InAngularBreakViolation);
 	static PHYSICSCORE_API void SetPlasticityLimits_AssumesLocked(const FPhysicsConstraintHandle& InConstraintRef, float InLinearPlasticityLimit, float InAngularPlasticityLimit, EConstraintPlasticityType InLinearPlasticityType);
 	static PHYSICSCORE_API void SetContactTransferScale_AssumesLocked(const FPhysicsConstraintHandle& InConstraintRef, float InContactTransferScale);
 	static PHYSICSCORE_API void SetLocalPose(const FPhysicsConstraintHandle& InConstraintRef, const FTransform& InPose, EConstraintFrame::Type InFrame);
