Allow IAD cache size to be configured

diff --git a/Engine/Source/Runtime/Experimental/IoStore/OnDemand/Private/IoStoreOnDemand.cpp b/Engine/Source/Runtime/Experimental/IoStore/OnDemand/Private/IoStoreOnDemand.cpp
--- a/Engine/Source/Runtime/Experimental/IoStore/OnDemand/Private/IoStoreOnDemand.cpp
+++ b/Engine/Source/Runtime/Experimental/IoStore/OnDemand/Private/IoStoreOnDemand.cpp
@@ -58,7 +58,7 @@ static FAutoConsoleVariableRef CVar_DistributedEndpointFallbackUrl(
 );
 
 ////////////////////////////////////////////////////////////////////////////////
-static int64 ParseSizeParam(FStringView Value)
+int64 ParseSizeParam(FStringView Value)
 {
 	Value = Value.TrimStartAndEnd();
 
diff --git a/Engine/Source/Runtime/Experimental/IoStore/OnDemand/Private/OnDemandIoStore.cpp b/Engine/Source/Runtime/Experimental/IoStore/OnDemand/Private/OnDemandIoStore.cpp
--- a/Engine/Source/Runtime/Experimental/IoStore/OnDemand/Private/OnDemandIoStore.cpp
+++ b/Engine/Source/Runtime/Experimental/IoStore/OnDemand/Private/OnDemandIoStore.cpp
@@ -25,6 +25,7 @@
 #include "IO/IoContainerHeader.h"
 #include "IO/PackageStore.h"
 #include "Misc/CommandLine.h"
+#include "Misc/ConfigCacheIni.h"
 #include "Misc/CoreDelegates.h"
 #include "Misc/CoreDelegatesInternal.h"
 #include "Misc/CoreMisc.h"
@@ -776,6 +777,16 @@ FIoStatus FOnDemandIoStore::InitializeOnDemandInstallCache()
 	if (bUseInstallCache)
 	{
 		FOnDemandInstallCacheConfig CacheConfig;
+
+		if (FString ValueStr; GConfig->GetString(TEXT("OnDemandInstall"), TEXT("FileCache.DiskQuota"), ValueStr, GEngineIni))
+		{
+			int64 DiskQuota = ParseSizeParam(ValueStr);
+			if (DiskQuota > 0)
+			{
+				CacheConfig.DiskQuota = DiskQuota;
+			}
+		}
+
 		CacheConfig.RootDirectory = Private::GetInstallCacheDirectory();
 #if !UE_BUILD_SHIPPING
 		CacheConfig.bDropCache = FParse::Param(FCommandLine::Get(), TEXT("Iad.DropCache"));
diff --git a/Engine/Source/Runtime/Experimental/IoStore/OnDemand/Public/IO/IoStoreOnDemand.h b/Engine/Source/Runtime/Experimental/IoStore/OnDemand/Public/IO/IoStoreOnDemand.h
--- a/Engine/Source/Runtime/Experimental/IoStore/OnDemand/Public/IO/IoStoreOnDemand.h
+++ b/Engine/Source/Runtime/Experimental/IoStore/OnDemand/Public/IO/IoStoreOnDemand.h
@@ -52,6 +52,10 @@ namespace UE::IoStore
 
 ////////////////////////////////////////////////////////////////////////////////
 
+int64 ParseSizeParam(FStringView Value);
+
+////////////////////////////////////////////////////////////////////////////////
+
 bool TryParseConfigFile(const FString& ConfigPath, FOnDemandEndpoint& OutEndpoint);
 
 ////////////////////////////////////////////////////////////////////////////////
