Fix OnDemandIoStore not initializing on server after first child is forked.
Add additional logging

diff --git a/Engine/Source/Runtime/Experimental/IoStore/OnDemand/Private/OnDemandIoStore.cpp b/Engine/Source/Runtime/Experimental/IoStore/OnDemand/Private/OnDemandIoStore.cpp
--- a/Engine/Source/Runtime/Experimental/IoStore/OnDemand/Private/OnDemandIoStore.cpp
+++ b/Engine/Source/Runtime/Experimental/IoStore/OnDemand/Private/OnDemandIoStore.cpp
@@ -581,6 +581,7 @@ FIoStatus FOnDemandIoStore::Initialize()
 	}
 #endif
 
+	UE_LOG(LogIoStoreOnDemand, Display, TEXT("Initialized"));
 	return EIoErrorCode::Ok;
 }
 
@@ -749,13 +750,14 @@ FOnDemandChunkInfo FOnDemandIoStore::GetInstalledChunkInfo(const FIoChunkId& Chu
 
 void FOnDemandIoStore::OnPostFork(EForkProcessRole ProcessRole)
 {
-	FCoreDelegates::OnPostFork.RemoveAll(this);
-
 	if (ProcessRole != EForkProcessRole::Child)
 	{
+		UE_LOG(LogIoStoreOnDemand, Display, TEXT("OnPostFork ProcessRole Parent"));
 		return;
 	}
 
+	UE_LOG(LogIoStoreOnDemand, Display, TEXT("OnPostFork ProcessRole Child"));
+
 	const FIoStatus Status = Initialize();
 	if (!Status.IsOk())
 	{
@@ -809,9 +811,14 @@ FIoStatus FOnDemandIoStore::InitializeOnDemandInstallCache()
 		else
 		{
 			// Only warn until this is properly tested
-			UE_LOG(LogIoStoreOnDemand, Warning, TEXT("Failed to initialize install cache"));
+			UE_LOG(LogIoStoreOnDemand, Error, TEXT("Failed to initialize install cache"));
+			return FIoStatusBuilder(EIoErrorCode::InvalidParameter) << TEXT("Failed to initialize install cache");
 		}
 	}
+	else
+	{
+		UE_LOG(LogIoStoreOnDemand, Warning, TEXT("Install cache disabled"));
+	}
 
 #if !(UE_BUILD_SHIPPING|UE_BUILD_TEST)
 	FString ParamValue;
