Mitigate against possible crash

diff --git a/Engine/Source/Runtime/Experimental/ChaosVehicles/ChaosVehiclesCore/Private/SimModule/ModuleInput.cpp b/Engine/Source/Runtime/Experimental/ChaosVehicles/ChaosVehiclesCore/Private/SimModule/ModuleInput.cpp
--- a/Engine/Source/Runtime/Experimental/ChaosVehicles/ChaosVehiclesCore/Private/SimModule/ModuleInput.cpp
+++ b/Engine/Source/Runtime/Experimental/ChaosVehicles/ChaosVehiclesCore/Private/SimModule/ModuleInput.cpp
@@ -258,7 +258,8 @@ void FModuleInputContainer::Lerp(const FModuleInputContainer& Min, const FModule
 
 void FModuleInputContainer::Merge(const FModuleInputContainer& From)
 {
-	for (int I = 0; I < InputValues.Num(); I++)
+	int Num = FMath::Min(InputValues.Num(), From.InputValues.Num());
+	for (int I = 0; I < Num; I++)
 	{
 		InputValues[I].Merge(From.InputValues[I]);
 	}
