Fix bug where input container array was growing. Not being reset on copy/serialization.

diff --git a/Engine/Source/Runtime/Experimental/ChaosVehicles/ChaosVehiclesCore/Private/SimModule/ModuleInput.cpp b/Engine/Source/Runtime/Experimental/ChaosVehicles/ChaosVehiclesCore/Private/SimModule/ModuleInput.cpp
--- a/Engine/Source/Runtime/Experimental/ChaosVehicles/ChaosVehiclesCore/Private/SimModule/ModuleInput.cpp
+++ b/Engine/Source/Runtime/Experimental/ChaosVehicles/ChaosVehiclesCore/Private/SimModule/ModuleInput.cpp
@@ -218,7 +218,7 @@ void FModuleInputContainer::Serialize(FArchive& Ar, UPackageMap* Map, bool& bOut
 	Ar << Number;
 	if (Ar.IsLoading())
 	{
-		InputValues.Reserve(Number);
+		InputValues.Reset(Number);
 		for (int I = 0; I < Number; I++)
 		{
 			FModuleInputValue Value;
@@ -268,38 +268,47 @@ void FModuleInputContainer::Merge(const FModuleInputContainer& From)
 
 void FInputInterface::SetValue(const FName& InName, const FModuleInputValue& InValue)
 {
-	if (const int* Index = NameMap.Find(InName))
+	if (ValueContainer.GetNumInputs() > 0)
 	{
-		ValueContainer.SetValueAtIndex(*Index, InValue);
-	}
-	else
-	{
-		UE_LOG(LogModularInput, Warning, TEXT("Trying to set the value of an undefined control input %s"), *InName.ToString());
+		if (const int* Index = NameMap.Find(InName))
+		{
+			ValueContainer.SetValueAtIndex(*Index, InValue);
+		}
+		else
+		{
+			UE_LOG(LogModularInput, Warning, TEXT("Trying to set the value of an undefined control input %s"), *InName.ToString());
+		}
 	}
 }
 
 void FInputInterface::MergeValue(const FName& InName, const FModuleInputValue& InValue)
 {
-	if (const int* Index = NameMap.Find(InName))
-	{
-		ValueContainer.MergeValueAtIndex(*Index, InValue);
-	}
-	else
+	if (ValueContainer.GetNumInputs() > 0)
 	{
-		UE_LOG(LogModularInput, Warning, TEXT("Trying to set the value of an undefined control input %s"), *InName.ToString());
+		if (const int* Index = NameMap.Find(InName))
+		{
+			ValueContainer.MergeValueAtIndex(*Index, InValue);
+		}
+		else
+		{
+			UE_LOG(LogModularInput, Warning, TEXT("Trying to set the value of an undefined control input %s"), *InName.ToString());
+		}
 	}
 }
 
 
 FModuleInputValue FInputInterface::GetValue(const FName& InName) const
 {
-	if (const int* Index = NameMap.Find(InName))
-	{
-		return ValueContainer.GetValueAtIndex(*Index);
-	}
-	else
+	if (ValueContainer.GetNumInputs() > 0)
 	{
-		UE_LOG(LogModularInput, Warning, TEXT("Trying to get the value of an undefined control input %s"), *InName.ToString());
+		if (const int* Index = NameMap.Find(InName))
+		{
+			return ValueContainer.GetValueAtIndex(*Index);
+		}
+		else
+		{
+			UE_LOG(LogModularInput, Warning, TEXT("Trying to get the value of an undefined control input %s"), *InName.ToString());
+		}
 	}
 
 	return FModuleInputValue(EModuleInputValueType::MBoolean, FVector::ZeroVector);
@@ -307,9 +316,12 @@ FModuleInputValue FInputInterface::GetValue(const FName& InName) const
 
 float FInputInterface::GetMagnitude(const FName& InName) const
 {
-	if (const int* Index = NameMap.Find(InName))
+	if (ValueContainer.GetNumInputs() > 0)
 	{
-		return ValueContainer.GetValueAtIndex(*Index).GetMagnitude();
+		if (const int* Index = NameMap.Find(InName))
+		{
+			return ValueContainer.GetValueAtIndex(*Index).GetMagnitude();
+		}
 	}
 
 	return 0.0f;
diff --git a/Engine/Source/Runtime/Experimental/ChaosVehicles/ChaosVehiclesCore/Public/SimModule/ModuleInput.h b/Engine/Source/Runtime/Experimental/ChaosVehicles/ChaosVehiclesCore/Public/SimModule/ModuleInput.h
--- a/Engine/Source/Runtime/Experimental/ChaosVehicles/ChaosVehiclesCore/Public/SimModule/ModuleInput.h
+++ b/Engine/Source/Runtime/Experimental/ChaosVehicles/ChaosVehiclesCore/Public/SimModule/ModuleInput.h
@@ -370,6 +370,7 @@ public:
 	{
 		if (&Other != this)
 		{
+			InputValues.Reset();
 			// perform a deep copy
 			if (Other.InputValues.Num())
 			{
