Geometry collection : add support for collision visualization for the nanite scene proxy

diff --git a/Engine/Source/Runtime/Experimental/GeometryCollectionEngine/Private/GeometryCollection/GeometryCollectionSceneProxy.cpp b/Engine/Source/Runtime/Experimental/GeometryCollectionEngine/Private/GeometryCollection/GeometryCollectionSceneProxy.cpp
--- a/Engine/Source/Runtime/Experimental/GeometryCollectionEngine/Private/GeometryCollection/GeometryCollectionSceneProxy.cpp
+++ b/Engine/Source/Runtime/Experimental/GeometryCollectionEngine/Private/GeometryCollection/GeometryCollectionSceneProxy.cpp
@@ -1189,6 +1189,8 @@ FNaniteGeometryCollectionSceneProxy::FNaniteGeometryCollectionSceneProxy(UGeomet
 	TArray<FMatrix44f> RestTransforms;
 	Component->GetRestTransforms(RestTransforms);
 
+	CollisionResponse = Component->GetCollisionResponseToChannels();
+
 	FGeometryCollectionDynamicData* DynamicData = Component->InitDynamicData(true);
 	SetDynamicData_RenderThread(DynamicData, Component->GetRenderMatrix());
 }
@@ -1220,6 +1222,9 @@ FPrimitiveViewRelevance FNaniteGeometryCollectionSceneProxy::GetViewRelevance(co
 	// Always render the Nanite mesh data with static relevance.
 	Result.bStaticRelevance = true;
 
+	// dynamic still relevance still must be used when drawing collisions
+	Result.bDynamicRelevance = ShowCollisionMeshes(View->Family->EngineShowFlags);
+
 	// Should always be covered by constructor of Nanite scene proxy.
 	Result.bRenderInMainPass = true;
 
@@ -1412,6 +1417,63 @@ void FNaniteGeometryCollectionSceneProxy::OnMotionEnd()
 }
 
 
+bool FNaniteGeometryCollectionSceneProxy::ShowCollisionMeshes(const FEngineShowFlags& EngineShowFlags) const
+{
+	if (IsCollisionEnabled())
+	{
+		if (EngineShowFlags.CollisionPawn && CollisionResponse.GetResponse(ECC_Pawn) != ECR_Ignore)
+		{
+			return true;
+		}
+		if (EngineShowFlags.CollisionVisibility && CollisionResponse.GetResponse(ECC_Visibility) != ECR_Ignore)
+		{
+			return true;
+		}
+		if (EngineShowFlags.Collision)
+		{
+			return true;
+		}
+	}
+	return false;
+}
+
+void FNaniteGeometryCollectionSceneProxy::GetDynamicMeshElements(const TArray<const FSceneView*>& Views, const FSceneViewFamily& ViewFamily, uint32 VisibilityMap, FMeshElementCollector& Collector) const
+{
+	QUICK_SCOPE_CYCLE_COUNTER(STAT_NaniteGeometryCollectionSceneProxy_GetDynamicMeshElements);
+
+	const FEngineShowFlags& EngineShowFlags = ViewFamily.EngineShowFlags;
+	const bool bDrawWireframeCollision = EngineShowFlags.Collision && IsCollisionEnabled();
+
+	// draw extra stuff ( collision , bounds ... )
+	for (int32 ViewIndex = 0; ViewIndex < Views.Num(); ViewIndex++)
+	{
+		if (VisibilityMap & (1 << ViewIndex))
+		{
+			// collision modes
+			if (ShowCollisionMeshes(EngineShowFlags) && GeometryCollection && GeometryCollection->GetGeometryCollection() && AllowDebugViewmodes())
+			{
+				FTransform GeomTransform(GetLocalToWorld());
+				if (bDrawWireframeCollision)
+				{
+					GeometryCollectionDebugDraw::DrawWireframe(*GeometryCollection->GetGeometryCollection(), GeomTransform, Collector, ViewIndex, GetWireframeColor().ToFColor(true));
+				}
+				else
+				{
+					FMaterialRenderProxy* CollisionMaterialInstance = new FColoredMaterialRenderProxy(GEngine->ShadedLevelColorationUnlitMaterial->GetRenderProxy(), GetWireframeColor());
+					Collector.RegisterOneFrameMaterialProxy(CollisionMaterialInstance);
+					GeometryCollectionDebugDraw::DrawSolid(*GeometryCollection->GetGeometryCollection(), GeomTransform, Collector, ViewIndex, CollisionMaterialInstance);
+				}
+			}
+
+			// render bounds
+#if !(UE_BUILD_SHIPPING || UE_BUILD_TEST)
+			RenderBounds(Collector.GetPDI(ViewIndex), ViewFamily.EngineShowFlags, GetBounds(), IsSelected());
+#endif
+		}
+	}
+}
+
+
 FGeometryCollectionDynamicDataPool::FGeometryCollectionDynamicDataPool()
 {
 	FreeList.SetNum(32);
diff --git a/Engine/Source/Runtime/Experimental/GeometryCollectionEngine/Private/GeometryCollection/GeometryCollectionSceneProxy.h b/Engine/Source/Runtime/Experimental/GeometryCollectionEngine/Private/GeometryCollection/GeometryCollectionSceneProxy.h
--- a/Engine/Source/Runtime/Experimental/GeometryCollectionEngine/Private/GeometryCollection/GeometryCollectionSceneProxy.h
+++ b/Engine/Source/Runtime/Experimental/GeometryCollectionEngine/Private/GeometryCollection/GeometryCollectionSceneProxy.h
@@ -343,6 +343,7 @@ public:
 	virtual void CreateRenderThreadResources(FRHICommandListBase& RHICmdList) override;
 	virtual SIZE_T GetTypeHash() const override;
 	virtual FPrimitiveViewRelevance	GetViewRelevance(const FSceneView* View) const override;
+	virtual void GetDynamicMeshElements(const TArray<const FSceneView*>& Views, const FSceneViewFamily& ViewFamily, uint32 VisibilityMap, FMeshElementCollector& Collector) const override;
 #if GEOMETRYCOLLECTION_EDITOR_SELECTION
 	virtual HHitProxy* CreateHitProxies(UPrimitiveComponent* Component, TArray<TRefCountPtr<HHitProxy> >& OutHitProxies) override;
 #endif
@@ -386,6 +387,7 @@ public:
 protected:
 	// TODO : Copy required data from UObject instead of using unsafe object pointer.
 	const UGeometryCollection* GeometryCollection = nullptr;
+	FCollisionResponseContainer CollisionResponse;
 
 	struct FGeometryNaniteData
 	{
@@ -423,6 +425,9 @@ protected:
 		GEOMETRYCOLLECTIONENGINE_API virtual FLightInteraction GetInteraction(const FLightSceneProxy* LightSceneProxy) const override;
 	};
 
+private:
+	bool ShowCollisionMeshes(const FEngineShowFlags& EngineShowFlags) const;
+
 private:
 	static FEmptyLightCacheInfo EmptyLightCacheInfo;
 };
