Chaos Modular Vehicles:
Serialize FSuspensionTargetPoint for networking.

diff --git a/Engine/Source/Runtime/Experimental/ChaosVehicles/ChaosVehiclesCore/Public/SimModule/SuspensionBaseInterface.h b/Engine/Source/Runtime/Experimental/ChaosVehicles/ChaosVehiclesCore/Public/SimModule/SuspensionBaseInterface.h
--- a/Engine/Source/Runtime/Experimental/ChaosVehicles/ChaosVehiclesCore/Public/SimModule/SuspensionBaseInterface.h
+++ b/Engine/Source/Runtime/Experimental/ChaosVehicles/ChaosVehiclesCore/Public/SimModule/SuspensionBaseInterface.h
@@ -40,6 +40,15 @@ struct CHAOSVEHICLESCORE_API FSuspensionTargetPoint
 		, SurfaceType(InSurfaceType)
 	{}
 
+	void Serialize(FArchive& Ar)
+	{
+		Ar << TargetPosition;
+		Ar << ImpactNormal;
+		Ar << HitDistance;
+		Ar << bWheelInContact;
+		Ar << SurfaceType;
+	}
+
 	FVector TargetPosition = FVector::ZeroVector;
 	FVector ImpactNormal = FVector::ZeroVector;
 	float HitDistance = 0.0f;
@@ -68,7 +77,7 @@ public:
 	virtual void GetWorldRaycastLocation(const FTransform& BodyTransform, float WheelRadius, FSpringTrace& OutTrace) = 0;
 
 	void SetTargetPoint(const FSuspensionTargetPoint& InTargetPoint);
-	
+	const FSuspensionTargetPoint& GetTargetPoint() const { return TargetPoint; }
 	bool IsWheelInContact() const { return TargetPoint.bWheelInContact; }
 	
 	void SetWheelSimTreeIndex(int WheelTreeIndexIn) { WheelSimTreeIndex = WheelTreeIndexIn; }
