IOStoreOnDemand Cache full purge functionality
- Added iostore.PurgeOnDemandInstallCache debug command

diff --git a/Engine/Source/Runtime/Experimental/IoStore/OnDemand/Private/IoStoreOnDemand.cpp b/Engine/Source/Runtime/Experimental/IoStore/OnDemand/Private/IoStoreOnDemand.cpp
--- a/Engine/Source/Runtime/Experimental/IoStore/OnDemand/Private/IoStoreOnDemand.cpp
+++ b/Engine/Source/Runtime/Experimental/IoStore/OnDemand/Private/IoStoreOnDemand.cpp
@@ -44,6 +44,35 @@ DEFINE_LOG_CATEGORY(LogIas);
 namespace UE::IoStore
 {
 
+///////////////////////////////////////////////////////////////////////////////
+static FAutoConsoleCommand OnDemandPurgeCacheCommand(
+	TEXT("iostore.PurgeOnDemandInstallCache"),
+	TEXT("Purge On Demand Install Cache"),
+	FConsoleCommandDelegate::CreateStatic([]()
+	{
+		FIoStoreOnDemandModule* IOStoreOnDemandModule = FModuleManager::Get().GetModulePtr<FIoStoreOnDemandModule>(TEXT("IoStoreOnDemand"));
+		if(!IOStoreOnDemandModule)
+		{
+			UE_LOG(LogIoStoreOnDemand, Error, TEXT("Could not find IoStoreOnDemand module"));
+			return;
+		}
+
+		UE_LOG(LogIoStoreOnDemand, Display, TEXT("Purging on demand install cache"));
+		IOStoreOnDemandModule->Purge(FOnDemandPurgeArgs(), [](const FOnDemandPurgeResult& Result)
+		{
+			if (Result.Status.IsOk())
+			{
+				UE_LOG(LogIoStoreOnDemand, Display, TEXT("Purged on demand install cache"));
+			}
+			else
+			{
+				UE_LOG(LogIoStoreOnDemand, Error, TEXT("Failed Purged on demand install cache: %s"), *Result.Status.ToString());
+			}
+		});
+	}),
+	ECVF_Cheat
+);
+
 ////////////////////////////////////////////////////////////////////////////////
 FString GIasOnDemandTocExt = TEXT(".uondemandtoc");
 
@@ -1084,6 +1113,22 @@ void FIoStoreOnDemandModule::Install(
 	IoStore->Install(MoveTemp(Args), MoveTemp(OnCompleted), MoveTemp(OnProgress), CancellationToken);
 }
 
+void FIoStoreOnDemandModule::Purge(FOnDemandPurgeArgs&& Args, FOnDemandPurgeCompleted&& OnCompleted)
+{
+	if (IoStore.IsValid() == false)
+	{
+		IoStore = MakeShared<FOnDemandIoStore>();
+		if (FIoStatus Status = IoStore->Initialize(); !Status.IsOk())
+		{
+			UE_LOG(LogIas, Error, TEXT("Failed to initialize I/O store on-demand, reason '%s'"), *Status.ToString());
+			IoStore.Reset();
+			return OnCompleted(FOnDemandPurgeResult{ .Status = Status });
+		}
+	}
+
+	IoStore->Purge(MoveTemp(Args), MoveTemp(OnCompleted));
+}
+
 FIoStatus FIoStoreOnDemandModule::Unmount(FStringView MountId)
 {
 	if (IoStore.IsValid())
diff --git a/Engine/Source/Runtime/Experimental/IoStore/OnDemand/Private/OnDemandInstallCache.cpp b/Engine/Source/Runtime/Experimental/IoStore/OnDemand/Private/OnDemandInstallCache.cpp
--- a/Engine/Source/Runtime/Experimental/IoStore/OnDemand/Private/OnDemandInstallCache.cpp
+++ b/Engine/Source/Runtime/Experimental/IoStore/OnDemand/Private/OnDemandInstallCache.cpp
@@ -4,6 +4,7 @@
 #include "OnDemandHttpClient.h"
 #include "OnDemandIoStore.h"
 
+#include "Algo/Accumulate.h"
 #include "Async/Mutex.h"
 #include "Async/UniqueLock.h"
 #include "Async/AsyncFileHandle.h"
@@ -100,7 +101,7 @@ struct FCas
 	using FLastAccess		= TMap<FCasBlockId, int64>;
 
 	FIoStatus				Initialize(FStringView Directory);
-	FCasLocation			FindChunk(const FIoHash& Hash);
+	FCasLocation			FindChunk(const FIoHash& Hash) const;
 	FCasBlockId				CreateBlock();
 	FIoStatus				DeleteBlock(FCasBlockId BlockId, TArray<FCasAddr>& OutAddrs);
 	FString					GetBlockFilename(FCasBlockId BlockId) const;
@@ -119,7 +120,7 @@ struct FCas
 	FReadHandles		ReadHandles;	
 	uint32				MaxBlockSize = 32 << 20; //TODO: Make configurable
 	FCasBlockId			CurrentBlock;
-	UE::FMutex			Mutex;
+	mutable UE::FMutex	Mutex;
 };
 
 ///////////////////////////////////////////////////////////////////////////////
@@ -152,7 +153,7 @@ FIoStatus FCas::Initialize(FStringView Directory)
 	return EIoErrorCode::Ok;
 };
 
-FCasLocation FCas::FindChunk(const FIoHash& Hash)
+FCasLocation FCas::FindChunk(const FIoHash& Hash) const
 {
 	const FCasAddr* Addr	= reinterpret_cast<const FCasAddr*>(&Hash);
 	const uint32 TypeHash	= GetTypeHash(*Addr);
@@ -913,9 +914,15 @@ public:
 	virtual bool				IsChunkCached(const FIoHash& ChunkHash) override;
 	virtual FIoStatus			PutChunk(FIoBuffer&& Chunk, const FIoHash& ChunkHash) override;
 	virtual FIoStatus			Purge(TMap<FIoHash, uint64>&& ChunksToIntall) override;
+	virtual FIoStatus			PurgeAllUnreferenced() override;
 	virtual FIoStatus			Flush() override;
 
 private:
+	void						AddReferencesToBlocks(
+									const TArray<FSharedOnDemandContainer>& Containers, 
+									const TArray<TBitArray<>>& ChunkEntryIndices,
+									FCasBlockInfoMap& BlockInfo) const;
+	FIoStatus					Purge(const FCasBlockInfoMap& BlockInfo, uint64 TotalBytesToPurge, uint64& OutTotalPurgedBytes);
 	bool						Resolve(FIoRequestImpl* Request);
 	void						CompleteRequest(FIoRequestImpl* Request, bool bFileReadWasCancelled);
 	FIoStatus					FlushPendingChunks(FPendingChunks& Block);
@@ -1256,11 +1263,95 @@ FIoStatus FOnDemandInstallCache::Purge(TMap<FIoHash, uint64>&& ChunksToIntall)
 	IoStore.GetReferencedContent(Containers, ChunkEntryIndices);
 	check(Containers.Num() == ChunkEntryIndices.Num());
 
+	AddReferencesToBlocks(Containers, ChunkEntryIndices, BlockInfo);
+
+	//TODO: Compute fragmentation metric and redownload chunks when this number gets too high
+
+	BlockInfo.ValueSort([](const FCasBlockInfo& LHS, const FCasBlockInfo& RHS)
+	{
+		return LHS.LastAccess < RHS.LastAccess;
+	});
+
+	const uint64 TotalReferencedBytes = Algo::TransformAccumulate(BlockInfo,
+		[](const TPair<FCasBlockId, FCasBlockInfo>& Kv){ return (Kv.Value.RefCount > 0) ? Kv.Value.FileSize : uint64(0); },
+		uint64(0));
+
+	UE_LOG(LogIoStoreOnDemand, Log, TEXT("Purging install cache, MaxCacheSize=%.2lf MiB, CacheSize=%.2lf MiB, UncachedSize=%.2lf MiB, ReferencedBytes=%.2lf MiB"),
+		ToMiB(MaxCacheSize), ToMiB(TotalCachedBytes), ToMiB(TotalUncachedBytes), ToMiB(TotalReferencedBytes));
+
+	const uint64	TotalBytesToPurge	= TotalRequiredBytes - MaxCacheSize;
+	uint64			TotalPurgedBytes	= 0;
+
+	FIoStatus Status = Purge(BlockInfo, TotalBytesToPurge, TotalPurgedBytes);
+
+	if (TotalPurgedBytes > 0)
+	{
+		UE_LOG(LogIoStoreOnDemand, Log, TEXT("Purged %.2lf MiB (%.2lf%%) from install cache"),
+			ToMiB(TotalPurgedBytes), 100.0 * (double(TotalPurgedBytes) / double(TotalCachedBytes)));
+	}
+
+	const uint64 NewCachedBytes = TotalCachedBytes - TotalPurgedBytes;
+	UE_CLOG(NewCachedBytes > MaxCacheSize,
+		LogIoStoreOnDemand, Warning, TEXT("Max install cache size exceeded by %.2lf MiB (%.2lf%%)"),
+			ToMiB(NewCachedBytes - MaxCacheSize), 100.0 * (double(NewCachedBytes - MaxCacheSize) / double(MaxCacheSize)));
+
+	if (TotalPurgedBytes < TotalBytesToPurge)
+	{
+		return FIoStatusBuilder(EIoErrorCode::WriteError) << FString::Printf(TEXT("Failed to purge %" UINT64_FMT " from install cache"), TotalBytesToPurge);
+	}
+
+	return Status;
+}
+
+FIoStatus FOnDemandInstallCache::PurgeAllUnreferenced()
+{
+	FCasBlockInfoMap	BlockInfo;
+	const uint64		TotalCachedBytes = Cas.GetBlockInfo(BlockInfo);
+	uint64				TotalUncachedBytes = 0;
+
+	TArray<FSharedOnDemandContainer>	Containers;
+	TArray<TBitArray<>>					ChunkEntryIndices;
+
+	IoStore.GetReferencedContent(Containers, ChunkEntryIndices);
+	check(Containers.Num() == ChunkEntryIndices.Num());
+
+	AddReferencesToBlocks(Containers, ChunkEntryIndices, BlockInfo);
+
 	//TODO: Compute fragmentation metric and redownload chunks when this number gets too high
+
+	const uint64 TotalReferencedBytes = Algo::TransformAccumulate(BlockInfo,
+		[](const TPair<FCasBlockId, FCasBlockInfo>& Kv) { return (Kv.Value.RefCount > 0) ? Kv.Value.FileSize : uint64(0); },
+		uint64(0));
+
+	UE_LOG(LogIoStoreOnDemand, Log, TEXT("Purging install cache, MaxCacheSize=%.2lf MiB, CacheSize=%.2lf MiB, UncachedSize=%.2lf MiB, ReferencedBytes=%.2lf MiB"),
+		ToMiB(MaxCacheSize), ToMiB(TotalCachedBytes), ToMiB(TotalUncachedBytes), ToMiB(TotalReferencedBytes));
+
+	uint64 TotalPurgedBytes = 0;
+	FIoStatus Status = Purge(BlockInfo, MaxCacheSize, TotalPurgedBytes);
+
+	if (TotalPurgedBytes > 0)
+	{
+		UE_LOG(LogIoStoreOnDemand, Log, TEXT("Purged %.2lf MiB (%.2lf%%) from install cache"),
+			ToMiB(TotalPurgedBytes), 100.0 * (double(TotalPurgedBytes) / double(TotalCachedBytes)));
+	}
+
+	const uint64 NewCachedBytes = TotalCachedBytes - TotalPurgedBytes;
+	UE_CLOG(NewCachedBytes > MaxCacheSize,
+		LogIoStoreOnDemand, Warning, TEXT("Max install cache size exceeded by %.2lf MiB (%.2lf%%)"),
+			ToMiB(NewCachedBytes - MaxCacheSize), 100.0 * (double(NewCachedBytes - MaxCacheSize) / double(MaxCacheSize)));
+
+	return Status;
+}
+
+void FOnDemandInstallCache::AddReferencesToBlocks( 
+	const TArray<FSharedOnDemandContainer>& Containers, 
+	const TArray<TBitArray<>>& ChunkEntryIndices,
+	FCasBlockInfoMap& BlockInfo) const
+{
 	for (int32 Index = 0; FSharedOnDemandContainer Container : Containers)
 	{
 		const TBitArray<>& IsReferenced = ChunkEntryIndices[Index++];
-		for (int32 EntryIndex = 0; const FOnDemandChunkEntry& Entry : Container->ChunkEntries)
+		for (int32 EntryIndex = 0; const FOnDemandChunkEntry & Entry : Container->ChunkEntries)
 		{
 			if (bool bIsReferenced = IsReferenced[EntryIndex++]; bIsReferenced == false)
 			{
@@ -1276,26 +1367,11 @@ FIoStatus FOnDemandInstallCache::Purge(TMap<FIoHash, uint64>&& ChunksToIntall)
 			}
 		}
 	}
+}
 
-	uint64 TotalReferencedBytes = 0;
-	for (const TPair<FCasBlockId, FCasBlockInfo>& Kv : BlockInfo)
-	{
-		if (Kv.Value.RefCount > 0)
-		{
-			TotalReferencedBytes += Kv.Value.FileSize; 
-		}
-	}
-
-	BlockInfo.ValueSort([](const FCasBlockInfo& LHS, const FCasBlockInfo& RHS)
-	{
-		return LHS.LastAccess < RHS.LastAccess;
-	});
-
-	UE_LOG(LogIoStoreOnDemand, Log, TEXT("Purging install cache, MaxCacheSize=%.2lf MiB, CacheSize=%.2lf MiB, UncachedSize=%.2lf MiB, ReferencedBytes=%.2lf MiB"),
-		ToMiB(MaxCacheSize), ToMiB(TotalCachedBytes), ToMiB(TotalUncachedBytes), ToMiB(TotalReferencedBytes));
-
-	const uint64	TotalBytesToPurge	= TotalRequiredBytes - MaxCacheSize;
-	uint64			TotalPurgedBytes	= 0;
+FIoStatus FOnDemandInstallCache::Purge(const FCasBlockInfoMap& BlockInfo, const uint64 TotalBytesToPurge, uint64& OutTotalPurgedBytes)
+{
+	OutTotalPurgedBytes = 0;
 
 	for (const TPair<FCasBlockId, FCasBlockInfo>& Kv : BlockInfo)
 	{
@@ -1313,6 +1389,8 @@ FIoStatus FOnDemandInstallCache::Purge(TMap<FIoHash, uint64>&& ChunksToIntall)
 			return Status;
 		}
 
+		OutTotalPurgedBytes += Info.FileSize;
+
 		for (const FCasAddr& Addr : RemovedChunks)
 		{
 			Transaction.ChunkLocation(FCasLocation::Invalid, Addr);
@@ -1323,22 +1401,13 @@ FIoStatus FOnDemandInstallCache::Purge(TMap<FIoHash, uint64>&& ChunksToIntall)
 		{
 			return Status;
 		}
-
-		TotalPurgedBytes += Info.FileSize;
-		if (TotalPurgedBytes >= TotalBytesToPurge)
+		
+		if (OutTotalPurgedBytes >= TotalBytesToPurge)
 		{
 			break;
 		}
 	}
 
-	UE_LOG(LogIoStoreOnDemand, Log, TEXT("Purged %.2lf MiB (%.2lf%%) from install cache"),
-		ToMiB(TotalPurgedBytes), 100.0 * (double(TotalPurgedBytes) / double(TotalCachedBytes)));
-
-	const uint64 NewCachedBytes = TotalCachedBytes - TotalPurgedBytes;
-	UE_CLOG(NewCachedBytes > MaxCacheSize,
-		LogIoStoreOnDemand, Warning, TEXT("Max install cache size exceeded by %.2lf MiB (%.2lf%%)"),
-			ToMiB(NewCachedBytes - MaxCacheSize), 100.0 * (double(NewCachedBytes - MaxCacheSize) / double(MaxCacheSize)));
-
 	return FIoStatus::Ok;
 }
 
diff --git a/Engine/Source/Runtime/Experimental/IoStore/OnDemand/Private/OnDemandInstallCache.h b/Engine/Source/Runtime/Experimental/IoStore/OnDemand/Private/OnDemandInstallCache.h
--- a/Engine/Source/Runtime/Experimental/IoStore/OnDemand/Private/OnDemandInstallCache.h
+++ b/Engine/Source/Runtime/Experimental/IoStore/OnDemand/Private/OnDemandInstallCache.h
@@ -22,6 +22,7 @@ public:
 	virtual bool		IsChunkCached(const FIoHash& ChunkHash) = 0;
 	virtual FIoStatus	PutChunk(FIoBuffer&& Chunk, const FIoHash& ChunkHash) = 0;
 	virtual FIoStatus	Purge(TMap<FIoHash, uint64>&& ChunksToInstall) = 0;
+	virtual FIoStatus	PurgeAllUnreferenced() = 0;
 	virtual FIoStatus	Flush() = 0;
 };
 
diff --git a/Engine/Source/Runtime/Experimental/IoStore/OnDemand/Private/OnDemandIoStore.cpp b/Engine/Source/Runtime/Experimental/IoStore/OnDemand/Private/OnDemandIoStore.cpp
--- a/Engine/Source/Runtime/Experimental/IoStore/OnDemand/Private/OnDemandIoStore.cpp
+++ b/Engine/Source/Runtime/Experimental/IoStore/OnDemand/Private/OnDemandIoStore.cpp
@@ -587,18 +587,12 @@ FIoStatus FOnDemandIoStore::Initialize()
 
 void FOnDemandIoStore::Mount(FOnDemandMountArgs&& Args, FOnDemandMountCompleted&& OnCompleted)
 {
-	if (Args.MountId.IsEmpty())
-	{
-		return OnCompleted(FOnDemandMountResult{ .Status = FIoStatus(EIoErrorCode::InvalidParameter, TEXT("Invalid mount ID")) });
-	}
+	FSharedMountRequest MountRequest = MakeShared<FMountRequest>();
+	MountRequest->Args = MoveTemp(Args);
+	MountRequest->OnCompleted = MoveTemp(OnCompleted);
 
+	UE_LOG(LogIoStoreOnDemand, Log, TEXT("Enqueing mount request, MountId='%s'"), *MountRequest->Args.MountId);
 	{
-		FSharedMountRequest MountRequest = MakeShared<FMountRequest>();
-		MountRequest->Args = MoveTemp(Args);
-		MountRequest->OnCompleted = MoveTemp(OnCompleted);
-
-		UE_LOG(LogIoStoreOnDemand, Log, TEXT("Enqueing mount request, MountId='%s'"), *MountRequest->Args.MountId);
-
 		UE::TUniqueLock Lock(MountRequestMutex);
 		MountRequests.Add(MoveTemp(MountRequest));
 	}
@@ -613,10 +607,10 @@ void FOnDemandIoStore::Install(
 	const FOnDemandCancellationToken* CancellationToken)
 {
 	FSharedInstallRequest InstallRequest = MakeShared<FInstallRequest>();
-	InstallRequest->Args		= MoveTemp(Args);
-	InstallRequest->OnCompleted = MoveTemp(OnCompleted);
-	InstallRequest->OnProgressed = MoveTemp(OnProgress);
-	InstallRequest->CancellationToken = CancellationToken;
+	InstallRequest->Args				= MoveTemp(Args);
+	InstallRequest->OnCompleted			= MoveTemp(OnCompleted);
+	InstallRequest->OnProgressed		= MoveTemp(OnProgress);
+	InstallRequest->CancellationToken	= CancellationToken;
 
 	{
 		UE::TUniqueLock Lock(MountRequestMutex);
@@ -626,6 +620,20 @@ void FOnDemandIoStore::Install(
 	TryEnterTickLoop();
 }
 
+void FOnDemandIoStore::Purge(FOnDemandPurgeArgs&& Args, FOnDemandPurgeCompleted&& OnCompleted)
+{
+	FSharedPurgeRequest PurgeRequest = MakeShared<FPurgeRequest>();
+	PurgeRequest->Args			= MoveTemp(Args);
+	PurgeRequest->OnCompleted	= MoveTemp(OnCompleted);
+
+	{
+		UE::TUniqueLock Lock(MountRequestMutex);
+		PurgeRequests.Add(MoveTemp(PurgeRequest));
+	}
+
+	TryEnterTickLoop();
+}
+
 FIoStatus FOnDemandIoStore::Unmount(FStringView MountId)
 {
 	UE_LOG(LogIoStoreOnDemand, Log, TEXT("Unmounting '%s'"), *WriteToString<256>(MountId));
@@ -969,6 +977,28 @@ bool FOnDemandIoStore::Tick()
 			});
 	}
 
+	TArray<FSharedPurgeRequest> LocalPurgeRequests;
+	{
+		UE::TUniqueLock Lock(MountRequestMutex);
+		LocalPurgeRequests = MoveTemp(PurgeRequests);
+		PurgeRequests.Empty();
+	}
+	
+	// Tick purge requests
+	if (LocalPurgeRequests.IsEmpty() == false)
+	{
+		FOnDemandPurgeResult PurgeResult;
+
+		const double StartTime = FPlatformTime::Seconds();
+		PurgeResult.Status = InstallCache->PurgeAllUnreferenced();
+		PurgeResult.DurationInSeconds = FPlatformTime::Seconds() - StartTime;
+
+		for (FSharedPurgeRequest& Request : LocalPurgeRequests)
+		{
+			CompletePurgeRequest(*Request, FOnDemandPurgeResult(PurgeResult));
+		}
+	}
+
 	TArray<FSharedInstallRequest> LocalInstallRequests;
 	{
 		UE::TUniqueLock Lock(MountRequestMutex);
@@ -1005,6 +1035,11 @@ FIoStatus FOnDemandIoStore::TickMountRequest(FMountRequest& MountRequest)
 
 	FOnDemandMountArgs& Args = MountRequest.Args;
 
+	if (Args.MountId.IsEmpty())
+	{
+		return FIoStatusBuilder(EIoErrorCode::InvalidParameter) << TEXT("Invalid mount ID");
+	}
+
 	bool bAnyPendingEncryptionKey = false;
 	bool bFoundContainers = false;
 
@@ -1553,6 +1588,29 @@ void FOnDemandIoStore::ProgressInstallRequest(const FInstallRequest& Request, co
 	}
 }
 
+void FOnDemandIoStore::CompletePurgeRequest(FPurgeRequest& Request, FOnDemandPurgeResult&& Result)
+{
+	if (!Request.OnCompleted)
+	{
+		return;
+	}
+
+	if (EnumHasAnyFlags(Request.Args.Options, EOnDemandInstallOptions::CallbackOnGameThread))
+	{
+		ExecuteOnGameThread(
+			UE_SOURCE_LOCATION,
+			[OnCompleted = MoveTemp(Request.OnCompleted), Result = MoveTemp(Result)]() mutable
+			{
+				OnCompleted(MoveTemp(Result));
+			});
+	}
+	else
+	{
+		FOnDemandPurgeCompleted OnCompleted = MoveTemp(Request.OnCompleted);
+		OnCompleted(MoveTemp(Result));
+	}
+}
+
 void FOnDemandIoStore::OnEncryptionKeyAdded(const FGuid& Id, const FAES::FAESKey& Key)
 {
 	TRACE_CPUPROFILER_EVENT_SCOPE(FOnDemandIoStore::OnEncryptionKeyAdded);
diff --git a/Engine/Source/Runtime/Experimental/IoStore/OnDemand/Private/OnDemandIoStore.h b/Engine/Source/Runtime/Experimental/IoStore/OnDemand/Private/OnDemandIoStore.h
--- a/Engine/Source/Runtime/Experimental/IoStore/OnDemand/Private/OnDemandIoStore.h
+++ b/Engine/Source/Runtime/Experimental/IoStore/OnDemand/Private/OnDemandIoStore.h
@@ -257,7 +257,7 @@ class FOnDemandIoStore
 		double					DurationInSeconds = 0.0;
 	};
 
-	using FSharedMountRequest	= TSharedPtr<FMountRequest>;
+	using FSharedMountRequest	= TSharedRef<FMountRequest>;
 
 	struct FInstallRequest
 	{
@@ -267,7 +267,15 @@ class FOnDemandIoStore
 		const FOnDemandCancellationToken*	CancellationToken = nullptr;
 	};
 
-	using FSharedInstallRequest	= TSharedPtr<FInstallRequest>;
+	using FSharedInstallRequest	= TSharedRef<FInstallRequest>;
+
+	struct FPurgeRequest
+	{
+		FOnDemandPurgeArgs					Args;
+		FOnDemandPurgeCompleted				OnCompleted;
+	};
+
+	using FSharedPurgeRequest = TSharedRef<FPurgeRequest>;
 
 public:
 	FOnDemandIoStore();
@@ -284,6 +292,7 @@ public:
 								FOnDemandInstallCompleted&& OnCompleted,
 								FOnDemandInstallProgressed&& OnProgress = nullptr,
 								const FOnDemandCancellationToken* CancellationToken = nullptr);
+	void					Purge(FOnDemandPurgeArgs&& Args, FOnDemandPurgeCompleted&& OnCompleted);
 	FIoStatus				Unmount(FStringView MountId);
 	TIoStatusOr<uint64>		GetInstallSize(const FOnDemandGetInstallSizeArgs& Args) const;
 	FIoStatus				GetInstallSizesByMountId(const FOnDemandGetInstallSizeArgs& Args, TMap<FString, uint64>& OutSizesByMountId) const;
@@ -314,6 +323,7 @@ private:
 	FOnDemandInstallResult	TickInstallRequest(const FInstallRequest& InstallRequest);
 	void					CompleteInstallRequest(FInstallRequest& InstallRequest, FOnDemandInstallResult&& InstallResult);
 	void					ProgressInstallRequest(const FInstallRequest& InstallRequest, const FOnDemandInstallProgress& Progress);
+	void					CompletePurgeRequest(FPurgeRequest& PurgeRequest, FOnDemandPurgeResult&& Result);
 	void					OnEncryptionKeyAdded(const FGuid& Id, const FAES::FAESKey& Key);
 	static void				CreateContainersFromToc(
 								FStringView MountId,
@@ -330,6 +340,7 @@ private:
 
 	TArray<FSharedMountRequest>			MountRequests;
 	TArray<FSharedInstallRequest>		InstallRequests;
+	TArray<FSharedPurgeRequest>			PurgeRequests;
 	UE::FMutex							MountRequestMutex;
 
 	bool								bTicking = false;
diff --git a/Engine/Source/Runtime/Experimental/IoStore/OnDemand/Public/IO/IoStoreOnDemand.h b/Engine/Source/Runtime/Experimental/IoStore/OnDemand/Public/IO/IoStoreOnDemand.h
--- a/Engine/Source/Runtime/Experimental/IoStore/OnDemand/Public/IO/IoStoreOnDemand.h
+++ b/Engine/Source/Runtime/Experimental/IoStore/OnDemand/Public/IO/IoStoreOnDemand.h
@@ -366,6 +366,25 @@ struct FOnDemandInstallResult
 /** Install completion callback. */
 using FOnDemandInstallCompleted = TUniqueFunction<void(FOnDemandInstallResult)>;
 
+/** Arguments for purging on-demand content. */
+struct FOnDemandPurgeArgs
+{
+	/** Purge options. */
+	EOnDemandInstallOptions Options = EOnDemandInstallOptions::None;
+};
+
+/** Holdes information about a purge request */
+struct FOnDemandPurgeResult
+{
+	/** The status of the purge request. */
+	FIoStatus Status;
+	/** Duration in seconds. */
+	double DurationInSeconds = 0.0;
+};
+
+/** Purge completion callback */
+using FOnDemandPurgeCompleted = TUniqueFunction<void(FOnDemandPurgeResult)>;
+
 /** Arguments for getting the size of on-demand content. */
 struct FOnDemandGetInstallSizeArgs 
 {
@@ -433,6 +452,8 @@ public:
 		FOnDemandInstallProgressed&& OnProgress = nullptr,
 		const FOnDemandCancellationToken* CancellationToken = nullptr);
 
+	UE_API void Purge(FOnDemandPurgeArgs&& Args, FOnDemandPurgeCompleted&& OnCompleted);
+
 	UE_API TIoStatusOr<uint64> GetInstallSize(const FOnDemandGetInstallSizeArgs& Args) const;
 
 	UE_API FIoStatus GetInstallSizesByMountId(const FOnDemandGetInstallSizeArgs& Args, TMap<FString, uint64>& OutSizesByMountId) const;
