[NetPhysics] Fix bug where applying state from FEvolutionResimCache would not update kinematic particles.

diff --git a/Engine/Source/Runtime/Experimental/Chaos/Private/Chaos/PBDRigidsEvolutionGBF.cpp b/Engine/Source/Runtime/Experimental/Chaos/Private/Chaos/PBDRigidsEvolutionGBF.cpp
--- a/Engine/Source/Runtime/Experimental/Chaos/Private/Chaos/PBDRigidsEvolutionGBF.cpp
+++ b/Engine/Source/Runtime/Experimental/Chaos/Private/Chaos/PBDRigidsEvolutionGBF.cpp
@@ -391,33 +391,70 @@ void FPBDRigidsEvolutionGBF::SetDebugDrawScene(const ChaosDD::Private::FChaosDDS
 
 void FPBDRigidsEvolutionGBF::ReloadParticlesCache()
 {
-	
-	// @todo(chaos): Parallelize
 	FEvolutionResimCache* ResimCache = GetCurrentStepResimCache();
 	if ((ResimCache != nullptr) && ResimCache->IsResimming())
 	{
-		for (int32 IslandIndex = 0; IslandIndex < GetIslandManager().GetNumIslands(); ++IslandIndex)
+		// Helper function to set particle state
+		auto ResimHelper = [this, ResimCache](auto Particle) -> bool
 		{
-			Private::FPBDIsland* Island = GetIslandManager().GetIsland(IslandIndex);
-			bool bIsUsingCache = false;
-			if (!Island->IsSleeping() && !Island->NeedsResim())
+			const bool bParticleAltered = ResimCache->ReloadParticlePostSolve(*Particle->Handle());
+
+			if (bParticleAltered)
 			{
-#if !(UE_BUILD_SHIPPING || UE_BUILD_TEST)
-				if (Chaos::FPhysicsSolverBase::IsNetworkPhysicsPredictionEnabled() && Chaos::FPhysicsSolverBase::CanDebugNetworkPhysicsPrediction())
+				TArray<const Private::FPBDIsland*> Islands = GetIslandManager().FindParticleIslands(Particle->Handle());
+				for (const Private::FPBDIsland* Island : Islands)
 				{
-					UE_LOG(LogChaos, Log, TEXT("Reloading Island[%d] cache for %d particles"), IslandIndex, Island->GetParticles().Num());
-				}
-#endif
-				for (Private::FPBDIslandParticle* IslandParticle : Island->GetParticles())
-				{
-					if (auto Rigid = IslandParticle->GetParticle()->CastToRigidParticle())
+					if (Island && !Island->IsUsingCache() && !Island->NeedsResim())
 					{
-						ResimCache->ReloadParticlePostSolve(*Rigid);
+						Private::FPBDIsland* IslandNonConst = GetIslandManager().GetIsland(Island->GetArrayIndex());
+						IslandNonConst->SetIsUsingCache(true);
 					}
 				}
-				bIsUsingCache = true;
 			}
-			Island->SetIsUsingCache(bIsUsingCache);
+
+			return bParticleAltered;
+		};
+
+		auto IterationHelper = [this, ResimHelper](FTransientGeometryParticleHandle& Particle)
+		{
+			if (Particle.ResimType() == EResimType::FullResim)
+			{
+				// Don't set state from ResimCache each tick for FullResim particles
+				return;
+			}
+
+			if (FTransientPBDRigidParticleHandle* Rigid = Particle.CastToRigidParticle())
+			{
+				if (ResimHelper(Rigid))
+				{
+					Particles.SetDynamicParticleSOA(Rigid->Handle()); // Update SOA views
+				}
+			}
+			else if (FTransientPBDRigidClusteredParticleHandle* Clustered = Particle.CastToClustered())
+			{
+				if (ResimHelper(Clustered))
+				{
+					Particles.SetClusteredParticleSOA(Clustered->Handle()); // Update SOA views
+				}
+			}
+		};
+
+		// Iterate over particles in the different active views
+		for (FTransientGeometryParticleHandle& Particle : Particles.GetActiveDynamicMovingKinematicParticlesView())
+		{
+			IterationHelper(Particle);
+		}
+		for (FTransientGeometryParticleHandle& Particle : Particles.GetActiveKinematicParticlesView())
+		{
+			IterationHelper(Particle);
+		}
+		for (FTransientGeometryParticleHandle& Particle : Particles.GetActiveMovingKinematicParticlesView())
+		{
+			IterationHelper(Particle);
+		}
+		for (FTransientGeometryParticleHandle& Particle : Particles.GetActiveParticlesView())
+		{
+			IterationHelper(Particle);
 		}
 	}
 }
diff --git a/Engine/Source/Runtime/Experimental/Chaos/Public/Chaos/EvolutionResimCache.h b/Engine/Source/Runtime/Experimental/Chaos/Public/Chaos/EvolutionResimCache.h
--- a/Engine/Source/Runtime/Experimental/Chaos/Public/Chaos/EvolutionResimCache.h
+++ b/Engine/Source/Runtime/Experimental/Chaos/Public/Chaos/EvolutionResimCache.h
@@ -50,17 +50,19 @@ namespace Chaos
 			Cache.W = Particle.GetW();
 		}
 
-		void ReloadParticlePostSolve(FPBDRigidParticleHandle& Particle) const
+		bool ReloadParticlePostSolve(FPBDRigidParticleHandle& Particle) const
 		{
 			//if this function is called it means the particle is in sync, which means we should have a cached value
 			const FPBDSolveCache* Cache = ParticleToCachedSolve.Find(Particle.UniqueIdx());
-			if (Cache)
+			if (Cache && (Particle.GetP() != Cache->P || Particle.GetQ() != Cache->Q || Particle.GetV() != Cache->V || Particle.GetW() != Cache->W))
 			{
 				Particle.SetP(Cache->P);
 				Particle.SetQ(Cache->Q);
 				Particle.SetV(Cache->V);
 				Particle.SetW(Cache->W);
+				return true;
 			}
+			return false;
 		}
 
 		void SaveConstraints(TArrayView<const FPBDCollisionConstraint* const> CollisionsArray)
