[JsonStringify] prevent null access when climbing up object hierarchy

diff --git a/Engine/Source/Runtime/Experimental/JsonObjectGraph/Private/JsonStringifyImpl.cpp b/Engine/Source/Runtime/Experimental/JsonObjectGraph/Private/JsonStringifyImpl.cpp
--- a/Engine/Source/Runtime/Experimental/JsonObjectGraph/Private/JsonStringifyImpl.cpp
+++ b/Engine/Source/Runtime/Experimental/JsonObjectGraph/Private/JsonStringifyImpl.cpp
@@ -231,7 +231,8 @@ static void WriteObjectPath(FUtf8StringBuilderBase& OutPath, const UObject* ForO
 {
 	const auto WriteObjectPathImpl = [&OutPath, OuterLimit](const UObject* ForObject, auto Self)
 		{
-			if (ForObject == OuterLimit)
+			if (ForObject == nullptr
+				|| ForObject == OuterLimit)
 			{
 				return;
 			}
