Iris
- NetTokens now has a concept of authority and can override local exported tokens with token received from authority so that will be used instead.
- Removed IsLocal from FNameAsNetTokenNetSerializer as it is no longer needed
- Cleaned up how we access type specific NetTokenDataStores as a step towards exposing NetTokens for use outside of iris.
- Refactored LastResortNetSerializer  and CharacterNetworkSerializationMovemenPackedBitsNetSerializer to support capturing and exporting FNames and NetTokens in addition to ObjectReferences.
    - Can be disabled with "net.iris.EnableIrisPackageMapNameExports 0"

diff --git a/Engine/Plugins/Runtime/ReplicationSystemTestPlugin/Source/Private/Tests/ReplicationSystem/TestReplicationSystem.cpp b/Engine/Plugins/Runtime/ReplicationSystemTestPlugin/Source/Private/Tests/ReplicationSystem/TestReplicationSystem.cpp
--- a/Engine/Plugins/Runtime/ReplicationSystemTestPlugin/Source/Private/Tests/ReplicationSystem/TestReplicationSystem.cpp
+++ b/Engine/Plugins/Runtime/ReplicationSystemTestPlugin/Source/Private/Tests/ReplicationSystem/TestReplicationSystem.cpp
@@ -492,6 +492,7 @@ public:
 	FStringTokenStore* ClientStringTokenStore = nullptr;
 	FReplicationSystemTestClient* Client = nullptr;
 	const FNetTokenStoreState* ClientRemoteNetTokenState;
+	const FNetTokenStoreState* ServerRemoteNetTokenState;
 
 	FNetToken CreateAndExportNetToken(const FString& TokenString)
 	{
@@ -501,20 +502,27 @@ public:
 		return Token;
 	}
 
+	FNetToken CreateAndExportNetTokenOnClient(const FString& TokenString)
+	{
+		FNetToken Token = ClientStringTokenStore->GetOrCreateToken(TokenString);
+		Client->GetConnectionInfo(Client->LocalConnectionId).NetTokenDataStream->AddNetTokenForExplicitExport(Token);
+
+		return Token;
+	}
+
 	virtual void SetUp() override
 	{
 		FReplicationSystemServerClientTestFixture::SetUp();
 
 		Client = CreateClient();
 
-		ServerStringTokenStore = Server->GetReplicationSystem()->GetStringTokenStore();
-		ClientStringTokenStore = Client->GetReplicationSystem()->GetStringTokenStore();
+		ServerStringTokenStore = Server->GetReplicationSystem()->GetNetTokenStore()->GetDataStore<FStringTokenStore>();
+		ServerRemoteNetTokenState = Server->GetConnectionInfo(Client->ConnectionIdOnServer).NetTokenDataStream->GetRemoteNetTokenStoreState();
+		ClientStringTokenStore = Client->GetReplicationSystem()->GetNetTokenStore()->GetDataStore<FStringTokenStore>();
 		ClientRemoteNetTokenState = Client->GetConnectionInfo(Client->LocalConnectionId).NetTokenDataStream->GetRemoteNetTokenStoreState();
 	}
 };
 
-
-
 UE_NET_TEST_FIXTURE(FTestNetTokensFixture, NetToken)
 {
 	// Create token
@@ -702,6 +710,62 @@ UE_NET_TEST_FIXTURE(FTestNetTokensFixture, NetTokenResendAndDataInSamePacketTest
 	UE_NET_ASSERT_EQ(TestStrings[1], FString(ClientStringTokenStore->ResolveRemoteToken(StringTokenB, *ClientRemoteNetTokenState)));
 }
 
+UE_NET_TEST_FIXTURE(FTestNetTokensFixture, NetTokenAuthority)
+{
+	// Create token
+	FString TokenStringA(TEXT("MyStringToken"));
+	FNetToken NonAuthToken = CreateAndExportNetTokenOnClient(TokenStringA);
+
+	UE_NET_ASSERT_EQ(NonAuthToken.IsAssignedByAuthority(), false);
+
+	// Send from server
+	Server->UpdateAndSend({Client});
+
+	// Send from client
+	Client->UpdateAndSend(Server);
+
+	// We should be able to resolve the token on the server using remote
+	UE_NET_ASSERT_EQ(TokenStringA, FString(ServerStringTokenStore->ResolveToken(NonAuthToken, ServerRemoteNetTokenState)));
+
+	// Find server token.
+	FNetToken AuthToken = CreateAndExportNetToken(TokenStringA);
+
+	// It should be a different token as the server is authoriative
+	UE_NET_ASSERT_FALSE(AuthToken == NonAuthToken);
+
+	// Send from server
+	Server->UpdateAndSend({Client});
+
+	// Client should be able to resolve ServerToken
+	UE_NET_ASSERT_EQ(TokenStringA, FString(ClientStringTokenStore->ResolveToken(AuthToken, ClientRemoteNetTokenState)));
+
+	// If we now try to create a token for the string also received from the authority we expect it to give us the server token and allow us to use that instead of the local exported token.
+	FNetToken NewClientToken = ClientStringTokenStore->GetOrCreateToken(TokenStringA);
+
+	// We expect the tokens to be identical.
+	UE_NET_ASSERT_TRUE(AuthToken == NewClientToken);
+}
+
+UE_NET_TEST_FIXTURE(FTestNetTokensFixture, NetTokenAuthTokenIsNotExportedFromClient)
+{
+	// Create token
+	FString TokenStringA(TEXT("MyStringToken"));
+	FNetToken AuthToken = CreateAndExportNetToken(TokenStringA);
+
+	UE_NET_ASSERT_EQ(AuthToken.IsAssignedByAuthority(), true);
+
+	// Send from server
+	Server->UpdateAndSend({Client});
+
+	// Expect to get auth token
+	FNetToken ClientExpectedAuthToken = CreateAndExportNetTokenOnClient(TokenStringA);
+	UE_NET_ASSERT_EQ(ClientExpectedAuthToken.IsAssignedByAuthority(), true);
+
+	// Send from client
+	Client->UpdateAndSend(Server);
+
+	// $TODO: Expose some stats that we can query for exports.
+}
 
 UE_NET_TEST_FIXTURE(FReplicationSystemServerClientTestFixture, AddRemoveFromConnectionScopeTest)
 {
diff --git a/Engine/Source/Runtime/Engine/Classes/GameFramework/CharacterMovementReplication.h b/Engine/Source/Runtime/Engine/Classes/GameFramework/CharacterMovementReplication.h
--- a/Engine/Source/Runtime/Engine/Classes/GameFramework/CharacterMovementReplication.h
+++ b/Engine/Source/Runtime/Engine/Classes/GameFramework/CharacterMovementReplication.h
@@ -59,8 +59,9 @@ struct FCharacterNetworkSerializationPackedBits
 	TBitArray<TInlineAllocator<CHARACTER_SERIALIZATION_PACKEDBITS_RESERVED_SIZE / NumBitsPerDWORD>> DataBits;
 
 #if UE_WITH_IRIS
-	// Since this struct uses custom serialization path we need to explicitly capture object references, this is managed by the use of a custom packagemap.
-	UIrisObjectReferencePackageMap::FObjectReferenceArray ObjectReferences;
+	// Since this struct uses custom serialization path we need to explicitly capture exports in order to forward them to iris
+	// This is managed by the use of a custom packagemap.
+	UE::Net::FIrisPackageMapExports PackageMapExports;
 #endif
 
 private:
diff --git a/Engine/Source/Runtime/Engine/Private/Components/CharacterMovementComponent.cpp b/Engine/Source/Runtime/Engine/Private/Components/CharacterMovementComponent.cpp
--- a/Engine/Source/Runtime/Engine/Private/Components/CharacterMovementComponent.cpp
+++ b/Engine/Source/Runtime/Engine/Private/Components/CharacterMovementComponent.cpp
@@ -31,6 +31,7 @@
 #if UE_WITH_IRIS
 #include "Net/Iris/ReplicationSystem/ReplicationSystemUtil.h"
 #include "Net/Iris/ReplicationSystem/ActorReplicationBridge.h"
+#include "Iris/ReplicationSystem/ReplicationSystem.h"
 #endif
 
 #include UE_INLINE_GENERATED_CPP_BY_NAME(CharacterMovementComponent)
@@ -8849,7 +8850,7 @@ void UCharacterMovementComponent::ReplicateMoveToServer(float DeltaTime, const F
 namespace UE::Private
 {
 
-static UIrisObjectReferencePackageMap* GetIrisPackageMapToCaptureReferences(UNetConnection* NetConnection, UIrisObjectReferencePackageMap::FObjectReferenceArray* InObjectReferences)
+static UIrisObjectReferencePackageMap* GetIrisPackageMapToCaptureReferences(UNetConnection* NetConnection, FCharacterNetworkSerializationPackedBits& PackedBits)
 {
 	using namespace UE::Net;
 
@@ -8857,7 +8858,12 @@ static UIrisObjectReferencePackageMap* GetIrisPackageMapToCaptureReferences(UNet
 	{
 		if (UIrisObjectReferencePackageMap* ObjectReferencePackageMap = Bridge->GetObjectReferencePackageMap())
 		{
-			ObjectReferencePackageMap->InitForWrite(InObjectReferences);
+			ObjectReferencePackageMap->InitForWrite(&PackedBits.PackageMapExports);
+
+			// $IRIS: $TODO: figure out how this should be done
+			//UReplicationSystem* ReplicationSystem = UE::Net::FReplicationSystemUtil::GetReplicationSystem(NetConnection->GetDriver());
+			//ReplicationSystem->InitNetTokenExportContext(*ObjectReferencePackageMap->GetNetTokenExportContext(), NetConnection->GetConnectionId());
+
 			return ObjectReferencePackageMap;
 		}
 	}
@@ -8865,14 +8871,19 @@ static UIrisObjectReferencePackageMap* GetIrisPackageMapToCaptureReferences(UNet
 	return nullptr;
 }
 
-static UIrisObjectReferencePackageMap* GetIrisPackageMapToReadReferences(const UNetConnection* NetConnection, const UIrisObjectReferencePackageMap::FObjectReferenceArray* InObjectReferences)
+static UIrisObjectReferencePackageMap* GetIrisPackageMapToReadReferences(const UNetConnection* NetConnection, const FCharacterNetworkSerializationPackedBits& PackedBits)
 {
 	using namespace UE::Net;
 	if (const UActorReplicationBridge* Bridge = FReplicationSystemUtil::GetActorReplicationBridge(NetConnection))
 	{
 		if (UIrisObjectReferencePackageMap* ObjectReferencePackageMap = Bridge->GetObjectReferencePackageMap())
 		{
-			ObjectReferencePackageMap->InitForRead(InObjectReferences);
+			ObjectReferencePackageMap->InitForRead(&PackedBits.PackageMapExports);
+
+			// $IRIS: $TODO: figure out how this should be done
+			//UReplicationSystem* ReplicationSystem = UE::Net::FReplicationSystemUtil::GetReplicationSystem(NetConnection->GetDriver());
+			//ReplicationSystem->InitNetTokenExportContext(*ObjectReferencePackageMap->GetNetTokenExportContext(), NetConnection->GetConnectionId());
+
 			return ObjectReferencePackageMap;
 		}
 	}
@@ -8898,7 +8909,7 @@ void UCharacterMovementComponent::CallServerMovePacked(const FSavedMove_Characte
 	UNetConnection* NetConnection = CharacterOwner->GetNetConnection();	
 
 #if UE_WITH_IRIS
-	if (UPackageMap* PackageMap = UE::Private::GetIrisPackageMapToCaptureReferences(NetConnection, &PackedBits.ObjectReferences))
+	if (UPackageMap* PackageMap = UE::Private::GetIrisPackageMapToCaptureReferences(NetConnection, PackedBits))
 	{
 		ServerMoveBitWriter.PackageMap = PackageMap;
 	}
@@ -9629,7 +9640,7 @@ void UCharacterMovementComponent::ServerMovePacked_ServerReceive(const FCharacte
 	ServerMoveBitReader.SetData((uint8*)PackedBits.DataBits.GetData(), NumBits);
 
 #if UE_WITH_IRIS
-	if (UPackageMap* PackageMap = UE::Private::GetIrisPackageMapToReadReferences(CharacterOwner->GetNetConnection(), &PackedBits.ObjectReferences))
+	if (UIrisObjectReferencePackageMap* PackageMap = UE::Private::GetIrisPackageMapToReadReferences(CharacterOwner->GetNetConnection(), PackedBits))
 	{
 		ServerMoveBitReader.PackageMap = PackageMap;
 	}
@@ -10492,7 +10503,7 @@ void UCharacterMovementComponent::MoveResponsePacked_ClientReceive(const FCharac
 	MoveResponseBitReader.SetData((uint8*)PackedBits.DataBits.GetData(), NumBits);
 
 #if UE_WITH_IRIS
-	if (UPackageMap* PackageMap = UE::Private::GetIrisPackageMapToReadReferences(CharacterOwner->GetNetConnection(), &PackedBits.ObjectReferences))
+	if (UPackageMap* PackageMap = UE::Private::GetIrisPackageMapToReadReferences(CharacterOwner->GetNetConnection(), PackedBits))
 	{
 		MoveResponseBitReader.PackageMap = PackageMap;
 	}
@@ -10537,7 +10548,7 @@ void UCharacterMovementComponent::ServerSendMoveResponse(const FClientAdjustment
 
 	// Extract the net package map used for serializing object references.
 #if UE_WITH_IRIS
-	if (UPackageMap* PackageMap = UE::Private::GetIrisPackageMapToCaptureReferences(NetConnection, &PackedBits.ObjectReferences))	
+	if (UPackageMap* PackageMap = UE::Private::GetIrisPackageMapToCaptureReferences(NetConnection, PackedBits))	
 	{
 		MoveResponseBitWriter.PackageMap = PackageMap;
 	}
diff --git a/Engine/Source/Runtime/Engine/Private/GameFramework/CharacterNetworkSerializationPackedBitsNetSerializer.cpp b/Engine/Source/Runtime/Engine/Private/GameFramework/CharacterNetworkSerializationPackedBitsNetSerializer.cpp
--- a/Engine/Source/Runtime/Engine/Private/GameFramework/CharacterNetworkSerializationPackedBitsNetSerializer.cpp
+++ b/Engine/Source/Runtime/Engine/Private/GameFramework/CharacterNetworkSerializationPackedBitsNetSerializer.cpp
@@ -17,6 +17,7 @@
 #include "GameFramework/CharacterMovementReplication.h"
 #include "Iris/Serialization/NetReferenceCollector.h"
 #include "Iris/Serialization/NetSerializerArrayStorage.h"
+#include "Iris/Serialization/IrisPackageMapExportUtil.h"
 
 namespace UE::Net::Private
 {
@@ -32,10 +33,9 @@ struct FCharacterNetworkSerializationPackedBitsNetSerializerQuantizedType
 	static constexpr uint32 MaxInlinedObjectRefs = 4;
 	static constexpr uint32 InlinedWordCount = Private::CalculateRequiredWordCount(CHARACTER_SERIALIZATION_PACKEDBITS_RESERVED_SIZE);
 
-	typedef FNetSerializerArrayStorage<FNetObjectReference, AllocationPolicies::TInlinedElementAllocationPolicy<MaxInlinedObjectRefs>> FObjectReferenceStorage;
 	typedef FNetSerializerArrayStorage<WordType, AllocationPolicies::TInlinedElementAllocationPolicy<InlinedWordCount>> FDataBitsStorage;
 
-	FObjectReferenceStorage ObjectReferenceStorage;
+	FIrisPackageMapExportsQuantizedType QuantizedExports;
 	FDataBitsStorage DataBitsStorage;
 	uint32 NumDataBits;
 };
@@ -91,7 +91,6 @@ private:
 	static void FreeDynamicState(FNetSerializationContext&, QuantizedType& Value);
 
 	static FCharacterNetworkSerializationPackedBitsNetSerializer::FNetSerializerRegistryDelegates NetSerializerRegistryDelegates;
-	static const FNetSerializer* ObjectNetSerializer;
 	inline static IConsoleVariable* CVarNetPackedMovementMaxBits = nullptr;
 };
 
@@ -99,7 +98,6 @@ UE_NET_IMPLEMENT_SERIALIZER(FCharacterNetworkSerializationPackedBitsNetSerialize
 
 const FCharacterNetworkSerializationPackedBitsNetSerializer::ConfigType FCharacterNetworkSerializationPackedBitsNetSerializer::DefaultConfig;
 FCharacterNetworkSerializationPackedBitsNetSerializer::FNetSerializerRegistryDelegates FCharacterNetworkSerializationPackedBitsNetSerializer::NetSerializerRegistryDelegates;
-const FNetSerializer* FCharacterNetworkSerializationPackedBitsNetSerializer::ObjectNetSerializer = &UE_NET_GET_SERIALIZER(FObjectNetSerializer);
 
 void FCharacterNetworkSerializationPackedBitsNetSerializer::Serialize(FNetSerializationContext& Context, const FNetSerializeArgs& Args)
 {
@@ -107,22 +105,8 @@ void FCharacterNetworkSerializationPackedBitsNetSerializer::Serialize(FNetSerial
 	
 	FNetBitStreamWriter* Writer = Context.GetBitStreamWriter();
 
-	const uint32 NumReferences = Value.ObjectReferenceStorage.Num();
-
-	// If we have any references, export them!
-	if (Writer->WriteBool(NumReferences != 0))
-	{
-		UE::Net::WritePackedUint32(Writer, NumReferences);		
-		FObjectNetSerializerConfig ObjectSerializerConfig;
-		for (const FNetObjectReference& Ref : MakeArrayView(Value.ObjectReferenceStorage.GetData(), NumReferences))
-		{
-			FNetSerializeArgs ObjectArgs;
-			ObjectArgs.NetSerializerConfig = &ObjectSerializerConfig;
-			ObjectArgs.Source = NetSerializerValuePointer(&Ref);
-
-			ObjectNetSerializer->Serialize(Context, ObjectArgs);
-		}
-	}
+	// Serialize captured references and exports
+	FIrisPackageMapExportsUtil::Serialize(Context, Value.QuantizedExports);
 
 	// Write data bits
 	const uint32 NumDataBits = Value.NumDataBits;
@@ -135,7 +119,9 @@ void FCharacterNetworkSerializationPackedBitsNetSerializer::Serialize(FNetSerial
 
 void FCharacterNetworkSerializationPackedBitsNetSerializer::FreeDynamicState(FNetSerializationContext& Context, QuantizedType& Value)
 {
-	Value.ObjectReferenceStorage.Free(Context);
+	// Free quantized state for captured reference and exports
+	FIrisPackageMapExportsUtil::FreeDynamicState(Context, Value.QuantizedExports);
+
 	Value.DataBitsStorage.Free(Context);
 	Value.NumDataBits = 0;
 }
@@ -147,34 +133,8 @@ void FCharacterNetworkSerializationPackedBitsNetSerializer::Deserialize(FNetSeri
 	
 	FNetBitStreamReader* Reader = Context.GetBitStreamReader();
 
-	// Read any object references
-	const bool bHasObjectReferences = Reader->ReadBool();
-	if (bHasObjectReferences)
-	{
-		const uint32 NumReferences = UE::Net::ReadPackedUint32(Reader);
-
-		if (NumReferences > Config->MaxAllowedObjectReferences)
-		{
-			Context.SetError(GNetError_ArraySizeTooLarge);
-			return;
-		}
-
-		TargetValue.ObjectReferenceStorage.AdjustSize(Context, NumReferences);
-
-		FObjectNetSerializerConfig ObjectSerializerConfig;
-		for (FNetObjectReference& Ref : MakeArrayView(TargetValue.ObjectReferenceStorage.GetData(), TargetValue.ObjectReferenceStorage.Num()))
-		{
-			FNetDeserializeArgs ObjectArgs;
-			ObjectArgs.NetSerializerConfig = &ObjectSerializerConfig;
-			ObjectArgs.Target = NetSerializerValuePointer(&Ref);
-
-			ObjectNetSerializer->Deserialize(Context, ObjectArgs);
-		}
-	}
-	else
-	{
-		TargetValue.ObjectReferenceStorage.Free(Context);
-	}
+	// Deserialize captured references and exports
+	FIrisPackageMapExportsUtil::Deserialize(Context, TargetValue.QuantizedExports);
 
 	const bool bHasDataBits = Reader->ReadBool();
 	if (bHasDataBits)
@@ -208,24 +168,8 @@ void FCharacterNetworkSerializationPackedBitsNetSerializer::Quantize(FNetSeriali
 	const SourceType& SourceValue = *reinterpret_cast<const SourceType*>(Args.Source);
 	QuantizedType& TargetValue = *reinterpret_cast<QuantizedType*>(Args.Target);
 
-	const uint32 NumObjectReferences = SourceValue.ObjectReferences.Num();
-	TargetValue.ObjectReferenceStorage.AdjustSize(Context, NumObjectReferences);
-
-	if (NumObjectReferences > 0)
-	{
-		FObjectNetSerializerConfig ObjectNetSerializerConfig;
-		const TObjectPtr<UObject>* SourceReferences = SourceValue.ObjectReferences.GetData();
-		FNetObjectReference* TargetReferences = TargetValue.ObjectReferenceStorage.GetData();
-		for (uint32 ReferenceIndex = 0; ReferenceIndex < NumObjectReferences; ++ReferenceIndex)
-		{
-			FNetQuantizeArgs ObjectArgs;
-			ObjectArgs.NetSerializerConfig = &ObjectNetSerializerConfig;
-			ObjectArgs.Source = NetSerializerValuePointer(SourceReferences + ReferenceIndex);
-			ObjectArgs.Target = NetSerializerValuePointer(TargetReferences + ReferenceIndex);
-
-			ObjectNetSerializer->Quantize(Context, ObjectArgs);
-		}
-	}
+	// Quantize captured references and exports
+	FIrisPackageMapExportsUtil::Quantize(Context, SourceValue.PackageMapExports, TargetValue.QuantizedExports);
 
 	uint32 NumDataBits = SourceValue.DataBits.Num();
 
@@ -253,22 +197,8 @@ void FCharacterNetworkSerializationPackedBitsNetSerializer::Dequantize(FNetSeria
 	const QuantizedType& Source = *reinterpret_cast<const QuantizedType*>(Args.Source);
 	SourceType& Target = *reinterpret_cast<SourceType*>(Args.Target);
 
-	// References
-	const uint32 NumObjectReferences = Source.ObjectReferenceStorage.Num();	
-	Target.ObjectReferences.SetNumUninitialized(NumObjectReferences);
-
-	FObjectNetSerializerConfig Config;
-	const FNetObjectReference* SourceReferences = Source.ObjectReferenceStorage.GetData();
-	TObjectPtr<UObject>* TargetReferences = Target.ObjectReferences.GetData();
-	for (uint32 ReferenceIndex = 0; ReferenceIndex < NumObjectReferences; ++ReferenceIndex)
-	{
-		FNetDequantizeArgs ObjectArgs;
-		ObjectArgs.NetSerializerConfig = &Config;
-		ObjectArgs.Source = NetSerializerValuePointer(SourceReferences + ReferenceIndex);
-		ObjectArgs.Target = NetSerializerValuePointer(TargetReferences + ReferenceIndex);
-
-		ObjectNetSerializer->Dequantize(Context, ObjectArgs);
-	}
+	// Dequantize captured references and exports and inject into target
+	FIrisPackageMapExportsUtil::Dequantize(Context, Source.QuantizedExports, Target.PackageMapExports);
 
 	// DataBits
 	Target.DataBits.SetNumUninitialized(Source.NumDataBits);
@@ -282,18 +212,19 @@ bool FCharacterNetworkSerializationPackedBitsNetSerializer::IsEqual(FNetSerializ
 		const QuantizedType& Value0 = *reinterpret_cast<const QuantizedType*>(Args.Source0);
 		const QuantizedType& Value1 = *reinterpret_cast<const QuantizedType*>(Args.Source1);
 
-		if (Value0.NumDataBits != Value1.NumDataBits || Value0.ObjectReferenceStorage.Num() != Value1.ObjectReferenceStorage.Num())
+		if (Value0.NumDataBits != Value1.NumDataBits)
 		{
 			return false;
 		}
 
-		const uint32 RequiredWords = Private::CalculateRequiredWordCount(Value0.NumDataBits);
-		if (RequiredWords > 0 && FMemory::Memcmp(Value0.DataBitsStorage.GetData(), Value1.DataBitsStorage.GetData(), sizeof(WordType) * RequiredWords) != 0)
+		// Compare references and exports
+		if (!FIrisPackageMapExportsUtil::IsEqual(Context, Value0.QuantizedExports, Value1.QuantizedExports))
 		{
 			return false;
 		}
 
-		if (Value0.ObjectReferenceStorage.Num() > 0 && FMemory::Memcmp(Value0.ObjectReferenceStorage.GetData(), Value1.ObjectReferenceStorage.GetData(), sizeof(FNetObjectReference) * Value0.ObjectReferenceStorage.Num()) != 0)
+		const uint32 RequiredWords = Private::CalculateRequiredWordCount(Value0.NumDataBits);
+		if (RequiredWords > 0 && FMemory::Memcmp(Value0.DataBitsStorage.GetData(), Value1.DataBitsStorage.GetData(), sizeof(WordType) * RequiredWords) != 0)
 		{
 			return false;
 		}
@@ -315,7 +246,7 @@ bool FCharacterNetworkSerializationPackedBitsNetSerializer::Validate(FNetSeriali
 
 	const uint32 MaxNumDataBits = CVarNetPackedMovementMaxBits ? static_cast<uint32>(CVarNetPackedMovementMaxBits->GetInt()) : Config->MaxAllowedDataBits;
 
-	if (SourceValue.ObjectReferenceStorage.Num() > Config->MaxAllowedObjectReferences || SourceValue.NumDataBits > MaxNumDataBits)
+	if (!FIrisPackageMapExportsUtil::Validate(Context, SourceValue.QuantizedExports))
 	{
 		return false;
 	}
@@ -327,7 +258,8 @@ void FCharacterNetworkSerializationPackedBitsNetSerializer::CloneDynamicState(FN
 	const QuantizedType& SourceValue = *reinterpret_cast<const QuantizedType*>(Args.Source);
 	QuantizedType& TargetValue = *reinterpret_cast<QuantizedType*>(Args.Target);
 
-	TargetValue.ObjectReferenceStorage.Clone(Context, SourceValue.ObjectReferenceStorage);
+	FIrisPackageMapExportsUtil::CloneDynamicState(Context, TargetValue.QuantizedExports, SourceValue.QuantizedExports);
+
 	TargetValue.DataBitsStorage.Clone(Context, SourceValue.DataBitsStorage);
 }
 
@@ -343,11 +275,7 @@ void FCharacterNetworkSerializationPackedBitsNetSerializer::CollectNetReferences
 	const QuantizedType& Value = *reinterpret_cast<const QuantizedType*>(Args.Source);
 	FNetReferenceCollector& Collector = *reinterpret_cast<UE::Net::FNetReferenceCollector*>(Args.Collector);
 
-	const FNetReferenceInfo ReferenceInfo(FNetReferenceInfo::EResolveType::ResolveOnClient);	
-	for (const FNetObjectReference& Ref : MakeArrayView(Value.ObjectReferenceStorage.GetData(), Value.ObjectReferenceStorage.Num()))
-	{
-		Collector.Add(ReferenceInfo, Ref, Args.ChangeMaskInfo);
-	}
+	FIrisPackageMapExportsUtil::CollectNetReferences(Context, Value.QuantizedExports, Args.ChangeMaskInfo, Collector);
 }
 
 static const FName PropertyNetSerializerRegistry_NAME_CharacterMoveResponsePackedBits("CharacterMoveResponsePackedBits");
diff --git a/Engine/Source/Runtime/Engine/Public/GameFramework/CharacterNetworkSerializationPackedBitsNetSerializer.h b/Engine/Source/Runtime/Engine/Public/GameFramework/CharacterNetworkSerializationPackedBitsNetSerializer.h
--- a/Engine/Source/Runtime/Engine/Public/GameFramework/CharacterNetworkSerializationPackedBitsNetSerializer.h
+++ b/Engine/Source/Runtime/Engine/Public/GameFramework/CharacterNetworkSerializationPackedBitsNetSerializer.h
@@ -13,10 +13,6 @@ struct FCharacterNetworkSerializationPackedBitsNetSerializerConfig : public FNet
 	// Value used to sanity check incoming data so that we do not over-allocate dynamic memory
 	UPROPERTY()
 	uint32 MaxAllowedDataBits = 8192; 
-
-	// Value used to sanity check incoming data so that we do not over-allocate dynamic memory
-	UPROPERTY()
-	uint32 MaxAllowedObjectReferences = 128; 
 };
 
 namespace UE::Net
diff --git a/Engine/Source/Runtime/Experimental/Iris/Core/Private/Iris/ReplicationSystem/NameTokenStore.cpp b/Engine/Source/Runtime/Experimental/Iris/Core/Private/Iris/ReplicationSystem/NameTokenStore.cpp
--- a/Engine/Source/Runtime/Experimental/Iris/Core/Private/Iris/ReplicationSystem/NameTokenStore.cpp
+++ b/Engine/Source/Runtime/Experimental/Iris/Core/Private/Iris/ReplicationSystem/NameTokenStore.cpp
@@ -34,10 +34,10 @@ FNetToken FNameTokenStore::GetOrCreateToken(FName Name)
 		}
 		else
 		{
-			const FNetToken NewToken = CreateToken(Key, *TokenStore.GetLocalNetTokenStoreState());
+			const FNetToken NewToken = CreateToken(TokenStore.IsAuthority() ? FNetToken::ENetTokenAuthority::Authority : FNetToken::ENetTokenAuthority::None, Key, *TokenStore.GetLocalNetTokenStoreState());
 			StoredTokens[NameIndex] = NewToken;
 
-			UE_LOG_FNAMETOKEN(TEXT("FNameTokenStore::GetOrCreateToken - Created new token %u for %s"), NewToken.GetIndex(), *Name.ToString());
+			UE_LOG_FNAMETOKEN(TEXT("FNameTokenStore::GetOrCreateToken - Created %s for %s"), *NewToken.ToString(), *Name.ToString());
 
 			return NewToken;
 		}
@@ -77,21 +77,19 @@ FNameTokenStore::FNameTokenStore(FNetTokenStore& InTokenStore)
 	StoredTokens.Add(FNetToken());
 
 	// Register
-	TokenStore.RegisterDataStore(this);
+	TokenStore.RegisterDataStore(this, FNameTokenStore::GetTokenStoreName());
 }
 
-FName FNameTokenStore::ResolveToken(FNetToken Token) const
+FName FNameTokenStore::ResolveToken(FNetToken Token, const FNetTokenStoreState* NetTokenStoreState) const
 {
-	return ResolveRemoteToken(Token, *TokenStore.GetLocalNetTokenStoreState());
-}
-
-FName FNameTokenStore::ResolveRemoteToken(FNetToken Token, const FNetTokenStoreState& NetTokenStoreState) const
-{
-	const FNetTokenStoreKey StoreKey = GetTokenKey(Token, NetTokenStoreState);
-
-	if (StoreKey.IsValid() && StoreKey.GetKeyValue() < (uint32)StoredFNames.Num())
+	const FNetTokenStoreState* TokenStoreState = TokenStore.IsLocalToken(Token) ? TokenStore.GetLocalNetTokenStoreState() : NetTokenStoreState;
+	if (Token.IsValid() && ensureMsgf(TokenStoreState, TEXT("FNameTokenStore::ResolveToken Needs valid remote NetTokenStoreState to resolve remote %s"), *Token.ToString()))
 	{
-		return StoredFNames[StoreKey.GetKeyValue()];
+		const FNetTokenStoreKey StoreKey = GetTokenKey(Token, *TokenStoreState);
+		if (StoreKey.IsValid() && StoreKey.GetKeyValue() < (uint32)StoredFNames.Num())
+		{
+			return StoredFNames[StoreKey.GetKeyValue()];
+		}
 	}
 
 	return FName();
@@ -99,13 +97,15 @@ FName FNameTokenStore::ResolveRemoteToken(FNetToken Token, const FNetTokenStoreS
 
 void FNameTokenStore::WriteTokenData(FNetSerializationContext& Context, FNetTokenStoreKey TokenStoreKey) const
 {
-	//UE_NET_TRACE_DYNAMIC_NAME_SCOPE(StoredFNames[TokenStoreKey.GetKeyValue()], *Context.GetBitStreamWriter(), Context.GetTraceCollector(), ENetTraceVerbosity::VeryVerbose);
+	UE_NET_TRACE_DYNAMIC_NAME_SCOPE(StoredFNames[TokenStoreKey.GetKeyValue()], *Context.GetBitStreamWriter(), Context.GetTraceCollector(), ENetTraceVerbosity::VeryVerbose);
 	// $TODO: $IRIS: We can be a bit smarter here and utilize the string-number split of FNames to export less data.. JIRA: UE-221753
 	WriteString(Context.GetBitStreamWriter(), StoredFNames[TokenStoreKey.GetKeyValue()].ToString());
 }
 
-FNetTokenStoreKey FNameTokenStore::ReadTokenData(FNetSerializationContext& Context)
+FNetTokenStoreKey FNameTokenStore::ReadTokenData(FNetSerializationContext& Context, const FNetToken& NetToken)
 {
+	UE_NET_TRACE_NAMED_DYNAMIC_NAME_SCOPE(TokenScope, FName(), *Context.GetBitStreamReader(), Context.GetTraceCollector(), ENetTraceVerbosity::VeryVerbose);
+
 	FNetBitStreamReader* Reader = Context.GetBitStreamReader();
 
 	// Read the token data and add it to the store without assigning LocalToken
@@ -115,7 +115,20 @@ FNetTokenStoreKey FNameTokenStore::ReadTokenData(FNetSerializationContext& Conte
 	if (!Reader->IsOverflown())
 	{
 		FName Name(Temp);
-		return GetOrCreateTokenStoreKey(Name);
+		UE_NET_TRACE_SET_SCOPE_NAME(TokenScope, Name);
+
+		FNetTokenStoreKey Key = GetOrCreateTokenStoreKey(Name);
+
+		// If this is an authtoken and we are not the authority, update the stored key
+		if (Key.IsValid() && NetToken.IsAssignedByAuthority() && !TokenStore.IsAuthority())
+		{
+			UE_LOG_FNAMETOKEN(TEXT("NameTokenStore::ReadTokenData - Replaced local key %s with %s for name %s"), *StoredTokens[Key.GetKeyValue()].ToString(), *NetToken.ToString(), *Temp);	
+
+			// This way the next time we lookup this key during assignment use the imported authoritative key which we do not have to export
+			StoredTokens[Key.GetKeyValue()] = NetToken;
+		}
+
+		return Key;
 	}
 	else
 	{
diff --git a/Engine/Source/Runtime/Experimental/Iris/Core/Private/Iris/ReplicationSystem/NetTokenDataStream.cpp b/Engine/Source/Runtime/Experimental/Iris/Core/Private/Iris/ReplicationSystem/NetTokenDataStream.cpp
--- a/Engine/Source/Runtime/Experimental/Iris/Core/Private/Iris/ReplicationSystem/NetTokenDataStream.cpp
+++ b/Engine/Source/Runtime/Experimental/Iris/Core/Private/Iris/ReplicationSystem/NetTokenDataStream.cpp
@@ -74,7 +74,7 @@ void UNetTokenDataStream::Init(const FInitParameters& Params)
 		NetTokensPendingExport.Reserve(LocalTokenCount - 1);
 		for (uint32 Index = 1; Index < LocalTokenCount; ++Index)
 		{
-			NetTokensPendingExport.Add(FNetToken::MakeNetToken(Index));
+			NetTokensPendingExport.Add(FNetToken::MakeNetToken(Index, NetTokenStore->IsAuthority() ? FNetToken::ENetTokenAuthority::Authority : FNetToken::ENetTokenAuthority::None));
 		}
 	}
 }
@@ -103,8 +103,7 @@ UDataStream::EWriteResult UNetTokenDataStream::WriteData(UE::Net::FNetSerializat
 		return EWriteResult::NoData;
 	}
 
-	UReplicationSystem* ReplicationSystem = UE::Net::GetReplicationSystem(ReplicationSystemId);
-	const FStringTokenStore* StringTokenStore = ReplicationSystem->GetStringTokenStore();
+	const FStringTokenStore* StringTokenStore = NetTokenStore->GetDataStore<const FStringTokenStore>();
 	FNetExportContext* ExportContext = Context.GetExportContext();
 
 	// Write data until we have no more data to write or it does not fit
@@ -116,6 +115,7 @@ UDataStream::EWriteResult UNetTokenDataStream::WriteData(UE::Net::FNetSerializat
 	FNetBitStreamWriter SubStream = Writer->CreateSubstream(Writer->GetBitsLeft() - 1U);
 	FNetSerializationContext SubContext = Context.MakeSubContext(&SubStream);
 
+	const bool bIsNetTokenAuthority = NetTokenStore->IsAuthority();
 	while (NetTokensPendingExport.Num())
 	{
 		FNetBitStreamRollbackScope SequenceRollback(SubStream);
@@ -124,12 +124,13 @@ UDataStream::EWriteResult UNetTokenDataStream::WriteData(UE::Net::FNetSerializat
 		// Peek at front index
 		const FNetToken& Token = NetTokensPendingExport.GetAtIndexNoCheck(0);
 
-		if (!ExportContext->IsExported(Token))
+		// We do not need to export tokens assigned by authority unless we are the authority.
+		if (!(Token.IsAssignedByAuthority() && !bIsNetTokenAuthority) && !ExportContext->IsExported(Token))
 		{
 			UE_NET_TRACE_NAMED_SCOPE(ExportScope, NetTokenExport, SubStream, SubContext.GetTraceCollector(), ENetTraceVerbosity::Verbose);
 
 			SubStream.WriteBool(true);
-			WriteNetToken(&SubStream, Token);
+			WriteNetToken(SubContext, Token);
 			NetTokenStore->WriteTokenData(SubContext, Token);
 
 			if (SubStream.IsOverflown())
@@ -199,7 +200,7 @@ void UNetTokenDataStream::ReadData(UE::Net::FNetSerializationContext& Context)
 			break;
 		}
 
-		FNetToken Token = ReadNetToken(Reader);
+		FNetToken Token = ReadNetToken(Context);
 		if (Token.IsValid())
 		{
 			NetTokenStore->ReadTokenData(Context, Token, *RemoteNetTokenStoreState);
diff --git a/Engine/Source/Runtime/Experimental/Iris/Core/Private/Iris/ReplicationSystem/NetTokenStore.cpp b/Engine/Source/Runtime/Experimental/Iris/Core/Private/Iris/ReplicationSystem/NetTokenStore.cpp
--- a/Engine/Source/Runtime/Experimental/Iris/Core/Private/Iris/ReplicationSystem/NetTokenStore.cpp
+++ b/Engine/Source/Runtime/Experimental/Iris/Core/Private/Iris/ReplicationSystem/NetTokenStore.cpp
@@ -9,6 +9,7 @@
 #include "Iris/Serialization/NetSerializationContext.h"
 #include "Iris/Serialization/InternalNetSerializationContext.h"
 #include "Iris/Serialization/NetExportContext.h"
+#include "Net/Core/Trace/NetTrace.h"
 
 namespace UE::Net
 {
@@ -36,7 +37,7 @@ FNetTokenStoreKey FNetTokenDataStore::GetTokenKey(FNetToken Token, const FNetTok
 	return FNetTokenStoreKey();
 }
 
-FNetToken FNetTokenDataStore::CreateToken(FNetTokenStoreKey Key, FNetTokenStoreState& TokenStoreState)
+FNetToken FNetTokenDataStore::CreateToken(FNetToken::ENetTokenAuthority Authority, FNetTokenStoreKey Key, FNetTokenStoreState& TokenStoreState)
 {
 	const uint32 NextTokenIndex = TokenStoreState.TokenInfos.Num();
 
@@ -48,26 +49,35 @@ FNetToken FNetTokenDataStore::CreateToken(FNetTokenStoreKey Key, FNetTokenStoreS
 	// Store token info
 	TokenStoreState.TokenInfos.Add(Key);
 
-	return FNetToken::MakeNetToken(NextTokenIndex);
+	return FNetToken::MakeNetToken(NextTokenIndex, Authority);
 }
 
-FNetToken ReadNetToken(FNetBitStreamReader* Reader)
+FNetToken ReadNetToken(UE::Net::FNetSerializationContext& Context)
 {
-	uint32 Index = ReadPackedUint32(Reader);
+	UE_NET_TRACE_NAMED_DYNAMIC_NAME_SCOPE(TokenScope, FName(), *Context.GetBitStreamReader(), Context.GetTraceCollector(), ENetTraceVerbosity::VeryVerbose);
+	FNetBitStreamReader* Reader = Context.GetBitStreamReader();
+
+	const uint32 Index = ReadPackedUint32(Reader);
+	const bool bIsAssignedByAuthority = Reader->ReadBool();
+
+	FNetToken ReadToken;
 	
 	if (!Reader->IsOverflown())
 	{
-		return FNetToken(Index);
-	}
-	else
-	{
-		return FNetToken();
+		ReadToken = FNetToken::MakeNetToken(Index, bIsAssignedByAuthority ? FNetToken::ENetTokenAuthority::Authority : FNetToken::ENetTokenAuthority::None);
+		UE_NET_TRACE_SET_SCOPE_NAME(TokenScope, *ReadToken.ToString());
 	}
+
+	return ReadToken;
 }
 
-void WriteNetToken(FNetBitStreamWriter* Writer, FNetToken Token)
+void WriteNetToken(UE::Net::FNetSerializationContext& Context, FNetToken Token)
 {
+	UE_NET_TRACE_DYNAMIC_NAME_SCOPE(*Token.ToString(), *Context.GetBitStreamWriter(), Context.GetTraceCollector(), ENetTraceVerbosity::VeryVerbose);
+	FNetBitStreamWriter* Writer = Context.GetBitStreamWriter();
+
 	WritePackedUint32(Writer, Token.GetIndex());
+	Writer->WriteBool(Token.IsAssignedByAuthority());
 }
 
 FNetTokenStore::FNetTokenStore()
@@ -80,7 +90,24 @@ FNetTokenStore::~FNetTokenStore()
 	delete LocalNetTokenStoreState;
 }
 
-bool FNetTokenStore::RegisterDataStore(FNetTokenDataStore* DataStore)
+void FNetTokenStore::Init(FNetTokenStore::FInitParams& InParams)
+{
+	Params = InParams;
+}
+
+const FNetTokenDataStore* FNetTokenStore::GetDataStore(FName Name) const
+{
+	const TTuple<FName, FNetTokenDataStore*>* Entry = TokenDataStores.FindByPredicate([Name](const TTuple<FName, FNetTokenDataStore*>& Entry){ return Name == Entry.Get<0>(); });
+	return Entry ? Entry->Get<1>() : nullptr;
+}
+
+ FNetTokenDataStore* FNetTokenStore::GetDataStore(FName Name)
+{
+	TTuple<FName, FNetTokenDataStore*>* Entry = TokenDataStores.FindByPredicate([Name](TTuple<FName, FNetTokenDataStore*>& Entry){ return Name == Entry.Get<0>(); });
+	return Entry ? Entry->Get<1>() : nullptr;
+}
+
+bool FNetTokenStore::RegisterDataStore(FNetTokenDataStore* DataStore, FName TokenStoreName)
 {
 	if (TokenDataStores.Num() >= FNetToken::MaxTypeIdCount)
 	{
@@ -92,6 +119,12 @@ bool FNetTokenStore::RegisterDataStore(FNetTokenDataStore* DataStore)
 		return false;
 	}
 
+	if (!ensure(GetDataStore(TokenStoreName) == nullptr))
+	{
+		// Already registered
+		return false;
+	}
+
 	if (!ensure(DataStore->GetTypeId() == FNetToken::InvalidTokenTypeId))
 	{
 		// Already registered
@@ -99,85 +132,82 @@ bool FNetTokenStore::RegisterDataStore(FNetTokenDataStore* DataStore)
 	}
 
 	DataStore->TypeId = TokenDataStores.Num();
-	TokenDataStores.Add(DataStore);
+	TokenDataStores.Emplace(TokenStoreName, DataStore);
 
 	return true;
 }
 
-void FNetTokenStore::WriteTokenData(FNetSerializationContext& Context, const FNetToken& NetToken) const
+void FNetTokenStore::WriteTokenData(FNetSerializationContext& Context, const FNetToken NetToken) const
 {
 	if (NetToken.IsValid())
 	{
-		WriteTokenDataForIndex(Context, NetToken.GetIndex());
+		FNetBitStreamWriter* Writer = Context.GetBitStreamWriter();
+
+		// Write type
+		const FNetTokenStoreKey& TokenKey = LocalNetTokenStoreState->TokenInfos[NetToken.GetIndex()];
+		Writer->WriteBits(TokenKey.TypeId, FNetToken::TokenTypeIdBits);
+
+		// Write token data
+		TokenDataStores[TokenKey.TypeId].Get<1>()->WriteTokenData(Context, TokenKey);
 	}
 }
 
-void FNetTokenStore::ReadTokenData(FNetSerializationContext& Context, const FNetToken& NetToken, FNetTokenStoreState& RemoteNetTokenStoreState)
+void FNetTokenStore::ReadTokenData(FNetSerializationContext& Context, const FNetToken NetToken, FNetTokenStoreState& RemoteNetTokenStoreState)
 {
 	if (NetToken.IsValid())
 	{
+		FNetBitStreamReader* Reader = Context.GetBitStreamReader();
+
 		// TODO: Guard this better, a map might be better after-all!
 		// Also need to protect this if we process inbound data in parallel.
 		if (!RemoteNetTokenStoreState.ReserveTokenCount(NetToken.GetIndex() + 1))
 		{
-			Context.GetBitStreamReader()->DoOverflow();
+			Reader->DoOverflow();
 			return;
 		}
 
-		ReadTokenDataForIndex(Context, NetToken.GetIndex(), RemoteNetTokenStoreState);
-	}
-}
+		// Read type
+		FNetToken::FTypeId TokenTypeId = (FNetToken::FTypeId)Reader->ReadBits(FNetToken::TokenTypeIdBits);
 
-void FNetTokenStore::WriteTokenDataForIndex(FNetSerializationContext& Context, uint32 TokenIndex) const
-{
-	FNetBitStreamWriter* Writer = Context.GetBitStreamWriter();
-
-	// Write type
-	const FNetTokenStoreKey& TokenKey = LocalNetTokenStoreState->TokenInfos[TokenIndex];
-	Writer->WriteBits(TokenKey.TypeId, FNetToken::TokenTypeIdBits);
-
-	// Write token data
-	TokenDataStores[TokenKey.TypeId]->WriteTokenData(Context, TokenKey);
-}
+		// Validate that we managed to read and verify the type
+		if (Reader->IsOverflown() || TokenTypeId >= (uint32)TokenDataStores.Num() )
+		{
+			Reader->DoOverflow();
+			return;
+		}
 
-void FNetTokenStore::ReadTokenDataForIndex(FNetSerializationContext& Context, uint32 TokenIndex, FNetTokenStoreState& TokenStoreState)
-{
-	FNetBitStreamReader* Reader = Context.GetBitStreamReader();
+		const FNetTokenStoreKey StoreKey = TokenDataStores[TokenTypeId].Get<1>()->ReadTokenData(Context, NetToken);
 
-	// Read type
-	FNetToken::FTypeId TokenTypeId = (FNetToken::FTypeId)Reader->ReadBits(FNetToken::TokenTypeIdBits);
+		// Validate that we managed to read the data for the key
+		if (Reader->IsOverflown() || !StoreKey.IsValid())
+		{
+			Reader->DoOverflow();
+			return;
+		}
 
-	// Validate that we managed to read and verify the type
-	if (Reader->IsOverflown() || TokenTypeId >= (uint32)TokenDataStores.Num() )
-	{
-		Reader->DoOverflow();
-		return;
+		// Since the same tokendata might be exported multiple times we better validate that it is the same data
+		const FNetTokenStoreKey& ExistingStoreKey = RemoteNetTokenStoreState.TokenInfos[NetToken.GetIndex()];
+		if (!ensureAlways(!ExistingStoreKey.IsValid() || StoreKey == ExistingStoreKey))
+		{
+			Reader->DoOverflow();
+			return;
+		}
+	
+		// Store
+		RemoteNetTokenStoreState.TokenInfos[NetToken.GetIndex()] = StoreKey;
 	}
+}
 
-	const FNetTokenStoreKey StoreKey = TokenDataStores[TokenTypeId]->ReadTokenData(Context);
-
-	// Validate that we managed to read the data for the key
-	if (Reader->IsOverflown() || !StoreKey.IsValid())
-	{
-		Reader->DoOverflow();
-		return;
-	}
+void FNetTokenStore::ConditionalWriteNetTokenData(FNetSerializationContext& Context, Private::FNetExportContext* ExportContext, const FNetToken NetToken) const
+{
+	FNetBitStreamWriter* Writer = Context.GetBitStreamWriter();
 
-	// Since the same tokendata might be exported multiple times we better validate that it is the same data
-	const FNetTokenStoreKey& ExistingStoreKey = TokenStoreState.TokenInfos[TokenIndex];
-	if (!ensureAlways(!ExistingStoreKey.IsValid() || StoreKey == ExistingStoreKey))
+	// We should not try to export tokens received from remote
+	if (IsAuthority() != NetToken.IsAssignedByAuthority())
 	{
-		Reader->DoOverflow();
+		Writer->WriteBool(false);
 		return;
 	}
-	
-	// Store
-	TokenStoreState.TokenInfos[TokenIndex] = StoreKey;
-}
-
-void FNetTokenStore::ConditionalWriteNetTokenData(FNetSerializationContext& Context, Private::FNetExportContext* ExportContext, const FNetToken& NetToken) const
-{
-	FNetBitStreamWriter* Writer = Context.GetBitStreamWriter();
 
 	if (ExportContext)
 	{
@@ -194,7 +224,7 @@ void FNetTokenStore::ConditionalWriteNetTokenData(FNetSerializationContext& Cont
 	}
 }
 
-void FNetTokenStore::ConditionalReadNetTokenData(FNetSerializationContext& Context, const FNetToken& NetToken)
+void FNetTokenStore::ConditionalReadNetTokenData(FNetSerializationContext& Context, const FNetToken NetToken)
 {
 	FNetBitStreamReader* Reader = Context.GetBitStreamReader();
 
@@ -212,4 +242,36 @@ void FNetTokenStore::ConditionalReadNetTokenData(FNetSerializationContext& Conte
 	}
 }
 
+void FNetTokenStore::AppendExportOrWriteInlinedExportData(FNetSerializationContext& Context, FNetToken NetToken)
+{
+	using namespace UE::Net::Private;
+
+	FNetExportContext* ExportContext = Context.GetExportContext();
+	const FInternalNetSerializationContext* InternalContext = Context.GetInternalContext();
+	
+	if (InternalContext->bInlineObjectReferenceExports == 0U)
+	{
+		ExportContext->AddPendingExport(NetToken);
+	}
+	else
+	{
+		FNetTokenStore* NetTokenStore = Context.GetNetTokenStore();
+		NetTokenStore->ConditionalWriteNetTokenData(Context, ExportContext, NetToken);
+	}
+}
+
+void FNetTokenStore::ReadInlinedExportData(FNetSerializationContext& Context, FNetToken NetToken)
+{
+	using namespace UE::Net::Private;
+
+	const FInternalNetSerializationContext* InternalContext = Context.GetInternalContext();
+
+	// Conditionally import if using inlined exports
+	if (InternalContext->bInlineObjectReferenceExports == 1U)
+	{
+		FNetTokenStore* NetTokenStore = Context.GetNetTokenStore();
+		NetTokenStore->ConditionalReadNetTokenData(Context, NetToken);
+	}
+}
+
 }
diff --git a/Engine/Source/Runtime/Experimental/Iris/Core/Private/Iris/ReplicationSystem/ObjectReferenceCache.cpp b/Engine/Source/Runtime/Experimental/Iris/Core/Private/Iris/ReplicationSystem/ObjectReferenceCache.cpp
--- a/Engine/Source/Runtime/Experimental/Iris/Core/Private/Iris/ReplicationSystem/ObjectReferenceCache.cpp
+++ b/Engine/Source/Runtime/Experimental/Iris/Core/Private/Iris/ReplicationSystem/ObjectReferenceCache.cpp
@@ -94,7 +94,7 @@ void FObjectReferenceCache::Init(UReplicationSystem* InReplicationSystem)
 	ReplicationSystem = InReplicationSystem;
 	ReplicationBridge = InReplicationSystem->GetReplicationBridgeAs<UObjectReplicationBridge>();
 	NetTokenStore =  &InReplicationSystem->GetReplicationSystemInternal()->GetNetTokenStore();
-	StringTokenStore = InReplicationSystem->GetStringTokenStore();
+	StringTokenStore = NetTokenStore->GetDataStore<FStringTokenStore>();
 	NetRefHandleManager = &InReplicationSystem->GetReplicationSystemInternal()->GetNetRefHandleManager();
 	bIsAuthority = InReplicationSystem->IsServer();
 }
@@ -1233,7 +1233,7 @@ void FObjectReferenceCache::WriteFullReferenceInternal(FNetSerializationContext&
 
 		if (Writer->WriteBool(bHasPath))
 		{
-			WriteNetToken(Writer, CachedObject->RelativePath);
+			WriteNetToken(Context, CachedObject->RelativePath);
 			NetTokenStore->ConditionalWriteNetTokenData(Context, ExportContext, CachedObject->RelativePath);
 			WriteFullReferenceInternal(Context, FNetObjectReference(CachedObject->OuterNetRefHandle));
 		}
@@ -1300,7 +1300,7 @@ void FObjectReferenceCache::ReadFullReferenceInternal(FNetSerializationContext&
 
 		if (Reader->ReadBool())
 		{
-			RelativePath = ReadNetToken(Reader);
+			RelativePath = ReadNetToken(Context);
 			NetTokenStore->ConditionalReadNetTokenData(Context, RelativePath);
 			ReadFullReferenceInternal(Context, OuterRef, RecursionCount + 1U);
 		}
@@ -1386,7 +1386,7 @@ void FObjectReferenceCache::ReadFullReference(FNetSerializationContext& Context,
 		UE_NET_TRACE_SCOPE(ClientAssigned, *Reader, Context.GetTraceCollector(), ENetTraceVerbosity::Verbose);
 
 		const FNetRefHandle NetRefHandle = ReadNetRefHandle(Context);
-		FNetToken RelativePath = ReadNetToken(Reader);
+		FNetToken RelativePath = ReadNetToken(Context);
 		NetTokenStore->ConditionalReadNetTokenData(Context, RelativePath);
 
 		OutRef.RefHandle = NetRefHandle;
@@ -1425,7 +1425,7 @@ void FObjectReferenceCache::WriteFullReference(FNetSerializationContext& Context
 		UE_NET_TRACE_SCOPE(ClientAssigned, *Writer, Context.GetTraceCollector(), ENetTraceVerbosity::Verbose);
 
 		WriteNetRefHandle(Context, Ref.GetRefHandle());
-		WriteNetToken(Writer, Ref.PathToken);
+		WriteNetToken(Context, Ref.PathToken);
 		NetTokenStore->ConditionalWriteNetTokenData(Context, Context.GetExportContext(), Ref.PathToken);
 		return;
 	}
@@ -1457,7 +1457,7 @@ void FObjectReferenceCache::WriteReference(FNetSerializationContext& Context, FN
 	if (Writer->WriteBool(bIsClientAssignedReference))
 	{
 		WriteNetRefHandle(Context, Ref.GetRefHandle());
-		WriteNetToken(Writer, Ref.PathToken);
+		WriteNetToken(Context, Ref.PathToken);
 
 		return;
 	}
@@ -1493,7 +1493,7 @@ void FObjectReferenceCache::ReadReference(FNetSerializationContext& Context, FNe
 	if (const bool bIsClientAssignedReference = Reader->ReadBool())
 	{
 		const FNetRefHandle NetRefHandle = ReadNetRefHandle(Context);
-		FNetToken RelativePath = ReadNetToken(Reader);
+		FNetToken RelativePath = ReadNetToken(Context);
 
 		OutRef.RefHandle = NetRefHandle;
 		OutRef.PathToken = RelativePath;
@@ -1579,25 +1579,25 @@ FObjectReferenceCache::EWriteExportsResult FObjectReferenceCache::WritePendingEx
 		FForceInlineExportScope ForceInlineExportScope(Context.GetInternalContext());
 
 		// For now we export NetTokens here
+		TArrayView<const FNetToken> NetTokenExportsView = MakeArrayView(ExportContext->GetBatchExports().NetTokensPendingExportInCurrentBatch);
 		{
 			UE_NET_TRACE_SCOPE(NetTokenExports, Writer, Context.GetTraceCollector(), ENetTraceVerbosity::Verbose);
-
-			TArrayView<const FNetToken> NetTokenExportsView = MakeArrayView(ExportContext->GetBatchExports().NetTokensPendingExportInCurrentBatch);
+			const bool bIsNetTokenAuthority = NetTokenStore->IsAuthority();
 
 			for (const FNetToken& NetToken : NetTokenExportsView)
 			{
-				if (!ExportContext->IsExported(NetToken))
+				// We do not need to export tokens assigned by authority unless we are the authority.
+				if (!(NetToken.IsAssignedByAuthority() && !bIsNetTokenAuthority) && !ExportContext->IsExported(NetToken))
 				{
-					//UE_NET_TRACE_DYNAMIC_NAME_SCOPE(*FString::Printf(TEXT("%u"), NetToken.GetIndex()), *Context.GetBitStreamWriter(), Context.GetTraceCollector(), ENetTraceVerbosity::VeryVerbose);
 					Writer.WriteBool(true);
-					WriteNetToken(&Writer, NetToken);
+					WriteNetToken(Context, NetToken);
 					NetTokenStore->WriteTokenData(Context, NetToken);
 					ExportContext->AddExported(NetToken);			
 				}
 			}
-			// Write stop bit
-			Writer.WriteBool(false);
 		}
+		// Write stop bit
+		Writer.WriteBool(false);
 
 		for (const FNetObjectReference& Reference : NetObjectReferenceExportsView)
 		{
@@ -1656,7 +1656,7 @@ bool FObjectReferenceCache::ReadExports(FNetSerializationContext& Context, TArra
 			FForceInlineExportScope ForceInlineExportScope(Context.GetInternalContext());
 			while (bHasExportsToRead && !Context.HasErrorOrOverflow())
 			{
-				FNetToken ImportedNetToken = ReadNetToken(&Reader);
+				FNetToken ImportedNetToken = ReadNetToken(Context);
 				NetTokenStore->ReadTokenData(Context, ImportedNetToken, *ResolveContext.RemoteNetTokenStoreState);
 				bHasExportsToRead = Reader.ReadBool();
 			}
diff --git a/Engine/Source/Runtime/Experimental/Iris/Core/Private/Iris/ReplicationSystem/ReplicationReader.h b/Engine/Source/Runtime/Experimental/Iris/Core/Private/Iris/ReplicationSystem/ReplicationReader.h
--- a/Engine/Source/Runtime/Experimental/Iris/Core/Private/Iris/ReplicationSystem/ReplicationReader.h
+++ b/Engine/Source/Runtime/Experimental/Iris/Core/Private/Iris/ReplicationSystem/ReplicationReader.h
@@ -19,6 +19,9 @@ namespace UE::Net
 	class FNetSerializationContext;
 	class FNetTokenStoreState;
 	class FReplicationStateStorage;
+	class FNetBitStreamReader;
+	class FNetBitStreamWriter;
+
 	namespace Private
 	{
 		class FReplicationSystemInternal;
diff --git a/Engine/Source/Runtime/Experimental/Iris/Core/Private/Iris/ReplicationSystem/ReplicationSystem.cpp b/Engine/Source/Runtime/Experimental/Iris/Core/Private/Iris/ReplicationSystem/ReplicationSystem.cpp
--- a/Engine/Source/Runtime/Experimental/Iris/Core/Private/Iris/ReplicationSystem/ReplicationSystem.cpp
+++ b/Engine/Source/Runtime/Experimental/Iris/Core/Private/Iris/ReplicationSystem/ReplicationSystem.cpp
@@ -107,6 +107,13 @@ public:
 #endif
 		const uint32 ReplicationSystemId = ReplicationSystem->GetId();
 
+		FNetTokenStore& NetTokenStore = ReplicationSystemInternal.GetNetTokenStore();
+		{
+			FNetTokenStore::FInitParams NetTokenStoreInitParams;
+			NetTokenStoreInitParams.Authority = Params.bIsServer ? FNetToken::ENetTokenAuthority::Authority : FNetToken::ENetTokenAuthority::None;
+			NetTokenStore.Init(NetTokenStoreInitParams);
+		}
+
 		FNetRefHandleManager& NetRefHandleManager = ReplicationSystemInternal.GetNetRefHandleManager();
 		{
 			FNetRefHandleManager::FInitParams NetRefHandleManagerInitParams;
@@ -1016,24 +1023,20 @@ UNetObjectPrioritizer* UReplicationSystem::GetPrioritizer(const FName Prioritize
 	return Impl->ReplicationSystemInternal.GetPrioritization().GetPrioritizer(PrioritizerName);
 }
 
-const UE::Net::FStringTokenStore* UReplicationSystem::GetStringTokenStore() const
-{
-	return &Impl->ReplicationSystemInternal.GetStringTokenStore();
-}
-
-UE::Net::FStringTokenStore* UReplicationSystem::GetStringTokenStore()
+const UE::Net::FNetTokenStore* UReplicationSystem::GetNetTokenStore() const
 {
-	return &Impl->ReplicationSystemInternal.GetStringTokenStore();
+	return &Impl->ReplicationSystemInternal.GetNetTokenStore();
 }
 
-const UE::Net::FNameTokenStore* UReplicationSystem::GetNameTokenStore() const
+UE::Net::FNetTokenStore* UReplicationSystem::GetNetTokenStore()
 {
-	return &Impl->ReplicationSystemInternal.GetNameTokenStore();
+	return &Impl->ReplicationSystemInternal.GetNetTokenStore();
 }
 
-UE::Net::FNameTokenStore* UReplicationSystem::GetNameTokenStore()
+void UReplicationSystem::InitNetTokenExportContext(UE::Net::FNetTokenExportContext& Context, uint32 ConnectionId)
 {
-	return &Impl->ReplicationSystemInternal.GetNameTokenStore();
+	Context.TokenStore = &Impl->ReplicationSystemInternal.GetNetTokenStore();
+	Context.RemoteState = &Impl->ReplicationSystemInternal.GetConnections().GetRemoteNetTokenStoreState(ConnectionId);
 }
 
 bool UReplicationSystem::RegisterNetBlobHandler(UNetBlobHandler* Handler)
diff --git a/Engine/Source/Runtime/Experimental/Iris/Core/Private/Iris/ReplicationSystem/StringTokenStore.cpp b/Engine/Source/Runtime/Experimental/Iris/Core/Private/Iris/ReplicationSystem/StringTokenStore.cpp
--- a/Engine/Source/Runtime/Experimental/Iris/Core/Private/Iris/ReplicationSystem/StringTokenStore.cpp
+++ b/Engine/Source/Runtime/Experimental/Iris/Core/Private/Iris/ReplicationSystem/StringTokenStore.cpp
@@ -37,10 +37,10 @@ FNetToken FStringTokenStore::GetOrCreateToken(const TCHAR* Name, uint32 Length)
 		}
 		else
 		{
-			const FNetToken NewToken = CreateToken(Key, *TokenStore.GetLocalNetTokenStoreState());
+			const FNetToken NewToken = CreateToken(TokenStore.IsAuthority() ? FNetToken::ENetTokenAuthority::Authority : FNetToken::ENetTokenAuthority::None, Key, *TokenStore.GetLocalNetTokenStoreState());
 			StoredTokens[StringIndex] = NewToken;
 
-			UE_LOG_STRINGTOKEN(TEXT("FStringTokenStore::GetOrCreateToken - Created new token %u for %s"), NewToken.GetIndex(), Name);
+			UE_LOG_STRINGTOKEN(TEXT("FStringTokenStore::GetOrCreateToken - Created %s for %s"), *NewToken.ToString(), Name);
 
 			return NewToken;
 		}
@@ -94,21 +94,19 @@ FStringTokenStore::FStringTokenStore(FNetTokenStore& InTokenStore)
 	StoredTokens.Add(FNetToken());
 
 	// Register
-	TokenStore.RegisterDataStore(this);
+	TokenStore.RegisterDataStore(this, GetTokenStoreName());
 }
 
-const TCHAR* FStringTokenStore::ResolveToken(FNetToken Token) const
+const TCHAR* FStringTokenStore::ResolveToken(FNetToken Token, const FNetTokenStoreState* NetTokenStoreState) const
 {
-	return ResolveRemoteToken(Token, *TokenStore.GetLocalNetTokenStoreState());
-}
-
-const TCHAR* FStringTokenStore::ResolveRemoteToken(FNetToken Token, const FNetTokenStoreState& NetTokenStoreState) const
-{
-	const FNetTokenStoreKey StoreKey = GetTokenKey(Token, NetTokenStoreState);
-
-	if (StoreKey.IsValid() && StoreKey.GetKeyValue() < (uint32)StoredStrings.Num())
+	const FNetTokenStoreState* TokenStoreState = TokenStore.IsLocalToken(Token) ? TokenStore.GetLocalNetTokenStoreState() : NetTokenStoreState;
+	if (Token.IsValid() && ensureMsgf(TokenStoreState, TEXT("FStringTokenStore::ResolveToken Needs valid remote NetTokenStoreState to resolve remote %s"), *Token.ToString()))
 	{
-		return StoredStrings[StoreKey.GetKeyValue()];
+		const FNetTokenStoreKey StoreKey = GetTokenKey(Token, *TokenStoreState);
+		if (StoreKey.IsValid() && StoreKey.GetKeyValue() < (uint32)StoredStrings.Num())
+		{
+			return StoredStrings[StoreKey.GetKeyValue()];
+		}
 	}
 
 	return nullptr;
@@ -120,7 +118,7 @@ void FStringTokenStore::WriteTokenData(FNetSerializationContext& Context, FNetTo
 	WriteString(Context.GetBitStreamWriter(), FStringView(StoredStrings[TokenStoreKey.GetKeyValue()]));
 }
 
-FNetTokenStoreKey FStringTokenStore::ReadTokenData(FNetSerializationContext& Context)
+FNetTokenStoreKey FStringTokenStore::ReadTokenData(FNetSerializationContext& Context, const FNetToken& NetToken)
 {
 	FNetBitStreamReader* Reader = Context.GetBitStreamReader();
 
@@ -130,7 +128,17 @@ FNetTokenStoreKey FStringTokenStore::ReadTokenData(FNetSerializationContext& Con
 
 	if (!Reader->IsOverflown())
 	{
-		return GetOrCreatePersistentString(*Temp, Temp.Len());
+		FNetTokenStoreKey Key = GetOrCreatePersistentString(*Temp, Temp.Len());
+
+		// If this is an authtoken and we are not the authority, update the stored key
+		if (Key.IsValid() && NetToken.IsAssignedByAuthority() && !TokenStore.IsAuthority())
+		{
+			UE_LOG_STRINGTOKEN(TEXT("FStringTokenStore::ReadTokenData - Replaced local key %s with %s for string %s"), *StoredTokens[Key.GetKeyValue()].ToString(), *NetToken.ToString(), *Temp);	
+
+			// This way the next time we lookup this key during assignment use the imported authoritative key which we do not have to export
+			StoredTokens[Key.GetKeyValue()] = NetToken;
+		}
+		return Key;
 	}
 	else
 	{
diff --git a/Engine/Source/Runtime/Experimental/Iris/Core/Private/Iris/Serialization/InternalNetSerializers.h b/Engine/Source/Runtime/Experimental/Iris/Core/Private/Iris/Serialization/InternalNetSerializers.h
--- a/Engine/Source/Runtime/Experimental/Iris/Core/Private/Iris/Serialization/InternalNetSerializers.h
+++ b/Engine/Source/Runtime/Experimental/Iris/Core/Private/Iris/Serialization/InternalNetSerializers.h
@@ -73,10 +73,6 @@ struct FLastResortPropertyNetSerializerConfig : public FNetSerializerConfig
 public:
 	UPROPERTY()
 	TFieldPath<FProperty> Property;
-
-	// Value used to sanity check incoming data so that we do not over-allocate dynamic memory
-	UPROPERTY()
-	uint32 MaxAllowedObjectReferences = 128; 
 };
 
 // ENetRole. With role swapping at deserialization
diff --git a/Engine/Source/Runtime/Experimental/Iris/Core/Private/Iris/Serialization/IrisObjectReferencePackageMap.cpp b/Engine/Source/Runtime/Experimental/Iris/Core/Private/Iris/Serialization/IrisObjectReferencePackageMap.cpp
--- a/Engine/Source/Runtime/Experimental/Iris/Core/Private/Iris/Serialization/IrisObjectReferencePackageMap.cpp
+++ b/Engine/Source/Runtime/Experimental/Iris/Core/Private/Iris/Serialization/IrisObjectReferencePackageMap.cpp
@@ -1,15 +1,26 @@
 // Copyright Epic Games, Inc. All Rights Reserved.
 
 #include "Iris/Serialization/IrisObjectReferencePackageMap.h"
+#include "HAL/IConsoleManager.h"
+
+namespace UE
+{
+	namespace Net
+	{
+		bool bEnableIrisPackageMapNameExports = true;
+		static FAutoConsoleVariableRef CVarEnableIrisPackageMapNameExports(TEXT("net.iris.EnableIrisPackageMapNameExports"), bEnableIrisPackageMapNameExports, TEXT("If enabled, iris captures and exports fnames when calling into old serialziation code instead of serializing a strings."));
+	}
+}
 
 bool UIrisObjectReferencePackageMap::SerializeObject(FArchive& Ar, UClass* InClass, UObject*& Obj, FNetworkGUID* OutNetGUID)
 {
-	if (!References)
+	if (!PackageMapExports)
 	{
 		return false;
 	}
 
 	constexpr uint8 MaxNumReferences = 255U;
+	UE::Net::FIrisPackageMapExports::FObjectReferenceArray* References = &PackageMapExports->References;
 
 	if (Ar.IsSaving())
 	{
@@ -51,17 +62,119 @@ bool UIrisObjectReferencePackageMap::SerializeObject(FArchive& Ar, UClass* InCla
 	return true;
 }
 
-void UIrisObjectReferencePackageMap::InitForRead(const FObjectReferenceArray* InReferences)
+bool UIrisObjectReferencePackageMap::SerializeName(FArchive& Ar, FName& InName)
+{
+	using namespace UE::Net;
+
+	if (!(PackageMapExports && bEnableIrisPackageMapNameExports))
+	{
+		return Super::SerializeName(Ar, InName);
+	}
+
+	FIrisPackageMapExports::FNameArray* Names = &PackageMapExports->Names;
+	constexpr uint8 MaxNumNames = 255U;
+
+	if (Ar.IsSaving())
+	{
+		int32 Index = MaxNumNames;
+		if (!Names->Find(InName, Index) && Names->Num() < MaxNumNames)
+		{
+			Index = Names->Add(InName);
+		}
+
+		if (Names->IsValidIndex(Index))
+		{
+			uint8 IndexByte = static_cast<uint8>(Index);
+			Ar << IndexByte;
+		}
+		else
+		{
+			ensureMsgf(false, TEXT("UIrisObjectReferencePackageMap::SerializeName, failed to serialize name with Index %u (%s). A Maximum of %u name are currently supported by this PackageMap"),
+				Index, *InName.ToString(), MaxNumNames);
+			uint8 IndexByte = MaxNumNames;
+			Ar << IndexByte;
+			return false;
+		}
+	}
+	else
+	{
+		uint8 IndexByte = MaxNumNames;
+		Ar << IndexByte;
+		if (Names->IsValidIndex(IndexByte) && IndexByte < MaxNumNames)
+		{
+			InName = (*Names)[IndexByte];
+		}
+		else
+		{
+			ensureMsgf(false, TEXT("UIrisObjectReferencePackageMap::SerializeName, failed to read name index %u is out of bounds. Current Name num: %u"), IndexByte, Names->Num());
+			return false;
+		}
+	}
+
+	return true;
+}
+
+bool UIrisObjectReferencePackageMap::SerializeNetToken(FArchive& Ar, UE::Net::FNetToken& InNetToken)
+{
+	if (!PackageMapExports)
+	{
+		return false;
+	}
+
+	UE::Net::FIrisPackageMapExports::FNetTokensArray* NetTokens = &PackageMapExports->NetTokens;
+	constexpr uint8 MaxNumNetTokens = 255U;
+
+	if (Ar.IsSaving())
+	{
+		int32 Index = MaxNumNetTokens;
+		if (!NetTokens->Find(InNetToken, Index) && NetTokens->Num() < MaxNumNetTokens)
+		{
+			Index = NetTokens->Add(InNetToken);
+		}
+
+		if (NetTokens->IsValidIndex(Index))
+		{
+			uint8 IndexByte = static_cast<uint8>(Index);
+			Ar << IndexByte;
+		}
+		else
+		{
+			ensureMsgf(false, TEXT("UIrisObjectReferencePackageMap::SerializeNetToken, failed to serialize NetToken with Index %u (%s). A Maximum of %u nettokens are currently supported by this PackageMap"),
+				Index, *InNetToken.ToString(), MaxNumNetTokens);
+			uint8 IndexByte = MaxNumNetTokens;
+			Ar << IndexByte;
+			return false;
+		}
+	}
+	else
+	{
+		uint8 IndexByte = MaxNumNetTokens;
+		Ar << IndexByte;
+		if (NetTokens->IsValidIndex(IndexByte) && IndexByte < MaxNumNetTokens)
+		{
+			InNetToken = (*NetTokens)[IndexByte];
+		}
+		else
+		{
+			ensureMsgf(false, TEXT("UIrisObjectReferencePackageMap::SerializeName, failed to read name index %u is out of bounds. Current Name num: %u"), IndexByte, NetTokens->Num());
+			return false;
+		}
+	}
+
+	return true;
+}
+
+void UIrisObjectReferencePackageMap::InitForRead(const UE::Net::FIrisPackageMapExports* InPackageMapExports)
 { 
-	References = const_cast<FObjectReferenceArray*>(InReferences);
+	PackageMapExports = const_cast<UE::Net::FIrisPackageMapExports*>(InPackageMapExports);
 }
 
-void UIrisObjectReferencePackageMap::InitForWrite(FObjectReferenceArray* InReferences)
+void UIrisObjectReferencePackageMap::InitForWrite(UE::Net::FIrisPackageMapExports* InPackageMapExports)
 {
-	References = InReferences; 
-	if (References)
+	PackageMapExports = InPackageMapExports;
+	if (ensureMsgf(PackageMapExports, TEXT("UIrisObjectReferencePackageMap requires valid PackageMapExports to capture exports.")))
 	{
-		References->Reset();
+		PackageMapExports->Reset();
 	}
 }
 
diff --git a/Engine/Source/Runtime/Experimental/Iris/Core/Private/Iris/Serialization/IrisPackageMapExportUtil.cpp b/Engine/Source/Runtime/Experimental/Iris/Core/Private/Iris/Serialization/IrisPackageMapExportUtil.cpp
new file mode 100644
--- /dev/null
+++ b/Engine/Source/Runtime/Experimental/Iris/Core/Private/Iris/Serialization/IrisPackageMapExportUtil.cpp
@@ -0,0 +1,386 @@
+// Copyright Epic Games, Inc. All Rights Reserved.
+
+#include "Iris/Serialization/IrisPackageMapExportutil.h"
+
+#include "Iris/Core/IrisLog.h"
+#include "Iris/Serialization/NetBitStreamReader.h"
+#include "Iris/Serialization/NetBitStreamUtil.h"
+#include "Iris/Serialization/NetSerializers.h"
+#include "Iris/Serialization/NetBitStreamWriter.h"
+#include "Iris/Serialization/NetReferenceCollector.h"
+#include "Iris/Serialization/NetSerializerArrayStorage.h"
+#include "Iris/Serialization/IrisObjectReferencePackageMap.h"
+#include "Iris/ReplicationSystem/NetTokenStore.h"
+#include "Net/Core/Trace/NetTrace.h"
+
+namespace UE::Net
+{
+
+const FNetSerializer* FIrisPackageMapExportsUtil::ObjectNetSerializer = &UE_NET_GET_SERIALIZER(FObjectNetSerializer);
+const FNetSerializer* FIrisPackageMapExportsUtil::NameNetSerializer = &UE_NET_GET_SERIALIZER(FNameNetSerializer);
+
+void FIrisPackageMapExportsUtil::Serialize(FNetSerializationContext& Context, const QuantizedType& Value)
+{
+	FNetBitStreamWriter* Writer = Context.GetBitStreamWriter();
+
+	// If we have any references, export them!
+	{
+		const uint32 NumReferences = Value.ObjectReferenceStorage.Num();
+		if (Writer->WriteBool(NumReferences != 0))
+		{
+			UE::Net::WritePackedUint32(Writer, NumReferences);		
+			FObjectNetSerializerConfig ObjectSerializerConfig;
+			for (const FNetObjectReference& Ref : MakeArrayView(Value.ObjectReferenceStorage.GetData(), NumReferences))
+			{
+				FNetSerializeArgs ObjectArgs;
+				ObjectArgs.NetSerializerConfig = &ObjectSerializerConfig;
+				ObjectArgs.Source = NetSerializerValuePointer(&Ref);
+
+				ObjectNetSerializer->Serialize(Context, ObjectArgs);
+			}
+		}
+	}
+
+	// If we have any names, export them!
+	{
+		const uint32 NumNames = Value.NameStorage.Num();
+		if (Writer->WriteBool(NumNames != 0))
+		{
+			UE::Net::WritePackedUint32(Writer, NumNames);
+			FNetSerializerConfig NameSerializerConfig;
+			for (const QuantizedType::FQuantizedName& QuantizedName : MakeArrayView(Value.NameStorage.GetData(), NumNames))
+			{
+				FNetSerializeArgs NameArgs;
+				NameArgs.NetSerializerConfig = &NameSerializerConfig;
+				NameArgs.Source = NetSerializerValuePointer(&QuantizedName);
+
+				NameNetSerializer->Serialize(Context, NameArgs);
+			}
+		}
+	}
+
+	// For NetTokens
+	{
+		// $TODO: For asymmetrical data like nettokens we will need to do something special to support validating the default state hash as the local quantized state will differ
+		// For now we ignore this in default state hash
+		if (Context.IsInitializingDefaultState())
+		{
+			return;
+		}
+
+		const uint32 NumNetTokens = Value.NetTokenStorage.Num();
+		if (Writer->WriteBool(NumNetTokens != 0))
+		{
+			UE::Net::WritePackedUint32(Writer, NumNetTokens);
+			for (const FNetToken& NetToken : MakeArrayView(Value.NetTokenStorage.GetData(), NumNetTokens))
+			{
+				// Always write the token
+				WriteNetToken(Context, NetToken);
+				// Export or add to pending exports for later export
+				FNetTokenStore::AppendExportOrWriteInlinedExportData(Context, NetToken);
+			}
+		}
+	}
+}
+
+void FIrisPackageMapExportsUtil::Deserialize(FNetSerializationContext& Context, QuantizedType& Value)
+{
+	FNetBitStreamReader* Reader = Context.GetBitStreamReader();
+
+	// Read any object references
+	{
+		const bool bHasObjectReferences = Reader->ReadBool();
+		if (bHasObjectReferences)
+		{
+			const uint32 NumReferences = UE::Net::ReadPackedUint32(Reader);
+
+			if (NumReferences > MaxExports)
+			{
+				UE_LOG(LogIris, Error, TEXT("FIrisPackageMapExportsUtil::Received too many object reference exports %u > max:%u"), NumReferences, MaxExports);
+				Context.SetError(GNetError_ArraySizeTooLarge);
+				return;
+			}
+
+			Value.ObjectReferenceStorage.AdjustSize(Context, NumReferences);
+
+			FObjectNetSerializerConfig ObjectSerializerConfig;
+			for (FNetObjectReference& Ref : MakeArrayView(Value.ObjectReferenceStorage.GetData(), Value.ObjectReferenceStorage.Num()))
+			{
+				FNetDeserializeArgs ObjectArgs;
+				ObjectArgs.NetSerializerConfig = &ObjectSerializerConfig;
+				ObjectArgs.Target = NetSerializerValuePointer(&Ref);
+
+				ObjectNetSerializer->Deserialize(Context, ObjectArgs);
+			}
+		}
+		else
+		{
+			Value.ObjectReferenceStorage.Free(Context);
+		}
+	}
+
+	// Read any exported names
+	{
+		const bool bHasNames = Reader->ReadBool();
+		if (bHasNames)
+		{
+			const uint32 NumNames = UE::Net::ReadPackedUint32(Reader);
+
+			if (NumNames > MaxExports)
+			{
+				UE_LOG(LogIris, Error, TEXT("FIrisPackageMapExportsUtil::Received too many name exports %u > max:%u"), NumNames, MaxExports);
+				Context.SetError(GNetError_ArraySizeTooLarge);
+				return;
+			}
+
+			Value.NameStorage.AdjustSize(Context, NumNames);
+
+			FNetSerializerConfig NameSerializerConfig;
+			for (QuantizedType::FQuantizedName& QuantizedName : MakeArrayView(Value.NameStorage.GetData(), Value.NameStorage.Num()))
+			{
+				FNetDeserializeArgs NameArgs;
+				NameArgs.NetSerializerConfig = &NameSerializerConfig;
+				NameArgs.Target = NetSerializerValuePointer(&QuantizedName);
+
+				NameNetSerializer->Deserialize(Context, NameArgs);
+			}
+		}
+		else
+		{
+			Value.NameStorage.Free(Context);
+		}
+	}
+
+	// Read any exported nettokens
+	{
+		const bool bHasNetTokens = Reader->ReadBool();
+		if (bHasNetTokens)
+		{
+			const uint32 NumNetTokens = UE::Net::ReadPackedUint32(Reader);
+
+			if (NumNetTokens > MaxExports)
+			{
+				UE_LOG(LogIris, Error, TEXT("FIrisPackageMapExportsUtil::Received too many NetToken exports %u > max:%u"), NumNetTokens, MaxExports);
+				Context.SetError(GNetError_ArraySizeTooLarge);
+				return;
+			}
+
+			Value.NetTokenStorage.AdjustSize(Context, NumNetTokens);
+
+			for (FNetToken& TargetToken : MakeArrayView(Value.NetTokenStorage.GetData(), Value.NetTokenStorage.Num()))
+			{
+				// Always Read the token
+				FNetToken NetToken = ReadNetToken(Context);
+
+				if (Reader->IsOverflown())
+				{
+					return;
+				}
+
+				// Read inlined exports if there are any
+				FNetTokenStore::ReadInlinedExportData(Context, NetToken);
+
+				TargetToken = NetToken;
+			}
+		}
+		else
+		{
+			Value.NetTokenStorage.Free(Context);
+		}
+	}
+}
+
+void FIrisPackageMapExportsUtil::Quantize(FNetSerializationContext& Context, const UE::Net::FIrisPackageMapExports& PackageMapExports, QuantizedType& Value)
+{
+	// Quantize captured references
+	{
+		const FIrisPackageMapExports::FObjectReferenceArray& ObjectReferences = PackageMapExports.References;
+		const uint32 NumObjectReferences = ObjectReferences.Num();
+		Value.ObjectReferenceStorage.AdjustSize(Context, NumObjectReferences);
+		if (NumObjectReferences > 0)
+		{
+			FObjectNetSerializerConfig ObjectNetSerializerConfig;
+			const TObjectPtr<UObject>* SourceReferences = ObjectReferences.GetData();
+			FNetObjectReference* TargetReferences = Value.ObjectReferenceStorage.GetData();
+			for (uint32 ReferenceIndex = 0; ReferenceIndex < NumObjectReferences; ++ReferenceIndex)
+			{
+				FNetQuantizeArgs ObjectArgs;
+				ObjectArgs.NetSerializerConfig = &ObjectNetSerializerConfig;
+				ObjectArgs.Source = NetSerializerValuePointer(SourceReferences + ReferenceIndex);
+				ObjectArgs.Target = NetSerializerValuePointer(TargetReferences + ReferenceIndex);
+
+				ObjectNetSerializer->Quantize(Context, ObjectArgs);
+			}
+		}
+	}
+
+	// Quantize captured names
+	{
+		const FIrisPackageMapExports::FNameArray& Names = PackageMapExports.Names;
+		const uint32 NumNames = Names.Num();
+		Value.NameStorage.AdjustSize(Context, NumNames);
+		if (NumNames > 0)
+		{
+			FNetSerializerConfig NameNetSerializerConfig;
+			const FName* SourceNames = Names.GetData();
+			QuantizedType::FQuantizedName* TargetNames = Value.NameStorage.GetData();
+			for (uint32 ReferenceIndex = 0; ReferenceIndex < NumNames; ++ReferenceIndex)
+			{
+				FNetQuantizeArgs NameArgs;
+				NameArgs.NetSerializerConfig = &NameNetSerializerConfig;
+				NameArgs.Source = NetSerializerValuePointer(SourceNames + ReferenceIndex);
+				NameArgs.Target = NetSerializerValuePointer(TargetNames + ReferenceIndex);
+
+				NameNetSerializer->Quantize(Context, NameArgs);
+			}
+		}
+	}
+
+	// Just store captured nettokens
+	{
+		const FIrisPackageMapExports::FNetTokensArray& NetTokens = PackageMapExports.NetTokens;
+		const uint32 NumNetTokens = NetTokens.Num();
+		Value.NetTokenStorage.AdjustSize(Context, NumNetTokens);
+		if (NumNetTokens > 0)
+		{
+			const FNetToken* SourceTokens = NetTokens.GetData();
+			FNetToken* TargetTokens = Value.NetTokenStorage.GetData();
+			for (uint32 ReferenceIndex = 0; ReferenceIndex < NumNetTokens; ++ReferenceIndex)
+			{
+				TargetTokens[ReferenceIndex] = SourceTokens[ReferenceIndex];
+			}
+		}
+	}
+}
+
+void FIrisPackageMapExportsUtil::Dequantize(FNetSerializationContext& Context, const QuantizedType& Source, UE::Net::FIrisPackageMapExports& PackageMapExports)
+{
+	Private::FInternalNetSerializationContext* InternalContext = Context.GetInternalContext();
+
+	// References
+	{
+		UE::Net::FIrisPackageMapExports::FObjectReferenceArray& ObjectReferences = PackageMapExports.References;
+		const uint32 NumObjectReferences = Source.ObjectReferenceStorage.Num();
+		if (NumObjectReferences > 0U)
+		{		
+			ObjectReferences.SetNumUninitialized(NumObjectReferences);
+
+			FObjectNetSerializerConfig ObjectNetSerializerConfig;
+			const FNetObjectReference* SourceReferences = Source.ObjectReferenceStorage.GetData();
+			TObjectPtr<UObject>* TargetReferences = ObjectReferences.GetData();
+			for (uint32 ReferenceIndex = 0; ReferenceIndex < NumObjectReferences; ++ReferenceIndex)
+			{
+				FNetDequantizeArgs ObjectArgs;
+				ObjectArgs.NetSerializerConfig = &ObjectNetSerializerConfig;
+				ObjectArgs.Source = NetSerializerValuePointer(SourceReferences + ReferenceIndex);
+				ObjectArgs.Target = NetSerializerValuePointer(TargetReferences + ReferenceIndex);
+
+				ObjectNetSerializer->Dequantize(Context, ObjectArgs);
+			}
+		}
+	}
+
+	// Names
+	{
+		UE::Net::FIrisPackageMapExports::FNameArray& Names = PackageMapExports.Names;
+
+		const uint32 NumNames = Source.NameStorage.Num();
+		if (NumNames > 0U)
+		{		
+			Names.SetNumUninitialized(NumNames);
+
+			FNetSerializerConfig NameNetSerializerConfig;
+			const QuantizedType::FQuantizedName* SourceNames = Source.NameStorage.GetData();
+			FName* TargetNames = Names.GetData();
+			for (uint32 ReferenceIndex = 0; ReferenceIndex < NumNames; ++ReferenceIndex)
+			{
+				FNetDequantizeArgs NameArgs;
+				NameArgs.NetSerializerConfig = &NameNetSerializerConfig;
+				NameArgs.Source = NetSerializerValuePointer(SourceNames + ReferenceIndex);
+				NameArgs.Target = NetSerializerValuePointer(TargetNames + ReferenceIndex);
+
+				NameNetSerializer->Dequantize(Context, NameArgs);
+			}
+		}
+	}
+
+	// NetTokens
+	{
+		UE::Net::FIrisPackageMapExports::FNetTokensArray& NetTokens = PackageMapExports.NetTokens;
+
+		const uint32 NumNetTokens = Source.NetTokenStorage.Num();
+		if (NumNetTokens > 0U)
+		{		
+			NetTokens.SetNumUninitialized(NumNetTokens);
+
+			const FNetToken* SourceTokens = Source.NetTokenStorage.GetData();
+			FNetToken* TargetTokens = NetTokens.GetData();
+			for (uint32 ReferenceIndex = 0; ReferenceIndex < NumNetTokens; ++ReferenceIndex)
+			{
+				TargetTokens[ReferenceIndex] = SourceTokens[ReferenceIndex];
+			}
+		}
+	}
+}
+
+bool FIrisPackageMapExportsUtil::IsEqual(FNetSerializationContext& Context, const QuantizedType& Value0, const QuantizedType& Value1)
+{
+	if ((Value0.ObjectReferenceStorage.Num() != Value1.ObjectReferenceStorage.Num()) || (Value0.NameStorage.Num() != Value1.NameStorage.Num()) || (Value0.NetTokenStorage.Num() != Value1.NetTokenStorage.Num()))
+	{
+		return false;
+	}
+
+	if (Value0.ObjectReferenceStorage.Num() > 0 && FMemory::Memcmp(Value0.ObjectReferenceStorage.GetData(), Value1.ObjectReferenceStorage.GetData(), sizeof(FNetObjectReference) * Value0.ObjectReferenceStorage.Num()) != 0)
+	{
+		return false;
+	}
+
+	if (Value0.NameStorage.Num() > 0 && FMemory::Memcmp(Value0.NameStorage.GetData(), Value1.NameStorage.GetData(), sizeof(FName) * Value0.NameStorage.Num()) != 0)
+	{
+		return false;
+	}
+
+	// Note: this can be improved by implementing a resolved token compare. Currently this is pessimistic and does not allow compare between auth and non-auth tokens even if they would
+	// resolve to same value.
+	if (Value0.NetTokenStorage.Num() > 0 && FMemory::Memcmp(Value0.NetTokenStorage.GetData(), Value1.NetTokenStorage.GetData(), sizeof(FNetToken) * Value0.NetTokenStorage.Num()) != 0)
+	{
+		return false;
+	}
+
+	return true;
+}
+
+void FIrisPackageMapExportsUtil::CloneDynamicState(FNetSerializationContext& Context, QuantizedType& Target, const QuantizedType& Source)
+{
+	Target.ObjectReferenceStorage.Clone(Context, Source.ObjectReferenceStorage);
+	Target.NameStorage.Clone(Context, Source.NameStorage);
+	Target.NetTokenStorage.Clone(Context, Source.NetTokenStorage);
+}
+
+void FIrisPackageMapExportsUtil::FreeDynamicState(FNetSerializationContext& Context, QuantizedType& Value)
+{
+	// Clear all info
+	Value.ObjectReferenceStorage.Free(Context);
+	Value.NameStorage.Free(Context);
+	Value.NetTokenStorage.Free(Context);
+}
+
+void FIrisPackageMapExportsUtil::CollectNetReferences(FNetSerializationContext& Context, const QuantizedType& Value, const FNetSerializerChangeMaskParam& ChangeMaskInfo, FNetReferenceCollector& Collector)
+{
+	const FNetReferenceInfo ReferenceInfo(FNetReferenceInfo::EResolveType::ResolveOnClient);	
+	for (const FNetObjectReference& Ref : MakeArrayView(Value.ObjectReferenceStorage.GetData(), Value.ObjectReferenceStorage.Num()))
+	{
+		Collector.Add(ReferenceInfo, Ref, ChangeMaskInfo);
+	}
+}
+
+bool FIrisPackageMapExportsUtil::Validate(FNetSerializationContext& Context, const QuantizedType& SourceValue)
+{
+	if ((SourceValue.ObjectReferenceStorage.Num() > MaxExports) || (SourceValue.NameStorage.Num() > MaxExports) || (SourceValue.NameStorage.Num() > MaxExports))
+	{
+		return false;
+	}
+	return true;
+}
+
+}
diff --git a/Engine/Source/Runtime/Experimental/Iris/Core/Private/Iris/Serialization/LastResortPropertyNetSerializer.cpp b/Engine/Source/Runtime/Experimental/Iris/Core/Private/Iris/Serialization/LastResortPropertyNetSerializer.cpp
--- a/Engine/Source/Runtime/Experimental/Iris/Core/Private/Iris/Serialization/LastResortPropertyNetSerializer.cpp
+++ b/Engine/Source/Runtime/Experimental/Iris/Core/Private/Iris/Serialization/LastResortPropertyNetSerializer.cpp
@@ -10,6 +10,8 @@
 #include "Iris/Serialization/NetReferenceCollector.h"
 #include "Iris/Serialization/NetSerializerArrayStorage.h"
 #include "Iris/Serialization/IrisObjectReferencePackageMap.h"
+#include "Iris/Serialization/IrisPackageMapExportUtil.h"
+#include "Iris/Serialization/NetExportContext.h"
 #include "Net/Core/Trace/NetTrace.h"
 
 namespace UE::Net
@@ -17,10 +19,7 @@ namespace UE::Net
 
 struct FFLastResortPropertyNetSerializerQuantizedType
 {
-	static constexpr uint32 MaxInlinedObjectRefs = 4;
-	typedef FNetSerializerArrayStorage<FNetObjectReference, AllocationPolicies::TInlinedElementAllocationPolicy<MaxInlinedObjectRefs>> FObjectReferenceStorage;
-
-	FObjectReferenceStorage ObjectReferenceStorage;
+	FIrisPackageMapExportsQuantizedType QuantizedExports;
 
 	// How many bytes the current allocation can hold.
 	uint16 ByteCapacity = 0;
@@ -73,8 +72,6 @@ private:
 	static void GrowDynamicStateInternal(FNetSerializationContext&, QuantizedType& Value, uint16 NewBitCount);
 	static void ShrinkDynamicStateInternal(FNetSerializationContext&, QuantizedType& Value, uint16 NewBitCount);
 	static void AdjustStorageSize(FNetSerializationContext&, QuantizedType& Value, uint16 NewBitCount);
-
-	static inline const FNetSerializer* ObjectNetSerializer = &UE_NET_GET_SERIALIZER(FObjectNetSerializer);
 };
 UE_NET_IMPLEMENT_SERIALIZER_INTERNAL(FLastResortPropertyNetSerializer);
 
@@ -83,25 +80,13 @@ void FLastResortPropertyNetSerializer::Serialize(FNetSerializationContext& Conte
 	const ConfigType* Config = static_cast<const ConfigType*>(Args.NetSerializerConfig);
 	const QuantizedType& Value = *reinterpret_cast<const QuantizedType*>(Args.Source);
 	FNetBitStreamWriter* Writer = Context.GetBitStreamWriter();
-	const uint32 NumReferences = Value.ObjectReferenceStorage.Num();
 
 	UE_NET_TRACE_DYNAMIC_NAME_SCOPE(Config->Property.Get()->GetName(), *Writer, Context.GetTraceCollector(), ENetTraceVerbosity::VeryVerbose);
 
-	// If we have any references, export them!
-	if (Writer->WriteBool(NumReferences != 0))
-	{
-		UE::Net::WritePackedUint32(Writer, NumReferences);		
-		FObjectNetSerializerConfig ObjectSerializerConfig;
-		for (const FNetObjectReference& Ref : MakeArrayView(Value.ObjectReferenceStorage.GetData(), NumReferences))
-		{
-			FNetSerializeArgs ObjectArgs;
-			ObjectArgs.NetSerializerConfig = &ObjectSerializerConfig;
-			ObjectArgs.Source = NetSerializerValuePointer(&Ref);
-
-			ObjectNetSerializer->Serialize(Context, ObjectArgs);
-		}
-	}
+	// If we have any captured exports, serialize them.
+	FIrisPackageMapExportsUtil::Serialize(Context, Value.QuantizedExports);
 
+	// Write data.
 	WritePackedUint32(Writer, Value.BitCount);
 	if (Value.BitCount > 0)
 	{
@@ -119,35 +104,10 @@ void FLastResortPropertyNetSerializer::Deserialize(FNetSerializationContext& Con
 
 	UE_NET_TRACE_DYNAMIC_NAME_SCOPE(Config->Property.Get()->GetName(), *Reader, Context.GetTraceCollector(), ENetTraceVerbosity::VeryVerbose);
 
-	// Read any object references
-	const bool bHasObjectReferences = Reader->ReadBool();
-	if (bHasObjectReferences)
-	{
-		const uint32 NumReferences = UE::Net::ReadPackedUint32(Reader);
-
-		if (NumReferences > Config->MaxAllowedObjectReferences)
-		{
-			Context.SetError(GNetError_ArraySizeTooLarge);
-			return;
-		}
-
-		Value.ObjectReferenceStorage.AdjustSize(Context, NumReferences);
-
-		FObjectNetSerializerConfig ObjectSerializerConfig;
-		for (FNetObjectReference& Ref : MakeArrayView(Value.ObjectReferenceStorage.GetData(), Value.ObjectReferenceStorage.Num()))
-		{
-			FNetDeserializeArgs ObjectArgs;
-			ObjectArgs.NetSerializerConfig = &ObjectSerializerConfig;
-			ObjectArgs.Target = NetSerializerValuePointer(&Ref);
-
-			ObjectNetSerializer->Deserialize(Context, ObjectArgs);
-		}
-	}
-	else
-	{
-		Value.ObjectReferenceStorage.Free(Context);
-	}
+	// Read exports for packagemap.
+	FIrisPackageMapExportsUtil::Deserialize(Context, Value.QuantizedExports);
 
+	// Read the data
 	const uint32 NewBitCount = ReadPackedUint32(Reader);
 	if (NewBitCount > 65535U)
 	{
@@ -159,20 +119,29 @@ void FLastResortPropertyNetSerializer::Deserialize(FNetSerializationContext& Con
 
 	Reader->ReadBitStream(static_cast<uint32*>(Value.Storage), NewBitCount);
 }
+
 void FLastResortPropertyNetSerializer::Quantize(FNetSerializationContext& Context, const FNetQuantizeArgs& Args)
 {
 	const ConfigType* Config = static_cast<const ConfigType*>(Args.NetSerializerConfig);
 	const FProperty* Property = Config->Property.Get();
 	QuantizedType& Value = *reinterpret_cast<QuantizedType*>(Args.Target);
 
-	// Capture references
+	// Setup UIrisObjectReferencePackageMap to capture exports
 	Private::FInternalNetSerializationContext* InternalContext = Context.GetInternalContext();
 	UIrisObjectReferencePackageMap* PackageMap = InternalContext ? InternalContext->PackageMap : nullptr;
-	
-	UIrisObjectReferencePackageMap::FObjectReferenceArray ObjectReferences;
+
+	UE::Net::FIrisPackageMapExports PackageMapExports;
+
 	if (PackageMap)
 	{
-		PackageMap->InitForWrite(&ObjectReferences);
+		PackageMap->InitForWrite(&PackageMapExports);
+
+		// $IRIS: $TODO: TBD: how to propagate context
+		//{
+		//	UE::Net::FNetTokenExportContext* NetTokenExportContext = PackageMap->GetNetTokenExportContext();
+		//	NetTokenExportContext->TokenStore = Context.GetNetTokenStore();
+		//	NetTokenExportContext->RemoteState = Context.GetInternalContext()->ResolveContext.RemoteNetTokenStoreState;
+		//}
 	}
 
 	// Use the Property serialization and store as binary blob.
@@ -186,25 +155,8 @@ void FLastResortPropertyNetSerializer::Quantize(FNetSerializationContext& Contex
 		return;
 	}
 
-
-	// Capture any references
-	const uint32 NumObjectReferences = ObjectReferences.Num();
-	Value.ObjectReferenceStorage.AdjustSize(Context, NumObjectReferences);
-	if (NumObjectReferences > 0)
-	{
-		FObjectNetSerializerConfig ObjectNetSerializerConfig;
-		const TObjectPtr<UObject>* SourceReferences = ObjectReferences.GetData();
-		FNetObjectReference* TargetReferences = Value.ObjectReferenceStorage.GetData();
-		for (uint32 ReferenceIndex = 0; ReferenceIndex < NumObjectReferences; ++ReferenceIndex)
-		{
-			FNetQuantizeArgs ObjectArgs;
-			ObjectArgs.NetSerializerConfig = &ObjectNetSerializerConfig;
-			ObjectArgs.Source = NetSerializerValuePointer(SourceReferences + ReferenceIndex);
-			ObjectArgs.Target = NetSerializerValuePointer(TargetReferences + ReferenceIndex);
-
-			ObjectNetSerializer->Quantize(Context, ObjectArgs);
-		}
-	}
+	// Quantize captured exports
+	FIrisPackageMapExportsUtil::Quantize(Context, PackageMapExports, Value.QuantizedExports);
 
 	// Deal with serialized data
 	AdjustStorageSize(Context, Value, static_cast<uint16>(BitCount));
@@ -221,33 +173,24 @@ void FLastResortPropertyNetSerializer::Dequantize(FNetSerializationContext& Cont
 
 	QuantizedType& Source = *reinterpret_cast<QuantizedType*>(Args.Source);
 
-	// Inject references
+	// Dequantize and inject exports
 	Private::FInternalNetSerializationContext* InternalContext = Context.GetInternalContext();
 	UIrisObjectReferencePackageMap* PackageMap = InternalContext ? InternalContext->PackageMap : nullptr;
-	UIrisObjectReferencePackageMap::FObjectReferenceArray ObjectReferences;
-
-	// References
-	const uint32 NumObjectReferences = Source.ObjectReferenceStorage.Num();
-	if (NumObjectReferences > 0U && ensureAlwaysMsgf(PackageMap, TEXT("FLastResortPropertyNetSerializer::Dequantize must have a packagemap to be able to dequantize object references. Make sure it is set in the InternalContext.")))
-	{		
-		ObjectReferences.SetNumUninitialized(NumObjectReferences);
-
-		FObjectNetSerializerConfig ObjectNetSerializerConfig;
-		const FNetObjectReference* SourceReferences = Source.ObjectReferenceStorage.GetData();
-		TObjectPtr<UObject>* TargetReferences = ObjectReferences.GetData();
-		for (uint32 ReferenceIndex = 0; ReferenceIndex < NumObjectReferences; ++ReferenceIndex)
-		{
-			FNetDequantizeArgs ObjectArgs;
-			ObjectArgs.NetSerializerConfig = &ObjectNetSerializerConfig;
-			ObjectArgs.Source = NetSerializerValuePointer(SourceReferences + ReferenceIndex);
-			ObjectArgs.Target = NetSerializerValuePointer(TargetReferences + ReferenceIndex);
 
-			ObjectNetSerializer->Dequantize(Context, ObjectArgs);
-		}
+	UE::Net::FIrisPackageMapExports PackageMapExports;
 
-		PackageMap->InitForRead(&ObjectReferences);
-	}
+	FIrisPackageMapExportsUtil::Dequantize(Context, Source.QuantizedExports, PackageMapExports);	
+
+	PackageMap->InitForRead(&PackageMapExports);
+
+	// $IRIS: $TODO: TBD: how to propagate context
+	//{
+	//	UE::Net::FNetTokenExportContext* NetTokenExportContext = PackageMap->GetNetTokenExportContext();
+	//	NetTokenExportContext->TokenStore = Context.GetNetTokenStore();
+	//	NetTokenExportContext->RemoteState = Context.GetInternalContext()->ResolveContext.RemoteNetTokenStoreState;
+	//}
 
+	// Read data
 	if (Source.BitCount)
 	{
 		FNetBitReader Archive(PackageMap, static_cast<uint8*>(Source.Storage), Source.BitCount);
@@ -265,12 +208,12 @@ bool FLastResortPropertyNetSerializer::IsEqual(FNetSerializationContext& Context
 	{
 		const QuantizedType& Value0 = *reinterpret_cast<const QuantizedType*>(Args.Source0);
 		const QuantizedType& Value1 = *reinterpret_cast<const QuantizedType*>(Args.Source1);
-		if ((Value0.BitCount != Value1.BitCount) || (Value0.ObjectReferenceStorage.Num() != Value1.ObjectReferenceStorage.Num()))
+		if ((Value0.BitCount != Value1.BitCount))
 		{
 			return false;
 		}
 
-		if (Value0.ObjectReferenceStorage.Num() > 0 && FMemory::Memcmp(Value0.ObjectReferenceStorage.GetData(), Value1.ObjectReferenceStorage.GetData(), sizeof(FNetObjectReference) * Value0.ObjectReferenceStorage.Num()) != 0)
+		if (!FIrisPackageMapExportsUtil::IsEqual(Context, Value0.QuantizedExports, Value1.QuantizedExports))
 		{
 			return false;
 		}
@@ -295,6 +238,9 @@ void FLastResortPropertyNetSerializer::CloneDynamicState(FNetSerializationContex
 	QuantizedType& Target = *reinterpret_cast<QuantizedType*>(Args.Target);
 	const QuantizedType& Source = *reinterpret_cast<const QuantizedType*>(Args.Source);
 
+	// Clone captured exports
+	FIrisPackageMapExportsUtil::CloneDynamicState(Context, Target.QuantizedExports, Source.QuantizedExports);
+
 	const uint16 ByteCount = static_cast<uint16>(Align((Source.BitCount + 7U)/8U, AllocationAlignment));
 
 	void* Storage = nullptr;
@@ -307,7 +253,6 @@ void FLastResortPropertyNetSerializer::CloneDynamicState(FNetSerializationContex
 	Target.BitCount = Source.BitCount;
 	Target.Storage = Storage;
 
-	Target.ObjectReferenceStorage.Clone(Context, Source.ObjectReferenceStorage);
 }
 
 void FLastResortPropertyNetSerializer::FreeDynamicState(FNetSerializationContext& Context, const FNetFreeDynamicStateArgs& Args)
@@ -318,7 +263,10 @@ void FLastResortPropertyNetSerializer::FreeDynamicState(FNetSerializationContext
 void FLastResortPropertyNetSerializer::FreeDynamicStateInternal(FNetSerializationContext& Context, QuantizedType& Value)
 {
 	// Clear all info
-	Value.ObjectReferenceStorage.Free(Context);
+
+	// Free captured export data.
+	FIrisPackageMapExportsUtil::FreeDynamicState(Context, Value.QuantizedExports);
+	
 	Context.GetInternalContext()->Free(Value.Storage);
 
 	Value.BitCount = 0;
@@ -381,11 +329,7 @@ void FLastResortPropertyNetSerializer::CollectNetReferences(FNetSerializationCon
 	const QuantizedType& Value = *reinterpret_cast<const QuantizedType*>(Args.Source);
 	FNetReferenceCollector& Collector = *reinterpret_cast<UE::Net::FNetReferenceCollector*>(Args.Collector);
 
-	const FNetReferenceInfo ReferenceInfo(FNetReferenceInfo::EResolveType::ResolveOnClient);	
-	for (const FNetObjectReference& Ref : MakeArrayView(Value.ObjectReferenceStorage.GetData(), Value.ObjectReferenceStorage.Num()))
-	{
-		Collector.Add(ReferenceInfo, Ref, Args.ChangeMaskInfo);
-	}
+	FIrisPackageMapExportsUtil::CollectNetReferences(Context, Value.QuantizedExports, Args.ChangeMaskInfo, Collector);
 }
 
 
diff --git a/Engine/Source/Runtime/Experimental/Iris/Core/Private/Iris/Serialization/NetSerializationContext.cpp b/Engine/Source/Runtime/Experimental/Iris/Core/Private/Iris/Serialization/NetSerializationContext.cpp
--- a/Engine/Source/Runtime/Experimental/Iris/Core/Private/Iris/Serialization/NetSerializationContext.cpp
+++ b/Engine/Source/Runtime/Experimental/Iris/Core/Private/Iris/Serialization/NetSerializationContext.cpp
@@ -80,4 +80,29 @@ UObject* FNetSerializationContext::GetLocalConnectionUserData(uint32 ConnectionI
 	return UserData;
 }
 
+const FNetTokenStore* FNetSerializationContext:: GetNetTokenStore() const
+{
+	if (InternalContext == nullptr)
+	{
+		return nullptr;
+	}
+
+	const UReplicationSystem* ReplicationSystem = InternalContext->ReplicationSystem;
+
+	return ReplicationSystem ? ReplicationSystem->GetNetTokenStore() : nullptr;
+}
+
+FNetTokenStore* FNetSerializationContext::GetNetTokenStore()
+{
+	if (InternalContext == nullptr)
+	{
+		return nullptr;
+	}
+
+	UReplicationSystem* ReplicationSystem = InternalContext->ReplicationSystem;
+
+	return ReplicationSystem ? ReplicationSystem->GetNetTokenStore() : nullptr;
+}
+
+
 }
diff --git a/Engine/Source/Runtime/Experimental/Iris/Core/Private/Iris/Serialization/StringNetSerializers.cpp b/Engine/Source/Runtime/Experimental/Iris/Core/Private/Iris/Serialization/StringNetSerializers.cpp
--- a/Engine/Source/Runtime/Experimental/Iris/Core/Private/Iris/Serialization/StringNetSerializers.cpp
+++ b/Engine/Source/Runtime/Experimental/Iris/Core/Private/Iris/Serialization/StringNetSerializers.cpp
@@ -51,6 +51,8 @@ struct FNameNetSerializer
 		void* ElementStorage;
 	};
 
+	static_assert(GetNameNetSerializerSafeQuantizedSize() >= sizeof(FQuantizedType));
+
 	typedef FName SourceType;
 	typedef FQuantizedType QuantizedType;
 	typedef FNameNetSerializerConfig ConfigType;
@@ -435,11 +437,10 @@ struct FNameAsNetTokenNetSerializerQuantizedType
 	uint32 bIsString : 1U;
 	// When serializing/deserializing the number is expressed as (MAX_int32 - number)
 	uint32 bEncodeNumberFromIntMax : 1U;
-
-	// if this is a local (obtained from quantize)
-	uint32 bIsLocal : 1U;
 };
 
+static_assert(GetNameNetSerializerSafeQuantizedSize() >= sizeof(FNameAsNetTokenNetSerializerQuantizedType));
+
 } // End of namespace UE::Net
 
 template<> struct TIsPODType<UE::Net::FNameAsNetTokenNetSerializerQuantizedType> { enum { Value = true }; };
@@ -475,10 +476,7 @@ struct FNameAsNetTokenNetSerializer
 private:
 
 	// Utility methods, consolidate with other changes to NetTokenStore as next step.
-	static FName ResolveNetToken(FNetSerializationContext&, FNetToken NetToken, bool bIsLocal);
-
-	static void AppendExportOrWriteInlinedExportData(FNetSerializationContext&, FNetToken NetToken);
-	static void ReadInlinedExportData(FNetSerializationContext&, FNetToken NetToken);
+	static FName ResolveNetToken(FNetSerializationContext&, FNetToken NetToken);
 
 	static const uint32 BitCountNeededForEName;
 };
@@ -503,13 +501,11 @@ void FNameAsNetTokenNetSerializer::Serialize(FNetSerializationContext& Context,
 	FNetBitStreamWriter* Writer = Context.GetBitStreamWriter();
 	if (Writer->WriteBool(Value.bIsString))
 	{
-		//UE_NET_TRACE_DYNAMIC_NAME_SCOPE(*FString::Printf(TEXT("%u"), Value.NetToken.GetIndex()), *Context.GetBitStreamWriter(), Context.GetTraceCollector(), ENetTraceVerbosity::VeryVerbose);
-
 		// Always write the token
-		WriteNetToken(Writer, Value.NetToken);
+		WriteNetToken(Context, Value.NetToken);
 
 		// Export or add to pending exports for later export
-		FNameAsNetTokenNetSerializer::AppendExportOrWriteInlinedExportData(Context, Value.NetToken);
+		FNetTokenStore::AppendExportOrWriteInlinedExportData(Context, Value.NetToken);
 	}
 	else
 	{
@@ -535,7 +531,7 @@ void FNameAsNetTokenNetSerializer::Deserialize(FNetSerializationContext& Context
 		Target.bIsString = 1;
 
 		// Always Read the token
-		FNetToken NetToken = ReadNetToken(Reader);
+		FNetToken NetToken = ReadNetToken(Context);
 
 		if (Reader->IsOverflown())
 		{
@@ -543,10 +539,9 @@ void FNameAsNetTokenNetSerializer::Deserialize(FNetSerializationContext& Context
 		}
 
 		// Read inlined exports if there are any
-		ReadInlinedExportData(Context, NetToken);
+		FNetTokenStore::ReadInlinedExportData(Context, NetToken);
 
 		Target.NetToken = NetToken;
-		Target.bIsLocal = 0U;
 	}
 	else
 	{
@@ -562,7 +557,6 @@ void FNameAsNetTokenNetSerializer::Deserialize(FNetSerializationContext& Context
 		Target.ENameOrNumber = ENameNumber;
 	
 		Target.NetToken = FNetToken();
-		Target.bIsLocal = 0U;
 	}
 }
 
@@ -582,11 +576,8 @@ void FNameAsNetTokenNetSerializer::Quantize(FNetSerializationContext& Context, c
 		TargetName.bEncodeNumberFromIntMax = 0;
 
 		// Store as NetToken
-		FInternalNetSerializationContext* InternalContext = Context.GetInternalContext();
-		using namespace UE::Net::Private;
-			
-		TargetName.NetToken = InternalContext->ReplicationSystem->GetNameTokenStore()->GetOrCreateToken(SourceName);
-		TargetName.bIsLocal = 1U;
+		FNameTokenStore* NameTokenStore = Context.GetNetTokenStore()->GetDataStore<FNameTokenStore>();
+		TargetName.NetToken = NameTokenStore->GetOrCreateToken(SourceName);
 	}
 	else
 	{
@@ -596,53 +587,17 @@ void FNameAsNetTokenNetSerializer::Quantize(FNetSerializationContext& Context, c
 		// The EName we do care about!
 		TargetName.ENameOrNumber = static_cast<int32>(static_cast<uint32>(*AsEName));
 
-		TargetName.bIsLocal = 0;
 		TargetName.NetToken = FNetToken();
 	}
 }
 
-FName FNameAsNetTokenNetSerializer::ResolveNetToken(FNetSerializationContext& Context, FNetToken NetToken, bool bIsLocal)
+FName FNameAsNetTokenNetSerializer::ResolveNetToken(FNetSerializationContext& Context, FNetToken NetToken)
 {
 	using namespace UE::Net::Private;
 
-	FName Result;
 	FInternalNetSerializationContext* InternalContext = Context.GetInternalContext();
-	FNameTokenStore* NameTokenStore = InternalContext->ReplicationSystem->GetNameTokenStore();
-
-	return bIsLocal ? NameTokenStore->ResolveToken(NetToken) : NameTokenStore->ResolveRemoteToken(NetToken, *InternalContext->ResolveContext.RemoteNetTokenStoreState);
-}
-
-void FNameAsNetTokenNetSerializer::AppendExportOrWriteInlinedExportData(FNetSerializationContext& Context, FNetToken NetToken)
-{
-	using namespace UE::Net::Private;
-
-	FNetExportContext* ExportContext = Context.GetExportContext();
-	const FInternalNetSerializationContext* InternalContext = Context.GetInternalContext();
-	
-	if (InternalContext->bInlineObjectReferenceExports == 0U)
-	{
-		ExportContext->AddPendingExport(NetToken);
-	}
-	else
-	{
-		FNetTokenStore& NetTokenStore = InternalContext->ReplicationSystem->GetReplicationSystemInternal()->GetNetTokenStore();
-
-		NetTokenStore.ConditionalWriteNetTokenData(Context, ExportContext, NetToken);
-	}
-}
-
-void FNameAsNetTokenNetSerializer::ReadInlinedExportData(FNetSerializationContext& Context, FNetToken NetToken)
-{
-	using namespace UE::Net::Private;
-
-	const FInternalNetSerializationContext* InternalContext = Context.GetInternalContext();
-
-	// Conditionally import if using inlined exports
-	if (InternalContext->bInlineObjectReferenceExports == 1U)
-	{
-		FNetTokenStore& NetTokenStore = InternalContext->ReplicationSystem->GetReplicationSystemInternal()->GetNetTokenStore();
-		NetTokenStore.ConditionalReadNetTokenData(Context, NetToken);
-	}
+	FNameTokenStore* NameTokenStore = Context.GetNetTokenStore()->GetDataStore<FNameTokenStore>();
+	return NameTokenStore->ResolveToken(NetToken, InternalContext->ResolveContext.RemoteNetTokenStoreState);
 }
 
 void FNameAsNetTokenNetSerializer::Dequantize(FNetSerializationContext& Context, const FNetDequantizeArgs& Args)
@@ -655,7 +610,7 @@ void FNameAsNetTokenNetSerializer::Dequantize(FNetSerializationContext& Context,
 	if (Source.bIsString)
 	{
 		// Resolve from RemoteNetToken
-		Target = FNameAsNetTokenNetSerializer::ResolveNetToken(Context, Source.NetToken, Source.bIsLocal);
+		Target = FNameAsNetTokenNetSerializer::ResolveNetToken(Context, Source.NetToken);
 	}
 	else
 	{
@@ -679,14 +634,14 @@ bool FNameAsNetTokenNetSerializer::IsEqual(FNetSerializationContext& Context, co
 			return false;
 		}
 
-		// We are pessimistic here
+		// Need to compare actual FNames to properly compare non-auth and auth token
 		if (Value0.bIsString)
 		{
-			if (Value0.bIsLocal != Value1.bIsLocal)
+			if (Value0.NetToken.IsAssignedByAuthority() != Value1.NetToken.IsAssignedByAuthority())
 			{
-				const FName Name0 = FNameAsNetTokenNetSerializer::ResolveNetToken(Context, Value0.NetToken, Value0.bIsLocal);
-				const FName Name1 = FNameAsNetTokenNetSerializer::ResolveNetToken(Context, Value1.NetToken, Value1.bIsLocal);;
-				// Need to compare actual FNames to handle compare of local and remote data.
+				const FName Name0 = FNameAsNetTokenNetSerializer::ResolveNetToken(Context, Value0.NetToken);
+				const FName Name1 = FNameAsNetTokenNetSerializer::ResolveNetToken(Context, Value1.NetToken);
+				
 				if (Name0 != Name1)
 				{
 					return false;
diff --git a/Engine/Source/Runtime/Experimental/Iris/Core/Public/Iris/ReplicationSystem/NameTokenStore.h b/Engine/Source/Runtime/Experimental/Iris/Core/Public/Iris/ReplicationSystem/NameTokenStore.h
--- a/Engine/Source/Runtime/Experimental/Iris/Core/Public/Iris/ReplicationSystem/NameTokenStore.h
+++ b/Engine/Source/Runtime/Experimental/Iris/Core/Public/Iris/ReplicationSystem/NameTokenStore.h
@@ -18,26 +18,27 @@ class FNameTokenStore : public FNetTokenDataStore
 public:
 	explicit FNameTokenStore(FNetTokenStore& TokenStore);
 
-	// Create a string token for the provided string
+	// Create a NetToken for the provided name
 	IRISCORE_API FNetToken GetOrCreateToken(FName Name);
 
-	// Resolve a local token
-	IRISCORE_API FName ResolveToken(FNetToken Token) const;
+	// Resolve NetToken, to resolve remote tokens RemoteTokenStoreState must be valid
+	IRISCORE_API FName ResolveToken(FNetToken Token, const FNetTokenStoreState* RemoteTokenStoreState = nullptr) const;	
 
-	// Resolve a token received from remote
-	IRISCORE_API FName ResolveRemoteToken(FNetToken Token, const FNetTokenStoreState& NetTokenStoreState) const;
+	static FName GetTokenStoreName() { return NameTokenStoreName; }
 
 protected:
 	// Serialize data for a token, note there is not validation in this function
 	virtual void WriteTokenData(FNetSerializationContext& Context, FNetTokenStoreKey TokenStoreKey) const override;
 
 	// Read data for a token, returns a valid StoreKey if successful read
-	virtual FNetTokenStoreKey ReadTokenData(FNetSerializationContext& Context) override;
+	virtual FNetTokenStoreKey ReadTokenData(FNetSerializationContext& Context, const FNetToken& NetToken) override;
 
 	// Create a persistent string
 	FNetTokenStoreKey GetOrCreateTokenStoreKey(FName Name);
 
 private:
+	inline static FName NameTokenStoreName = TEXT("NameTokenStore");
+
 	FNetTokenStore& TokenStore;
 	TMap<FName, FNetTokenStoreKey> FNameToKey;
 	TArray<FName> StoredFNames;
diff --git a/Engine/Source/Runtime/Experimental/Iris/Core/Public/Iris/ReplicationSystem/NetToken.h b/Engine/Source/Runtime/Experimental/Iris/Core/Public/Iris/ReplicationSystem/NetToken.h
--- a/Engine/Source/Runtime/Experimental/Iris/Core/Public/Iris/ReplicationSystem/NetToken.h
+++ b/Engine/Source/Runtime/Experimental/Iris/Core/Public/Iris/ReplicationSystem/NetToken.h
@@ -8,8 +8,7 @@
 
 namespace UE::Net
 {
-	class FNetBitStreamReader;
-	class FNetBitStreamWriter;
+	class FNetSerializationContext;
 }
 
 namespace UE::Net
@@ -20,6 +19,8 @@ class FNetToken
 public:
 	typedef uint32 FTypeId;
 
+	enum : uint32 { Invalid = 0U };
+
 	enum : uint32 { InvalidTokenTypeId = ~FTypeId(0) };
 	enum : uint32 { InvalidTokenIndex = 0U };
 	enum : uint32 { TokenTypeIdBits = 2U };
@@ -27,21 +28,42 @@ public:
 	enum : uint32 { MaxTypeIdCount = 1U << TokenTypeIdBits };
 	enum : uint32 { MaxNetTokenCount = 1U << TokenBits };
 
+	enum class ENetTokenAuthority
+	{
+		None,
+		Authority,
+	};
+
 public:	
-	FNetToken() : Index(InvalidTokenIndex) {}
-	inline bool IsValid() const { return Index != InvalidTokenIndex; }
+	FNetToken() : Value(Invalid) {}
+	inline bool IsValid() const { return Value != Invalid; }
+	bool IsAssignedByAuthority() const { return bIsAssignedByAuthority != 0U; }
 	uint32 GetIndex() const { return Index; }
-	bool operator==(const FNetToken& Other) const { return Index == Other.Index; }
+	bool operator==(const FNetToken& Other) const { return IsAssignedByAuthority() == Other.IsAssignedByAuthority() && Index == Other.Index; }
 	FString ToString() const;
-	static FNetToken MakeNetToken(uint32 Index) { check(Index < MaxNetTokenCount); return FNetToken(Index); }
+
+	friend uint32 GetTypeHash(const FNetToken& Token)
+	{
+		return ::GetTypeHash(Token.Value);
+	}
+
+	static FNetToken MakeNetToken(uint32 Index, ENetTokenAuthority Authority) { check(Index < MaxNetTokenCount); return FNetToken(Index, Authority); }
 
 private:
-	explicit FNetToken(uint32 InIndex) : Index(InIndex) {}
+	explicit FNetToken(uint32 InIndex, ENetTokenAuthority Authority) { Index = InIndex, bIsAssignedByAuthority = Authority == ENetTokenAuthority::Authority ? 1U : 0U; }
 
 private:
-	friend IRISCORE_API FNetToken ReadNetToken(FNetBitStreamReader*);
+	friend IRISCORE_API FNetToken ReadNetToken(UE::Net::FNetSerializationContext&);
 
-	uint32 Index;
+	union 
+	{
+		struct
+		{
+			uint32 Index : TokenBits;
+			uint32 bIsAssignedByAuthority : 1U;
+		};
+		uint32 Value;
+	};
 };
 
 class FNetTokenStoreKey
@@ -73,21 +95,16 @@ private:
 
 static_assert(sizeof(FNetTokenStoreKey) == sizeof(uint32), "FNetTokenKey should fit in a uint32");
 
-FORCEINLINE uint32 GetTypeHash(const FNetToken& Token)
-{
-	return ::GetTypeHash(Token.GetIndex());
-}
-
 inline FString FNetToken::ToString() const
 {
 	FString Result;
-	Result = FString::Printf(TEXT("NetToken (Index=%u)"), Index);
+	Result = FString::Printf(TEXT("NetToken (Auth:%u Index=%u)"), IsAssignedByAuthority(), Index);
 	return Result;
 }
 
 // Read and write tokens
-IRISCORE_API FNetToken ReadNetToken(FNetBitStreamReader* Reader);
-IRISCORE_API void WriteNetToken(FNetBitStreamWriter* Writer, FNetToken Token);
+IRISCORE_API FNetToken ReadNetToken(UE::Net::FNetSerializationContext& Context);
+IRISCORE_API void WriteNetToken(UE::Net::FNetSerializationContext& Context, FNetToken Token);
 
 }
 
diff --git a/Engine/Source/Runtime/Experimental/Iris/Core/Public/Iris/ReplicationSystem/NetTokenStore.h b/Engine/Source/Runtime/Experimental/Iris/Core/Public/Iris/ReplicationSystem/NetTokenStore.h
--- a/Engine/Source/Runtime/Experimental/Iris/Core/Public/Iris/ReplicationSystem/NetTokenStore.h
+++ b/Engine/Source/Runtime/Experimental/Iris/Core/Public/Iris/ReplicationSystem/NetTokenStore.h
@@ -4,6 +4,9 @@
 
 #include "CoreTypes.h"
 #include "NetToken.h"
+#include "UObject/NameTypes.h"
+#include "Containers/Array.h"
+#include "Templates/Tuple.h"
 
 class UNetTokenDataStream;
 namespace UE::Net
@@ -26,20 +29,19 @@ namespace UE::Net
 class FNetTokenDataStore
 {
 public:
-	
 	virtual ~FNetTokenDataStore();
 
 protected:
 	FNetTokenDataStore();
 
 	virtual void WriteTokenData(FNetSerializationContext& Context, FNetTokenStoreKey Key) const = 0;
-	virtual FNetTokenStoreKey ReadTokenData(FNetSerializationContext& Context) = 0;
+	virtual FNetTokenStoreKey ReadTokenData(FNetSerializationContext& Context, const FNetToken& NetToken) = 0;
 
 	FNetTokenStoreKey GetTokenKey(FNetToken Token, const FNetTokenStoreState& TokenStoreState) const;
 	inline FNetToken::FTypeId GetTypeId() const { return TypeId; }
 
 	// Create new NetToken
-	FNetToken CreateToken(FNetTokenStoreKey Key, FNetTokenStoreState& TokenStoreState);
+	FNetToken CreateToken(FNetToken::ENetTokenAuthority Authority, FNetTokenStoreKey Key, FNetTokenStoreState& TokenStoreState);
 
 	// Make NetTokenStoreKey
 	static FNetTokenStoreKey MakeNetTokenStoreKey(FNetToken::FTypeId TokenTypeId, uint32 TokenKeyValue);
@@ -49,7 +51,7 @@ private:
 	FNetToken::FTypeId TypeId;
 };
 
-// This is the token store, we have one per game instance
+// This is the token store, we currently have one per ReplicationSystem but it is possible we will share this across game instance.
 class FNetTokenStore
 {
 	UE_NONCOPYABLE(FNetTokenStore);
@@ -57,38 +59,58 @@ public:
 	FNetTokenStore();
 	~FNetTokenStore();
 
+	/** External configuration variables used to initialize the NetTokenStore */
+	struct FInitParams
+	{
+		FNetToken::ENetTokenAuthority Authority;
+	};
+	void Init(FInitParams& InitParams);
+
+	bool IsAuthority() const { return Params.Authority == FNetToken::ENetTokenAuthority::Authority; }
+
+	// A token is local if the authority of the NetTokenStore and the token matches, Invalid tokens are always local.
+	bool IsLocalToken(const FNetToken NetToken) const { return !NetToken.IsValid() || IsAuthority() == NetToken.IsAssignedByAuthority(); }
+
 	// Register DataStore and return true if it was registered.
-	bool RegisterDataStore(FNetTokenDataStore* DataStore);
+	bool RegisterDataStore(FNetTokenDataStore* DataStore, FName TokenStoreName);
+
+	IRISCORE_API const FNetTokenDataStore* GetDataStore(FName Name) const;
+	IRISCORE_API FNetTokenDataStore* GetDataStore(FName Name);
+
+	/** Return data store of specified type. */
+	template<typename T>
+	T* GetDataStore() { return static_cast<T*>(GetDataStore(T::GetTokenStoreName())); }
+
+	/** Return data store of specified type. */
+	template<typename T>
+	const T* GetDataStore() const { return static_cast<const T*>(GetDataStore(T::GetTokenStoreName())); }
 
 	const FNetTokenStoreState* GetLocalNetTokenStoreState() const { return LocalNetTokenStoreState; }
 	FNetTokenStoreState* GetLocalNetTokenStoreState() { return LocalNetTokenStoreState; }
 
 	// Write data associated with the NetToken
-	void WriteTokenData(FNetSerializationContext& Context, const FNetToken& NetToken) const;
+	void WriteTokenData(FNetSerializationContext& Context, const FNetToken NetToken) const;
 
 	// Read data associated with the NetToken
-	void ReadTokenData(FNetSerializationContext& Context, const FNetToken& NetToken, FNetTokenStoreState& RemoteNetTokenStoreState);
+	void ReadTokenData(FNetSerializationContext& Context, const FNetToken NetToken, FNetTokenStoreState& RemoteNetTokenStoreState);
 
 	// Conditionally write NetTokenData unless already exported
-	void ConditionalWriteNetTokenData(FNetSerializationContext& Context, Private::FNetExportContext* ExportContext, const FNetToken& NetToken) const;
+	void ConditionalWriteNetTokenData(FNetSerializationContext& Context, Private::FNetExportContext* ExportContext, const FNetToken NetToken) const;
 
 	// Conditionally read NetTokenData if exported
-	void ConditionalReadNetTokenData(FNetSerializationContext& Context, const FNetToken& NetToken);
-
-private:
-
-	// Write token data for the TokenIndex
-	void WriteTokenDataForIndex(FNetSerializationContext& Context, uint32 TokenIndex) const;
+	void ConditionalReadNetTokenData(FNetSerializationContext& Context, const FNetToken NetToken);
 
-	// Read token data for the TokenIndex
-	void ReadTokenDataForIndex(FNetSerializationContext& Context, uint32 TokenIndex, FNetTokenStoreState& Resolver);
+	// Utility methods, consolidate with other changes to NetTokenStore as next step.
+	IRISCORE_API static void AppendExportOrWriteInlinedExportData(FNetSerializationContext&, FNetToken NetToken);
+	IRISCORE_API static void ReadInlinedExportData(FNetSerializationContext&, FNetToken NetToken);
 
 private:
 	friend UNetTokenDataStream;
 	friend FNetTokenDataStore;
 
 	FNetTokenStoreState* LocalNetTokenStoreState;
-	TArray<FNetTokenDataStore*> TokenDataStores;
+	TArray<TTuple<FName, FNetTokenDataStore*> > TokenDataStores;
+	FInitParams Params;
 };
 
 inline FNetTokenStoreKey FNetTokenDataStore::MakeNetTokenStoreKey(FNetToken::FTypeId TokenTypeId, uint32 TokenKeyValue)
diff --git a/Engine/Source/Runtime/Experimental/Iris/Core/Public/Iris/ReplicationSystem/ReplicationBridge.h b/Engine/Source/Runtime/Experimental/Iris/Core/Public/Iris/ReplicationSystem/ReplicationBridge.h
--- a/Engine/Source/Runtime/Experimental/Iris/Core/Public/Iris/ReplicationSystem/ReplicationBridge.h
+++ b/Engine/Source/Runtime/Experimental/Iris/Core/Public/Iris/ReplicationSystem/ReplicationBridge.h
@@ -29,6 +29,8 @@ namespace UE::Net
 	struct FReplicationInstanceProtocol;
 	struct FReplicationProtocol;
 	class FNetSerializationContext;
+	class FNetBitStreamReader;
+	class FNetBitStreamWriter;
 	namespace Private
 	{
 		typedef uint32 FInternalNetRefIndex;
diff --git a/Engine/Source/Runtime/Experimental/Iris/Core/Public/Iris/ReplicationSystem/ReplicationSystem.h b/Engine/Source/Runtime/Experimental/Iris/Core/Public/Iris/ReplicationSystem/ReplicationSystem.h
--- a/Engine/Source/Runtime/Experimental/Iris/Core/Public/Iris/ReplicationSystem/ReplicationSystem.h
+++ b/Engine/Source/Runtime/Experimental/Iris/Core/Public/Iris/ReplicationSystem/ReplicationSystem.h
@@ -32,6 +32,7 @@ namespace UE::Net
 	class FNetObjectAttachment;	
 	class FNetObjectReference;
 	class FReplicationSystemFactory;
+	class FNetTokenStore;
 	class FStringTokenStore;
 	class FNameTokenStore;
 	class FWorldLocations;
@@ -47,6 +48,8 @@ namespace UE::Net
 	enum class ENetRefHandleError : uint32;
 	enum class EReplicationCondition : uint32;
 
+	struct FNetTokenExportContext;
+
 	namespace Private
 	{
 		class FReplicationSystemImpl;
@@ -330,29 +333,22 @@ public:
 	T* GetReplicationBridgeAs() const { return Cast<T>(GetReplicationBridge()); }
 
 	/**
-	 * @return A const version of the string token store.
-	 * @see UE::Net::FStringTokenStore
+	 * @return The Net token store.
+	 * @see UE::Net::FNetTokenStore
 	 */
-	IRISCORE_API const UE::Net::FStringTokenStore* GetStringTokenStore() const;
+	IRISCORE_API UE::Net::FNetTokenStore* GetNetTokenStore();
 
 	/**
-	 * @return The string token store.
-	 * @see UE::Net::FStringTokenStore
+	 * @return A const version of the NetTokenStore.
+	 * @see UE::Net::FNetTokenStore
 	 */
-	IRISCORE_API UE::Net::FStringTokenStore* GetStringTokenStore();
+	IRISCORE_API const UE::Net::FNetTokenStore* GetNetTokenStore() const;
 
 	/**
-	 * @return A const version of the FName token store.
-	 * @see UE::Net::FFNameTokenStore
+	 * Init NetTokenExportContext required to quantize and resolve NetTokens
+	 * @see UE::Net::FNetTokenStore
 	 */
-	IRISCORE_API const UE::Net::FNameTokenStore* GetNameTokenStore() const;
-
-	/**
-	 * @return The string token store.
-	 * @see UE::Net::FNameTokenStore
-	 */
-	IRISCORE_API UE::Net::FNameTokenStore* GetNameTokenStore();
-
+	IRISCORE_API void InitNetTokenExportContext(UE::Net::FNetTokenExportContext& NetTokenExportContext, uint32 ConnectionId);
 
 	/**
 	 * Check whether a FNetRefHandle is still associated with a replicated object.
diff --git a/Engine/Source/Runtime/Experimental/Iris/Core/Public/Iris/ReplicationSystem/StringTokenStore.h b/Engine/Source/Runtime/Experimental/Iris/Core/Public/Iris/ReplicationSystem/StringTokenStore.h
--- a/Engine/Source/Runtime/Experimental/Iris/Core/Public/Iris/ReplicationSystem/StringTokenStore.h
+++ b/Engine/Source/Runtime/Experimental/Iris/Core/Public/Iris/ReplicationSystem/StringTokenStore.h
@@ -22,23 +22,27 @@ public:
 	IRISCORE_API FNetToken GetOrCreateToken(const FString& String);
 	IRISCORE_API FNetToken GetOrCreateToken(const TCHAR* Name, uint32 Length);
 
-	// Resolve a local token
-	IRISCORE_API const TCHAR* ResolveToken(FNetToken Token) const;
+	// Resolve NetToken, to resolve remote tokens RemoteTokenStoreState must be valid
+	IRISCORE_API const TCHAR* ResolveToken(FNetToken Token, const FNetTokenStoreState* RemoteTokenStoreState = nullptr) const;	
 
 	// Resolve a token received from remote
-	IRISCORE_API const TCHAR* ResolveRemoteToken(FNetToken Token, const FNetTokenStoreState& NetTokenStoreState) const;
+	IRISCORE_API const TCHAR* ResolveRemoteToken(FNetToken Token, const FNetTokenStoreState& NetTokenStoreState) const { return ResolveToken(Token, &NetTokenStoreState); }
+
+	static FName GetTokenStoreName() { return StringTokenStoreName; }
 
 protected:
 	// Serialize data for a token, note there is not validation in this function
 	virtual void WriteTokenData(FNetSerializationContext& Context, FNetTokenStoreKey TokenStoreKey) const override;
 
 	// Read data for a token, returns a valid StoreKey if successful read
-	virtual FNetTokenStoreKey ReadTokenData(FNetSerializationContext& Context) override;
+	virtual FNetTokenStoreKey ReadTokenData(FNetSerializationContext& Context, const FNetToken& NetToken) override;
 
 	// Create a persistent string
 	FNetTokenStoreKey GetOrCreatePersistentString(const TCHAR* Name, uint32 Length);
 
 private:
+	inline static FName StringTokenStoreName = TEXT("StringTokenStore");
+
 	FNetTokenStore& TokenStore;
 	TMap<uint64, FNetTokenStoreKey> HashToKey;
 	TArray<const TCHAR*> StoredStrings;
diff --git a/Engine/Source/Runtime/Experimental/Iris/Core/Public/Iris/Serialization/IrisObjectReferencePackageMap.h b/Engine/Source/Runtime/Experimental/Iris/Core/Public/Iris/Serialization/IrisObjectReferencePackageMap.h
--- a/Engine/Source/Runtime/Experimental/Iris/Core/Public/Iris/Serialization/IrisObjectReferencePackageMap.h
+++ b/Engine/Source/Runtime/Experimental/Iris/Core/Public/Iris/Serialization/IrisObjectReferencePackageMap.h
@@ -3,10 +3,44 @@
 #pragma once
 
 #include "UObject/CoreNet.h"
+#include "Iris/ReplicationSystem/NetToken.h"
 #include "IrisObjectReferencePackageMap.generated.h"
 
+// Forward declarations
 class FNetworkGUID;
 
+namespace UE::Net
+{
+	class FNetTokenStoreState;
+	class FNetTokenStore;
+
+	struct FNetTokenExportContext
+	{
+		FNetTokenStore* TokenStore = nullptr;
+		FNetTokenStoreState* RemoteState = nullptr;		
+	};
+
+	// In order to properly capture exported data when calling in to old style NetSerialize methods
+	// we need to capture and inject certain types.
+	struct FIrisPackageMapExports
+	{
+		typedef TArray<TObjectPtr<UObject>, TInlineAllocator<4>> FObjectReferenceArray;
+		typedef TArray<FName, TInlineAllocator<4>> FNameArray;
+		typedef TArray<UE::Net::FNetToken, TInlineAllocator<4>> FNetTokensArray;
+
+		void Reset()
+		{
+			References.Reset();
+			Names.Reset();
+			NetTokens.Reset();
+		}
+
+		FObjectReferenceArray References;
+		FNameArray Names;
+		FNetTokensArray NetTokens;
+	};
+}
+
 /**
  * Custom packagemap implementation used to be able to capture UObject* references from external serialization.
  * Any object references written when using this packagemap will be added to the References array and serialized as an index.
@@ -19,17 +53,31 @@ public:
 	GENERATED_BODY()
 
 	typedef TArray<TObjectPtr<UObject>, TInlineAllocator<4>> FObjectReferenceArray;
+	typedef TArray<FName, TInlineAllocator<4>> FNameArray;
+	typedef TArray<UE::Net::FNetToken, TInlineAllocator<4>> FNetTokensArray;
 
-	// We override SerialzierObject in order to be able to capture object references
+	// We override SerializeObject in order to be able to capture object references
 	virtual bool SerializeObject(FArchive& Ar, UClass* InClass, UObject*& Obj, FNetworkGUID* OutNetGUID) override;
 
+	// Override SerializeName in order to be able to capture name and serialize them with iris instead.
+	virtual bool SerializeName(FArchive& Ar, FName& InName);
+
+	// Allow serialization of NetTokens to be able to support export of custom token types from existing NetSerialize implementations.
+	IRISCORE_API bool SerializeNetToken(FArchive& Ar, UE::Net::FNetToken& InNetToken);
+
 	// Init for read, we need a reference array to be able to resolve references
-	IRISCORE_API void InitForRead(const FObjectReferenceArray* InReferences);
+	IRISCORE_API void InitForRead(const UE::Net::FIrisPackageMapExports* PackageMapExports);
+
+	// Init for write, all captured exports will be serialized as in index and added to the PackageMapExports for later export using iris.
+	IRISCORE_API void InitForWrite(UE::Net::FIrisPackageMapExports* PackageMapExports);
 
-	// Init for write, all written references will be added to the references array which will be resetted by this call
-	IRISCORE_API void InitForWrite(FObjectReferenceArray* InReferences);
+	// Gets NetTokenExportContext to allow access to NetTokenStore
+	UE::Net::FNetTokenExportContext* GetNetTokenExportContext() { return &NetTokenExportContext; }
+
+	// Gets NetTokenExportContext to allow access to NetTokenStore
+	const UE::Net::FNetTokenExportContext* GetNetTokenExportContext() const { return &NetTokenExportContext; }
 
 protected:
-	FObjectReferenceArray* References = nullptr;
+	UE::Net::FIrisPackageMapExports* PackageMapExports = nullptr;
+	UE::Net::FNetTokenExportContext NetTokenExportContext;
 };
-
diff --git a/Engine/Source/Runtime/Experimental/Iris/Core/Public/Iris/Serialization/IrisPackageMapExportUtil.h b/Engine/Source/Runtime/Experimental/Iris/Core/Public/Iris/Serialization/IrisPackageMapExportUtil.h
new file mode 100644
--- /dev/null
+++ b/Engine/Source/Runtime/Experimental/Iris/Core/Public/Iris/Serialization/IrisPackageMapExportUtil.h
@@ -0,0 +1,71 @@
+// Copyright Epic Games, Inc. All Rights Reserved.
+#pragma once
+
+// Note: This util is intended for internal use and should not be included or used outside of NetSerializers having to deal with calling into existing NetSerialzie functions.
+ 
+#include "NetSerializer.h"
+#include "UObject/ObjectMacros.h"
+#include "Iris/Core/NetObjectReference.h"
+#include "Iris/ReplicationSystem/NetToken.h"
+#include "Iris/Serialization/NetSerializerArrayStorage.h"
+#include "Iris/Serialization/StringNetSerializers.h"
+
+namespace UE::Net
+{
+struct FIrisPackageMapExports;
+class FNetReferenceCollector;
+
+struct FIrisPackageMapExportsQuantizedType
+{
+	static constexpr uint32 MaxInlinedObjectRefs = 4;
+	typedef FNetSerializerArrayStorage<FNetObjectReference, AllocationPolicies::TInlinedElementAllocationPolicy<MaxInlinedObjectRefs>> FObjectReferenceStorage;
+
+	struct FQuantizedName
+	{
+		alignas(8) uint8 Name[GetNameNetSerializerSafeQuantizedSize()];
+	};
+
+	static constexpr uint32 MaxInlinedNames = 4;
+	typedef FNetSerializerArrayStorage<FQuantizedName, AllocationPolicies::TInlinedElementAllocationPolicy<MaxInlinedNames>> FNamesStorage;
+
+	static constexpr uint32 MaxInlinedNetTokens = 4;
+	typedef FNetSerializerArrayStorage<FNetToken, AllocationPolicies::TInlinedElementAllocationPolicy<MaxInlinedNetTokens>> FNetTokenStorage;
+
+	FObjectReferenceStorage ObjectReferenceStorage;
+	FNamesStorage NameStorage;
+	FNetTokenStorage NetTokenStorage;
+};
+
+}
+
+template <> struct TIsPODType<UE::Net::FIrisPackageMapExportsQuantizedType> { enum { Value = true }; };
+
+namespace UE::Net
+{
+
+// Util to facilitate capture and export of supported types when calling into old NetSerialize() methods.
+struct FIrisPackageMapExportsUtil
+{
+	typedef FIrisPackageMapExportsQuantizedType QuantizedType;
+
+	// We need some sort of limit here.
+	static constexpr uint32 MaxExports = 128U;
+
+	// Matches NetSerializer functions
+	IRISCORE_API static void Serialize(FNetSerializationContext& Context, const QuantizedType& Value);
+	IRISCORE_API static void Deserialize(FNetSerializationContext& Context, QuantizedType& Value);
+	IRISCORE_API static void Quantize(FNetSerializationContext& Context, const UE::Net::FIrisPackageMapExports& PackageMapExport, QuantizedType& Target);
+	IRISCORE_API static void Dequantize(FNetSerializationContext& Context, const QuantizedType& Source, UE::Net::FIrisPackageMapExports& PackageMapExports);
+	IRISCORE_API static bool IsEqual(FNetSerializationContext& Context, const QuantizedType& Value0, const QuantizedType& Value1);
+	IRISCORE_API static void CloneDynamicState(FNetSerializationContext& Context, QuantizedType& Target, const QuantizedType& Source);
+	IRISCORE_API static void FreeDynamicState(FNetSerializationContext& Context, QuantizedType& Value);
+	IRISCORE_API static void CollectNetReferences(FNetSerializationContext& Context, const QuantizedType& Value, const FNetSerializerChangeMaskParam& ChangeMaskInfo, FNetReferenceCollector& Collector);
+	IRISCORE_API static bool Validate(FNetSerializationContext&, const QuantizedType& Value);
+
+private:
+
+	static const FNetSerializer* ObjectNetSerializer;
+	static const FNetSerializer* NameNetSerializer;
+};
+
+}
diff --git a/Engine/Source/Runtime/Experimental/Iris/Core/Public/Iris/Serialization/NetSerializationContext.h b/Engine/Source/Runtime/Experimental/Iris/Core/Public/Iris/Serialization/NetSerializationContext.h
--- a/Engine/Source/Runtime/Experimental/Iris/Core/Public/Iris/Serialization/NetSerializationContext.h
+++ b/Engine/Source/Runtime/Experimental/Iris/Core/Public/Iris/Serialization/NetSerializationContext.h
@@ -13,6 +13,7 @@ namespace UE::Net
 	class FNetBitArrayView;
 	class FNetBitStreamReader;
 	class FNetBitStreamWriter;
+	class FNetTokenStore;
 
 	namespace Private
 	{
@@ -88,6 +89,9 @@ public:
 	 */
 	IRISCORE_API UObject* GetLocalConnectionUserData(uint32 ConnectionId);
 
+	IRISCORE_API const UE::Net::FNetTokenStore* GetNetTokenStore() const;
+	IRISCORE_API UE::Net::FNetTokenStore* GetNetTokenStore();
+
 	void SetInternalContext(Private::FInternalNetSerializationContext* InInternalContext) { InternalContext = InInternalContext; }
 	Private::FInternalNetSerializationContext* GetInternalContext() { return InternalContext; }
 
diff --git a/Engine/Source/Runtime/Experimental/Iris/Core/Public/Iris/Serialization/StringNetSerializers.h b/Engine/Source/Runtime/Experimental/Iris/Core/Public/Iris/Serialization/StringNetSerializers.h
--- a/Engine/Source/Runtime/Experimental/Iris/Core/Public/Iris/Serialization/StringNetSerializers.h
+++ b/Engine/Source/Runtime/Experimental/Iris/Core/Public/Iris/Serialization/StringNetSerializers.h
@@ -31,5 +31,6 @@ UE_NET_DECLARE_SERIALIZER(FNameNetSerializer, IRISCORE_API)
 UE_NET_DECLARE_SERIALIZER(FNameAsNetTokenNetSerializer, IRISCORE_API)
 UE_NET_DECLARE_SERIALIZER(FStringNetSerializer, IRISCORE_API)
 
+constexpr SIZE_T GetNameNetSerializerSafeQuantizedSize() { return 24; }
 }
 
