Iris
- Fixed static analysis error

diff --git a/Engine/Source/Runtime/Experimental/Iris/Core/Private/Iris/Serialization/IrisPackageMapExportUtil.cpp b/Engine/Source/Runtime/Experimental/Iris/Core/Private/Iris/Serialization/IrisPackageMapExportUtil.cpp
--- a/Engine/Source/Runtime/Experimental/Iris/Core/Private/Iris/Serialization/IrisPackageMapExportUtil.cpp
+++ b/Engine/Source/Runtime/Experimental/Iris/Core/Private/Iris/Serialization/IrisPackageMapExportUtil.cpp
@@ -376,7 +376,7 @@ void FIrisPackageMapExportsUtil::CollectNetReferences(FNetSerializationContext&
 
 bool FIrisPackageMapExportsUtil::Validate(FNetSerializationContext& Context, const QuantizedType& SourceValue)
 {
-	if ((SourceValue.ObjectReferenceStorage.Num() > MaxExports) || (SourceValue.NameStorage.Num() > MaxExports) || (SourceValue.NameStorage.Num() > MaxExports))
+	if ((SourceValue.ObjectReferenceStorage.Num() > MaxExports) || (SourceValue.NameStorage.Num() > MaxExports) || (SourceValue.NetTokenStorage.Num() > MaxExports))
 	{
 		return false;
 	}
