[AutoRTFM] Fix another crash in FChaosRefCountedObject.

This CL significantly reworks the underlying behavior; GetRefCount should always return a value at least as large as the correct number of held references. A deprecated class has been replaced with std::atomic, and a transient-or-persistent boolean is replaced with an easier-to-read enum.

diff --git a/Engine/Source/Runtime/Experimental/Chaos/Private/Tests/AutoRTFMRefCountedObjectTests.cpp b/Engine/Source/Runtime/Experimental/Chaos/Private/Tests/AutoRTFMRefCountedObjectTests.cpp
--- a/Engine/Source/Runtime/Experimental/Chaos/Private/Tests/AutoRTFMRefCountedObjectTests.cpp
+++ b/Engine/Source/Runtime/Experimental/Chaos/Private/Tests/AutoRTFMRefCountedObjectTests.cpp
@@ -8,6 +8,9 @@
 
 #if WITH_DEV_AUTOMATION_TESTS
 
+#define CHECK_EQ(A, B) \
+	UTEST_EQUAL(TEXT(__FILE__ ":" UE_STRINGIZE(__LINE__) ": UTEST_EQUAL_EXPR(" #A ", " #B ")"), A, B)
+
 IMPLEMENT_SIMPLE_AUTOMATION_TEST(FAutoRTFMChaosRefCountedObject, "AutoRTFM + ChaosRefCountedObject", \
 	                             EAutomationTestFlags::EngineFilter | EAutomationTestFlags::ClientContext | \
 								 EAutomationTestFlags::ServerContext | EAutomationTestFlags::CommandletContext)
@@ -37,6 +40,49 @@ bool FAutoRTFMChaosRefCountedObject::RunTest(const FString& Parameters)
 			});
 	}
 
+	// AddRef on an object with non-zero refcount is rolled back properly when a transaction is aborted.
+	{
+		auto* TransientObject = new Chaos::FChaosRefCountedObject();
+		TransientObject->AddRef();
+		CHECK_EQ(TransientObject->GetRefCount(), 1);
+		int RefCountInsideTransaction;
+		AutoRTFM::Transact([&]
+			{
+				TransientObject->AddRef();
+				RefCountInsideTransaction = TransientObject->GetRefCount();
+				AutoRTFM::AbortTransaction();
+			});
+		CHECK_EQ(RefCountInsideTransaction, 2);
+		CHECK_EQ(TransientObject->GetRefCount(), 1);
+		TransientObject->Release();
+	}
+
+	// Release is rolled back properly when a transaction is aborted.
+	{
+		auto* TransientObject = new Chaos::FChaosRefCountedObject();
+		TransientObject->AddRef();
+		CHECK_EQ(TransientObject->GetRefCount(), 1);
+		AutoRTFM::Transact([&]
+			{
+				TransientObject->Release();
+				AutoRTFM::AbortTransaction();
+			});
+		CHECK_EQ(TransientObject->GetRefCount(), 1);
+		TransientObject->Release();
+	}
+
+	// AddRef on a zero-refcount object is rolled back properly when a transaction is aborted.
+	// (That is, the refcount is restored and the object is not destroyed.)
+	{
+		auto TransientObject = MakeUnique<Chaos::FChaosRefCountedObject>();
+		AutoRTFM::Transact([&]
+			{
+				TransientObject->AddRef();
+				AutoRTFM::AbortTransaction();
+			});
+		CHECK_EQ(TransientObject->GetRefCount(), 0);
+	}
+
 	// Adding and releasing a reference to a persistent object will not delete it.
 	// This tests uses TUniquePtr to perform the deletion when the object falls out of scope.
 	{
diff --git a/Engine/Source/Runtime/Experimental/Chaos/Public/Chaos/RefCountedObject.h b/Engine/Source/Runtime/Experimental/Chaos/Public/Chaos/RefCountedObject.h
--- a/Engine/Source/Runtime/Experimental/Chaos/Public/Chaos/RefCountedObject.h
+++ b/Engine/Source/Runtime/Experimental/Chaos/Public/Chaos/RefCountedObject.h
@@ -3,79 +3,114 @@
 
 #include <atomic>
 #include "AutoRTFM/AutoRTFM.h"
-#include "Runtime/Core/Public/HAL/ThreadSafeCounter.h"
 
 namespace Chaos
 {
 
-// Chaos ref counted object
-//  * @note AutoRTFM means that the return value of AddRef/Release is nonsense (as the ref-count doesn't change until the
-//  *       transaction is committed), but this is fine for use with TRefCountPtr (as it doesn't use those return values).
+// Chaos ref-counted object
+//  * @note In AutoRTFM, the return value of AddRef/Release may be higher than expected, because the refcount won't decrease
+//          until the transaction is committed. This is fine for use with TRefCountPtr, as it doesn't use the refcount directly.
 class FChaosRefCountedObject
 {
 public:
-	FChaosRefCountedObject() : NumRefs(0) {}
+	FChaosRefCountedObject() {}
 	virtual ~FChaosRefCountedObject()
 	{
-		UE_AUTORTFM_ONCOMMIT2(this)
+		// We want to report an error if we attempt to destroy a ref-counted object with leaked references.
+		// Sometimes, these objects exist ephemerally on the stack. If so, it should have a refcount of zero right now.
+		if (GetRefCount() != 0)
 		{
-			check(NumRefs.GetValue() == 0);
-		};
+			// If not, it might still have references that are queued up to Release at ONCOMMIT time. So, we check a second
+			// time during OnCommit. (If an ephemeral stack object is destroyed a non-zero refcount, this is user error;
+			// we might report this by check-failing here with a garbage value for its refcount, as its stack representation
+			// might already be overwritten.)
+			UE_AUTORTFM_ONCOMMIT2(this)
+			{
+				check(GetRefCount() == 0);
+			};
+		}
 	}
 	FChaosRefCountedObject(const FChaosRefCountedObject& Rhs) = delete;
 	FChaosRefCountedObject& operator=(const FChaosRefCountedObject& Rhs) = delete;
+	
 	uint32 AddRef() const
 	{
-		UE_AUTORTFM_ONCOMMIT2(this)
+		bool bIsFirstReference;
+
+		UE_AUTORTFM_OPEN2
 		{
-			NumRefs.Increment();
+			bIsFirstReference = (++NumRefs == 1);
 		};
-
-		// Note: TRefCountPtr doesn't use the return value
+		UE_AUTORTFM_ONABORT2(=, this)
+		{
+			if (bIsFirstReference)
+			{
+				// We took the first reference, and then aborted. This should undo the taking of
+				// the reference, but shouldn't delete the object if it is transient.
+				--NumRefs;
+			}
+			else
+			{
+				// After an object gains its initial reference, an AddRef call can be balanced
+				// out with a matching Release.
+				Release();
+			} 
+		};
+		// TRefCountPtr doesn't use the return value.
 		return 0;
 	}
+
 	uint32 Release() const
 	{
 		UE_AUTORTFM_ONCOMMIT2(this)
 		{
-			uint32 Refs = uint32(NumRefs.Decrement());
-			if (Refs == 0)
+			if (--NumRefs == 0 && RefCountMode == ERCM_Transient)
 			{
-				if (bTransientFlag)
-				{
-					delete this;
-				}
+				delete this;
 			}
 		};
-
-		// Note: TRefCountPtr doesn't use the return value
+		// TRefCountPtr doesn't use the return value.
 		return 0;
 	}
+
 	uint32 GetRefCount() const
 	{
-		uint32 Ret = 0;
+		uint32 Ret;
 		UE_AUTORTFM_OPEN2
 		{
-			Ret = uint32(NumRefs.GetValue());
+			Ret = uint32(NumRefs.load());
 		};
-
 		return Ret;
 	}
 
 	void MakePersistent() const
 	{
-		UE_AUTORTFM_ONCOMMIT2(this)
+		ERCM_RefCountMode OriginalMode = RefCountMode;
+
+		UE_AUTORTFM_OPEN2
 		{
-			bTransientFlag = false;
+			RefCountMode = ERCM_Persistent;
+		};
+		UE_AUTORTFM_ONABORT2(=, this)
+		{
+			RefCountMode = OriginalMode;
 		};
 	}
 
 private:
-	// Number of refs onto the object
-	mutable FThreadSafeCounter NumRefs;
-
-	// Transient flag to trigger the automatic deletion, or not
-	mutable std::atomic<bool> bTransientFlag = true;
+	// Number of refs onto the object.
+	mutable std::atomic<int32> NumRefs = 0;
+	
+	enum ERCM_RefCountMode {
+		// An object is considered Transient by default.
+		// After an initial AddRef, when the reference count reaches zero, it automatically deletes itself.
+		ERCM_Transient,
+		// Calling MakePersistent will convert an object to Persistent. 
+		// A Persistent object no longer deletes itself when the reference count reaches zero; the caller
+		// is responsible for deletion. (Basically, this opts out of the reference-counting mechanism.)
+		ERCM_Persistent,
+	};
+	mutable std::atomic<ERCM_RefCountMode> RefCountMode = ERCM_Transient;
 };
 
 }  // namespace Chaos
\ No newline at end of file
