Better fix for writing references to OFPA objects - use the provided outerlimit when a package object is found in the outer chain

#rnx

diff --git a/Engine/Source/Runtime/Experimental/JsonObjectGraph/Private/JsonStringifyImpl.cpp b/Engine/Source/Runtime/Experimental/JsonObjectGraph/Private/JsonStringifyImpl.cpp
--- a/Engine/Source/Runtime/Experimental/JsonObjectGraph/Private/JsonStringifyImpl.cpp
+++ b/Engine/Source/Runtime/Experimental/JsonObjectGraph/Private/JsonStringifyImpl.cpp
@@ -231,8 +231,16 @@ static void WriteObjectPath(FUtf8StringBuilderBase& OutPath, const UObject* ForO
 {
 	const auto WriteObjectPathImpl = [&OutPath, OuterLimit](const UObject* ForObject, auto Self)
 		{
-			if (ForObject == nullptr
-				|| ForObject == OuterLimit)
+			check(ForObject);
+			// ofpa files still have an outer package, but the GetPackage() terminator
+			// will not be reachable via the outer chain, replace any encoutered UPackage
+			// w/ the OuterLimit (acquired via GetPackage) when serializing:
+			if(Cast<UPackage>(ForObject) && ForObject != OuterLimit)
+			{
+				ForObject = OuterLimit;
+			}
+
+			if (ForObject == OuterLimit)
 			{
 				return;
 			}
