Iris - Fixed missing incrementing of item ReplicationKeys when received. This prevented the items from being serialized during client demo recording.

diff --git a/Engine/Source/Runtime/Experimental/Iris/Core/Public/Iris/ReplicationSystem/Private/FastArrayReplicationFragmentInternal.h b/Engine/Source/Runtime/Experimental/Iris/Core/Public/Iris/ReplicationSystem/Private/FastArrayReplicationFragmentInternal.h
--- a/Engine/Source/Runtime/Experimental/Iris/Core/Public/Iris/ReplicationSystem/Private/FastArrayReplicationFragmentInternal.h
+++ b/Engine/Source/Runtime/Experimental/Iris/Core/Public/Iris/ReplicationSystem/Private/FastArrayReplicationFragmentInternal.h
@@ -258,6 +258,7 @@ void FFastArrayReplicationFragmentHelper::ApplyReplicatedState(FastArrayType* Ds
 				continue;
 			}
 
+			ItemType* DstItem = nullptr;
 			if (int32* ExistingIndex = DstArraySerializer->ItemMap.Find(SrcItems[It].ReplicationID))
 			{
 				// Only compare if the changemask indicate that this might be a dirty entry, the compare is required since we do share entries in the changemask.
@@ -268,7 +269,8 @@ void FFastArrayReplicationFragmentHelper::ApplyReplicatedState(FastArrayType* Ds
 					ModifiedIndices.Add(*ExistingIndex);
 
 					// We use per element apply since we do not want to overwrite data that is not replicated
-					InternalApplyArrayElement(ArrayElementDescriptor, &(*DstWrappedArray)[*ExistingIndex], &SrcItems[It]);
+					DstItem = &(*DstWrappedArray)[*ExistingIndex];
+					InternalApplyArrayElement(ArrayElementDescriptor, DstItem, &SrcItems[It]);
 				}
 			}
 			else
@@ -281,12 +283,22 @@ void FFastArrayReplicationFragmentHelper::ApplyReplicatedState(FastArrayType* Ds
 					UE_LOG(LogNetFastTArray, Log, TEXT("   New. ID: %d. New Element! local Idx: %d"), SrcItems[It].ReplicationID, AddedIndex);
 
 					// We need to propagate the ReplicationID in order to find our object
-					(*DstWrappedArray)[AddedIndex].ReplicationID = SrcItems[It].ReplicationID;
+					DstItem = &(*DstWrappedArray)[AddedIndex];
+					DstItem->ReplicationID = SrcItems[It].ReplicationID;
 
 					// should we store ids or indices?
 					AddedIndices.Add(AddedIndex);
 				}
 			}
+
+			if (DstItem != nullptr)
+			{
+				// Update the item's most recent array replication key
+				DstItem->MostRecentArrayReplicationKey = DstArraySerializer->ArrayReplicationKey;
+
+				// Update the item's replication key so that a client can re-serialize the array for client replay recording
+				DstItem->ReplicationKey++;
+			}
 		}
 	}
 
