Fix accidental leak and early-exit in unit test.

Two problems collided:
- The `RefCountInsideTransaction` variable was ill-conceived, since its value was stored during a transaction, which was later aborted. This meant that the unit test was looking at a variable which was initialized, but then rolled back to its initial uninitialized-memory state.
- `UTEST_EQUAL` embeds an early `return false`, stopping the test short.

Returning early caused Horde to detect this function as leaky.

diff --git a/Engine/Source/Runtime/Experimental/Chaos/Private/Tests/AutoRTFMRefCountedObjectTests.cpp b/Engine/Source/Runtime/Experimental/Chaos/Private/Tests/AutoRTFMRefCountedObjectTests.cpp
--- a/Engine/Source/Runtime/Experimental/Chaos/Private/Tests/AutoRTFMRefCountedObjectTests.cpp
+++ b/Engine/Source/Runtime/Experimental/Chaos/Private/Tests/AutoRTFMRefCountedObjectTests.cpp
@@ -9,7 +9,7 @@
 #if WITH_DEV_AUTOMATION_TESTS
 
 #define CHECK_EQ(A, B) \
-	UTEST_EQUAL(TEXT(__FILE__ ":" UE_STRINGIZE(__LINE__) ": UTEST_EQUAL_EXPR(" #A ", " #B ")"), A, B)
+	do { TestEqual(TEXT(__FILE__ ":" UE_STRINGIZE(__LINE__) ": TestEqual(" #A ", " #B ")"), (A), (B)); } while (0)
 
 IMPLEMENT_SIMPLE_AUTOMATION_TEST(FAutoRTFMChaosRefCountedObject, "AutoRTFM + ChaosRefCountedObject", \
 	                             EAutomationTestFlags::EngineFilter | EAutomationTestFlags::ClientContext | \
@@ -45,14 +45,12 @@ bool FAutoRTFMChaosRefCountedObject::RunTest(const FString& Parameters)
 		auto* TransientObject = new Chaos::FChaosRefCountedObject();
 		TransientObject->AddRef();
 		CHECK_EQ(TransientObject->GetRefCount(), 1);
-		int RefCountInsideTransaction;
 		AutoRTFM::Transact([&]
 			{
 				TransientObject->AddRef();
-				RefCountInsideTransaction = TransientObject->GetRefCount();
+				CHECK_EQ(TransientObject->GetRefCount(), 2);
 				AutoRTFM::AbortTransaction();
 			});
-		CHECK_EQ(RefCountInsideTransaction, 2);
 		CHECK_EQ(TransientObject->GetRefCount(), 1);
 		TransientObject->Release();
 	}
