Constraints: improved constraint sorting (follow-up to 32886233)
    - use topological sorting to handle more complex ordering schemas
    - use topological sorting to handle more complex ordering schemas

Constraints: improving constraint sorting (tracking 32886233)
    - use topological sorting to manage more complex dependency patterns
    - constraint ordering is now entirely managed by the evaluation graph to avoid code duplication (+ faster computation as sorting is only done once when building the graph)

#rnx

diff --git a/Engine/Source/Runtime/Experimental/Animation/Constraints/Private/ConstraintsEvaluationGraph.cpp b/Engine/Source/Runtime/Experimental/Animation/Constraints/Private/ConstraintsEvaluationGraph.cpp
--- a/Engine/Source/Runtime/Experimental/Animation/Constraints/Private/ConstraintsEvaluationGraph.cpp
+++ b/Engine/Source/Runtime/Experimental/Animation/Constraints/Private/ConstraintsEvaluationGraph.cpp
@@ -4,6 +4,7 @@
 
 #include "ConstraintsManager.h"
 #include "ConstraintSubsystem.h"
+#include "Algo/TopologicalSort.h"
 
 namespace ConstraintsEvaluationGraph
 {
@@ -108,86 +109,49 @@ void FConstraintsEvaluationGraph::Rebuild()
 	{
 		return;
 	}
+	
+	UE::Constraints::Graph::BuildGraph(ConstraintsInWorld.World.Get(), Constraints, Nodes);
+	
+	State = ReadyForEvaluation;
 
-	// build nodes
-	for (int32 ConstraintIndex = 0; ConstraintIndex < Constraints.Num(); ConstraintIndex++)
+	Dump();
+}
+
+bool FConstraintsEvaluationGraph::GetSortedConstraints(TArray<TWeakObjectPtr<UTickableConstraint>>& OutConstraints)
+{
+	OutConstraints.Reset();
+
+	if (!ensure(ConstraintsInWorld.World.Get()))
 	{
-		if (Constraints[ConstraintIndex].IsValid())
-		{
-			FConstraintNode& Node = GetNode(Constraints[ConstraintIndex]);
-			Node.ConstraintIndex = ConstraintIndex;
-		}
+		return false;
 	}
 
-	// sort nodes by tick dependencies
-	auto EvaluationOrderPredicate = [](const FConstraintNode& LHS, const FConstraintNode& RHS)
+	if (State == InvalidData)
 	{
-		const TArray<FTickPrerequisite>& RHSPrerex = RHS.ConstraintTick->GetPrerequisites();
-		const FConstraintTickFunction* LHSTickFunction = LHS.ConstraintTick;
-		const bool bIsLHSAPrerexOfRHS = RHSPrerex.ContainsByPredicate([LHSTickFunction](const FTickPrerequisite& Prerex)
-		{
-			return Prerex.PrerequisiteTickFunction == LHSTickFunction;
-		});
-		// if not a prerequisite then compare constraints indices
-		return bIsLHSAPrerexOfRHS ? bIsLHSAPrerexOfRHS :  LHS.ConstraintIndex < RHS.ConstraintIndex;
-	};
-	Algo::Sort(Nodes, EvaluationOrderPredicate);
+		Rebuild();
+	}
 
-	// store node index after re-ordering
-	for (int32 NodeIndex = 0; NodeIndex < Nodes.Num(); NodeIndex++)
+	if (Nodes.IsEmpty())
 	{
-		Nodes[NodeIndex].NodeIndex = NodeIndex;
+		return false;
 	}
 
-	// build edges	
-	for (FConstraintNode& Node: Nodes)
+	using ConstraintPtr = TWeakObjectPtr<UTickableConstraint>;
+	const TArray<ConstraintPtr>& Constraints = ConstraintsInWorld.Constraints;
+	
+	OutConstraints.Reserve(Nodes.Num());
+	for (const FConstraintNode& Node: Nodes)
 	{
-		const TArray<FTickPrerequisite>& Prerequisites = Node.ConstraintTick->GetPrerequisites();
-		for (const FTickPrerequisite& Prerex: Prerequisites)
+		const ConstraintPtr Constraint = Constraints[Node.ConstraintIndex];
+		if (Constraint.IsValid())
 		{
-			if (const FTickFunction* PrerexFunction = Prerex.Get())
-			{
-				FConstraintNode* PrerexNode = Nodes.FindByPredicate([PrerexFunction, Node](const FConstraintNode& OtherNode)
-				{
-					if (OtherNode.NodeIndex != Node.NodeIndex)
-					{
-						if (OtherNode.ConstraintTick == PrerexFunction)
-						{
-							return true;
-						}
-					}
-					return false;
-				});
-
-				if (PrerexNode)
-				{
-					if (!PrerexNode->Parents.Contains(Node.NodeIndex))
-					{
-						Node.Parents.Add(PrerexNode->NodeIndex);
-					}
-					else
-					{
-						// we may have create a cycle
-						ensure(false);
-					}
-
-					if (!Node.Children.Contains(PrerexNode->NodeIndex))
-					{
-						PrerexNode->Children.Add(Node.NodeIndex);
-					}
-					else
-					{
-						// we may have create a cycle
-						ensure(false);
-					}
-				}
-			}
+			OutConstraints.Add(Constraint);
 		}
 	}
-	
-	State = ReadyForEvaluation;
 
-	Dump();
+	ensure(OutConstraints.Num() == Constraints.Num());
+	
+	return true;
 }
 
 bool FConstraintsEvaluationGraph::IsPendingEvaluation() const
@@ -322,5 +286,327 @@ void FConstraintsEvaluationGraph::Dump() const
 		ensure(Constraints.IsValidIndex(Node.ConstraintIndex));
 		ensure(Node.NodeIndex == NodeIndex);
 		UE_LOG(LogTemp, Warning, TEXT("\tNode[%d] = %s [%d]"), NodeIndex, *GetConstraintLabel(Constraints[Node.ConstraintIndex]), Node.ConstraintIndex);
+		for (const uint32 ChildIndex: Node.Children)
+		{
+			if (ensure(Nodes.IsValidIndex(ChildIndex)))
+			{
+				const FConstraintNode& ChildNode = Nodes[ChildIndex];
+				UE_LOG(LogTemp, Warning, TEXT("\t\tChild[%d] = %s [%d]"), ChildIndex, *GetConstraintLabel(Constraints[ChildNode.ConstraintIndex]), ChildNode.ConstraintIndex);
+			}
+		}
+	}
+}
+
+namespace UE::Constraints::Graph::Private
+{
+
+struct FGraphBuildHelper
+{
+	FGraphBuildHelper(UWorld* InWorld, const TArrayView<const ConstraintPtr>& InConstraints, TArray<FConstraintNode>& OutNodes)
+		: World(InWorld)
+		, Constraints(InConstraints)
+		, Nodes(OutNodes)
+	{}
+
+	void Build()
+	{
+		Nodes.Reset();
+	
+		if (!World || Constraints.IsEmpty())
+		{
+			return;
+		}
+
+		BuildNodes();	
+		BuildNodeDependencies();
+		SortNodes();
+	}
+
+private:
+
+	UWorld* World = nullptr;
+	const TArrayView<const ConstraintPtr>& Constraints;
+	TArray<FConstraintNode>& Nodes;
+
+	FConstraintNode& GetNode(const ConstraintPtr& InConstraint) const
+	{
+		const int32 Found = Nodes.IndexOfByPredicate([InConstraint](const FConstraintNode& Node)
+		{
+			return InConstraint.IsValid() && Node.ConstraintID == InConstraint->ConstraintID; 
+		});
+
+		if (Found != INDEX_NONE)
+		{
+			return Nodes[Found];
+		}
+
+		FConstraintNode Node;
+		Node.ConstraintID = InConstraint->ConstraintID;
+		Node.ConstraintTick = &InConstraint->GetTickFunction(World);
+		return Nodes.Emplace_GetRef(MoveTemp(Node)); 
 	}
-}
\ No newline at end of file
+
+	// build vertices
+	void BuildNodes() const
+	{
+		for (int32 ConstraintIndex = 0, NumConstraints = Constraints.Num(), NodeIndex = 0; ConstraintIndex < NumConstraints; ConstraintIndex++)
+		{
+			if (Constraints[ConstraintIndex].IsValid())
+			{
+				FConstraintNode& Node = GetNode(Constraints[ConstraintIndex]);
+				Node.NodeIndex = NodeIndex++;
+				Node.ConstraintIndex = ConstraintIndex;
+			}
+		}
+	}
+
+	// build edges
+	void BuildNodeDependencies() const
+	{
+		for (FConstraintNode& Node: Nodes)
+		{
+			const TArray<FTickPrerequisite>& Prerequisites = Node.ConstraintTick->GetPrerequisites();
+			for (const FTickPrerequisite& Prerex: Prerequisites)
+			{
+				if (const FTickFunction* PrerexFunction = Prerex.Get())
+				{
+					FConstraintNode* PrerexNode = Nodes.FindByPredicate([PrerexFunction, Node](const FConstraintNode& OtherNode)
+					{
+						if (OtherNode.NodeIndex != Node.NodeIndex)
+						{
+							if (OtherNode.ConstraintTick == PrerexFunction)
+							{
+								return true;
+							}
+						}
+						return false;
+					});
+
+					if (PrerexNode)
+					{
+						if (!PrerexNode->Parents.Contains(Node.NodeIndex))
+						{
+							Node.Parents.Add(PrerexNode->NodeIndex);
+						}
+						else
+						{
+							// we may have create a cycle
+							ensure(false);
+						}
+
+						if (!Node.Children.Contains(PrerexNode->NodeIndex))
+						{
+							PrerexNode->Children.Add(Node.NodeIndex);
+						}
+						else
+						{
+							// we may have create a cycle
+							ensure(false);
+						}
+					}
+				}
+			}
+		}
+	}
+	
+	// sort using dependencies
+	void SortNodes() const
+	{
+		const int32 NumNodes = Nodes.Num();
+
+		// switch to Indices
+		TArray<int32> Indices;
+		Indices.Reserve(NumNodes);
+		Algo::Transform(Nodes, Indices, [](const FConstraintNode& Node) { return Node.NodeIndex; });
+
+		// store copy
+		const TArray IndicesBeforeSort(Indices);
+		const TArray NodesBeforeSort(Nodes);
+		
+		// try topological sort
+		auto GetDependencies = [this](const int32 Index)
+		{
+			return Nodes[Index].Parents.Array();
+		};
+		const bool bSucceeded = Algo::TopologicalSort(Indices, GetDependencies);
+
+		TMap<int32, int32> OldToNewIndex;
+		
+		if (ensure(bSucceeded))
+		{
+			// store a old to new index map to re-index things out
+			bool bHasBeenReordered = false;
+			for (int32 Index = 0; Index < NumNodes; ++Index)
+			{
+				const int32 OldIndex = IndicesBeforeSort.IndexOfByKey(Indices[Index]);
+				OldToNewIndex.Emplace(OldIndex, Index);
+				bHasBeenReordered |= OldIndex != Index;
+			}
+
+			// early exits if it hasn't been reordered 
+			if (!bHasBeenReordered)
+			{
+				if (ConstraintsEvaluationGraph::bDebugGraph)
+				{
+					UE_LOG(LogTemp, Warning, TEXT("No need to re-index constraints."));
+				}
+				return;
+			}
+			
+			// switch back to Nodes
+			for (int32 Index = 0; Index < NumNodes; ++Index)
+			{
+				const int32 NewIndex = Indices[Index];
+				Nodes[Index] = NodesBeforeSort[NewIndex];
+			}
+		}
+		else
+		{
+			// if it failed (mostly due to cycles) return to the basic Algo::Sort to still trying to sort things out.
+
+			auto Predicate = [](const FConstraintNode& LHS, const FConstraintNode& RHS)
+			{
+				const FConstraintTickFunction* LHSTickFunction = LHS.ConstraintTick;
+				const FConstraintTickFunction* RHSTickFunction = RHS.ConstraintTick;
+    
+				if (!LHSTickFunction || !RHSTickFunction)
+				{
+					return LHS.ConstraintIndex < RHS.ConstraintIndex;
+				}
+        
+				const TArray<FTickPrerequisite>& RHSPrerex = RHSTickFunction->GetPrerequisites();
+				const bool bIsLHSAPrerexOfRHS = RHSPrerex.ContainsByPredicate([LHSTickFunction](const FTickPrerequisite& Prerex)
+				{
+					return Prerex.PrerequisiteTickFunction == LHSTickFunction;
+				});
+
+				if (bIsLHSAPrerexOfRHS)
+				{
+					return true;
+				}
+    
+				const TArray<FTickPrerequisite>& LHSPrerex = LHSTickFunction->GetPrerequisites();
+				const bool bIsRHSAPrerexOfLHS = LHSPrerex.ContainsByPredicate([RHSTickFunction](const FTickPrerequisite& Prerex)
+				{
+					return Prerex.PrerequisiteTickFunction == RHSTickFunction;
+				});
+        	
+				if (bIsRHSAPrerexOfLHS)
+				{
+					return false;
+				}
+    
+				// if not a prerequisite then compare constraints indices
+				return LHS.ConstraintIndex < RHS.ConstraintIndex;
+			};
+
+			Algo::Sort(Nodes, Predicate);
+
+			// store old to new index map to re-index things out
+			bool bHasBeenReordered = false;
+			for (int32 Index = 0; Index < NumNodes; ++Index)
+			{
+				const int32 OldIndex = IndicesBeforeSort.IndexOfByKey(Nodes[Index].NodeIndex);
+				OldToNewIndex.Emplace(OldIndex, Index);
+				bHasBeenReordered |= OldIndex != Index;
+			}
+
+			// early exits if it hasn't been reordered 
+			if (!bHasBeenReordered)
+			{
+				if (ConstraintsEvaluationGraph::bDebugGraph)
+				{
+					UE_LOG(LogTemp, Warning, TEXT("No need to re-index constraints."));
+				}
+				return;
+			}
+		}
+
+		if (ensure(OldToNewIndex.Num() == NumNodes))
+		{
+			// reindex node + parents + children
+			for (int32 Index = 0; Index < NumNodes; ++Index)
+			{
+				FConstraintNode& Node = Nodes[Index];
+				ReIndexNode(Node, OldToNewIndex);
+			}
+		}
+	}
+
+	static void ReIndexNode(FConstraintNode& InOutNode, const TMap<int32, int32>& InOldToNewIndex)
+	{
+		// update node + NodeIndex
+		InOutNode.NodeIndex = InOldToNewIndex[InOutNode.NodeIndex];
+
+		// update parents indices
+		if (!InOutNode.Parents.IsEmpty())
+		{
+			TSet<uint32> NewParents;
+			NewParents.Reserve(InOutNode.Parents.Num());
+			for (const int32 OldParentIndex: InOutNode.Parents)
+			{
+				NewParents.Add(InOldToNewIndex[OldParentIndex]);
+			}
+			InOutNode.Parents = MoveTemp(NewParents);
+		}
+
+		// update children indices
+		if (!InOutNode.Children.IsEmpty())
+		{
+			TSet<uint32> NewChildren;
+			NewChildren.Reserve(InOutNode.Children.Num());
+			for (const int32 OldChildIndex: InOutNode.Children)
+			{
+				NewChildren.Add(InOldToNewIndex[OldChildIndex]);
+			}	
+			InOutNode.Children = MoveTemp(NewChildren);
+		}
+	}
+};
+
+}
+
+namespace UE::Constraints::Graph
+{
+
+void BuildGraph(UWorld* InWorld, const TArrayView<const ConstraintPtr>& InConstraints, TArray<FConstraintNode>& OutNodes)
+{
+	OutNodes.Empty();
+	
+	if (!ensure(InWorld))
+	{
+		return;
+	}
+
+	if (InConstraints.IsEmpty())
+	{
+		return;
+	}
+	
+	Private::FGraphBuildHelper BuildHelper(InWorld, InConstraints, OutNodes);
+	BuildHelper.Build();
+}
+	
+void SortConstraints(UWorld* InWorld, TArray<ConstraintPtr>& InOutConstraints)
+{
+	// build graph
+	TArray<FConstraintNode> Nodes;
+	BuildGraph(InWorld, InOutConstraints, Nodes);
+
+	if (!Nodes.IsEmpty())
+	{
+		// re-order constraints
+		const TArray<ConstraintPtr> ConstraintsCopy(InOutConstraints);
+		InOutConstraints.Reset(Nodes.Num());
+		for (int32 Index = 0, NumNodes = Nodes.Num(); Index < NumNodes; ++Index)
+		{
+			const int32 ConstraintIndex = Nodes[Index].ConstraintIndex;
+			if (ConstraintsCopy.IsValidIndex(ConstraintIndex))
+			{
+				InOutConstraints.Add(ConstraintsCopy[ConstraintIndex]);
+			}
+		}
+	}
+}
+
+}
diff --git a/Engine/Source/Runtime/Experimental/Animation/Constraints/Private/ConstraintsManager.cpp b/Engine/Source/Runtime/Experimental/Animation/Constraints/Private/ConstraintsManager.cpp
--- a/Engine/Source/Runtime/Experimental/Animation/Constraints/Private/ConstraintsManager.cpp
+++ b/Engine/Source/Runtime/Experimental/Animation/Constraints/Private/ConstraintsManager.cpp
@@ -766,55 +766,6 @@ bool FConstraintsManagerController::RemoveAllConstraints(bool bDoNotCompensate)
 	return false;
 }
 
-static void SortConstraints(UWorld* InWorld, TArray< TWeakObjectPtr<UTickableConstraint> >& InOutSortedConstraints)
-{
-	using ConstraintPtr = TWeakObjectPtr<UTickableConstraint>;
-
-	// store original indices
-	TArray<int32> ConstraintIndices;
-	ConstraintIndices.Reserve(InOutSortedConstraints.Num());
-	for (int32 Index = 0; Index < InOutSortedConstraints.Num(); ++Index)
-	{
-		ConstraintIndices.Add(Index);
-	}
-	
-	// LHS ticks before RHS = LHS is a prerex of RHS 
-	auto TicksBefore = [InWorld, &InOutSortedConstraints](const int32 LHSIndex, const int32 RHSIndex)
-	{
-		const ConstraintPtr& LHS = InOutSortedConstraints[LHSIndex];
-		const ConstraintPtr& RHS = InOutSortedConstraints[RHSIndex];
-		if (RHS.IsValid() == false || LHS.IsValid() == false)
-		{
-			return false;
-		}
-	
-		const TArray<FTickPrerequisite>& RHSPrerex = RHS->GetTickFunction(InWorld).GetPrerequisites();
-		const FConstraintTickFunction& LHSTickFunction = LHS->GetTickFunction(InWorld);
-		const bool bIsLHSAPrerexOfRHS = RHSPrerex.ContainsByPredicate([&LHSTickFunction](const FTickPrerequisite& Prerex)
-		{
-			return Prerex.PrerequisiteTickFunction == &LHSTickFunction;
-		});
-
-		// if not a prerequisite then compare constraints indices
-		return bIsLHSAPrerexOfRHS ? bIsLHSAPrerexOfRHS : LHSIndex < RHSIndex;
-	};
-
-	// sort
-	Algo::Sort(ConstraintIndices, [TicksBefore](const int32 LHSIndex, const int32 RHSIndex)
-	{
-		return TicksBefore(LHSIndex, RHSIndex);
-	});
-
-	// re-order using sorted indices
-	TArray<TWeakObjectPtr<UTickableConstraint> > SortedConstraints;
-	SortedConstraints.Reserve(InOutSortedConstraints.Num());
-	for (int32 Index = 0; Index < ConstraintIndices.Num(); ++Index)
-	{
-		SortedConstraints.Add(InOutSortedConstraints[ConstraintIndices[Index]]);
-	}
-	InOutSortedConstraints = SortedConstraints;
-}
-
 TArray< TWeakObjectPtr<UTickableConstraint> > FConstraintsManagerController::GetStaticConstraints(const bool bSorted) const
 {
 	UConstraintsManager* Manager = FindManager();
@@ -836,7 +787,7 @@ TArray< TWeakObjectPtr<UTickableConstraint> > FConstraintsManagerController::Get
 		return Constraints;
 	}
 
-	SortConstraints(World, Constraints);
+	UE::Constraints::Graph::SortConstraints(World, Constraints);
 
 	return Constraints;
 }
@@ -850,6 +801,16 @@ TArray< TWeakObjectPtr<UTickableConstraint> > FConstraintsManagerController::Get
 		return Empty;
 	}
 
+	// use evaluation graph cached data
+	if (FConstraintsEvaluationGraph::UseEvaluationGraph() && bSorted)
+	{
+		TArray<TWeakObjectPtr<UTickableConstraint>> SortedConstraints;
+		if (Subsystem->GetEvaluationGraph(World).GetSortedConstraints(SortedConstraints))
+		{
+		   return SortedConstraints;
+		}
+	}
+
 	TArray<TWeakObjectPtr<UTickableConstraint>> Constraints = Subsystem->GetConstraints(World);
 
 	// Remove stale constraints (note that GCd constraints should already have been caught by UConstraintSubsystem)
@@ -863,7 +824,7 @@ TArray< TWeakObjectPtr<UTickableConstraint> > FConstraintsManagerController::Get
 		return Constraints;
 	}
 
-	SortConstraints(World, Constraints);
+	UE::Constraints::Graph::SortConstraints(World, Constraints);
 	
 	return Constraints;
 }
diff --git a/Engine/Source/Runtime/Experimental/Animation/Constraints/Public/ConstraintsEvaluationGraph.h b/Engine/Source/Runtime/Experimental/Animation/Constraints/Public/ConstraintsEvaluationGraph.h
--- a/Engine/Source/Runtime/Experimental/Animation/Constraints/Public/ConstraintsEvaluationGraph.h
+++ b/Engine/Source/Runtime/Experimental/Animation/Constraints/Public/ConstraintsEvaluationGraph.h
@@ -61,7 +61,10 @@ public:
 
 	/* Should the graph be evaluated ? */
 	bool IsPendingEvaluation() const;
-	
+
+	/* Get Sorted Constraints */
+	bool GetSortedConstraints(TArray<TWeakObjectPtr<UTickableConstraint>>& OutConstraints);
+
 	static bool UseEvaluationGraph();
 	
 private:
@@ -101,4 +104,24 @@ private:
 	
 	/* Current graph state. */
 	EGraphState	State = InvalidData;
-};
\ No newline at end of file
+};
+
+namespace UE::Constraints::Graph
+{
+	
+using ConstraintPtr = TWeakObjectPtr<UTickableConstraint>;
+
+/** Builds a directed graph from a constraints array using the tick prerequisites to define dependencies.
+* @param InWorld - The UWorld in which tick functions are registered.
+* @param InConstraints - An array of constraints.
+* @param OutNodes - A ordered data structure storing constraints' dependencies. 
+*/
+void BuildGraph(UWorld* InWorld, const TArrayView<const ConstraintPtr>& InConstraints, TArray<FConstraintNode>& OutNodes);
+
+/** Sorts the constraint array by building a directed graph from it, using the tick prerequisites to define dependencies.
+* @param InWorld - The UWorld in which tick functions are registered.
+* @param InOutConstraints - An array of constraints that will be re-ordered from the constraints' dependencies.   
+*/
+void SortConstraints(UWorld* InWorld, TArray<ConstraintPtr>& InOutConstraints);
+
+}
\ No newline at end of file
