Iris
* Demystify BeginReplication functions in the bridges by renaming them to what type of replicated object they create.
* Renamed FCreateNetRefHandleParams to FRootObjectReplicationParams.
* Removed FCreateNetRefHandleParams from SubObject BeginReplication function. All of the params could only be set on Root Objects.

diff --git a/Engine/Plugins/Runtime/ReplicationSystemTestPlugin/Source/Private/Tests/ReplicationSystem/Filtering/TestGridFilter.cpp b/Engine/Plugins/Runtime/ReplicationSystemTestPlugin/Source/Private/Tests/ReplicationSystem/Filtering/TestGridFilter.cpp
--- a/Engine/Plugins/Runtime/ReplicationSystemTestPlugin/Source/Private/Tests/ReplicationSystem/Filtering/TestGridFilter.cpp
+++ b/Engine/Plugins/Runtime/ReplicationSystemTestPlugin/Source/Private/Tests/ReplicationSystem/Filtering/TestGridFilter.cpp
@@ -85,7 +85,7 @@ private:
 UE_NET_TEST_FIXTURE(FTestGridFilterFixture, TestWorldLocGridFilter)
 {
 	// Spawn object with WorldLocation's on server
-	UObjectReplicationBridge::FCreateNetRefHandleParams Params;
+	UObjectReplicationBridge::FRootObjectReplicationParams Params;
 	Params.bNeedsWorldLocationUpdate = true;
 	Params.bUseClassConfigDynamicFilter = true;
 
@@ -183,7 +183,7 @@ UE_NET_TEST_FIXTURE(FTestGridFilterFixture, TestWorldLocationIsFrequentlyUpdated
 	Server->ReplicationSystem->SetReplicationView(Client->ConnectionIdOnServer, ReplicationView);
 
 	// Spawn object with WorldLocation's on server
-	UObjectReplicationBridge::FCreateNetRefHandleParams Params;
+	UObjectReplicationBridge::FRootObjectReplicationParams Params;
 	Params.bNeedsWorldLocationUpdate = true;
 	Params.bIsDormant = false;
 	Params.bUseClassConfigDynamicFilter = false;
@@ -265,7 +265,7 @@ UE_NET_TEST_FIXTURE(FTestGridFilterFixture, TestFilterPerformance)
 	FGenericPlatformMath::SRandInit(0);
 	for (int32 i = 0; i < NUM_OBJECTS; i++)
 	{
-		UObjectReplicationBridge::FCreateNetRefHandleParams Params;
+		UObjectReplicationBridge::FRootObjectReplicationParams Params;
 		Params.bNeedsWorldLocationUpdate = true;
 		Params.bUseClassConfigDynamicFilter = true;
 
diff --git a/Engine/Plugins/Runtime/ReplicationSystemTestPlugin/Source/Private/Tests/ReplicationSystem/Prioritization/TestNetObjectPrioritizerFixture.cpp b/Engine/Plugins/Runtime/ReplicationSystemTestPlugin/Source/Private/Tests/ReplicationSystem/Prioritization/TestNetObjectPrioritizerFixture.cpp
--- a/Engine/Plugins/Runtime/ReplicationSystemTestPlugin/Source/Private/Tests/ReplicationSystem/Prioritization/TestNetObjectPrioritizerFixture.cpp
+++ b/Engine/Plugins/Runtime/ReplicationSystemTestPlugin/Source/Private/Tests/ReplicationSystem/Prioritization/TestNetObjectPrioritizerFixture.cpp
@@ -90,7 +90,7 @@ FTestNetObjectPrioritizerFixture::FPrioritizationResult FTestNetObjectPrioritize
 
 	// Begin replication for all created objects
 	{
-		UObjectReplicationBridge::FCreateNetRefHandleParams CreateNetRefHandleParams = { .bNeedsWorldLocationUpdate = true };
+		UObjectReplicationBridge::FRootObjectReplicationParams CreateNetRefHandleParams = { .bNeedsWorldLocationUpdate = true };
 
 		for (UTestPrioritizationWithWorldLocationObject* Object : WorldLocObjects)
 		{
@@ -223,7 +223,7 @@ FNetObjectPrioritizerImage FTestNetObjectPrioritizerFixture::Visualize(FNetObjec
 
 	// Begin replication for all created objects
 	{
-		UObjectReplicationBridge::FCreateNetRefHandleParams CreateNetRefHandleParams = { .bNeedsWorldLocationUpdate = true };
+		UObjectReplicationBridge::FRootObjectReplicationParams CreateNetRefHandleParams = { .bNeedsWorldLocationUpdate = true };
 
 		for (UTestPrioritizationWithWorldLocationObject* Object : WorldLocObjects)
 		{
diff --git a/Engine/Plugins/Runtime/ReplicationSystemTestPlugin/Source/Private/Tests/ReplicationSystem/ReplicatedTestObject.cpp b/Engine/Plugins/Runtime/ReplicationSystemTestPlugin/Source/Private/Tests/ReplicationSystem/ReplicatedTestObject.cpp
--- a/Engine/Plugins/Runtime/ReplicationSystemTestPlugin/Source/Private/Tests/ReplicationSystem/ReplicatedTestObject.cpp
+++ b/Engine/Plugins/Runtime/ReplicationSystemTestPlugin/Source/Private/Tests/ReplicationSystem/ReplicatedTestObject.cpp
@@ -50,8 +50,8 @@ const UE::Net::FReplicationInstanceProtocol* UReplicatedTestObjectBridge::GetRep
 UE::Net::FNetRefHandle UReplicatedTestObjectBridge::BeginReplication(UReplicatedTestObject* Instance)
 {
 	// Create NetRefHandle for the registered fragments
-	Super::FCreateNetRefHandleParams Params;
-	FNetRefHandle Handle = Super::BeginReplication(Instance, Params);
+	FRootObjectReplicationParams Params;
+	FNetRefHandle Handle = StartReplicatingRootObject(Instance, Params);
 
 	// This is optional but typically we want to cache at least the NetRefHandle in the game instance to avoid doing map lookups to find it
 	if (Handle.IsValid())
@@ -62,10 +62,10 @@ UE::Net::FNetRefHandle UReplicatedTestObjectBridge::BeginReplication(UReplicated
 	return Handle;
 }
 
-UE::Net::FNetRefHandle UReplicatedTestObjectBridge::BeginReplication(UReplicatedTestObject* Instance, const UObjectReplicationBridge::FCreateNetRefHandleParams& Params)
+UE::Net::FNetRefHandle UReplicatedTestObjectBridge::BeginReplication(UReplicatedTestObject* Instance, const UObjectReplicationBridge::FRootObjectReplicationParams& Params)
 {
 	// Create NetRefHandle for the registered fragments
-	FNetRefHandle Handle = Super::BeginReplication(Instance, Params);
+	FNetRefHandle Handle = Super::StartReplicatingRootObject(Instance, Params);
 
 	// This is optional but typically we want to cache at least the NetRefHandle in the game instance to avoid doing map lookups to find it
 	if (Handle.IsValid())
@@ -81,8 +81,7 @@ UE::Net::FNetRefHandle UReplicatedTestObjectBridge::BeginReplication(FNetRefHand
 	check(OwnerHandle.IsValid());
 
 	// Create NetRefHandle for the registered fragments
-	Super::FCreateNetRefHandleParams Params;
-	FNetRefHandle Handle = Super::BeginReplication(OwnerHandle, SubObjectInstance, InsertRelativeToSubObjectHandle, Params, InsertionOrder);
+	FNetRefHandle Handle = Super::StartReplicatingSubObject(OwnerHandle, SubObjectInstance, InsertRelativeToSubObjectHandle, InsertionOrder);
 
 	if (Handle.IsValid())
 	{
@@ -94,6 +93,10 @@ UE::Net::FNetRefHandle UReplicatedTestObjectBridge::BeginReplication(FNetRefHand
 	return Handle;
 }
 
+void UReplicatedTestObjectBridge::EndReplication(UReplicatedTestObject* Instance, EEndReplicationFlags Flags)
+{
+	StopReplicatingNetObject(Instance, Flags);
+}
 
 bool UReplicatedTestObjectBridge::WriteCreationHeader(UE::Net::FNetSerializationContext& Context, FNetRefHandle Handle)
 {
diff --git a/Engine/Plugins/Runtime/ReplicationSystemTestPlugin/Source/Private/Tests/ReplicationSystem/ReplicatedTestObject.h b/Engine/Plugins/Runtime/ReplicationSystemTestPlugin/Source/Private/Tests/ReplicationSystem/ReplicatedTestObject.h
--- a/Engine/Plugins/Runtime/ReplicationSystemTestPlugin/Source/Private/Tests/ReplicationSystem/ReplicatedTestObject.h
+++ b/Engine/Plugins/Runtime/ReplicationSystemTestPlugin/Source/Private/Tests/ReplicationSystem/ReplicatedTestObject.h
@@ -539,9 +539,11 @@ public:
 	// This is the local Interface, it is up to each bridge implementation to define the interface for that type
 	// In this example we have methods that directly uses UReplicatedTestObject;
 	FNetRefHandle BeginReplication(UReplicatedTestObject* Instance);
-	FNetRefHandle BeginReplication(UReplicatedTestObject* Instance, const UObjectReplicationBridge::FCreateNetRefHandleParams& Params);
+	FNetRefHandle BeginReplication(UReplicatedTestObject* Instance, const UObjectReplicationBridge::FRootObjectReplicationParams& Params);
 	FNetRefHandle BeginReplication(FNetRefHandle OwnerHandle, UReplicatedTestObject* Instance, FNetRefHandle InsertRelativeToSubObjectHandle = FNetRefHandle::GetInvalid(), ESubObjectInsertionOrder InsertionOrder = UReplicationBridge::ESubObjectInsertionOrder::None);
 
+	void EndReplication(UReplicatedTestObject* Instance, EEndReplicationFlags Flags=EEndReplicationFlags::Destroy);
+
 	// For testing we expose some things that normally are not accessible
 	const UE::Net::FReplicationInstanceProtocol* GetReplicationInstanceProtocol(FNetRefHandle Handle) const;
 
diff --git a/Engine/Plugins/Runtime/ReplicationSystemTestPlugin/Source/Private/Tests/ReplicationSystem/ReplicationSystemServerClientTestFixture.cpp b/Engine/Plugins/Runtime/ReplicationSystemTestPlugin/Source/Private/Tests/ReplicationSystem/ReplicationSystemServerClientTestFixture.cpp
--- a/Engine/Plugins/Runtime/ReplicationSystemTestPlugin/Source/Private/Tests/ReplicationSystem/ReplicationSystemServerClientTestFixture.cpp
+++ b/Engine/Plugins/Runtime/ReplicationSystemTestPlugin/Source/Private/Tests/ReplicationSystem/ReplicationSystemServerClientTestFixture.cpp
@@ -119,7 +119,7 @@ uint32 FReplicationSystemTestNode::GetNetTraceId() const
 	return ReplicationSystem ? ReplicationSystem->GetId() : ~0U;
 }
 
-UTestReplicatedIrisObject* FReplicationSystemTestNode::CreateObject(const UObjectReplicationBridge::FCreateNetRefHandleParams& Params, UTestReplicatedIrisObject::FComponents* ComponentsToCreate)
+UTestReplicatedIrisObject* FReplicationSystemTestNode::CreateObject(const UObjectReplicationBridge::FRootObjectReplicationParams& Params, UTestReplicatedIrisObject::FComponents* ComponentsToCreate)
 {
 	UTestReplicatedIrisObject* CreatedObject = NewObject<UTestReplicatedIrisObject>();
 
diff --git a/Engine/Plugins/Runtime/ReplicationSystemTestPlugin/Source/Private/Tests/ReplicationSystem/ReplicationSystemServerClientTestFixture.h b/Engine/Plugins/Runtime/ReplicationSystemTestPlugin/Source/Private/Tests/ReplicationSystem/ReplicationSystemServerClientTestFixture.h
--- a/Engine/Plugins/Runtime/ReplicationSystemTestPlugin/Source/Private/Tests/ReplicationSystem/ReplicationSystemServerClientTestFixture.h
+++ b/Engine/Plugins/Runtime/ReplicationSystemTestPlugin/Source/Private/Tests/ReplicationSystem/ReplicationSystemServerClientTestFixture.h
@@ -128,7 +128,7 @@ public:
 		return static_cast<T*>(ReplicationBridge->GetReplicatedObject(Handle));
 	}
 
-	UTestReplicatedIrisObject* CreateObject(const UObjectReplicationBridge::FCreateNetRefHandleParams& Params, UTestReplicatedIrisObject::FComponents* ComponentsToCreate=nullptr);
+	UTestReplicatedIrisObject* CreateObject(const UObjectReplicationBridge::FRootObjectReplicationParams& Params, UTestReplicatedIrisObject::FComponents* ComponentsToCreate=nullptr);
 	UTestReplicatedIrisObject* CreateObject(uint32 NumComponents, uint32 NumIrisComponents);
 	UTestReplicatedIrisObject* CreateSubObject(FNetRefHandle Owner, uint32 NumComponents, uint32 NumIrisComponents);
 	UTestReplicatedIrisObject* CreateObject(const UTestReplicatedIrisObject::FComponents& Components = UTestReplicatedIrisObject::FComponents());
diff --git a/Engine/Plugins/Runtime/ReplicationSystemTestPlugin/Source/Private/Tests/ReplicationSystem/TestDependentObjects.cpp b/Engine/Plugins/Runtime/ReplicationSystemTestPlugin/Source/Private/Tests/ReplicationSystem/TestDependentObjects.cpp
--- a/Engine/Plugins/Runtime/ReplicationSystemTestPlugin/Source/Private/Tests/ReplicationSystem/TestDependentObjects.cpp
+++ b/Engine/Plugins/Runtime/ReplicationSystemTestPlugin/Source/Private/Tests/ReplicationSystem/TestDependentObjects.cpp
@@ -269,7 +269,7 @@ UE_NET_TEST_FIXTURE(FReplicationSystemServerClientTestFixture, TestDependentObje
 
 	// Spawn second object on server that later will be added as a dependent object
 	// With high PollFramePeriod so that it will not replicate in a while unless it is a dependent
-	UObjectReplicationBridge::FCreateNetRefHandleParams Params;
+	UObjectReplicationBridge::FRootObjectReplicationParams Params;
 	Params.PollFrequency = Server->ConvertPollPeriodIntoFrequency(255U);
 	UTestReplicatedIrisObject* ServerDependentObject = Server->CreateObject(Params);
 
@@ -336,7 +336,7 @@ UE_NET_TEST_FIXTURE(FReplicationSystemServerClientTestFixture, TestDependentObje
 	FReplicationSystemTestClient* Client = CreateClient();
 
 	// Spawn object on server
-	UObjectReplicationBridge::FCreateNetRefHandleParams Params;
+	UObjectReplicationBridge::FRootObjectReplicationParams Params;
 
 	// Setup different poll frequencies for the objects
 	Params.PollFrequency = Server->ConvertPollPeriodIntoFrequency(10U);
@@ -415,7 +415,7 @@ UE_NET_TEST_FIXTURE(FReplicationSystemServerClientTestFixture, TestDependentObje
 	FReplicationSystemTestClient* Client = CreateClient();
 
 	// Spawn object on server
-	UObjectReplicationBridge::FCreateNetRefHandleParams Params;
+	UObjectReplicationBridge::FRootObjectReplicationParams Params;
 	Params.PollFrequency = Server->ConvertPollPeriodIntoFrequency(14U);
 	UTestReplicatedIrisObject* ServerObject = Server->CreateObject(Params);
 
diff --git a/Engine/Plugins/Runtime/ReplicationSystemTestPlugin/Source/Private/Tests/ReplicationSystem/TestDirtyNetObjectTracker.cpp b/Engine/Plugins/Runtime/ReplicationSystemTestPlugin/Source/Private/Tests/ReplicationSystem/TestDirtyNetObjectTracker.cpp
--- a/Engine/Plugins/Runtime/ReplicationSystemTestPlugin/Source/Private/Tests/ReplicationSystem/TestDirtyNetObjectTracker.cpp
+++ b/Engine/Plugins/Runtime/ReplicationSystemTestPlugin/Source/Private/Tests/ReplicationSystem/TestDirtyNetObjectTracker.cpp
@@ -142,7 +142,7 @@ UE_NET_TEST_FIXTURE(FReplicationSystemServerClientTestFixture, GlobalDirtyTracke
 	FReplicationSystemTestClient* Client = CreateClient();
 
 	// Spawn object on server that is polled only every 3 frames
-	UObjectReplicationBridge::FCreateNetRefHandleParams Params;
+	UObjectReplicationBridge::FRootObjectReplicationParams Params;
 	const uint32 PollPeriod = 2;
 	const float PollFrequency = Server->ConvertPollPeriodIntoFrequency(PollPeriod);
 	Params.PollFrequency = PollFrequency;
@@ -186,7 +186,7 @@ UE_NET_TEST_FIXTURE(FReplicationSystemServerClientTestFixture, NetForceUpdateOth
 	FReplicationSystemTestClient* Client = CreateClient();
 
 	// Spawn object on server that is polled late in order to test ForceNetUpdate
-	UObjectReplicationBridge::FCreateNetRefHandleParams Params;
+	UObjectReplicationBridge::FRootObjectReplicationParams Params;
 	const uint32 PollPeriod = 100;
 	const float PollFrequency = Server->ConvertPollPeriodIntoFrequency(PollPeriod);
 	Params.PollFrequency = PollFrequency;
@@ -306,7 +306,7 @@ UE_NET_TEST_FIXTURE(FReplicationSystemServerClientTestFixture, DirtyOtherObjectI
 	FReplicationSystemTestClient* Client = CreateClient();
 
 	// Spawn object on server that is polled every frame.
-	UObjectReplicationBridge::FCreateNetRefHandleParams Params;
+	UObjectReplicationBridge::FRootObjectReplicationParams Params;
 	Params.bUseClassConfigDynamicFilter = true;
 	Params.bNeedsPreUpdate = true;
 
@@ -366,7 +366,7 @@ UE_NET_TEST_FIXTURE(FReplicationSystemServerClientTestFixture, PushModelMarkSelf
 	FReplicationSystemTestClient* Client = CreateClient();
 
 	// Spawn object with a PreUpdate call
-	UObjectReplicationBridge::FCreateNetRefHandleParams Params;
+	UObjectReplicationBridge::FRootObjectReplicationParams Params;
 	Params.bNeedsPreUpdate = true;
 	UTestReplicatedIrisObject::FComponents ComponentsToCreate = { .ObjectReferenceComponentCount = 1 };
 
diff --git a/Engine/Plugins/Runtime/ReplicationSystemTestPlugin/Source/Private/Tests/ReplicationSystem/TestFlushBeforeDestroy.cpp b/Engine/Plugins/Runtime/ReplicationSystemTestPlugin/Source/Private/Tests/ReplicationSystem/TestFlushBeforeDestroy.cpp
--- a/Engine/Plugins/Runtime/ReplicationSystemTestPlugin/Source/Private/Tests/ReplicationSystem/TestFlushBeforeDestroy.cpp
+++ b/Engine/Plugins/Runtime/ReplicationSystemTestPlugin/Source/Private/Tests/ReplicationSystem/TestFlushBeforeDestroy.cpp
@@ -713,7 +713,7 @@ UE_NET_TEST_FIXTURE(FTestFlushBeforeDestroyFixture, TestDroppedPendingTearOffIsC
 	}
 
 	// Request tearoff
-	Server->ReplicationBridge->EndReplication(ServerObject->NetRefHandle, EEndReplicationFlags::TearOff);
+	Server->ReplicationBridge->EndReplication(ServerObject, EEndReplicationFlags::TearOff);
 
 	// Send packet so that we have creationdata in flight
 	Server->PreSendUpdate();
@@ -721,7 +721,7 @@ UE_NET_TEST_FIXTURE(FTestFlushBeforeDestroyFixture, TestDroppedPendingTearOffIsC
 	Server->PostSendUpdate();
 
 	// Force destroy object already pending tearoff/flush. DestroyLocalNetHandle will invalidate cached creationinfo.
-	Server->ReplicationBridge->EndReplication(ServerObject->NetRefHandle, EEndReplicationFlags::Destroy);
+	Server->ReplicationBridge->EndReplication(ServerObject, EEndReplicationFlags::Destroy);
 
 	// Drop and notify that the packet while object still is the state waitoncreateconfirmation as we have not yet updated scope.
 	// When this failed it did put the state of the object back in PendingCreate even though we no longer had any cached creationinfo.
@@ -772,14 +772,14 @@ UE_NET_TEST_FIXTURE(FTestFlushBeforeDestroyFixture, TestPendingCreateTearOffIsCa
 	}
 
 	// Request tearoff
-	Server->ReplicationBridge->EndReplication(ServerObject->NetRefHandle, EEndReplicationFlags::TearOff);
+	Server->ReplicationBridge->EndReplication(ServerObject, EEndReplicationFlags::TearOff);
 
 	// PreUpdate to update scoping to get the object into the PendingCreate state.
 	Server->PreSendUpdate();
 	Server->PostSendUpdate();
 
 	// Force destroy object already pending tearoff/flush. DestroyLocalNetHandle will invalidate cached creationinfo.
-	Server->ReplicationBridge->EndReplication(ServerObject->NetRefHandle, EEndReplicationFlags::Destroy);
+	Server->ReplicationBridge->EndReplication(ServerObject, EEndReplicationFlags::Destroy);
 
 	// Send a packet.
 	Server->PreSendUpdate();
@@ -811,14 +811,14 @@ UE_NET_TEST_FIXTURE(FTestFlushBeforeDestroyFixture, TestPendingCreateTearOffIsNo
 	}
 
 	// Request tearoff
-	Server->ReplicationBridge->EndReplication(ServerObject->NetRefHandle, EEndReplicationFlags::TearOff);
+	Server->ReplicationBridge->EndReplication(ServerObject, EEndReplicationFlags::TearOff);
 
 	// PreUpdate to update scoping to get the object into the PendingCreate state.
 	Server->PreSendUpdate();
 	Server->PostSendUpdate();
 
 	// This should be ignored as we are already pending tear off.
-	Server->ReplicationBridge->EndReplication(ServerObject->NetRefHandle, EEndReplicationFlags::Destroy);
+	Server->ReplicationBridge->EndReplication(ServerObject, EEndReplicationFlags::Destroy);
 
 	// Deliver a packet, this should flush the object and deliver the attachment
 	Server->PreSendUpdate();
@@ -854,7 +854,7 @@ UE_NET_TEST_FIXTURE(FTestFlushBeforeDestroyFixture, TestDroppedTearOffIsNotCance
 	}
 
 	// Request tearoff
-	Server->ReplicationBridge->EndReplication(ServerObject->NetRefHandle, EEndReplicationFlags::TearOff);
+	Server->ReplicationBridge->EndReplication(ServerObject, EEndReplicationFlags::TearOff);
 
 	// Send packet so that we have creationdata in flight
 	Server->PreSendUpdate();
@@ -862,7 +862,7 @@ UE_NET_TEST_FIXTURE(FTestFlushBeforeDestroyFixture, TestDroppedTearOffIsNotCance
 	Server->PostSendUpdate();
 
 	// Force destroy object already pending tearoff/flush. This should be ignored
-	Server->ReplicationBridge->EndReplication(ServerObject->NetRefHandle, EEndReplicationFlags::Destroy);
+	Server->ReplicationBridge->EndReplication(ServerObject, EEndReplicationFlags::Destroy);
 
 	// Drop and notify that the packet while object still is the state waitoncreateconfirmation as we have not yet updated scope.
 	Server->DeliverTo(Client, false);
diff --git a/Engine/Plugins/Runtime/ReplicationSystemTestPlugin/Source/Private/Tests/ReplicationSystem/TestObjectPollFrequencyLimiter.cpp b/Engine/Plugins/Runtime/ReplicationSystemTestPlugin/Source/Private/Tests/ReplicationSystem/TestObjectPollFrequencyLimiter.cpp
--- a/Engine/Plugins/Runtime/ReplicationSystemTestPlugin/Source/Private/Tests/ReplicationSystem/TestObjectPollFrequencyLimiter.cpp
+++ b/Engine/Plugins/Runtime/ReplicationSystemTestPlugin/Source/Private/Tests/ReplicationSystem/TestObjectPollFrequencyLimiter.cpp
@@ -121,7 +121,7 @@ UE_NET_TEST_FIXTURE(FReplicationSystemServerClientTestFixture, TestSubObjectsPol
 	// Spawn object on server polled every 3 frames
 	const uint32 PollPeriod = 3;
 	const float PollFrequency = Server->ConvertPollPeriodIntoFrequency(PollPeriod);
-	UObjectReplicationBridge::FCreateNetRefHandleParams Params;
+	UObjectReplicationBridge::FRootObjectReplicationParams Params;
 	Params.PollFrequency = PollFrequency;
 	Params.bUseClassConfigDynamicFilter = true;
 	Params.bNeedsPreUpdate = true;
diff --git a/Engine/Plugins/Runtime/ReplicationSystemTestPlugin/Source/Private/Tests/ReplicationSystem/TestObjectSplitting.cpp b/Engine/Plugins/Runtime/ReplicationSystemTestPlugin/Source/Private/Tests/ReplicationSystem/TestObjectSplitting.cpp
--- a/Engine/Plugins/Runtime/ReplicationSystemTestPlugin/Source/Private/Tests/ReplicationSystem/TestObjectSplitting.cpp
+++ b/Engine/Plugins/Runtime/ReplicationSystemTestPlugin/Source/Private/Tests/ReplicationSystem/TestObjectSplitting.cpp
@@ -1111,7 +1111,7 @@ UE_NET_TEST_FIXTURE(FSplitObjectTestFixture, TestDependentObjectCannotBeDestroye
 	FReplicationSystemTestClient* Client = CreateClient();
 
 	UTestReplicatedIrisObject* ServerObject = CreateHugeObject(Server);
-	UTestReplicatedIrisObject* ServerDependentObject = Server->CreateObject(UObjectReplicationBridge::FCreateNetRefHandleParams{});
+	UTestReplicatedIrisObject* ServerDependentObject = Server->CreateObject(UObjectReplicationBridge::FRootObjectReplicationParams{});
 
 	ServerBridge->AddDependentObject(ServerObject->NetRefHandle, ServerDependentObject->NetRefHandle);
 
diff --git a/Engine/Plugins/Runtime/ReplicationSystemTestPlugin/Source/Private/Tests/ReplicationSystem/TestPreUpdateCallback.cpp b/Engine/Plugins/Runtime/ReplicationSystemTestPlugin/Source/Private/Tests/ReplicationSystem/TestPreUpdateCallback.cpp
--- a/Engine/Plugins/Runtime/ReplicationSystemTestPlugin/Source/Private/Tests/ReplicationSystem/TestPreUpdateCallback.cpp
+++ b/Engine/Plugins/Runtime/ReplicationSystemTestPlugin/Source/Private/Tests/ReplicationSystem/TestPreUpdateCallback.cpp
@@ -21,7 +21,7 @@ UE_NET_TEST_FIXTURE(FReplicationSystemServerClientTestFixture, CreateObjectsInsi
 	FReplicationSystemTestClient* Client = CreateClient();
 
 	// Spawn object on server that is polled late in order to test ForceNetUpdate
-	UObjectReplicationBridge::FCreateNetRefHandleParams Params;
+	UObjectReplicationBridge::FRootObjectReplicationParams Params;
 	const uint32 PollPeriod = 100;
 	const float PollFrequency = Server->ConvertPollPeriodIntoFrequency(PollPeriod);
 	Params.PollFrequency = PollFrequency;
diff --git a/Engine/Plugins/Runtime/ReplicationSystemTestPlugin/Source/Private/Tests/ReplicationSystem/TestReplicationOperations.cpp b/Engine/Plugins/Runtime/ReplicationSystemTestPlugin/Source/Private/Tests/ReplicationSystem/TestReplicationOperations.cpp
--- a/Engine/Plugins/Runtime/ReplicationSystemTestPlugin/Source/Private/Tests/ReplicationSystem/TestReplicationOperations.cpp
+++ b/Engine/Plugins/Runtime/ReplicationSystemTestPlugin/Source/Private/Tests/ReplicationSystem/TestReplicationOperations.cpp
@@ -537,7 +537,7 @@ UE_NET_TEST_FIXTURE(FTestReplicationOperationsForObjectsFixture, StaleObjectPoin
 	ReplicationSystem->PostSendUpdate();
 
 	// Destroy first object and invalidate references to it.
-	ReplicationBridge->EndReplication(HandleToObject0);
+	ReplicationBridge->EndReplication(Object0);
 	DestroyObject(Object0);
 	constexpr bool bPerformFullPurge = false;
 	CollectGarbage(GARBAGE_COLLECTION_KEEPFLAGS, bPerformFullPurge);
diff --git a/Engine/Plugins/Runtime/ReplicationSystemTestPlugin/Source/Private/Tests/ReplicationSystem/TestReplicationSystem.cpp b/Engine/Plugins/Runtime/ReplicationSystemTestPlugin/Source/Private/Tests/ReplicationSystem/TestReplicationSystem.cpp
--- a/Engine/Plugins/Runtime/ReplicationSystemTestPlugin/Source/Private/Tests/ReplicationSystem/TestReplicationSystem.cpp
+++ b/Engine/Plugins/Runtime/ReplicationSystemTestPlugin/Source/Private/Tests/ReplicationSystem/TestReplicationSystem.cpp
@@ -1766,7 +1766,7 @@ UE_NET_TEST_FIXTURE(FReplicationSystemServerClientTestFixture, TestObjectPollFra
 	UTestReplicatedIrisObject* ServerObject = Server->CreateObject(0, 0);
 
 	// Spawn second object on server that later will be added as a dependent object
-	UObjectReplicationBridge::FCreateNetRefHandleParams Params;
+	UObjectReplicationBridge::FRootObjectReplicationParams Params;
 	Params.PollFrequency = Server->ConvertPollPeriodIntoFrequency(1U);
 	UTestReplicatedIrisObject* ServerObjectPolledEveryOtherFrame = Server->CreateObject(Params);
 
diff --git a/Engine/Source/Runtime/Engine/Classes/GameFramework/Actor.h b/Engine/Source/Runtime/Engine/Classes/GameFramework/Actor.h
--- a/Engine/Source/Runtime/Engine/Classes/GameFramework/Actor.h
+++ b/Engine/Source/Runtime/Engine/Classes/GameFramework/Actor.h
@@ -45,7 +45,7 @@ class UDataLayerInstance;
 class AWorldDataLayers;
 class IWorldPartitionCell;
 #if UE_WITH_IRIS
-struct FActorBeginReplicationParams;
+struct FActorReplicationParams;
 #endif // UE_WITH_IRIS
 class UActorFolder;
 struct FActorDataLayer;
@@ -3385,7 +3385,7 @@ protected:
 	 * Helper to BeginReplication passing on additional parameters to the ReplicationSystem, typically called from code overriding normal BeginReplication()
 	 * @param Params Additional parameters we want to pass on
 	 */
-	ENGINE_API void BeginReplication(const FActorBeginReplicationParams& Params);
+	ENGINE_API void BeginReplication(const FActorReplicationParams& Params);
 #endif // UE_WITH_IRIS
 
 	/**
diff --git a/Engine/Source/Runtime/Engine/Private/ActorReplication.cpp b/Engine/Source/Runtime/Engine/Private/ActorReplication.cpp
--- a/Engine/Source/Runtime/Engine/Private/ActorReplication.cpp
+++ b/Engine/Source/Runtime/Engine/Private/ActorReplication.cpp
@@ -1076,7 +1076,7 @@ void AActor::RegisterReplicationFragments(UE::Net::FFragmentRegistrationContext&
 	}
 }
 
-void AActor::BeginReplication(const FActorBeginReplicationParams& Params)
+void AActor::BeginReplication(const FActorReplicationParams& Params)
 {
 	UE::Net::FReplicationSystemUtil::BeginReplication(this, Params);
 	UpdateOwningNetConnection();
@@ -1084,7 +1084,7 @@ void AActor::BeginReplication(const FActorBeginReplicationParams& Params)
 
 void AActor::BeginReplication()
 {
-	const FActorBeginReplicationParams BeginReplicationParams;
+	const FActorReplicationParams BeginReplicationParams;
 
 	BeginReplication(BeginReplicationParams);
 }
diff --git a/Engine/Source/Runtime/Engine/Private/Net/Iris/ReplicationSystem/ActorReplicationBridge.cpp b/Engine/Source/Runtime/Engine/Private/Net/Iris/ReplicationSystem/ActorReplicationBridge.cpp
--- a/Engine/Source/Runtime/Engine/Private/Net/Iris/ReplicationSystem/ActorReplicationBridge.cpp
+++ b/Engine/Source/Runtime/Engine/Private/Net/Iris/ReplicationSystem/ActorReplicationBridge.cpp
@@ -181,7 +181,7 @@ void UActorReplicationBridge::Deinitialize()
 	ObjectReferencePackageMap = nullptr;
 }
 
-UE::Net::FNetRefHandle UActorReplicationBridge::BeginReplication(AActor* Actor, const FActorBeginReplicationParams& Params)
+UE::Net::FNetRefHandle UActorReplicationBridge::StartReplicatingActor(AActor* Actor, const FActorReplicationParams& Params)
 {
 	using namespace UE::Net;
 
@@ -267,7 +267,7 @@ UE::Net::FNetRefHandle UActorReplicationBridge::BeginReplication(AActor* Actor,
 	}
 
 	// Create handles for the registered fragments
-	Super::FCreateNetRefHandleParams CreateNetRefHandleParams = {
+	Super::FRootObjectReplicationParams RootObjectParams = {
 		.bNeedsPreUpdate = 1U,
 		.bNeedsWorldLocationUpdate = 1U,
 		.bIsDormant = Actor->NetDormancy > DORM_Awake,
@@ -276,10 +276,10 @@ UE::Net::FNetRefHandle UActorReplicationBridge::BeginReplication(AActor* Actor,
 	};
 
 #if !UE_BUILD_SHIPPING
-	ensureMsgf(!(Actor->bAlwaysRelevant || Actor->bOnlyRelevantToOwner) || CreateNetRefHandleParams.StaticPriority >= 1.0f, TEXT("Very low NetPriority %.02f for always relevant or owner relevant Actor %s. Set it to 1.0f or higher."), Actor->NetPriority, ToCStr(Actor->GetName()));
+	ensureMsgf(!(Actor->bAlwaysRelevant || Actor->bOnlyRelevantToOwner) || RootObjectParams.StaticPriority >= 1.0f, TEXT("Very low NetPriority %.02f for always relevant or owner relevant Actor %s. Set it to 1.0f or higher."), Actor->NetPriority, ToCStr(Actor->GetName()));
 #endif
 
-	FNetRefHandle ActorRefHandle = Super::BeginReplication(Actor, CreateNetRefHandleParams);
+	FNetRefHandle ActorRefHandle = StartReplicatingRootObject(Actor, RootObjectParams);
 
 	if (!ActorRefHandle.IsValid())
 	{
@@ -287,7 +287,7 @@ UE::Net::FNetRefHandle UActorReplicationBridge::BeginReplication(AActor* Actor,
 		return FNetRefHandle::GetInvalid();
 	}
 
-    UE_CLOG(Actor->bAlwaysRelevant, LogIrisBridge, Verbose, TEXT("BeginReplication of AlwaysRelevant actor %s"), *PrintObjectFromNetRefHandle(ActorRefHandle));
+    UE_CLOG(Actor->bAlwaysRelevant, LogIrisBridge, Verbose, TEXT("StartReplicatingActor of AlwaysRelevant actor %s"), *PrintObjectFromNetRefHandle(ActorRefHandle));
 
 	// Set owning connection filtering if actor is only relevant to owner
 	{
@@ -336,8 +336,7 @@ UE::Net::FNetRefHandle UActorReplicationBridge::BeginReplication(AActor* Actor,
 			UObject* SubObjectToReplicate = SubObjectInfo.GetSubObject();
 			if (IsValid(SubObjectToReplicate) && SubObjectInfo.NetCondition != ELifetimeCondition::COND_Never)
 			{
-				const UObjectReplicationBridge::FCreateNetRefHandleParams CreateNetRefParams;
-				FNetRefHandle SubObjectRefHandle = UObjectReplicationBridge::BeginReplication(ActorRefHandle, SubObjectToReplicate, CreateNetRefParams);
+				FNetRefHandle SubObjectRefHandle = StartReplicatingSubObject(ActorRefHandle, SubObjectToReplicate);
 				if (SubObjectRefHandle.IsValid() && SubObjectInfo.NetCondition != ELifetimeCondition::COND_None)
 				{
 					UObjectReplicationBridge::SetSubObjectNetCondition(SubObjectRefHandle, SubObjectInfo.NetCondition);
@@ -360,7 +359,7 @@ UE::Net::FNetRefHandle UActorReplicationBridge::BeginReplication(AActor* Actor,
 	return ActorRefHandle;
 }
 
-UE::Net::FNetRefHandle UActorReplicationBridge::BeginReplication(FNetRefHandle OwnerHandle, UActorComponent* SubObject)
+UE::Net::FNetRefHandle UActorReplicationBridge::StartReplicatingComponent(FNetRefHandle OwnerHandle, UActorComponent* SubObject)
 {
 	using namespace UE::Net;
 
@@ -396,7 +395,7 @@ UE::Net::FNetRefHandle UActorReplicationBridge::BeginReplication(FNetRefHandle O
 		}
 
 		// Start replicating the subobject with its owner.
-		ReplicatedComponentHandle = Super::BeginReplication(OwnerHandle, SubObject, UObjectReplicationBridge::FCreateNetRefHandleParams());
+		ReplicatedComponentHandle = StartReplicatingSubObject(OwnerHandle, SubObject);
 	}
 
 	if (!ReplicatedComponentHandle.IsValid())
@@ -412,13 +411,12 @@ UE::Net::FNetRefHandle UActorReplicationBridge::BeginReplication(FNetRefHandle O
 	}
 
 	// Begin replication for any SubObjects registered by the component
-	const UObjectReplicationBridge::FCreateNetRefHandleParams CreateNetRefParams;
 	for (const FSubObjectRegistry::FEntry& SubObjectInfo : RepComponentInfo->SubObjects.GetRegistryList())
 	{
 		UObject* SubObjectToReplicate = SubObjectInfo.GetSubObject();
 		if (IsValid(SubObjectToReplicate) && SubObjectInfo.NetCondition != ELifetimeCondition::COND_Never)
 		{
-			FNetRefHandle SubObjectHandle = UObjectReplicationBridge::BeginReplication(OwnerHandle, SubObjectToReplicate, ReplicatedComponentHandle, CreateNetRefParams, UReplicationBridge::ESubObjectInsertionOrder::ReplicateWith);
+			FNetRefHandle SubObjectHandle = StartReplicatingSubObject(OwnerHandle, SubObjectToReplicate, ReplicatedComponentHandle, UReplicationBridge::ESubObjectInsertionOrder::ReplicateWith);
 			if (SubObjectHandle.IsValid() && SubObjectInfo.NetCondition != ELifetimeCondition::COND_None)
 			{
 				SetSubObjectNetCondition(SubObjectHandle, SubObjectInfo.NetCondition);
@@ -429,59 +427,62 @@ UE::Net::FNetRefHandle UActorReplicationBridge::BeginReplication(FNetRefHandle O
 	return ReplicatedComponentHandle;
 }
 
-void UActorReplicationBridge::EndReplication(AActor* Actor, EEndPlayReason::Type EndPlayReason)
+void UActorReplicationBridge::StopReplicatingActor(AActor* Actor, EEndPlayReason::Type EndPlayReason)
 {
 	using namespace UE::Net;
 
 	FNetRefHandle RefHandle = GetReplicatedRefHandle(Actor, EGetRefHandleFlags::EvenIfGarbage);
-	if (RefHandle.IsValid())
+	if (!RefHandle.IsValid())
 	{
-		UE_LOG(LogIrisBridge, Verbose, TEXT("EndReplication for %s. Reason %s "), *PrintObjectFromNetRefHandle(RefHandle), *UEnum::GetValueAsString(TEXT("Engine.EEndPlayReason"), EndPlayReason));
-		ensureMsgf(IsValid(Actor), TEXT("Calling EndReplication for Invalid Object: %s."), *PrintObjectFromNetRefHandle(RefHandle));
+		// Already not replicated
+		return;
+	}
+
+	UE_LOG(LogIrisBridge, Verbose, TEXT("StopReplicatingActor %s. Reason %s "), *PrintObjectFromNetRefHandle(RefHandle), *UEnum::GetValueAsString(TEXT("Engine.EEndPlayReason"), EndPlayReason));
+	ensureMsgf(IsValid(Actor), TEXT("StopReplicatingActor called on invalid actor tied to handle: %s."), *PrintObjectFromNetRefHandle(RefHandle));
 	
-		EEndReplicationFlags Flags = EEndReplicationFlags::None;
-		const bool bShouldDestroyObject = EndPlayReason == EEndPlayReason::Destroyed;		
-		if (bShouldDestroyObject)
-		{ 
-			Flags |= EEndReplicationFlags::Destroy | EEndReplicationFlags::DestroyNetHandle | EEndReplicationFlags::ClearNetPushId;
-		}
+	EEndReplicationFlags Flags = EEndReplicationFlags::None;
+	const bool bShouldDestroyObject = EndPlayReason == EEndPlayReason::Destroyed;		
+	if (bShouldDestroyObject)
+	{ 
+		Flags |= EEndReplicationFlags::Destroy | EEndReplicationFlags::DestroyNetHandle | EEndReplicationFlags::ClearNetPushId;
+	}
 
-		// If we are shutting down or the actor is a nettemporary we do not need to validate that we are not detaching remote instances by accident.
-		const bool bIsShuttingDown = (EndPlayReason == EEndPlayReason::EndPlayInEditor) || (EndPlayReason == EEndPlayReason::Quit);
-		if (bIsShuttingDown || Actor->bNetTemporary)
-		{
-			Flags |= EEndReplicationFlags::SkipPendingEndReplicationValidation;
-		}
+	// If we are shutting down or the actor is a nettemporary we do not need to validate that we are not detaching remote instances by accident.
+	const bool bIsShuttingDown = (EndPlayReason == EEndPlayReason::EndPlayInEditor) || (EndPlayReason == EEndPlayReason::Quit);
+	if (bIsShuttingDown || Actor->bNetTemporary)
+	{
+		Flags |= EEndReplicationFlags::SkipPendingEndReplicationValidation;
+	}
 					
-		const bool bShouldCreateDestructionInfo = RefHandle.IsStatic() && bShouldDestroyObject && GetReplicationSystem()->IsServer();
-		if (bShouldCreateDestructionInfo)
+	const bool bShouldCreateDestructionInfo = RefHandle.IsStatic() && bShouldDestroyObject && GetReplicationSystem()->IsServer();
+	if (bShouldCreateDestructionInfo)
+	{
+		UObjectReplicationBridge::FEndReplicationParameters EndReplicationParameters
 		{
-			UObjectReplicationBridge::FEndReplicationParameters EndReplicationParameters;
-
-			EndReplicationParameters.Location = Actor->GetActorLocation();
-			EndReplicationParameters.Level = Actor->GetLevel();
-			EndReplicationParameters.bUseDistanceBasedPrioritization = Actor->bAlwaysRelevant == false;
+			.Location = Actor->GetActorLocation(),
+			.Level = Actor->GetLevel(),
+			.bUseDistanceBasedPrioritization = (Actor->bAlwaysRelevant == false)
+		};
 
-			UObjectReplicationBridge::EndReplication(RefHandle, Flags, &EndReplicationParameters);		
-		}
-		else
-		{
-			UObjectReplicationBridge::EndReplication(RefHandle, Flags, nullptr);
-		}
+		// Store this static actor destruction info so it can be replicated if a client can know to delete it in the future.
+		StoreDestructionInfo(RefHandle, EndReplicationParameters);
 	}
+
+	StopReplicatingNetRefHandle(RefHandle, Flags);
 }
 
-void UActorReplicationBridge::EndReplicationForActorComponent(UActorComponent* ActorComponent, EEndReplicationFlags EndReplicationFlags)
+void UActorReplicationBridge::StopReplicatingComponent(UActorComponent* ActorComponent, EEndReplicationFlags EndReplicationFlags)
 {
 	using namespace UE::Net;
 
 	FNetRefHandle ComponentHandle = GetReplicatedRefHandle(ActorComponent, EGetRefHandleFlags::EvenIfGarbage);
 	if (ComponentHandle.IsValid())
 	{
-		UE_LOG(LogIrisBridge, Verbose, TEXT("EndReplicationForActorComponent for %s %s."), *GetNameSafe(ActorComponent), *ComponentHandle.ToString());
-		ensureMsgf(IsValid(ActorComponent), TEXT("Calling EndReplication for Invalid Object for %s %s."), *GetNameSafe(ActorComponent), *ComponentHandle.ToString());
+		UE_LOG(LogIrisBridge, Verbose, TEXT("StopReplicatingComponent for %s %s."), *GetNameSafe(ActorComponent), *ComponentHandle.ToString());
+		ensureMsgf(IsValid(ActorComponent), TEXT("StopReplicatingActor called on invalid component tied to handle: %s."), *ComponentHandle.ToString());
 		
-		UObjectReplicationBridge::EndReplication(ComponentHandle, EndReplicationFlags, nullptr);
+		StopReplicatingNetRefHandle(ComponentHandle, EndReplicationFlags);
 	}
 }
 
diff --git a/Engine/Source/Runtime/Engine/Private/Net/Iris/ReplicationSystem/ReplicationSystemUtil.cpp b/Engine/Source/Runtime/Engine/Private/Net/Iris/ReplicationSystem/ReplicationSystemUtil.cpp
--- a/Engine/Source/Runtime/Engine/Private/Net/Iris/ReplicationSystem/ReplicationSystemUtil.cpp
+++ b/Engine/Source/Runtime/Engine/Private/Net/Iris/ReplicationSystem/ReplicationSystemUtil.cpp
@@ -116,7 +116,7 @@ FNetHandle FReplicationSystemUtil::GetNetHandle(const UObject* Object)
 	return NetHandle;
 }
 
-void FReplicationSystemUtil::BeginReplication(AActor* Actor, const FActorBeginReplicationParams& Params)
+void FReplicationSystemUtil::BeginReplication(AActor* Actor, const FActorReplicationParams& Params)
 {
 	if (const UWorld* World = Actor->GetWorld())
 	{
@@ -126,7 +126,7 @@ void FReplicationSystemUtil::BeginReplication(AActor* Actor, const FActorBeginRe
 			{
 				if (UActorReplicationBridge* Bridge = Cast<UActorReplicationBridge>(ReplicationSystem->GetReplicationBridge()))
 				{
-					Bridge->BeginReplication(Actor, Params);
+					Bridge->StartReplicatingActor(Actor, Params);
 				}
 			}
 		});
@@ -135,7 +135,7 @@ void FReplicationSystemUtil::BeginReplication(AActor* Actor, const FActorBeginRe
 
 void FReplicationSystemUtil::BeginReplication(AActor* Actor)
 {
-	const FActorBeginReplicationParams BeginReplicationParams;
+	const FActorReplicationParams BeginReplicationParams;
 	return BeginReplication(Actor, BeginReplicationParams);
 }
 
@@ -147,7 +147,7 @@ void FReplicationSystemUtil::EndReplication(AActor* Actor, EEndPlayReason::Type
 	{
 		if (UActorReplicationBridge* Bridge = Cast<UActorReplicationBridge>(ReplicationSystem->GetReplicationBridge()))
 		{
-			Bridge->EndReplication(Actor, EndPlayReason);
+			Bridge->StopReplicatingActor(Actor, EndPlayReason);
 		}
 	});
 }
@@ -180,7 +180,7 @@ void FReplicationSystemUtil::BeginReplicationForActorComponent(FNetHandle ActorH
 						const FNetRefHandle OwnerRefHandle = Bridge->GetReplicatedRefHandle(ActorHandle);
 						if (OwnerRefHandle.IsValid())
 						{
-							Bridge->BeginReplication(OwnerRefHandle, ActorComp);
+							Bridge->StartReplicatingComponent(OwnerRefHandle, ActorComp);
 						}
 					}
 				}
@@ -209,7 +209,7 @@ void FReplicationSystemUtil::BeginReplicationForActorComponent(const AActor* Act
 					const FNetRefHandle ActorRefHandle = Bridge->GetReplicatedRefHandle(ActorHandle);
 					if (ActorRefHandle.IsValid())
 					{
-						Bridge->BeginReplication(ActorRefHandle, ActorComp);
+						Bridge->StartReplicatingComponent(ActorRefHandle, ActorComp);
 					}
 				}
 			}
@@ -242,8 +242,7 @@ void FReplicationSystemUtil::BeginReplicationForActorSubObject(const AActor* Act
 					const FNetRefHandle ActorRefHandle = Bridge->GetReplicatedRefHandle(ActorHandle);
 					if (ActorRefHandle.IsValid())
 					{
-						const UObjectReplicationBridge::FCreateNetRefHandleParams CreateNetRefParams;
-						const FNetRefHandle SubObjectRefHandle = Bridge->BeginReplication(ActorRefHandle, ActorSubObject, CreateNetRefParams);
+						const FNetRefHandle SubObjectRefHandle = Bridge->StartReplicatingSubObject(ActorRefHandle, ActorSubObject);
 						if (SubObjectRefHandle.IsValid() && NetCondition != ELifetimeCondition::COND_None)
 						{
 							Bridge->SetSubObjectNetCondition(SubObjectRefHandle, NetCondition);
@@ -278,9 +277,7 @@ void FReplicationSystemUtil::BeginReplicationForActorComponentSubObject(UActorCo
 						const FNetRefHandle ActorComponentRefHandle = Bridge->GetReplicatedRefHandle(ActorComponent);
 						if (ActorRefHandle.IsValid() && ActorComponentRefHandle.IsValid())
 						{
-							const UObjectReplicationBridge::FCreateNetRefHandleParams CreateNetRefParams;
-
-							const FNetRefHandle SubObjectRefHandle = Bridge->BeginReplication(ActorRefHandle, SubObject, ActorComponentRefHandle, CreateNetRefParams, UReplicationBridge::ESubObjectInsertionOrder::ReplicateWith);
+							const FNetRefHandle SubObjectRefHandle = Bridge->StartReplicatingSubObject(ActorRefHandle, SubObject, ActorComponentRefHandle, UReplicationBridge::ESubObjectInsertionOrder::ReplicateWith);
 							if (SubObjectRefHandle.IsValid() && NetCondition != ELifetimeCondition::COND_None)
 							{
 								Bridge->SetSubObjectNetCondition(SubObjectRefHandle, NetCondition);
@@ -300,7 +297,7 @@ void FReplicationSystemUtil::EndReplicationForActorComponent(UActorComponent* Ac
 		if (UActorReplicationBridge* Bridge = Cast<UActorReplicationBridge>(ReplicationSystem->GetReplicationBridge()))
 		{
 			constexpr EEndReplicationFlags EndReplicationFlags = EEndReplicationFlags::DestroyNetHandle | EEndReplicationFlags::ClearNetPushId;
-			Bridge->EndReplicationForActorComponent(ActorComponent, EndReplicationFlags);
+			Bridge->StopReplicatingComponent(ActorComponent, EndReplicationFlags);
 		}
 	});
 }
@@ -312,7 +309,7 @@ void FReplicationSystemUtil::EndReplicationForActorSubObject(const AActor* Actor
 		if (UActorReplicationBridge* Bridge = Cast<UActorReplicationBridge>(ReplicationSystem->GetReplicationBridge()))
 		{				
 			constexpr EEndReplicationFlags EndReplicationFlags = EEndReplicationFlags::Destroy | EEndReplicationFlags::DestroyNetHandle | EEndReplicationFlags::ClearNetPushId;
-			Bridge->EndReplication(SubObject, EndReplicationFlags);
+			Bridge->StopReplicatingNetObject(SubObject, EndReplicationFlags);
 		}
 	});
 }
@@ -351,8 +348,8 @@ void FReplicationSystemUtil::AddDependentActor(const AActor* Parent, AActor* Chi
 					FNetRefHandle ChildRefHandle = Bridge->GetReplicatedRefHandle(Child);
 					if (!ChildRefHandle.IsValid())
 					{
-						const FActorBeginReplicationParams BeginReplicationParams;
-						ChildRefHandle = Bridge->BeginReplication(Child, BeginReplicationParams);
+						const FActorReplicationParams BeginReplicationParams;
+						ChildRefHandle = Bridge->StartReplicatingActor(Child, BeginReplicationParams);
 					}
 					if (ensureMsgf(ChildRefHandle.IsValid(), TEXT("FReplicationSystemUtil::AddDependentActor Child %s is not replicated"), *GetPathNameSafe(Child)))
 					{
diff --git a/Engine/Source/Runtime/Engine/Private/NetDriver.cpp b/Engine/Source/Runtime/Engine/Private/NetDriver.cpp
--- a/Engine/Source/Runtime/Engine/Private/NetDriver.cpp
+++ b/Engine/Source/Runtime/Engine/Private/NetDriver.cpp
@@ -4051,7 +4051,7 @@ void UNetDriver::DeleteSubObjectOnClients(AActor* Actor, UObject* SubObject)
 		if (UActorReplicationBridge* Bridge = Cast<UActorReplicationBridge>(ReplicationSystem->GetReplicationBridge()))
 		{
 			constexpr EEndReplicationFlags EndReplicationFlags = EEndReplicationFlags::Destroy | EEndReplicationFlags::DestroyNetHandle | EEndReplicationFlags::ClearNetPushId;
-			Bridge->EndReplication(SubObject, EndReplicationFlags);
+			Bridge->StopReplicatingNetObject(SubObject, EndReplicationFlags);
 		}
 	}
 	else
@@ -4070,8 +4070,7 @@ void UNetDriver::TearOffSubObjectOnClients(AActor* Actor, UObject* SubObject)
 	{
 		if (UActorReplicationBridge* Bridge = Cast<UActorReplicationBridge>(ReplicationSystem->GetReplicationBridge()))
 		{
-			constexpr EEndReplicationFlags EndReplicationFlags = EEndReplicationFlags::TearOff;
-			Bridge->EndReplication(SubObject, EndReplicationFlags);
+			Bridge->StopReplicatingNetObject(SubObject, EEndReplicationFlags::TearOff);
 		}
 	}
 	else
diff --git a/Engine/Source/Runtime/Engine/Public/Net/Iris/ReplicationSystem/ActorReplicationBridge.h b/Engine/Source/Runtime/Engine/Public/Net/Iris/ReplicationSystem/ActorReplicationBridge.h
--- a/Engine/Source/Runtime/Engine/Public/Net/Iris/ReplicationSystem/ActorReplicationBridge.h
+++ b/Engine/Source/Runtime/Engine/Public/Net/Iris/ReplicationSystem/ActorReplicationBridge.h
@@ -36,8 +36,7 @@ ENGINE_API bool ShouldUseIrisReplication(const UObject* Actor);
 
 }
 
-/** Parameters passed to UActorReplicationBridge::BeginReplication. */
-struct FActorBeginReplicationParams
+struct FActorReplicationParams
 {
 	/** When true we ignore the configured dynamic filter for this actor type and use the explicit filter instead */
 	bool bOverrideDynamicFilterConfig = false;
@@ -68,19 +67,19 @@ public:
 	inline UNetDriver* GetNetDriver() const { return NetDriver; }
 
 	/** Begin replication of an actor and its registered ActorComponents and SubObjects. */
-	ENGINE_API FNetRefHandle BeginReplication(AActor* Instance, const FActorBeginReplicationParams& Params);
+	ENGINE_API FNetRefHandle StartReplicatingActor(AActor* Instance, const FActorReplicationParams& Params);
 
 	/** Stop replicating an actor. Will destroy handle for actor and registered subobjects. */
-	ENGINE_API void EndReplication(AActor* Actor, EEndPlayReason::Type EndPlayReason);
+	ENGINE_API void StopReplicatingActor(AActor* Actor, EEndPlayReason::Type EndPlayReason);
 		
 	/**
 	 * Begin replication of an ActorComponent and its registered SubObjects, 
 	 * if the ActorComponent already is replicated any set NetObjectConditions will be updated.
 	*/
-	ENGINE_API FNetRefHandle BeginReplication(FNetRefHandle OwnerHandle, UActorComponent* ActorComponent);
+	ENGINE_API FNetRefHandle StartReplicatingComponent(FNetRefHandle OwnerHandle, UActorComponent* ActorComponent);
 
 	/** Stop replicating an ActorComponent and its associated SubObjects. */
-	ENGINE_API void EndReplicationForActorComponent(UActorComponent* ActorComponent, EEndReplicationFlags EndReplicationFlags = EEndReplicationFlags::None);
+	ENGINE_API void StopReplicatingComponent(UActorComponent* ActorComponent, EEndReplicationFlags EndReplicationFlags = EEndReplicationFlags::None);
 
 	/** Get object reference packagemap. Used in special cases where serialization hasn't been converted to use NetSerializers.  */
 	UIrisObjectReferencePackageMap* GetObjectReferencePackageMap() const { return ObjectReferencePackageMap; }
@@ -101,8 +100,6 @@ public:
 	 */
 	ENGINE_API void ConsumeNetMetrics(TArray<FAnalyticsEventAttribute>& OutAttrs);
 
-	using UObjectReplicationBridge::EndReplication;
-
 protected:
 
 	// UObjectReplicationBridge
diff --git a/Engine/Source/Runtime/Engine/Public/Net/Iris/ReplicationSystem/ReplicationSystemUtil.h b/Engine/Source/Runtime/Engine/Public/Net/Iris/ReplicationSystem/ReplicationSystemUtil.h
--- a/Engine/Source/Runtime/Engine/Public/Net/Iris/ReplicationSystem/ReplicationSystemUtil.h
+++ b/Engine/Source/Runtime/Engine/Public/Net/Iris/ReplicationSystem/ReplicationSystemUtil.h
@@ -12,7 +12,7 @@
 
 class AActor;
 class UActorComponent;
-struct FActorBeginReplicationParams;
+struct FActorReplicationParams;
 class UActorReplicationBridge;
 class UObject;
 class FRepChangedPropertyTracker;
@@ -54,7 +54,7 @@ struct FReplicationSystemUtil
 	ENGINE_API static FNetHandle GetNetHandle(const UObject* Object);
 
 	/** Begins replication of an actor and all of its registered subobjects. If any ReplicationSystem wants to replicate the actor a NetHandle will be created. */
-	ENGINE_API static void BeginReplication(AActor* Actor, const FActorBeginReplicationParams& Params);
+	ENGINE_API static void BeginReplication(AActor* Actor, const FActorReplicationParams& Params);
 
 	/** Begins replication of an actor and all of its registered subobjects. If any ReplicationSystem wants to replicate the actor a NetHandle will be created. */
 	ENGINE_API static void BeginReplication(AActor* Actor);
diff --git a/Engine/Source/Runtime/Experimental/Iris/Core/Private/Iris/ReplicationSystem/ObjectReplicationBridge.cpp b/Engine/Source/Runtime/Experimental/Iris/Core/Private/Iris/ReplicationSystem/ObjectReplicationBridge.cpp
--- a/Engine/Source/Runtime/Experimental/Iris/Core/Private/Iris/ReplicationSystem/ObjectReplicationBridge.cpp
+++ b/Engine/Source/Runtime/Experimental/Iris/Core/Private/Iris/ReplicationSystem/ObjectReplicationBridge.cpp
@@ -214,7 +214,7 @@ void UObjectReplicationBridge::AddStaticDestructionInfo(const FString& ObjectPat
 	UE::Net::FNetObjectReference ObjectRef = GetOrCreateObjectReference(ObjectPath, Outer);
 	if (ObjectRef.IsValid())
 	{
-		InternalAddDestructionInfo(ObjectRef.GetRefHandle(), Parameters);
+		StoreDestructionInfo(ObjectRef.GetRefHandle(), Parameters);
 	}
 }
 
@@ -245,19 +245,11 @@ UE::Net::FNetRefHandle UObjectReplicationBridge::GetReplicatedRefHandle(FNetHand
 	return ObjectData.RefHandle;
 }
 
-UE::Net::FNetRefHandle UObjectReplicationBridge::BeginReplication(UObject* Instance, const FCreateNetRefHandleParams& Params)
+UE::Net::FNetRefHandle UObjectReplicationBridge::StartReplicatingNetObject(UObject* Instance, UE::Net::EReplicationFragmentTraits Traits)
 {
-	LLM_SCOPE_BYTAG(IrisState);
-
 	using namespace UE::Net;
 	using namespace UE::Net::Private;
 
-	if (bBlockBeginReplication)
-	{
-		ensureMsgf(false, TEXT("BeginReplication is not allowed during this operation. %s will not be replicated"), *GetNameSafe(Instance));
-		return FNetRefHandle::GetInvalid();
-	}
-
 	FNetRefHandle AllocatedRefHandle = ObjectReferenceCache->CreateObjectReferenceHandle(Instance);
 
 	// If we failed to assign a handle, or if the Handle already is replicating, just return the handle
@@ -271,12 +263,7 @@ UE::Net::FNetRefHandle UObjectReplicationBridge::BeginReplication(UObject* Insta
 		return AllocatedRefHandle;
 	}
 
-	IRIS_PROFILER_SCOPE(BeginReplication);
-	
-	// Register fragments
-	EReplicationFragmentTraits Traits = EReplicationFragmentTraits::CanReplicate;
-	Traits |= Params.bNeedsPreUpdate ? EReplicationFragmentTraits::NeedsPreSendUpdate : EReplicationFragmentTraits::None;
-	Traits |= Params.bNeedsWorldLocationUpdate ? EReplicationFragmentTraits::NeedsWorldLocationUpdate : EReplicationFragmentTraits::None;
+	IRIS_PROFILER_SCOPE(StartReplicatingNetObject);
 	
 	FFragmentRegistrationContext FragmentRegistrationContext(GetReplicationStateDescriptorRegistry(), GetReplicationSystem(), Traits);
 
@@ -325,103 +312,134 @@ UE::Net::FNetRefHandle UObjectReplicationBridge::BeginReplication(UObject* Insta
 		const bool bIsValidProtocol = FReplicationProtocolManager::ValidateReplicationProtocol(ReplicationProtocol, RegisteredFragments);
 		if (!bIsValidProtocol)
 		{
-			UE_LOG_OBJECTREPLICATIONBRIDGE(Error, TEXT("BeginReplication Found invalid protocol ProtocolId:0x%x for Object named %s"), ReplicationProtocol->ProtocolIdentifier, *Instance->GetName());
+			UE_LOG_OBJECTREPLICATIONBRIDGE(Error, TEXT("StartReplicatingNetObject Found invalid protocol ProtocolId:0x%x for Object named %s"), ReplicationProtocol->ProtocolIdentifier, *Instance->GetName());
 			return FNetRefHandle::GetInvalid();
 		}
 	}
 #endif
 		
-	if (ReplicationProtocol)
+	if (!ReplicationProtocol)
 	{
-		IRIS_PROFILER_PROTOCOL_NAME(ReplicationProtocol->DebugName->Name);		
+		UE_LOG(LogIris, Error, TEXT("StartReplicatingNetObject - Failed to create protocol for %s"), ToCStr(Instance->GetPathName()));
+		return FNetRefHandle::GetInvalid();
+	}
 
-		// Create NetHandle and bind instance
-		FNetHandle NetHandle = FNetHandleManager::GetOrCreateNetHandle(Instance);
-		FNetRefHandle RefHandle = InternalCreateNetObject(AllocatedRefHandle, NetHandle, ReplicationProtocol);
-		if (RefHandle.IsValid())
-		{
-			// Attach the instance and bind the instance protocol to dirty tracking
-			constexpr bool bBindInstanceProtocol = true;
-			InternalAttachInstanceToNetRefHandle(RefHandle, bBindInstanceProtocol, InstanceProtocol.Get(), Instance, NetHandle);
+	IRIS_PROFILER_PROTOCOL_NAME(ReplicationProtocol->DebugName->Name);		
+
+	// Create NetHandle and bind instance
+	FNetHandle NetHandle = FNetHandleManager::GetOrCreateNetHandle(Instance);
+	FNetRefHandle RefHandle = InternalCreateNetObject(AllocatedRefHandle, NetHandle, ReplicationProtocol);
+
+	if (!RefHandle.IsValid())
+	{
+		UE_LOG_OBJECTREPLICATIONBRIDGE(Warning, TEXT("StartReplicatingNetObject Failed to create NetRefHandle with ProtocolId:0x%x for Object named %s"), (ReplicationProtocol != nullptr ? ReplicationProtocol->ProtocolIdentifier : FReplicationProtocolIdentifier(0)), *Instance->GetName());
+		return FNetRefHandle::GetInvalid();
+	}
+
+	UE_LOG_OBJECTREPLICATIONBRIDGE(Verbose, TEXT("StartReplicatingNetObject Created %s with ProtocolId:0x%x for Object named %s"), *RefHandle.ToString(), ReplicationProtocol->ProtocolIdentifier, *Instance->GetName());
+
+	// Attach the instance and bind the instance protocol to dirty tracking
+	constexpr bool bBindInstanceProtocol = true;
+	InternalAttachInstanceToNetRefHandle(RefHandle, bBindInstanceProtocol, InstanceProtocol.Get(), Instance, NetHandle);
 #if WITH_PUSH_MODEL
-			SetNetPushIdOnInstance(InstanceProtocol.Get(), NetHandle);
+	SetNetPushIdOnInstance(InstanceProtocol.Get(), NetHandle);
 #endif
 
-			const FInternalNetRefIndex InternalReplicationIndex = NetRefHandleManager->GetInternalIndex(RefHandle);
+	const FInternalNetRefIndex InternalReplicationIndex = NetRefHandleManager->GetInternalIndex(RefHandle);
 
-			// Initialize conditionals
-			FReplicationSystemInternal* ReplicationSystemInternal = GetReplicationSystem()->GetReplicationSystemInternal();
-			ReplicationSystemInternal->GetConditionals().InitPropertyCustomConditions(InternalReplicationIndex);
+	// Initialize conditionals
+	FReplicationSystemInternal* ReplicationSystemInternal = GetReplicationSystem()->GetReplicationSystemInternal();
+	ReplicationSystemInternal->GetConditionals().InitPropertyCustomConditions(InternalReplicationIndex);
 
-			// Keep track of handles with object references for garbage collection's sake.
-			ObjectsWithObjectReferences.SetBitValue(InternalReplicationIndex, EnumHasAnyFlags(InstanceProtocol->InstanceTraits, EReplicationInstanceProtocolTraits::HasObjectReference));
+	// Keep track of handles with object references for garbage collection's sake.
+	ObjectsWithObjectReferences.SetBitValue(InternalReplicationIndex, EnumHasAnyFlags(InstanceProtocol->InstanceTraits, EReplicationInstanceProtocolTraits::HasObjectReference));
 
-			// Set poll frame period
-			float PollFrequency = Params.PollFrequency;
-			FindOrCachePollFrequency(Instance->GetClass(), PollFrequency);
-			
-			uint8 PollFramePeriod = ConvertPollFrequencyIntoFrames(PollFrequency);
-			PollFrequencyLimiter->SetPollFramePeriod(InternalReplicationIndex, PollFramePeriod);
+	if (ShouldClassBeDeltaCompressed(Instance->GetClass()))
+	{
+		ReplicationSystem->SetDeltaCompressionStatus(RefHandle, ENetObjectDeltaCompressionStatus::Allow);
+	}
 
-			UE_LOG_OBJECTREPLICATIONBRIDGE(Verbose, TEXT("BeginReplication Created %s with ProtocolId:0x%x for Object named %s"), *RefHandle.ToString(), ReplicationProtocol->ProtocolIdentifier, *Instance->GetName());
+	// Release instance protocol from the UniquePtr as it is now successfully bound to the handle
+	(void)InstanceProtocol.Release();
 
-			{
-				FWorldLocations& WorldLocations = ReplicationSystem->GetReplicationSystemInternal()->GetWorldLocations();
+	return RefHandle;
+}
 
-				if (Params.bNeedsWorldLocationUpdate)
-				{
-					WorldLocations.InitObjectInfoCache(InternalReplicationIndex);
-				}
-			}
+UE::Net::FNetRefHandle UObjectReplicationBridge::StartReplicatingRootObject(UObject* Instance, const FRootObjectReplicationParams& Params)
+{
+	LLM_SCOPE_BYTAG(IrisState);
 
-			// Set prioritizer
-			const bool bRequireForceEnabled = Params.StaticPriority > 0.0f;
-			const FNetObjectPrioritizerHandle PrioritizerHandle = GetPrioritizer(Instance->GetClass(), bRequireForceEnabled);
-			// Set static priority if valid unless we have a force enabled prioritizer.
-			if (Params.StaticPriority > 0.0f && PrioritizerHandle == InvalidNetObjectPrioritizerHandle)
-			{
-				ReplicationSystem->SetStaticPriority(RefHandle, Params.StaticPriority);
-			}
-			else
-			{
-				if (PrioritizerHandle != InvalidNetObjectPrioritizerHandle)
-				{
-					ReplicationSystem->SetPrioritizer(RefHandle, PrioritizerHandle);
-				}
-				else if (Params.bNeedsWorldLocationUpdate || HasRepTag(ReplicationProtocol, RepTag_WorldLocation))
-				{
-					ReplicationSystem->SetPrioritizer(RefHandle, DefaultSpatialNetObjectPrioritizerHandle);
-				}
-			}
+	using namespace UE::Net;
+	using namespace UE::Net::Private;
 
-			AssignDynamicFilter(Instance, Params, RefHandle);
+	if (bBlockStartRootObjectReplication)
+	{
+		ensureMsgf(false, TEXT("StartReplicatingRootObject is not allowed during this operation. %s will not be replicated"), *GetNameSafe(Instance));
+		return FNetRefHandle::GetInvalid();
+	}
+	
+	EReplicationFragmentTraits Traits = EReplicationFragmentTraits::CanReplicate;
+	Traits |= Params.bNeedsPreUpdate ? EReplicationFragmentTraits::NeedsPreSendUpdate : EReplicationFragmentTraits::None;
+	Traits |= Params.bNeedsWorldLocationUpdate ? EReplicationFragmentTraits::NeedsWorldLocationUpdate : EReplicationFragmentTraits::None;
 
-			if (ShouldClassBeDeltaCompressed(Instance->GetClass()))
-			{
-				ReplicationSystem->SetDeltaCompressionStatus(RefHandle, ENetObjectDeltaCompressionStatus::Allow);
-			}
+	// Register the object
+	FNetRefHandle RefHandle = StartReplicatingNetObject(Instance, Traits);
 
-			// Spatially filtered non-dormant objects requires frequent world location updates. Expecting a better solution that instead of us polling will inform us when locations change, UE-193004.
-			if (Params.bNeedsWorldLocationUpdate && !Params.bIsDormant)
-			{
-				OptionallySetObjectRequiresFrequentWorldLocationUpdate(RefHandle, true);
-			}
+	if (!RefHandle.IsValid())
+	{
+		return FNetRefHandle::GetInvalid();;
+	}
+	
+	// Now configure the RootObject
+	const FInternalNetRefIndex ObjectIndex = NetRefHandleManager->GetInternalIndex(RefHandle);
 
-			// Release instance protocol from the UniquePtr as it is now successfully bound to the handle
-			(void)InstanceProtocol.Release();
+	// Set poll frame period
+	float PollFrequency = Params.PollFrequency;
+	FindOrCachePollFrequency(Instance->GetClass(), PollFrequency);
+			
+	uint8 PollFramePeriod = ConvertPollFrequencyIntoFrames(PollFrequency);
+	PollFrequencyLimiter->SetPollFramePeriod(ObjectIndex, PollFramePeriod);
+
+	if (Params.bNeedsWorldLocationUpdate)
+	{
+		FWorldLocations& WorldLocations = ReplicationSystem->GetReplicationSystemInternal()->GetWorldLocations();
+		WorldLocations.InitObjectInfoCache(ObjectIndex);
+	}
 
-			return RefHandle;
+	// Set prioritizer
+	const bool bRequireForceEnabled = Params.StaticPriority > 0.0f;
+	const FNetObjectPrioritizerHandle PrioritizerHandle = GetPrioritizer(Instance->GetClass(), bRequireForceEnabled);
+
+	// Set static priority if valid unless we have a force enabled prioritizer.
+	if (Params.StaticPriority > 0.0f && PrioritizerHandle == InvalidNetObjectPrioritizerHandle)
+	{
+		ReplicationSystem->SetStaticPriority(RefHandle, Params.StaticPriority);
+	}
+	else
+	{
+		if (PrioritizerHandle != InvalidNetObjectPrioritizerHandle)
+		{
+			ReplicationSystem->SetPrioritizer(RefHandle, PrioritizerHandle);
+		}
+		else if (Params.bNeedsWorldLocationUpdate || HasRepTag(ReplicationSystem->GetReplicationProtocol(RefHandle), RepTag_WorldLocation))
+		{
+			ReplicationSystem->SetPrioritizer(RefHandle, DefaultSpatialNetObjectPrioritizerHandle);
 		}
-		UE_LOG_OBJECTREPLICATIONBRIDGE(Warning, TEXT("BeginReplication Failed to create NetRefHandle with ProtocolId:0x%x for Object named %s"), (ReplicationProtocol != nullptr ? ReplicationProtocol->ProtocolIdentifier : FReplicationProtocolIdentifier(0)), *Instance->GetName());
 	}
 
-	// If we get here, it means that we failed to assign an internal handle for the object. We've probably run out of handles which currently is a fatal error.
-	UE_LOG(LogIris, Error, TEXT("UObjectReplicationBridge::BeginReplication - Failed to create NetRefHandle for object %s"), ToCStr(Instance->GetPathName()));
+	// Set dynamic filter
+	AssignDynamicFilter(Instance, Params, RefHandle);
+
+	// Spatially filtered non-dormant objects requires frequent world location updates. Expecting a better solution that instead of us polling will inform us when locations change, UE-193004.
+	if (Params.bNeedsWorldLocationUpdate && !Params.bIsDormant)
+	{
+		OptionallySetObjectRequiresFrequentWorldLocationUpdate(RefHandle, true);
+	}
 
-	return FNetRefHandle::GetInvalid();
+	return RefHandle;
 }
 
-void UObjectReplicationBridge::AssignDynamicFilter(UObject* Instance, const FCreateNetRefHandleParams& Params, FNetRefHandle RefHandle)
+void UObjectReplicationBridge::AssignDynamicFilter(UObject* Instance, const FRootObjectReplicationParams& Params, FNetRefHandle RefHandle)
 {
 	using namespace UE::Net;
 
@@ -451,13 +469,13 @@ void UObjectReplicationBridge::AssignDynamicFilter(UObject* Instance, const FCre
 
 	if (FilterHandle != InvalidNetObjectFilterHandle)
 	{
-		UE_LOG_OBJECTREPLICATIONBRIDGE(Verbose, TEXT("BeginReplication Filter: %s will be used for %s. (FilterProfile: %s)"), *(ReplicationSystem->GetFilterName(FilterHandle).ToString()), *NetRefHandleManager->PrintObjectFromNetRefHandle(RefHandle), *FilterConfigProfile.ToString());
+		UE_LOG_OBJECTREPLICATIONBRIDGE(Verbose, TEXT("StartReplicatingRootObject Filter: %s will be used for %s. (FilterProfile: %s)"), *(ReplicationSystem->GetFilterName(FilterHandle).ToString()), *NetRefHandleManager->PrintObjectFromNetRefHandle(RefHandle), *FilterConfigProfile.ToString());
 		ReplicationSystem->SetFilter(RefHandle, FilterHandle, FilterConfigProfile);
 	}
 	
 }
 
-UE::Net::FNetRefHandle UObjectReplicationBridge::BeginReplication(FNetRefHandle OwnerRefHandle, UObject* Instance, FNetRefHandle InsertRelativeToSubObjectRefHandle, const FCreateNetRefHandleParams& Params, ESubObjectInsertionOrder InsertionOrder)
+UE::Net::FNetRefHandle UObjectReplicationBridge::StartReplicatingSubObject(FNetRefHandle OwnerRefHandle, UObject* Instance, FNetRefHandle InsertRelativeToSubObjectRefHandle, ESubObjectInsertionOrder InsertionOrder)
 {
 	LLM_SCOPE_BYTAG(IrisState);
 
@@ -480,34 +498,27 @@ UE::Net::FNetRefHandle UObjectReplicationBridge::BeginReplication(FNetRefHandle
 		check(OwnerRefHandle == LocalNetRefHandleManager.GetRootObjectOfSubObject(SubObjectRefHandle));
 		return SubObjectRefHandle;
 	}
-	else
-	{
-		// We support replicating new subobjects even during internal operations.
-		TGuardValue<bool> AllowBeginReplication(bBlockBeginReplication, false);
+	
+	const EReplicationFragmentTraits Traits = EReplicationFragmentTraits::CanReplicate;
+	SubObjectRefHandle = StartReplicatingNetObject(Instance, Traits);
 
-		FCreateNetRefHandleParams SubObjectCreateParams = Params;
-		// The filtering system ignores subobjects so let's not waste cycles figuring out which filter to use.
-		SubObjectCreateParams.bUseClassConfigDynamicFilter = false;
-		SubObjectRefHandle = BeginReplication(Instance, SubObjectCreateParams);
+	if (!SubObjectRefHandle.IsValid())
+	{
+		return FNetRefHandle::GetInvalid();
 	}
 
-	if (SubObjectRefHandle.IsValid())
-	{
-		// Add subobject
-		InternalAddSubObject(OwnerRefHandle, SubObjectRefHandle, InsertRelativeToSubObjectRefHandle, InsertionOrder);
+	// Add subobject
+	InternalAddSubObject(OwnerRefHandle, SubObjectRefHandle, InsertRelativeToSubObjectRefHandle, InsertionOrder);
 
-		UE_LOG_OBJECTREPLICATIONBRIDGE(Verbose, TEXT("BeginReplication Added %s RelativeToSubObjectHandle %s"), *PrintObjectFromNetRefHandle(SubObjectRefHandle), *PrintObjectFromNetRefHandle(InsertRelativeToSubObjectRefHandle));
+	UE_LOG_OBJECTREPLICATIONBRIDGE(Verbose, TEXT("StartReplicatingSubObject Added %s RelativeToSubObjectHandle %s"), *PrintObjectFromNetRefHandle(SubObjectRefHandle), *PrintObjectFromNetRefHandle(InsertRelativeToSubObjectRefHandle));
 
-		// SubObjects should always poll with owner
-		SetPollWithObject(OwnerRefHandle, SubObjectRefHandle);
+	// SubObjects should always poll with owner
+	SetPollWithObject(OwnerRefHandle, SubObjectRefHandle);
 
-		// Copy pending dormancy from owner
-		SetSubObjectDormancyStatus(SubObjectRefHandle, OwnerRefHandle);
+	// Copy pending dormancy from owner
+	SetSubObjectDormancyStatus(SubObjectRefHandle, OwnerRefHandle);
 	
-		return SubObjectRefHandle;
-	}
-
-	return FNetRefHandle::GetInvalid();
+	return SubObjectRefHandle;
 }
 
 void UObjectReplicationBridge::SetSubObjectDormancyStatus(FNetRefHandle SubObjectRefHandle, FNetRefHandle OwnerRefHandle)
@@ -640,13 +651,13 @@ bool UObjectReplicationBridge::CacheNetRefHandleCreationInfo(FNetRefHandle Handl
 	return false;
 }
 
-void UObjectReplicationBridge::EndReplication(UObject* Instance, EEndReplicationFlags EndReplicationFlags, FEndReplicationParameters* Parameters)
+void UObjectReplicationBridge::StopReplicatingNetObject(UObject* Instance, EEndReplicationFlags EndReplicationFlags)
 {
 	const FNetRefHandle RefHandle = GetReplicatedRefHandle(Instance, EGetRefHandleFlags::EvenIfGarbage);
 	if (RefHandle.IsValid())
 	{
-		ensureMsgf(IsValid(Instance), TEXT("Calling EndReplication for Invalid Object for %s %s."), *GetNameSafe(Instance), *RefHandle.ToString());
-		UReplicationBridge::EndReplication(RefHandle, EndReplicationFlags, Parameters);
+		ensureMsgf(IsValid(Instance), TEXT("Calling EndReplication for Invalid Object: %s."), *PrintObjectFromNetRefHandle(RefHandle));
+		StopReplicatingNetRefHandle(RefHandle, EndReplicationFlags);
 	}
 }
 
@@ -856,7 +867,7 @@ void UObjectReplicationBridge::PreSendUpdateSingleHandle(FNetRefHandle RefHandle
 void UObjectReplicationBridge::OnStartPreSendUpdate()
 {
 	// During SendUpdate it is not supported to start replication of new root objects.
-	bBlockBeginReplication = true;
+	bBlockStartRootObjectReplication = true;
 }
 
 void UObjectReplicationBridge::PreSendUpdate()
@@ -879,7 +890,7 @@ void UObjectReplicationBridge::PreSendUpdate()
 
 void UObjectReplicationBridge::OnPostSendUpdate()
 {
-	bBlockBeginReplication = false;
+	bBlockStartRootObjectReplication = false;
 }
 
 void UObjectReplicationBridge::PruneStaleObjects()
@@ -924,7 +935,7 @@ void UObjectReplicationBridge::PruneStaleObjects()
 	// EndReplication/detach stale instances
 	for (FNetRefHandle Handle : MakeArrayView(StaleObjects.GetData(), StaleObjects.Num()))
 	{
-		UReplicationBridge::EndReplication(Handle);
+		StopReplicatingNetRefHandle(Handle, EEndReplicationFlags::Destroy);
 	}
 
 	// Mark poll override info as dirty
diff --git a/Engine/Source/Runtime/Experimental/Iris/Core/Private/Iris/ReplicationSystem/ReplicationBridge.cpp b/Engine/Source/Runtime/Experimental/Iris/Core/Private/Iris/ReplicationSystem/ReplicationBridge.cpp
--- a/Engine/Source/Runtime/Experimental/Iris/Core/Private/Iris/ReplicationSystem/ReplicationBridge.cpp
+++ b/Engine/Source/Runtime/Experimental/Iris/Core/Private/Iris/ReplicationSystem/ReplicationBridge.cpp
@@ -582,7 +582,7 @@ void UReplicationBridge::InternalDestroySubObjects(FNetRefHandle OwnerHandle, EE
 	}
 }
 
-void UReplicationBridge::EndReplication(FNetRefHandle Handle, EEndReplicationFlags EndReplicationFlags, FEndReplicationParameters* Parameters)
+void UReplicationBridge::StopReplicatingNetRefHandle(FNetRefHandle Handle, EEndReplicationFlags EndReplicationFlags)
 {
 	using namespace UE::Net::Private;
 
@@ -617,14 +617,6 @@ void UReplicationBridge::EndReplication(FNetRefHandle Handle, EEndReplicationFla
 		}
 		else 
 		{
-			const bool bShouldCreateDestructionInfo = Handle.IsStatic() && EnumHasAnyFlags(EndReplicationFlags, EEndReplicationFlags::Destroy);		
-			if (bShouldCreateDestructionInfo && Parameters)
-			{
-				// Add static destruction info for the Handle, should we just copy group memberships?
-				// need to copy cached reference as well so that we can "keep" it around even after pointer reuse
-				InternalAddDestructionInfo(Handle, *Parameters);
-			}
-
 			// New objects, destroyed during the same frame with posted attachments(RPC):s needs to request a flush to ensure that they get a scope update
 			const UE::Net::Private::FNetBlobManager& NetBlobManager = GetReplicationSystem()->GetReplicationSystemInternal()->GetNetBlobManager();
 			const bool bAllowAutoFlushOfUnProcessedReliableRPCs = bEnableFlushReliableRPCOnDestroy && ObjectData.bNeedsFullCopyAndQuantize;
@@ -960,7 +952,7 @@ void UReplicationBridge::InternalTearOff(FNetRefHandle Handle)
 	ObjectData.bShouldPropagateChangedStates = 0U;
 }
 
-UE::Net::FNetRefHandle UReplicationBridge::InternalAddDestructionInfo(FNetRefHandle Handle, const FEndReplicationParameters& Parameters)
+UE::Net::FNetRefHandle UReplicationBridge::StoreDestructionInfo(FNetRefHandle Handle, const FEndReplicationParameters& Parameters)
 {
 	using namespace UE::Net;
 	using namespace UE::Net::Private;
@@ -1003,7 +995,7 @@ UE::Net::FNetRefHandle UReplicationBridge::InternalAddDestructionInfo(FNetRefHan
 		GetReplicationSystem()->SetPrioritizer(DestructionInfoHandle, DefaultSpatialNetObjectPrioritizerHandle);
 	}
 
-	UE_LOG_REPLICATIONBRIDGE(Verbose, TEXT("UReplicationBridge::InternalAddDestructionInfo %s (InternalIndex: %u ) for %s GroupIndex: %u"), *DestructionInfoHandle.ToString(), InternalReplicationIndex,  *PrintObjectFromNetRefHandle(Handle), LevelGroupHandle.GetGroupIndex());
+	UE_LOG_REPLICATIONBRIDGE(Verbose, TEXT("UReplicationBridge::StoreDestructionInfo %s (InternalIndex: %u ) for %s GroupIndex: %u"), *DestructionInfoHandle.ToString(), InternalReplicationIndex,  *PrintObjectFromNetRefHandle(Handle), LevelGroupHandle.GetGroupIndex());
 	
 	return DestructionInfoHandle;
 }
diff --git a/Engine/Source/Runtime/Experimental/Iris/Core/Public/Iris/ReplicationSystem/ObjectReplicationBridge.h b/Engine/Source/Runtime/Experimental/Iris/Core/Public/Iris/ReplicationSystem/ObjectReplicationBridge.h
--- a/Engine/Source/Runtime/Experimental/Iris/Core/Public/Iris/ReplicationSystem/ObjectReplicationBridge.h
+++ b/Engine/Source/Runtime/Experimental/Iris/Core/Public/Iris/ReplicationSystem/ObjectReplicationBridge.h
@@ -26,6 +26,8 @@ namespace UE::Net
 
 	class FNetObjectReference;
 
+	enum class EReplicationFragmentTraits : uint32;
+
 	namespace Private
 	{
 		typedef uint32 FInternalNetRefIndex;
@@ -50,10 +52,10 @@ class UObjectReplicationBridge : public UReplicationBridge
 	GENERATED_BODY()
 
 public:
-	using UReplicationBridge::EndReplication;
+
 	using EGetRefHandleFlags = UE::Net::EGetRefHandleFlags;
 
-	struct FCreateNetRefHandleParams
+	struct FRootObjectReplicationParams
 	{
 		/** When true it means the object wants to receive a PreUpdate callback just before it gets polled. */
 		bool bNeedsPreUpdate = false;
@@ -113,19 +115,33 @@ public:
 	/** Get or create NetObjectReference for object identified by path relative to outer. */
 	IRISCORE_API UE::Net::FNetObjectReference GetOrCreateObjectReference(const FString& Path, const UObject* Outer) const;
 
-	/** Begin replicating the Instance and return a valid NetRefHandle for the Instance if successful. */
-	IRISCORE_API FNetRefHandle BeginReplication(UObject* Instance, const FCreateNetRefHandleParams& Params);
+	/** 
+	 * Start replicating a RootObject and return a valid NetRefHandle if successful. 
+	 * 
+	 * @param Instance The instance that needs to be replicated
+	 * @param Params Optional configuration parameters to specify how the root object will be replicated.
+	 * @return The NetRefHandle associated to this object if the operation succeeded.
+	 */
+	IRISCORE_API FNetRefHandle StartReplicatingRootObject(UObject* Instance, const FRootObjectReplicationParams& Params);
 
 	/**
-	 * Begin replicating the Instance as a subobject of the OwnerHandle. If InsertRelativeSubObjectHandle is valid
-	 * the new subobject will be inserted in the subobject replication list next to the specified handle and the wanted insertion order.
-	 * Default behavior is to always add new subobjects at the end of the list. Returns a valid NetRefHandle for the Instance if successful.
+	 * Start replicating a SubObject and return a valid NetRefHandle if successful.
+	 * 
+	 * @param OwnerHandle The owner of the subobject
+	 * @param Instance The instance that needs to be replicated
+	 * @param InsertRelativeToSubObjectHandle When valid the new subobject will be inserted in the subobject replication list next to the specified handle and the wanted insertion order.
+	 * @param InsertionOrder When none will always add the new subobject at the end of the list. 
+	 * @return The NetRefHandle of the subobject if successful.
 	 */
-	IRISCORE_API FNetRefHandle BeginReplication(FNetRefHandle OwnerHandle, UObject* Instance, FNetRefHandle InsertRelativeToSubObjectHandle, const FCreateNetRefHandleParams& Params, ESubObjectInsertionOrder InsertionOrder = UReplicationBridge::ESubObjectInsertionOrder::None);
+	IRISCORE_API FNetRefHandle StartReplicatingSubObject(FNetRefHandle OwnerHandle, UObject* Instance, FNetRefHandle InsertRelativeToSubObjectHandle = FNetRefHandle::GetInvalid(), ESubObjectInsertionOrder InsertionOrder = UReplicationBridge::ESubObjectInsertionOrder::None);
 
-	/** Create handle and start replicating the Instance as a SubObject of the OwnerHandle. */
-	/** Begin replicating the Instance as a subobject of the OwnerHandle and return a valid NetRefHandle for the Instance if successful. */
-	IRISCORE_API FNetRefHandle BeginReplication(FNetRefHandle OwnerHandle, UObject* Instance, const FCreateNetRefHandleParams& Params);
+	/** 
+	 * Stop replicating any type of NetObject. @see UReplicationBridge::StopReplicatingNetObject for more details.
+	 * @param Instance The instance that won't be replicated anymore
+	 * @param EndReplicationFlags Optional settings to modify the function behavior. Defaults to destroy the  instance on remote clients.
+	 * 
+	 * */
+	IRISCORE_API void StopReplicatingNetObject(UObject* Instance, EEndReplicationFlags EndReplicationFlags = EEndReplicationFlags::Destroy);
 
 	/** 
 	 * Set NetCondition for a subobject, the condition is used to determine if the SubObject should replicate or not.
@@ -137,9 +153,6 @@ public:
 	/** Get the handle of the root object of any replicated subobject. */
 	IRISCORE_API FNetRefHandle GetRootObjectOfSubObject(FNetRefHandle SubObjectHandle) const;
 
-	/** Stop replicating the object. */
-	IRISCORE_API void EndReplication(UObject* Instance, EEndReplicationFlags EndReplicationFlags = EEndReplicationFlags::Destroy, FEndReplicationParameters* Parameters = nullptr);
-
 	/** Add static destruction info, this is used when stably named objects are destroyed prior to starting replication */
 	IRISCORE_API void AddStaticDestructionInfo(const FString& ObjectPath, const UObject* Outer, const FEndReplicationParameters& Parameters);
 	
@@ -379,6 +392,9 @@ private:
 	/** Tries to load the classes used in poll period overrides. */
 	void FindClassesInPollPeriodOverrides();
 
+	/** Does all that is necessary to replicate any given UObject. Returns the object's handle if successful  */
+	UE::Net::FNetRefHandle StartReplicatingNetObject(UObject* Instance, UE::Net::EReplicationFragmentTraits Traits);
+
 protected:
 	/** Retrieves the dynamic filter to set for the given class. Will return an invalid handle if no dynamic filter should be set. */
 	IRISCORE_API UE::Net::FNetObjectFilterHandle GetDynamicFilter(const UClass* Class, bool bRequireForceEnabled, FName& OutFilterProfile);
@@ -388,7 +404,7 @@ private:
 	UE::Net::FNetObjectPrioritizerHandle GetPrioritizer(const UClass* Class, bool bRequireForceEnabled);
 
 	/** Assign the proper dynamic filter to a new object */
-	void AssignDynamicFilter(UObject* Instance, const FCreateNetRefHandleParams& Params, FNetRefHandle RefHandle);
+	void AssignDynamicFilter(UObject* Instance, const FRootObjectReplicationParams& Params, FNetRefHandle RefHandle);
 
 	/** Returns true if instances of this class should be delta compressed */
 	bool ShouldClassBeDeltaCompressed(const UClass* Class);
@@ -487,15 +503,9 @@ private:
 	bool bHasPollOverrides = false;
 	bool bHasDirtyClassesInPollPeriodOverrides = false;
 
-	/** Set to true when the system does not allow new objects to begin replication at this moment. Useful when calling into user code and to warn them of illegal operations. */
-	bool bBlockBeginReplication = false;
+	/** Set to true when the system does not allow new root objects to be replicated at this moment. Useful when calling into user code and to warn them of illegal operations. */
+	bool bBlockStartRootObjectReplication = false;
 
 protected:
 	bool bSuppressCreateInstanceFailedEnsure = false;
 };
-
-
-inline UE::Net::FNetRefHandle UObjectReplicationBridge::BeginReplication(UE::Net::FNetRefHandle OwnerHandle, UObject* Instance, const FCreateNetRefHandleParams& Params)
-{
-	return BeginReplication(OwnerHandle, Instance, FNetRefHandle::GetInvalid(), Params, ESubObjectInsertionOrder::None);
-}
diff --git a/Engine/Source/Runtime/Experimental/Iris/Core/Public/Iris/ReplicationSystem/ReplicationBridge.h b/Engine/Source/Runtime/Experimental/Iris/Core/Public/Iris/ReplicationSystem/ReplicationBridge.h
--- a/Engine/Source/Runtime/Experimental/Iris/Core/Public/Iris/ReplicationSystem/ReplicationBridge.h
+++ b/Engine/Source/Runtime/Experimental/Iris/Core/Public/Iris/ReplicationSystem/ReplicationBridge.h
@@ -151,13 +151,16 @@ public:
 	};
 
 	/**
-	 * Stop replicating the NetObject associated with the handle and mark the handle to be destroyed.
-	 * If EEndReplication::TearOff is set the remote instance will be Torn-off rather than being destroyed on the receiving end, after the call, any state changes will not be replicated
-     * If EEndReplication::Flush is set all pending states will be delivered before the remote instance is destroyed, final state will be immediately copied so it is safe to remove the object after this call
-	 * If EEndReplication::Destroy is set the remote instance will be destroyed, if this is set for a static instance and the EndReplicationParameters are set a permanent destruction info will be added
-	 * Dynamic instances are always destroyed unless the TearOff flag is set.
-	 */
-	IRISCORE_API void EndReplication(FNetRefHandle Handle, EEndReplicationFlags EndReplicationFlags = EEndReplicationFlags::Destroy, FEndReplicationParameters* Parameters = nullptr);
+	* Stop replicating the NetObject associated with the handle and mark the handle to be destroyed.
+	* If EEndReplication::TearOff is set the remote instance will be Torn-off rather than being destroyed on the receiving end, after the call, any state changes will not be replicated
+	* If EEndReplication::Flush is set all pending states will be delivered before the remote instance is destroyed, final state will be immediately copied so it is safe to remove the object after this call
+	* If EEndReplication::Destroy is set the remote instance will be destroyed, if this is set for a static instance and the EndReplicationParameters are set a permanent destruction info will be added
+	* Dynamic instances are always destroyed unless the TearOff flag is set.
+	*/
+	IRISCORE_API void StopReplicatingNetRefHandle(FNetRefHandle Handle, EEndReplicationFlags EndReplicationFlags);
+
+	/** Store destruction info for the referenced object. */
+	IRISCORE_API FNetRefHandle StoreDestructionInfo(FNetRefHandle Handle, const FEndReplicationParameters& Parameters);
 
 	/** Returns true if the handle is replicated. */
 	IRISCORE_API bool IsReplicatedHandle(FNetRefHandle Handle) const;
@@ -256,9 +259,6 @@ protected:
 	/** Add SubObjectHandle as SubObject to OwnerHandle. */
 	IRISCORE_API void InternalAddSubObject(FNetRefHandle OwnerHandle, FNetRefHandle SubObjectHandle, FNetRefHandle InsertRelativeToSubObjectHandle, ESubObjectInsertionOrder InsertionOrder);
 
-	/** Add destruction info for the referenced object. */
-	FNetRefHandle InternalAddDestructionInfo(FNetRefHandle Handle, const FEndReplicationParameters& Parameters);
-
 	inline UE::Net::Private::FReplicationProtocolManager* GetReplicationProtocolManager() const { return ReplicationProtocolManager; }
 	inline UReplicationSystem* GetReplicationSystem() const { return ReplicationSystem; }
 	inline UE::Net::Private::FReplicationStateDescriptorRegistry* GetReplicationStateDescriptorRegistry() const { return ReplicationStateDescriptorRegistry; }
