Rename/move file from dataflow plugin into runtime/experimental + bug fixes (panel timeline, proxy creation...) + flesh solver preview options + async skelmesh animation

#virtualized

diff --git a/Engine/Plugins/ChaosClothAsset/Source/ChaosClothAssetEngine/Private/ChaosClothAsset/ClothComponent.cpp b/Engine/Plugins/ChaosClothAsset/Source/ChaosClothAssetEngine/Private/ChaosClothAsset/ClothComponent.cpp
--- a/Engine/Plugins/ChaosClothAsset/Source/ChaosClothAssetEngine/Private/ChaosClothAsset/ClothComponent.cpp
+++ b/Engine/Plugins/ChaosClothAsset/Source/ChaosClothAssetEngine/Private/ChaosClothAsset/ClothComponent.cpp
@@ -187,6 +187,24 @@ PRAGMA_ENABLE_DEPRECATION_WARNINGS
 
 	Super::PostEditChangeProperty(PropertyChangedEvent);
 }
+
+bool UChaosClothComponent::CanEditChange(const FProperty* InProperty) const
+{
+	if (!Super::CanEditChange(InProperty))
+	{
+		return false;
+	}
+
+	const FName& Name = InProperty->GetFName();
+
+	if (Name == GET_MEMBER_NAME_CHECKED(ThisClass, SimulationAsset))
+	{
+		static const auto CVarEnableSimulationDataflow = IConsoleManager::Get().FindConsoleVariable(TEXT("p.Dataflow.EnableSimulation"));
+		return CVarEnableSimulationDataflow->GetBool();
+	}
+
+	return true;
+}
 #endif // WITH_EDITOR
 
 void UChaosClothComponent::OnRegister()
@@ -210,7 +228,7 @@ void UChaosClothComponent::OnRegister()
 	UpdateVisibility();
 
 	// Register the dataflow simulation interface
-	Dataflow::RegisterSimulationInterface(this);
+	UE::Dataflow::RegisterSimulationInterface(this);
 }
 
 void UChaosClothComponent::OnUnregister()
@@ -226,7 +244,7 @@ void UChaosClothComponent::OnUnregister()
 	PropertyCollections.Empty();
 
 	// Unregister the dataflow simulation interface
-	Dataflow::UnregisterSimulationInterface(this);
+	UE::Dataflow::UnregisterSimulationInterface(this);
 }
 
 bool UChaosClothComponent::IsComponentTickEnabled() const
@@ -640,6 +658,14 @@ void UChaosClothComponent::WriteToSimulation(const float DeltaTime)
 	}
 }
 
+void UChaosClothComponent::PreProcessSimulation(const float DeltaTime)
+{
+	if (ClothSimulationProxy.IsValid() && ClothSimulationProxy->HasCacheData())
+	{
+		WriteToSimulation(DeltaTime);
+	}
+}
+
 void UChaosClothComponent::ReadFromSimulation(const float DeltaTime)
 {
 	if (ClothSimulationProxy.IsValid())
diff --git a/Engine/Plugins/ChaosClothAsset/Source/ChaosClothAssetEngine/Private/ChaosClothAsset/ClothSimulationProxy.cpp b/Engine/Plugins/ChaosClothAsset/Source/ChaosClothAssetEngine/Private/ChaosClothAsset/ClothSimulationProxy.cpp
--- a/Engine/Plugins/ChaosClothAsset/Source/ChaosClothAssetEngine/Private/ChaosClothAsset/ClothSimulationProxy.cpp
+++ b/Engine/Plugins/ChaosClothAsset/Source/ChaosClothAssetEngine/Private/ChaosClothAsset/ClothSimulationProxy.cpp
@@ -176,6 +176,11 @@ PRAGMA_ENABLE_DEPRECATION_WARNINGS
 	{
 		CompleteParallelSimulation_GameThread();
 	}
+	
+	void FClothSimulationProxy::PreSolveProxy(const float DeltaTime)
+	{
+		SetupSimulationData(DeltaTime);
+	}
 
 	bool FClothSimulationProxy::SetupSimulationData(float DeltaTime)
 	{
diff --git a/Engine/Plugins/ChaosClothAsset/Source/ChaosClothAssetEngine/Public/ChaosClothAsset/ClothComponent.h b/Engine/Plugins/ChaosClothAsset/Source/ChaosClothAssetEngine/Public/ChaosClothAsset/ClothComponent.h
--- a/Engine/Plugins/ChaosClothAsset/Source/ChaosClothAssetEngine/Public/ChaosClothAsset/ClothComponent.h
+++ b/Engine/Plugins/ChaosClothAsset/Source/ChaosClothAssetEngine/Public/ChaosClothAsset/ClothComponent.h
@@ -3,7 +3,7 @@
 #pragma once
 
 #include "Components/SkinnedMeshComponent.h"
-#include "Interfaces/DataflowPhysicsSolver.h"
+#include "Dataflow/Interfaces/DataflowPhysicsSolver.h"
 #include "ClothComponent.generated.h"
 
 class UChaosClothAsset;
@@ -153,6 +153,7 @@ protected:
 	virtual void PostLoad() override;
 #if WITH_EDITOR
 	virtual void PostEditChangeProperty(FPropertyChangedEvent& PropertyChangedEvent) override;
+	virtual bool CanEditChange(const FProperty* InProperty) const override;
 #endif // WITH_EDITOR
 	//~ End UObject Interface
 
@@ -190,6 +191,7 @@ protected:
 	virtual void ResetSimulationProxy() override;
 	virtual void WriteToSimulation(const float DeltaTime) override;
 	virtual void ReadFromSimulation(const float DeltaTime) override;
+	virtual void PreProcessSimulation(const float DeltaTime) override;
 	// End IDataflowPhysicsSolverInterface overrides
 
 	/** Override this function for setting up custom simulation proxies when the component is registered. */
@@ -216,7 +218,7 @@ private:
 #endif
 
 	/* Solver dataflow asset used to advance in time */
-	UPROPERTY(EditAnywhere, Category = ClothComponent)
+	UPROPERTY(EditAnywhere, Category = ClothComponent, meta=(EditConditionHides), AdvancedDisplay)
 	FDataflowSimulationAsset SimulationAsset;
 	
 	/** If enabled, and the parent is another Skinned Mesh Component (e.g. another Cloth Component, Poseable Mesh Component, Skeletal Mesh Component, ...etc.), use its pose. */
diff --git a/Engine/Plugins/ChaosClothAsset/Source/ChaosClothAssetEngine/Public/ChaosClothAsset/ClothSimulationProxy.h b/Engine/Plugins/ChaosClothAsset/Source/ChaosClothAssetEngine/Public/ChaosClothAsset/ClothSimulationProxy.h
--- a/Engine/Plugins/ChaosClothAsset/Source/ChaosClothAssetEngine/Public/ChaosClothAsset/ClothSimulationProxy.h
+++ b/Engine/Plugins/ChaosClothAsset/Source/ChaosClothAssetEngine/Public/ChaosClothAsset/ClothSimulationProxy.h
@@ -7,7 +7,7 @@
 #include "Math/Transform.h"
 #include "Templates/UniquePtr.h"
 #include "ClothingSystemRuntimeTypes.h"
-#include "Interfaces/DataflowPhysicsSolver.h"
+#include "Dataflow/Interfaces/DataflowPhysicsSolver.h"
 
 namespace Chaos
 {
@@ -83,6 +83,7 @@ PRAGMA_ENABLE_DEPRECATION_WARNINGS
 		float GetLinearSolveError() const { return LastLinearSolveError; }
 		float GetSimulationTime() const { return SimulationTime; }
 		bool IsTeleported() const { return bIsTeleported; }
+		bool HasCacheData() const { return CacheData.IsValid(); }
 
 	protected:
 		void Tick();
@@ -97,6 +98,7 @@ PRAGMA_ENABLE_DEPRECATION_WARNINGS
 
 		// Begin FDataflowPhysicsSolverProxy overrides
 		virtual void AdvanceSolverDatas(const float DeltaTime) override {Tick();}
+		virtual void PreSolveProxy(const float DeltaTime);
 		// End FDataflowPhysicsSolverProxy overrides
 
 		// Internal physics thread object
diff --git a/Engine/Plugins/Experimental/ChaosFlesh/Source/ChaosFleshEditor/Private/ChaosFlesh/Asset/AssetDefinition_FleshAsset.cpp b/Engine/Plugins/Experimental/ChaosFlesh/Source/ChaosFleshEditor/Private/ChaosFlesh/Asset/AssetDefinition_FleshAsset.cpp
--- a/Engine/Plugins/Experimental/ChaosFlesh/Source/ChaosFleshEditor/Private/ChaosFlesh/Asset/AssetDefinition_FleshAsset.cpp
+++ b/Engine/Plugins/Experimental/ChaosFlesh/Source/ChaosFleshEditor/Private/ChaosFlesh/Asset/AssetDefinition_FleshAsset.cpp
@@ -74,7 +74,7 @@ EAssetCommandResult UAssetDefinition_FleshAsset::OpenAssets(const FAssetOpenArgs
 				UAssetEditorSubsystem* const AssetEditorSubsystem = GEditor->GetEditorSubsystem<UAssetEditorSubsystem>();
 				UDataflowEditor* const AssetEditor = NewObject<UDataflowEditor>(AssetEditorSubsystem, NAME_None, RF_Transient);
 
-				UClass* ActorClass = StaticLoadClass(AActor::StaticClass(), nullptr,
+				const TSubclassOf<AActor> ActorClass = StaticLoadClass(AActor::StaticClass(), nullptr,
 					TEXT("/ChaosFlesh/BP_FleshPreview.BP_FleshPreview_C"), nullptr, LOAD_None, nullptr);
 				AssetEditor->Initialize({ FleshAsset }, ActorClass);
 				
diff --git a/Engine/Plugins/Experimental/ChaosFlesh/Source/ChaosFleshEngine/Private/ChaosFlesh/ChaosDeformableSolverComponent.cpp b/Engine/Plugins/Experimental/ChaosFlesh/Source/ChaosFleshEngine/Private/ChaosFlesh/ChaosDeformableSolverComponent.cpp
--- a/Engine/Plugins/Experimental/ChaosFlesh/Source/ChaosFleshEngine/Private/ChaosFlesh/ChaosDeformableSolverComponent.cpp
+++ b/Engine/Plugins/Experimental/ChaosFlesh/Source/ChaosFleshEngine/Private/ChaosFlesh/ChaosDeformableSolverComponent.cpp
@@ -207,6 +207,59 @@ void UDeformableSolverComponent::EndPlay(const EEndPlayReason::Type EndPlayReaso
 	}
 }
 
+void UDeformableSolverComponent::ResetSimulationProperties(const FSolverTimingGroup& TimingGroup, const FSolverEvolutionGroup& EvolutionGroup,
+		FSolverCollisionsGroup CollisionsGroup, FSolverConstraintsGroup ConstraintsGroup, FSolverForcesGroup ForcesGroup,
+		FSolverDebuggingGroup DebuggingGroup, FSolverMuscleActivationGroup MuscleActivationGroup)
+{
+	SolverTiming = TimingGroup;
+	SolverEvolution = EvolutionGroup;
+	SolverCollisions = CollisionsGroup;
+	SolverConstraints = ConstraintsGroup;
+	SolverForces = ForcesGroup;
+	SolverDebugging = DebuggingGroup;
+	SolverMuscleActivation = MuscleActivationGroup;
+
+	if(FleshSolverProxy.IsValid())
+	{
+		FDeformableSolver::FPhysicsThreadAccess PhysicsThreadSolver = PhysicsThreadAccess();
+		PhysicsThreadSolver.Reset(Chaos::Softs::FDeformableSolverProperties(
+			SolverTiming.NumSubSteps
+			, SolverTiming.NumSolverIterations
+			, SolverTiming.FixTimeStep
+			, SolverTiming.TimeStepSize
+			, SolverDebugging.CacheToFile
+			, SolverConstraints.bEnableKinematics
+			, SolverCollisions.bUseFloor
+			, false /*SolverCollisions.SolverGridBasedCollisions.bUseGridBasedConstraints*/
+			, 25. /*SolverCollisions.SolverGridBasedCollisions.GridDx*/
+			, SolverEvolution.SolverQuasistatics.bDoQuasistatics
+			, SolverForces.YoungModulus
+			, SolverConstraints.CorotatedConstraints.bDoBlended
+			, SolverConstraints.CorotatedConstraints.BlendedZeta
+			, SolverForces.Damping
+			, SolverForces.bEnableGravity
+			, SolverConstraints.CorotatedConstraints.bEnableCorotatedConstraint
+			, SolverConstraints.bEnablePositionTargets
+			, SolverConstraints.GaussSeidelConstraints.bUseGaussSeidelConstraints
+			, SolverConstraints.GaussSeidelConstraints.bUseSOR
+			, SolverConstraints.GaussSeidelConstraints.OmegaSOR
+			, SolverConstraints.GaussSeidelConstraints.bUseGSNeohookean
+			, SolverConstraints.GaussSeidelConstraints.SpringCollision.bDoSpringCollision
+			, SolverConstraints.GaussSeidelConstraints.SpringCollision.InComponentSpringCollision.bDoInComponentSpringCollision
+			, SolverConstraints.GaussSeidelConstraints.SpringCollision.InComponentSpringCollision.NRingExcluded
+			, SolverConstraints.GaussSeidelConstraints.SpringCollision.CollisionSearchRadius
+			, SolverConstraints.GaussSeidelConstraints.SpringCollision.SpringCollisionStiffness
+			, SolverConstraints.GaussSeidelConstraints.SpringCollision.bAllowSliding
+			, SolverConstraints.GaussSeidelConstraints.SphereRepulsion.bDoSphereRepulsion
+			, SolverConstraints.GaussSeidelConstraints.SphereRepulsion.SphereRepulsionRadius
+			, SolverConstraints.GaussSeidelConstraints.SphereRepulsion.SphereRepulsionStiffness
+			, SolverMuscleActivation.bDoMuscleActivation
+			, SolverConstraints.GaussSeidelConstraints.SpringCollision.bCollideWithFullmesh
+			, SolverConstraints.GaussSeidelConstraints.bEnableDynamicSprings
+		));
+	}
+}
+
 void UDeformableSolverComponent::BuildSimulationProxy()
 {
 	SCOPE_CYCLE_COUNTER(STAT_ChaosDeformable_UDeformableSolverComponent_Reset);
@@ -386,6 +439,28 @@ void UDeformableSolverComponent::ReadFromSimulation(float DeltaTime)
 	}
 }
 
+#if WITH_EDITOR
+
+bool UDeformableSolverComponent::CanEditChange(const FProperty* InProperty) const
+{
+	if (!Super::CanEditChange(InProperty))
+	{
+		return false;
+	}
+
+	const FName& Name = InProperty->GetFName();
+
+	if (Name == GET_MEMBER_NAME_CHECKED(ThisClass, SimulationAsset))
+	{
+		static const auto CVarEnableSimulationDataflow = IConsoleManager::Get().FindConsoleVariable(TEXT("p.Dataflow.EnableSimulation"));
+		return CVarEnableSimulationDataflow->GetBool();
+	}
+
+	return true;
+}
+
+#endif
+
 
 
 
diff --git a/Engine/Plugins/Experimental/ChaosFlesh/Source/ChaosFleshEngine/Private/ChaosFlesh/FleshAsset.cpp b/Engine/Plugins/Experimental/ChaosFlesh/Source/ChaosFleshEngine/Private/ChaosFlesh/FleshAsset.cpp
--- a/Engine/Plugins/Experimental/ChaosFlesh/Source/ChaosFleshEngine/Private/ChaosFlesh/FleshAsset.cpp
+++ b/Engine/Plugins/Experimental/ChaosFlesh/Source/ChaosFleshEngine/Private/ChaosFlesh/FleshAsset.cpp
@@ -75,7 +75,8 @@ void UFleshAsset::Serialize(FArchive& Ar)
 
 TObjectPtr<UDataflowBaseContent> UFleshAsset::CreateDataflowContent()
 {
-	TObjectPtr<UDataflowSkeletalContent> SkeletalContent = DataflowContextHelpers::CreateNewDataflowContent<UDataflowSkeletalContent>(this);
+	TObjectPtr<UDataflowFleshContent> SkeletalContent = NewObject<UDataflowFleshContent>(this, UDataflowFleshContent::StaticClass());
+	SkeletalContent->SetIsSaved(false);
 
 	SkeletalContent->SetDataflowOwner(this);
 	SkeletalContent->SetTerminalAsset(this);
@@ -87,7 +88,7 @@ TObjectPtr<UDataflowBaseContent> UFleshAsset::CreateDataflowContent()
 
 void UFleshAsset::WriteDataflowContent(const TObjectPtr<UDataflowBaseContent>& DataflowContent) const
 {
-	if(const TObjectPtr<UDataflowSkeletalContent> SkeletalContent = Cast<UDataflowSkeletalContent>(DataflowContent))
+	if(const TObjectPtr<UDataflowFleshContent> SkeletalContent = Cast<UDataflowFleshContent>(DataflowContent))
 	{
 		SkeletalContent->SetDataflowAsset(DataflowAsset);
 		SkeletalContent->SetDataflowTerminal(DataflowTerminal);
@@ -97,16 +98,30 @@ void UFleshAsset::WriteDataflowContent(const TObjectPtr<UDataflowBaseContent>& D
 
 #if WITH_EDITORONLY_DATA
 		SkeletalContent->SetAnimationAsset(PreviewAnimationAsset.Get());
+		SkeletalContent->SolverTiming = PreviewSolverTiming;
+		SkeletalContent->SolverEvolution = PreviewSolverEvolution;
+		SkeletalContent->SolverCollisions = PreviewSolverCollisions;
+		SkeletalContent->SolverConstraints = PreviewSolverConstraints;
+		SkeletalContent->SolverForces = PreviewSolverForces;
+		SkeletalContent->SolverDebugging = PreviewSolverDebugging;
+		SkeletalContent->SolverMuscleActivation = PreviewSolverMuscleActivation;
 #endif
 	}
 }
 
 void UFleshAsset::ReadDataflowContent(const TObjectPtr<UDataflowBaseContent>& DataflowContent)
 {
-	if(const TObjectPtr<UDataflowSkeletalContent> SkeletalContent = Cast<UDataflowSkeletalContent>(DataflowContent))
+	if(const TObjectPtr<UDataflowFleshContent> SkeletalContent = Cast<UDataflowFleshContent>(DataflowContent))
 	{
 #if WITH_EDITORONLY_DATA
 		PreviewAnimationAsset = SkeletalContent->GetAnimationAsset();
+		PreviewSolverTiming = SkeletalContent->SolverTiming;
+		PreviewSolverEvolution = SkeletalContent->SolverEvolution;
+		PreviewSolverCollisions = SkeletalContent->SolverCollisions;
+		PreviewSolverConstraints = SkeletalContent->SolverConstraints;
+		PreviewSolverForces = SkeletalContent->SolverForces;
+		PreviewSolverDebugging = SkeletalContent->SolverDebugging;
+		PreviewSolverMuscleActivation = SkeletalContent->SolverMuscleActivation;
 #endif
 	}
 }
@@ -137,4 +152,32 @@ void UFleshAsset::PostEditChangeProperty(FPropertyChangedEvent& PropertyChangedE
 }
 #endif //if WITH_EDITOR
 
+UDataflowFleshContent::UDataflowFleshContent() : Super()
+{
+	bHideSkeletalMesh = false;
+	bHideAnimationAsset = false;
+}
+
+void UDataflowFleshContent::AddReferencedObjects(UObject* InThis, FReferenceCollector& Collector)
+{
+	Super::AddReferencedObjects(InThis, Collector);
+	UDataflowFleshContent* This = CastChecked<UDataflowFleshContent>(InThis);
+	Super::AddReferencedObjects(InThis, Collector);
+}
+
+void UDataflowFleshContent::SetActorProperties(TObjectPtr<AActor>& PreviewActor) const
+{
+	Super::SetActorProperties(PreviewActor);
+	OverrideStructProperty(PreviewActor, SolverTiming, TEXT("SolverTiming"));
+	OverrideStructProperty(PreviewActor, SolverEvolution, TEXT("SolverEvolution"));
+	OverrideStructProperty(PreviewActor, SolverCollisions, TEXT("SolverCollisions"));
+	OverrideStructProperty(PreviewActor, SolverConstraints, TEXT("SolverConstraints"));
+	OverrideStructProperty(PreviewActor, SolverForces, TEXT("SolverForces"));
+	OverrideStructProperty(PreviewActor, SolverDebugging, TEXT("SolverDebugging"));
+	OverrideStructProperty(PreviewActor, SolverMuscleActivation, TEXT("SolverMuscleActivation"));
+}
+
+
+
+
 
diff --git a/Engine/Plugins/Experimental/ChaosFlesh/Source/ChaosFleshEngine/Public/ChaosFlesh/ChaosDeformableSolverComponent.h b/Engine/Plugins/Experimental/ChaosFlesh/Source/ChaosFleshEngine/Public/ChaosFlesh/ChaosDeformableSolverComponent.h
--- a/Engine/Plugins/Experimental/ChaosFlesh/Source/ChaosFleshEngine/Public/ChaosFlesh/ChaosDeformableSolverComponent.h
+++ b/Engine/Plugins/Experimental/ChaosFlesh/Source/ChaosFleshEngine/Public/ChaosFlesh/ChaosDeformableSolverComponent.h
@@ -6,10 +6,10 @@
 #include "Chaos/Deformable/ChaosDeformableSolverProxy.h"
 #include "Chaos/Deformable/ChaosDeformableSolverTypes.h"
 #include "ChaosFlesh/ChaosDeformableSolverThreading.h"
-#include "ChaosFlesh/ChaosDeformableTypes.h"
+#include "ChaosFlesh/ChaosDeformableSolverGroups.h"
 #include "DeformableInterface.h"
 #include "Components/SceneComponent.h"
-#include "Interfaces/DataflowPhysicsSolver.h"
+#include "Dataflow/Interfaces/DataflowPhysicsSolver.h"
 #include "UObject/ObjectMacros.h"
 
 #include "ChaosDeformableSolverComponent.generated.h"
@@ -26,268 +26,6 @@ struct FConnectedObjectsGroup
 	TArray< TObjectPtr<UDeformablePhysicsComponent> > DeformableComponents;
 };
 
-USTRUCT(BlueprintType)
-struct FSolverTimingGroup
-{
-	GENERATED_USTRUCT_BODY()
-
-	UPROPERTY(EditAnywhere, Category = "SolverTiming")
-		int32 NumSubSteps = 2;
-
-	UPROPERTY(EditAnywhere, Category = "SolverTiming")
-		int32 NumSolverIterations = 5;
-
-	UPROPERTY(EditAnywhere, Category = "SolverTiming")
-		bool FixTimeStep = false;
-
-	UPROPERTY(EditAnywhere, Category = "SolverTiming")
-		float TimeStepSize = 0.05;
-
-	UPROPERTY(EditAnywhere, Category = "SolverTiming")
-		bool bDoThreadedAdvance = true;
-
-	/** ObjectType defines how to initialize the rigid objects state, Kinematic, Sleeping, Dynamic. */
-	UPROPERTY(EditAnywhere, BlueprintReadOnly, Category = "SolverTiming")
-		EDeformableExecutionModel ExecutionModel = EDeformableExecutionModel::Chaos_Deformable_PostPhysics;
-};
-
-
-USTRUCT(BlueprintType)
-struct FSolverDebuggingGroup
-{
-	GENERATED_USTRUCT_BODY()
-
-	UPROPERTY(EditAnywhere, Category = "Physics")
-	bool CacheToFile = false;
-};
-
-
-USTRUCT(BlueprintType)
-struct FSolverQuasistaticsGroup
-{
-	GENERATED_USTRUCT_BODY()
-
-	UPROPERTY(EditAnywhere, Category = "Quasistatics")
-	bool bDoQuasistatics = false;
-};
-
-
-USTRUCT(BlueprintType)
-struct FSolverEvolutionGroup
-{
-	GENERATED_USTRUCT_BODY()
-
-	UPROPERTY(EditAnywhere, Category = "Evolution")
-	FSolverQuasistaticsGroup SolverQuasistatics;
-
-};
-
-
-USTRUCT(BlueprintType)
-struct FSolverGridBasedCollisionsGroup
-{
-	GENERATED_USTRUCT_BODY()
-
-	UPROPERTY(EditAnywhere, Category = "GridBasedCollisions")
-	bool bUseGridBasedConstraints = false;
-
-	UPROPERTY(EditAnywhere, Category = "GridBasedCollisions")
-	float GridDx = 25.;
-};
-
-USTRUCT(BlueprintType)
-struct FInComponentSpringCollisionGroup
-{
-	GENERATED_USTRUCT_BODY()
-	/**
-	* If uses in-component spring self-collision
-	*/
-	UPROPERTY(EditAnywhere, Category = "InComponentSpringCollision")
-	bool bDoInComponentSpringCollision = false;
-	/**
-	* N ring to exclude for in-component spring self-collision
-	*/
-	UPROPERTY(EditAnywhere, Category = "InComponentSpringCollision")
-	int32 NRingExcluded = 1;
-};
-
-USTRUCT(BlueprintType)
-struct FSpringCollisionGroup
-{
-	GENERATED_USTRUCT_BODY()
-	/**
-	* If uses component-component spring collision
-	*/
-	UPROPERTY(EditAnywhere, Category = "SpringCollision")
-	bool bDoSpringCollision = false;
-	/**
-	* In-component spring self collision detection parameters
-	*/
-	UPROPERTY(EditAnywhere, BlueprintReadOnly, Category = "SpringCollision")
-	FInComponentSpringCollisionGroup InComponentSpringCollision;
-	/**
-	* Search radius for point triangle collision pairs
-	*/
-	UPROPERTY(EditAnywhere, Category = "SpringCollision")
-	float CollisionSearchRadius = 0.f;
-	/**
-	* Collision spring stiffness; larger value will stop penetration better
-	*/
-	UPROPERTY(EditAnywhere, Category = "SpringCollision")
-	float SpringCollisionStiffness = 500.f;
-	/**
-	* Anisotropic springs will allow sliding on the triangle
-	*/
-	UPROPERTY(EditAnywhere, Category = "SpringCollision")
-	bool bAllowSliding = true;
-	/**
-	* Do self collision with kinematic triangles as well
-	*/
-	UPROPERTY(EditAnywhere, Category = "SpringCollision")
-	bool bCollideWithFullmesh = true;
-};
-
-USTRUCT(BlueprintType)
-struct FSphereRepulsionGroup
-{
-	GENERATED_USTRUCT_BODY()
-	/**
-	* If uses sphere repulsion for collision
-	*/
-	UPROPERTY(EditAnywhere, Category = "SphereRepulsion")
-	bool bDoSphereRepulsion = false;
-	/**
-	* Search radius for repulsion pairs
-	*/
-	UPROPERTY(EditAnywhere, Category = "SphereRepulsion")
-	float SphereRepulsionRadius = 0.f;
-	/**
-	* Stiffness for sphere repulsion
-	*/
-	UPROPERTY(EditAnywhere, Category = "SphereRepulsion")
-	float SphereRepulsionStiffness = 500.f;
-};
-
-USTRUCT(BlueprintType)
-struct FSolverGaussSeidelConstraintsGroup
-{
-	GENERATED_USTRUCT_BODY()
-
-	/**
-	* Enable the Gauss Seidel solver instead of the existing XPBD.
-	*/
-	UPROPERTY(EditAnywhere, Category = "GaussSeidelConstraints")
-	bool bUseGaussSeidelConstraints = false;
-
-	/**
-	* Enable another model that runs simulation faster.
-	*/
-	UPROPERTY(EditAnywhere, Category = "GaussSeidelConstraints")
-	bool bUseGSNeohookean = false;
-
-	/**
-	* Enable acceleration technique for Gauss Seidel solver to make simulation look better within a limited budget.
-	*/
-	UPROPERTY(EditAnywhere, Category = "GaussSeidelConstraints")
-	bool bUseSOR = true;
-
-	/**
-	* Acceleration related parameter. Tune it down if simulation becomes unstable. 
-	*/
-	UPROPERTY(EditAnywhere, Category = "GaussSeidelConstraints")
-	float OmegaSOR = 1.6f;
-
-	/**
-	* Enable dynamic springs controlled by constraint manager. 
-	*/
-	UPROPERTY(EditAnywhere, Category = "GaussSeidelConstraints")
-	bool bEnableDynamicSprings = true;
-	
-	/**
-	* Component-component collision detection radius and stiffness
-	*/
-	UPROPERTY(EditAnywhere, BlueprintReadOnly, Category = "GaussSeidelConstraints")
-	FSpringCollisionGroup SpringCollision;
-
-	/**
-	* Sphere repulsion parameters
-	*/
-	UPROPERTY(EditAnywhere, BlueprintReadOnly, Category = "GaussSeidelConstraints")
-	FSphereRepulsionGroup SphereRepulsion;
-};
-
-USTRUCT(BlueprintType)
-struct FSolverCollisionsGroup
-{
-	GENERATED_USTRUCT_BODY()
-
-	UPROPERTY(EditAnywhere, Category = "Physics")
-	bool bUseFloor = true;
-
-	//UPROPERTY(EditAnywhere, Category = "Collisions")
-	//FSolverGridBasedCollisionsGroup SolverGridBasedCollisions;
-};
-
-USTRUCT(BlueprintType)
-struct FSolverCorotatedConstraintsGroup
-{
-	GENERATED_USTRUCT_BODY()
-
-	UPROPERTY(EditAnywhere, Category = "Corotated")
-	bool bEnableCorotatedConstraint = true;
-
-	UPROPERTY(EditAnywhere, Category = "Corotated")
-	bool bDoBlended = false;
-
-	UPROPERTY(EditAnywhere, Category = "Corotated")
-	float BlendedZeta = 0;
-};
-
-USTRUCT(BlueprintType)
-struct FSolverConstraintsGroup
-{
-	GENERATED_USTRUCT_BODY()
-
-	UPROPERTY(EditAnywhere, Category = "Constraints")
-	bool bEnablePositionTargets = true;
-
-	UPROPERTY(EditAnywhere, Category = "Constraints")
-	bool bEnableKinematics = true;
-
-	UPROPERTY(EditAnywhere, Category = "Constraints")
-	FSolverCorotatedConstraintsGroup CorotatedConstraints;
-
-	/**
-	* These are options for another solver. 
-	*/
-	UPROPERTY(EditAnywhere, BlueprintReadOnly, Category = "Constraints")
-	FSolverGaussSeidelConstraintsGroup GaussSeidelConstraints;
-};
-
-USTRUCT(BlueprintType)
-struct FSolverForcesGroup
-{
-	GENERATED_USTRUCT_BODY()
-
-	UPROPERTY(EditAnywhere, Category = "Forces")
-	float YoungModulus = 100000;
-
-	UPROPERTY(EditAnywhere, Category = "Forces")
-	float Damping = 0;
-
-	UPROPERTY(EditAnywhere, Category = "Forces")
-	bool bEnableGravity = true;
-};
-
-USTRUCT(BlueprintType)
-struct FSolverMuscleActivationGroup
-{
-	GENERATED_USTRUCT_BODY()
-
-	UPROPERTY(EditAnywhere, Category = "MuscleActivation")
-	bool bDoMuscleActivation = false;
-};
-
 USTRUCT(BlueprintType)
 struct FDataflowFleshSolverProxy : public FDataflowPhysicsSolverProxy
 {
@@ -296,11 +34,10 @@ struct FDataflowFleshSolverProxy : public FDataflowPhysicsSolverProxy
 	FDataflowFleshSolverProxy(Chaos::Softs::FDeformableSolverProperties InProp = Chaos::Softs::FDeformableSolverProperties()) :
 		FDataflowPhysicsSolverProxy()
 	{
-		Solver = MakeUnique<Chaos::Softs::FDeformableSolver>(InProp);
 	}
 	virtual ~FDataflowFleshSolverProxy() override = default;
 
-	// Begin FPhysicsSolverInterface overrides
+	//~ Begin FPhysicsSolverInterface interface
 	virtual void AdvanceSolverDatas(const float DeltaTime) override
 	{
 		Chaos::Softs::FDeformableSolver::FPhysicsThreadAccess PhysicsThreadAccess(Solver.Get(), Chaos::Softs::FPhysicsThreadAccessor());
@@ -312,7 +49,12 @@ struct FDataflowFleshSolverProxy : public FDataflowPhysicsSolverProxy
 		return PhysicsThreadAccess.GetProperties().TimeStepSize;
 	}
 	virtual bool IsValid() const override { return Solver.IsValid();}
-	// End FPhysicsSolverInterface overrides
+	
+	virtual const UScriptStruct* GetScriptStruct() const override
+	{
+		return StaticStruct();
+	}
+	//~ End FPhysicsSolverInterface interface
 
 	/** Chaos deformable solver that will be used in the component */
 	TUniquePtr<Chaos::Softs::FDeformableSolver> Solver;
@@ -354,6 +96,12 @@ public:
 	virtual void ReadFromSimulation(const float DeltaTime) override;
 	// End IDataflowPhysicsSolverInterface overrides
 
+	//~ Begin UObject Interface
+#if WITH_EDITOR
+	virtual bool CanEditChange(const FProperty* InProperty) const override;
+#endif // WITH_EDITOR
+	//~ End UObject Interface
+
 	// Begin UActorComponent overrides
 	virtual bool ShouldCreatePhysicsState() const override {return true;}
 	virtual void BeginPlay() override;
@@ -376,9 +124,15 @@ public:
 
 	/* Callback to trigger the deformable update after the simulation */
 	void UpdateDeformableEndTickState(bool bRegister);
+
+	/** Stop the simulation, and keep the cloth in its last pose. */
+	UFUNCTION(CallInEditor, BlueprintCallable, Category = "Physics")
+	void ResetSimulationProperties(const FSolverTimingGroup& TimingGroup, const FSolverEvolutionGroup& EvolutionGroup,
+		FSolverCollisionsGroup CollisionsGroup, FSolverConstraintsGroup ConstraintsGroup, FSolverForcesGroup ForcesGroup,
+		FSolverDebuggingGroup DebuggingGroup, FSolverMuscleActivationGroup MuscleActivationGroup);
 	
 	/* Solver dataflow asset used to advance in time */
-	UPROPERTY(EditAnywhere, BlueprintReadWrite, Category = "Physics")
+	UPROPERTY(EditAnywhere, BlueprintReadWrite, Category = "Physics", meta=(EditConditionHides), AdvancedDisplay)
 	FDataflowSimulationAsset SimulationAsset;
 
 	/* Properties : Do NOT place ungrouped properties in this class */
diff --git a/Engine/Plugins/Experimental/ChaosFlesh/Source/ChaosFleshEngine/Public/ChaosFlesh/ChaosDeformableSolverGroups.h b/Engine/Plugins/Experimental/ChaosFlesh/Source/ChaosFleshEngine/Public/ChaosFlesh/ChaosDeformableSolverGroups.h
new file mode 100644
--- /dev/null
+++ b/Engine/Plugins/Experimental/ChaosFlesh/Source/ChaosFleshEngine/Public/ChaosFlesh/ChaosDeformableSolverGroups.h
@@ -0,0 +1,238 @@
+﻿// Copyright Epic Games, Inc. All Rights Reserved.
+#pragma once
+
+#include "ChaosFlesh/ChaosDeformableTypes.h"
+#include "UObject/ObjectMacros.h"
+
+#include "ChaosDeformableSolverGroups.generated.h"
+USTRUCT(BlueprintType)
+struct FSolverTimingGroup
+{
+	GENERATED_BODY()
+
+	UPROPERTY(EditAnywhere, Category = "SolverTiming")
+	int32 NumSubSteps = 2;
+
+	UPROPERTY(EditAnywhere, Category = "SolverTiming")
+	int32 NumSolverIterations = 5;
+
+	UPROPERTY(EditAnywhere, Category = "SolverTiming")
+	bool FixTimeStep = false;
+
+	UPROPERTY(EditAnywhere, Category = "SolverTiming")
+	float TimeStepSize = 0.05;
+
+	UPROPERTY(EditAnywhere, Category = "SolverTiming")
+	bool bDoThreadedAdvance = true;
+
+	/** ObjectType defines how to initialize the rigid objects state, Kinematic, Sleeping, Dynamic. */
+	UPROPERTY(EditAnywhere, BlueprintReadOnly, Category = "SolverTiming")
+	EDeformableExecutionModel ExecutionModel = EDeformableExecutionModel::Chaos_Deformable_PostPhysics;
+};
+
+USTRUCT(BlueprintType)
+struct FSolverDebuggingGroup
+{
+	GENERATED_BODY()
+
+	UPROPERTY(EditAnywhere, Category = "Physics")
+	bool CacheToFile = false;
+};
+
+USTRUCT(BlueprintType)
+struct FSolverQuasistaticsGroup
+{
+	GENERATED_BODY()
+
+	UPROPERTY(EditAnywhere, Category = "Quasistatics")
+	bool bDoQuasistatics = false;
+};
+
+USTRUCT(BlueprintType)
+struct FSolverEvolutionGroup
+{
+	GENERATED_BODY()
+
+	UPROPERTY(EditAnywhere, Category = "Evolution")
+	FSolverQuasistaticsGroup SolverQuasistatics;
+
+};
+
+USTRUCT(BlueprintType)
+struct FSolverGridBasedCollisionsGroup
+{
+	GENERATED_BODY()
+
+	UPROPERTY(EditAnywhere, Category = "GridBasedCollisions")
+	bool bUseGridBasedConstraints = false;
+
+	UPROPERTY(EditAnywhere, Category = "GridBasedCollisions")
+	float GridDx = 25.;
+};
+
+USTRUCT(BlueprintType)
+struct FInComponentSpringCollisionGroup
+{
+	GENERATED_BODY()
+	
+	/** If uses in-component spring self-collision */
+	UPROPERTY(EditAnywhere, Category = "InComponentSpringCollision")
+	bool bDoInComponentSpringCollision = false;
+	
+	/** N ring to exclude for in-component spring self-collision */
+	UPROPERTY(EditAnywhere, Category = "InComponentSpringCollision")
+	int32 NRingExcluded = 1;
+};
+
+USTRUCT(BlueprintType)
+struct FSpringCollisionGroup
+{
+	GENERATED_BODY()
+	
+	/** If uses component-component spring collision */
+	UPROPERTY(EditAnywhere, Category = "SpringCollision")
+	bool bDoSpringCollision = false;
+	
+	/** In-component spring self collision detection parameters */
+	UPROPERTY(EditAnywhere, BlueprintReadOnly, Category = "SpringCollision")
+	FInComponentSpringCollisionGroup InComponentSpringCollision;
+	
+	/** Search radius for point triangle collision pairs */
+	UPROPERTY(EditAnywhere, Category = "SpringCollision")
+	float CollisionSearchRadius = 0.f;
+	
+	/** Collision spring stiffness; larger value will stop penetration better */
+	UPROPERTY(EditAnywhere, Category = "SpringCollision")
+	float SpringCollisionStiffness = 500.f;
+	
+	/**  Anisotropic springs will allow sliding on the triangle */
+	UPROPERTY(EditAnywhere, Category = "SpringCollision")
+	bool bAllowSliding = true;
+	
+	/** Do self collision with kinematic triangles as well */
+	UPROPERTY(EditAnywhere, Category = "SpringCollision")
+	bool bCollideWithFullmesh = true;
+};
+
+USTRUCT(BlueprintType)
+struct FSphereRepulsionGroup
+{
+	GENERATED_BODY()
+	
+	/** If uses sphere repulsion for collision */
+	UPROPERTY(EditAnywhere, Category = "SphereRepulsion")
+	bool bDoSphereRepulsion = false;
+	
+	/** Search radius for repulsion pairs */
+	UPROPERTY(EditAnywhere, Category = "SphereRepulsion")
+	float SphereRepulsionRadius = 0.f;
+	
+	/** Stiffness for sphere repulsion */
+	UPROPERTY(EditAnywhere, Category = "SphereRepulsion")
+	float SphereRepulsionStiffness = 500.f;
+};
+
+USTRUCT(BlueprintType)
+struct FSolverGaussSeidelConstraintsGroup
+{
+	GENERATED_BODY()
+
+	/** Enable the Gauss Seidel solver instead of the existing XPBD. */
+	UPROPERTY(EditAnywhere, Category = "GaussSeidelConstraints")
+	bool bUseGaussSeidelConstraints = false;
+
+	/** Enable another model that runs simulation faster. */
+	UPROPERTY(EditAnywhere, Category = "GaussSeidelConstraints")
+	bool bUseGSNeohookean = false;
+
+	/** Enable acceleration technique for Gauss Seidel solver to make simulation look better within a limited budget. */
+	UPROPERTY(EditAnywhere, Category = "GaussSeidelConstraints")
+	bool bUseSOR = true;
+
+	/** Acceleration related parameter. Tune it down if simulation becomes unstable. */
+	UPROPERTY(EditAnywhere, Category = "GaussSeidelConstraints")
+	float OmegaSOR = 1.6f;
+
+	/** Enable dynamic springs controlled by constraint manager. */
+	UPROPERTY(EditAnywhere, Category = "GaussSeidelConstraints")
+	bool bEnableDynamicSprings = true;
+	
+	/** Component-component collision detection radius and stiffness */
+	UPROPERTY(EditAnywhere, BlueprintReadOnly, Category = "GaussSeidelConstraints")
+	FSpringCollisionGroup SpringCollision;
+
+	/** Sphere repulsion parameters */
+	UPROPERTY(EditAnywhere, BlueprintReadOnly, Category = "GaussSeidelConstraints")
+	FSphereRepulsionGroup SphereRepulsion;
+};
+
+USTRUCT(BlueprintType)
+struct FSolverCollisionsGroup
+{
+	GENERATED_BODY()
+
+	UPROPERTY(EditAnywhere, Category = "Physics")
+	bool bUseFloor = true;
+
+	//UPROPERTY(EditAnywhere, Category = "Collisions")
+	//FSolverGridBasedCollisionsGroup SolverGridBasedCollisions;
+};
+
+USTRUCT(BlueprintType)
+struct FSolverCorotatedConstraintsGroup
+{
+	GENERATED_BODY()
+
+	UPROPERTY(EditAnywhere, Category = "Corotated")
+	bool bEnableCorotatedConstraint = true;
+
+	UPROPERTY(EditAnywhere, Category = "Corotated")
+	bool bDoBlended = false;
+
+	UPROPERTY(EditAnywhere, Category = "Corotated")
+	float BlendedZeta = 0;
+};
+
+USTRUCT(BlueprintType)
+struct FSolverConstraintsGroup
+{
+	GENERATED_BODY()
+
+	UPROPERTY(EditAnywhere, Category = "Constraints")
+	bool bEnablePositionTargets = true;
+
+	UPROPERTY(EditAnywhere, Category = "Constraints")
+	bool bEnableKinematics = true;
+
+	UPROPERTY(EditAnywhere, Category = "Constraints")
+	FSolverCorotatedConstraintsGroup CorotatedConstraints;
+
+	/** These are options for another solver. */
+	UPROPERTY(EditAnywhere, BlueprintReadOnly, Category = "Constraints")
+	FSolverGaussSeidelConstraintsGroup GaussSeidelConstraints;
+};
+
+USTRUCT(BlueprintType)
+struct FSolverForcesGroup
+{
+	GENERATED_BODY()
+
+	UPROPERTY(EditAnywhere, Category = "Forces")
+	float YoungModulus = 100000;
+
+	UPROPERTY(EditAnywhere, Category = "Forces")
+	float Damping = 0;
+
+	UPROPERTY(EditAnywhere, Category = "Forces")
+	bool bEnableGravity = true;
+};
+
+USTRUCT(BlueprintType)
+struct FSolverMuscleActivationGroup
+{
+	GENERATED_BODY()
+
+	UPROPERTY(EditAnywhere, Category = "MuscleActivation")
+	bool bDoMuscleActivation = false;
+};
+
diff --git a/Engine/Plugins/Experimental/ChaosFlesh/Source/ChaosFleshEngine/Public/ChaosFlesh/FleshAsset.h b/Engine/Plugins/Experimental/ChaosFlesh/Source/ChaosFleshEngine/Public/ChaosFlesh/FleshAsset.h
--- a/Engine/Plugins/Experimental/ChaosFlesh/Source/ChaosFleshEngine/Public/ChaosFlesh/FleshAsset.h
+++ b/Engine/Plugins/Experimental/ChaosFlesh/Source/ChaosFleshEngine/Public/ChaosFlesh/FleshAsset.h
@@ -4,6 +4,7 @@
 
 #include "CoreMinimal.h"
 
+#include "ChaosFlesh/ChaosDeformableSolverGroups.h"
 #include "ChaosFlesh/FleshCollection.h"
 #include "Dataflow/DataflowNodeParameters.h"
 #include "Dataflow/DataflowEngineTypes.h"
@@ -161,7 +162,70 @@ public:
 	/** Animation asset used in this asset */
 	UPROPERTY(DuplicateTransient, AssetRegistrySearchable)
 	TSoftObjectPtr<UAnimationAsset> PreviewAnimationAsset = nullptr;
+
+	UPROPERTY(DuplicateTransient, AssetRegistrySearchable)
+	FSolverTimingGroup PreviewSolverTiming;
+
+	UPROPERTY(DuplicateTransient, AssetRegistrySearchable)
+	FSolverEvolutionGroup PreviewSolverEvolution;
+
+	UPROPERTY(DuplicateTransient, AssetRegistrySearchable)
+	FSolverCollisionsGroup PreviewSolverCollisions;
+
+	UPROPERTY(DuplicateTransient, AssetRegistrySearchable)
+	FSolverConstraintsGroup PreviewSolverConstraints;
+
+	UPROPERTY(DuplicateTransient, AssetRegistrySearchable)
+	FSolverForcesGroup PreviewSolverForces;
+
+	UPROPERTY(DuplicateTransient, AssetRegistrySearchable)
+	FSolverDebuggingGroup PreviewSolverDebugging;
+
+	UPROPERTY(DuplicateTransient, AssetRegistrySearchable)
+	FSolverMuscleActivationGroup PreviewSolverMuscleActivation;
 	
 #endif // WITH_EDITORONLY_DATA
 
 };
+
+
+/** 
+ * Dataflow content owning dataflow and solver properties that will be used to evaluate the graph
+ */
+UCLASS()
+class CHAOSFLESHENGINE_API  UDataflowFleshContent : public UDataflowSkeletalContent
+{
+	GENERATED_BODY()
+
+public:
+	UDataflowFleshContent();
+	virtual ~UDataflowFleshContent() override {}
+
+	//~ UObject interface
+	static void AddReferencedObjects(UObject* InThis, FReferenceCollector& Collector);
+
+	/** Set all the preview actor exposed properties */
+	virtual void SetActorProperties(TObjectPtr<AActor>& PreviewActor) const override;
+
+	UPROPERTY(EditAnywhere, Category = "Preview", Transient, SkipSerialization)
+	FSolverTimingGroup SolverTiming;
+
+	UPROPERTY(EditAnywhere, Category = "Preview", Transient, SkipSerialization)
+	FSolverEvolutionGroup SolverEvolution;
+
+	UPROPERTY(EditAnywhere, Category = "Preview", Transient, SkipSerialization)
+	FSolverCollisionsGroup SolverCollisions;
+
+	UPROPERTY(EditAnywhere, Category = "Preview", Transient, SkipSerialization)
+	FSolverConstraintsGroup SolverConstraints;
+
+	UPROPERTY(EditAnywhere, Category = "Preview", Transient, SkipSerialization)
+	FSolverForcesGroup SolverForces;
+
+	UPROPERTY(EditAnywhere, Category = "Preview", Transient, SkipSerialization)
+	FSolverDebuggingGroup SolverDebugging;
+
+	UPROPERTY(EditAnywhere, Category = "Preview", Transient, SkipSerialization)
+	FSolverMuscleActivationGroup SolverMuscleActivation;
+};
+
diff --git a/Engine/Plugins/Experimental/Dataflow/Dataflow.uplugin b/Engine/Plugins/Experimental/Dataflow/Dataflow.uplugin
--- a/Engine/Plugins/Experimental/Dataflow/Dataflow.uplugin
+++ b/Engine/Plugins/Experimental/Dataflow/Dataflow.uplugin
@@ -35,11 +35,6 @@
 			"Name": "DataflowNodes",
 			"Type": "Runtime",
 			"LoadingPhase": "Default"
-		},
-		{
-			"Name": "DataflowSimulation",
-			"Type": "Runtime",
-			"LoadingPhase": "Default"
 		}
 	],
 	"Plugins": [
diff --git a/Engine/Plugins/Experimental/Dataflow/Source/DataflowEditor/Private/Dataflow/DataflowSimulationControls.cpp b/Engine/Plugins/Experimental/Dataflow/Source/DataflowEditor/Private/Dataflow/DataflowSimulationControls.cpp
new file mode 100644
--- /dev/null
+++ b/Engine/Plugins/Experimental/Dataflow/Source/DataflowEditor/Private/Dataflow/DataflowSimulationControls.cpp
@@ -0,0 +1,269 @@
+﻿// Copyright Epic Games, Inc. All Rights Reserved.
+
+#include "Dataflow/DataflowSimulationControls.h"
+
+#include "Animation/AnimSingleNodeInstance.h"
+#include "Animation/AnimSequenceBase.h"
+#include "Chaos/CacheManagerActor.h"
+#include "Chaos/Adapters/CacheAdapter.h"
+#include "Dataflow/DataflowContent.h"
+#include "Dataflow/DataflowObject.h"
+#include "Features/IModularFeatures.h"
+#include "Components/SkeletalMeshComponent.h"
+#include "Engine/SkeletalMesh.h"
+#include "Engine/World.h"
+#include "BonePose.h"
+#include "AnimationRuntime.h"
+
+#define LOCTEXT_NAMESPACE "DataflowSimulationGenerator"
+
+namespace UE::Dataflow
+{
+	bool ShouldResetWorld(const TObjectPtr<UDataflow>& SimulationGraph, const TObjectPtr<UWorld>& SimulationWorld, ::Dataflow::FTimestamp& LastTimeStamp)
+	{
+		if(const TSharedPtr<::Dataflow::FGraph> DataflowGraph = SimulationGraph->GetDataflow())
+		{
+			::Dataflow::FTimestamp MaxTimeStamp = ::Dataflow::FTimestamp::Invalid;
+			for(const TSharedPtr<FDataflowNode>& TerminalNode : DataflowGraph->GetFilteredNodes(FDataflowTerminalNode::StaticType()))
+			{
+				MaxTimeStamp.Value = FMath::Max(MaxTimeStamp.Value, TerminalNode->GetTimestamp().Value);
+			}
+			if(MaxTimeStamp.Value > LastTimeStamp.Value)
+			{
+				LastTimeStamp = MaxTimeStamp.Value;
+				return true;
+			}
+		}
+		return false;
+	}
+	
+	TObjectPtr<AActor> SpawnSimulatedActor(const TSubclassOf<AActor>& ActorClass,
+		const TObjectPtr<AChaosCacheManager>& CacheManager, const TObjectPtr<UChaosCacheCollection>& CacheCollection,
+		const bool bIsRecording, const TObjectPtr<UDataflowBaseContent>& DataflowContent)
+	{
+		if(CacheManager)
+		{
+			FActorSpawnParameters SpawnParameters;
+			SpawnParameters.Name = TEXT("CacheActor");
+			SpawnParameters.NameMode = FActorSpawnParameters::ESpawnActorNameMode::Requested;
+			SpawnParameters.Owner = CacheManager.Get(); 
+			SpawnParameters.bDeferConstruction = true;
+			
+			TObjectPtr<AActor> PreviewActor = CacheManager->GetWorld()->SpawnActor<AActor>(ActorClass, SpawnParameters);
+			if(PreviewActor)
+			{
+				// Link the editor content properties to the BP actor one 
+				DataflowContent->SetActorProperties(PreviewActor);
+
+				// Finish spawning
+				PreviewActor->FinishSpawning(FTransform::Identity, true);
+			}
+
+			CacheManager->CacheCollection = CacheCollection;
+			CacheManager->StartMode = EStartMode::Timed;
+			CacheManager->CacheMode = bIsRecording ? ECacheMode::Record : ECacheMode::None;
+			
+			// Get the implementation of our adapters for identifying compatible components
+			IModularFeatures&                      ModularFeatures = IModularFeatures::Get();
+			TArray<Chaos::FComponentCacheAdapter*> Adapters = ModularFeatures.GetModularFeatureImplementations<Chaos::FComponentCacheAdapter>(Chaos::FComponentCacheAdapter::FeatureName);
+
+			if(PreviewActor)
+			{
+				TInlineComponentArray<UPrimitiveComponent*> PrimComponents;
+				PreviewActor->GetComponents(PrimComponents);
+	
+				for(UPrimitiveComponent* PrimComponent : PrimComponents)
+				{
+					if(Chaos::FAdapterUtil::GetBestAdapterForClass(PrimComponent->GetClass(), false))
+					{
+						const FName ChannelName(PrimComponent->GetName());
+						CacheManager->FindOrAddObservedComponent(PrimComponent, ChannelName, true);
+					}
+				}
+			}
+			return PreviewActor;
+		}
+		return nullptr;
+	}
+	
+	void SetupSkeletonAnimation(const TObjectPtr<AActor>& PreviewActor)
+	{
+		if(PreviewActor)
+		{
+			TInlineComponentArray<UPrimitiveComponent*> PrimComponents;
+			PreviewActor->GetComponents(PrimComponents);
+
+			for(UPrimitiveComponent* PrimComponent : PrimComponents)
+			{
+				if(USkeletalMeshComponent* SkeletalMeshComponent = Cast<USkeletalMeshComponent>(PrimComponent))
+				{
+					SkeletalMeshComponent->SetAnimationMode(EAnimationMode::AnimationSingleNode);
+					SkeletalMeshComponent->InitAnim(true);
+					
+					if(UAnimSingleNodeInstance* AnimNodeInstance = SkeletalMeshComponent->GetSingleNodeInstance())
+					{
+						// Setup the animation instance
+						AnimNodeInstance->SetAnimationAsset(SkeletalMeshComponent->AnimationData.AnimToPlay);
+						AnimNodeInstance->InitializeAnimation();
+
+						// Update the anim data
+						SkeletalMeshComponent->AnimationData.PopulateFrom(AnimNodeInstance);
+#if WITH_EDITOR
+						SkeletalMeshComponent->ValidateAnimation();
+#endif
+
+						// Stop the animation 
+						AnimNodeInstance->SetLooping(true);
+						AnimNodeInstance->SetPlaying(false);
+						
+					}
+				}
+			}
+		}
+	}
+
+	static void FillAnimationDatas(const UAnimSequenceBase* AnimSequence, const float CurrentTime, USkeletalMesh* InSkeletalMesh, TArray<FTransform>& ComponentSpaceTransforms)
+	{
+		const FAnimExtractContext ExtractionContext(FMath::Clamp(CurrentTime, 0., AnimSequence->GetPlayLength()));
+
+		if(const FReferenceSkeleton* ReferenceSkeleton = &InSkeletalMesh->GetRefSkeleton())
+		{
+			const int32 NumBones = ReferenceSkeleton->GetNum();
+
+			TArray<FBoneIndexType> BoneIndices;
+			BoneIndices.SetNumUninitialized(NumBones);
+			for (int32 Index = 0; Index < NumBones; ++Index)
+			{
+				BoneIndices[Index] = StaticCast<FBoneIndexType>(Index);
+			}
+
+			FBoneContainer BoneContainer;
+			BoneContainer.SetUseRAWData(true);
+			BoneContainer.InitializeTo(BoneIndices, UE::Anim::FCurveFilterSettings(), *InSkeletalMesh);
+		
+			FCompactPose CompactPose;
+			CompactPose.SetBoneContainer(&BoneContainer);
+
+			FBlendedCurve BlendedCurve;
+			BlendedCurve.InitFrom(BoneContainer);
+
+			UE::Anim::FStackAttributeContainer TempAttributes;
+			FAnimationPoseData AnimationPoseData(CompactPose, BlendedCurve, TempAttributes);
+			AnimSequence->GetAnimationPose(AnimationPoseData, ExtractionContext);
+		
+			FAnimationRuntime::FillUpComponentSpaceTransforms(*ReferenceSkeleton, AnimationPoseData.GetPose().GetBones(), ComponentSpaceTransforms);
+		}
+	}
+
+	void ComputeSkeletonAnimation(const TObjectPtr<AActor>& PreviewActor, const float SimulationTime)
+	{
+		if(PreviewActor)
+		{
+			TInlineComponentArray<UPrimitiveComponent*> PrimComponents;
+			PreviewActor->GetComponents(PrimComponents);
+
+			// Update all the animation time 
+			for(UPrimitiveComponent* PrimComponent : PrimComponents)
+			{
+				if(USkeletalMeshComponent* SkeletalMeshComponent = Cast<USkeletalMeshComponent>(PrimComponent))
+				{
+					if(UAnimSingleNodeInstance* SingleNodeInstance = Cast<UAnimSingleNodeInstance>(SkeletalMeshComponent->GetAnimInstance()))
+					{
+						if(const UAnimSequenceBase* AnimSequence = Cast<UAnimSequenceBase>(SingleNodeInstance->GetAnimationAsset()))
+						{
+							TArray<FTransform> ComponentSpaceTransforms = SkeletalMeshComponent->GetComponentSpaceTransforms();
+							FillAnimationDatas(AnimSequence, SimulationTime, SkeletalMeshComponent->GetSkeletalMeshAsset(), ComponentSpaceTransforms);
+							
+							SkeletalMeshComponent->GetEditableComponentSpaceTransforms() = ComponentSpaceTransforms;
+							SkeletalMeshComponent->ApplyEditedComponentSpaceTransforms();
+						}
+					}
+				}
+			}
+		}
+	}
+
+	void UpdateSkeletonAnimation(const TObjectPtr<AActor>& PreviewActor, const float SimulationTime)
+	{
+		if(PreviewActor)
+		{
+			TInlineComponentArray<UPrimitiveComponent*> PrimComponents;
+			PreviewActor->GetComponents(PrimComponents);
+
+			// Update all the animation time 
+			for(UPrimitiveComponent* PrimComponent : PrimComponents)
+			{
+				if(USkeletalMeshComponent* SkeletalMeshComponent = Cast<USkeletalMeshComponent>(PrimComponent))
+				{
+					SkeletalMeshComponent->SetPosition(SimulationTime);
+					SkeletalMeshComponent->TickAnimation(0.f, false /*bNeedsValidRootMotion*/);
+					SkeletalMeshComponent->RefreshBoneTransforms(nullptr /*TickFunction*/);
+				}
+			}
+		}
+	}
+
+	void StartSkeletonAnimation(const TObjectPtr<AActor>& PreviewActor)
+	{
+		if(PreviewActor)
+		{
+			TInlineComponentArray<UPrimitiveComponent*> PrimComponents;
+			PreviewActor->GetComponents(PrimComponents);
+
+			for(UPrimitiveComponent* PrimComponent : PrimComponents)
+			{
+				if(const USkeletalMeshComponent* SkeletalMeshComponent = Cast<USkeletalMeshComponent>(PrimComponent))
+				{
+					if(UAnimSingleNodeInstance* AnimNodeInstance = SkeletalMeshComponent->GetSingleNodeInstance())
+					{
+						AnimNodeInstance->SetPlaying(true);
+					}
+				}
+			}
+		}
+	}
+	
+	void PauseSkeletonAnimation(const TObjectPtr<AActor>& PreviewActor)
+	{
+		if(PreviewActor)
+		{
+			TInlineComponentArray<UPrimitiveComponent*> PrimComponents;
+			PreviewActor->GetComponents(PrimComponents);
+
+			for(UPrimitiveComponent* PrimComponent : PrimComponents)
+			{
+				if(const USkeletalMeshComponent* SkeletalMeshComponent = Cast<USkeletalMeshComponent>(PrimComponent))
+				{
+					if(UAnimSingleNodeInstance* AnimNodeInstance = SkeletalMeshComponent->GetSingleNodeInstance())
+					{
+						AnimNodeInstance->SetPlaying(false);
+					}
+				}
+			}
+		}
+	}
+	
+	void StepSkeletonAnimation(const TObjectPtr<AActor>& PreviewActor)
+	{
+		if(PreviewActor)
+		{
+			TInlineComponentArray<UPrimitiveComponent*> PrimComponents;
+			PreviewActor->GetComponents(PrimComponents);
+
+			for(UPrimitiveComponent* PrimComponent : PrimComponents)
+			{
+				if(const USkeletalMeshComponent* SkeletalMeshComponent = Cast<USkeletalMeshComponent>(PrimComponent))
+				{
+					if(UAnimSingleNodeInstance* AnimNodeInstance = SkeletalMeshComponent->GetSingleNodeInstance())
+					{
+						AnimNodeInstance->SetPlaying(false);
+						AnimNodeInstance->StepForward();
+					}
+				}
+			}
+		}
+	}
+	
+}
+
+#undef LOCTEXT_NAMESPACE
diff --git a/Engine/Plugins/Experimental/Dataflow/Source/DataflowEditor/Private/Dataflow/DataflowSimulationGenerator.cpp b/Engine/Plugins/Experimental/Dataflow/Source/DataflowEditor/Private/Dataflow/DataflowSimulationGenerator.cpp
--- a/Engine/Plugins/Experimental/Dataflow/Source/DataflowEditor/Private/Dataflow/DataflowSimulationGenerator.cpp
+++ b/Engine/Plugins/Experimental/Dataflow/Source/DataflowEditor/Private/Dataflow/DataflowSimulationGenerator.cpp
@@ -2,6 +2,8 @@
 
 #include "Dataflow/DataflowSimulationGenerator.h"
 #include "Dataflow/DataflowSimulationManager.h"
+#include "Dataflow/DataflowSimulationControls.h"
+#include "Dataflow/DataflowSimulationUtils.h"
 #include "Chaos/CacheManagerActor.h"
 #include "Engine/World.h"
 #include "Engine/Engine.h"
@@ -11,12 +13,13 @@ DEFINE_LOG_CATEGORY(LogDataflowSimulationGenerator);
 
 #define LOCTEXT_NAMESPACE "DataflowSimulationGenerator"
 
-namespace Dataflow
+namespace UE::Dataflow
 {
 void FDataflowSimulationTask::DoWork()
 {
+	const int32 StartFrame = bBackgroundTask ? 1 : 0;
 	const int32 NumFrames = (MaxTime-MinTime) / DeltaTime;
-	for (int32 FrameIndex = 0; FrameIndex < NumFrames; ++FrameIndex)
+	for (int32 FrameIndex = StartFrame; FrameIndex < NumFrames; ++FrameIndex)
 	{
 		if (!TaskManager->bCancelled.load())
 		{
@@ -25,9 +28,20 @@ void FDataflowSimulationTask::DoWork()
 
 			if(bBackgroundTask)
 			{
+				// Compute all the skelmesh animations at the simulation time
+				Dataflow::ComputeSkeletonAnimation(TaskManager->PreviewActor, SimulationTime);
+				
 				// Background task : run directly the advance simulation data without coming back to the game thread
 				UDataflowSimulationManager* DataflowManager = SimulationWorld->GetSubsystem<UDataflowSimulationManager>();
-                DataflowManager->AdvanceSimulationData(DeltaTime, SimulationTime);
+
+				// Pre advance the proxies
+				DataflowManager->PreAdvanceProxies(DeltaTime);
+
+				// Advance the simulation proxies
+                DataflowManager->AdvanceSimulationProxies(DeltaTime, SimulationTime);
+
+				// Post advanc e the simulation proxies
+				DataflowManager->PostAdvanceProxies(DeltaTime);
 			}
 			else
 			{
@@ -59,6 +73,7 @@ bool FDataflowTaskManager::AllocateSimulationResource(const FVector2f& TimeRange
 	CacheManager = SimulationWorld->SpawnActor<AChaosCacheManager>();
 
 	PreviewActor = Dataflow::SpawnSimulatedActor(ActorClass, CacheManager, CacheAsset, true, DataflowContent);
+	Dataflow::SetupSkeletonAnimation(PreviewActor);
 
 	// Init the cache manager
 	CacheManager->SetObservedComponentProperties(CacheManager->CacheMode);
@@ -77,10 +92,19 @@ bool FDataflowTaskManager::AllocateSimulationResource(const FVector2f& TimeRange
 		SimulationTask->GetTask().DeltaTime = (TimeRange[1]-TimeRange[0]) / NumFrames;
 
 		UDataflowSimulationManager* DataflowManager = SimulationWorld->GetSubsystem<UDataflowSimulationManager>();
+		if(SimulationTask->GetTask().bBackgroundTask)
+		{
+			// Update all the skelmesh animations at the simulation time
+			Dataflow::UpdateSkeletonAnimation(PreviewActor, SimulationTask->GetTask().MinTime+SimulationTask->GetTask().DeltaTime);
+				
+			// Foreground task : Run the world ticking 
+			SimulationWorld->Tick(ELevelTick::LEVELTICK_All, SimulationTask->GetTask().DeltaTime);
+
+			// Init simulation proxies from interface
+			DataflowManager->InitSimulationInterfaces();
+		}
 		DataflowManager->SetSimulationEnabled(!SimulationTask->GetTask().bBackgroundTask);
-		DataflowManager->WriteSimulationData(SimulationTask->GetTask().DeltaTime);
 	}
-	Dataflow::SetupSkeletonAnimation(PreviewActor);
 
 	return true;
 }
@@ -90,9 +114,12 @@ void FDataflowTaskManager::FreeSimulationResource()
 	if (SimulationTask.IsValid())
 	{
 		SimulationTask->EnsureCompletion();
-		
-		UDataflowSimulationManager* DataflowManager = SimulationWorld->GetSubsystem<UDataflowSimulationManager>();
-		DataflowManager->ReadSimulationData(SimulationTask->GetTask().DeltaTime);
+
+		if(SimulationTask->GetTask().bBackgroundTask)
+		{
+			UDataflowSimulationManager* DataflowManager = SimulationWorld->GetSubsystem<UDataflowSimulationManager>();
+			DataflowManager->ResetSimulationInterfaces();
+		}
 	}
 
 	if(CacheManager)
diff --git a/Engine/Plugins/Experimental/Dataflow/Source/DataflowEditor/Private/Dataflow/DataflowSimulationPanel.cpp b/Engine/Plugins/Experimental/Dataflow/Source/DataflowEditor/Private/Dataflow/DataflowSimulationPanel.cpp
--- a/Engine/Plugins/Experimental/Dataflow/Source/DataflowEditor/Private/Dataflow/DataflowSimulationPanel.cpp
+++ b/Engine/Plugins/Experimental/Dataflow/Source/DataflowEditor/Private/Dataflow/DataflowSimulationPanel.cpp
@@ -33,7 +33,7 @@ void SDataflowSimulationPanel::Construct( const SDataflowSimulationPanel::FArgum
 	this->ChildSlot
 	[
 		SNew(SHorizontalBox)
-		.AddMetaData<FTagMetaData>(TEXT("ClothAnimScrub.Scrub"))
+		.AddMetaData<FTagMetaData>(TEXT("DataflowSimulationScrub.Scrub"))
 		+SHorizontalBox::Slot()
 		.HAlign(HAlign_Fill) 
 		.VAlign(VAlign_Center)
@@ -136,7 +136,7 @@ FReply SDataflowSimulationPanel::OnClick_Forward_End()
 {
 	if (const TSharedPtr<FDataflowSimulationScene> PreviewScene = SimulationScene.Pin())
 	{
-		PreviewScene->SimulationTime = GetSequenceLength();
+		PreviewScene->SimulationTime = PreviewScene->GetTimeRange()[1];
 	}
 
 	return FReply::Handled();
@@ -157,7 +157,7 @@ FReply SDataflowSimulationPanel::OnClick_Backward_End()
 {
 	if (const TSharedPtr<FDataflowSimulationScene> PreviewScene = SimulationScene.Pin())
 	{
-		PreviewScene->SimulationTime = 0.0f;
+		PreviewScene->SimulationTime = PreviewScene->GetTimeRange()[0];
 	}
 
 	return FReply::Handled();
@@ -229,23 +229,23 @@ void SDataflowSimulationPanel::OnTickPlayback(double InCurrentTime, float InDelt
 		
 		if(PreviewPlaybackMode == EDataflowPlaybackMode::Looping)
 		{
-			PreviewScene->SimulationTime = SimulationTime - SequenceLength * FMath::Floor(SimulationTime / SequenceLength);
+			PreviewScene->SimulationTime = SimulationTime - SequenceLength * FMath::Floor((SimulationTime-PreviewScene->GetTimeRange()[0]) / SequenceLength);
 		}
 		else
 		{
 			if(PreviewPlaybackMode == EDataflowPlaybackMode::PingPong)
 			{
-				if((PlaybackMode == EPlaybackMode::PlayingForward) && (SimulationTime >= SequenceLength))
+				if((PlaybackMode == EPlaybackMode::PlayingForward) && (SimulationTime >= PreviewScene->GetTimeRange()[1]))
 				{
 					PlaybackMode = EPlaybackMode::PlayingReverse;
 				}
 
-				if((PlaybackMode == EPlaybackMode::PlayingReverse) && (SimulationTime <= 0.0f))
+				if((PlaybackMode == EPlaybackMode::PlayingReverse) && (SimulationTime <= PreviewScene->GetTimeRange()[0]))
 				{
 					PlaybackMode = EPlaybackMode::PlayingForward;
 				}
 			}
-			PreviewScene->SimulationTime = FMath::Clamp(SimulationTime, 0.0f,  SequenceLength);
+			PreviewScene->SimulationTime = FMath::Clamp(SimulationTime, PreviewScene->GetTimeRange()[0],  PreviewScene->GetTimeRange()[1]);
 		}
 	}
 }
@@ -259,7 +259,7 @@ void SDataflowSimulationPanel::OnValueChanged(float NewValue)
 {
 	if (const TSharedPtr<FDataflowSimulationScene> PreviewScene = SimulationScene.Pin())
 	{
-		PreviewScene->SimulationTime = NewValue;
+		PreviewScene->SimulationTime = NewValue + PreviewScene->GetTimeRange()[0];
 	}
 }
 
@@ -289,7 +289,7 @@ float SDataflowSimulationPanel::GetScrubValue() const
 {
 	if (const TSharedPtr<const FDataflowSimulationScene> PreviewScene = SimulationScene.Pin())
 	{
-		return PreviewScene->SimulationTime;
+		return PreviewScene->SimulationTime-PreviewScene->GetTimeRange()[0];
 	}
 
 	return 0.0f;
@@ -299,7 +299,7 @@ bool SDataflowSimulationPanel::GetDisplayDrag() const
 {
 	if (const TSharedPtr<const FDataflowSimulationScene> PreviewScene = SimulationScene.Pin())
 	{
-			return true;
+		return true;
 	}
 	return false;
 }
diff --git a/Engine/Plugins/Experimental/Dataflow/Source/DataflowEditor/Private/Dataflow/DataflowSimulationScene.cpp b/Engine/Plugins/Experimental/Dataflow/Source/DataflowEditor/Private/Dataflow/DataflowSimulationScene.cpp
--- a/Engine/Plugins/Experimental/Dataflow/Source/DataflowEditor/Private/Dataflow/DataflowSimulationScene.cpp
+++ b/Engine/Plugins/Experimental/Dataflow/Source/DataflowEditor/Private/Dataflow/DataflowSimulationScene.cpp
@@ -2,6 +2,7 @@
 
 #include "Dataflow/DataflowSimulationScene.h"
 #include "Dataflow/DataflowSimulationManager.h"
+#include "Dataflow/DataflowSimulationControls.h"
 #include "Dataflow/DataflowEditor.h"
 #include "Components/PrimitiveComponent.h"
 #include "Chaos/CacheManagerActor.h"
@@ -23,7 +24,7 @@ FDataflowSimulationScene::FDataflowSimulationScene(FPreviewScene::ConstructionVa
 	SceneDescription = NewObject<UDataflowSimulationSceneDescription>();
 	SceneDescription->SetSimulationScene(this);
 
-	SimulationGenerator = MakeShared<Dataflow::FDataflowSimulationGenerator>();
+	SimulationGenerator = MakeShared<UE::Dataflow::FDataflowSimulationGenerator>();
 	RootSceneActor = GetWorld()->SpawnActor<AChaosCacheManager>();
 
 	if(GetEditorContent())
@@ -94,7 +95,7 @@ void FDataflowSimulationScene::PauseSimulationScene() const
 	if(SceneDescription && (SceneDescription->CacheAsset == nullptr))
 	{
 		GetWorld()->GetSubsystem<UDataflowSimulationManager>()->SetSimulationEnabled(false);
-		Dataflow::PauseSkeletonAnimation(PreviewActor);
+		UE::Dataflow::PauseSkeletonAnimation(PreviewActor);
 	}
 }
 
@@ -103,7 +104,7 @@ void FDataflowSimulationScene::StartSimulationScene() const
 	if(SceneDescription && (SceneDescription->CacheAsset == nullptr))
 	{
 		GetWorld()->GetSubsystem<UDataflowSimulationManager>()->SetSimulationEnabled(true);
-		Dataflow::StartSkeletonAnimation(PreviewActor);
+		UE::Dataflow::StartSkeletonAnimation(PreviewActor);
 	}
 }
 
@@ -113,7 +114,7 @@ void FDataflowSimulationScene::StepSimulationScene() const
 	{
 		GetWorld()->GetSubsystem<UDataflowSimulationManager>()->SetSimulationEnabled(true);
 		GetWorld()->GetSubsystem<UDataflowSimulationManager>()->SetSimulationStepping(true);
-		Dataflow::StepSkeletonAnimation(PreviewActor);
+		UE::Dataflow::StepSkeletonAnimation(PreviewActor);
 	}
 }
 
@@ -159,11 +160,11 @@ void FDataflowSimulationScene::CreateSimulationScene()
 		TimeRange = SceneDescription->CacheParams.TimeRange;
 		NumFrames = (TimeRange[1] > TimeRange[0]) ? FMath::Floor((TimeRange[1] - TimeRange[0]) * SceneDescription->CacheParams.FrameRate) : 0;
 		
-		PreviewActor = Dataflow::SpawnSimulatedActor(SceneDescription->BlueprintClass, Cast<AChaosCacheManager>(RootSceneActor),
+		PreviewActor = UE::Dataflow::SpawnSimulatedActor(SceneDescription->BlueprintClass, Cast<AChaosCacheManager>(RootSceneActor),
 			SceneDescription->CacheAsset, false, GetEditorContent());
 
 		// Setup all the skelmesh animations
-		Dataflow::SetupSkeletonAnimation(PreviewActor);
+		UE::Dataflow::SetupSkeletonAnimation(PreviewActor);
 		
 		GetWorld()->GetSubsystem<UDataflowSimulationManager>()->SetSimulationEnabled(false);
 	}
@@ -176,7 +177,7 @@ void FDataflowSimulationScene::UpdateSimulationCache()
 {
 	if(SimulationGenerator.IsValid())
 	{
-		SimulationGenerator->RequestGeneratorAction(Dataflow::EDataflowGeneratorActions::StartGenerate);
+		SimulationGenerator->RequestGeneratorAction(UE::Dataflow::EDataflowGeneratorActions::StartGenerate);
 	}
 }
 
@@ -186,7 +187,7 @@ void FDataflowSimulationScene::TickDataflowScene(const float DeltaSeconds)
 
 	if(const TObjectPtr<UDataflowBaseContent>& EditorContent = GetEditorContent())
 	{
-		if(Dataflow::ShouldResetWorld(EditorContent->GetDataflowAsset(), GetWorld(), LastTimeStamp) || EditorContent->IsSimulationDirty())
+		if(UE::Dataflow::ShouldResetWorld(EditorContent->GetDataflowAsset(), GetWorld(), LastTimeStamp) || EditorContent->IsSimulationDirty())
 		{
 			// Unregister components, cache manager, selection...
 			ResetSimulationScene();
@@ -206,7 +207,7 @@ void FDataflowSimulationScene::TickDataflowScene(const float DeltaSeconds)
 				Cast<AChaosCacheManager>(RootSceneActor)->SetStartTime(SimulationTime);
 			}
 			// Update all the skelmesh animations at the simulation time
-			Dataflow::UpdateSkeletonAnimation(PreviewActor, SimulationTime);
+			UE::Dataflow::UpdateSkeletonAnimation(PreviewActor, SimulationTime);
 		}
 	}
 }
diff --git a/Engine/Plugins/Experimental/Dataflow/Source/DataflowEditor/Private/Dataflow/DataflowSimulationViewport.cpp b/Engine/Plugins/Experimental/Dataflow/Source/DataflowEditor/Private/Dataflow/DataflowSimulationViewport.cpp
--- a/Engine/Plugins/Experimental/Dataflow/Source/DataflowEditor/Private/Dataflow/DataflowSimulationViewport.cpp
+++ b/Engine/Plugins/Experimental/Dataflow/Source/DataflowEditor/Private/Dataflow/DataflowSimulationViewport.cpp
@@ -171,12 +171,12 @@ void SDataflowSimulationViewport::OnFloatingButtonClicked()
 
 float SDataflowSimulationViewport::GetViewMinInput() const
 {
-	return GetSimulationScene()->GetTimeRange()[0];
+	return 0.0f;
 }
 
 float SDataflowSimulationViewport::GetViewMaxInput() const
 {
-	return GetSimulationScene()->GetTimeRange()[1];
+	return GetSimulationScene()->GetTimeRange()[1]-GetSimulationScene()->GetTimeRange()[0];
 }
 
 
diff --git a/Engine/Plugins/Experimental/Dataflow/Source/DataflowEditor/Public/Dataflow/DataflowSimulationControls.h b/Engine/Plugins/Experimental/Dataflow/Source/DataflowEditor/Public/Dataflow/DataflowSimulationControls.h
new file mode 100644
--- /dev/null
+++ b/Engine/Plugins/Experimental/Dataflow/Source/DataflowEditor/Public/Dataflow/DataflowSimulationControls.h
@@ -0,0 +1,43 @@
+﻿// Copyright Epic Games, Inc. All Rights Reserved.
+
+#pragma once
+
+#include "Dataflow/DataflowObjectInterface.h"
+#include "Dataflow/DataflowSimulationContext.h"
+
+class UDataflow;
+class UChaosCacheCollection;
+class AChaosCacheManager;
+class UDataflowBaseContent;
+
+namespace UE::Dataflow
+{
+	/** Check if the simulation cache nodes have changed to trigger a reset */
+	bool ShouldResetWorld(const TObjectPtr<UDataflow>& SimulationGraph,
+		const TObjectPtr<UWorld>& SimulationWorld, ::Dataflow::FTimestamp& LastTimeStamp);
+
+	/** Spawn an actor given a class type and attach it to the cache manager */
+	TObjectPtr<AActor> SpawnSimulatedActor(const TSubclassOf<AActor>& ActorClass,
+		const TObjectPtr<AChaosCacheManager>& CacheManager, const TObjectPtr<UChaosCacheCollection>& CacheCollection,
+		const bool bIsRecording, const TObjectPtr<UDataflowBaseContent>& DataflowContent);
+
+	/** Setup the skelmesh animations to be used in the scene/generator */
+	void SetupSkeletonAnimation(const TObjectPtr<AActor>& PreviewActor);
+
+	/** Update the skelmesh animation at some point in time (GT) */
+	void UpdateSkeletonAnimation(const TObjectPtr<AActor>& PreviewActor, const float SimulationTime);
+
+	/** Compute the skelmesh animation at some point in time (PT) */
+	void ComputeSkeletonAnimation(const TObjectPtr<AActor>& PreviewActor, const float SimulationTime);
+
+	/** Start the skelmesh animation */
+	void StartSkeletonAnimation(const TObjectPtr<AActor>& PreviewActor);
+
+	/** Pause the skelmesh animation */
+	void PauseSkeletonAnimation(const TObjectPtr<AActor>& PreviewActor);
+	
+	/** Step the skelmesh animation */
+	void StepSkeletonAnimation(const TObjectPtr<AActor>& PreviewActor);
+};
+
+
diff --git a/Engine/Plugins/Experimental/Dataflow/Source/DataflowEditor/Public/Dataflow/DataflowSimulationGenerator.h b/Engine/Plugins/Experimental/Dataflow/Source/DataflowEditor/Public/Dataflow/DataflowSimulationGenerator.h
--- a/Engine/Plugins/Experimental/Dataflow/Source/DataflowEditor/Public/Dataflow/DataflowSimulationGenerator.h
+++ b/Engine/Plugins/Experimental/Dataflow/Source/DataflowEditor/Public/Dataflow/DataflowSimulationGenerator.h
@@ -9,14 +9,13 @@
 #include "Misc/AsyncTaskNotification.h"
 #include "Dataflow/DataflowContent.h"
 #include "Dataflow/DataflowPreview.h"
-#include "Dataflow/DataflowSimulationUtils.h"
 
 DECLARE_LOG_CATEGORY_EXTERN(LogDataflowSimulationGenerator, Log, All);
 
 class UChaosCacheCollection;
 class AChaosCacheManager;
 
-namespace Dataflow
+namespace UE::Dataflow
 {
 	/** Simulation Task to be run on the async thread */
 	class FDataflowSimulationTask : public FNonAbandonableTask
diff --git a/Engine/Plugins/Experimental/Dataflow/Source/DataflowEditor/Public/Dataflow/DataflowSimulationScene.h b/Engine/Plugins/Experimental/Dataflow/Source/DataflowEditor/Public/Dataflow/DataflowSimulationScene.h
--- a/Engine/Plugins/Experimental/Dataflow/Source/DataflowEditor/Public/Dataflow/DataflowSimulationScene.h
+++ b/Engine/Plugins/Experimental/Dataflow/Source/DataflowEditor/Public/Dataflow/DataflowSimulationScene.h
@@ -121,7 +121,7 @@ private:
 	TObjectPtr<UDataflowSimulationSceneDescription> SceneDescription;
 
 	/** Simulation generator to record the simulation result */
-	TSharedPtr<Dataflow::FDataflowSimulationGenerator> SimulationGenerator;
+	TSharedPtr<UE::Dataflow::FDataflowSimulationGenerator> SimulationGenerator;
 
 	/** Cache time range in seconds */
 	FVector2f TimeRange;
diff --git a/Engine/Plugins/Experimental/Dataflow/Source/DataflowEnginePlugin/Private/Dataflow/DataflowEnginePlugin.cpp b/Engine/Plugins/Experimental/Dataflow/Source/DataflowEnginePlugin/Private/Dataflow/DataflowEnginePlugin.cpp
--- a/Engine/Plugins/Experimental/Dataflow/Source/DataflowEnginePlugin/Private/Dataflow/DataflowEnginePlugin.cpp
+++ b/Engine/Plugins/Experimental/Dataflow/Source/DataflowEnginePlugin/Private/Dataflow/DataflowEnginePlugin.cpp
@@ -29,6 +29,7 @@ IMPLEMENT_MODULE( FDataflowEnginePlugin, DataflowEnginePlugin)
 void FDataflowEnginePlugin::StartupModule()
 {
 	Dataflow::ContextCachingCallbacks();
+	FModuleManager::Get().LoadModule("DataflowSimulation");
 }
 
 
diff --git a/Engine/Plugins/Experimental/Dataflow/Source/DataflowSimulation/DataflowSimulation.Build.cs b/Engine/Plugins/Experimental/Dataflow/Source/DataflowSimulation/DataflowSimulation.Build.cs
deleted file mode 100644
diff --git a/Engine/Plugins/Experimental/Dataflow/Source/DataflowSimulation/Private/Dataflow/DataflowSimulationContext.cpp b/Engine/Plugins/Experimental/Dataflow/Source/DataflowSimulation/Private/Dataflow/DataflowSimulationContext.cpp
deleted file mode 100644
diff --git a/Engine/Plugins/Experimental/Dataflow/Source/DataflowSimulation/Private/Dataflow/DataflowSimulationInterface.cpp b/Engine/Plugins/Experimental/Dataflow/Source/DataflowSimulation/Private/Dataflow/DataflowSimulationInterface.cpp
deleted file mode 100644
diff --git a/Engine/Plugins/Experimental/Dataflow/Source/DataflowSimulation/Private/Dataflow/DataflowSimulationManager.cpp b/Engine/Plugins/Experimental/Dataflow/Source/DataflowSimulation/Private/Dataflow/DataflowSimulationManager.cpp
deleted file mode 100644
diff --git a/Engine/Plugins/Experimental/Dataflow/Source/DataflowSimulation/Private/Dataflow/DataflowSimulationNodes.cpp b/Engine/Plugins/Experimental/Dataflow/Source/DataflowSimulation/Private/Dataflow/DataflowSimulationNodes.cpp
deleted file mode 100644
diff --git a/Engine/Plugins/Experimental/Dataflow/Source/DataflowSimulation/Private/Dataflow/DataflowSimulationPlugin.cpp b/Engine/Plugins/Experimental/Dataflow/Source/DataflowSimulation/Private/Dataflow/DataflowSimulationPlugin.cpp
deleted file mode 100644
diff --git a/Engine/Plugins/Experimental/Dataflow/Source/DataflowSimulation/Private/Dataflow/DataflowSimulationProxy.cpp b/Engine/Plugins/Experimental/Dataflow/Source/DataflowSimulation/Private/Dataflow/DataflowSimulationProxy.cpp
deleted file mode 100644
diff --git a/Engine/Plugins/Experimental/Dataflow/Source/DataflowSimulation/Private/Dataflow/DataflowSimulationUtils.cpp b/Engine/Plugins/Experimental/Dataflow/Source/DataflowSimulation/Private/Dataflow/DataflowSimulationUtils.cpp
deleted file mode 100644
diff --git a/Engine/Plugins/Experimental/Dataflow/Source/DataflowSimulation/Private/Interfaces/DataflowCollisionObject.cpp b/Engine/Plugins/Experimental/Dataflow/Source/DataflowSimulation/Private/Interfaces/DataflowCollisionObject.cpp
deleted file mode 100644
diff --git a/Engine/Plugins/Experimental/Dataflow/Source/DataflowSimulation/Private/Interfaces/DataflowConstraintObject.cpp b/Engine/Plugins/Experimental/Dataflow/Source/DataflowSimulation/Private/Interfaces/DataflowConstraintObject.cpp
deleted file mode 100644
diff --git a/Engine/Plugins/Experimental/Dataflow/Source/DataflowSimulation/Private/Interfaces/DataflowPhysicsObject.cpp b/Engine/Plugins/Experimental/Dataflow/Source/DataflowSimulation/Private/Interfaces/DataflowPhysicsObject.cpp
deleted file mode 100644
diff --git a/Engine/Plugins/Experimental/Dataflow/Source/DataflowSimulation/Private/Interfaces/DataflowPhysicsSolver.cpp b/Engine/Plugins/Experimental/Dataflow/Source/DataflowSimulation/Private/Interfaces/DataflowPhysicsSolver.cpp
deleted file mode 100644
diff --git a/Engine/Plugins/Experimental/Dataflow/Source/DataflowSimulation/Public/Dataflow/DataflowSimulationContext.h b/Engine/Plugins/Experimental/Dataflow/Source/DataflowSimulation/Public/Dataflow/DataflowSimulationContext.h
deleted file mode 100644
diff --git a/Engine/Plugins/Experimental/Dataflow/Source/DataflowSimulation/Public/Dataflow/DataflowSimulationInterface.h b/Engine/Plugins/Experimental/Dataflow/Source/DataflowSimulation/Public/Dataflow/DataflowSimulationInterface.h
deleted file mode 100644
diff --git a/Engine/Plugins/Experimental/Dataflow/Source/DataflowSimulation/Public/Dataflow/DataflowSimulationManager.h b/Engine/Plugins/Experimental/Dataflow/Source/DataflowSimulation/Public/Dataflow/DataflowSimulationManager.h
deleted file mode 100644
diff --git a/Engine/Plugins/Experimental/Dataflow/Source/DataflowSimulation/Public/Dataflow/DataflowSimulationNodes.h b/Engine/Plugins/Experimental/Dataflow/Source/DataflowSimulation/Public/Dataflow/DataflowSimulationNodes.h
deleted file mode 100644
diff --git a/Engine/Plugins/Experimental/Dataflow/Source/DataflowSimulation/Public/Dataflow/DataflowSimulationPlugin.h b/Engine/Plugins/Experimental/Dataflow/Source/DataflowSimulation/Public/Dataflow/DataflowSimulationPlugin.h
deleted file mode 100644
diff --git a/Engine/Plugins/Experimental/Dataflow/Source/DataflowSimulation/Public/Dataflow/DataflowSimulationProxy.h b/Engine/Plugins/Experimental/Dataflow/Source/DataflowSimulation/Public/Dataflow/DataflowSimulationProxy.h
deleted file mode 100644
diff --git a/Engine/Plugins/Experimental/Dataflow/Source/DataflowSimulation/Public/Dataflow/DataflowSimulationUtils.h b/Engine/Plugins/Experimental/Dataflow/Source/DataflowSimulation/Public/Dataflow/DataflowSimulationUtils.h
deleted file mode 100644
diff --git a/Engine/Plugins/Experimental/Dataflow/Source/DataflowSimulation/Public/Interfaces/DataflowCollisionObject.h b/Engine/Plugins/Experimental/Dataflow/Source/DataflowSimulation/Public/Interfaces/DataflowCollisionObject.h
deleted file mode 100644
diff --git a/Engine/Plugins/Experimental/Dataflow/Source/DataflowSimulation/Public/Interfaces/DataflowConstraintObject.h b/Engine/Plugins/Experimental/Dataflow/Source/DataflowSimulation/Public/Interfaces/DataflowConstraintObject.h
deleted file mode 100644
diff --git a/Engine/Plugins/Experimental/Dataflow/Source/DataflowSimulation/Public/Interfaces/DataflowPhysicsObject.h b/Engine/Plugins/Experimental/Dataflow/Source/DataflowSimulation/Public/Interfaces/DataflowPhysicsObject.h
deleted file mode 100644
diff --git a/Engine/Plugins/Experimental/Dataflow/Source/DataflowSimulation/Public/Interfaces/DataflowPhysicsSolver.h b/Engine/Plugins/Experimental/Dataflow/Source/DataflowSimulation/Public/Interfaces/DataflowPhysicsSolver.h
deleted file mode 100644
diff --git a/Engine/Source/Runtime/Experimental/Dataflow/Engine/Private/Dataflow/DataflowContent.cpp b/Engine/Source/Runtime/Experimental/Dataflow/Engine/Private/Dataflow/DataflowContent.cpp
--- a/Engine/Source/Runtime/Experimental/Dataflow/Engine/Private/Dataflow/DataflowContent.cpp
+++ b/Engine/Source/Runtime/Experimental/Dataflow/Engine/Private/Dataflow/DataflowContent.cpp
@@ -391,7 +391,7 @@ void UDataflowBaseContent::AddReferencedObjects(UObject* InThis, FReferenceColle
 	Super::AddReferencedObjects(InThis, Collector);
 }
 
-void OverrideActorProperty(const TObjectPtr<AActor>& PreviewActor, TObjectPtr<UObject> PropertyValue, const FName& PropertyName)
+void UDataflowBaseContent::OverrideActorProperty(const TObjectPtr<AActor>& PreviewActor, TObjectPtr<UObject> PropertyValue, const FName& PropertyName)
 {
 	if(PreviewActor && PropertyValue)
 	{
diff --git a/Engine/Source/Runtime/Experimental/Dataflow/Engine/Private/Dataflow/DataflowObject.cpp b/Engine/Source/Runtime/Experimental/Dataflow/Engine/Private/Dataflow/DataflowObject.cpp
--- a/Engine/Source/Runtime/Experimental/Dataflow/Engine/Private/Dataflow/DataflowObject.cpp
+++ b/Engine/Source/Runtime/Experimental/Dataflow/Engine/Private/Dataflow/DataflowObject.cpp
@@ -12,15 +12,14 @@
 
 #define LOCTEXT_NAMESPACE "UDataflow"
 
-
-namespace Dataflow
+namespace UE::Dataflow::CVars
 {
-	namespace CVars
-	{
-		/** Enable the simulation dataflow (for now WIP) */
-        DATAFLOWENGINE_API bool bEnableSimulationDataflow = false;
-        FAutoConsoleVariableRef CVarEnableSimulationDataflow(TEXT("p.Dataflow.EnableSimulation"), bEnableSimulationDataflow, TEXT("If true enable the use of simulation dataflow (WIP)"));
-	}
+	/** Enable the simulation dataflow (for now WIP) */
+	TAutoConsoleVariable<bool> CVarEnableSimulationDataflow(
+			TEXT("p.Dataflow.EnableSimulation"),
+			false,
+			TEXT("If true enable the use of simulation dataflow (WIP)"),
+			ECVF_Default);
 }
 
 FDataflowAssetEdit::FDataflowAssetEdit(UDataflow* InAsset, FPostEditFunctionCallback InCallback)
@@ -224,7 +223,7 @@ bool UDataflow::CanEditChange(const FProperty* InProperty) const
 
 	if (Name == GET_MEMBER_NAME_CHECKED(ThisClass, Type))
 	{
-		return Dataflow::CVars::bEnableSimulationDataflow == true;
+		return UE::Dataflow::CVars::CVarEnableSimulationDataflow.GetValueOnGameThread();
 	}
 
 	return true;
diff --git a/Engine/Source/Runtime/Experimental/Dataflow/Engine/Public/Dataflow/DataflowContent.h b/Engine/Source/Runtime/Experimental/Dataflow/Engine/Public/Dataflow/DataflowContent.h
--- a/Engine/Source/Runtime/Experimental/Dataflow/Engine/Public/Dataflow/DataflowContent.h
+++ b/Engine/Source/Runtime/Experimental/Dataflow/Engine/Public/Dataflow/DataflowContent.h
@@ -5,8 +5,10 @@
 #include "DataflowContextObject.h"
 #include "UObject/Interface.h"
 #include "Templates/SharedPointer.h"
+#include "GameFramework/Actor.h"
 #include "DataflowContent.generated.h"
 
+
 class FDataflowEditorToolkit;
 class USkeletalMesh;
 class USkeleton;
@@ -15,7 +17,6 @@ class UAnimationAsset;
 class UDataflowBaseContent;
 class FPreviewScene;
 class UAnimSingleNodeInstance;
-class AActor;
 
 namespace DataflowContextHelpers
 {
@@ -169,8 +170,36 @@ protected:
 
 	/** Preview actor class that could be used to visualize the result */
 	TSubclassOf<AActor> PreviewClass = nullptr;
+
+	/** override actor properties from BP */
+	static void OverrideActorProperty(const TObjectPtr<AActor>& PreviewActor, TObjectPtr<UObject> PropertyValue, const FName& PropertyName);
+
+	/** override struct properties from BP */
+	template<typename StructType>
+	static void OverrideStructProperty(const TObjectPtr<AActor>& PreviewActor, const StructType& PropertyValue, const FName& PropertyName);
 };
 
+template<typename StructType>
+void UDataflowBaseContent::OverrideStructProperty(const TObjectPtr<AActor>& PreviewActor, const StructType& PropertyValue, const FName& PropertyName)
+{
+	if(PreviewActor)
+	{
+		if(const FProperty* DataflowProperty = PreviewActor->GetClass()->FindPropertyByName(PropertyName))
+		{
+			if (const FStructProperty* StructProperty = CastField<FStructProperty>(DataflowProperty))
+			{
+				if(StructProperty->Struct == StructType::StaticStruct())
+				{
+					if(StructType* PropertyStruct = DataflowProperty->ContainerPtrToValuePtr<StructType>(PreviewActor))
+					{
+						(*PropertyStruct) = PropertyValue;
+					}
+				}
+			}
+		}
+	}
+}
+
 /** 
  * Dataflow content owning dataflow and skelmesh assets that that will be used to evaluate the graph
  */
diff --git a/Engine/Source/Runtime/Experimental/Dataflow/Engine/Public/Dataflow/DataflowObject.h b/Engine/Source/Runtime/Experimental/Dataflow/Engine/Public/Dataflow/DataflowObject.h
--- a/Engine/Source/Runtime/Experimental/Dataflow/Engine/Public/Dataflow/DataflowObject.h
+++ b/Engine/Source/Runtime/Experimental/Dataflow/Engine/Public/Dataflow/DataflowObject.h
@@ -101,7 +101,7 @@ public:
 	UPROPERTY(EditAnywhere, Category = "Render")
 	TObjectPtr<UMaterial> Material = nullptr;
 
-	UPROPERTY(EditAnywhere, Category = "Evaluation")
+	UPROPERTY(EditAnywhere, Category = "Evaluation",meta=(EditConditionHides))
 	EDataflowType Type = EDataflowType::Construction;
 
 public:
diff --git a/Engine/Source/Runtime/Experimental/Dataflow/Engine/Public/Dataflow/DataflowPreview.h b/Engine/Source/Runtime/Experimental/Dataflow/Engine/Public/Dataflow/DataflowPreview.h
--- a/Engine/Source/Runtime/Experimental/Dataflow/Engine/Public/Dataflow/DataflowPreview.h
+++ b/Engine/Source/Runtime/Experimental/Dataflow/Engine/Public/Dataflow/DataflowPreview.h
@@ -27,6 +27,6 @@ public :
 	FVector2f TimeRange = FVector2f(0.0f, 5.0f);
 
 	/** Boolean to check if the caching will be done on an async thread (if yes no GT dependency) */
-	UPROPERTY(EditAnywhere, Category = "Caching")
+	UPROPERTY(EditAnywhere, AdvancedDisplay, Category = "Caching")
 	bool bBackgroundTask = true;
 };
diff --git a/Engine/Source/Runtime/Experimental/Dataflow/Simulation/DataflowSimulation.Build.cs b/Engine/Source/Runtime/Experimental/Dataflow/Simulation/DataflowSimulation.Build.cs
new file mode 100644
--- /dev/null
+++ b/Engine/Source/Runtime/Experimental/Dataflow/Simulation/DataflowSimulation.Build.cs
@@ -0,0 +1,35 @@
+﻿// Copyright Epic Games, Inc. All Rights Reserved.
+
+namespace UnrealBuildTool.Rules
+{
+	public class DataflowSimulation : ModuleRules
+	{
+		public DataflowSimulation(ReadOnlyTargetRules Target) : base(Target)
+		{
+			bTreatAsEngineModule = true;
+			PublicDependencyModuleNames.AddRange(
+				new string[]
+				{
+					"Core",
+					"CoreUObject",
+					"DataflowCore",
+					"DataflowEngine",
+				}
+			);
+			PrivateDependencyModuleNames.AddRange(
+				new string[]
+				{
+					// NOTE: UVEditorTools is a separate module so that it doesn't rely on the editor.
+					// So, do not add editor dependencies here.
+	
+					"Engine",
+					"RenderCore",
+					"RHI",
+					"DataflowCore",
+					"DataflowEngine",
+					"Chaos"
+				}
+			);
+		}
+	}
+}
\ No newline at end of file
diff --git a/Engine/Source/Runtime/Experimental/Dataflow/Simulation/Private/Dataflow/DataflowSimulationContext.cpp b/Engine/Source/Runtime/Experimental/Dataflow/Simulation/Private/Dataflow/DataflowSimulationContext.cpp
new file mode 100644
--- /dev/null
+++ b/Engine/Source/Runtime/Experimental/Dataflow/Simulation/Private/Dataflow/DataflowSimulationContext.cpp
@@ -0,0 +1,97 @@
+﻿// Copyright Epic Games, Inc. All Rights Reserved.
+
+#include "Dataflow/DataflowSimulationContext.h"
+#include "Dataflow/DataflowSimulationProxy.h"
+
+namespace UE::Dataflow
+{
+	
+template<typename Base>
+void TSimulationContext<Base>::GetSimulationProxies(const FString& ProxyType, const TArray<FString>& SimulationGroups, TArray<FDataflowSimulationProxy*>& FilteredProxies) const
+{
+	TBitArray<> ProxyGroups(false, GroupIndices.Num());
+	BuildGroupBits(SimulationGroups, ProxyGroups);
+	
+	if(const TSet<FDataflowSimulationProxy*>* TypedProxies = SimulationProxies.Find(ProxyType))
+	{
+		for(FDataflowSimulationProxy* SimulationProxy : *TypedProxies)
+		{
+			if(SimulationProxy && SimulationProxy->HasGroupBit(ProxyGroups))
+			{
+				FilteredProxies.Add(SimulationProxy);
+			}
+		}
+	}
+}
+
+template<typename Base>
+void TSimulationContext<Base>::BuildGroupBits(const TArray<FString>& SimulationGroups, TBitArray<>& GroupBits) const
+{
+	GroupBits.Init(false, GroupIndices.Num());
+    for(const FString& SimulationGroup : SimulationGroups)
+    {
+    	if(const uint32* GroupIndex = GroupIndices.Find(SimulationGroup))
+    	{
+    		GroupBits[*GroupIndex] = true;
+    	}
+    }
+}
+	
+template<typename Base>
+void TSimulationContext<Base>::RegisterProxyGroups()
+{
+	GroupIndices.Reset();
+	TArray<uint32> IndexArray;
+	
+	for(const TPair<FString, TSet<FDataflowSimulationProxy*>>& TypedProxies : SimulationProxies)
+	{
+		for(FDataflowSimulationProxy* SimulationProxy : TypedProxies.Value)
+		{
+			if(SimulationProxy)
+			{
+				IndexArray.Reset(SimulationProxy->GetSimulationGroups().Num());
+				for(const FString& SimulationGroup : SimulationProxy->GetSimulationGroups())
+				{
+					const uint32* IndexPtr = GroupIndices.Find(SimulationGroup);
+					IndexArray.Add( (IndexPtr == nullptr) ?
+							GroupIndices.Add(SimulationGroup, GroupIndices.Num()) : *IndexPtr);
+				}
+				SimulationProxy->GroupBits.Init(false, GroupIndices.Num());
+				for(const uint32& GroupIndex : IndexArray)
+				{
+					SimulationProxy->GroupBits[GroupIndex] = true;
+				}
+			}
+		}
+	}
+}
+
+template<typename Base>
+void TSimulationContext<Base>::AddSimulationProxy(const FString& ProxyType, FDataflowSimulationProxy* SimulationProxy)
+{
+	SimulationProxies.FindOrAdd(ProxyType).Add(SimulationProxy);
+}
+	
+template<typename Base>
+void TSimulationContext<Base>::RemoveSimulationProxy(const FString& ProxyType, const FDataflowSimulationProxy* SimulationProxy)
+{
+	SimulationProxies.FindOrAdd(ProxyType).Remove(SimulationProxy);
+}
+
+template<typename Base>
+void TSimulationContext<Base>::ResetSimulationProxies()
+{
+	SimulationProxies.Reset();
+}
+	
+template<typename Base>
+int32 TSimulationContext<Base>::NumSimulationProxies(const FString& ProxyType) const
+{
+	const TSet<FDataflowSimulationProxy*>* TypedProxies = SimulationProxies.Find(ProxyType);
+	return TypedProxies ? TypedProxies->Num() : 0;
+}
+
+template class TSimulationContext<::Dataflow::FContextSingle>;
+template class TSimulationContext<::Dataflow::FContextThreaded>;
+	
+}
\ No newline at end of file
diff --git a/Engine/Source/Runtime/Experimental/Dataflow/Simulation/Private/Dataflow/DataflowSimulationInterface.cpp b/Engine/Source/Runtime/Experimental/Dataflow/Simulation/Private/Dataflow/DataflowSimulationInterface.cpp
new file mode 100644
--- /dev/null
+++ b/Engine/Source/Runtime/Experimental/Dataflow/Simulation/Private/Dataflow/DataflowSimulationInterface.cpp
@@ -0,0 +1,39 @@
+﻿// Copyright Epic Games, Inc. All Rights Reserved.
+
+#include "Dataflow/DataflowSimulationInterface.h"
+#include "Dataflow/DataflowSimulationManager.h"
+#include "Engine/World.h"
+	
+void IDataflowSimulationInterface::RegisterManagerInterface(const TObjectPtr<UWorld>& SimulationWorld)
+{
+	if(SimulationWorld)
+	{
+		if (UDataflowSimulationManager* SimulationManager = SimulationWorld->GetSubsystem<UDataflowSimulationManager>())
+		{
+			SimulationManager->AddSimulationInterface(this);
+		}
+	}
+}
+
+bool IDataflowSimulationInterface::IsInterfaceRegistered(const TObjectPtr<UWorld>& SimulationWorld) const
+{
+	if(SimulationWorld)
+	{
+		if (const UDataflowSimulationManager* SimulationManager = SimulationWorld->GetSubsystem<UDataflowSimulationManager>())
+		{
+			return SimulationManager->HasSimulationInterface(this);
+		}
+	}
+	return false;
+}
+
+void IDataflowSimulationInterface::UnregisterManagerInterface(const TObjectPtr<UWorld>& SimulationWorld) const
+{
+	if(SimulationWorld)
+	{
+		if (UDataflowSimulationManager* SimulationManager = SimulationWorld->GetSubsystem<UDataflowSimulationManager>())
+		{
+			SimulationManager->RemoveSimulationInterface(this); 
+		}
+	}
+}
diff --git a/Engine/Source/Runtime/Experimental/Dataflow/Simulation/Private/Dataflow/DataflowSimulationManager.cpp b/Engine/Source/Runtime/Experimental/Dataflow/Simulation/Private/Dataflow/DataflowSimulationManager.cpp
new file mode 100644
--- /dev/null
+++ b/Engine/Source/Runtime/Experimental/Dataflow/Simulation/Private/Dataflow/DataflowSimulationManager.cpp
@@ -0,0 +1,560 @@
+// Copyright Epic Games, Inc. All Rights Reserved.
+
+#include "Dataflow/DataflowSimulationManager.h"
+#include "Dataflow/DataflowSimulationUtils.h"
+#include "Dataflow/DataflowSimulationInterface.h"
+#include "Dataflow/DataflowSimulationProxy.h"
+#include "Components/ActorComponent.h"
+#include "Async/TaskGraphInterfaces.h"
+#include "Containers/Map.h"
+#include "Dataflow/Interfaces/DataflowPhysicsSolver.h"
+#include "Engine/World.h"
+#include "Kismet/GameplayStatics.h"
+
+namespace UE::Dataflow
+{
+	namespace Private
+	{
+		enum class ESimulationThreadingMode : uint8
+		{
+			GameThread,
+			BlockingThread,
+			AsyncThread,
+		};
+
+		/** Simulation task priority */
+		static FAutoConsoleTaskPriority CVarDataflowSimulationTaskPriority(
+			TEXT("TaskGraph.TaskPriorities.DataflowSimulationTask"),
+			TEXT("Task and thread priority for the dataflow simulation."),
+			ENamedThreads::HighThreadPriority, // If we have high priority task threads, then use them...
+			ENamedThreads::NormalTaskPriority, // .. at normal task priority
+			ENamedThreads::HighTaskPriority);  // If we don't have high priority threads, then use normal priority threads at high task priority instead
+		
+		/** Simulation threading mode */
+		int32  DataflowSimulationThreadingMode = (uint8)ESimulationThreadingMode::AsyncThread;
+		FAutoConsoleVariableRef CVarDataflowSimulationThreadingMode(TEXT("p.Dataflow.Simulation.ThreadingMode"), DataflowSimulationThreadingMode,
+			TEXT("0 : run simulation on GT | 1 : run simulation on PT (GT is blocked in manager Tick) | 2 : run simulation on PT (GT will be blocked at the end of the world tick)"));
+
+		/** Simulation task use to run async the dataflow evaluation */
+		class FDataflowSimulationTask
+		{
+		public:
+			FDataflowSimulationTask(const TObjectPtr<UDataflow>& InDataflowAsset, const TSharedPtr<FDataflowSimulationContext>& InSimulationContext,
+				const float InDeltaTime, const float InSimulationTime)
+				: DataflowAsset(InDataflowAsset), SimulationContext(InSimulationContext), DeltaTime(InDeltaTime), SimulationTime(InSimulationTime)
+			{
+			}
+
+			TStatId GetStatId() const
+			{
+				RETURN_QUICK_DECLARE_CYCLE_STAT(FDataflowSimulationProxyParallelTask, STATGROUP_TaskGraphTasks);
+			}
+
+			static ENamedThreads::Type GetDesiredThread()
+			{
+				if (CVarDataflowSimulationTaskPriority.Get() != 0)
+				{
+					return CVarDataflowSimulationTaskPriority.Get();
+				}
+				return ENamedThreads::GameThread;
+			}
+
+			static ESubsequentsMode::Type GetSubsequentsMode()
+			{
+				return ESubsequentsMode::TrackSubsequents;
+			}
+
+			void DoTask(ENamedThreads::Type CurrentThread, const FGraphEventRef& MyCompletionGraphEvent)
+			{
+				Dataflow::EvaluateSimulationGraph(DataflowAsset, SimulationContext, DeltaTime, SimulationTime);
+			}
+
+		private:
+			/** Dataflow graph asset used to launch the simulation */
+			TObjectPtr<UDataflow> DataflowAsset;
+
+			/** Simulation context */
+			TSharedPtr<Dataflow::FDataflowSimulationContext> SimulationContext;
+
+			/** Delta time used to advance the simulation */
+			float DeltaTime;
+
+			/** World simulation time  */
+			float SimulationTime;
+		};
+
+		inline void PreSimulationTick(const TObjectPtr<UObject>& SimulationWorld, const float SimulationTime, const float DeltaTime)
+		{
+			if(SimulationWorld)
+			{
+				TArray<AActor*> Actors;
+				UGameplayStatics::GetAllActorsWithInterface(SimulationWorld, UDataflowSimulationActor::StaticClass(), Actors);
+
+				for (AActor* CurrentActor : Actors)
+				{
+					IDataflowSimulationActor::Execute_PreDataflowSimulationTick(CurrentActor, SimulationTime, DeltaTime);
+				}
+			}
+		}
+
+		inline void PostSimulationTick(const TObjectPtr<UObject>& SimulationWorld, const float SimulationTime, const float DeltaTime)
+		{
+			if(SimulationWorld)
+			{
+				TArray<AActor*> Actors;
+				UGameplayStatics::GetAllActorsWithInterface(SimulationWorld, UDataflowSimulationActor::StaticClass(), Actors);
+
+				for (AActor* CurrentActor : Actors)
+				{
+					IDataflowSimulationActor::Execute_PostDataflowSimulationTick(CurrentActor, SimulationTime, DeltaTime);
+				}
+			}
+		}
+	}
+
+	void RegisterSimulationInterface(const TObjectPtr<UObject>& SimulationObject)
+	{
+		if(IDataflowSimulationInterface* SimulationInterface = Cast<IDataflowSimulationInterface>(SimulationObject))
+		{
+			if(SimulationInterface->GetSimulationAsset().DataflowAsset)
+			{
+				FDataflowSimulationProxy* SimulationProxy = SimulationInterface->GetSimulationProxy();
+				
+				if(!SimulationProxy || (SimulationProxy && !SimulationProxy->IsValid()))
+				{
+					// Build the simulation proxy
+					SimulationInterface->BuildSimulationProxy();
+				}
+
+				// Register the simulation interface to the manager
+				SimulationInterface->RegisterManagerInterface(SimulationObject->GetWorld());
+			}
+		}
+	}
+
+	void UnregisterSimulationInterface(const TObjectPtr<UObject>& SimulationObject)
+	{
+		if(IDataflowSimulationInterface* SimulationInterface = Cast<IDataflowSimulationInterface>(SimulationObject))
+		{
+			if(SimulationInterface->GetSimulationAsset().DataflowAsset)
+			{
+				FDataflowSimulationProxy* SimulationProxy = SimulationInterface->GetSimulationProxy();
+				
+				if(SimulationProxy && SimulationProxy->IsValid())
+				{
+					// Reset the simulation proxy
+					SimulationInterface->ResetSimulationProxy();
+				}
+
+				// Unregister the simulation interface from the manager
+				SimulationInterface->UnregisterManagerInterface(SimulationObject->GetWorld());
+			}
+		}
+	}
+}
+
+FDelegateHandle UDataflowSimulationManager::OnObjectPropertyChangedHandle;
+FDelegateHandle UDataflowSimulationManager::OnWorldTickEndHandle;
+FDelegateHandle UDataflowSimulationManager::OnCreatePhysicsStateHandle;
+FDelegateHandle UDataflowSimulationManager::OnDestroyPhysicsStateHandle;
+
+UDataflowSimulationManager::UDataflowSimulationManager()
+{}
+
+void UDataflowSimulationManager::Tick(float DeltaTime)
+{
+	Super::Tick(DeltaTime);
+	
+	check(IsInGameThread());
+	check(SimulationTasks.IsEmpty());
+	
+	PreProcessSimulation(DeltaTime);
+	
+	if(bIsSimulationEnabled)
+	{
+		// Transfer data from GT -> PT
+		ReadSimulationInterfaces(DeltaTime);
+		
+		if(UE::Dataflow::Private::DataflowSimulationThreadingMode == static_cast<uint8>(UE::Dataflow::Private::ESimulationThreadingMode::GameThread))
+		{
+			// Advance the simulation in time
+			AdvanceSimulationProxies(DeltaTime, GetWorld()->GetTimeSeconds());
+
+			// Transfer data from PT -> GT
+			WriteSimulationInterfaces(DeltaTime);
+		}
+		else
+		{
+			// Start all the simulation tasks in parallel
+			StartSimulationTasks(DeltaTime, GetWorld()->GetTimeSeconds());
+
+			if(UE::Dataflow::Private::DataflowSimulationThreadingMode == static_cast<uint8>(UE::Dataflow::Private::ESimulationThreadingMode::BlockingThread))
+            {
+            	// Wait until all tasks are complete
+            	CompleteSimulationTasks();
+            
+            	// Transfer data from PT -> GT
+            	WriteSimulationInterfaces(DeltaTime);
+            }
+		}
+	}
+	PostProcessSimulation(DeltaTime);
+}
+
+void UDataflowSimulationManager::OnStartup()
+{
+	OnWorldTickEndHandle = FWorldDelegates::OnWorldTickEnd.AddLambda(
+	[](const UWorld* SimulationWorld, ELevelTick LevelTick, const float DeltaSeconds)
+	{
+		if(SimulationWorld)
+		{
+			if (UDataflowSimulationManager* DataflowManager = SimulationWorld->GetSubsystem<UDataflowSimulationManager>())
+			{
+				if(DataflowManager->bIsSimulationEnabled)
+				{
+					if(UE::Dataflow::Private::DataflowSimulationThreadingMode == static_cast<uint8>(UE::Dataflow::Private::ESimulationThreadingMode::AsyncThread))
+					{
+						// Wait until all tasks are complete
+						DataflowManager->CompleteSimulationTasks();
+					
+						// Transfer data from PT -> GT
+						DataflowManager->WriteSimulationInterfaces(DeltaSeconds);
+					}
+				}
+			}
+		}
+		
+	});
+
+	OnCreatePhysicsStateHandle = UActorComponent::GlobalCreatePhysicsDelegate.AddLambda([](UActorComponent* ActorComponent)
+	{
+		UE::Dataflow::RegisterSimulationInterface(ActorComponent);
+	});
+
+	OnDestroyPhysicsStateHandle = UActorComponent::GlobalDestroyPhysicsDelegate.AddLambda([](UActorComponent* ActorComponent)
+	{
+		UE::Dataflow::UnregisterSimulationInterface(ActorComponent);
+	});
+
+#if WITH_EDITOR
+	OnObjectPropertyChangedHandle = FCoreUObjectDelegates::OnObjectPropertyChanged.AddLambda([](UObject* ModifiedObject, FPropertyChangedEvent& ChangedProperty)
+	{
+		if(IDataflowSimulationInterface* SimulationInterface = Cast<IDataflowSimulationInterface>(ModifiedObject))
+		{
+			if(!SimulationInterface->IsInterfaceRegistered(ModifiedObject->GetWorld()))
+			{
+				// Unregister the simulation interface from the manager
+				SimulationInterface->UnregisterManagerInterface(ModifiedObject->GetWorld());
+                
+                // Register the simulation interface to the manager
+                SimulationInterface->RegisterManagerInterface(ModifiedObject->GetWorld());
+			}
+		}
+	});
+#endif
+}
+
+void UDataflowSimulationManager::OnShutdown()
+{
+#if WITH_EDITOR
+	FCoreUObjectDelegates::OnObjectPropertyChanged.Remove(OnObjectPropertyChangedHandle);
+#endif
+	FWorldDelegates::OnWorldTickEnd.Remove(OnWorldTickEndHandle);
+	UActorComponent::GlobalCreatePhysicsDelegate.Remove(OnCreatePhysicsStateHandle);
+	UActorComponent::GlobalDestroyPhysicsDelegate.Remove(OnDestroyPhysicsStateHandle);
+}
+
+void UDataflowSimulationManager::PreProcessSimulation(const float DeltaTime)
+{
+	for(const TPair<TObjectPtr<UDataflow>, UE::Dataflow::Private::FDataflowSimulationData>& DataflowData : SimulationData)
+	{
+		for(const TPair<FString,TSet<IDataflowSimulationInterface*>>& SimulationInterfaces : DataflowData.Value.SimulationInterfaces)
+		{
+			for(IDataflowSimulationInterface* SimulationInterface : SimulationInterfaces.Value)
+			{
+				if(SimulationInterface)
+				{
+					SimulationInterface->PreProcessSimulation(DeltaTime);
+				}
+			}
+		}
+	}
+}
+
+void UDataflowSimulationManager::PostProcessSimulation(const float DeltaTime)
+{
+	for(const TPair<TObjectPtr<UDataflow>, UE::Dataflow::Private::FDataflowSimulationData>& DataflowData : SimulationData)
+	{
+		for(const TPair<FString,TSet<IDataflowSimulationInterface*>>& SimulationInterfaces : DataflowData.Value.SimulationInterfaces)
+		{
+			for(IDataflowSimulationInterface* SimulationInterface : SimulationInterfaces.Value)
+			{
+				if(SimulationInterface)
+				{
+					SimulationInterface->PostProcessSimulation(DeltaTime);
+				}
+			}
+		}
+	}
+}
+
+void UDataflowSimulationManager::PreAdvanceProxies(const float DeltaTime)
+{
+	for(const TPair<TObjectPtr<UDataflow>, UE::Dataflow::Private::FDataflowSimulationData>& DataflowData : SimulationData)
+	{
+		if(DataflowData.Value.SimulationContext.IsValid())
+		{
+			TArray<FDataflowPhysicsSolverProxy*> SolverProxies;
+			DataflowData.Value.SimulationContext->GetTypedProxies(SolverProxies);
+		
+			for(FDataflowPhysicsSolverProxy* SolverProxy : SolverProxies)
+			{
+				SolverProxy->PreSolveProxy(DeltaTime);
+			}
+		}
+	}
+}
+
+void UDataflowSimulationManager::PostAdvanceProxies(const float DeltaTime)
+{
+	for(const TPair<TObjectPtr<UDataflow>, UE::Dataflow::Private::FDataflowSimulationData>& DataflowData : SimulationData)
+	{
+		if(DataflowData.Value.SimulationContext.IsValid())
+		{
+			TArray<FDataflowPhysicsSolverProxy*> SolverProxies;
+			DataflowData.Value.SimulationContext->GetTypedProxies(SolverProxies);
+
+			for(FDataflowPhysicsSolverProxy* SolverProxy : SolverProxies)
+			{
+				SolverProxy->PostSolveProxy(DeltaTime);
+			}
+		}
+	}
+}
+
+
+void UDataflowSimulationManager::ReadSimulationInterfaces(const float DeltaTime)
+{
+	// Pre-simulation callback that could be used in BP before the simulation
+	//Dataflow::PreSimulationTick(GetWorld(), GetWorld()->GetTimeSeconds(), DeltaTime);
+
+	InitSimulationInterfaces();
+	for(const TPair<TObjectPtr<UDataflow>, UE::Dataflow::Private::FDataflowSimulationData>& DataflowData : SimulationData)
+	{
+		for(const TPair<FString,TSet<IDataflowSimulationInterface*>>& SimulationInterfaces : DataflowData.Value.SimulationInterfaces)
+		{
+			for(IDataflowSimulationInterface* SimulationInterface : SimulationInterfaces.Value)
+			{
+				if(SimulationInterface)
+				{
+					SimulationInterface->WriteToSimulation(DeltaTime);
+				}
+			}
+		}
+	}
+}
+
+void UDataflowSimulationManager::InitSimulationInterfaces()
+{
+	for(const TPair<TObjectPtr<UDataflow>, UE::Dataflow::Private::FDataflowSimulationData>& DataflowData : SimulationData)
+	{
+		if(DataflowData.Value.SimulationContext.IsValid())
+		{
+			DataflowData.Value.SimulationContext->ResetSimulationProxies();
+		}
+		for(const TPair<FString,TSet<IDataflowSimulationInterface*>>& SimulationInterfaces : DataflowData.Value.SimulationInterfaces)
+		{
+			for(IDataflowSimulationInterface* SimulationInterface : SimulationInterfaces.Value)
+			{
+				if(SimulationInterface)
+				{
+					if(!SimulationInterface->GetSimulationProxy())
+					{
+						SimulationInterface->BuildSimulationProxy();
+					}
+					if(SimulationInterface->GetSimulationProxy())
+					{
+						SimulationInterface->GetSimulationProxy()->SetSimulationGroups(SimulationInterface->GetSimulationAsset().SimulationGroups);
+						if(DataflowData.Value.SimulationContext.IsValid())
+						{
+							DataflowData.Value.SimulationContext->AddSimulationProxy(SimulationInterfaces.Key, SimulationInterface->GetSimulationProxy());
+						}
+					}
+				}
+			}
+		}
+		if(DataflowData.Value.SimulationContext.IsValid())
+		{
+			DataflowData.Value.SimulationContext->RegisterProxyGroups();
+		}
+	}
+}
+
+void UDataflowSimulationManager::ResetSimulationInterfaces()
+{
+	for(const TPair<TObjectPtr<UDataflow>, UE::Dataflow::Private::FDataflowSimulationData>& DataflowData : SimulationData)
+	{
+		if(DataflowData.Value.SimulationContext.IsValid())
+		{
+			DataflowData.Value.SimulationContext->ResetSimulationProxies();
+		}
+	}
+}
+
+void UDataflowSimulationManager::WriteSimulationInterfaces(const float DeltaTime)
+{
+	for(const TPair<TObjectPtr<UDataflow>, UE::Dataflow::Private::FDataflowSimulationData>& DataflowData : SimulationData)
+	{
+		for(const TPair<FString,TSet<IDataflowSimulationInterface*>>& SimulationInterfaces : DataflowData.Value.SimulationInterfaces)
+		{
+			for(IDataflowSimulationInterface* SimulationInterface : SimulationInterfaces.Value)
+			{
+				if(SimulationInterface)
+				{
+					SimulationInterface->ReadFromSimulation(DeltaTime);
+				}
+			}
+		}
+	}
+	if(bStepSimulationScene)
+	{
+		bIsSimulationEnabled = false;
+		bStepSimulationScene = false;
+	}
+	ResetSimulationInterfaces();
+	// Post-simulation callback that could be used in BP after the simulation
+	//Dataflow::PostSimulationTick(GetWorld(), GetWorld()->GetTimeSeconds(), DeltaTime);
+}
+
+void UDataflowSimulationManager::AdvanceSimulationProxies(const float DeltaTime, const float SimulationTime)
+{
+	for(TPair<TObjectPtr<UDataflow>, UE::Dataflow::Private::FDataflowSimulationData>& DataflowData : SimulationData)
+	{
+		UE::Dataflow::EvaluateSimulationGraph(DataflowData.Key, DataflowData.Value.SimulationContext, DeltaTime, SimulationTime);
+	}
+}
+
+void UDataflowSimulationManager::StartSimulationTasks(const float DeltaTime, const float SimulationTime)
+{
+	check(IsInGameThread());
+	check(SimulationTasks.IsEmpty());
+
+	for(TPair<TObjectPtr<UDataflow>, UE::Dataflow::Private::FDataflowSimulationData>& DataflowData : SimulationData)
+	{
+		if(!DataflowData.Value.IsEmpty())
+		{
+			// Add a simulation task linked to that solver
+			SimulationTasks.Add(TGraphTask<UE::Dataflow::Private::FDataflowSimulationTask>::CreateTask(nullptr, ENamedThreads::GameThread).ConstructAndDispatchWhenReady(
+					DataflowData.Key, DataflowData.Value.SimulationContext, DeltaTime, SimulationTime));
+		}
+	}
+}
+
+void UDataflowSimulationManager::CompleteSimulationTasks()
+{
+	check(IsInGameThread());
+
+	for(FGraphEventRef& SimulationTask : SimulationTasks)
+	{
+		if (IsValidRef(SimulationTask))
+		{
+			// There's a simulation in flight
+			FTaskGraphInterface::Get().WaitUntilTaskCompletes(SimulationTask, ENamedThreads::GameThread);
+
+			// No longer need this task, it has completed
+			SimulationTask.SafeRelease();
+		}
+	}
+	SimulationTasks.Reset();
+}
+
+TSharedPtr<UE::Dataflow::FDataflowSimulationContext> UDataflowSimulationManager::GetSimulationContext(const TObjectPtr<UDataflow>& DataflowAsset) const
+{
+	if(DataflowAsset)
+	{
+		if(const UE::Dataflow::Private::FDataflowSimulationData* DataflowData = SimulationData.Find(DataflowAsset))
+		{
+			return DataflowData->SimulationContext;
+		}
+	}
+	return nullptr;
+}
+
+bool UDataflowSimulationManager::HasSimulationInterface(const IDataflowSimulationInterface* SimulationInterface) const
+{
+	if(SimulationInterface)
+	{
+		if(const TObjectPtr<UDataflow> DataflowAsset = SimulationInterface->GetSimulationAsset().DataflowAsset)
+		{
+			if(const UE::Dataflow::Private::FDataflowSimulationData* DataflowData = SimulationData.Find(DataflowAsset))
+			{
+				if(const TSet<IDataflowSimulationInterface*>* SimulationInterfaces =
+					DataflowData->SimulationInterfaces.Find(SimulationInterface->GetSimulationType()))
+				{
+					return SimulationInterfaces->Find(SimulationInterface) != nullptr;
+				}
+			}
+		}
+	}
+	return false;
+}
+
+void UDataflowSimulationManager::AddSimulationInterface(IDataflowSimulationInterface* SimulationInterface)
+{
+	if(SimulationInterface)
+	{
+		if(TObjectPtr<UDataflow> DataflowAsset = SimulationInterface->GetSimulationAsset().DataflowAsset)
+		{
+			UE::Dataflow::Private::FDataflowSimulationData& DataflowData = SimulationData.FindOrAdd(DataflowAsset);
+			if(!DataflowData.SimulationContext.IsValid())
+			{
+				DataflowData.SimulationContext = MakeShared<UE::Dataflow::FDataflowSimulationContext>(DataflowAsset);
+			}
+			DataflowData.SimulationInterfaces.FindOrAdd(SimulationInterface->GetSimulationType()).Add(SimulationInterface);
+		}
+	}
+}
+
+void UDataflowSimulationManager::RemoveSimulationInterface(const IDataflowSimulationInterface* SimulationInterface)
+{
+	if(SimulationInterface)
+	{
+		for(TPair<TObjectPtr<UDataflow>, UE::Dataflow::Private::FDataflowSimulationData>& DataflowData : SimulationData)
+		{
+			if(TSet<IDataflowSimulationInterface*>* SimulationInterfaces =
+					DataflowData.Value.SimulationInterfaces.Find(SimulationInterface->GetSimulationType()))
+			{
+				SimulationInterfaces->Remove(SimulationInterface);
+			}
+		}
+	}
+}
+
+ETickableTickType UDataflowSimulationManager::GetTickableTickType() const
+{
+	return HasAnyFlags(RF_ClassDefaultObject) || !GetWorld() || GetWorld()->IsNetMode(NM_DedicatedServer) ? ETickableTickType::Never : ETickableTickType::Always;
+}
+
+TStatId UDataflowSimulationManager::GetStatId() const
+{
+	RETURN_QUICK_DECLARE_CYCLE_STAT(UDataflowSimulationManager, STATGROUP_Tickables);
+}
+
+bool UDataflowSimulationManager::DoesSupportWorldType(EWorldType::Type WorldType) const
+{
+	return Super::DoesSupportWorldType(WorldType) || WorldType == EWorldType::EditorPreview || WorldType == EWorldType::GamePreview || WorldType == EWorldType::GameRPC;
+}
+
+void UDataflowSimulationManager::Initialize(FSubsystemCollectionBase& Collection)
+{
+	Super::Initialize(Collection);
+}
+
+void UDataflowSimulationManager::Deinitialize()
+{
+	Super::Deinitialize();
+
+	CompleteSimulationTasks();
+}
+
+
diff --git a/Engine/Source/Runtime/Experimental/Dataflow/Simulation/Private/Dataflow/DataflowSimulationNodes.cpp b/Engine/Source/Runtime/Experimental/Dataflow/Simulation/Private/Dataflow/DataflowSimulationNodes.cpp
new file mode 100644
--- /dev/null
+++ b/Engine/Source/Runtime/Experimental/Dataflow/Simulation/Private/Dataflow/DataflowSimulationNodes.cpp
@@ -0,0 +1,99 @@
+﻿// Copyright Epic Games, Inc. All Rights Reserved.
+
+#include "Dataflow/DataflowSimulationNodes.h"
+#include "Dataflow/DataflowSimulationManager.h"
+#include "Dataflow/DataflowNodeFactory.h"
+#include "Dataflow/DataflowNodeColorsRegistry.h"
+#include "Dataflow/Interfaces/DataflowPhysicsSolver.h"
+
+#include UE_INLINE_GENERATED_CPP_BY_NAME(DataflowSimulationNodes)
+
+namespace Dataflow
+{
+	void RegisterDataflowSimulationNodes()
+	{
+		// Common nodes
+		DATAFLOW_NODE_REGISTER_CREATION_FACTORY(FGetSimulationTimeDataflowNode);
+		
+		// Solvers nodes
+		DATAFLOW_NODE_REGISTER_CREATION_FACTORY(FGetPhysicsSolversDataflowNode);
+		DATAFLOW_NODE_REGISTER_CREATION_FACTORY(FAdvancePhysicsSolversDataflowNode);
+
+		// Proxies nodes
+		DATAFLOW_NODE_REGISTER_CREATION_FACTORY(FFilterSimulationProxiesDataflowNode);
+		DATAFLOW_NODE_REGISTER_CREATION_FACTORY(FSimulationProxiesTerminalDataflowNode);
+		
+		static constexpr FLinearColor CDefaultNodeBodyTintColor = FLinearColor(0.f, 0.f, 0.0f, 0.5f);
+		
+		DATAFLOW_NODE_REGISTER_CREATION_FACTORY_NODE_COLORS_BY_CATEGORY("Terminal", FLinearColor(1.0f, 0.0f, 0.0f), CDefaultNodeBodyTintColor);
+		DATAFLOW_NODE_REGISTER_CREATION_FACTORY_NODE_COLORS_BY_CATEGORY("Setup", FLinearColor(1.0f, 1.0f, 0.0f), CDefaultNodeBodyTintColor);
+		DATAFLOW_NODE_REGISTER_CREATION_FACTORY_NODE_COLORS_BY_CATEGORY("Physics", FLinearColor(0.0f, 1.0f, 0.0f), CDefaultNodeBodyTintColor);
+	}
+}
+
+void FGetSimulationTimeDataflowNode::EvaluateSimulation(UE::Dataflow::FDataflowSimulationContext& SimulationContext, const FDataflowOutput* Output) const
+{
+	SetValue(SimulationContext, FDataflowSimulationTime(SimulationContext.GetDeltaTime(), SimulationContext.GetSimulationTime()), &SimulationTime);
+}
+
+void FAdvancePhysicsSolversDataflowNode::EvaluateSimulation(UE::Dataflow::FDataflowSimulationContext& SimulationContext, const FDataflowOutput* Output) const
+{
+	const TArray<FDataflowSimulationProperty> SolverProperties = GetValue(SimulationContext, &PhysicsSolvers);
+	const float SimulationDeltaTime = GetValue(SimulationContext, &SimulationTime).DeltaTime;
+	
+	if(!SolverProperties.IsEmpty())
+	{
+		for(const FDataflowSimulationProperty& SolverProperty : SolverProperties)
+		{
+			if(SolverProperty.SimulationProxy)
+			{
+				if(FDataflowPhysicsSolverProxy* SolverProxy = SolverProperty.SimulationProxy->AsType<FDataflowPhysicsSolverProxy>())
+				{
+					SolverProxy->AdvanceSolverDatas(SimulationDeltaTime);
+				}
+			}
+		}
+	}
+	SetValue(SimulationContext, SolverProperties, &PhysicsSolvers);
+}
+
+void FGetPhysicsSolversDataflowNode::EvaluateSimulation(UE::Dataflow::FDataflowSimulationContext& SimulationContext, const FDataflowOutput* Output) const
+{
+	TArray<FDataflowSimulationProxy*> SimulationProxies;
+	SimulationContext.GetSimulationProxies(FDataflowPhysicsSolverProxy::StaticStruct()->GetName(), SimulationGroups, SimulationProxies);
+
+	TArray<FDataflowSimulationProperty> SolverProperties;
+	for(FDataflowSimulationProxy* SimulationProxy : SimulationProxies)
+	{
+		SolverProperties.Add({SimulationProxy});
+	}
+	
+	SetValue(SimulationContext, SolverProperties, &PhysicsSolvers);
+}
+
+void FFilterSimulationProxiesDataflowNode::EvaluateSimulation(UE::Dataflow::FDataflowSimulationContext& SimulationContext, const FDataflowOutput* Output) const
+{
+	const TArray<FDataflowSimulationProperty> SimulationProperties = GetValue(SimulationContext, &SimulationProxies);
+
+	TArray<FDataflowSimulationProperty> FilteredProperties;
+	if(!SimulationProperties.IsEmpty())
+	{
+		TBitArray<> GroupBits;
+		SimulationContext.BuildGroupBits(SimulationGroups, GroupBits);
+
+		for(const FDataflowSimulationProperty& SimulationProperty : SimulationProperties)
+		{
+			if(SimulationProperty.SimulationProxy && SimulationProperty.SimulationProxy->HasGroupBit(GroupBits))
+			{
+				FilteredProperties.Add({SimulationProperty.SimulationProxy});
+			}
+		}
+	}
+	
+	SetValue(SimulationContext, FilteredProperties, &FilteredProxies);
+}
+
+void FSimulationProxiesTerminalDataflowNode::EvaluateSimulation(UE::Dataflow::FDataflowSimulationContext& SimulationContext, const FDataflowOutput* Output) const
+{
+	const TArray<FDataflowSimulationProperty> SolverProperties = GetValue(SimulationContext, &SimulationProxies);
+}
diff --git a/Engine/Source/Runtime/Experimental/Dataflow/Simulation/Private/Dataflow/DataflowSimulationPlugin.cpp b/Engine/Source/Runtime/Experimental/Dataflow/Simulation/Private/Dataflow/DataflowSimulationPlugin.cpp
new file mode 100644
--- /dev/null
+++ b/Engine/Source/Runtime/Experimental/Dataflow/Simulation/Private/Dataflow/DataflowSimulationPlugin.cpp
@@ -0,0 +1,29 @@
+﻿// Copyright Epic Games, Inc. All Rights Reserved.
+
+#include "Dataflow/DataflowSimulationPlugin.h"
+#include "Dataflow/DataflowSimulationNodes.h"
+#include "Dataflow/DataflowGraph.h"
+#include "Dataflow/DataflowSimulationManager.h"
+#include "CoreMinimal.h"
+
+#define LOCTEXT_NAMESPACE "DataflowSimulation"
+
+void IDataflowSimulationPlugin::StartupModule()
+{
+	Dataflow::RegisterDataflowSimulationNodes();
+
+	Dataflow::RegisterNodeFilter(FDataflowSimulationNode::StaticType());
+	Dataflow::RegisterNodeFilter(FDataflowInvalidNode::StaticType());
+	Dataflow::RegisterNodeFilter(FDataflowExecutionNode::StaticType());
+
+	UDataflowSimulationManager::OnStartup();
+}
+
+void IDataflowSimulationPlugin::ShutdownModule()
+{
+	UDataflowSimulationManager::OnShutdown();
+}
+
+IMPLEMENT_MODULE(IDataflowSimulationPlugin, DataflowSimulation)
+
+#undef LOCTEXT_NAMESPACE
diff --git a/Engine/Source/Runtime/Experimental/Dataflow/Simulation/Private/Dataflow/DataflowSimulationProxy.cpp b/Engine/Source/Runtime/Experimental/Dataflow/Simulation/Private/Dataflow/DataflowSimulationProxy.cpp
new file mode 100644
--- /dev/null
+++ b/Engine/Source/Runtime/Experimental/Dataflow/Simulation/Private/Dataflow/DataflowSimulationProxy.cpp
@@ -0,0 +1,19 @@
+﻿// Copyright Epic Games, Inc. All Rights Reserved.
+
+#include "Dataflow/DataflowSimulationProxy.h"
+
+void FDataflowSimulationProxy::SetSimulationGroups(const TSet<FString>& InSimulationGroups)
+{
+	SimulationGroups = InSimulationGroups;
+}
+
+bool FDataflowSimulationProxy::HasSimulationGroup(const FString& SimulationGroup) const
+{
+	return (SimulationGroups.Find(SimulationGroup) != nullptr);
+}
+
+bool FDataflowSimulationProxy::HasGroupBit(const TBitArray<>& SimulationBits) const
+{
+	const TBitArray<> CommonGroups = TBitArray<>::BitwiseAND(GroupBits, SimulationBits, EBitwiseOperatorFlags::MinSize);
+	return (CommonGroups.CountSetBits() != 0);
+}
diff --git a/Engine/Source/Runtime/Experimental/Dataflow/Simulation/Private/Dataflow/DataflowSimulationUtils.cpp b/Engine/Source/Runtime/Experimental/Dataflow/Simulation/Private/Dataflow/DataflowSimulationUtils.cpp
new file mode 100644
--- /dev/null
+++ b/Engine/Source/Runtime/Experimental/Dataflow/Simulation/Private/Dataflow/DataflowSimulationUtils.cpp
@@ -0,0 +1,40 @@
+// Copyright Epic Games, Inc. All Rights Reserved.
+
+#include "Dataflow/DataflowSimulationUtils.h"
+#include "Dataflow/DataflowObject.h"
+#include "Dataflow/DataflowSimulationNodes.h"
+#include "Dataflow/DataflowObjectInterface.h"
+
+#define LOCTEXT_NAMESPACE "DataflowSimulationGenerator"
+
+namespace UE::Dataflow
+{
+	void EvaluateSimulationGraph(const TObjectPtr<UDataflow>& SimulationGraph, const TSharedPtr<FDataflowSimulationContext>& SimulationContext,
+		const float DeltaTime, const float SimulationTime)
+	{
+		if(SimulationContext.IsValid())
+		{
+			SimulationContext->SetTimingInfos(DeltaTime, SimulationTime);
+			
+			if(SimulationGraph)
+			{
+				if(const TSharedPtr<::Dataflow::FGraph> DataflowGraph = SimulationGraph->GetDataflow())
+				{
+					// Invalidation of all the simulation nodes that are always dirty
+					for(const TSharedPtr<FDataflowNode>& InvalidNode : DataflowGraph->GetFilteredNodes(FDataflowInvalidNode::StaticType()))
+					{
+						InvalidNode->Invalidate();
+					}
+					
+					// Pull the graph evaluation from the solver nodes
+					for(const TSharedPtr<FDataflowNode>& ExecutionNode : DataflowGraph->GetFilteredNodes(FDataflowExecutionNode::StaticType()))
+					{
+						SimulationContext->Evaluate(ExecutionNode.Get(), nullptr);
+					}
+				}
+			}
+		}
+	}
+}
+
+#undef LOCTEXT_NAMESPACE
diff --git a/Engine/Source/Runtime/Experimental/Dataflow/Simulation/Private/Dataflow/Interfaces/DataflowCollisionObject.cpp b/Engine/Source/Runtime/Experimental/Dataflow/Simulation/Private/Dataflow/Interfaces/DataflowCollisionObject.cpp
new file mode 100644
--- /dev/null
+++ b/Engine/Source/Runtime/Experimental/Dataflow/Simulation/Private/Dataflow/Interfaces/DataflowCollisionObject.cpp
@@ -0,0 +1,9 @@
+﻿// Copyright Epic Games, Inc. All Rights Reserved.
+
+#include "Dataflow/Interfaces/DataflowCollisionObject.h"
+
+FDataflowCollisionObjectProxy::FDataflowCollisionObjectProxy()
+{}
+
+IDataflowCollisionObjectInterface::IDataflowCollisionObjectInterface()
+{}
diff --git a/Engine/Source/Runtime/Experimental/Dataflow/Simulation/Private/Dataflow/Interfaces/DataflowConstraintObject.cpp b/Engine/Source/Runtime/Experimental/Dataflow/Simulation/Private/Dataflow/Interfaces/DataflowConstraintObject.cpp
new file mode 100644
--- /dev/null
+++ b/Engine/Source/Runtime/Experimental/Dataflow/Simulation/Private/Dataflow/Interfaces/DataflowConstraintObject.cpp
@@ -0,0 +1,9 @@
+﻿// Copyright Epic Games, Inc. All Rights Reserved.
+
+#include "Dataflow/Interfaces/DataflowConstraintObject.h"
+
+FDataflowConstraintObjectProxy::FDataflowConstraintObjectProxy()
+{}
+
+IDataflowConstraintObjectInterface::IDataflowConstraintObjectInterface()
+{}
diff --git a/Engine/Source/Runtime/Experimental/Dataflow/Simulation/Private/Dataflow/Interfaces/DataflowPhysicsObject.cpp b/Engine/Source/Runtime/Experimental/Dataflow/Simulation/Private/Dataflow/Interfaces/DataflowPhysicsObject.cpp
new file mode 100644
--- /dev/null
+++ b/Engine/Source/Runtime/Experimental/Dataflow/Simulation/Private/Dataflow/Interfaces/DataflowPhysicsObject.cpp
@@ -0,0 +1,9 @@
+﻿// Copyright Epic Games, Inc. All Rights Reserved.
+
+#include "Dataflow/Interfaces/DataflowPhysicsObject.h"
+
+FDataflowPhysicsObjectProxy::FDataflowPhysicsObjectProxy()
+{}
+
+IDataflowPhysicsObjectInterface::IDataflowPhysicsObjectInterface()
+{}
diff --git a/Engine/Source/Runtime/Experimental/Dataflow/Simulation/Private/Dataflow/Interfaces/DataflowPhysicsSolver.cpp b/Engine/Source/Runtime/Experimental/Dataflow/Simulation/Private/Dataflow/Interfaces/DataflowPhysicsSolver.cpp
new file mode 100644
--- /dev/null
+++ b/Engine/Source/Runtime/Experimental/Dataflow/Simulation/Private/Dataflow/Interfaces/DataflowPhysicsSolver.cpp
@@ -0,0 +1,9 @@
+﻿// Copyright Epic Games, Inc. All Rights Reserved.
+
+#include "Dataflow/Interfaces/DataflowPhysicsSolver.h"
+
+FDataflowPhysicsSolverProxy::FDataflowPhysicsSolverProxy()
+{}
+
+IDataflowPhysicsSolverInterface::IDataflowPhysicsSolverInterface()
+{}
diff --git a/Engine/Source/Runtime/Experimental/Dataflow/Simulation/Public/Dataflow/DataflowSimulationContext.h b/Engine/Source/Runtime/Experimental/Dataflow/Simulation/Public/Dataflow/DataflowSimulationContext.h
new file mode 100644
--- /dev/null
+++ b/Engine/Source/Runtime/Experimental/Dataflow/Simulation/Public/Dataflow/DataflowSimulationContext.h
@@ -0,0 +1,119 @@
+// Copyright Epic Games, Inc. All Rights Reserved.
+
+#pragma once
+
+#include "CoreMinimal.h"
+#include "UObject/ObjectMacros.h"
+#include "Dataflow/DataflowObjectInterface.h"
+#include "Dataflow/DataflowSimulationProxy.h"
+
+namespace UE::Dataflow
+{
+	/** Simulation context that will be used by all the simulation nodes*/
+	template<typename Base = ::Dataflow::FContextSingle>
+	class DATAFLOWSIMULATION_API TSimulationContext : public ::Dataflow::TEngineContext<Base>
+	{
+	public:
+		DATAFLOW_CONTEXT_INTERNAL(::Dataflow::TEngineContext<Base>, TSimulationContext);
+
+		explicit TSimulationContext(const TObjectPtr<UObject>& InOwner)
+				: Super(InOwner)
+		{}
+		
+		virtual ~TSimulationContext() override {};
+		
+		/** Set the timing infos */
+		void SetTimingInfos(const float DeltaSeconds, const float TimeSeconds) {DeltaTime = DeltaSeconds; SimulationTime = TimeSeconds;}
+
+		/** Get the delta time in seconds */
+		float GetDeltaTime() const {return DeltaTime;}
+		
+		/** Get the simulation time in seconds */
+		float GetSimulationTime() const {return SimulationTime;}
+
+		/** Get the typed proxies */
+		template<typename ProxyType>
+		void GetTypedProxies(TArray<ProxyType*>& FilteredProxies) const;
+
+		/** Filter the physics solvers matching the groups */
+		void GetSimulationProxies(const FString& ProxyType, const TArray<FString>& SimulationGroups, TArray<FDataflowSimulationProxy*>& FilteredProxies) const;
+
+		/** Add simulation proxy to the context */
+		void AddSimulationProxy(const FString& ProxyType, FDataflowSimulationProxy* SimulationProxy);
+		
+		/** Remove simulation proxy from the context */
+    	void RemoveSimulationProxy(const FString& ProxyType, const FDataflowSimulationProxy* SimulationProxy);
+
+		/** Reset all the simulation proxies */
+		void ResetSimulationProxies();
+		
+		/** Return the number of physics solvers */
+		int32 NumSimulationProxies(const FString& ProxyType) const;
+
+		/** Register all the proxy groups used in the proxy */
+		void RegisterProxyGroups();
+
+		/** Build group bits based on the string array */
+		void BuildGroupBits(const TArray<FString>& SimulationGroups, TBitArray<>& GroupBits) const;
+
+		/** Push another level of iteration indices */
+		void PushIterationIndex()
+		{
+			IterationIndices.Push(0);
+		}
+
+		/** Pop the last level of iteration indices */
+		void PopIterationIndex()
+		{
+			IterationIndices.Pop();
+		}
+
+		/** Set the last iteration index */
+		void SetIterationIndex(const int32 IterationIndex)
+		{
+			IterationIndices.Last() = IterationIndex;
+		}
+
+		/** Get the last iteration index */
+		int32 GetIterationIndex()
+		{
+			return IterationIndices.Last();
+		}
+		
+	private :
+
+		/** Simulation time */
+		float SimulationTime = 0.0f;
+
+		/** Delta time */
+		float DeltaTime = 0.0f;
+
+		/** List of all the simulation proxies within the context sorted by type */
+		TMap<FString, TSet<FDataflowSimulationProxy*>> SimulationProxies;
+
+		/** List of all the simulation groups indices */
+		TMap<FString, uint32> GroupIndices;
+
+		/** List of all the nested iteration indices */
+		TArray<int32> IterationIndices;
+	};
+
+	template<typename Base>
+	template<typename ProxyType>
+	void TSimulationContext<Base>::GetTypedProxies(TArray<ProxyType*>& FilteredProxies) const
+	{
+		if(const TSet<FDataflowSimulationProxy*>* TypedProxies = SimulationProxies.Find(ProxyType::StaticStruct()->GetName()))
+		{
+			for(FDataflowSimulationProxy* SimulationProxy : *TypedProxies)
+			{
+				if(SimulationProxy)
+				{
+					FilteredProxies.Add(SimulationProxy->AsType<ProxyType>());
+				}
+			}
+		}
+	}
+
+	using FDataflowSimulationContext = TSimulationContext<::Dataflow::FContextSingle>;
+	using FDataflowSimulationContextThreaded =  TSimulationContext<::Dataflow::FContextThreaded>;
+}
diff --git a/Engine/Source/Runtime/Experimental/Dataflow/Simulation/Public/Dataflow/DataflowSimulationInterface.h b/Engine/Source/Runtime/Experimental/Dataflow/Simulation/Public/Dataflow/DataflowSimulationInterface.h
new file mode 100644
--- /dev/null
+++ b/Engine/Source/Runtime/Experimental/Dataflow/Simulation/Public/Dataflow/DataflowSimulationInterface.h
@@ -0,0 +1,90 @@
+﻿// Copyright Epic Games, Inc. All Rights Reserved.
+
+#pragma once
+
+#include "CoreMinimal.h"
+#include "Dataflow/DataflowObject.h"
+#include "UObject/Interface.h"
+
+#include "DataflowSimulationInterface.generated.h"
+
+struct FDataflowSimulationProxy;
+
+/**
+ * Dataflow simulation asset (should be in the interface children)
+ */
+USTRUCT(BlueprintType)
+struct DATAFLOWSIMULATION_API FDataflowSimulationAsset
+{
+	GENERATED_BODY()
+
+	/* Simulation dataflow asset used to advance in time on Pt */
+	UPROPERTY(EditAnywhere, BlueprintReadWrite, Category = "Physics")
+	TObjectPtr<UDataflow> DataflowAsset;
+
+	/* Simulation groups used to filter within the simulation nodes*/
+	UPROPERTY(EditAnywhere, BlueprintReadWrite, Category = "Physics")
+	TSet<FString> SimulationGroups;
+};
+
+UINTERFACE()
+class DATAFLOWSIMULATION_API UDataflowSimulationInterface : public UInterface
+{
+	GENERATED_BODY()
+};
+
+/**
+ * Dataflow simulation interface to send/receive datas (GT <-> PT)
+ */
+class DATAFLOWSIMULATION_API IDataflowSimulationInterface
+{
+	GENERATED_BODY()
+	
+public:
+	IDataflowSimulationInterface() {}
+
+	/** Get the dataflow simulation asset */
+	virtual FDataflowSimulationAsset& GetSimulationAsset() = 0;
+
+	/** Get the const dataflow simulation asset */
+	virtual const FDataflowSimulationAsset& GetSimulationAsset() const = 0;
+	
+	/** Build the simulation proxy */
+	virtual void BuildSimulationProxy() = 0;
+	
+	/** Reset the simulation proxy */
+	virtual void ResetSimulationProxy() = 0;
+
+	/** Get the const simulation proxy */
+	virtual const FDataflowSimulationProxy* GetSimulationProxy() const = 0;
+
+	/** Get the simulation proxy */
+	virtual FDataflowSimulationProxy* GetSimulationProxy() = 0;
+
+	/** Get the simulation name */
+	virtual FString GetSimulationName() const = 0;
+
+	/** Preprocess data before simulation */
+	virtual void PreProcessSimulation(const float DeltaTime) {};
+
+	/** Write data to be sent to the simulation proxy */
+	virtual void WriteToSimulation(const float DeltaTime) {};
+
+	/** Read data received from the simulation proxy */
+	virtual void ReadFromSimulation(const float DeltaTime) {};
+
+	/** Postprocess data after simulation */
+	virtual void PostProcessSimulation(const float DeltaTime) {};
+	
+	/** Get the simulation type */
+    virtual FString GetSimulationType() const {return TEXT("");};
+
+	/** Register simulation interface solver to manager */
+	void RegisterManagerInterface(const TObjectPtr<UWorld>& SimulationWorld);
+
+	/** Unregister simulation interface from the manager */
+	void UnregisterManagerInterface(const TObjectPtr<UWorld>& SimulationWorld) const;
+
+	/** Check if the interface has been registered to the manager */
+	bool IsInterfaceRegistered(const TObjectPtr<UWorld>& SimulationWorld) const;
+};
diff --git a/Engine/Source/Runtime/Experimental/Dataflow/Simulation/Public/Dataflow/DataflowSimulationManager.h b/Engine/Source/Runtime/Experimental/Dataflow/Simulation/Public/Dataflow/DataflowSimulationManager.h
new file mode 100644
--- /dev/null
+++ b/Engine/Source/Runtime/Experimental/Dataflow/Simulation/Public/Dataflow/DataflowSimulationManager.h
@@ -0,0 +1,173 @@
+﻿// Copyright Epic Games, Inc. All Rights Reserved.
+
+#pragma once
+
+#include "CoreMinimal.h"
+#include "UObject/ObjectMacros.h"
+#include "UObject/Interface.h"
+#include "Subsystems/WorldSubsystem.h"
+#include "Templates/SharedPointer.h"
+#include "Dataflow/DataflowSimulationContext.h"
+
+#include "DataflowSimulationManager.generated.h"
+
+struct FDataflowSimulationProxy;
+class IDataflowSimulationInterface;
+
+namespace UE::Dataflow
+{
+	namespace Private
+	{
+		/** Per dataflow graph simulation data type (data interfaces + simulation context) */
+		struct  FDataflowSimulationData
+		{
+			/** List of all simulation interfaces used in this this dataflow graph */
+			TMap<FString, TSet<IDataflowSimulationInterface*>> SimulationInterfaces;
+
+			/** Simulation context used to evaluate the graph on PT */
+			TSharedPtr<UE::Dataflow::FDataflowSimulationContext> SimulationContext;
+
+			/** Check is there is any datas to process */
+			bool IsEmpty() const
+			{
+				for(const TPair<FString, TSet<IDataflowSimulationInterface*>>& InterfacesPair : SimulationInterfaces)
+				{
+					if(!InterfacesPair.Value.IsEmpty())
+					{
+						return false;
+					}
+				}
+				return true;
+			}
+		};
+	}
+	DATAFLOWSIMULATION_API void RegisterSimulationInterface(const TObjectPtr<UObject>& SimulationObject);
+	DATAFLOWSIMULATION_API void UnregisterSimulationInterface(const TObjectPtr<UObject>& SimulationObject);
+}
+
+UCLASS()
+class DATAFLOWSIMULATION_API UDataflowSimulationManager : public UTickableWorldSubsystem
+{
+	GENERATED_BODY()
+
+public:
+	UDataflowSimulationManager();
+	virtual ~UDataflowSimulationManager() override  = default;
+
+	/** Static function to add world delegates */
+	static void OnStartup();
+	
+	/** Static function to remove world delegates */
+	static void OnShutdown();
+
+	// Begin FTickableGameObject overrides
+	virtual void Tick(float DeltaTime) override;
+	virtual bool IsTickableInEditor() const override { return true; }
+	virtual ETickableTickType GetTickableTickType() const override;
+	virtual bool DoesSupportWorldType(const EWorldType::Type WorldType) const override;
+	virtual TStatId GetStatId() const override;
+	// End FTickableGameObject overrides
+	
+	// Begin USubsystem overrides
+	virtual void Initialize(FSubsystemCollectionBase& Collection) override;
+	virtual void Deinitialize() override;
+	// End USubsystem overrides
+
+	/** Pre solve all the solver proxies */
+	void PreAdvanceProxies(const float DeltaTime);
+
+	/** Post solve all the solver proxies */
+	void PostAdvanceProxies(const float DeltaTime);
+
+	/** Advance in time the registered simulation data (PT) */
+    void AdvanceSimulationProxies(const float DeltaTime, const float SimulationTime);
+
+	/** Check if the manager has a simulation interface */
+	bool HasSimulationInterface(const IDataflowSimulationInterface* SimulationInterface) const;
+
+	/** Add a dataflow simulation interface to the manager*/
+	void AddSimulationInterface(IDataflowSimulationInterface* SimulationInterface);
+
+	/** Remove a dataflow simulation interface from the manager */
+	void RemoveSimulationInterface(const IDataflowSimulationInterface* SimulationInterface);
+
+	/** Read the simulation interfaces and Write all the data to the simulation proxies (to be send from GT->PT) */
+	void ReadSimulationInterfaces(const float DeltaTime);
+
+	/** Read all the data from the simulation proxies and write the result onto the interfaces (received from PT->GT) */
+	void WriteSimulationInterfaces(const float DeltaTime);
+
+	/** Init all the simulation interfaces*/
+	void InitSimulationInterfaces();
+
+	/** Reset all the simulation interfaces*/
+	void ResetSimulationInterfaces();
+
+	/** Complete all the simulation tasks */
+	void CompleteSimulationTasks();
+
+	/** Start the simulation tasks given a delta time */
+	void StartSimulationTasks(const float DeltaTime, const float SimulationTime);
+
+	/** Set the simulation flag to enable/disable the simulation */
+	void SetSimulationEnabled(const bool bSimulationEnabled) {bIsSimulationEnabled = bSimulationEnabled;}
+
+	/** Set the simulation scene stepping bool */
+	void SetSimulationStepping(const bool bSimulationStepping) {bStepSimulationScene = bSimulationStepping;}
+
+	/** Get the simulation context for a given asset */
+	TSharedPtr<UE::Dataflow::FDataflowSimulationContext> GetSimulationContext(const TObjectPtr<UDataflow>& DataflowAsset) const;
+	
+private :
+
+	/** Pre process before the simulation step */
+	void PreProcessSimulation(const float DeltaTime);
+
+	/** Post process after the simulation step */
+	void PostProcessSimulation(const float DeltaTime);
+
+	/** static delegate for object property changed */
+	static FDelegateHandle OnObjectPropertyChangedHandle;
+
+	/** static delegate for end of frame update */
+	static FDelegateHandle OnWorldTickEndHandle;
+
+	/** static delegate for physics state creation */
+	static FDelegateHandle OnCreatePhysicsStateHandle;
+
+	/** static delegate for physics state destruction */
+	static FDelegateHandle OnDestroyPhysicsStateHandle;
+
+	/** Dataflow simulation data registered to the manager */
+	TMap<TObjectPtr<UDataflow>, UE::Dataflow::Private::FDataflowSimulationData> SimulationData;
+	
+	/** Simulation tasks in which the graph will be evaluated */
+	TArray<FGraphEventRef> SimulationTasks;
+ 
+	/** Boolean to control if the simulation should be disabled or not */
+	bool bIsSimulationEnabled = true;
+	
+	/** Boolean to check if we are stepping the simulation scene */
+	bool bStepSimulationScene = false;
+};
+
+/** Dataflow simulation actor interface to be able to call BP events before/after the manager ticking in case we need it */
+UINTERFACE(MinimalAPI, Blueprintable)
+class UDataflowSimulationActor : public UInterface
+{
+	GENERATED_BODY()
+};
+
+class IDataflowSimulationActor
+{
+	GENERATED_BODY()
+ 
+public:
+	/** Pre simulation callback function that can be implemented in C++ or Blueprint. */
+	UFUNCTION(BlueprintCallable, BlueprintNativeEvent, Category="Dataflow")
+	void PreDataflowSimulationTick(const float SimulationTime, const float DeltaTime);
+
+	/** Post simulation callback function that can be implemented in C++ or Blueprint. */
+	UFUNCTION(BlueprintCallable, BlueprintNativeEvent, Category="Dataflow")
+	void PostDataflowSimulationTick(const float SimulationTime, const float DeltaTime);
+};
diff --git a/Engine/Source/Runtime/Experimental/Dataflow/Simulation/Public/Dataflow/DataflowSimulationNodes.h b/Engine/Source/Runtime/Experimental/Dataflow/Simulation/Public/Dataflow/DataflowSimulationNodes.h
new file mode 100644
--- /dev/null
+++ b/Engine/Source/Runtime/Experimental/Dataflow/Simulation/Public/Dataflow/DataflowSimulationNodes.h
@@ -0,0 +1,294 @@
+﻿// Copyright Epic Games, Inc. All Rights Reserved.
+
+#pragma once
+
+#include "Dataflow/DataflowObject.h"
+#include "Dataflow/DataflowSimulationProxy.h"
+#include "Dataflow/DataflowSimulationContext.h"
+
+#include "DataflowSimulationNodes.generated.h"
+
+/**
+ * Dataflow simulation property 
+ */
+USTRUCT(BlueprintType)
+struct DATAFLOWSIMULATION_API FDataflowSimulationProperty
+{
+	GENERATED_BODY()
+
+	/* Simulation proxy used to pass information in between nodes */
+	FDataflowSimulationProxy* SimulationProxy;
+};
+
+/**
+* FDataflowSimulationNode
+*		Base class for simulation nodes within the Dataflow graph. 
+* 
+*		Simulation nodes are used to simulate data from the calling client. 
+*/
+USTRUCT()
+struct FDataflowSimulationNode : public FDataflowNode
+{
+	GENERATED_BODY()
+
+	FDataflowSimulationNode()
+		: Super() { }
+
+	FDataflowSimulationNode(const Dataflow::FNodeParameters& Param, FGuid InGuid = FGuid::NewGuid())
+		: Super(Param,InGuid) {
+	}
+
+	/** FDataflowNode interface */
+	virtual ~FDataflowSimulationNode() {}
+	
+	static FName StaticType() { return FName("FDataflowSimulationNode"); }
+
+	virtual bool IsA(FName InType) const override 
+	{ 
+		return InType.ToString().Equals(StaticType().ToString()) 
+			|| Super::IsA(InType); 
+	}
+
+	/** Evaluate simulation dispatch */
+	virtual void Evaluate(Dataflow::FContext& Context, const FDataflowOutput* Output) const override
+	{
+		if(Context.IsA(UE::Dataflow::FDataflowSimulationContext::StaticType()))
+		{
+			UE::Dataflow::FDataflowSimulationContext& SimulationContext = StaticCast<UE::Dataflow::FDataflowSimulationContext&>(Context);
+			EvaluateSimulation(SimulationContext, Output);
+			
+		}
+	}
+
+protected :
+
+	/** Evaluate simulation node given a simulation context */
+	virtual void EvaluateSimulation(UE::Dataflow::FDataflowSimulationContext& SimulationContext, const FDataflowOutput* Output) const {};
+};
+
+/**
+* FDataflowInvalidNode
+*		Base class for invalid nodes within the Dataflow graph. 
+* 
+*		Invalid nodes will be always invalidated while simulating
+*/
+USTRUCT()
+struct FDataflowInvalidNode : public FDataflowSimulationNode
+{
+	GENERATED_BODY()
+
+	FDataflowInvalidNode()
+		: Super() { }
+
+	FDataflowInvalidNode(const Dataflow::FNodeParameters& Param, FGuid InGuid = FGuid::NewGuid())
+		: Super(Param,InGuid) {
+	}
+
+	/** FDataflowNode interface */
+	virtual ~FDataflowInvalidNode() {}
+	
+	static FName StaticType() { return FName("FDataflowInvalidNode"); }
+
+	virtual bool IsA(FName InType) const override 
+	{ 
+		return InType.ToString().Equals(StaticType().ToString()) 
+			|| Super::IsA(InType); 
+	}
+};
+
+/**
+* FDataflowExecutionNode
+*		Base class for the execute the dataflow simulation graph. 
+* 
+*		Execution nodes are used to pull the graph from the calling client. 
+*/
+USTRUCT()
+struct FDataflowExecutionNode : public FDataflowSimulationNode
+{
+	GENERATED_BODY()
+
+	FDataflowExecutionNode()
+		: Super() { }
+
+	FDataflowExecutionNode(const Dataflow::FNodeParameters& Param, FGuid InGuid = FGuid::NewGuid())
+		: Super(Param,InGuid) {
+	}
+
+	/** FDataflowNode interface */
+	virtual ~FDataflowExecutionNode() {}
+	
+	static FName StaticType() { return FName("FDataflowExecutionNode"); }
+
+	virtual bool IsA(FName InType) const override 
+	{ 
+		return InType.ToString().Equals(StaticType().ToString()) 
+			|| Super::IsA(InType); 
+	}
+};
+
+/** Get the dataflow simulation time */
+USTRUCT()
+struct DATAFLOWSIMULATION_API FDataflowSimulationTime 
+{
+	GENERATED_BODY()
+
+public:
+	FDataflowSimulationTime(): DeltaTime(0.0f), CurrentTime(0.0f)
+	{}
+	
+	FDataflowSimulationTime(const float InDeltaTime, const float InCurrentTime) : DeltaTime(InDeltaTime), CurrentTime(InCurrentTime)
+	{}
+	
+	/** Delta time in seconds coming from the context */
+	UPROPERTY(Transient, SkipSerialization)
+	float DeltaTime = 0.0f;
+
+	/** Current time in seconds coming from the context */
+	UPROPERTY(Transient, SkipSerialization)
+	float CurrentTime = 0.0f;
+
+	/** Time offset that will need to be added to the current time to compute the sub-stepping */
+	UPROPERTY(Transient, SkipSerialization)
+	float TimeOffset = 0.0f;
+};
+
+/** Get the context simulation time */
+USTRUCT(meta = (DataflowSimulation))
+struct DATAFLOWSIMULATION_API FGetSimulationTimeDataflowNode : public FDataflowInvalidNode
+{
+	GENERATED_BODY()
+	DATAFLOW_NODE_DEFINE_INTERNAL(FGetSimulationTimeDataflowNode, "GetSimulationTime", "Physics|Common", UDataflow::SimulationTag)
+
+public:
+	
+	FGetSimulationTimeDataflowNode(const Dataflow::FNodeParameters& InParam, FGuid InGuid = FGuid::NewGuid())
+	: FDataflowInvalidNode(InParam, InGuid)
+	{
+		RegisterOutputConnection(&SimulationTime);
+	}
+	
+	/** Simulation time property coming from the context */
+	UPROPERTY(Transient, SkipSerialization, Meta = (DataflowOutput))
+	FDataflowSimulationTime SimulationTime = FDataflowSimulationTime(0.0f, 0.0f);
+
+	/** Evaluate simulation node given a simulation context */
+	virtual void EvaluateSimulation(UE::Dataflow::FDataflowSimulationContext& SimulationContext, const FDataflowOutput* Output) const override;
+};
+
+/** Get physics solvers from context */
+USTRUCT(meta = (DataflowSimulation))
+struct FGetPhysicsSolversDataflowNode : public FDataflowInvalidNode
+{
+	GENERATED_BODY()
+	DATAFLOW_NODE_DEFINE_INTERNAL(FGetPhysicsSolversDataflowNode, "GetPhysicsSolvers", "Physics|Solver", UDataflow::SimulationTag)
+
+public:
+	
+	FGetPhysicsSolversDataflowNode(const Dataflow::FNodeParameters& InParam, FGuid InGuid = FGuid::NewGuid())
+	: FDataflowInvalidNode(InParam, InGuid)
+	{
+		RegisterOutputConnection(&PhysicsSolvers);
+	}
+
+	/** Physics solvers coming from the context and filtered with the groups */
+	UPROPERTY(Transient, SkipSerialization, Meta = (DataflowOutput))
+	TArray<FDataflowSimulationProperty> PhysicsSolvers;
+
+	/** Simulation groups to filter the output solvers properties */
+	UPROPERTY(EditAnywhere, Category="Simulation")
+	TArray<FString> SimulationGroups;
+
+	/** Evaluate simulation node given a simulation context */
+	virtual void EvaluateSimulation(UE::Dataflow::FDataflowSimulationContext& SimulationContext, const FDataflowOutput* Output) const override;
+};
+
+/** Advance the simulation physics solver in time */
+USTRUCT(meta = (DataflowSimulation))
+struct FAdvancePhysicsSolversDataflowNode : public FDataflowSimulationNode
+{
+	GENERATED_BODY()
+	DATAFLOW_NODE_DEFINE_INTERNAL(FAdvancePhysicsSolversDataflowNode, "AdvancePhysicsSolvers", "Physics|Solver", UDataflow::SimulationTag)
+
+public:
+	
+	FAdvancePhysicsSolversDataflowNode(const Dataflow::FNodeParameters& InParam, FGuid InGuid = FGuid::NewGuid())
+	: FDataflowSimulationNode(InParam, InGuid)
+	{
+		RegisterInputConnection(&SimulationTime);
+		RegisterInputConnection(&PhysicsSolvers);
+		RegisterOutputConnection(&PhysicsSolvers, &PhysicsSolvers);
+	}
+
+	/** Delta time to use to advance the solver*/
+	UPROPERTY(Transient, SkipSerialization, Meta = (DataflowInput))
+	FDataflowSimulationTime SimulationTime = FDataflowSimulationTime(0.0f,0.0f);
+
+	/** Physics solvers to advance in time */
+	UPROPERTY(Transient, SkipSerialization, Meta = (DataflowInput, DataflowOutput, DataflowPassthrough = "PhysicsSolvers"))
+	TArray<FDataflowSimulationProperty> PhysicsSolvers;
+
+	/** Evaluate simulation node given a simulation context */
+	virtual void EvaluateSimulation(UE::Dataflow::FDataflowSimulationContext& SimulationContext, const FDataflowOutput* Output) const override;
+};
+
+/** Filter simulation proxies from context */
+USTRUCT(meta = (DataflowSimulation))
+struct FFilterSimulationProxiesDataflowNode : public FDataflowSimulationNode
+{
+	GENERATED_BODY()
+	DATAFLOW_NODE_DEFINE_INTERNAL(FFilterSimulationProxiesDataflowNode, "FilterSimulationProxies", "Physics|Proxy", UDataflow::SimulationTag)
+
+public:
+	
+	FFilterSimulationProxiesDataflowNode(const Dataflow::FNodeParameters& InParam, FGuid InGuid = FGuid::NewGuid())
+	: FDataflowSimulationNode(InParam, InGuid)
+	{
+		RegisterInputConnection(&SimulationProxies);
+		RegisterOutputConnection(&FilteredProxies);
+	}
+
+	/** Simulation proxies coming from the context and filtered with the groups */
+	UPROPERTY(Transient, SkipSerialization, Meta = (DataflowInput))
+	TArray<FDataflowSimulationProperty> SimulationProxies;
+ 
+	/** Simulation proxies coming from the context and filtered with the groups */
+	UPROPERTY(Transient, SkipSerialization, Meta = (DataflowOutput))
+	TArray<FDataflowSimulationProperty> FilteredProxies;
+
+	/** Simulation groups to filter the output solvers properties */
+	UPROPERTY(EditAnywhere, Category="Simulation")
+	TArray<FString> SimulationGroups;
+
+	/** Evaluate simulation node given a simulation context */
+	virtual void EvaluateSimulation(UE::Dataflow::FDataflowSimulationContext& SimulationContext, const FDataflowOutput* Output) const override;
+};
+
+/** Main terminal node for simulation proxies */
+USTRUCT(meta = (DataflowSimulation))
+struct FSimulationProxiesTerminalDataflowNode : public FDataflowExecutionNode
+{
+	GENERATED_BODY()
+	DATAFLOW_NODE_DEFINE_INTERNAL(FSimulationProxiesTerminalDataflowNode, "SimulationProxiesTerminal", "Terminal|Proxy", UDataflow::SimulationTag)
+
+public:
+	
+	FSimulationProxiesTerminalDataflowNode(const Dataflow::FNodeParameters& InParam, FGuid InGuid = FGuid::NewGuid())
+	: FDataflowExecutionNode(InParam, InGuid)
+	{
+		RegisterInputConnection(&SimulationProxies);
+	}
+
+	/** Physics solvers to evaluate  */
+	UPROPERTY(Transient, SkipSerialization, Meta = (DataflowInput))
+	TArray<FDataflowSimulationProperty> SimulationProxies;
+	
+	/** Evaluate simulation node given a simulation context */
+	virtual void EvaluateSimulation(UE::Dataflow::FDataflowSimulationContext& SimulationContext, const FDataflowOutput* Output) const override;
+};
+
+namespace Dataflow
+{
+	void RegisterDataflowSimulationNodes();
+}
+
+
diff --git a/Engine/Source/Runtime/Experimental/Dataflow/Simulation/Public/Dataflow/DataflowSimulationPlugin.h b/Engine/Source/Runtime/Experimental/Dataflow/Simulation/Public/Dataflow/DataflowSimulationPlugin.h
new file mode 100644
--- /dev/null
+++ b/Engine/Source/Runtime/Experimental/Dataflow/Simulation/Public/Dataflow/DataflowSimulationPlugin.h
@@ -0,0 +1,17 @@
+﻿// Copyright Epic Games, Inc. All Rights Reserved.
+
+#pragma once
+
+#include "CoreMinimal.h"
+#include "Modules/ModuleInterface.h"
+
+/**
+ * The public interface to this module
+ */
+class IDataflowSimulationPlugin : public IModuleInterface
+{
+public:
+	virtual void StartupModule();
+	virtual void ShutdownModule();
+};
+
diff --git a/Engine/Source/Runtime/Experimental/Dataflow/Simulation/Public/Dataflow/DataflowSimulationProxy.h b/Engine/Source/Runtime/Experimental/Dataflow/Simulation/Public/Dataflow/DataflowSimulationProxy.h
new file mode 100644
--- /dev/null
+++ b/Engine/Source/Runtime/Experimental/Dataflow/Simulation/Public/Dataflow/DataflowSimulationProxy.h
@@ -0,0 +1,70 @@
+﻿// Copyright Epic Games, Inc. All Rights Reserved.
+
+#pragma once
+
+#include "CoreMinimal.h"
+#include "DataflowSimulationProxy.generated.h"
+
+/**
+ * Dataflow simulation proxy used for simulation (PT)
+ */
+USTRUCT()
+struct DATAFLOWSIMULATION_API FDataflowSimulationProxy
+{
+	GENERATED_BODY()
+	
+	FDataflowSimulationProxy() = default;
+	virtual ~FDataflowSimulationProxy() = default;
+	
+	/** Check if the proxy is valid */
+	virtual bool IsValid() const { return true;}
+	
+	/** Get the proxy script struct */
+	virtual const UScriptStruct* GetScriptStruct() const
+	{
+		return StaticStruct();
+	}
+	
+	/** Cast the proxy to child */
+	template<class T>
+	const T* AsType() const
+	{
+		const UScriptStruct* ScriptStruct = GetScriptStruct();
+		if(ScriptStruct && ScriptStruct->IsChildOf(T::StaticStruct()))
+		{
+			return static_cast<T*>(this);
+		}
+		return nullptr;
+	}
+	
+	/** Cast the const proxy to child */
+	template<class T>
+	T* AsType()
+	{
+		const UScriptStruct* ScriptStruct = GetScriptStruct();
+		if(ScriptStruct && ScriptStruct->IsChildOf(T::StaticStruct()))
+		{
+			return static_cast<T*>(this);
+		}
+		return nullptr;
+	}
+
+	/** Get the simulation groups from the proxy */
+	const TSet<FString>& GetSimulationGroups() const {return SimulationGroups;};
+
+	/** Set the simulation groups onto the proxy */
+	void SetSimulationGroups(const TSet<FString>& InSimulationGroups);
+
+	/** Check if the given group is within the proxy simulation groups*/
+	bool HasSimulationGroup(const FString& SimulationGroup) const;
+
+	/** Check if the proxy has at least one valid bit */
+    bool HasGroupBit(const TBitArray<>& SimulationBits) const;
+
+	/** Bit array matching the simulation groups for fast access */
+	TBitArray<> GroupBits;
+
+private:
+	/** List of simulation groups this proxy belongs to */
+	TSet<FString> SimulationGroups;
+};
\ No newline at end of file
diff --git a/Engine/Source/Runtime/Experimental/Dataflow/Simulation/Public/Dataflow/DataflowSimulationUtils.h b/Engine/Source/Runtime/Experimental/Dataflow/Simulation/Public/Dataflow/DataflowSimulationUtils.h
new file mode 100644
--- /dev/null
+++ b/Engine/Source/Runtime/Experimental/Dataflow/Simulation/Public/Dataflow/DataflowSimulationUtils.h
@@ -0,0 +1,17 @@
+﻿// Copyright Epic Games, Inc. All Rights Reserved.
+
+#pragma once
+
+#include "Dataflow/DataflowObjectInterface.h"
+#include "Dataflow/DataflowSimulationContext.h"
+
+class UDataflow;
+
+namespace UE::Dataflow
+{
+	/** Evaluate the simulation graph given a simulation context and a simulation flag */
+	DATAFLOWSIMULATION_API void EvaluateSimulationGraph(const TObjectPtr<UDataflow>& SimulationGraph,
+		const TSharedPtr<UE::Dataflow::FDataflowSimulationContext>& SimulationContext, const float DeltaTime, const float SimulationTime);
+};
+
+
diff --git a/Engine/Source/Runtime/Experimental/Dataflow/Simulation/Public/Dataflow/Interfaces/DataflowCollisionObject.h b/Engine/Source/Runtime/Experimental/Dataflow/Simulation/Public/Dataflow/Interfaces/DataflowCollisionObject.h
new file mode 100644
--- /dev/null
+++ b/Engine/Source/Runtime/Experimental/Dataflow/Simulation/Public/Dataflow/Interfaces/DataflowCollisionObject.h
@@ -0,0 +1,46 @@
+﻿// Copyright Epic Games, Inc. All Rights Reserved.
+
+#pragma once
+#include "Dataflow/DataflowSimulationInterface.h"
+#include "Dataflow/DataflowSimulationProxy.h"
+#include "DataflowCollisionObject.generated.h"
+
+/**
+ * Dataflow collision object proxy (PT)
+ */
+USTRUCT()
+struct DATAFLOWSIMULATION_API FDataflowCollisionObjectProxy : public FDataflowSimulationProxy
+{
+	GENERATED_BODY()
+	
+	FDataflowCollisionObjectProxy();
+	virtual ~FDataflowCollisionObjectProxy() override = default;
+
+	/** Get the proxy script struct */
+	virtual const UScriptStruct* GetScriptStruct() const override
+	{
+		return StaticStruct();
+	}
+};
+
+UINTERFACE()
+class DATAFLOWSIMULATION_API UDataflowCollisionObjectInterface : public UDataflowSimulationInterface
+{
+	GENERATED_BODY()
+};
+
+/**
+ * Dataflow collision object interface to send/receive datas (GT <-> PT)
+ */
+class DATAFLOWSIMULATION_API IDataflowCollisionObjectInterface : public IDataflowSimulationInterface
+{
+	GENERATED_BODY()
+public:
+	IDataflowCollisionObjectInterface();
+
+	/** Get the simulation type */
+	virtual FString GetSimulationType() const override
+	{
+		return FDataflowCollisionObjectProxy::StaticStruct()->GetName();
+	}
+};
diff --git a/Engine/Source/Runtime/Experimental/Dataflow/Simulation/Public/Dataflow/Interfaces/DataflowConstraintObject.h b/Engine/Source/Runtime/Experimental/Dataflow/Simulation/Public/Dataflow/Interfaces/DataflowConstraintObject.h
new file mode 100644
--- /dev/null
+++ b/Engine/Source/Runtime/Experimental/Dataflow/Simulation/Public/Dataflow/Interfaces/DataflowConstraintObject.h
@@ -0,0 +1,47 @@
+﻿// Copyright Epic Games, Inc. All Rights Reserved.
+
+#pragma once
+#include "Dataflow/DataflowSimulationInterface.h"
+#include "Dataflow/DataflowSimulationProxy.h"
+#include "DataflowConstraintObject.generated.h"
+
+/**
+ * Dataflow collision object proxy (PT)
+ */
+USTRUCT()
+struct DATAFLOWSIMULATION_API FDataflowConstraintObjectProxy : public FDataflowSimulationProxy
+{
+	GENERATED_BODY()
+	
+	FDataflowConstraintObjectProxy();
+	virtual ~FDataflowConstraintObjectProxy() override = default;
+
+	/** Get the proxy script struct */
+	virtual const UScriptStruct* GetScriptStruct() const override
+	{
+		return StaticStruct();
+	}
+};
+
+UINTERFACE()
+class DATAFLOWSIMULATION_API UDataflowConstraintObjectInterface : public UDataflowSimulationInterface
+{
+	GENERATED_BODY()
+};
+
+/**
+ * Dataflow collision object interface to send/receive datas (GT <-> PT)
+ */
+class DATAFLOWSIMULATION_API IDataflowConstraintObjectInterface : public IDataflowSimulationInterface
+{
+	GENERATED_BODY()
+	
+public:
+	IDataflowConstraintObjectInterface();
+
+	/** Get the simulation type */
+	virtual FString GetSimulationType() const override
+	{
+		return FDataflowConstraintObjectProxy::StaticStruct()->GetName();
+	}
+};
\ No newline at end of file
diff --git a/Engine/Source/Runtime/Experimental/Dataflow/Simulation/Public/Dataflow/Interfaces/DataflowPhysicsObject.h b/Engine/Source/Runtime/Experimental/Dataflow/Simulation/Public/Dataflow/Interfaces/DataflowPhysicsObject.h
new file mode 100644
--- /dev/null
+++ b/Engine/Source/Runtime/Experimental/Dataflow/Simulation/Public/Dataflow/Interfaces/DataflowPhysicsObject.h
@@ -0,0 +1,47 @@
+﻿// Copyright Epic Games, Inc. All Rights Reserved.
+
+#pragma once
+#include "Dataflow/DataflowSimulationInterface.h"
+#include "Dataflow/DataflowSimulationProxy.h"
+#include "DataflowPhysicsObject.generated.h"
+
+/**
+ * Dataflow physics object proxy (PT)
+ */
+USTRUCT()
+struct DATAFLOWSIMULATION_API FDataflowPhysicsObjectProxy : public FDataflowSimulationProxy
+{
+	GENERATED_BODY()
+	
+	FDataflowPhysicsObjectProxy();
+	virtual ~FDataflowPhysicsObjectProxy() override = default;
+
+	/** Get the proxy script struct */
+	virtual const UScriptStruct* GetScriptStruct() const override
+	{
+		return StaticStruct();
+	}
+};
+
+UINTERFACE()
+class DATAFLOWSIMULATION_API UDataflowPhysicsObjectInterface : public UDataflowSimulationInterface
+{
+	GENERATED_BODY()
+};
+
+/**
+ * Dataflow physics object interface to send/receive datas (GT <-> PT)
+ */
+class DATAFLOWSIMULATION_API IDataflowPhysicsObjectInterface : public IDataflowSimulationInterface
+{
+	GENERATED_BODY()
+	
+public:
+	IDataflowPhysicsObjectInterface();
+
+	/** Get the simulation type */
+	virtual FString GetSimulationType() const override
+	{
+		return FDataflowPhysicsObjectProxy::StaticStruct()->GetName();
+	}
+};
diff --git a/Engine/Source/Runtime/Experimental/Dataflow/Simulation/Public/Dataflow/Interfaces/DataflowPhysicsSolver.h b/Engine/Source/Runtime/Experimental/Dataflow/Simulation/Public/Dataflow/Interfaces/DataflowPhysicsSolver.h
new file mode 100644
--- /dev/null
+++ b/Engine/Source/Runtime/Experimental/Dataflow/Simulation/Public/Dataflow/Interfaces/DataflowPhysicsSolver.h
@@ -0,0 +1,63 @@
+﻿// Copyright Epic Games, Inc. All Rights Reserved.
+
+#pragma once
+#include "Dataflow/DataflowSimulationInterface.h"
+#include "Dataflow/DataflowSimulationProxy.h"
+#include "DataflowPhysicsSolver.generated.h"
+
+/**
+ * Dataflow simulation physics solver proxy (PT)
+ */
+USTRUCT()
+struct DATAFLOWSIMULATION_API FDataflowPhysicsSolverProxy : public FDataflowSimulationProxy
+{
+	GENERATED_BODY()
+
+	FDataflowPhysicsSolverProxy();
+	virtual ~FDataflowPhysicsSolverProxy() override = default;
+
+	/** Get the proxy script struct */
+	virtual const UScriptStruct* GetScriptStruct() const override
+	{
+		return StaticStruct();
+	}
+	
+	/** Advance the solver datas in time */
+	virtual void AdvanceSolverDatas(const float DeltaTime) {}
+	
+	/** Pre-solve the proxy */
+	virtual void PreSolveProxy(const float DeltaTime) {}
+
+	/** Post-solve the proxy */
+	virtual void PostSolveProxy(const float DeltaTime) {}
+
+	/** Get the solver time step */
+	virtual float GetTimeStep() { return 0.033f;}
+};
+
+UINTERFACE()
+class DATAFLOWSIMULATION_API UDataflowPhysicsSolverInterface : public UDataflowSimulationInterface
+{
+	GENERATED_BODY()
+};
+
+/**
+ * Dataflow physics solver interface to send/receive (GT)
+ */
+class DATAFLOWSIMULATION_API IDataflowPhysicsSolverInterface : public IDataflowSimulationInterface
+{
+	GENERATED_BODY()
+	
+public:
+	IDataflowPhysicsSolverInterface();
+	
+	/** Get the simulation type */
+	virtual FString GetSimulationType() const override
+	{
+		return FDataflowPhysicsSolverProxy::StaticStruct()->GetName();
+	}
+};
+
+
+
+
