Chaos Cloth: Fix lift/drag weight maps not working correctly when not using new outer lift and drag.

diff --git a/Engine/Source/Runtime/Experimental/Chaos/Private/Chaos/VelocityField.cpp b/Engine/Source/Runtime/Experimental/Chaos/Private/Chaos/VelocityField.cpp
--- a/Engine/Source/Runtime/Experimental/Chaos/Private/Chaos/VelocityField.cpp
+++ b/Engine/Source/Runtime/Experimental/Chaos/Private/Chaos/VelocityField.cpp
@@ -44,10 +44,21 @@ void FVelocityAndPressureField::SetProperties(
 				TConstArrayView<TVec3<int32>>(Elements),
 				Offset,
 				NumParticles);
+
+			if (OuterDragIndex == INDEX_NONE)
+			{
+				// OuterDrag is not specified, so it should match Drag.
+				OuterDrag = Drag;
+			}
 		}
 		else
 		{
 			Drag.SetWeightedValue(WeightedValue.ClampAxes(MinCoefficient, MaxCoefficient));
+			if (OuterDragIndex == INDEX_NONE)
+			{
+				// OuterDrag is not specified, so it should match Drag.
+				OuterDrag.SetWeightedValue(WeightedValue.ClampAxes(MinCoefficient, MaxCoefficient));
+			}
 		}
 	}
 
@@ -83,10 +94,21 @@ void FVelocityAndPressureField::SetProperties(
 				TConstArrayView<TVec3<int32>>(Elements),
 				Offset,
 				NumParticles);
+
+			if (OuterLiftIndex == INDEX_NONE)
+			{
+				// OuterLift is not specified, so it should match Lift.
+				OuterLift = Lift;
+			}
 		}
 		else
 		{
 			Lift.SetWeightedValue(WeightedValue.ClampAxes(MinCoefficient, MaxCoefficient));
+			if (OuterLiftIndex == INDEX_NONE)
+			{
+				// OuterLift is not specified, so it should match Lift.
+				OuterLift.SetWeightedValue(WeightedValue.ClampAxes(MinCoefficient, MaxCoefficient));
+			}
 		}
 	}
 
@@ -352,11 +374,11 @@ void FVelocityAndPressureField::SetMultipliers(
 	const TConstArrayView<FRealSingle> DragMultipliers = (DragIndex != INDEX_NONE) ?
 		Weightmaps.FindRef(GetDragString(PropertyCollection)) : TConstArrayView<FRealSingle>();
 	const TConstArrayView<FRealSingle> OuterDragMultipliers = (OuterDragIndex != INDEX_NONE) ?
-		Weightmaps.FindRef(GetOuterDragString(PropertyCollection)) : TConstArrayView<FRealSingle>();
+		Weightmaps.FindRef(GetOuterDragString(PropertyCollection)) : DragMultipliers; // OuterDrag defaults to Drag
 	const TConstArrayView<FRealSingle> LiftMultipliers = (LiftIndex != INDEX_NONE) ?
 		Weightmaps.FindRef(GetLiftString(PropertyCollection)) : TConstArrayView<FRealSingle>();
 	const TConstArrayView<FRealSingle> OuterLiftMultipliers = (OuterLiftIndex != INDEX_NONE) ?
-		Weightmaps.FindRef(GetOuterLiftString(PropertyCollection)) : TConstArrayView<FRealSingle>();
+		Weightmaps.FindRef(GetOuterLiftString(PropertyCollection)) : LiftMultipliers; // OuterLift defaults to Lift
 	const TConstArrayView<FRealSingle> PressureMultipliers = (PressureIndex != INDEX_NONE) ?
 		Weightmaps.FindRef(GetPressureString(PropertyCollection)) : TConstArrayView<FRealSingle>();
 
