Dataflow : Better tooltip support for AnyTypes as well as convenient function to change a connection type

diff --git a/Engine/Source/Runtime/Experimental/Dataflow/Core/Private/Dataflow/DataflowConnection.cpp b/Engine/Source/Runtime/Experimental/Dataflow/Core/Private/Dataflow/DataflowConnection.cpp
--- a/Engine/Source/Runtime/Experimental/Dataflow/Core/Private/Dataflow/DataflowConnection.cpp
+++ b/Engine/Source/Runtime/Experimental/Dataflow/Core/Private/Dataflow/DataflowConnection.cpp
@@ -144,9 +144,49 @@ void FDataflowConnection::ForceSimpleType(FName InType)
 
 void FDataflowConnection::FixAndPropagateType()
 {
+	check(Property);
 	FString ExtendedType;
 	const FString CPPType = Property->GetCPPType(&ExtendedType);
 	FName FixedType(CPPType + ExtendedType);
 
 	FixAndPropagateType(FixedType);
+}
+
+FString FDataflowConnection::GetPropertyTooltip() const
+{
+#if WITH_EDITORONLY_DATA
+	check(Property);
+	return Property->GetToolTipText().ToString();
+#else
+	return {};
+#endif // WITH_EDITORONLY_DATA
+}
+
+FString FDataflowConnection::GetPropertyTypeNameTooltip() const
+{
+#if WITH_EDITORONLY_DATA
+	FString TypeNameStr = Type.ToString();
+	if (bIsAnyType)
+	{
+		if (!HasConcreteType())
+		{
+			check(Property);
+			TypeNameStr = TEXT("Wildcard");
+		}
+		if (Property->GetClass()->IsChildOf(FStructProperty::StaticClass()))
+		{
+			if (const FStructProperty* StructProperty = CastField<FStructProperty>(Property))
+			{
+				if (StructProperty->Struct && StructProperty->Struct->IsChildOf<FDataflowAnyType>())
+				{
+					TypeNameStr += TEXT("\n");
+					TypeNameStr += StructProperty->Struct->GetToolTipText().ToString();
+				}
+			}
+		}
+	}
+	return TypeNameStr;
+#else
+	return {};
+#endif
 }
\ No newline at end of file
diff --git a/Engine/Source/Runtime/Experimental/Dataflow/Core/Private/Dataflow/DataflowNode.cpp b/Engine/Source/Runtime/Experimental/Dataflow/Core/Private/Dataflow/DataflowNode.cpp
--- a/Engine/Source/Runtime/Experimental/Dataflow/Core/Private/Dataflow/DataflowNode.cpp
+++ b/Engine/Source/Runtime/Experimental/Dataflow/Core/Private/Dataflow/DataflowNode.cpp
@@ -1085,14 +1085,22 @@ PRAGMA_ENABLE_DEPRECATION_WARNINGS
 	return bHasValidConnections;
 }
 
-FString FDataflowNode::GetToolTip()
+TUniquePtr<const FStructOnScope> FDataflowNode::NewStructOnScopeConst() const
+{
+	// New StructOnScopeConst is non-const and virtual , changing it would be quite difficult
+	// we have to use a const cast though to accomodate this current constraint
+	// but we make sure to return a const version of the FStructOnScope 
+	return TUniquePtr<const FStructOnScope>(const_cast<FDataflowNode*>(this)->NewStructOnScope());
+}
+
+FString FDataflowNode::GetToolTip() const
 {
 	Dataflow::FFactoryParameters FactoryParameters = ::Dataflow::FNodeFactory::GetInstance()->GetParameters(GetType());
 
 	return FactoryParameters.ToolTip;
 }
 
-FText FDataflowNode::GetPinDisplayName(const FName& PropertyFullName, const Dataflow::FPin::EDirection Direction)
+FText FDataflowNode::GetPinDisplayName(const FName& PropertyFullName, const Dataflow::FPin::EDirection Direction) const
 {
 	int32 ContainerIndex = INDEX_NONE;
 
@@ -1111,7 +1119,7 @@ FText FDataflowNode::GetPinDisplayName(const FName& PropertyFullName, const Data
 		}
 	}
 
-	if (const TUniquePtr<FStructOnScope> ScriptOnStruct = TUniquePtr<FStructOnScope>(NewStructOnScope()))
+	if (const TUniquePtr<const FStructOnScope> ScriptOnStruct = NewStructOnScopeConst())
 	{
 		if (const UStruct* const Struct = ScriptOnStruct->GetStruct())
 		{
@@ -1126,7 +1134,7 @@ FText FDataflowNode::GetPinDisplayName(const FName& PropertyFullName, const Data
 	return FText();
 }
 
-FString FDataflowNode::GetPinToolTip(const FName& PropertyFullName, const Dataflow::FPin::EDirection Direction)
+FString FDataflowNode::GetPinToolTip(const FName& PropertyFullName, const Dataflow::FPin::EDirection Direction) const
 {
 #if WITH_EDITORONLY_DATA
 	if (Direction == Dataflow::FPin::EDirection::INPUT)
@@ -1149,7 +1157,7 @@ FString FDataflowNode::GetPinToolTip(const FName& PropertyFullName, const Datafl
 			}
 		}
 	}
-	else if (const TUniquePtr<FStructOnScope> ScriptOnStruct = TUniquePtr<FStructOnScope>(NewStructOnScope()))
+	else if (const TUniquePtr<const FStructOnScope> ScriptOnStruct = NewStructOnScopeConst())
 	{
 		if (const UStruct* const Struct = ScriptOnStruct->GetStruct())
 		{
@@ -1161,10 +1169,10 @@ FString FDataflowNode::GetPinToolTip(const FName& PropertyFullName, const Datafl
 	}
 #endif
 
-	return "";
+	return {};
 }
 
-TArray<FString> FDataflowNode::GetPinMetaData(const FName& PropertyFullName, const Dataflow::FPin::EDirection Direction)
+TArray<FString> FDataflowNode::GetPinMetaData(const FName& PropertyFullName, const Dataflow::FPin::EDirection Direction) const
 {
 #if WITH_EDITORONLY_DATA
 	if (Direction == Dataflow::FPin::EDirection::INPUT)
@@ -1187,7 +1195,7 @@ TArray<FString> FDataflowNode::GetPinMetaData(const FName& PropertyFullName, con
 			}
 		}
 	}
-	else if (const TUniquePtr<FStructOnScope> ScriptOnStruct = TUniquePtr<FStructOnScope>(NewStructOnScope()))
+	else if (const TUniquePtr<const FStructOnScope> ScriptOnStruct = NewStructOnScopeConst())
 	{
 		if (const UStruct* const Struct = ScriptOnStruct->GetStruct())
 		{
@@ -1286,3 +1294,29 @@ bool FDataflowNode::SetOutputConcreteType(const Dataflow::FConnectionReference&
 	}
 	return false;
 }
+
+bool FDataflowNode::SetAllConnectionConcreteType(FName NewType)
+{
+	bool bChanged = false;
+	for (const TPair<Dataflow::FConnectionKey, FDataflowInput*>& InputEntry : ExpandedInputs)
+	{
+		FDataflowInput* const Input = InputEntry.Value;
+		if (Input && Input->GetType() != NewType)
+		{
+			bChanged |= Input->SetConcreteType(NewType);
+		}
+	}
+
+	PRAGMA_DISABLE_DEPRECATION_WARNINGS  // Until Outputs becomes private
+	for (const TPair<int32, FDataflowOutput*>& OutputEntry : Outputs)
+	PRAGMA_ENABLE_DEPRECATION_WARNINGS
+	{
+		FDataflowOutput* const Output = OutputEntry.Value;
+		if (Output && Output->GetType() != NewType)
+		{
+			bChanged |= Output->SetConcreteType(NewType);
+		}
+	}
+
+	return bChanged;
+}
\ No newline at end of file
diff --git a/Engine/Source/Runtime/Experimental/Dataflow/Core/Public/Dataflow/DataflowConnection.h b/Engine/Source/Runtime/Experimental/Dataflow/Core/Public/Dataflow/DataflowConnection.h
--- a/Engine/Source/Runtime/Experimental/Dataflow/Core/Public/Dataflow/DataflowConnection.h
+++ b/Engine/Source/Runtime/Experimental/Dataflow/Core/Public/Dataflow/DataflowConnection.h
@@ -187,6 +187,9 @@ public:
 
 	const FProperty* GetProperty() const { return Property; }
 
+	DATAFLOWCORE_API FString GetPropertyTooltip() const;
+	DATAFLOWCORE_API FString GetPropertyTypeNameTooltip() const;
+
 	Dataflow::FPin::EDirection GetDirection() const { return Direction; }
 	uint32 GetOffset() const { return Offset; }
 	virtual int32 GetContainerIndex() const { return INDEX_NONE; }
diff --git a/Engine/Source/Runtime/Experimental/Dataflow/Core/Public/Dataflow/DataflowNode.h b/Engine/Source/Runtime/Experimental/Dataflow/Core/Public/Dataflow/DataflowNode.h
--- a/Engine/Source/Runtime/Experimental/Dataflow/Core/Public/Dataflow/DataflowNode.h
+++ b/Engine/Source/Runtime/Experimental/Dataflow/Core/Public/Dataflow/DataflowNode.h
@@ -90,10 +90,10 @@ PRAGMA_ENABLE_DEPRECATION_WARNINGS
 	virtual FName GetDisplayName() const { return ""; }
 	virtual FName GetCategory() const { return ""; }
 	virtual FString GetTags() const { return ""; }
-	DATAFLOWCORE_API virtual FString GetToolTip();
-	DATAFLOWCORE_API FString GetPinToolTip(const FName& PropertyName, const Dataflow::FPin::EDirection Direction = Dataflow::FPin::EDirection::NONE);
-	DATAFLOWCORE_API FText GetPinDisplayName(const FName& PropertyName, const Dataflow::FPin::EDirection Direction = Dataflow::FPin::EDirection::NONE);
-	DATAFLOWCORE_API TArray<FString> GetPinMetaData(const FName& PropertyName, const Dataflow::FPin::EDirection Direction = Dataflow::FPin::EDirection::NONE);
+	DATAFLOWCORE_API virtual FString GetToolTip() const;
+	DATAFLOWCORE_API FString GetPinToolTip(const FName& PropertyName, const Dataflow::FPin::EDirection Direction = Dataflow::FPin::EDirection::NONE) const;
+	DATAFLOWCORE_API FText GetPinDisplayName(const FName& PropertyName, const Dataflow::FPin::EDirection Direction = Dataflow::FPin::EDirection::NONE) const;
+	DATAFLOWCORE_API TArray<FString> GetPinMetaData(const FName& PropertyName, const Dataflow::FPin::EDirection Direction = Dataflow::FPin::EDirection::NONE) const;
 	virtual TArray<Dataflow::FRenderingParameter> GetRenderParameters() const { return GetRenderParametersImpl(); }
 	// Copy node property values from another node
 	UE_DEPRECATED(5.4, "FDataflowNode::CopyNodeProperties is deprecated.")
@@ -225,6 +225,8 @@ PRAGMA_ENABLE_DEPRECATION_WARNINGS
 	virtual FStructOnScope* NewStructOnScope() { return nullptr; }
 	virtual const UScriptStruct* TypedScriptStruct() const { return nullptr; }
 
+	TUniquePtr<const FStructOnScope> NewStructOnScopeConst() const;
+
 	/** Register the Input and Outputs after the creation in the factory. Use PropertyName to disambiguate a struct name from its first property. */
 	template<typename T>
 	FDataflowInput& RegisterInputConnection(const Dataflow::TConnectionReference<T>& Reference, const FName& PropertyName = NAME_None)
@@ -565,9 +567,24 @@ protected:
 	// returns true if the input type was changed successfully
 	DATAFLOWCORE_API bool SetInputConcreteType(const Dataflow::FConnectionReference& InputReference, FName NewType);
 
+	template <typename T>
+	bool SetInputConcreteType(const Dataflow::FConnectionReference& InputReference)
+	{
+		return SetInputConcreteType(InputReference, TDataflowPolicyTypeName<T>::GetName());
+	}
+
 	// returns true if the output type was changed successfully
 	DATAFLOWCORE_API bool SetOutputConcreteType(const Dataflow::FConnectionReference& OutputReference, FName NewType);
 
+	template <typename T>
+	bool SetOutputConcreteType(const Dataflow::FConnectionReference& OutputReference)
+	{
+		return SetOutputConcreteType(OutputReference, TDataflowPolicyTypeName<T>::GetName());
+	}
+
+	// returns true if any of the types was changed successfully
+	DATAFLOWCORE_API bool SetAllConnectionConcreteType(FName NewType);
+
 	DATAFLOWCORE_API FDataflowInput& RegisterInputConnectionInternal(const Dataflow::FConnectionReference& Reference, const FName& PropertyName = NAME_None);
 	DATAFLOWCORE_API FDataflowOutput& RegisterOutputConnectionInternal(const Dataflow::FConnectionReference& Reference, const FName& PropertyName = NAME_None);
 	DATAFLOWCORE_API FDataflowInput& RegisterInputArrayConnectionInternal(const Dataflow::FConnectionReference& Reference, const FName& ElementPropertyName = NAME_None,
diff --git a/Engine/Source/Runtime/Experimental/Dataflow/Engine/Private/Dataflow/DataflowEdNode.cpp b/Engine/Source/Runtime/Experimental/Dataflow/Engine/Private/Dataflow/DataflowEdNode.cpp
--- a/Engine/Source/Runtime/Experimental/Dataflow/Engine/Private/Dataflow/DataflowEdNode.cpp
+++ b/Engine/Source/Runtime/Experimental/Dataflow/Engine/Private/Dataflow/DataflowEdNode.cpp
@@ -462,30 +462,18 @@ FLinearColor UDataflowEdNode::GetNodeTitleColor() const
 
 FLinearColor UDataflowEdNode::GetNodeBodyTintColor() const
 {
-	if (DataflowGraph)
+	if (const TSharedPtr<const FDataflowNode> DataflowNode = GetDataflowNode())
 	{
-		if (DataflowNodeGuid.IsValid())
-		{
-			if (TSharedPtr<FDataflowNode> DataflowNode = DataflowGraph->FindBaseNode(DataflowNodeGuid))
-			{
-				return Dataflow::FNodeColorsRegistry::Get().GetNodeBodyTintColor(DataflowNode->GetCategory());
-			}
-		}
+		Dataflow::FNodeColorsRegistry::Get().GetNodeBodyTintColor(DataflowNode->GetCategory());
 	}
 	return FDataflowNode::DefaultNodeBodyTintColor;
 }
 
 FText UDataflowEdNode::GetTooltipText() const
 {
-	if (DataflowGraph)
+	if (const TSharedPtr<const FDataflowNode> DataflowNode = GetDataflowNode()) 
 	{
-		if (DataflowNodeGuid.IsValid())
-		{
-			if (TSharedPtr<FDataflowNode> DataflowNode = DataflowGraph->FindBaseNode(DataflowNodeGuid))
-			{
-				return FText::FromString(DataflowNode->GetToolTip());
-			}
-		}
+		return FText::FromString(DataflowNode->GetToolTip());
 	}
 
 	return FText::FromString("");
@@ -494,17 +482,14 @@ FText UDataflowEdNode::GetTooltipText() const
 
 FText UDataflowEdNode::GetPinDisplayName(const UEdGraphPin* Pin) const
 {
-	if (Pin && DataflowGraph)
+	if (Pin)
 	{
-		if (DataflowNodeGuid.IsValid())
+		if (const TSharedPtr<const FDataflowNode> DataflowNode = GetDataflowNode())
 		{
-			if (TSharedPtr<FDataflowNode> DataflowNode = DataflowGraph->FindBaseNode(DataflowNodeGuid))
+			const FText DisplayName = DataflowNode->GetPinDisplayName(Pin->PinName, Dataflow::Private::EdPinDirectionToDataflowDirection(Pin->Direction));
+			if (!DisplayName.IsEmpty())
 			{
-				const FText DisplayName = DataflowNode->GetPinDisplayName(Pin->PinName, Dataflow::Private::EdPinDirectionToDataflowDirection(Pin->Direction));
-				if (!DisplayName.IsEmpty())
-				{
-					return DisplayName;
-				}
+				return DisplayName;
 			}
 		}
 	}
@@ -513,36 +498,58 @@ FText UDataflowEdNode::GetPinDisplayName(const UEdGraphPin* Pin) const
 
 void UDataflowEdNode::GetPinHoverText(const UEdGraphPin& Pin, FString& HoverTextOut) const
 {
-	if (DataflowGraph)
+	if (const TSharedPtr<const FDataflowNode> DataflowNode = GetDataflowNode())
 	{
-		if (DataflowNodeGuid.IsValid())
-		{
-			if (TSharedPtr<FDataflowNode> DataflowNode = DataflowGraph->FindBaseNode(DataflowNodeGuid))
-			{		
-				FString MetaDataStr;
-				
-				TArray<FString> PinMetaData = DataflowNode->GetPinMetaData(Pin.PinName, Dataflow::Private::EdPinDirectionToDataflowDirection(Pin.Direction));
+		const Dataflow::FPin::EDirection PinDirection = Dataflow::Private::EdPinDirectionToDataflowDirection(Pin.Direction);
+
+		FString MetaDataStr;
+		TArray<FString> PinMetaData = DataflowNode->GetPinMetaData(Pin.PinName, PinDirection);
 			
-				if (Pin.Direction == EGPD_Input && PinMetaData.Contains(FDataflowNode::DataflowIntrinsic.ToString()))
-				{
-					MetaDataStr = "[Intrinsic]";
-				}
-				if (Pin.Direction == EGPD_Output && PinMetaData.Contains(FDataflowNode::DataflowPassthrough.ToString()))
-				{
-					MetaDataStr = "[Passthrough]";
-				}
+		if (Pin.Direction == EGPD_Input && PinMetaData.Contains(FDataflowNode::DataflowIntrinsic.ToString()))
+		{
+			MetaDataStr = "[Intrinsic]";
+		}
+		if (Pin.Direction == EGPD_Output && PinMetaData.Contains(FDataflowNode::DataflowPassthrough.ToString()))
+		{
+			MetaDataStr = "[Passthrough]";
+		}
 
-				FString NameStr = Pin.PinName.ToString();
-				if (MetaDataStr.Len() > 0)
-				{
-					NameStr.Appendf(TEXT(" %s"), *MetaDataStr);
-				}
+		FString NameStr = Pin.PinName.ToString();
+		if (MetaDataStr.Len() > 0)
+		{
+			NameStr.Appendf(TEXT(" %s"), *MetaDataStr);
+		}
 
-				HoverTextOut.Appendf(TEXT("%s\n%s\n\n%s"), *NameStr, *Pin.PinType.PinCategory.ToString(),
-					*DataflowNode->GetPinToolTip(Pin.PinName, Dataflow::Private::EdPinDirectionToDataflowDirection(Pin.Direction)));
-			}
+		// find type information 
+		FString TypeNameStr = Pin.PinType.PinCategory.ToString();
+
+		const FDataflowConnection* Connection = nullptr;
+		if (Pin.Direction == EGPD_Input)
+		{
+			Connection = DataflowNode->FindInput(Pin.PinName);
 		}
-	}
+		else if (Pin.Direction == EGPD_Output)
+		{
+			Connection = DataflowNode->FindOutput(Pin.PinName);
+		}
+		if (Connection)
+		{
+			TypeNameStr = Connection->GetPropertyTypeNameTooltip();
+		}
+
+		const FString PropertyTooltip = DataflowNode->GetPinToolTip(Pin.PinName, PinDirection);
+
+		// put it all together 
+		if (PropertyTooltip.IsEmpty())
+		{
+			HoverTextOut.Appendf(TEXT("%s\n%s"), *NameStr, *TypeNameStr);
+		}
+		else
+		{
+			HoverTextOut.Appendf(TEXT("%s\n%s\n\n%s"), *NameStr, *TypeNameStr, *PropertyTooltip);
+		}
+		
+}
 }
 
 void UDataflowEdNode::AutowireNewNode(UEdGraphPin* FromPin)
