[NetPhysics] Allow RewindData to run with variable tick physics

diff --git a/Engine/Source/Runtime/Experimental/Chaos/Private/PBDRigidsSolver.cpp b/Engine/Source/Runtime/Experimental/Chaos/Private/PBDRigidsSolver.cpp
--- a/Engine/Source/Runtime/Experimental/Chaos/Private/PBDRigidsSolver.cpp
+++ b/Engine/Source/Runtime/Experimental/Chaos/Private/PBDRigidsSolver.cpp
@@ -1102,12 +1102,6 @@ namespace Chaos
 
 	void FPBDRigidsSolver::EnableRewindCapture(int32 NumFrames, bool InUseCollisionResimCache)
 	{
-		if (!IsUsingFixedDt())
-		{
-			UE_LOG(LogChaos, Warning, TEXT("PBDRigidsSolver::EnableRewindCapture Failed - Physics need to run with fixed tick. Project Settings -> Physics -> Framerate -> Tick Physics Async needs to be enabled."));
-			return;
-		}
-
 		//TODO: this function calls both internal and external - sort of assumed during initialization. Should decide what thread it's called on and mark it as either external or internal
 		if (MRewindData.IsValid())
 		{
