Iris
- Fixed crash in IrisPackageMapExportUtil, was caused by use of garbage memory after resizing array.
- Zero out memory on resize in NetSerializerArrayStorage
- Changed from using NameNetSerializer to NameAsNetTokenNetSerializer for exported names
- Renabled net.iris.EnableIrisPackageMapNameExports 1

diff --git a/Engine/Source/Runtime/Experimental/Iris/Core/Private/Iris/Serialization/IrisObjectReferencePackageMap.cpp b/Engine/Source/Runtime/Experimental/Iris/Core/Private/Iris/Serialization/IrisObjectReferencePackageMap.cpp
--- a/Engine/Source/Runtime/Experimental/Iris/Core/Private/Iris/Serialization/IrisObjectReferencePackageMap.cpp
+++ b/Engine/Source/Runtime/Experimental/Iris/Core/Private/Iris/Serialization/IrisObjectReferencePackageMap.cpp
@@ -7,7 +7,7 @@ namespace UE
 {
 	namespace Net
 	{
-		bool bEnableIrisPackageMapNameExports = false;
+		bool bEnableIrisPackageMapNameExports = true;
 		static FAutoConsoleVariableRef CVarEnableIrisPackageMapNameExports(TEXT("net.iris.EnableIrisPackageMapNameExports"), bEnableIrisPackageMapNameExports, TEXT("If enabled, iris captures and exports fnames when calling into old serialziation code instead of serializing a strings."));
 	}
 }
diff --git a/Engine/Source/Runtime/Experimental/Iris/Core/Private/Iris/Serialization/IrisPackageMapExportUtil.cpp b/Engine/Source/Runtime/Experimental/Iris/Core/Private/Iris/Serialization/IrisPackageMapExportUtil.cpp
--- a/Engine/Source/Runtime/Experimental/Iris/Core/Private/Iris/Serialization/IrisPackageMapExportUtil.cpp
+++ b/Engine/Source/Runtime/Experimental/Iris/Core/Private/Iris/Serialization/IrisPackageMapExportUtil.cpp
@@ -17,7 +17,8 @@ namespace UE::Net
 {
 
 const FNetSerializer* FIrisPackageMapExportsUtil::ObjectNetSerializer = &UE_NET_GET_SERIALIZER(FObjectNetSerializer);
-const FNetSerializer* FIrisPackageMapExportsUtil::NameNetSerializer = &UE_NET_GET_SERIALIZER(FNameNetSerializer);
+const FNetSerializer* FIrisPackageMapExportsUtil::NameNetSerializer = &UE_NET_GET_SERIALIZER(FNameAsNetTokenNetSerializer);
+//const FNetSerializer* FIrisPackageMapExportsUtil::NameNetSerializer = &UE_NET_GET_SERIALIZER(FNameNetSerializer);
 
 void FIrisPackageMapExportsUtil::Serialize(FNetSerializationContext& Context, const QuantizedType& Value)
 {
@@ -25,6 +26,7 @@ void FIrisPackageMapExportsUtil::Serialize(FNetSerializationContext& Context, co
 
 	// If we have any references, export them!
 	{
+		UE_NET_TRACE_SCOPE(ObjectReferences, *Writer, Context.GetTraceCollector(), ENetTraceVerbosity::VeryVerbose);
 		const uint32 NumReferences = Value.ObjectReferenceStorage.Num();
 		if (Writer->WriteBool(NumReferences != 0))
 		{
@@ -43,6 +45,7 @@ void FIrisPackageMapExportsUtil::Serialize(FNetSerializationContext& Context, co
 
 	// If we have any names, export them!
 	{
+		UE_NET_TRACE_SCOPE(Names, *Writer, Context.GetTraceCollector(), ENetTraceVerbosity::VeryVerbose);
 		const uint32 NumNames = Value.NameStorage.Num();
 		if (Writer->WriteBool(NumNames != 0))
 		{
@@ -61,6 +64,7 @@ void FIrisPackageMapExportsUtil::Serialize(FNetSerializationContext& Context, co
 
 	// For NetTokens
 	{
+		UE_NET_TRACE_SCOPE(NetTokens, *Writer, Context.GetTraceCollector(), ENetTraceVerbosity::VeryVerbose);
 		// $TODO: For asymmetrical data like nettokens we will need to do something special to support validating the default state hash as the local quantized state will differ
 		// For now we ignore this in default state hash
 		if (Context.IsInitializingDefaultState())
@@ -89,6 +93,7 @@ void FIrisPackageMapExportsUtil::Deserialize(FNetSerializationContext& Context,
 
 	// Read any object references
 	{
+		UE_NET_TRACE_SCOPE(ObjectReferences, *Reader, Context.GetTraceCollector(), ENetTraceVerbosity::VeryVerbose);
 		const bool bHasObjectReferences = Reader->ReadBool();
 		if (bHasObjectReferences)
 		{
@@ -121,6 +126,7 @@ void FIrisPackageMapExportsUtil::Deserialize(FNetSerializationContext& Context,
 
 	// Read any exported names
 	{
+		UE_NET_TRACE_SCOPE(Names, *Reader, Context.GetTraceCollector(), ENetTraceVerbosity::VeryVerbose);
 		const bool bHasNames = Reader->ReadBool();
 		if (bHasNames)
 		{
@@ -153,6 +159,7 @@ void FIrisPackageMapExportsUtil::Deserialize(FNetSerializationContext& Context,
 
 	// Read any exported nettokens
 	{
+		UE_NET_TRACE_SCOPE(NetTokens, *Reader, Context.GetTraceCollector(), ENetTraceVerbosity::VeryVerbose);
 		const bool bHasNetTokens = Reader->ReadBool();
 		if (bHasNetTokens)
 		{
diff --git a/Engine/Source/Runtime/Experimental/Iris/Core/Public/Iris/Serialization/NetSerializerArrayStorage.h b/Engine/Source/Runtime/Experimental/Iris/Core/Public/Iris/Serialization/NetSerializerArrayStorage.h
--- a/Engine/Source/Runtime/Experimental/Iris/Core/Public/Iris/Serialization/NetSerializerArrayStorage.h
+++ b/Engine/Source/Runtime/Experimental/Iris/Core/Public/Iris/Serialization/NetSerializerArrayStorage.h
@@ -229,6 +229,13 @@ void FNetSerializerArrayStorage<QuantizedElementType, AllocationPolicy>::AdjustS
 	const SizeType NewCapacity = AllocatorInstance.CalculateNewCapacity(InNum);
 	AllocatorInstance.ResizeAllocation(Context, ArrayNum, NewCapacity);
 	ArrayMaxCapacity = NewCapacity;
+
+	if (NewCapacity > ArrayNum)
+	{
+		// To avoid issues with bad data in padding we always zero initialize new memory.
+		FMemory::Memzero(GetData() + ArrayNum, (NewCapacity - ArrayNum) * sizeof(ElementType));
+	}
+
 	ArrayNum = FMath::Min(InNum, NewCapacity);
 }
 
