Enable AutoRTFM here

diff --git a/Engine/Source/Runtime/Experimental/GeometryCollectionEngine/GeometryCollectionEngine.Build.cs b/Engine/Source/Runtime/Experimental/GeometryCollectionEngine/GeometryCollectionEngine.Build.cs
--- a/Engine/Source/Runtime/Experimental/GeometryCollectionEngine/GeometryCollectionEngine.Build.cs
+++ b/Engine/Source/Runtime/Experimental/GeometryCollectionEngine/GeometryCollectionEngine.Build.cs
@@ -49,8 +49,6 @@ namespace UnrealBuildTool.Rules
 			}
 
 			PrivateDefinitions.Add("CHAOS_INCLUDE_LEVEL_1=1");
-
-			bDisableAutoRTFMInstrumentation = true;
 		}
 	}
 }
