Fixed Sphere-Convex trace not taking radius into account.

If we're not computing MTD for a sphere trace we run down a path that can lead to FConvex::Raycast. A previous change added FConvex::RaycastFast to hande this without deferring to GJK for the result however that implementation didn't take Thickness into account for the ray which is how we implement sphere casts.

In the case of a sphere cast we now defer to GJK for the correct result

diff --git a/Engine/Source/Runtime/Experimental/Chaos/Private/Chaos/Convex.cpp b/Engine/Source/Runtime/Experimental/Chaos/Private/Chaos/Convex.cpp
--- a/Engine/Source/Runtime/Experimental/Chaos/Private/Chaos/Convex.cpp
+++ b/Engine/Source/Runtime/Experimental/Chaos/Private/Chaos/Convex.cpp
@@ -33,13 +33,18 @@ namespace Chaos
 	bool FConvex::Raycast(const FVec3& StartPoint, const FVec3& Dir, const FReal Length, const FReal Thickness, FReal& OutTime, FVec3& OutPosition, FVec3& OutNormal, int32& OutFaceIndex) const
 	{
 #if CHAOS_CONVEX_USE_FAST_RAYCAST
-		return RaycastFast(StartPoint, Dir, Length, Thickness, OutTime, OutPosition, OutNormal, OutFaceIndex);
-#else
+		// If we are just doing a raycast (Thickness == 0) we can run a faster direct ray test, but if it's
+		// a sphere-cast we'll still need to defer to GJK for an accurate result
+		if(Thickness == 0)
+		{
+			return RaycastFast(StartPoint, Dir, Length, Thickness, OutTime, OutPosition, OutNormal, OutFaceIndex);
+		}
+#endif
+
 		OutFaceIndex = INDEX_NONE;	//finding face is expensive, should be called directly by user
 		const FRigidTransform3 StartTM(StartPoint, FRotation3::FromIdentity());
 		const TSphere<FReal, 3> Sphere(FVec3(0), Thickness);
 		return GJKRaycast(*this, Sphere, StartTM, Dir, Length, OutTime, OutPosition, OutNormal);
-#endif
 	}
 
 	bool FConvex::RaycastFast(const FVec3& StartPoint, const FVec3& Dir, const FReal Length, const FReal Thickness, FReal& OutTime, FVec3& OutPosition, FVec3& OutNormal, int32& OutFaceIndex) const
