Iris
- Added log of NetBitStream reader readjournal when reporting error from ReplicationReader to display more context.

diff --git a/Engine/Source/Runtime/Experimental/Iris/Core/Private/Iris/ReplicationSystem/ReplicationReader.cpp b/Engine/Source/Runtime/Experimental/Iris/Core/Private/Iris/ReplicationSystem/ReplicationReader.cpp
--- a/Engine/Source/Runtime/Experimental/Iris/Core/Private/Iris/ReplicationSystem/ReplicationReader.cpp
+++ b/Engine/Source/Runtime/Experimental/Iris/Core/Private/Iris/ReplicationSystem/ReplicationReader.cpp
@@ -2188,7 +2188,7 @@ void FReplicationReader::ReadObjects(FNetSerializationContext& Context, uint32 O
 		--ObjectBatchCountToRead;
 	}
 
-	UE_CLOG(Context.HasErrorOrOverflow(), LogIris, Error, TEXT("Overflow: %c Error: %s Bit stream bits left: %u position: %u"), TEXT("YN")[Context.HasError()], ToCStr(Context.GetError().ToString()), Reader.GetBitsLeft(), Reader.GetPosBits())
+	UE_CLOG(Context.HasErrorOrOverflow(), LogIris, Error, TEXT("Overflow: %c Error: %s Bit stream bits left: %u position: %u %s"), TEXT("YN")[Context.HasError()], ToCStr(Context.GetError().ToString()), Reader.GetBitsLeft(), Reader.GetPosBits(), *Context.PrintReadJournal())
 	ensure(!Context.HasErrorOrOverflow());
 }
 
