Change Module Type from a string to a Hash for network serialization

diff --git a/Engine/Plugins/Experimental/ChaosModularVehicle/Source/ChaosModularVehicleEngine/Private/ChaosModularVehicle/ChaosSimModuleManagerAsyncCallback.cpp b/Engine/Plugins/Experimental/ChaosModularVehicle/Source/ChaosModularVehicleEngine/Private/ChaosModularVehicle/ChaosSimModuleManagerAsyncCallback.cpp
--- a/Engine/Plugins/Experimental/ChaosModularVehicle/Source/ChaosModularVehicleEngine/Private/ChaosModularVehicle/ChaosSimModuleManagerAsyncCallback.cpp
+++ b/Engine/Plugins/Experimental/ChaosModularVehicle/Source/ChaosModularVehicleEngine/Private/ChaosModularVehicle/ChaosSimModuleManagerAsyncCallback.cpp
@@ -319,28 +319,30 @@ bool FNetworkModularVehicleStates::NetSerialize(FArchive& Ar, class UPackageMap*
 		{
 			if (NumNetModules > 0)
 			{
-				FName ModuleType = NAME_None;
+				uint32 ModuleTypeHash = 0;
 				int32 SimArrayIndex = 0;
-				Ar << ModuleType;
+				Ar << ModuleTypeHash;
 				Ar << SimArrayIndex;
 
 				if (I >= ModuleData.Num())
 				{
-					if (TSharedPtr<Chaos::FModuleNetData> Data = Chaos::FModuleFactoryRegister::Get().GenerateNetData(ModuleType, SimArrayIndex))
+					if (TSharedPtr<Chaos::FModuleNetData> Data = Chaos::FModuleFactoryRegister::Get().GenerateNetData(ModuleTypeHash, SimArrayIndex))
 					{
 						ModuleData.Emplace(Data);
 					}
 				}
 				if(I <= ModuleData.Num() && ModuleData[I].IsValid())
 				{
+					check(ModuleTypeHash == GetTypeHash(ModuleData[I]->GetSimType()));
 					ModuleData[I]->Serialize(Ar);
 				}
 			}
 		}
 		else
 		{
-			FName ModuleType = ModuleData[I]->GetSimType();
-			Ar << ModuleType;
+			int ModuleTypeHash = GetTypeHash(ModuleData[I]->GetSimType());
+			
+			Ar << ModuleTypeHash;
 			Ar << ModuleData[I]->SimArrayIndex;
 			ModuleData[I]->Serialize(Ar);
 		}
diff --git a/Engine/Source/Runtime/Experimental/ChaosVehicles/ChaosVehiclesCore/Private/SimModule/ModuleFactoryRegister.cpp b/Engine/Source/Runtime/Experimental/ChaosVehicles/ChaosVehiclesCore/Private/SimModule/ModuleFactoryRegister.cpp
--- a/Engine/Source/Runtime/Experimental/ChaosVehicles/ChaosVehiclesCore/Private/SimModule/ModuleFactoryRegister.cpp
+++ b/Engine/Source/Runtime/Experimental/ChaosVehicles/ChaosVehiclesCore/Private/SimModule/ModuleFactoryRegister.cpp
@@ -7,19 +7,25 @@ namespace Chaos
 {
 	void FModuleFactoryRegister::RegisterFactory(const FName TypeName, TWeakPtr<IFactoryModule> InFactory)
 	{
-		if (!ContainsFactory(TypeName))
+		RegisterFactory(GetTypeHash(TypeName), InFactory);
+	}
+
+	void FModuleFactoryRegister::RegisterFactory(const uint32 TypeNameHash, TWeakPtr<IFactoryModule> InFactory)
+	{
+		if (!ContainsFactory(TypeNameHash))
 		{
-			RegisteredFactoriesByName.Add(TypeName, InFactory);
+			RegisteredFactoriesByName.Add(TypeNameHash, InFactory);
 		}
 	}
 
 	void FModuleFactoryRegister::RemoveFactory(TWeakPtr<IFactoryModule> InFactory)
 	{
-		for (TPair<FName, TWeakPtr<IFactoryModule>> Pair : RegisteredFactoriesByName)
+		for (TPair<uint32, TWeakPtr<IFactoryModule>> Pair : RegisteredFactoriesByName)
 		{
 			if (Pair.Value == InFactory)
-			{RegisteredFactoriesByName.Remove(Pair.Key);
-				
+			{
+				RegisteredFactoriesByName.Remove(Pair.Key);
+			
 				return;
 			}
 		}
@@ -32,16 +38,21 @@ namespace Chaos
 
 	bool FModuleFactoryRegister::ContainsFactory(const FName TypeName) const
 	{
-		return RegisteredFactoriesByName.Contains(TypeName);
+		return ContainsFactory(GetTypeHash(TypeName));
+	}
+
+	bool FModuleFactoryRegister::ContainsFactory(const uint32 TypeNameHash) const
+	{
+		return RegisteredFactoriesByName.Contains(TypeNameHash);
 	}
 
-	TSharedPtr<Chaos::FModuleNetData> FModuleFactoryRegister::GenerateNetData(const FName TypeName, const int32 SimArrayIndex)
+	TSharedPtr<Chaos::FModuleNetData> FModuleFactoryRegister::GenerateNetData(const uint32 TypeNameHash, const int32 SimArrayIndex)
 	{
 		using namespace Chaos;
 
-		if (RegisteredFactoriesByName.Contains(TypeName))
+		if (RegisteredFactoriesByName.Contains(TypeNameHash))
 		{
-			TSharedPtr<IFactoryModule> PinnedFactory = RegisteredFactoriesByName[TypeName].Pin();
+			TSharedPtr<IFactoryModule> PinnedFactory = RegisteredFactoriesByName[TypeNameHash].Pin();
 
 			if (PinnedFactory.IsValid())
 			{
@@ -50,7 +61,7 @@ namespace Chaos
 		}
 		else
 		{
-			UE_LOG(LogTemp, Error, TEXT("No factory registered for name '%s'"), *TypeName.ToString());
+			UE_LOG(LogTemp, Error, TEXT("No factory registered for hashed type %d"), TypeNameHash);
 		}
 
 		return nullptr;
diff --git a/Engine/Source/Runtime/Experimental/ChaosVehicles/ChaosVehiclesCore/Public/SimModule/ModuleFactoryRegister.h b/Engine/Source/Runtime/Experimental/ChaosVehicles/ChaosVehiclesCore/Public/SimModule/ModuleFactoryRegister.h
--- a/Engine/Source/Runtime/Experimental/ChaosVehicles/ChaosVehiclesCore/Public/SimModule/ModuleFactoryRegister.h
+++ b/Engine/Source/Runtime/Experimental/ChaosVehicles/ChaosVehiclesCore/Public/SimModule/ModuleFactoryRegister.h
@@ -22,15 +22,17 @@ namespace Chaos
 		}
 		
 		void RegisterFactory(const FName TypeName, TWeakPtr<IFactoryModule> InFactory);
+		void RegisterFactory(const uint32 TypeNameHash, TWeakPtr<IFactoryModule> InFactory);
 		void RemoveFactory(TWeakPtr<IFactoryModule> InFactory);
 		void Reset();
 		bool ContainsFactory(const FName TypeName) const;
-		TSharedPtr<Chaos::FModuleNetData> GenerateNetData(const FName TypeName, const int32 SimArrayIndex);
+		bool ContainsFactory(const uint32 TypeNameHash) const;
+		TSharedPtr<Chaos::FModuleNetData> GenerateNetData(const uint32 TypeNameHash, const int32 SimArrayIndex);
 
 	protected:
 
 		FModuleFactoryRegister() = default;
-		TMap<FName, TWeakPtr<IFactoryModule>> RegisteredFactoriesByName;
+		TMap<int32, TWeakPtr<IFactoryModule>> RegisteredFactoriesByName;
 	};
 	
 	template<typename _To, typename ..._Rest>
@@ -44,14 +46,15 @@ namespace Chaos
 		{
 			return false;
 		}
-		if(FModuleFactoryRegister::Get().ContainsFactory(SimTypeName))
+		uint32 SimTypeNameHash = GetTypeHash(SimTypeName);
+		if(FModuleFactoryRegister::Get().ContainsFactory(SimTypeNameHash))
 		{
 			return true;
 		}
 		static TSharedPtr<T> SharedFactory = MakeShared<T>(args...);
 		if (SharedFactory.IsValid())
 		{
-			FModuleFactoryRegister::Get().RegisterFactory(SimTypeName, SharedFactory);
+			FModuleFactoryRegister::Get().RegisterFactory(SimTypeNameHash, SharedFactory);
 			return true;
 		}
 		return false;
