[ChaosVD] Record disabled particles of inside active cluster particle.

This allows record all updates of particles that are in a cluster union an therefore disabled. Without this change these particles only got updated every 10 seconds which is the interval we have use between keyframes.

As this change is not free perf wise when a recording is active, I gated it behing a new data channel that is disabled by default, and can be enabled as needed.

diff --git a/Engine/Source/Runtime/Experimental/Chaos/Private/Chaos/PBDRigidsEvolutionGBF.cpp b/Engine/Source/Runtime/Experimental/Chaos/Private/Chaos/PBDRigidsEvolutionGBF.cpp
--- a/Engine/Source/Runtime/Experimental/Chaos/Private/Chaos/PBDRigidsEvolutionGBF.cpp
+++ b/Engine/Source/Runtime/Experimental/Chaos/Private/Chaos/PBDRigidsEvolutionGBF.cpp
@@ -508,7 +508,7 @@ void FPBDRigidsEvolutionGBF::AdvanceOneTimeStepImpl(const FReal Dt, const FSubSt
 
 	{
 		CVD_SCOPE_TRACE_SOLVER_STEP(CVDDC_EvolutionStart, TEXT("Evolution Start"));
-		CVD_TRACE_PARTICLES_SOA(Particles);
+		CVD_TRACE_PARTICLES_SOA(Particles, &Clustering);
 	}
 
 	// Update the collision solver type (used to support runtime comparisons of solver types for debugging/testing)
@@ -543,7 +543,7 @@ void FPBDRigidsEvolutionGBF::AdvanceOneTimeStepImpl(const FReal Dt, const FSubSt
 
 	{
 		CVD_SCOPE_TRACE_SOLVER_STEP(CVDDC_PostIntegrate, TEXT("Post Integrate"));
-		CVD_TRACE_PARTICLES_SOA(Particles);
+		CVD_TRACE_PARTICLES_SOA(Particles, &Clustering);
 	}
 
 	if (PostIntegrateCallback != nullptr)
@@ -663,7 +663,7 @@ void FPBDRigidsEvolutionGBF::AdvanceOneTimeStepImpl(const FReal Dt, const FSubSt
 
 	{
 		CVD_SCOPE_TRACE_SOLVER_STEP(CVDDC_PreConstraintSolve, TEXT("Pre Solve"));
-		CVD_TRACE_PARTICLES_SOA(Particles);
+		CVD_TRACE_PARTICLES_SOA(Particles, &Clustering);
 	}
 
 	// Assign all islands to a set of groups. Each group is solved in parallel with the others.
@@ -705,7 +705,7 @@ void FPBDRigidsEvolutionGBF::AdvanceOneTimeStepImpl(const FReal Dt, const FSubSt
 
 	{
 		CVD_SCOPE_TRACE_SOLVER_STEP(CVDDC_PostConstraintSolve, TEXT("Post Solve"));
-		CVD_TRACE_PARTICLES_SOA(Particles);
+		CVD_TRACE_PARTICLES_SOA(Particles, &Clustering);
 	}
 
 	if (PostSolveCallback != nullptr)
@@ -790,7 +790,7 @@ void FPBDRigidsEvolutionGBF::AdvanceOneTimeStepImpl(const FReal Dt, const FSubSt
 		CVD_TRACE_JOINT_CONSTRAINTS(CVDDC_JointConstraints, JointConstraints);
 		CVD_TRACE_CHARACTER_GROUND_CONSTRAINTS(CVDDC_CharacterGroundConstraints, CharacterGroundConstraints);
 
-		CVD_TRACE_PARTICLES_SOA(Particles);
+		CVD_TRACE_PARTICLES_SOA(Particles, &Clustering);
 	}
 
 #if !UE_BUILD_SHIPPING
diff --git a/Engine/Source/Runtime/Experimental/Chaos/Private/ChaosVisualDebugger/ChaosVDOptionalDataChannel.cpp b/Engine/Source/Runtime/Experimental/Chaos/Private/ChaosVisualDebugger/ChaosVDOptionalDataChannel.cpp
--- a/Engine/Source/Runtime/Experimental/Chaos/Private/ChaosVisualDebugger/ChaosVDOptionalDataChannel.cpp
+++ b/Engine/Source/Runtime/Experimental/Chaos/Private/ChaosVisualDebugger/ChaosVDOptionalDataChannel.cpp
@@ -86,5 +86,6 @@ CVD_DEFINE_OPTIONAL_DATA_CHANNEL(SceneQueries, EChaosVDDataChannelInitialization
 CVD_DEFINE_OPTIONAL_DATA_CHANNEL(JointConstraints, EChaosVDDataChannelInitializationFlags::CanChangeEnabledState);
 CVD_DEFINE_OPTIONAL_DATA_CHANNEL(CharacterGroundConstraints, EChaosVDDataChannelInitializationFlags::CanChangeEnabledState);
 CVD_DEFINE_OPTIONAL_DATA_CHANNEL(AccelerationStructures, EChaosVDDataChannelInitializationFlags::CanChangeEnabledState);
+CVD_DEFINE_OPTIONAL_DATA_CHANNEL(ClusterParticlesChildData, EChaosVDDataChannelInitializationFlags::CanChangeEnabledState);
 
 #endif
diff --git a/Engine/Source/Runtime/Experimental/Chaos/Private/ChaosVisualDebugger/ChaosVisualDebuggerTrace.cpp b/Engine/Source/Runtime/Experimental/Chaos/Private/ChaosVisualDebugger/ChaosVisualDebuggerTrace.cpp
--- a/Engine/Source/Runtime/Experimental/Chaos/Private/ChaosVisualDebugger/ChaosVisualDebuggerTrace.cpp
+++ b/Engine/Source/Runtime/Experimental/Chaos/Private/ChaosVisualDebugger/ChaosVisualDebuggerTrace.cpp
@@ -4,6 +4,8 @@
 
 #if WITH_CHAOS_VISUAL_DEBUGGER
 
+#include "Chaos/PBDRigidClustering.h"
+
 #include "Chaos/Character/CharacterGroundConstraintContainer.h"
 #include "Chaos/Framework/PhysicsSolverBase.h"
 #include "Chaos/ImplicitObject.h"
@@ -225,7 +227,41 @@ void FChaosVisualDebuggerTrace::TraceParticleDestroyed(const Chaos::FGeometryPar
 		<< ChaosVDParticleDestroyed.ParticleID(ParticleHandle->UniqueIdx().Idx);
 }
 
-void FChaosVisualDebuggerTrace::TraceParticlesSoA(const Chaos::FPBDRigidsSOAs& ParticlesSoA)
+void FChaosVisualDebuggerTrace::TraceParticleClusterChildData(const Chaos::TParticleView<Chaos::TPBDRigidParticles<Chaos::FReal, 3>>& ParticlesView, Chaos::FRigidClustering* ClusteringData, const FChaosVDContext& CVDContextData)
+{
+	if (!IsTracing())
+	{
+		return;
+	}
+
+	if (!ClusteringData)
+	{
+		return;
+	}
+
+	if (!CVDDC_ClusterParticlesChildData->IsChannelEnabled())
+	{
+		return;
+	}
+
+	ParticlesView.ParallelFor([CopyContext = CVDContextData, ClusteringData](auto& Particle, int32 Index)
+	{
+		if (Chaos::FPBDRigidClusteredParticleHandle* ClusteredParticle = Particle.Handle()->CastToClustered())
+		{
+			CVD_SCOPE_CONTEXT(CopyContext);
+			if (const TArray<Chaos::FPBDRigidParticleHandle*>* ChildrenHandles = ClusteringData->GetChildrenMap().Find(ClusteredParticle))
+			{
+				const TArray<Chaos::FPBDRigidParticleHandle*>& ChildrenHandlesArray = *ChildrenHandles;
+				for (Chaos::FPBDRigidParticleHandle* ParticleHandle : ChildrenHandlesArray)
+				{
+					TraceParticle(ParticleHandle);
+				}
+			}
+		}
+	});
+}
+
+void FChaosVisualDebuggerTrace::TraceParticlesSoA(const Chaos::FPBDRigidsSOAs& ParticlesSoA, Chaos::FRigidClustering* ClusteringData)
 {
 	using namespace Chaos::VisualDebugger::Utils;
 	if (!IsTracing())
@@ -248,6 +284,9 @@ void FChaosVisualDebuggerTrace::TraceParticlesSoA(const Chaos::FPBDRigidsSOAs& P
 	}
 
 	TraceParticlesView(ParticlesSoA.GetDirtyParticlesView());
+
+	// If we are recording a delta frame, we need to also record the child particles of any cluster (if we have clustering data available)
+	TraceParticleClusterChildData(ParticlesSoA.GetDirtyParticlesView(), ClusteringData, *CVDContextData);
 }
 
 void FChaosVisualDebuggerTrace::SetupForFullCaptureIfNeeded(int32 SolverID, bool& bOutFullCaptureRequested)
diff --git a/Engine/Source/Runtime/Experimental/Chaos/Public/ChaosVisualDebugger/ChaosVDOptionalDataChannel.h b/Engine/Source/Runtime/Experimental/Chaos/Public/ChaosVisualDebugger/ChaosVDOptionalDataChannel.h
--- a/Engine/Source/Runtime/Experimental/Chaos/Public/ChaosVisualDebugger/ChaosVDOptionalDataChannel.h
+++ b/Engine/Source/Runtime/Experimental/Chaos/Public/ChaosVisualDebugger/ChaosVDOptionalDataChannel.h
@@ -155,6 +155,10 @@ CVD_DECLARE_OPTIONAL_DATA_CHANNEL(JointConstraints);
 CVD_DECLARE_OPTIONAL_DATA_CHANNEL(CharacterGroundConstraints);
 CVD_DECLARE_OPTIONAL_DATA_CHANNEL(AccelerationStructures);
 
+// When a particle is added to a cluster will get disabled and therefore not updated by CVD's delta frames implementation.
+// If this channel is enabled, we will record all child particles of any active dirty cluster particle
+CVD_DECLARE_OPTIONAL_DATA_CHANNEL(ClusterParticlesChildData);
+
 #else  //WITH_CHAOS_VISUAL_DEBUGGER
 
 /** Declares an Optional CVD Data channel to be available globally. The data channel can be accessed by using CVDDC_TheNameOfTheChannelUsedWithThisMacro */
diff --git a/Engine/Source/Runtime/Experimental/Chaos/Public/ChaosVisualDebugger/ChaosVDTraceMacros.h b/Engine/Source/Runtime/Experimental/Chaos/Public/ChaosVisualDebugger/ChaosVDTraceMacros.h
--- a/Engine/Source/Runtime/Experimental/Chaos/Public/ChaosVisualDebugger/ChaosVDTraceMacros.h
+++ b/Engine/Source/Runtime/Experimental/Chaos/Public/ChaosVisualDebugger/ChaosVDTraceMacros.h
@@ -50,8 +50,8 @@
 	#endif
 
 	#ifndef CVD_TRACE_PARTICLES_SOA
-		#define CVD_TRACE_PARTICLES_SOA(ParticleSoA) \
-		FChaosVisualDebuggerTrace::TraceParticlesSoA(ParticleSoA);
+		#define CVD_TRACE_PARTICLES_SOA(ParticleSoA, ...) \
+		FChaosVisualDebuggerTrace::TraceParticlesSoA(ParticleSoA, ##__VA_ARGS__);
 	#endif
 
 	#ifndef CVD_TRACE_SOLVER_START_FRAME
@@ -257,7 +257,7 @@
 	#endif
 
 	#ifndef CVD_TRACE_PARTICLES_SOA
-		#define CVD_TRACE_PARTICLES_SOA(ParticleSoA)
+		#define CVD_TRACE_PARTICLES_SOA(ParticleSoA, ...)
 	#endif
 
 	#ifndef CVD_TRACE_PARTICLE_DESTROYED
diff --git a/Engine/Source/Runtime/Experimental/Chaos/Public/ChaosVisualDebugger/ChaosVisualDebuggerTrace.h b/Engine/Source/Runtime/Experimental/Chaos/Public/ChaosVisualDebugger/ChaosVisualDebuggerTrace.h
--- a/Engine/Source/Runtime/Experimental/Chaos/Public/ChaosVisualDebugger/ChaosVisualDebuggerTrace.h
+++ b/Engine/Source/Runtime/Experimental/Chaos/Public/ChaosVisualDebugger/ChaosVisualDebuggerTrace.h
@@ -29,6 +29,7 @@
 
 namespace Chaos
 {
+class FRigidClustering;
 class FAccelerationStructureHandle;
 class FCharacterGroundConstraintContainer;
 class FPBDConstraintContainer;
@@ -195,7 +196,7 @@ public:
 	 * Traces data from particles on the provided FPBDRigidsSOAs. It traces only the DirtyParticles view unless a full capture was requested
 	 * @param ParticlesSoA Particles SoA to evaluate and trace
 	 */
-	static CHAOS_API void TraceParticlesSoA(const Chaos::FPBDRigidsSOAs& ParticlesSoA);
+	static CHAOS_API void TraceParticlesSoA(const Chaos::FPBDRigidsSOAs& ParticlesSoA, Chaos::FRigidClustering* ClusteringData = nullptr);
 
 	/** Traces the provided particle view in parallel */
 	template<typename ParticleType>
@@ -376,10 +377,18 @@ public:
 
 private:
 
+	/**
+	 * Traces data from all Child Particles from any Cluster Particle inside the provided view array, using the provided CVD Context and Clustering Data
+	 * @param ParticlesView Particles Array view to process and Trace
+	 * @param ClusteringData Object containing the mappings required to find all child particles for a given cluster particle
+	 * @param CVDContextData Context to be used to tie this Trace event with specific solver frame and step
+	 */
+	static CHAOS_API void TraceParticleClusterChildData(const Chaos::TParticleView<Chaos::TPBDRigidParticles<Chaos::FReal, 3>>& ParticlesView, Chaos::FRigidClustering* ClusteringData, const FChaosVDContext& CVDContextData);
+
 	/**
 	 * Traces data from a Particle Handle using the provided CVD Context
 	 * @param ParticleHandle Handle to process and Trace
-	 * @param ContextData Context to be used to tied this Trace event to a specific solver frame and step
+	 * @param ContextData Context to be used to tie this Trace event with specific solver frame and step
 	 */
 	static CHAOS_API void TraceParticle(Chaos::FGeometryParticleHandle* ParticleHandle, const FChaosVDContext& ContextData);
 
