FIoStoreOnDemandModule
Add GetCacheUsage() and console command.

diff --git a/Engine/Source/Runtime/Experimental/IoStore/OnDemand/Private/IoStoreOnDemand.cpp b/Engine/Source/Runtime/Experimental/IoStore/OnDemand/Private/IoStoreOnDemand.cpp
--- a/Engine/Source/Runtime/Experimental/IoStore/OnDemand/Private/IoStoreOnDemand.cpp
+++ b/Engine/Source/Runtime/Experimental/IoStore/OnDemand/Private/IoStoreOnDemand.cpp
@@ -73,6 +73,34 @@ static FAutoConsoleCommand OnDemandPurgeCacheCommand(
 	ECVF_Cheat
 );
 
+static FAutoConsoleCommand OnDemandCacheUsageCommand(
+	TEXT("iostore.CacheUsage"),
+	TEXT("print cache usage"),
+	FConsoleCommandDelegate::CreateStatic([]()
+	{
+		FIoStoreOnDemandModule* IOStoreOnDemandModule = FModuleManager::Get().GetModulePtr<FIoStoreOnDemandModule>(TEXT("IoStoreOnDemand"));
+		if (!IOStoreOnDemandModule)
+		{
+			UE_LOG(LogIoStoreOnDemand, Error, TEXT("Could not find IoStoreOnDemand module"));
+			return;
+		}
+
+		TIoStatusOr<FOnDemandCacheUsage> MaybeUsage = IOStoreOnDemandModule->GetCacheUsage();
+		if (!MaybeUsage.IsOk())
+		{
+			UE_LOG(LogIoStoreOnDemand, Error, TEXT("iostore.CacheUsage failed: %s"), *MaybeUsage.Status().ToString());
+			return;
+		}
+
+		const FOnDemandCacheUsage& Usage = MaybeUsage.ValueOrDie();
+		UE_LOG(LogIoStoreOnDemand, Display, TEXT("iostore.CacheUsage"));
+		UE_LOG(LogIoStoreOnDemand, Display, TEXT("\tMaxSize %" UINT64_FMT), Usage.MaxSize);
+		UE_LOG(LogIoStoreOnDemand, Display, TEXT("\tTotalSize %" UINT64_FMT), Usage.TotalSize);
+		UE_LOG(LogIoStoreOnDemand, Display, TEXT("\tReferencedBlockSize %" UINT64_FMT), Usage.ReferencedBlockSize);
+	}),
+	ECVF_Cheat
+);
+
 ////////////////////////////////////////////////////////////////////////////////
 FString GIasOnDemandTocExt = TEXT(".uondemandtoc");
 
@@ -1158,6 +1186,16 @@ FIoStatus FIoStoreOnDemandModule::GetInstallSizesByMountId(const FOnDemandGetIns
 	return FIoStatus(EIoErrorCode::InvalidCode, NotInitializedError);
 }
 
+TIoStatusOr<FOnDemandCacheUsage> FIoStoreOnDemandModule::GetCacheUsage() const
+{
+	if (IoStore)
+	{
+		return IoStore->GetCacheUsage();
+	}
+
+	return FIoStatus(EIoErrorCode::InvalidCode, NotInitializedError);
+}
+
 void FIoStoreOnDemandModule::InitializeInternal()
 {
 	LLM_SCOPE_BYTAG(Ias);
diff --git a/Engine/Source/Runtime/Experimental/IoStore/OnDemand/Private/OnDemandInstallCache.cpp b/Engine/Source/Runtime/Experimental/IoStore/OnDemand/Private/OnDemandInstallCache.cpp
--- a/Engine/Source/Runtime/Experimental/IoStore/OnDemand/Private/OnDemandInstallCache.cpp
+++ b/Engine/Source/Runtime/Experimental/IoStore/OnDemand/Private/OnDemandInstallCache.cpp
@@ -900,33 +900,34 @@ public:
 	virtual ~FOnDemandInstallCache();
 
 	// IIoDispatcherBackend
-	virtual void				Initialize(FSharedBackendContextRef Context) override;
-	virtual void				Shutdown() override;
-	virtual void				ResolveIoRequests(FIoRequestList Requests, FIoRequestList& OutUnresolved) override;
-	virtual FIoRequestImpl*		GetCompletedIoRequests() override;
-	virtual void				CancelIoRequest(FIoRequestImpl* Request) override;
-	virtual void				UpdatePriorityForIoRequest(FIoRequestImpl* Request) override;
-	virtual bool				DoesChunkExist(const FIoChunkId& ChunkId) const override;
-	virtual TIoStatusOr<uint64> GetSizeForChunk(const FIoChunkId& ChunkId) const override;
-	virtual TIoStatusOr<FIoMappedRegion> OpenMapped(const FIoChunkId& ChunkId, const FIoReadOptions& Options) override;
+	virtual void								Initialize(FSharedBackendContextRef Context) override;
+	virtual void								Shutdown() override;
+	virtual void								ResolveIoRequests(FIoRequestList Requests, FIoRequestList& OutUnresolved) override;
+	virtual FIoRequestImpl*						GetCompletedIoRequests() override;
+	virtual void								CancelIoRequest(FIoRequestImpl* Request) override;
+	virtual void								UpdatePriorityForIoRequest(FIoRequestImpl* Request) override;
+	virtual bool								DoesChunkExist(const FIoChunkId& ChunkId) const override;
+	virtual TIoStatusOr<uint64>					GetSizeForChunk(const FIoChunkId& ChunkId) const override;
+	virtual TIoStatusOr<FIoMappedRegion>		OpenMapped(const FIoChunkId& ChunkId, const FIoReadOptions& Options) override;
 
 	// IOnDemandInstallCache
-	virtual bool				IsChunkCached(const FIoHash& ChunkHash) override;
-	virtual FIoStatus			PutChunk(FIoBuffer&& Chunk, const FIoHash& ChunkHash) override;
-	virtual FIoStatus			Purge(TMap<FIoHash, uint64>&& ChunksToIntall) override;
-	virtual FIoStatus			PurgeAllUnreferenced() override;
-	virtual FIoStatus			Flush() override;
+	virtual bool								IsChunkCached(const FIoHash& ChunkHash) override;
+	virtual FIoStatus							PutChunk(FIoBuffer&& Chunk, const FIoHash& ChunkHash) override;
+	virtual FIoStatus							Purge(TMap<FIoHash, uint64>&& ChunksToIntall) override;
+	virtual FIoStatus							PurgeAllUnreferenced() override;
+	virtual FIoStatus							Flush() override;
+	virtual FOnDemandInstallCacheStorageUsage	GetStorageUsage() override;
 
 private:
-	void						AddReferencesToBlocks(
-									const TArray<FSharedOnDemandContainer>& Containers, 
-									const TArray<TBitArray<>>& ChunkEntryIndices,
-									FCasBlockInfoMap& BlockInfo) const;
-	FIoStatus					Purge(const FCasBlockInfoMap& BlockInfo, uint64 TotalBytesToPurge, uint64& OutTotalPurgedBytes);
-	bool						Resolve(FIoRequestImpl* Request);
-	void						CompleteRequest(FIoRequestImpl* Request, bool bFileReadWasCancelled);
-	FIoStatus					FlushPendingChunks(FPendingChunks& Block);
-	FString						GetJournalFilename() const { return CacheDirectory / TEXT("cas.jrn"); }
+	void										AddReferencesToBlocks(
+													const TArray<FSharedOnDemandContainer>& Containers, 
+													const TArray<TBitArray<>>& ChunkEntryIndices,
+													FCasBlockInfoMap& BlockInfo) const;
+	FIoStatus									Purge(const FCasBlockInfoMap& BlockInfo, uint64 TotalBytesToPurge, uint64& OutTotalPurgedBytes);
+	bool										Resolve(FIoRequestImpl* Request);
+	void										CompleteRequest(FIoRequestImpl* Request, bool bFileReadWasCancelled);
+	FIoStatus									FlushPendingChunks(FPendingChunks& Block);
+	FString										GetJournalFilename() const { return CacheDirectory / TEXT("cas.jrn"); }
 
 	FOnDemandIoStore&		IoStore;
 	FString					CacheDirectory;
@@ -1351,7 +1352,7 @@ void FOnDemandInstallCache::AddReferencesToBlocks(
 	for (int32 Index = 0; FSharedOnDemandContainer Container : Containers)
 	{
 		const TBitArray<>& IsReferenced = ChunkEntryIndices[Index++];
-		for (int32 EntryIndex = 0; const FOnDemandChunkEntry & Entry : Container->ChunkEntries)
+		for (int32 EntryIndex = 0; const FOnDemandChunkEntry& Entry : Container->ChunkEntries)
 		{
 			if (bool bIsReferenced = IsReferenced[EntryIndex++]; bIsReferenced == false)
 			{
@@ -1423,6 +1424,33 @@ FIoStatus FOnDemandInstallCache::Flush()
 	return FIoStatus::Ok;
 }
 
+FOnDemandInstallCacheStorageUsage FOnDemandInstallCache::GetStorageUsage()
+{
+	// If this is called from a thread other than the OnDemandIoStore tick thread
+	// then its possible the block info and containers may not be in sync with each other
+	// or the current state of the tick thread.
+	// This should only be used for debugging and telemetry purposes.
+
+	FCasBlockInfoMap	BlockInfo;
+	const uint64		TotalCachedBytes = Cas.GetBlockInfo(BlockInfo);
+
+	TArray<FSharedOnDemandContainer>	Containers;
+	TArray<TBitArray<>>					ChunkEntryIndices;
+	IoStore.GetReferencedContent(Containers, ChunkEntryIndices);
+	check(Containers.Num() == ChunkEntryIndices.Num());
+
+	AddReferencesToBlocks(Containers, ChunkEntryIndices, BlockInfo);
+
+	return FOnDemandInstallCacheStorageUsage
+	{
+		.MaxSize = MaxCacheSize,
+		.TotalSize = TotalCachedBytes,
+		.ReferencedBlockSize = Algo::TransformAccumulate(BlockInfo,
+			[](const TPair<FCasBlockId, FCasBlockInfo>& Kv) { return (Kv.Value.RefCount > 0) ? Kv.Value.FileSize : uint64(0); },
+			uint64(0))
+	};
+}
+
 FIoStatus FOnDemandInstallCache::FlushPendingChunks(FPendingChunks& Chunks)
 {
 	ON_SCOPE_EXIT { Chunks.Reset(); };
diff --git a/Engine/Source/Runtime/Experimental/IoStore/OnDemand/Private/OnDemandInstallCache.h b/Engine/Source/Runtime/Experimental/IoStore/OnDemand/Private/OnDemandInstallCache.h
--- a/Engine/Source/Runtime/Experimental/IoStore/OnDemand/Private/OnDemandInstallCache.h
+++ b/Engine/Source/Runtime/Experimental/IoStore/OnDemand/Private/OnDemandInstallCache.h
@@ -14,16 +14,24 @@ namespace UE::IoStore
 class FOnDemandIoStore;
 struct FOnDemandContainer;
 
+struct FOnDemandInstallCacheStorageUsage
+{
+	uint64 MaxSize = 0;
+	uint64 TotalSize = 0;
+	uint64 ReferencedBlockSize = 0;
+};
+
 class IOnDemandInstallCache 
 	: public IIoDispatcherBackend
 {
 public:
-	virtual				~IOnDemandInstallCache() = default;
-	virtual bool		IsChunkCached(const FIoHash& ChunkHash) = 0;
-	virtual FIoStatus	PutChunk(FIoBuffer&& Chunk, const FIoHash& ChunkHash) = 0;
-	virtual FIoStatus	Purge(TMap<FIoHash, uint64>&& ChunksToInstall) = 0;
-	virtual FIoStatus	PurgeAllUnreferenced() = 0;
-	virtual FIoStatus	Flush() = 0;
+	virtual										~IOnDemandInstallCache() = default;
+	virtual bool								IsChunkCached(const FIoHash& ChunkHash) = 0;
+	virtual FIoStatus							PutChunk(FIoBuffer&& Chunk, const FIoHash& ChunkHash) = 0;
+	virtual FIoStatus							Purge(TMap<FIoHash, uint64>&& ChunksToInstall) = 0;
+	virtual FIoStatus							PurgeAllUnreferenced() = 0;
+	virtual FIoStatus							Flush() = 0;
+	virtual FOnDemandInstallCacheStorageUsage	GetStorageUsage() = 0;
 };
 
 struct FOnDemandInstallCacheConfig
diff --git a/Engine/Source/Runtime/Experimental/IoStore/OnDemand/Private/OnDemandIoStore.cpp b/Engine/Source/Runtime/Experimental/IoStore/OnDemand/Private/OnDemandIoStore.cpp
--- a/Engine/Source/Runtime/Experimental/IoStore/OnDemand/Private/OnDemandIoStore.cpp
+++ b/Engine/Source/Runtime/Experimental/IoStore/OnDemand/Private/OnDemandIoStore.cpp
@@ -818,7 +818,6 @@ FIoStatus FOnDemandIoStore::InitializeOnDemandInstallCache()
 		}
 		else
 		{
-			// Only warn until this is properly tested
 			UE_LOG(LogIoStoreOnDemand, Error, TEXT("Failed to initialize install cache"));
 			return FIoStatusBuilder(EIoErrorCode::InvalidParameter) << TEXT("Failed to initialize install cache");
 		}
@@ -1930,6 +1929,17 @@ void FOnDemandIoStore::GetReferencedContent(TArray<FSharedOnDemandContainer>& Ou
 	}
 }
 
+FOnDemandCacheUsage FOnDemandIoStore::GetCacheUsage() const
+{
+	FOnDemandInstallCacheStorageUsage Usage = InstallCache->GetStorageUsage();
+	return FOnDemandCacheUsage
+	{
+		.MaxSize = Usage.MaxSize,
+		.TotalSize = Usage.TotalSize,
+		.ReferencedBlockSize = Usage.ReferencedBlockSize
+	};
+}
+
 ///////////////////////////////////////////////////////////////////////////////
 FOnDemandInternalContentHandle::~FOnDemandInternalContentHandle()
 {
diff --git a/Engine/Source/Runtime/Experimental/IoStore/OnDemand/Private/OnDemandIoStore.h b/Engine/Source/Runtime/Experimental/IoStore/OnDemand/Private/OnDemandIoStore.h
--- a/Engine/Source/Runtime/Experimental/IoStore/OnDemand/Private/OnDemandIoStore.h
+++ b/Engine/Source/Runtime/Experimental/IoStore/OnDemand/Private/OnDemandIoStore.h
@@ -300,6 +300,7 @@ public:
 	FOnDemandChunkInfo		GetInstalledChunkInfo(const FIoChunkId& ChunkId);
 	void					ReleaseContent(FOnDemandInternalContentHandle& ContentHandle);
 	void					GetReferencedContent(TArray<FSharedOnDemandContainer>& OutContainers, TArray<TBitArray<>>& OutChunkEntryIndices);
+	FOnDemandCacheUsage		GetCacheUsage() const;
 
 private:
 	FIoStatus				GetContainersForInstall(
diff --git a/Engine/Source/Runtime/Experimental/IoStore/OnDemand/Public/IO/IoStoreOnDemand.h b/Engine/Source/Runtime/Experimental/IoStore/OnDemand/Public/IO/IoStoreOnDemand.h
--- a/Engine/Source/Runtime/Experimental/IoStore/OnDemand/Public/IO/IoStoreOnDemand.h
+++ b/Engine/Source/Runtime/Experimental/IoStore/OnDemand/Public/IO/IoStoreOnDemand.h
@@ -373,7 +373,7 @@ struct FOnDemandPurgeArgs
 	EOnDemandInstallOptions Options = EOnDemandInstallOptions::None;
 };
 
-/** Holdes information about a purge request */
+/** Holds information about a purge request */
 struct FOnDemandPurgeResult
 {
 	/** The status of the purge request. */
@@ -396,6 +396,14 @@ struct FOnDemandGetInstallSizeArgs
 	TArray<FPackageId> PackageIds;
 };
 
+/** Holds information about cache usage */
+struct FOnDemandCacheUsage
+{
+	uint64 MaxSize = 0;
+	uint64 TotalSize = 0;
+	uint64 ReferencedBlockSize = 0;
+};
+
 /** Token used for signalling an operation to be cancelled. */
 class FOnDemandCancellationToken
 {
@@ -458,6 +466,12 @@ public:
 
 	UE_API FIoStatus GetInstallSizesByMountId(const FOnDemandGetInstallSizeArgs& Args, TMap<FString, uint64>& OutSizesByMountId) const;
 
+	/** This does not queue a request so the returned state may not be completely consistent
+	  * if there are requests currently being processed. 
+	  * This should only be used for purposes such as debugging telemetry. 
+	  */
+	UE_API TIoStatusOr<FOnDemandCacheUsage> GetCacheUsage() const;
+
 	virtual void StartupModule() override;
 	virtual void ShutdownModule() override;
 	
