Fixed an issue when running a construction script that was causing geometry collection component's custom render to try and spawn actors

diff --git a/Engine/Source/Runtime/Experimental/GeometryCollectionEngine/Private/GeometryCollection/GeometryCollectionComponent.cpp b/Engine/Source/Runtime/Experimental/GeometryCollectionEngine/Private/GeometryCollection/GeometryCollectionComponent.cpp
--- a/Engine/Source/Runtime/Experimental/GeometryCollectionEngine/Private/GeometryCollection/GeometryCollectionComponent.cpp
+++ b/Engine/Source/Runtime/Experimental/GeometryCollectionEngine/Private/GeometryCollection/GeometryCollectionComponent.cpp
@@ -6322,6 +6322,12 @@ bool UGeometryCollectionComponent::IsUsingCustomRenderer() const
 
 void UGeometryCollectionComponent::RefreshCustomRenderer()
 {
+	UWorld* World = GetWorld();
+	if (!World || World->bIsRunningConstructionScript)
+	{
+		return;
+	}
+
 	if (RestCollection == nullptr || !IsCustomRendererAvailable())
 	{
 		return;
