Iris
* Changed ensureAlways to ensure

diff --git a/Engine/Source/Runtime/Experimental/Iris/Core/Private/Iris/ReplicationSystem/NetRefHandleManager.cpp b/Engine/Source/Runtime/Experimental/Iris/Core/Private/Iris/ReplicationSystem/NetRefHandleManager.cpp
--- a/Engine/Source/Runtime/Experimental/Iris/Core/Private/Iris/ReplicationSystem/NetRefHandleManager.cpp
+++ b/Engine/Source/Runtime/Experimental/Iris/Core/Private/Iris/ReplicationSystem/NetRefHandleManager.cpp
@@ -269,7 +269,7 @@ FInternalNetRefIndex FNetRefHandleManager::InternalCreateNetObject(const FNetRef
 
 void FNetRefHandleManager::AttachInstanceProtocol(FInternalNetRefIndex InternalIndex, const FReplicationInstanceProtocol* InstanceProtocol, UObject* Instance)
 {
-	if (ensureAlways((InternalIndex != InvalidInternalIndex) && InstanceProtocol))
+	if (ensure((InternalIndex != InvalidInternalIndex) && InstanceProtocol))
 	{
 		FReplicatedObjectData& Data = ReplicatedObjectData[InternalIndex];
 		Data.InstanceProtocol = InstanceProtocol;
@@ -657,7 +657,7 @@ void FNetRefHandleManager::InternalRemoveSubObject(FInternalNetRefIndex OwnerInt
 				if (SubObjects.GetInternalChildSubObjectAndConditionalArrays(SubObjectData.SubObjectParentIndex, ChildSubObjectArray, SubObjectConditionsArray))
 				{
 					const int32 ArrayIndex = ChildSubObjectArray->Find(SubObjectInternalIndex);
-					if (ensureAlways(ArrayIndex != INDEX_NONE))
+					if (ensure(ArrayIndex != INDEX_NONE))
 					{
 						ChildSubObjectArray->RemoveAt(ArrayIndex);
 						if (SubObjectConditionsArray)
@@ -951,7 +951,7 @@ FNetRefHandle FNetRefHandleManager::MakeNetRefHandle(uint64 Id, uint32 Replicati
 FNetRefHandle FNetRefHandleManager::MakeNetRefHandleFromId(uint64 Id)
 {
 	// This is called on the receiving end when deserializing replicated objects. We don't want to crash on bit stream errors leading to invalid handle IDs being read.
-	ensureAlways((Id & FNetRefHandle::IdMask) == Id);
+	ensure((Id & FNetRefHandle::IdMask) == Id);
 
 	FNetRefHandle Handle;
 
