Chaos::add per-joint linear/nonlinear solver option to joints. It's available in physics editor. It can be used for long thin objects when they are not solved correctly by the linear solver.

diff --git a/Engine/Plugins/Experimental/PhysicsControl/Source/PhysicsControl/Private/RigidBodyWithControl.cpp b/Engine/Plugins/Experimental/PhysicsControl/Source/PhysicsControl/Private/RigidBodyWithControl.cpp
--- a/Engine/Plugins/Experimental/PhysicsControl/Source/PhysicsControl/Private/RigidBodyWithControl.cpp
+++ b/Engine/Plugins/Experimental/PhysicsControl/Source/PhysicsControl/Private/RigidBodyWithControl.cpp
@@ -196,6 +196,8 @@ ImmediatePhysics::FJointHandle* FAnimNode_RigidBodyWithControl::CreateConstraint
 
 		// For control, we shouldn't be in situations where mass conditioning is needed.
 		Settings.bMassConditioningEnabled = false;
+
+		Settings.bUseLinearSolver = true;
 		// It's not our job to change collision settings - that should come from the physics asset.
 		// However, the naming of this is unclear - if collisions are disabled in the physics asset,
 		// trust that this doesn't enable them.
diff --git a/Engine/Source/Runtime/Engine/Classes/PhysicsEngine/ConstraintInstance.h b/Engine/Source/Runtime/Engine/Classes/PhysicsEngine/ConstraintInstance.h
--- a/Engine/Source/Runtime/Engine/Classes/PhysicsEngine/ConstraintInstance.h
+++ b/Engine/Source/Runtime/Engine/Classes/PhysicsEngine/ConstraintInstance.h
@@ -148,6 +148,10 @@ struct FConstraintProfileProperties
 	UPROPERTY(EditAnywhere, Category = Projection)
 	uint8 bEnableMassConditioning : 1;
 
+	/** Whether to use linear solver for this joint. Linear solver is faster and less accurate. The nonlinear solver is slower and more convergent.  */
+	UPROPERTY(EditAnywhere, Category = Projection)
+	uint8 bUseLinearJointSolver : 1;
+
 	/** Whether it is possible to break the joint with angular force. */
 	UPROPERTY(EditAnywhere, AdvancedDisplay, Category = Angular)
 	uint8 bAngularBreakable : 1;
diff --git a/Engine/Source/Runtime/Engine/Private/PhysicsEngine/ConstraintInstance.cpp b/Engine/Source/Runtime/Engine/Private/PhysicsEngine/ConstraintInstance.cpp
--- a/Engine/Source/Runtime/Engine/Private/PhysicsEngine/ConstraintInstance.cpp
+++ b/Engine/Source/Runtime/Engine/Private/PhysicsEngine/ConstraintInstance.cpp
@@ -189,6 +189,7 @@ FConstraintProfileProperties::FConstraintProfileProperties()
 	, bEnableShockPropagation(false)
 	, bEnableProjection(true)
 	, bEnableMassConditioning(true)
+	, bUseLinearJointSolver(true)
 	, bAngularBreakable(false)
 	, bAngularPlasticity(false)
 	, bLinearBreakable(false)
@@ -554,6 +555,7 @@ void FConstraintProfileProperties::UpdateConstraintFlags_AssumesLocked(const FPh
 	FPhysicsInterface::SetShockPropagationEnabled_AssumesLocked(InConstraintRef, bEnableShockPropagation, ShockPropagationAlpha);
 	FPhysicsInterface::SetParentDominates_AssumesLocked(InConstraintRef, bParentDominates);
 	FPhysicsInterface::SetMassConditioningEnabled_AssumesLocked(InConstraintRef, bEnableMassConditioning);
+	FPhysicsInterface::SetUseLinearJointSolver_AssumesLocked(InConstraintRef, bUseLinearJointSolver);
 }
 
 
diff --git a/Engine/Source/Runtime/Engine/Private/PhysicsEngine/ImmediatePhysics/ImmediatePhysicsChaos/ImmediatePhysicsJointHandle_Chaos.cpp b/Engine/Source/Runtime/Engine/Private/PhysicsEngine/ImmediatePhysics/ImmediatePhysicsChaos/ImmediatePhysicsJointHandle_Chaos.cpp
--- a/Engine/Source/Runtime/Engine/Private/PhysicsEngine/ImmediatePhysics/ImmediatePhysicsChaos/ImmediatePhysicsJointHandle_Chaos.cpp
+++ b/Engine/Source/Runtime/Engine/Private/PhysicsEngine/ImmediatePhysics/ImmediatePhysicsChaos/ImmediatePhysicsJointHandle_Chaos.cpp
@@ -73,6 +73,8 @@ namespace ImmediatePhysics_Chaos
 		JointSettings.bShockPropagationEnabled = Profile.bEnableShockPropagation;
 		JointSettings.bMassConditioningEnabled = Profile.bEnableMassConditioning;
 
+		JointSettings.bUseLinearSolver = Profile.bUseLinearJointSolver;
+
 		JointSettings.LinearMotionTypes[0] = ConvertToJointMotionType(Profile.LinearLimit.XMotion);
 		JointSettings.LinearMotionTypes[1] = ConvertToJointMotionType(Profile.LinearLimit.YMotion);
 		JointSettings.LinearMotionTypes[2] = ConvertToJointMotionType(Profile.LinearLimit.ZMotion);
diff --git a/Engine/Source/Runtime/Engine/Private/PhysicsEngine/ImmediatePhysics/ImmediatePhysicsChaos/ImmediatePhysicsSimulation_Chaos.cpp b/Engine/Source/Runtime/Engine/Private/PhysicsEngine/ImmediatePhysics/ImmediatePhysicsChaos/ImmediatePhysicsSimulation_Chaos.cpp
--- a/Engine/Source/Runtime/Engine/Private/PhysicsEngine/ImmediatePhysics/ImmediatePhysicsChaos/ImmediatePhysicsSimulation_Chaos.cpp
+++ b/Engine/Source/Runtime/Engine/Private/PhysicsEngine/ImmediatePhysics/ImmediatePhysicsChaos/ImmediatePhysicsSimulation_Chaos.cpp
@@ -818,7 +818,10 @@ namespace ImmediatePhysics_Chaos
 
 		if (UseLinearJointSolver >= 0)
 		{
-			Implementation->Joints.SetUseLinearJointSolver(UseLinearJointSolver != 0);
+			for (int32 ConstraintIndex = 0; ConstraintIndex < Implementation->Joints.GetNumConstraints(); ConstraintIndex++)
+			{
+				Implementation->Joints.SetUseLinearSolver(ConstraintIndex, UseLinearJointSolver != 0);
+			}
 			Implementation->Joints.SetUseSimd(bChaosImmediate_UseSimdForLinearSolver); // RBAN solver can use safely the SIMD version
 		}
 
@@ -999,7 +1002,7 @@ namespace ImmediatePhysics_Chaos
 		SCOPE_CYCLE_COUNTER(STAT_ImmediateSimulate_Chaos_InertiaConditioning);
 
 		// Linear and Non-linear joint solvers have a separate cvar to enable inertia conditioning
-		const bool bUseInertiaConditioning = Implementation->Joints.GetSettings().bUseLinearSolver ? bChaosImmediate_LinearInertiaConditioningEnabled : bChaosImmediate_NonLinearInertiaConditioningEnabled;
+		const bool bUseInertiaConditioning = ChaosImmediate_Joint_UseLinearSolver ? bChaosImmediate_LinearInertiaConditioningEnabled : bChaosImmediate_NonLinearInertiaConditioningEnabled;
 
 		if (bUseInertiaConditioning)
 		{
diff --git a/Engine/Source/Runtime/Experimental/Chaos/Private/Chaos/Joint/PBDJointContainerSolver.cpp b/Engine/Source/Runtime/Experimental/Chaos/Private/Chaos/Joint/PBDJointContainerSolver.cpp
--- a/Engine/Source/Runtime/Experimental/Chaos/Private/Chaos/Joint/PBDJointContainerSolver.cpp
+++ b/Engine/Source/Runtime/Experimental/Chaos/Private/Chaos/Joint/PBDJointContainerSolver.cpp
@@ -120,27 +120,33 @@ namespace Chaos
 		{
 		}
 
-		bool FPBDJointContainerSolver::UseLinearSolver() const
+		bool FPBDJointContainerSolver::UseLinearSolver() const 
 		{
-			return ConstraintContainer.GetSettings().bUseLinearSolver;
+			return NonLinearConstraintSolvers.Num() == 0;
 		}
 
 		void FPBDJointContainerSolver::Reset(const int32 InMaxConstraints)
 		{
-			ContainerConstraintIndices.Reset(InMaxConstraints);
-
-			// NOTE: We presize the solver arrays to avoid repeated calls to Add(). We resize down later if
-			// the InMaxConstraints turns out to be an over-estimate (which is should never be, currently)
-			if (UseLinearSolver())
-			{
-				LinearConstraintSolvers.SetNum(InMaxConstraints);
-				NonLinearConstraintSolvers.Empty();
-			}
-			else
+			int32 NumNonLinearConstraints = 0;
+			int32 NumLinearConstraints = 0;
+			for (int32 i = 0; i < InMaxConstraints; i++)
 			{
-				LinearConstraintSolvers.Empty();
-				NonLinearConstraintSolvers.SetNum(InMaxConstraints);
+				if (i < ConstraintContainer.GetNumConstraints())
+				{
+					if (ConstraintContainer.GetConstraintSettings(i).bUseLinearSolver)
+					{
+						NumLinearConstraints += 1;
+					}
+					else
+					{
+						NumNonLinearConstraints += 1;
+					}
+				}
 			}
+			LinearConstraintSolvers.SetNum(NumLinearConstraints);
+			NonLinearConstraintSolvers.SetNum(NumNonLinearConstraints);
+			ContainerLinearConstraintGlobalIndices.Reset(NumLinearConstraints);
+			ContainerNonLinearConstraintGlobalIndices.Reset(NumNonLinearConstraints);
 		}
 
 		void FPBDJointContainerSolver::AddConstraints()
@@ -152,26 +158,30 @@ namespace Chaos
 			{
 				if (ConstraintContainer.IsConstraintEnabled(ContainerConstraintIndex))
 				{
-					AddConstraint(ContainerConstraintIndex);
+					AddConstraint(ContainerConstraintIndex, ConstraintContainer.GetConstraintSettings(ContainerConstraintIndex).bUseLinearSolver);
 				}
 			}
 		}
 
 		void FPBDJointContainerSolver::AddConstraints(const TArrayView<Private::FPBDIslandConstraint*>& IslandConstraints)
 		{
+			int32 LinearSolverIndex = 0;
+			int32 NonLinearSolverIndex = 0;
+			int32 CurrentIndex = 0;
 			for (Private::FPBDIslandConstraint* IslandConstraint : IslandConstraints)
 			{
 				// We will only ever be given constraints from our container (asserts in non-shipping)
 				const int32 ContainerConstraintIndex = IslandConstraint->GetConstraint()->AsUnsafe<FPBDJointConstraintHandle>()->GetConstraintIndex();
 
-				AddConstraint(ContainerConstraintIndex);
+				AddConstraint(ContainerConstraintIndex, ConstraintContainer.GetConstraintSettings(ContainerConstraintIndex).bUseLinearSolver);
 			}
 		}
 
-		void FPBDJointContainerSolver::AddConstraint(const int32 InContainerConstraintIndex)
+		void FPBDJointContainerSolver::AddConstraint(const int32 InContainerConstraintIndex, const bool bUseLinearSolver)
 		{
 			// If this triggers, Reset was called with the wrong constraint count
-			check(ContainerConstraintIndices.Num() < ContainerConstraintIndices.Max());
+			check(ContainerLinearConstraintGlobalIndices.Num() <= ContainerLinearConstraintGlobalIndices.Max());
+			check(ContainerNonLinearConstraintGlobalIndices.Num() <= ContainerNonLinearConstraintGlobalIndices.Max());
 
 			// Only add a constraint if it is working on at least one dynamic body
 			const FGenericParticleHandle Particle0 = GetJointParticle(ConstraintContainer, InContainerConstraintIndex, 0);
@@ -179,84 +189,83 @@ namespace Chaos
 
 			if (Particle0->IsDynamic() || Particle1->IsDynamic())
 			{
-				ContainerConstraintIndices.Add(InContainerConstraintIndex);
+				if (bUseLinearSolver)
+				{
+					ContainerLinearConstraintGlobalIndices.Add(InContainerConstraintIndex);
+				}
+				else
+				{
+					ContainerNonLinearConstraintGlobalIndices.Add(InContainerConstraintIndex);
+				}
 			}
 		}
 
-		void FPBDJointContainerSolver::AddBodies(FSolverBodyContainer& SolverBodyContainer)
+		template <typename SolverType>
+		void AddBodiesImpl(const FPBDJointContainerSolver& Container, const TArray<int32>& SolverGlobalIndices, FSolverBodyContainer& SolverBodyContainer, TArray<SolverType>& Solvers)
 		{
-			for (int32 SolverConstraintIndex = 0, SolverConstraintEndIndex = ContainerConstraintIndices.Num(); SolverConstraintIndex < SolverConstraintEndIndex; ++SolverConstraintIndex)
+			for (int32 i = 0; i < SolverGlobalIndices.Num(); i++)
 			{
-				const int32 ContainerConstraintIndex = ContainerConstraintIndices[SolverConstraintIndex];
-
-				FGenericParticleHandle Particle0 = GetJointParticle(ConstraintContainer, ContainerConstraintIndex, 0);
-				FGenericParticleHandle Particle1 = GetJointParticle(ConstraintContainer, ContainerConstraintIndex, 1);
+				const int32 ContainerConstraintIndex = SolverGlobalIndices[i];
+				FGenericParticleHandle Particle0 = GetJointParticle(Container.GetContainer(), ContainerConstraintIndex, 0);
+				FGenericParticleHandle Particle1 = GetJointParticle(Container.GetContainer(), ContainerConstraintIndex, 1);
 
 				FSolverBody* SolverBody0 = SolverBodyContainer.FindOrAdd(Particle0);
 				FSolverBody* SolverBody1 = SolverBodyContainer.FindOrAdd(Particle1);
 
-				if (UseLinearSolver())
-				{
-					LinearConstraintSolvers[SolverConstraintIndex].SetSolverBodies(SolverBody0, SolverBody1);
-				}
-				else
-				{
-					NonLinearConstraintSolvers[SolverConstraintIndex].SetSolverBodies(SolverBody0, SolverBody1);
-				}
+				Solvers[i].SetSolverBodies(SolverBody0, SolverBody1);
 			}
+
+		}
+
+		void FPBDJointContainerSolver::AddBodies(FSolverBodyContainer& SolverBodyContainer)
+		{
+			AddBodiesImpl(*this, ContainerLinearConstraintGlobalIndices, SolverBodyContainer, LinearConstraintSolvers);
+			AddBodiesImpl(*this, ContainerNonLinearConstraintGlobalIndices, SolverBodyContainer, NonLinearConstraintSolvers);
 		}
 
 		void FPBDJointContainerSolver::GatherInput(const FReal Dt)
 		{
-			GatherInput(Dt, 0, ContainerConstraintIndices.Num());
+			GatherInput(Dt, 0, GetNumConstraints());
 		}
 
 		template<typename SolverType>
-		void GatherInputImpl(const FPBDJointContainerSolver& Container, TArray<SolverType>& Solvers, const FReal Dt, const int32 SolverConstraintBeginIndex, const int32 SolverConstraintEndIndex)
+		void GatherInputImpl(const FPBDJointContainerSolver& Container, TArray<SolverType>& Solvers, const TArray<int32>& SolverGlobalIndices, const FReal Dt, const int32 SolverConstraintBeginIndex, const int32 SolverConstraintEndIndex, const bool bUseLinearSolver)
 		{
 			const FPBDJointSolverSettings& SolverSettings = Container.GetSettings();
 
-			for (int32 SolverConstraintIndex = SolverConstraintBeginIndex; SolverConstraintIndex < SolverConstraintEndIndex; ++SolverConstraintIndex)
+			for (int32 SolverConstraintIndex = FMath::Max(0, SolverConstraintBeginIndex); SolverConstraintIndex < FMath::Min(SolverConstraintEndIndex, SolverGlobalIndices.Num()); ++SolverConstraintIndex)
 			{
-				const FPBDJointSettings& JointSettings = Container.GetConstraintSettings(SolverConstraintIndex);
+				const FPBDJointSettings& JointSettings = Container.GetConstraintSettings(SolverConstraintIndex, bUseLinearSolver);
 
-				const int32 ContainerConstraintIndex = Container.GetContainerConstraintIndex(SolverConstraintIndex);
+				const int32 ContainerConstraintIndex = Container.GetContainerConstraintIndex(SolverConstraintIndex, bUseLinearSolver);
 				FGenericParticleHandle Particle0 = GetJointParticle(Container.GetContainer(), ContainerConstraintIndex, 0);
 				FGenericParticleHandle Particle1 = GetJointParticle(Container.GetContainer(), ContainerConstraintIndex, 1);
 				const FRigidTransform3& Frame0 = GetJointFrame(Container.GetContainer(), ContainerConstraintIndex, 0);
 				const FRigidTransform3& Frame1 = GetJointFrame(Container.GetContainer(), ContainerConstraintIndex, 1);
-
 				Solvers[SolverConstraintIndex].Init(Dt, SolverSettings, JointSettings, Particle0->GetComRelativeTransform(Frame0), Particle1->GetComRelativeTransform(Frame1));
 			}
 		}
 
 		void FPBDJointContainerSolver::GatherInput(const FReal Dt, const int32 ConstraintBeginIndex, const int32 ConstraintEndIndex)
 		{
-			if (UseLinearSolver())
-			{
-				GatherInputImpl(*this, LinearConstraintSolvers, Dt, ConstraintBeginIndex, ConstraintEndIndex);
-			}
-			else
-			{
-				GatherInputImpl(*this, NonLinearConstraintSolvers, Dt, ConstraintBeginIndex, ConstraintEndIndex);
-			}
+			GatherInputImpl(*this, LinearConstraintSolvers, ContainerLinearConstraintGlobalIndices, Dt, ConstraintBeginIndex, ConstraintEndIndex, true);
+			GatherInputImpl(*this, NonLinearConstraintSolvers, ContainerNonLinearConstraintGlobalIndices, Dt, ConstraintBeginIndex - ContainerLinearConstraintGlobalIndices.Num(), ConstraintEndIndex - -ContainerLinearConstraintGlobalIndices.Num(), false);
 		}
 
 		void FPBDJointContainerSolver::ScatterOutput(const FReal Dt)
 		{
-			ScatterOutput(Dt, 0, ContainerConstraintIndices.Num());
+			ScatterOutput(Dt, 0, ContainerLinearConstraintGlobalIndices.Num() + ContainerNonLinearConstraintGlobalIndices.Num());
 		}
 
 		template<typename SolverType>
-		void ScatterOutputImpl(const FPBDJointContainerSolver& Container, TArray<SolverType>& Solvers, const FReal Dt, const int32 SolverConstraintBeginIndex, const int32 SolverConstraintEndIndex)
+		void ScatterOutputImpl(const FPBDJointContainerSolver& Container, TArray<SolverType>& Solvers, const TArray<int32>& SolverGlobalIndices, const FReal Dt, const int32 SolverConstraintBeginIndex, const int32 SolverConstraintEndIndex, const bool bUseLinearSolver)
 		{
-			for (int32 SolverConstraintIndex = SolverConstraintBeginIndex; SolverConstraintIndex < SolverConstraintEndIndex; ++SolverConstraintIndex)
+			for (int32 SolverConstraintIndex = FMath::Max(0, SolverConstraintBeginIndex); SolverConstraintIndex < FMath::Min(SolverConstraintEndIndex, SolverGlobalIndices.Num()); ++SolverConstraintIndex)
 			{
-				const int32 ContainerConstraintIndex = Container.GetContainerConstraintIndex(SolverConstraintIndex);
+				const int32 GlobalContainerConstraintIndex = Container.GetContainerConstraintIndex(SolverConstraintIndex, bUseLinearSolver);
 				if (Dt > UE_SMALL_NUMBER)
 				{
 					SolverType& Solver = Solvers[SolverConstraintIndex];
-
 					// NOTE: Particle order was revered in the solver...
 					// NOTE: Solver impulses are positional impulses
 					const FVec3 LinearImpulse = -Solver.GetNetLinearImpulse() / Dt;
@@ -264,52 +273,39 @@ namespace Chaos
 					const FSolverBody* SolverBody0 = &Solver.Body0().SolverBody();
 					const FSolverBody* SolverBody1 = &Solver.Body1().SolverBody();
 					const bool bIsBroken = Solver.IsBroken();
+
 					const bool bIsViolating = Solver.IsViolating();
 					const float LinearViolation = bIsViolating ? FMath::Sqrt((float)Solver.GetLinearViolationSq()) : 0.f;
 					const float AngularViolation = bIsViolating ? (float)Solver.GetAngularViolation() : 0.f;
 
-					Container.GetContainer().SetSolverResults(ContainerConstraintIndex, LinearImpulse, AngularImpulse, LinearViolation, AngularViolation, bIsBroken, bIsViolating, SolverBody0, SolverBody1);
+					Container.GetContainer().SetSolverResults(GlobalContainerConstraintIndex, LinearImpulse, AngularImpulse, LinearViolation, AngularViolation, bIsBroken, bIsViolating, SolverBody0, SolverBody1);
 
 					Solver.Deinit();
 				}
 				else
 				{
-					Container.GetContainer().SetSolverResults(ContainerConstraintIndex, FVec3(0), FVec3(0), 0.f, 0.f, false, false, nullptr, nullptr);
+					Container.GetContainer().SetSolverResults(GlobalContainerConstraintIndex, FVec3(0), FVec3(0), 0.f, 0.f, false, false, nullptr, nullptr);
 				}
 			}
 		}
 
 		void FPBDJointContainerSolver::ScatterOutput(const FReal Dt, const int32 ConstraintBeginIndex, const int32 ConstraintEndIndex)
 		{
-			if (UseLinearSolver())
-			{
-				ScatterOutputImpl(*this, LinearConstraintSolvers, Dt, ConstraintBeginIndex, ConstraintEndIndex);
-			}
-			else
-			{
-				ScatterOutputImpl(*this, NonLinearConstraintSolvers, Dt, ConstraintBeginIndex, ConstraintEndIndex);
-			}
+			ScatterOutputImpl(*this, LinearConstraintSolvers, ContainerLinearConstraintGlobalIndices, Dt, ConstraintBeginIndex, ConstraintEndIndex, true);
+			ScatterOutputImpl(*this, NonLinearConstraintSolvers, ContainerNonLinearConstraintGlobalIndices, Dt, ConstraintBeginIndex - ContainerLinearConstraintGlobalIndices.Num(), ConstraintEndIndex - ContainerLinearConstraintGlobalIndices.Num(), false);
 		}
 
 		void FPBDJointContainerSolver::ResizeSolverArrays()
 		{
-			// We may have conservatively allocated the solver arrays. If so, reduce their size now
-			const int32 NumConstraints = GetNumConstraints();
-			if (UseLinearSolver())
-			{
-				check(LinearConstraintSolvers.Num() >= NumConstraints);
-				LinearConstraintSolvers.SetNum(NumConstraints);
-			}
-			else
-			{
-				check(NonLinearConstraintSolvers.Num() >= NumConstraints);
-				NonLinearConstraintSolvers.SetNum(NumConstraints);
-			}
+			check(LinearConstraintSolvers.Num() >= ContainerLinearConstraintGlobalIndices.Num());
+			LinearConstraintSolvers.SetNum(ContainerLinearConstraintGlobalIndices.Num());
+			check(NonLinearConstraintSolvers.Num() >= ContainerNonLinearConstraintGlobalIndices.Num());
+			NonLinearConstraintSolvers.SetNum(ContainerNonLinearConstraintGlobalIndices.Num());
 		}
 
 		// Apply position constraints for linear or non-linear solvers
 		template<typename SolverType>
-		void ApplyPositionConstraintsImpl(const FPBDJointContainerSolver& Container, TArray<SolverType>& Solvers, const FReal Dt, const int32 It, const int32 NumIts)
+		void ApplyPositionConstraintsImpl(const FPBDJointContainerSolver& Container, TArray<SolverType>& Solvers, const FReal Dt, const int32 It, const int32 NumIts, const bool bUseLinearSolver)
 		{
 			const FPBDJointSolverSettings& Settings = Container.GetSettings();
 			const FReal IterationStiffness = GetJointIterationStiffness(Settings, It, NumIts);
@@ -322,7 +318,7 @@ namespace Chaos
 					continue;
 				}
 
-				const FPBDJointSettings& JointSettings = Container.GetConstraintSettings(SolverConstraintIndex);
+				const FPBDJointSettings& JointSettings = Container.GetConstraintSettings(SolverConstraintIndex, bUseLinearSolver);
 				Solver.Update(Dt, Settings, JointSettings);
 
 				// Set parent inverse mass scale based on current shock propagation state
@@ -345,13 +341,15 @@ namespace Chaos
 		{
 			ResizeSolverArrays();
 
-			ApplyPositionConstraintsImpl(*this, LinearConstraintSolvers, Dt, It, NumIts);
-			ApplyPositionConstraintsImpl(*this, NonLinearConstraintSolvers, Dt, It, NumIts);
+			constexpr bool bUseSolveLinear = true;
+
+			ApplyPositionConstraintsImpl(*this, LinearConstraintSolvers, Dt, It, NumIts, bUseSolveLinear);
+			ApplyPositionConstraintsImpl(*this, NonLinearConstraintSolvers, Dt, It, NumIts, !bUseSolveLinear);
 		}
 
 		// Apply velocity constraints for linear or non-linear solvers
 		template<typename SolverType>
-		void ApplyVelocityConstraintsImpl(const FPBDJointContainerSolver& Container, TArray<SolverType>& Solvers, const FReal Dt, const int32 It, const int32 NumIts)
+		void ApplyVelocityConstraintsImpl(const FPBDJointContainerSolver& Container, TArray<SolverType>& Solvers, const FReal Dt, const int32 It, const int32 NumIts, const bool bUseLinearSolver)
 		{
 			const FPBDJointSolverSettings& Settings = Container.GetSettings();
 			const FReal IterationStiffness = GetJointIterationStiffness(Settings, It, NumIts);
@@ -364,7 +362,7 @@ namespace Chaos
 					continue;
 				}
 
-				const FPBDJointSettings& JointSettings = Container.GetConstraintSettings(SolverConstraintIndex);
+				const FPBDJointSettings& JointSettings = Container.GetConstraintSettings(SolverConstraintIndex, bUseLinearSolver);
 				Solver.Update(Dt, Settings, JointSettings);
 
 				// Set parent inverse mass scale based on current shock propagation state
@@ -379,8 +377,9 @@ namespace Chaos
 
 		void FPBDJointContainerSolver::ApplyVelocityConstraints(const FReal Dt, const int32 It, const int32 NumIts)
 		{
-			ApplyVelocityConstraintsImpl(*this, LinearConstraintSolvers, Dt, It, NumIts);
-			ApplyVelocityConstraintsImpl(*this, NonLinearConstraintSolvers, Dt, It, NumIts);
+			constexpr bool bUseSolveLinear = true;
+			ApplyVelocityConstraintsImpl(*this, LinearConstraintSolvers, Dt, It, NumIts, bUseSolveLinear);
+			ApplyVelocityConstraintsImpl(*this, NonLinearConstraintSolvers, Dt, It, NumIts, !bUseSolveLinear);
 		}
 
 		void FPBDJointContainerSolver::ApplyProjectionConstraints(const FReal Dt, const int32 It, const int32 NumIts)
@@ -391,6 +390,7 @@ namespace Chaos
 
 		void FPBDJointContainerSolver::ApplyLinearProjectionConstraints(const FReal Dt, const int32 It, const int32 NumIts)
 		{
+			constexpr bool bUseSolveLinear = true;
 			const FPBDJointSolverSettings& Settings = GetSettings();
 
 			if (It == 0)
@@ -407,7 +407,7 @@ namespace Chaos
 						continue;
 					}
 
-					const FPBDJointSettings& JointSettings = GetConstraintSettings(SolverConstraintIndex);
+					const FPBDJointSettings& JointSettings = GetConstraintSettings(SolverConstraintIndex, bUseSolveLinear);
 					if (!JointSettings.bProjectionEnabled)
 					{
 						continue;
@@ -425,7 +425,7 @@ namespace Chaos
 					continue;
 				}
 
-				const FPBDJointSettings& JointSettings = GetConstraintSettings(SolverConstraintIndex);
+				const FPBDJointSettings& JointSettings = GetConstraintSettings(SolverConstraintIndex, bUseSolveLinear);
 				if (!JointSettings.bProjectionEnabled)
 				{
 					continue;
@@ -443,6 +443,7 @@ namespace Chaos
 
 		void FPBDJointContainerSolver::ApplyNonLinearProjectionConstraints(const FReal Dt, const int32 It, const int32 NumIts)
 		{
+			constexpr bool bUseSolveLinear = false;
 			const FPBDJointSolverSettings& Settings = GetSettings();
 
 			for (int32 SolverConstraintIndex = 0; SolverConstraintIndex < NonLinearConstraintSolvers.Num(); ++SolverConstraintIndex)
@@ -453,7 +454,7 @@ namespace Chaos
 					continue;
 				}
 
-				const FPBDJointSettings& JointSettings = GetConstraintSettings(SolverConstraintIndex);
+				const FPBDJointSettings& JointSettings = GetConstraintSettings(SolverConstraintIndex, bUseSolveLinear);
 				if (!JointSettings.bProjectionEnabled)
 				{
 					continue;
diff --git a/Engine/Source/Runtime/Experimental/Chaos/Private/Chaos/PBDJointConstraints.cpp b/Engine/Source/Runtime/Experimental/Chaos/Private/Chaos/PBDJointConstraints.cpp
--- a/Engine/Source/Runtime/Experimental/Chaos/Private/Chaos/PBDJointConstraints.cpp
+++ b/Engine/Source/Runtime/Experimental/Chaos/Private/Chaos/PBDJointConstraints.cpp
@@ -341,6 +341,7 @@ namespace Chaos
 		, bProjectionEnabled(false)
 		, bShockPropagationEnabled(false)
 		, bMassConditioningEnabled(true)
+		, bUseLinearSolver(true)
 		, LinearMotionTypes({ EJointMotionType::Locked, EJointMotionType::Locked, EJointMotionType::Locked })
 		, LinearLimit(UE_MAX_FLT)
 		, AngularMotionTypes({ EJointMotionType::Free, EJointMotionType::Free, EJointMotionType::Free })
@@ -510,7 +511,6 @@ namespace Chaos
 		, MaxSolverStiffness(1)
 		, NumIterationsAtMaxSolverStiffness(1)
 		, NumShockPropagationIterations(0)
-		, bUseLinearSolver(true)
 		, bUseSimd(false)
 		, bSortEnabled(false)
 		, bSolvePositionLast(true)
@@ -537,6 +537,17 @@ namespace Chaos
 	{
 	}
 
+	PRAGMA_DISABLE_DEPRECATION_WARNINGS
+	
+	FPBDJointSolverSettings::FPBDJointSolverSettings(const FPBDJointSolverSettings& InSettings) = default;
+
+	FPBDJointSolverSettings FPBDJointSolverSettings::operator=(const FPBDJointSolverSettings Other)
+	{
+		return FPBDJointSolverSettings(Other);
+	}
+	
+	PRAGMA_ENABLE_DEPRECATION_WARNINGS
+
 
 	//
 	// Constraint Container
@@ -570,12 +581,12 @@ namespace Chaos
 		return Settings;
 	}
 
-	
+	PRAGMA_DISABLE_DEPRECATION_WARNINGS
 	void FPBDJointConstraints::SetSettings(const FPBDJointSolverSettings& InSettings)
 	{
 		Settings = InSettings;
 	}
-
+	PRAGMA_ENABLE_DEPRECATION_WARNINGS
 
 	int32 FPBDJointConstraints::NumConstraints() const
 	{
@@ -902,6 +913,11 @@ namespace Chaos
 		ConstraintSettings[ConstraintIndex].AngularDrivePositionTarget = InAngularDrivePositionTarget;
 	}
 
+	void FPBDJointConstraints::SetUseLinearSolver(int32 ConstraintIndex, bool InbUseLinearSolver)
+	{
+		ConstraintSettings[ConstraintIndex].bUseLinearSolver = InbUseLinearSolver;
+	}
+
 	int32 FPBDJointConstraints::GetConstraintIsland(int32 ConstraintIndex) const
 	{
 		return ConstraintStates[ConstraintIndex].Island;
diff --git a/Engine/Source/Runtime/Experimental/Chaos/Private/ChaosVisualDebugger/ChaosVDDataWrapperUtils.cpp b/Engine/Source/Runtime/Experimental/Chaos/Private/ChaosVisualDebugger/ChaosVDDataWrapperUtils.cpp
--- a/Engine/Source/Runtime/Experimental/Chaos/Private/ChaosVisualDebugger/ChaosVDDataWrapperUtils.cpp
+++ b/Engine/Source/Runtime/Experimental/Chaos/Private/ChaosVisualDebugger/ChaosVDDataWrapperUtils.cpp
@@ -298,6 +298,7 @@ FChaosVDJointConstraint FChaosVDDataWrapperUtils::BuildJointDataWrapper(const Ch
 		CVD_COPY_FIELD_TO_WRAPPER(ConstaintHanlde->GetJointSettings(), WrappedJointData.JointSettings, ParentInvMassScale);
 		CVD_COPY_FIELD_TO_WRAPPER(ConstaintHanlde->GetJointSettings(), WrappedJointData.JointSettings, bCollisionEnabled);
 		CVD_COPY_FIELD_TO_WRAPPER(ConstaintHanlde->GetJointSettings(), WrappedJointData.JointSettings, bMassConditioningEnabled);
+		CVD_COPY_FIELD_TO_WRAPPER(ConstaintHanlde->GetJointSettings(), WrappedJointData.JointSettings, bUseLinearSolver);
 		CVD_COPY_FIELD_TO_WRAPPER(ConstaintHanlde->GetJointSettings(), WrappedJointData.JointSettings, bSoftLinearLimitsEnabled);
 		CVD_COPY_FIELD_TO_WRAPPER(ConstaintHanlde->GetJointSettings(), WrappedJointData.JointSettings, bSoftTwistLimitsEnabled);
 		CVD_COPY_FIELD_TO_WRAPPER(ConstaintHanlde->GetJointSettings(), WrappedJointData.JointSettings, bAngularSLerpPositionDriveEnabled);
diff --git a/Engine/Source/Runtime/Experimental/Chaos/Private/PBDRigidsSolver.cpp b/Engine/Source/Runtime/Experimental/Chaos/Private/PBDRigidsSolver.cpp
--- a/Engine/Source/Runtime/Experimental/Chaos/Private/PBDRigidsSolver.cpp
+++ b/Engine/Source/Runtime/Experimental/Chaos/Private/PBDRigidsSolver.cpp
@@ -289,10 +289,6 @@ namespace Chaos
 		int32 ChaosSolverParticlePoolNumFrameUntilShrink = 30;
 		FAutoConsoleVariableRef CVarChaosSolverParticlePoolNumFrameUntilShrink(TEXT("p.Chaos.Solver.ParticlePoolNumFrameUntilShrink"), ChaosSolverParticlePoolNumFrameUntilShrink, TEXT("Num Frame until we can potentially shrink the pool"));
 
-		// Joint solver mode (linear vs non-linear)
-		bool bChaosSolverJointUseLinearSolver = true;
-		FAutoConsoleVariableRef CVarChaosSolverJointUseCachedSolver(TEXT("p.Chaos.Solver.Joint.UseLinearSolver"), bChaosSolverJointUseLinearSolver, TEXT("Use linear version of joint solver. (default is true"));
-
 		// Enable/Disable collisions
 		bool bChaosSolverCollisionEnabled = true;
 		FAutoConsoleVariableRef CVarChaosSolverCollisionDisable(TEXT("p.Chaos.Solver.Collision.Enabled"), bChaosSolverCollisionEnabled, TEXT("Enable/Disable collisions in the main scene."));
@@ -2808,7 +2804,6 @@ TRACE_COUNTER_SET(ChaosTraceCounter_##Name, Value)
 		JointsSettings.bUsePositionBasedDrives = bChaosSolverJointUsePositionBasedDrives;
 		JointsSettings.NumShockPropagationIterations = ChaosSolverJointNumShockProagationIterations;
 		JointsSettings.ShockPropagationOverride = ChaosSolverJointShockPropagation;
-		JointsSettings.bUseLinearSolver = bChaosSolverJointUseLinearSolver;
 		JointsSettings.bSortEnabled = false;
 		MEvolution->GetJointConstraints().SetSettings(JointsSettings);
 
diff --git a/Engine/Source/Runtime/Experimental/Chaos/Public/Chaos/Joint/PBDJointContainerSolver.h b/Engine/Source/Runtime/Experimental/Chaos/Public/Chaos/Joint/PBDJointContainerSolver.h
--- a/Engine/Source/Runtime/Experimental/Chaos/Public/Chaos/Joint/PBDJointContainerSolver.h
+++ b/Engine/Source/Runtime/Experimental/Chaos/Public/Chaos/Joint/PBDJointContainerSolver.h
@@ -28,7 +28,7 @@ namespace Chaos
 			~FPBDJointContainerSolver();
 
 			// FConstraintContainerSolver impl
-			virtual int32 GetNumConstraints() const override final { return ContainerConstraintIndices.Num(); }
+			virtual int32 GetNumConstraints() const override final { return ContainerLinearConstraintGlobalIndices.Num() + ContainerNonLinearConstraintGlobalIndices.Num();}
 			virtual void Reset(const int32 InMaxCollisions) override final;
 			virtual void AddConstraints() override final;
 			virtual void AddConstraints(const TArrayView<Private::FPBDIslandConstraint*>& IslandConstraints) override final;
@@ -43,24 +43,52 @@ namespace Chaos
 
 			FPBDJointConstraints& GetContainer() const { return ConstraintContainer; }
 			const FPBDJointSolverSettings& GetSettings() const { return ConstraintContainer.GetSettings(); }
-			const FPBDJointSettings& GetConstraintSettings(const int32 InConstraintIndex) const { return ConstraintContainer.GetConstraintSettings(ContainerConstraintIndices[InConstraintIndex]); }
-			int32 GetContainerConstraintIndex(const int32 InConstraintIndex) const { return ContainerConstraintIndices[InConstraintIndex]; }
+			const FPBDJointSettings& GetLinearConstraintSettings(const int32 InConstraintIndex) const { return ConstraintContainer.GetConstraintSettings(ContainerLinearConstraintGlobalIndices[InConstraintIndex]);}
+			const FPBDJointSettings& GetNonLinearConstraintSettings(const int32 InConstraintIndex) const { return ConstraintContainer.GetConstraintSettings(ContainerNonLinearConstraintGlobalIndices[InConstraintIndex]); }
+			const FPBDJointSettings& GetConstraintSettings(const int32 InConstraintIndex, const bool bUseLinearSolver = true) const 
+			{
+				if (bUseLinearSolver)
+				{
+					return ConstraintContainer.GetConstraintSettings(ContainerLinearConstraintGlobalIndices[InConstraintIndex]);
+				}
+				else
+				{
+					return ConstraintContainer.GetConstraintSettings(ContainerNonLinearConstraintGlobalIndices[InConstraintIndex]);
+				}
+			}
+			int32 GetContainerConstraintIndex(const int32 InConstraintIndex, const bool bUseLinearSolver) const 
+			{
+				if (bUseLinearSolver)
+				{
+					return ContainerLinearConstraintGlobalIndices[InConstraintIndex];
+				}
+				else
+				{
+					return ContainerNonLinearConstraintGlobalIndices[InConstraintIndex];
+				}
+			}
+			int32 GetContainerLinearConstraintGlobalIndex(const int32 InLocalConstraintIndex) const { return ContainerLinearConstraintGlobalIndices[InLocalConstraintIndex];}
+			int32 GetContainerNonLinearConstraintGlobalIndex(const int32 InLocalConstraintIndex) const { return ContainerNonLinearConstraintGlobalIndices[InLocalConstraintIndex]; }
+			int32 GetContainerLinearGlobalContraintIndex(const int32 InSolverIndex) const { return ContainerLinearConstraintGlobalIndices[InSolverIndex]; }
+			int32 GetContainerNonLinearGlobalContraintIndex(const int32 InSolverIndex) const { return ContainerNonLinearConstraintGlobalIndices[InSolverIndex]; }
+
 
 		private:
 			bool UseLinearSolver() const;
-			void AddConstraint(const int32 InContainerConstraintIndex);
+			void AddConstraint(const int32 InContainerConstraintIndex, const bool bUseLinearSolver = true);
 			void ResizeSolverArrays();
 			void ApplyLinearProjectionConstraints(const FReal Dt, const int32 It, const int32 NumIts);
 			void ApplyNonLinearProjectionConstraints(const FReal Dt, const int32 It, const int32 NumIts);
 
 			FPBDJointConstraints& ConstraintContainer;
 
-			// Index remapping from internal index [0,NumToSolve) to index in the joint container [0,NumJointsInWorld)
-			TArray<int32> ContainerConstraintIndices;
-
 			// The linear and non-linear joint solvers. One for each joint we wish to solve in the order that they are solved
 			TArray<FPBDJointCachedSolver> LinearConstraintSolvers;
 			TArray<FPBDJointSolver> NonLinearConstraintSolvers;
+
+			// Index remapping from respective internal index of each array to index in the joint container [0,NumJointsInWorld)
+			TArray<int32> ContainerLinearConstraintGlobalIndices;
+			TArray<int32> ContainerNonLinearConstraintGlobalIndices;
 		};
 
 	}
diff --git a/Engine/Source/Runtime/Experimental/Chaos/Public/Chaos/JointProperties.inl b/Engine/Source/Runtime/Experimental/Chaos/Public/Chaos/JointProperties.inl
--- a/Engine/Source/Runtime/Experimental/Chaos/Public/Chaos/JointProperties.inl
+++ b/Engine/Source/Runtime/Experimental/Chaos/Public/Chaos/JointProperties.inl
@@ -6,6 +6,7 @@ CHAOS_INNER_JOINT_PROPERTY(JointSettings, CollisionEnabled, bCollisionEnabled, b
 CHAOS_INNER_JOINT_PROPERTY(JointSettings, ProjectionEnabled, bProjectionEnabled, bool)
 CHAOS_INNER_JOINT_PROPERTY(JointSettings, ShockPropagationEnabled, bShockPropagationEnabled, bool)
 CHAOS_INNER_JOINT_PROPERTY(JointSettings, MassConditioningEnabled, bMassConditioningEnabled, bool)
+CHAOS_INNER_JOINT_PROPERTY(JointSettings, UseLinearSolver, bUseLinearSolver, bool)
 CHAOS_INNER_JOINT_PROPERTY(JointSettings, ProjectionLinearAlpha, LinearProjection, FReal)
 CHAOS_INNER_JOINT_PROPERTY(JointSettings, ProjectionAngularAlpha, AngularProjection, FReal)
 CHAOS_INNER_JOINT_PROPERTY(JointSettings, ProjectionLinearTolerance, TeleportDistance, FReal)
diff --git a/Engine/Source/Runtime/Experimental/Chaos/Public/Chaos/PBDJointConstraintTypes.h b/Engine/Source/Runtime/Experimental/Chaos/Public/Chaos/PBDJointConstraintTypes.h
--- a/Engine/Source/Runtime/Experimental/Chaos/Public/Chaos/PBDJointConstraintTypes.h
+++ b/Engine/Source/Runtime/Experimental/Chaos/Public/Chaos/PBDJointConstraintTypes.h
@@ -139,6 +139,8 @@ namespace Chaos
 		bool bShockPropagationEnabled;	// @chaos(todo): remove - implied by alpha
 		bool bMassConditioningEnabled;
 
+		bool bUseLinearSolver;
+
 		TVector<EJointMotionType, 3> LinearMotionTypes;
 		FReal LinearLimit;
 
@@ -207,6 +209,10 @@ namespace Chaos
 	public:
 		CHAOS_API FPBDJointSolverSettings();
 
+		CHAOS_API FPBDJointSolverSettings(const FPBDJointSolverSettings& InSettings);
+
+		CHAOS_API FPBDJointSolverSettings operator=(const FPBDJointSolverSettings Other);
+
 		// Tolerances
 		FReal SwingTwistAngleTolerance;
 		FReal PositionTolerance;
@@ -223,6 +229,7 @@ namespace Chaos
 		int32 NumShockPropagationIterations;
 
 		// Whether to use the linear or non-linear joint solver
+		UE_DEPRECATED(5.5, "No longer used, see per-constraint flag in FPBDJointSettings::bUseLinearSolver")
 		bool bUseLinearSolver;
 
 		// Whether use vectorization in single precision, only available on the linear solver
diff --git a/Engine/Source/Runtime/Experimental/Chaos/Public/Chaos/PBDJointConstraints.h b/Engine/Source/Runtime/Experimental/Chaos/Public/Chaos/PBDJointConstraints.h
--- a/Engine/Source/Runtime/Experimental/Chaos/Public/Chaos/PBDJointConstraints.h
+++ b/Engine/Source/Runtime/Experimental/Chaos/Public/Chaos/PBDJointConstraints.h
@@ -180,9 +180,10 @@ namespace Chaos
 
 		/**
 		* Whether to use a linear or non-linear joint solver. Non-linear is more stable but much more expensive.
-		* A linear solver is used by default (see FPBDJointSolverSettings).
+		* A linear solver is used by default. This API is deprecated. See FPBDJointSettings::bUseLinearSolver instead. 
 		*/
-		void SetUseLinearJointSolver(const bool bInEnable) { Settings.bUseLinearSolver = bInEnable; }
+		UE_DEPRECATED(5.5, "No longer used, see per-constraint flag in FPBDJointSettings::bUseLinearSolver")
+		void SetUseLinearJointSolver(const bool bInEnable) { }
 
 		/**
 		* Whether to use simd in single precision on the linear joint solver. 
@@ -346,6 +347,11 @@ namespace Chaos
 		*/
 		CHAOS_API void SetAngularDrivePositionTarget(int32 ConstraintIndex, FRotation3 InAngularDrivePositionTarget);
 
+		/**
+		 * Set the option to use linear solver for a constraint by index
+		*/
+		CHAOS_API void SetUseLinearSolver(int32 ConstraintIndex, bool InbUseLinearSolver);
+
 		/**
 		* The total linear impulse applied by the constraint
 		*/
diff --git a/Engine/Source/Runtime/Experimental/ChaosVisualDebugger/Private/DataWrappers/ChaosVDJointDataWrappers.cpp b/Engine/Source/Runtime/Experimental/ChaosVisualDebugger/Private/DataWrappers/ChaosVDJointDataWrappers.cpp
--- a/Engine/Source/Runtime/Experimental/ChaosVisualDebugger/Private/DataWrappers/ChaosVDJointDataWrappers.cpp
+++ b/Engine/Source/Runtime/Experimental/ChaosVisualDebugger/Private/DataWrappers/ChaosVDJointDataWrappers.cpp
@@ -187,6 +187,7 @@ bool FChaosVDJointSettingsDataWrapper::Serialize(FArchive& Ar)
 		CVD_PACK_BITFIELD_DATA(bLinearVelocityDriveEnabled0, Flags, EChaosVDJointSettingsFlags::LinearVelocityDriveEnabled0);
 		CVD_PACK_BITFIELD_DATA(bLinearVelocityDriveEnabled1, Flags, EChaosVDJointSettingsFlags::LinearVelocityDriveEnabled1);
 		CVD_PACK_BITFIELD_DATA(bLinearVelocityDriveEnabled2, Flags, EChaosVDJointSettingsFlags::LinearVelocityDriveEnabled2);
+		CVD_PACK_BITFIELD_DATA(bUseLinearSolver, Flags, EChaosVDJointSettingsFlags::UseLinearSolver);
 	}
 	else
 	{
@@ -207,7 +208,8 @@ bool FChaosVDJointSettingsDataWrapper::Serialize(FArchive& Ar)
 		CVD_UNPACK_BITFIELD_DATA(bLinearPositionDriveEnabled2, Flags, EChaosVDJointSettingsFlags::LinearPositionDriveEnable2);
 		CVD_UNPACK_BITFIELD_DATA(bLinearVelocityDriveEnabled0, Flags, EChaosVDJointSettingsFlags::LinearVelocityDriveEnabled0);
 		CVD_UNPACK_BITFIELD_DATA(bLinearVelocityDriveEnabled1, Flags, EChaosVDJointSettingsFlags::LinearVelocityDriveEnabled1);
-		CVD_UNPACK_BITFIELD_DATA(bLinearVelocityDriveEnabled2, Flags, EChaosVDJointSettingsFlags::LinearVelocityDriveEnabled2);;
+		CVD_UNPACK_BITFIELD_DATA(bLinearVelocityDriveEnabled2, Flags, EChaosVDJointSettingsFlags::LinearVelocityDriveEnabled2);
+		CVD_UNPACK_BITFIELD_DATA(bUseLinearSolver, Flags, EChaosVDJointSettingsFlags::UseLinearSolver);;
 
 		Ar << Flags;
 	}
diff --git a/Engine/Source/Runtime/Experimental/ChaosVisualDebugger/Public/DataWrappers/ChaosVDJointDataWrappers.h b/Engine/Source/Runtime/Experimental/ChaosVisualDebugger/Public/DataWrappers/ChaosVDJointDataWrappers.h
--- a/Engine/Source/Runtime/Experimental/ChaosVisualDebugger/Public/DataWrappers/ChaosVDJointDataWrappers.h
+++ b/Engine/Source/Runtime/Experimental/ChaosVisualDebugger/Public/DataWrappers/ChaosVDJointDataWrappers.h
@@ -206,7 +206,7 @@ enum class EChaosVDJointSettingsFlags : uint32
 	AngularTwistVelocityDriveEnabled = 1 << 5,
 	AngularSwingPositionDriveEnabled = 1 << 6,
 	AngularSwingVelocityDriveEnabled = 1 << 7,
-	SoftLinearLimitsEnabled= 1 << 8,
+	SoftLinearLimitsEnabled = 1 << 8,
 	SoftTwistLimitsEnabled = 1 << 9,
 	SoftSwingLimitsEnabled = 1 << 10,
 	LinearPositionDriveEnabled0 = 1 << 11,
@@ -215,6 +215,7 @@ enum class EChaosVDJointSettingsFlags : uint32
 	LinearVelocityDriveEnabled0 = 1 << 14,
 	LinearVelocityDriveEnabled1 = 1 << 15,
 	LinearVelocityDriveEnabled2 = 1 << 16,
+	UseLinearSolver = 1 << 17,
 };
 
 USTRUCT()
@@ -247,6 +248,8 @@ public:
 	UPROPERTY(VisibleAnywhere, Category="General")
 	uint16 bMassConditioningEnabled : 1 = false;
 	UPROPERTY(VisibleAnywhere, Category="General")
+	uint16 bUseLinearSolver : 1 = false;
+	UPROPERTY(VisibleAnywhere, Category="General")
 	uint16 bSoftLinearLimitsEnabled : 1 = false;
 	UPROPERTY(VisibleAnywhere, Category="General")
 	uint16 bSoftTwistLimitsEnabled : 1 = false;
diff --git a/Engine/Source/Runtime/PhysicsCore/Private/ChaosEngineInterface.cpp b/Engine/Source/Runtime/PhysicsCore/Private/ChaosEngineInterface.cpp
--- a/Engine/Source/Runtime/PhysicsCore/Private/ChaosEngineInterface.cpp
+++ b/Engine/Source/Runtime/PhysicsCore/Private/ChaosEngineInterface.cpp
@@ -1706,6 +1706,17 @@ void FChaosEngineInterface::SetMassConditioningEnabled_AssumesLocked(const FPhys
 	}
 }
 
+void FChaosEngineInterface::SetUseLinearJointSolver_AssumesLocked(const FPhysicsConstraintHandle& InConstraintRef, bool bInUseLinearJointSolver)
+{
+	if (InConstraintRef.IsValid() && InConstraintRef.Constraint->IsType(Chaos::EConstraintType::JointConstraintType))
+	{
+		if (Chaos::FJointConstraint* Constraint = static_cast<Chaos::FJointConstraint*>(InConstraintRef.Constraint))
+		{
+			Constraint->SetUseLinearSolver(bInUseLinearJointSolver);
+		}
+	}
+}
+
 void FChaosEngineInterface::SetBreakForces_AssumesLocked(const FPhysicsConstraintHandle& InConstraintRef,float InLinearBreakForce,float InAngularBreakTorque)
 {
 	if (InConstraintRef.IsValid() && InConstraintRef.Constraint->IsType(Chaos::EConstraintType::JointConstraintType))
diff --git a/Engine/Source/Runtime/PhysicsCore/Public/Chaos/ChaosEngineInterface.h b/Engine/Source/Runtime/PhysicsCore/Public/Chaos/ChaosEngineInterface.h
--- a/Engine/Source/Runtime/PhysicsCore/Public/Chaos/ChaosEngineInterface.h
+++ b/Engine/Source/Runtime/PhysicsCore/Public/Chaos/ChaosEngineInterface.h
@@ -516,6 +516,7 @@ public:
 	static PHYSICSCORE_API void SetShockPropagationEnabled_AssumesLocked(const FPhysicsConstraintHandle& InConstraintRef, bool bInShockPropagationEnabled, float InShockPropagationAlpha);
 	static PHYSICSCORE_API void SetParentDominates_AssumesLocked(const FPhysicsConstraintHandle& InConstraintRef, bool bInParentDominates);
 	static PHYSICSCORE_API void SetMassConditioningEnabled_AssumesLocked(const FPhysicsConstraintHandle& InConstraintRef, bool bInMassConditioningEnabled);
+	static PHYSICSCORE_API void SetUseLinearJointSolver_AssumesLocked(const FPhysicsConstraintHandle& InConstraintRef, bool bInUseLinearJointSolver);
 	static PHYSICSCORE_API void SetBreakForces_AssumesLocked(const FPhysicsConstraintHandle& InConstraintRef,float InLinearBreakForce,float InAngularBreakForce);
 	static PHYSICSCORE_API void SetViolationCallbackThresholds_AssumesLocked(const FPhysicsConstraintHandle& InConstraintRef,float InLinearBreakViolation,float InAngularBreakViolation);
 	static PHYSICSCORE_API void SetPlasticityLimits_AssumesLocked(const FPhysicsConstraintHandle& InConstraintRef, float InLinearPlasticityLimit, float InAngularPlasticityLimit, EConstraintPlasticityType InLinearPlasticityType);
