Geometry collection : add support for SetActorHiidenInGame when using custom renderers
- make sure to account for the hidden flag of the actor in the RefreshCustomrenderer method
- also change game specific code to account for this new behaviour

diff --git a/Engine/Source/Runtime/Experimental/GeometryCollectionEngine/Private/GeometryCollection/GeometryCollectionComponent.cpp b/Engine/Source/Runtime/Experimental/GeometryCollectionEngine/Private/GeometryCollection/GeometryCollectionComponent.cpp
--- a/Engine/Source/Runtime/Experimental/GeometryCollectionEngine/Private/GeometryCollection/GeometryCollectionComponent.cpp
+++ b/Engine/Source/Runtime/Experimental/GeometryCollectionEngine/Private/GeometryCollection/GeometryCollectionComponent.cpp
@@ -6346,8 +6346,11 @@ void UGeometryCollectionComponent::RefreshCustomRenderer()
 
 		const bool bRenderRootProxy = !bForceBrokenForCustomRenderer && !bIsBroken && (RestCollection->RootProxyData.ProxyMeshes.Num() > 0);
 
+		const AActor* Owner = GetOwner();
+		const bool bIsActorHidden = Owner ? Owner->IsHidden() : false;
+
 		uint32 StateFlags = 0;
-		StateFlags |= bHiddenInGame || !IsVisible() ? 0 : IGeometryCollectionExternalRenderInterface::EState_Visible;
+		StateFlags |= (bHiddenInGame || !IsVisible() || bIsActorHidden) ? 0 : IGeometryCollectionExternalRenderInterface::EState_Visible;
 		StateFlags |= bRenderRootProxy ? 0 : IGeometryCollectionExternalRenderInterface::EState_Broken;
 		StateFlags |= bForceBrokenForCustomRenderer ? IGeometryCollectionExternalRenderInterface::EState_ForcedBroken : 0;
 
