Geometry collection : Add GetLocalRestTransform blueprint function

diff --git a/Engine/Source/Runtime/Experimental/GeometryCollectionEngine/Private/GeometryCollection/GeometryCollectionComponent.cpp b/Engine/Source/Runtime/Experimental/GeometryCollectionEngine/Private/GeometryCollection/GeometryCollectionComponent.cpp
--- a/Engine/Source/Runtime/Experimental/GeometryCollectionEngine/Private/GeometryCollection/GeometryCollectionComponent.cpp
+++ b/Engine/Source/Runtime/Experimental/GeometryCollectionEngine/Private/GeometryCollection/GeometryCollectionComponent.cpp
@@ -6841,6 +6841,35 @@ TArray<FTransform> UGeometryCollectionComponent::GetInitialLocalRestTransforms()
 	return InitialLocalTransforms;
 }
 
+TArray<FTransform> UGeometryCollectionComponent::GetLocalRestTransforms(bool bInitialTransforms) const
+{
+	TArray<FTransform> CurrentLocalTransforms;
+	if (RestCollection && RestCollection->GetGeometryCollection())
+	{
+		const FGeometryCollection& RestGeometryCollection = *RestCollection->GetGeometryCollection();
+
+		const bool bHasRestTransformOverride = (RestTransforms.Num() > 0) && (RestTransforms.Num() == RestGeometryCollection.Transform.Num());
+		if (bInitialTransforms || !bHasRestTransformOverride)
+		{
+			GeometryCollectionAlgo::GlobalMatrices(RestCollection->GetGeometryCollection()->Transform, RestGeometryCollection.Parent, CurrentLocalTransforms);
+		}
+		else
+		{
+			GeometryCollectionAlgo::GlobalMatrices(TManagedArray<FTransform>(RestTransforms), RestGeometryCollection.Parent, CurrentLocalTransforms);
+		}
+
+		const TManagedArray<FTransform>* MassToLocal = RestGeometryCollection.FindAttribute<FTransform>("MassToLocal", FGeometryCollection::TransformGroup);
+		if (MassToLocal && CurrentLocalTransforms.Num() == MassToLocal->Num())
+		{
+			for (int32 TransformIndex = 0; TransformIndex < CurrentLocalTransforms.Num(); TransformIndex++)
+			{
+				CurrentLocalTransforms[TransformIndex] = (*MassToLocal)[TransformIndex] * CurrentLocalTransforms[TransformIndex];
+			}
+		}
+	}
+	return CurrentLocalTransforms;
+}
+
 void UGeometryCollectionComponent::SetLocalRestTransforms(const TArray<FTransform>& NewTransforms, bool bOnlyLeaves)
 {
 	if (RestCollection && RestCollection->GetGeometryCollection())
diff --git a/Engine/Source/Runtime/Experimental/GeometryCollectionEngine/Public/GeometryCollection/GeometryCollectionComponent.h b/Engine/Source/Runtime/Experimental/GeometryCollectionEngine/Public/GeometryCollection/GeometryCollectionComponent.h
--- a/Engine/Source/Runtime/Experimental/GeometryCollectionEngine/Public/GeometryCollection/GeometryCollectionComponent.h
+++ b/Engine/Source/Runtime/Experimental/GeometryCollectionEngine/Public/GeometryCollection/GeometryCollectionComponent.h
@@ -824,6 +824,14 @@ public:
 	UFUNCTION(BlueprintCallable, Category = "ChaosPhysics")
 	GEOMETRYCOLLECTIONENGINE_API TArray<FTransform> GetInitialLocalRestTransforms() const;
 
+	/**
+	* Get the rest transforms in component (local) space  space,
+	* if none have been set by SetLocalRestTransforms or if RestTransform property is empty , then the initial ones are returned
+	* @param bInitialTransform force the initial transforms to be returned ( false by default )
+	*/
+	UFUNCTION(BlueprintCallable, Category = "ChaosPhysics")
+	GEOMETRYCOLLECTIONENGINE_API TArray<FTransform> GetLocalRestTransforms(bool bInitialTransforms = false) const;
+
 	/** 
 	* Set the local rest transform, this may be different from the rest collection 
 	* If the geometry collection is already simulating those matrices will be overriden by the physics state updates
