Quantize input to 16bits primarily for networking. Local values used in simulation modules should also be using the quantized values so local and remote are simulating exactly the same.

diff --git a/Engine/Source/Runtime/Experimental/ChaosVehicles/ChaosVehiclesCore/Private/SimModule/ModuleInput.cpp b/Engine/Source/Runtime/Experimental/ChaosVehicles/ChaosVehiclesCore/Private/SimModule/ModuleInput.cpp
--- a/Engine/Source/Runtime/Experimental/ChaosVehicles/ChaosVehiclesCore/Private/SimModule/ModuleInput.cpp
+++ b/Engine/Source/Runtime/Experimental/ChaosVehicles/ChaosVehiclesCore/Private/SimModule/ModuleInput.cpp
@@ -62,25 +62,30 @@ void FModuleInputValue::SetMagnitude(float NewSize)
 
 void FModuleInputValue::Serialize(FArchive& Ar, UPackageMap* Map, bool& bOutSuccess)
 {
-	// #TODO: send only when value changes also quantize what is sent, but just get it working for now
+	// #TODO: send only value changes/deltas
 	Ar << ValueType;
 
 	switch (GetValueType())
 	{
 	case EModuleInputValueType::MBoolean:
-		Ar << Value.X;
+		{
+			uint8 State = (Value.X != 0);
+			Ar << State;
+			if (Ar.IsLoading())
+			{
+				Value.X = State ? 1.0f : 0.0f;
+			}
+		}
 		break;
 	case EModuleInputValueType::MAxis1D:
-		Ar << Value.X;
+		ModularQuantize::SerializeFixedFloat<1, 16>(Value.X, Ar);
 		break;
 	case EModuleInputValueType::MAxis2D:
-		Ar << Value.X;
-		Ar << Value.Y;
+		ModularQuantize::SerializeFixedFloat<1, 16>(Value.X, Ar);
+		ModularQuantize::SerializeFixedFloat<1, 16>(Value.Y, Ar);
 		break;
 	case EModuleInputValueType::MAxis3D:
-		Ar << Value.X;
-		Ar << Value.Y;
-		Ar << Value.Z;
+		Value.NetSerialize(Ar, Map, bOutSuccess);
 		break;
 	default:
 		checkf(false, TEXT("Unsupported value type for module input value!"));
@@ -214,12 +219,12 @@ void FModuleInputContainer::Serialize(FArchive& Ar, UPackageMap* Map, bool& bOut
 {
 	bOutSuccess = true;
 
-	int Number = InputValues.Num();
+	uint32 Number = InputValues.Num();
 	Ar << Number;
 	if (Ar.IsLoading())
 	{
 		InputValues.Reset(Number);
-		for (int I = 0; I < Number; I++)
+		for (uint32 I = 0; I < Number; I++)
 		{
 			FModuleInputValue Value;
 			Value.Serialize(Ar, Map, bOutSuccess);
@@ -228,7 +233,7 @@ void FModuleInputContainer::Serialize(FArchive& Ar, UPackageMap* Map, bool& bOut
 	}
 	else
 	{
-		for (int I = 0; I < Number; I++)
+		for (uint32 I = 0; I < Number; I++)
 		{
 			InputValues[I].Serialize(Ar, Map, bOutSuccess);
 		}
diff --git a/Engine/Source/Runtime/Experimental/ChaosVehicles/ChaosVehiclesCore/Public/SimModule/ModuleInput.h b/Engine/Source/Runtime/Experimental/ChaosVehicles/ChaosVehiclesCore/Public/SimModule/ModuleInput.h
--- a/Engine/Source/Runtime/Experimental/ChaosVehicles/ChaosVehiclesCore/Public/SimModule/ModuleInput.h
+++ b/Engine/Source/Runtime/Experimental/ChaosVehicles/ChaosVehiclesCore/Public/SimModule/ModuleInput.h
@@ -3,6 +3,7 @@
 #pragma once
 
 #include "CoreMinimal.h"
+//#include "Engine/NetSerialization.h"
 #include "Templates/SubclassOf.h"
 #include "ModuleInput.generated.h"
 
@@ -29,6 +30,126 @@ enum class EFunctionType : uint8
 	CustomCurve
 };
 
+namespace ModularQuantize
+{
+	template<int32 MaxValue, uint32 NumBits>
+	struct TCompressedFloatDetails
+	{
+		// NumBits = 8:
+		static constexpr int32 MaxBitValue = (1 << (NumBits - 1)) - 1;  //   0111 1111 - Max abs value we will serialize
+		static constexpr int32 Bias = (1 << (NumBits - 1));             //   1000 0000 - Bias to pivot around (in order to support signed values)
+		static constexpr int32 SerIntMax = (1 << (NumBits - 0));        // 1 0000 0000 - What we pass into SerializeInt
+		static constexpr int32 MaxDelta = (1 << (NumBits - 0)) - 1;     //   1111 1111 - Max delta is
+	};
+
+	template<int32 MaxValue, uint32 NumBits, typename T UE_REQUIRES(std::is_floating_point_v<T>&& NumBits < 32)>
+	bool ToCompressedFloat(const T InValue, uint32& OutCompressedFloat)
+	{
+		using Details = ModularQuantize::TCompressedFloatDetails<MaxValue, NumBits>;
+
+		bool clamp = false;
+		int64 ScaledValue;
+		if (MaxValue > Details::MaxBitValue)
+		{
+			// We have to scale this down
+			const T Scale = T(Details::MaxBitValue) / MaxValue;
+			ScaledValue = FMath::TruncToInt(Scale * InValue);
+		}
+		else
+		{
+			// We will scale up to get extra precision. But keep is a whole number preserve whole values
+			constexpr int32 Scale = Details::MaxBitValue / MaxValue;
+			ScaledValue = FMath::RoundToInt(Scale * InValue);
+		}
+
+		uint32 Delta = static_cast<uint32>(ScaledValue + Details::Bias);
+
+		if (Delta > Details::MaxDelta)
+		{
+			clamp = true;
+			Delta = static_cast<int32>(Delta) > 0 ? Details::MaxDelta : 0;
+		}
+
+		OutCompressedFloat = Delta;
+
+		return !clamp;
+	}
+
+	template<int32 MaxValue, uint32 NumBits, typename T UE_REQUIRES(std::is_floating_point_v<T>&& NumBits < 32)>
+	bool FromCompressedFloat(const uint32 InCompressed, T& OutValue )
+	{
+		using Details = ModularQuantize::TCompressedFloatDetails<MaxValue, NumBits>;
+
+		uint32 Delta = InCompressed;
+		T UnscaledValue = static_cast<T>(static_cast<int32>(Delta) - Details::Bias);
+
+		if constexpr (MaxValue > Details::MaxBitValue)
+		{
+			// We have to scale down, scale needs to be a float:
+			constexpr T InvScale = MaxValue / (T)Details::MaxBitValue;
+			OutValue = UnscaledValue * InvScale;
+		}
+		else
+		{
+			constexpr int32 Scale = Details::MaxBitValue / MaxValue;
+			constexpr T InvScale = T(1) / (T)Scale;
+
+			OutValue = UnscaledValue * InvScale;
+		}
+
+		return true;
+	}
+
+	template<int32 MaxValue, uint32 NumBits, typename T UE_REQUIRES(std::is_floating_point_v<T>&& NumBits < 32)>
+	bool WriteCompressedFloat(const T Value, FArchive& Ar)
+	{
+		using Details = ModularQuantize::TCompressedFloatDetails<MaxValue, NumBits>;
+
+		uint32 CompressedValue;
+		bool clamp = ModularQuantize::ToCompressedFloat<MaxValue, NumBits>(Value, CompressedValue);
+
+		Ar.SerializeInt(CompressedValue, Details::SerIntMax);
+
+		return !clamp;
+	}
+
+	template<int32 MaxValue, uint32 NumBits, typename T UE_REQUIRES(std::is_floating_point_v<T>&& NumBits < 32)>
+	bool ReadCompressedFloat(T& Value, FArchive& Ar)
+	{
+		using Details = ModularQuantize::TCompressedFloatDetails<MaxValue, NumBits>;
+
+		uint32 CompressedValue;
+		Ar.SerializeInt(CompressedValue, Details::SerIntMax);
+
+		ModularQuantize::FromCompressedFloat<MaxValue, NumBits>(CompressedValue, Value);
+
+		return true;
+	}
+
+	// Required because we seialize quantized vector in seperate parts depending on input type
+	template<int32 MaxValue, uint32 NumBits>
+	bool SerializeFixedFloat(double& InOutValue, FArchive& Ar)
+	{
+		if (Ar.IsSaving())
+		{
+			bool success = true;
+			success &= ModularQuantize::WriteCompressedFloat<MaxValue, NumBits>(InOutValue, Ar);
+			return success;
+		}
+
+		ModularQuantize::ReadCompressedFloat<MaxValue, NumBits>(InOutValue, Ar);
+		return true;
+	}
+
+	template<int32 MaxValue, uint32 NumBits, typename T UE_REQUIRES(std::is_floating_point_v<T>&& NumBits < 32)>
+	void QuantizeValue(T& Value)
+	{
+		uint32 CompressedValue = 0;
+		ModularQuantize::ToCompressedFloat<MaxValue, NumBits>(Value, CompressedValue);
+		ModularQuantize::FromCompressedFloat<MaxValue, NumBits>(CompressedValue, Value);
+	}
+
+}
 
 USTRUCT(BlueprintType)
 struct CHAOSVEHICLESCORE_API FModuleInputValue
@@ -49,9 +170,42 @@ public:
 	// Converting a value to a different type (e.g. Val = FVector(1, 1, 1); Val = true;) zeroes out any unused components to ensure getters continue to function correctly.
 	explicit FModuleInputValue(bool bInValue) : Value(bInValue ? 1.f : 0.f, 0.f, 0.f), ValueType(EModuleInputValueType::MBoolean) {}
 	FModuleInputValue(MAxis1D InValue) : Value(InValue, 0.f, 0.f), ValueType(EModuleInputValueType::MAxis1D) {}
-	FModuleInputValue(MAxis2D InValue) : Value(InValue, 0.f), ValueType(EModuleInputValueType::MAxis2D) {}
+	FModuleInputValue(MAxis2D InValue) : Value(InValue.X, InValue.Y, 0.f), ValueType(EModuleInputValueType::MAxis2D) {}
 	FModuleInputValue(MAxis3D InValue) : Value(InValue), ValueType(EModuleInputValueType::MAxis3D) {}
 
+	FModuleInputValue ReturnQuantized() const
+	{
+		FModuleInputValue OutValue = Value;
+
+		switch (ValueType)
+		{
+			case EModuleInputValueType::MBoolean:
+			case EModuleInputValueType::MAxis1D:
+				ModularQuantize::QuantizeValue<1, 16>(OutValue.Value.X);
+			break;
+
+			case EModuleInputValueType::MAxis2D:
+				ModularQuantize::QuantizeValue<1, 16>(OutValue.Value.X);
+				ModularQuantize::QuantizeValue<1, 16>(OutValue.Value.Y);
+			break;
+
+			case EModuleInputValueType::MAxis3D:
+				ModularQuantize::QuantizeValue<1, 16>(OutValue.Value.X);
+				ModularQuantize::QuantizeValue<1, 16>(OutValue.Value.Y);
+				ModularQuantize::QuantizeValue<1, 16>(OutValue.Value.Z);
+			break;
+
+			default:
+				checkf(false, TEXT("Unsupported value type for module input value!"));
+				break;
+
+		}
+
+		//UE_LOG(LogTemp, Warning, TEXT("Original %f vs Quantized %f"), Value.X, OutValue.Value.X);
+
+		return OutValue;
+	}
+
 	// Build a specific type with an arbitrary Axis3D value
 	FModuleInputValue(EModuleInputValueType InValueType, MAxis3D InValue) : Value(InValue), ValueType(InValueType)
 	{
@@ -164,7 +318,7 @@ public:
 	float GetMagnitudeSq() const;
 	float GetMagnitude() const;
 
-	// Serialize values, not quantized yet and not serializing deltas
+	// Serialize values
 	void Serialize(FArchive& Ar, UPackageMap* Map, bool& bOutSuccess);
 
 	void Lerp(const FModuleInputValue& Min, const FModuleInputValue& Max, float Alpha)
@@ -363,8 +517,8 @@ public:
 
 	int GetNumInputs() const { return InputValues.Num(); }
 	FModuleInputValue GetValueAtIndex(int Index) const { return InputValues[Index]; }
-	void SetValueAtIndex(int Index, const FModuleInputValue& InValue) { InputValues[Index] = InValue; }
-	void MergeValueAtIndex(int Index, const FModuleInputValue& InValue) { InputValues[Index].Merge(InValue); }
+	void SetValueAtIndex(int Index, const FModuleInputValue& InValue) { InputValues[Index] = InValue.ReturnQuantized(); }
+	void MergeValueAtIndex(int Index, const FModuleInputValue& InValue) { InputValues[Index].Merge(InValue.ReturnQuantized()); }
 
 	FModuleInputContainer& operator=(const FModuleInputContainer& Other)
 	{
@@ -441,5 +595,3 @@ public:
 	virtual void ProduceInput(int32 PhysicsStep, int32 NumSteps, const FInputNameMap& InNameMap, FModuleInputContainer& InOutContainer) {}
 };
 
-
-
