Backing out CL 35894470 as the hashing function not deterministic between server/client & ultimately causes crash.

diff --git a/Engine/Plugins/Experimental/ChaosModularVehicle/Source/ChaosModularVehicleEngine/Private/ChaosModularVehicle/ChaosSimModuleManagerAsyncCallback.cpp b/Engine/Plugins/Experimental/ChaosModularVehicle/Source/ChaosModularVehicleEngine/Private/ChaosModularVehicle/ChaosSimModuleManagerAsyncCallback.cpp
--- a/Engine/Plugins/Experimental/ChaosModularVehicle/Source/ChaosModularVehicleEngine/Private/ChaosModularVehicle/ChaosSimModuleManagerAsyncCallback.cpp
+++ b/Engine/Plugins/Experimental/ChaosModularVehicle/Source/ChaosModularVehicleEngine/Private/ChaosModularVehicle/ChaosSimModuleManagerAsyncCallback.cpp
@@ -319,30 +319,28 @@ bool FNetworkModularVehicleStates::NetSerialize(FArchive& Ar, class UPackageMap*
 		{
 			if (NumNetModules > 0)
 			{
-				uint32 ModuleTypeHash = 0;
+				FName ModuleType = NAME_None;
 				int32 SimArrayIndex = 0;
-				Ar << ModuleTypeHash;
+				Ar << ModuleType;
 				Ar << SimArrayIndex;
 
 				if (I >= ModuleData.Num())
 				{
-					if (TSharedPtr<Chaos::FModuleNetData> Data = Chaos::FModuleFactoryRegister::Get().GenerateNetData(ModuleTypeHash, SimArrayIndex))
+					if (TSharedPtr<Chaos::FModuleNetData> Data = Chaos::FModuleFactoryRegister::Get().GenerateNetData(ModuleType, SimArrayIndex))
 					{
 						ModuleData.Emplace(Data);
 					}
 				}
 				if(I <= ModuleData.Num() && ModuleData[I].IsValid())
 				{
-					check(ModuleTypeHash == GetTypeHash(ModuleData[I]->GetSimType()));
 					ModuleData[I]->Serialize(Ar);
 				}
 			}
 		}
 		else
 		{
-			int ModuleTypeHash = GetTypeHash(ModuleData[I]->GetSimType());
-			
-			Ar << ModuleTypeHash;
+			FName ModuleType = ModuleData[I]->GetSimType();
+			Ar << ModuleType;
 			Ar << ModuleData[I]->SimArrayIndex;
 			ModuleData[I]->Serialize(Ar);
 		}
diff --git a/Engine/Source/Runtime/Experimental/ChaosVehicles/ChaosVehiclesCore/Private/SimModule/ModuleFactoryRegister.cpp b/Engine/Source/Runtime/Experimental/ChaosVehicles/ChaosVehiclesCore/Private/SimModule/ModuleFactoryRegister.cpp
--- a/Engine/Source/Runtime/Experimental/ChaosVehicles/ChaosVehiclesCore/Private/SimModule/ModuleFactoryRegister.cpp
+++ b/Engine/Source/Runtime/Experimental/ChaosVehicles/ChaosVehiclesCore/Private/SimModule/ModuleFactoryRegister.cpp
@@ -7,25 +7,19 @@ namespace Chaos
 {
 	void FModuleFactoryRegister::RegisterFactory(const FName TypeName, TWeakPtr<IFactoryModule> InFactory)
 	{
-		RegisterFactory(GetTypeHash(TypeName), InFactory);
-	}
-
-	void FModuleFactoryRegister::RegisterFactory(const uint32 TypeNameHash, TWeakPtr<IFactoryModule> InFactory)
-	{
-		if (!ContainsFactory(TypeNameHash))
+		if (!ContainsFactory(TypeName))
 		{
-			RegisteredFactoriesByName.Add(TypeNameHash, InFactory);
+			RegisteredFactoriesByName.Add(TypeName, InFactory);
 		}
 	}
 
 	void FModuleFactoryRegister::RemoveFactory(TWeakPtr<IFactoryModule> InFactory)
 	{
-		for (TPair<uint32, TWeakPtr<IFactoryModule>> Pair : RegisteredFactoriesByName)
+		for (TPair<FName, TWeakPtr<IFactoryModule>> Pair : RegisteredFactoriesByName)
 		{
 			if (Pair.Value == InFactory)
-			{
-				RegisteredFactoriesByName.Remove(Pair.Key);
-			
+			{RegisteredFactoriesByName.Remove(Pair.Key);
+				
 				return;
 			}
 		}
@@ -38,21 +32,16 @@ namespace Chaos
 
 	bool FModuleFactoryRegister::ContainsFactory(const FName TypeName) const
 	{
-		return ContainsFactory(GetTypeHash(TypeName));
-	}
-
-	bool FModuleFactoryRegister::ContainsFactory(const uint32 TypeNameHash) const
-	{
-		return RegisteredFactoriesByName.Contains(TypeNameHash);
+		return RegisteredFactoriesByName.Contains(TypeName);
 	}
 
-	TSharedPtr<Chaos::FModuleNetData> FModuleFactoryRegister::GenerateNetData(const uint32 TypeNameHash, const int32 SimArrayIndex)
+	TSharedPtr<Chaos::FModuleNetData> FModuleFactoryRegister::GenerateNetData(const FName TypeName, const int32 SimArrayIndex)
 	{
 		using namespace Chaos;
 
-		if (RegisteredFactoriesByName.Contains(TypeNameHash))
+		if (RegisteredFactoriesByName.Contains(TypeName))
 		{
-			TSharedPtr<IFactoryModule> PinnedFactory = RegisteredFactoriesByName[TypeNameHash].Pin();
+			TSharedPtr<IFactoryModule> PinnedFactory = RegisteredFactoriesByName[TypeName].Pin();
 
 			if (PinnedFactory.IsValid())
 			{
@@ -61,7 +50,7 @@ namespace Chaos
 		}
 		else
 		{
-			UE_LOG(LogTemp, Error, TEXT("No factory registered for hashed type %d"), TypeNameHash);
+			UE_LOG(LogTemp, Error, TEXT("No factory registered for name '%s'"), *TypeName.ToString());
 		}
 
 		return nullptr;
diff --git a/Engine/Source/Runtime/Experimental/ChaosVehicles/ChaosVehiclesCore/Public/SimModule/ModuleFactoryRegister.h b/Engine/Source/Runtime/Experimental/ChaosVehicles/ChaosVehiclesCore/Public/SimModule/ModuleFactoryRegister.h
--- a/Engine/Source/Runtime/Experimental/ChaosVehicles/ChaosVehiclesCore/Public/SimModule/ModuleFactoryRegister.h
+++ b/Engine/Source/Runtime/Experimental/ChaosVehicles/ChaosVehiclesCore/Public/SimModule/ModuleFactoryRegister.h
@@ -22,17 +22,15 @@ namespace Chaos
 		}
 		
 		void RegisterFactory(const FName TypeName, TWeakPtr<IFactoryModule> InFactory);
-		void RegisterFactory(const uint32 TypeNameHash, TWeakPtr<IFactoryModule> InFactory);
 		void RemoveFactory(TWeakPtr<IFactoryModule> InFactory);
 		void Reset();
 		bool ContainsFactory(const FName TypeName) const;
-		bool ContainsFactory(const uint32 TypeNameHash) const;
-		TSharedPtr<Chaos::FModuleNetData> GenerateNetData(const uint32 TypeNameHash, const int32 SimArrayIndex);
+		TSharedPtr<Chaos::FModuleNetData> GenerateNetData(const FName TypeName, const int32 SimArrayIndex);
 
 	protected:
 
 		FModuleFactoryRegister() = default;
-		TMap<int32, TWeakPtr<IFactoryModule>> RegisteredFactoriesByName;
+		TMap<FName, TWeakPtr<IFactoryModule>> RegisteredFactoriesByName;
 	};
 	
 	template<typename _To, typename ..._Rest>
@@ -46,15 +44,14 @@ namespace Chaos
 		{
 			return false;
 		}
-		uint32 SimTypeNameHash = GetTypeHash(SimTypeName);
-		if(FModuleFactoryRegister::Get().ContainsFactory(SimTypeNameHash))
+		if(FModuleFactoryRegister::Get().ContainsFactory(SimTypeName))
 		{
 			return true;
 		}
 		static TSharedPtr<T> SharedFactory = MakeShared<T>(args...);
 		if (SharedFactory.IsValid())
 		{
-			FModuleFactoryRegister::Get().RegisterFactory(SimTypeNameHash, SharedFactory);
+			FModuleFactoryRegister::Get().RegisterFactory(SimTypeName, SharedFactory);
 			return true;
 		}
 		return false;
