Fix animation not working for flesh while caching in async mode

diff --git a/Engine/Plugins/Animation/MLDeformer/ChaosFleshGenerator/Source/ChaosFleshGenerator/Private/ChaosFleshGeneratorSimulation.cpp b/Engine/Plugins/Animation/MLDeformer/ChaosFleshGenerator/Source/ChaosFleshGenerator/Private/ChaosFleshGeneratorSimulation.cpp
--- a/Engine/Plugins/Animation/MLDeformer/ChaosFleshGenerator/Source/ChaosFleshGenerator/Private/ChaosFleshGeneratorSimulation.cpp
+++ b/Engine/Plugins/Animation/MLDeformer/ChaosFleshGenerator/Source/ChaosFleshGenerator/Private/ChaosFleshGeneratorSimulation.cpp
@@ -62,9 +62,9 @@ namespace UE::Chaos::FleshGenerator
 	
 		const TArray<FTransform> Transforms = GetBoneTransforms( SkeletalComponent, AnimFrame);
 		FleshComponent.Pose(SkeletalComponent, Transforms);
-		SolverComponent.WriteToSimulation(DeltaTime);
+		SolverComponent.WriteToSimulation(DeltaTime, false);
 		SolverComponent.Simulate(DeltaTime);
-		SolverComponent.ReadFromSimulation(DeltaTime);
+		SolverComponent.ReadFromSimulation(DeltaTime,false);
 		SimResource.SimulatedPositions[CacheFrame] = GetRenderPositions(SimResource);
 
 		SimResource.FinishFrame();
diff --git a/Engine/Plugins/ChaosClothAsset/Source/ChaosClothAssetEngine/Private/ChaosClothAsset/ClothComponent.cpp b/Engine/Plugins/ChaosClothAsset/Source/ChaosClothAssetEngine/Private/ChaosClothAsset/ClothComponent.cpp
--- a/Engine/Plugins/ChaosClothAsset/Source/ChaosClothAssetEngine/Private/ChaosClothAsset/ClothComponent.cpp
+++ b/Engine/Plugins/ChaosClothAsset/Source/ChaosClothAssetEngine/Private/ChaosClothAsset/ClothComponent.cpp
@@ -644,7 +644,7 @@ void UChaosClothComponent::ResetSimulationProxy()
 	ClothSimulationProxy.Reset();
 }
 
-void UChaosClothComponent::WriteToSimulation(const float DeltaTime)
+void UChaosClothComponent::WriteToSimulation(const float DeltaTime, const bool bAsyncTask)
 {
 	if (ClothSimulationProxy.IsValid())
 	{
@@ -662,11 +662,11 @@ void UChaosClothComponent::PreProcessSimulation(const float DeltaTime)
 {
 	if (ClothSimulationProxy.IsValid() && ClothSimulationProxy->HasCacheData())
 	{
-		WriteToSimulation(DeltaTime);
+		WriteToSimulation(DeltaTime, false);
 	}
 }
 
-void UChaosClothComponent::ReadFromSimulation(const float DeltaTime)
+void UChaosClothComponent::ReadFromSimulation(const float DeltaTime, const bool bAsyncTask)
 {
 	if (ClothSimulationProxy.IsValid())
 	{
diff --git a/Engine/Plugins/ChaosClothAsset/Source/ChaosClothAssetEngine/Public/ChaosClothAsset/ClothComponent.h b/Engine/Plugins/ChaosClothAsset/Source/ChaosClothAssetEngine/Public/ChaosClothAsset/ClothComponent.h
--- a/Engine/Plugins/ChaosClothAsset/Source/ChaosClothAssetEngine/Public/ChaosClothAsset/ClothComponent.h
+++ b/Engine/Plugins/ChaosClothAsset/Source/ChaosClothAssetEngine/Public/ChaosClothAsset/ClothComponent.h
@@ -195,8 +195,8 @@ protected:
 	virtual const FDataflowSimulationProxy* GetSimulationProxy() const  override;
 	virtual void BuildSimulationProxy() override;
 	virtual void ResetSimulationProxy() override;
-	virtual void WriteToSimulation(const float DeltaTime) override;
-	virtual void ReadFromSimulation(const float DeltaTime) override;
+	virtual void WriteToSimulation(const float DeltaTime, const bool bAsyncTask) override;
+	virtual void ReadFromSimulation(const float DeltaTime, const bool bAsyncTask) override;
 	virtual void PreProcessSimulation(const float DeltaTime) override;
 	// End IDataflowPhysicsSolverInterface overrides
 
diff --git a/Engine/Plugins/Experimental/ChaosFlesh/Source/ChaosFleshEngine/Private/ChaosFlesh/ChaosDeformableSolverComponent.cpp b/Engine/Plugins/Experimental/ChaosFlesh/Source/ChaosFleshEngine/Private/ChaosFlesh/ChaosDeformableSolverComponent.cpp
--- a/Engine/Plugins/Experimental/ChaosFlesh/Source/ChaosFleshEngine/Private/ChaosFlesh/ChaosDeformableSolverComponent.cpp
+++ b/Engine/Plugins/Experimental/ChaosFlesh/Source/ChaosFleshEngine/Private/ChaosFlesh/ChaosDeformableSolverComponent.cpp
@@ -116,6 +116,8 @@ bool UDeformableSolverComponent::IsSimulating(UDeformablePhysicsComponent* InCom
 }
 
 
+
+
 void UDeformableSolverComponent::UpdateDeformableEndTickState(bool bRegister)
 {
 	SCOPE_CYCLE_COUNTER(STAT_ChaosDeformable_UDeformableSolverComponent_UpdateDeformableEndTickState);
@@ -177,7 +179,7 @@ void UDeformableSolverComponent::TickComponent(float DeltaTime, enum ELevelTick
 		{
 			if(bSimulationTicking)
 			{
-				WriteToSimulation(DeltaTime);
+				WriteToSimulation(DeltaTime, false);
 
 				if (SolverTiming.bDoThreadedAdvance)
 				{
@@ -191,7 +193,7 @@ void UDeformableSolverComponent::TickComponent(float DeltaTime, enum ELevelTick
 				{
 					Simulate(DeltaTime);
 
-					ReadFromSimulation(DeltaTime);
+					ReadFromSimulation(DeltaTime, false);
 				}
 			}
 		}
@@ -377,7 +379,7 @@ void UDeformableSolverComponent::Simulate(float DeltaTime)
 	}
 }
 
-void UDeformableSolverComponent::WriteToSimulation(float DeltaTime)
+void UDeformableSolverComponent::WriteToSimulation(float DeltaTime, const bool bAsyncTask)
 {
 	SCOPE_CYCLE_COUNTER(STAT_ChaosDeformable_UDeformableSolverComponent_UpdateFromGameThread);
 	TRACE_CPUPROFILER_EVENT_SCOPE(ChaosDeformable_UDeformableSolverComponent_UpdateFromGameThread);
@@ -406,7 +408,7 @@ void UDeformableSolverComponent::WriteToSimulation(float DeltaTime)
 	}
 }
 
-void UDeformableSolverComponent::ReadFromSimulation(float DeltaTime)
+void UDeformableSolverComponent::ReadFromSimulation(float DeltaTime, const bool bAsyncTask)
 {
 	SCOPE_CYCLE_COUNTER(STAT_ChaosDeformable_UDeformableSolverComponent_UpdateFromSimulation);
 	TRACE_CPUPROFILER_EVENT_SCOPE(ChaosDeformable_UDeformableSolverComponent_UpdateFromSimulation);
diff --git a/Engine/Plugins/Experimental/ChaosFlesh/Source/ChaosFleshEngine/Private/ChaosFlesh/ChaosDeformableSolverThreading.cpp b/Engine/Plugins/Experimental/ChaosFlesh/Source/ChaosFleshEngine/Private/ChaosFlesh/ChaosDeformableSolverThreading.cpp
--- a/Engine/Plugins/Experimental/ChaosFlesh/Source/ChaosFleshEngine/Private/ChaosFlesh/ChaosDeformableSolverThreading.cpp
+++ b/Engine/Plugins/Experimental/ChaosFlesh/Source/ChaosFleshEngine/Private/ChaosFlesh/ChaosDeformableSolverThreading.cpp
@@ -45,7 +45,7 @@ void FDeformableEndTickFunction::ExecuteTick(float DeltaTime, enum ELevelTick Ti
 
 	if (DeformableSolverComponent)
 	{
-		DeformableSolverComponent->ReadFromSimulation(DeltaTime);
+		DeformableSolverComponent->ReadFromSimulation(DeltaTime,false);
 	}
 }
 
diff --git a/Engine/Plugins/Experimental/ChaosFlesh/Source/ChaosFleshEngine/Public/ChaosFlesh/ChaosDeformableSolverComponent.h b/Engine/Plugins/Experimental/ChaosFlesh/Source/ChaosFleshEngine/Public/ChaosFlesh/ChaosDeformableSolverComponent.h
--- a/Engine/Plugins/Experimental/ChaosFlesh/Source/ChaosFleshEngine/Public/ChaosFlesh/ChaosDeformableSolverComponent.h
+++ b/Engine/Plugins/Experimental/ChaosFlesh/Source/ChaosFleshEngine/Public/ChaosFlesh/ChaosDeformableSolverComponent.h
@@ -92,8 +92,8 @@ public:
 	virtual const FDataflowSimulationProxy* GetSimulationProxy() const  override {return &FleshSolverProxy;}
 	virtual void BuildSimulationProxy() override;
 	virtual void ResetSimulationProxy() override;
-	virtual void WriteToSimulation(const float DeltaTime) override;
-	virtual void ReadFromSimulation(const float DeltaTime) override;
+	virtual void WriteToSimulation(const float DeltaTime, const bool bAsyncTask) override;
+	virtual void ReadFromSimulation(const float DeltaTime, const bool bAsyncTask) override;
 	// End IDataflowPhysicsSolverInterface overrides
 
 	//~ Begin UObject Interface
diff --git a/Engine/Plugins/Experimental/Dataflow/Source/DataflowEditor/Private/Dataflow/DataflowSimulationGenerator.cpp b/Engine/Plugins/Experimental/Dataflow/Source/DataflowEditor/Private/Dataflow/DataflowSimulationGenerator.cpp
--- a/Engine/Plugins/Experimental/Dataflow/Source/DataflowEditor/Private/Dataflow/DataflowSimulationGenerator.cpp
+++ b/Engine/Plugins/Experimental/Dataflow/Source/DataflowEditor/Private/Dataflow/DataflowSimulationGenerator.cpp
@@ -35,13 +35,13 @@ void FDataflowSimulationTask::DoWork()
 				UDataflowSimulationManager* DataflowManager = SimulationWorld->GetSubsystem<UDataflowSimulationManager>();
 
 				// Pre advance the proxies
-				DataflowManager->PreAdvanceProxies(DeltaTime);
+				DataflowManager->ReadSimulationInterfaces(DeltaTime, true);
 
 				// Advance the simulation proxies
                 DataflowManager->AdvanceSimulationProxies(DeltaTime, SimulationTime);
 
-				// Post advanc e the simulation proxies
-				DataflowManager->PostAdvanceProxies(DeltaTime);
+				// Post advance the simulation proxies
+				DataflowManager->WriteSimulationInterfaces(DeltaTime, true);
 			}
 			else
 			{
@@ -66,6 +66,7 @@ bool FDataflowTaskManager::AllocateSimulationResource(const FVector2f& TimeRange
 		const TObjectPtr<UChaosCacheCollection>& CacheAsset, const TSubclassOf<AActor>& ActorClass, const TObjectPtr<UDataflowBaseContent>& DataflowContent)
 {
 	SimulationWorld = UWorld::CreateWorld(EWorldType::Editor, false);
+	SimulationWorld->bPostTickComponentUpdate = false;
 
 	FWorldContext& WorldContext = GEngine->CreateNewWorldContext( SimulationWorld->WorldType );
 	WorldContext.SetCurrentWorld(SimulationWorld);
diff --git a/Engine/Source/Runtime/Experimental/Chaos/Public/GeometryCollection/ManagedArray.h b/Engine/Source/Runtime/Experimental/Chaos/Public/GeometryCollection/ManagedArray.h
--- a/Engine/Source/Runtime/Experimental/Chaos/Public/GeometryCollection/ManagedArray.h
+++ b/Engine/Source/Runtime/Experimental/Chaos/Public/GeometryCollection/ManagedArray.h
@@ -843,6 +843,16 @@ public:
 
 		Exchange(*this, NewTypedArray);
 	}
+	
+	/**
+	* return true if index is in array range.
+	*
+	* @param Index Index to check.
+	*/
+	FORCEINLINE bool IsValidIndex(int32 Index) const
+	{
+		return Array.IsValidIndex(Index);
+	}
 
 	/**
 	* Returning a reference to the element at index.
diff --git a/Engine/Source/Runtime/Experimental/Dataflow/Simulation/Private/Dataflow/DataflowSimulationManager.cpp b/Engine/Source/Runtime/Experimental/Dataflow/Simulation/Private/Dataflow/DataflowSimulationManager.cpp
--- a/Engine/Source/Runtime/Experimental/Dataflow/Simulation/Private/Dataflow/DataflowSimulationManager.cpp
+++ b/Engine/Source/Runtime/Experimental/Dataflow/Simulation/Private/Dataflow/DataflowSimulationManager.cpp
@@ -173,7 +173,7 @@ void UDataflowSimulationManager::Tick(float DeltaTime)
 	if(bIsSimulationEnabled)
 	{
 		// Transfer data from GT -> PT
-		ReadSimulationInterfaces(DeltaTime);
+		ReadSimulationInterfaces(DeltaTime, false);
 		
 		if(UE::Dataflow::Private::DataflowSimulationThreadingMode == static_cast<uint8>(UE::Dataflow::Private::ESimulationThreadingMode::GameThread))
 		{
@@ -181,7 +181,7 @@ void UDataflowSimulationManager::Tick(float DeltaTime)
 			AdvanceSimulationProxies(DeltaTime, GetWorld()->GetTimeSeconds());
 
 			// Transfer data from PT -> GT
-			WriteSimulationInterfaces(DeltaTime);
+			WriteSimulationInterfaces(DeltaTime, false);
 		}
 		else
 		{
@@ -194,7 +194,7 @@ void UDataflowSimulationManager::Tick(float DeltaTime)
             	CompleteSimulationTasks();
             
             	// Transfer data from PT -> GT
-            	WriteSimulationInterfaces(DeltaTime);
+            	WriteSimulationInterfaces(DeltaTime, false);
             }
 		}
 	}
@@ -218,7 +218,7 @@ void UDataflowSimulationManager::OnStartup()
 						DataflowManager->CompleteSimulationTasks();
 					
 						// Transfer data from PT -> GT
-						DataflowManager->WriteSimulationInterfaces(DeltaSeconds);
+						DataflowManager->WriteSimulationInterfaces(DeltaSeconds, false);
 					}
 				}
 			}
@@ -298,42 +298,7 @@ void UDataflowSimulationManager::PostProcessSimulation(const float DeltaTime)
 	}
 }
 
-void UDataflowSimulationManager::PreAdvanceProxies(const float DeltaTime)
-{
-	for(const TPair<TObjectPtr<UDataflow>, UE::Dataflow::Private::FDataflowSimulationData>& DataflowData : SimulationData)
-	{
-		if(DataflowData.Value.SimulationContext.IsValid())
-		{
-			TArray<FDataflowPhysicsSolverProxy*> SolverProxies;
-			DataflowData.Value.SimulationContext->GetTypedProxies(SolverProxies);
-		
-			for(FDataflowPhysicsSolverProxy* SolverProxy : SolverProxies)
-			{
-				SolverProxy->PreSolveProxy(DeltaTime);
-			}
-		}
-	}
-}
-
-void UDataflowSimulationManager::PostAdvanceProxies(const float DeltaTime)
-{
-	for(const TPair<TObjectPtr<UDataflow>, UE::Dataflow::Private::FDataflowSimulationData>& DataflowData : SimulationData)
-	{
-		if(DataflowData.Value.SimulationContext.IsValid())
-		{
-			TArray<FDataflowPhysicsSolverProxy*> SolverProxies;
-			DataflowData.Value.SimulationContext->GetTypedProxies(SolverProxies);
-
-			for(FDataflowPhysicsSolverProxy* SolverProxy : SolverProxies)
-			{
-				SolverProxy->PostSolveProxy(DeltaTime);
-			}
-		}
-	}
-}
-
-
-void UDataflowSimulationManager::ReadSimulationInterfaces(const float DeltaTime)
+void UDataflowSimulationManager::ReadSimulationInterfaces(const float DeltaTime, const bool bAsyncTask)
 {
 	// Pre-simulation callback that could be used in BP before the simulation
 	//Dataflow::PreSimulationTick(GetWorld(), GetWorld()->GetTimeSeconds(), DeltaTime);
@@ -347,7 +312,7 @@ void UDataflowSimulationManager::ReadSimulationInterfaces(const float DeltaTime)
 			{
 				if(SimulationInterface)
 				{
-					SimulationInterface->WriteToSimulation(DeltaTime);
+					SimulationInterface->WriteToSimulation(DeltaTime, bAsyncTask);
 				}
 			}
 		}
@@ -401,7 +366,7 @@ void UDataflowSimulationManager::ResetSimulationInterfaces()
 	}
 }
 
-void UDataflowSimulationManager::WriteSimulationInterfaces(const float DeltaTime)
+void UDataflowSimulationManager::WriteSimulationInterfaces(const float DeltaTime, const bool bAsyncTask)
 {
 	for(const TPair<TObjectPtr<UDataflow>, UE::Dataflow::Private::FDataflowSimulationData>& DataflowData : SimulationData)
 	{
@@ -411,7 +376,7 @@ void UDataflowSimulationManager::WriteSimulationInterfaces(const float DeltaTime
 			{
 				if(SimulationInterface)
 				{
-					SimulationInterface->ReadFromSimulation(DeltaTime);
+					SimulationInterface->ReadFromSimulation(DeltaTime, bAsyncTask);
 				}
 			}
 		}
diff --git a/Engine/Source/Runtime/Experimental/Dataflow/Simulation/Public/Dataflow/DataflowSimulationInterface.h b/Engine/Source/Runtime/Experimental/Dataflow/Simulation/Public/Dataflow/DataflowSimulationInterface.h
--- a/Engine/Source/Runtime/Experimental/Dataflow/Simulation/Public/Dataflow/DataflowSimulationInterface.h
+++ b/Engine/Source/Runtime/Experimental/Dataflow/Simulation/Public/Dataflow/DataflowSimulationInterface.h
@@ -68,10 +68,10 @@ public:
 	virtual void PreProcessSimulation(const float DeltaTime) {};
 
 	/** Write data to be sent to the simulation proxy */
-	virtual void WriteToSimulation(const float DeltaTime) {};
+	virtual void WriteToSimulation(const float DeltaTime, const bool bAsyncTask) {};
 
 	/** Read data received from the simulation proxy */
-	virtual void ReadFromSimulation(const float DeltaTime) {};
+	virtual void ReadFromSimulation(const float DeltaTime, const bool bAsyncTask) {};
 
 	/** Postprocess data after simulation */
 	virtual void PostProcessSimulation(const float DeltaTime) {};
diff --git a/Engine/Source/Runtime/Experimental/Dataflow/Simulation/Public/Dataflow/DataflowSimulationManager.h b/Engine/Source/Runtime/Experimental/Dataflow/Simulation/Public/Dataflow/DataflowSimulationManager.h
--- a/Engine/Source/Runtime/Experimental/Dataflow/Simulation/Public/Dataflow/DataflowSimulationManager.h
+++ b/Engine/Source/Runtime/Experimental/Dataflow/Simulation/Public/Dataflow/DataflowSimulationManager.h
@@ -73,12 +73,6 @@ public:
 	virtual void Deinitialize() override;
 	// End USubsystem overrides
 
-	/** Pre solve all the solver proxies */
-	void PreAdvanceProxies(const float DeltaTime);
-
-	/** Post solve all the solver proxies */
-	void PostAdvanceProxies(const float DeltaTime);
-
 	/** Advance in time the registered simulation data (PT) */
     void AdvanceSimulationProxies(const float DeltaTime, const float SimulationTime);
 
@@ -92,10 +86,10 @@ public:
 	void RemoveSimulationInterface(const IDataflowSimulationInterface* SimulationInterface);
 
 	/** Read the simulation interfaces and Write all the data to the simulation proxies (to be send from GT->PT) */
-	void ReadSimulationInterfaces(const float DeltaTime);
+	void ReadSimulationInterfaces(const float DeltaTime, const bool bAsyncTask);
 
 	/** Read all the data from the simulation proxies and write the result onto the interfaces (received from PT->GT) */
-	void WriteSimulationInterfaces(const float DeltaTime);
+	void WriteSimulationInterfaces(const float DeltaTime, const bool bAsyncTask);
 
 	/** Init all the simulation interfaces*/
 	void InitSimulationInterfaces();
diff --git a/Engine/Source/Runtime/Experimental/Dataflow/Simulation/Public/Dataflow/Interfaces/DataflowPhysicsSolver.h b/Engine/Source/Runtime/Experimental/Dataflow/Simulation/Public/Dataflow/Interfaces/DataflowPhysicsSolver.h
--- a/Engine/Source/Runtime/Experimental/Dataflow/Simulation/Public/Dataflow/Interfaces/DataflowPhysicsSolver.h
+++ b/Engine/Source/Runtime/Experimental/Dataflow/Simulation/Public/Dataflow/Interfaces/DataflowPhysicsSolver.h
@@ -24,12 +24,6 @@ struct DATAFLOWSIMULATION_API FDataflowPhysicsSolverProxy : public FDataflowSimu
 	
 	/** Advance the solver datas in time */
 	virtual void AdvanceSolverDatas(const float DeltaTime) {}
-	
-	/** Pre-solve the proxy */
-	virtual void PreSolveProxy(const float DeltaTime) {}
-
-	/** Post-solve the proxy */
-	virtual void PostSolveProxy(const float DeltaTime) {}
 
 	/** Get the solver time step */
 	virtual float GetTimeStep() { return 0.033f;}
