Commit 1 of 3 - Decoupling ISMPool from GeometryCollection

Create a copy of the GC_ISM pool files in its own Runtime/Experimental folder.  This allows our non-GC assets to take advantage of the ISMPooling and custom instance data.

diff --git a/Engine/Source/Runtime/Experimental/ISMPool/ISMPool.Build.cs b/Engine/Source/Runtime/Experimental/ISMPool/ISMPool.Build.cs
new file mode 100644
--- /dev/null
+++ b/Engine/Source/Runtime/Experimental/ISMPool/ISMPool.Build.cs
@@ -0,0 +1,28 @@
+// Copyright Epic Games, Inc. All Rights Reserved.
+
+namespace UnrealBuildTool.Rules
+{
+	public class ISMPool : ModuleRules
+	{
+		public ISMPool(ReadOnlyTargetRules Target) : base(Target)
+		{
+			SetupModulePhysicsSupport(Target);
+
+			PublicDependencyModuleNames.AddRange(
+				new string[]
+				{
+					"Core",
+					"CoreUObject",
+					"Engine",
+				}
+				);
+
+			if (Target.bBuildEditor)
+			{
+				PublicDependencyModuleNames.Add("UnrealEd");
+			}
+
+			bDisableAutoRTFMInstrumentation = true;
+		}
+	}
+}
diff --git a/Engine/Source/Runtime/Experimental/ISMPool/Private/ISMPool/ISMPoolActor.cpp b/Engine/Source/Runtime/Experimental/ISMPool/Private/ISMPool/ISMPoolActor.cpp
new file mode 100644
--- /dev/null
+++ b/Engine/Source/Runtime/Experimental/ISMPool/Private/ISMPool/ISMPoolActor.cpp
@@ -0,0 +1,23 @@
+// Copyright Epic Games, Inc. All Rights Reserved.
+
+#include "ISMPool/ISMPoolActor.h"
+#include "ISMPool/ISMPoolComponent.h"
+#include "ISMPool/ISMPoolDebugDrawComponent.h"
+
+#include UE_INLINE_GENERATED_CPP_BY_NAME(ISMPoolActor)
+
+AISMPoolActor::AISMPoolActor(const FObjectInitializer& ObjectInitializer)
+	: Super(ObjectInitializer)
+{
+	ISMPoolComp = CreateDefaultSubobject<UISMPoolComponent>(TEXT("ISMPoolComp"));
+	RootComponent = ISMPoolComp;
+
+#if UE_ENABLE_DEBUG_DRAWING
+	ISMPoolDebugDrawComp = CreateDefaultSubobject<UISMPoolDebugDrawComponent>(TEXT("ISMPoolDebug"));
+	ISMPoolDebugDrawComp->SetCollisionEnabled(ECollisionEnabled::NoCollision);
+	ISMPoolDebugDrawComp->SetCanEverAffectNavigation(false);
+	ISMPoolDebugDrawComp->CanCharacterStepUpOn = ECanBeCharacterBase::ECB_No;
+	ISMPoolDebugDrawComp->SetGenerateOverlapEvents(false);
+	ISMPoolDebugDrawComp->SetupAttachment(ISMPoolComp);
+#endif
+}
diff --git a/Engine/Source/Runtime/Experimental/ISMPool/Private/ISMPool/ISMPoolComponent.cpp b/Engine/Source/Runtime/Experimental/ISMPool/Private/ISMPool/ISMPoolComponent.cpp
new file mode 100644
--- /dev/null
+++ b/Engine/Source/Runtime/Experimental/ISMPool/Private/ISMPool/ISMPoolComponent.cpp
@@ -0,0 +1,741 @@
+// Copyright Epic Games, Inc. All Rights Reserved.
+
+#include "ISMPool/ISMPoolComponent.h"
+
+#include "ChaosLog.h"
+#include "Components/InstancedStaticMeshComponent.h"
+#include "Engine/CollisionProfile.h"
+#include "Engine/StaticMesh.h"
+#include "Engine/World.h"
+
+#include UE_INLINE_GENERATED_CPP_BY_NAME(ISMPoolComponent)
+
+// Don't release ISM components when they empty, but keep them (and their scene proxy) alive.
+// This can remove the high cost associated with repeated registration, scene proxy creation and mesh draw command creation.
+// But it can also have a high memory overhead since the ISMs retain hard references to their static meshes.
+static bool GComponentKeepAlive = false; 
+FAutoConsoleVariableRef CVarISMPoolComponentKeepAlive(
+	TEXT("r.ISMPool.ComponentKeepAlive"),
+	GComponentKeepAlive,
+	TEXT("Keep ISM components alive when all their instances are removed."));
+
+// Use a FreeList to enable recycling of ISM components.
+// ISM components aren't unregistered, but their scene proxy is destroyed.
+// When recycling a component, a new mesh description can be used.
+// This removes the high CPU cost of unregister/register.
+// But there is more CPU cost to recycling a component then to simply keeping it alive because scene proxy creation and mesh draw command caching isn't cheap.
+// The component memory cost is kept bounded when compared to keeping components alive.
+static bool GComponentRecycle = true;
+FAutoConsoleVariableRef CVarISMPoolComponentRecycle(
+	TEXT("r.ISMPool.ComponentRecycle"),
+	GComponentRecycle,
+	TEXT("Recycle ISM components to a free list for reuse when all their instances are removed."));
+
+// Target free list size when recycling ISM components.
+// We try to maintain a pool of free components for fast allocation, but want to clean up when numbers get too high.
+static int32 GComponentFreeListTargetSize = 50;
+FAutoConsoleVariableRef CVarISMPoolComponentFreeListTargetSize(
+	TEXT("r.ISMPool.ComponentFreeListTargetSize"),
+	GComponentFreeListTargetSize,
+	TEXT("Target size for number of ISM components in the recycling free list."));
+
+// Keep copies of all custom instance data for restoration on readding an instance.
+static bool GShadowCopyCustomData = false;
+FAutoConsoleVariableRef CVarShadowCopyCustomData(
+	TEXT("r.ISMPool.ShadowCopyCustomData"),
+	GShadowCopyCustomData,
+	TEXT("Keeps a copy of custom instance data so it can be restored if the instance is removed and readded."));
+
+
+void FISMPoolMeshInfo::ShadowCopyCustomData(int32 InstanceCount, int32 NumCustomDataFloatsPerInstance, TArrayView<const float> CustomDataFloats)
+{
+	CustomData.SetNum(InstanceCount * NumCustomDataFloatsPerInstance + NumCustomDataFloatsPerInstance);
+
+	for (int32 InstanceIndex = 0; InstanceIndex < InstanceCount; ++InstanceIndex)
+	{
+		int32 Offset = InstanceIndex * NumCustomDataFloatsPerInstance;
+		FMemory::Memcpy(&CustomData[Offset], CustomDataFloats.GetData() + Offset, NumCustomDataFloatsPerInstance * CustomDataFloats.GetTypeSize());
+	}
+}
+
+TArrayView<const float> FISMPoolMeshInfo::CustomDataSlice(int32 InstanceIndex, int32 NumCustomDataFloatsPerInstance)
+{
+	TArrayView<const float> DataView = CustomData;
+	return DataView.Slice(InstanceIndex * NumCustomDataFloatsPerInstance, NumCustomDataFloatsPerInstance);
+}
+
+FISMPoolMeshGroup::FMeshId FISMPoolMeshGroup::AddMesh(const FISMPoolStaticMeshInstance& MeshInstance, int32 InstanceCount, const FISMPoolMeshInfo& ISMInstanceInfo, TArrayView<const float> CustomDataFloats)
+{
+	const FMeshId MeshInfoIndex = MeshInfos.Emplace(ISMInstanceInfo);
+
+	if (bAllowPerInstanceRemoval && GShadowCopyCustomData)
+	{
+		FISMPoolMeshInfo& MeshInfo = MeshInfos[MeshInfoIndex];
+		MeshInfo.ShadowCopyCustomData(InstanceCount, MeshInstance.Desc.NumCustomDataFloats, CustomDataFloats);
+	}
+
+	return MeshInfoIndex;
+}
+
+bool FISMPoolMeshGroup::BatchUpdateInstancesTransforms(FISMPool& ISMPool, FMeshId MeshId, int32 StartInstanceIndex, TArrayView<const FTransform> NewInstancesTransforms, bool bWorldSpace, bool bMarkRenderStateDirty, bool bTeleport)
+{
+	if (MeshInfos.IsValidIndex(MeshId))
+	{
+		return ISMPool.BatchUpdateInstancesTransforms(MeshInfos[MeshId], StartInstanceIndex, NewInstancesTransforms, bWorldSpace, bMarkRenderStateDirty, bTeleport, bAllowPerInstanceRemoval);
+	}
+	UE_LOG(LogChaos, Warning, TEXT("UISMPoolComponent : Invalid mesh Id (%d) for this mesh group"), MeshId);
+	return false;
+}
+
+void FISMPoolMeshGroup::BatchUpdateInstanceCustomData(FISMPool& ISMPool, int32 CustomFloatIndex, float CustomFloatValue)
+{
+	for (const FISMPoolMeshInfo& MeshInfo : MeshInfos)
+	{
+		ISMPool.BatchUpdateInstanceCustomData(MeshInfo, CustomFloatIndex, CustomFloatValue);
+	}
+}
+
+void FISMPoolMeshGroup::RemoveAllMeshes(FISMPool& ISMPool)
+{
+	for (const FISMPoolMeshInfo& MeshInfo: MeshInfos)
+	{
+		ISMPool.RemoveInstancesFromISM(MeshInfo);
+	}
+	MeshInfos.Empty();
+}
+
+void FISMPoolISM::CreateISM(USceneComponent* InOwningComponent)
+{
+	check(InOwningComponent);
+
+	AActor* OwningActor = InOwningComponent->GetOwner();
+	USceneComponent* RootComponent = OwningActor->GetRootComponent();
+
+	ISMComponent = NewObject<UInstancedStaticMeshComponent>(InOwningComponent, NAME_None, RF_Transient | RF_DuplicateTransient);
+
+	ISMComponent->SetRemoveSwap();
+	ISMComponent->SetCanEverAffectNavigation(false);
+	ISMComponent->SetCollisionEnabled(ECollisionEnabled::NoCollision);
+	ISMComponent->SetupAttachment(RootComponent);
+	ISMComponent->RegisterComponent();
+
+#if WITH_EDITOR
+	UWorld const* World = InOwningComponent->GetWorld();
+	const bool bShowInWorldOutliner = World && World->IsGameWorld();
+	if (bShowInWorldOutliner)
+	{
+		OwningActor->AddInstanceComponent(ISMComponent);
+	}
+#endif
+}
+
+void FISMPoolISM::InitISM(const FISMPoolStaticMeshInstance& InMeshInstance, bool bKeepAlive, bool bOverrideTransformUpdates)
+{
+	MeshInstance = InMeshInstance;
+	check(ISMComponent != nullptr);
+
+	UStaticMesh* StaticMesh = MeshInstance.StaticMesh.Get();
+	// We should only get here for valid static mesh objects.
+	check(StaticMesh != nullptr);
+
+#if WITH_EDITOR
+	const FName ISMName = MakeUniqueObjectName(ISMComponent->GetOwner(), UInstancedStaticMeshComponent::StaticClass(), StaticMesh->GetFName());
+	const FString ISMNameString = ISMName.ToString();
+	ISMComponent->Rename(*ISMNameString);
+#endif
+
+	ISMComponent->bUseAttachParentBound = bOverrideTransformUpdates;
+	ISMComponent->SetAbsolute(bOverrideTransformUpdates, bOverrideTransformUpdates, bOverrideTransformUpdates);
+
+	bool bDisallowNanite = false;
+
+	ISMComponent->EmptyOverrideMaterials();
+	for (int32 MaterialIndex = 0; MaterialIndex < MeshInstance.MaterialsOverrides.Num(); MaterialIndex++)
+	{
+		UMaterialInterface* Material = MeshInstance.MaterialsOverrides[MaterialIndex].Get();
+		// We should only get here for valid material objects.
+		check(Material != nullptr);
+		ISMComponent->SetMaterial(MaterialIndex, Material);
+
+		// Nanite doesn't support translucent materials.
+		bDisallowNanite |= Material->GetBlendMode() == BLEND_Translucent;
+	}
+
+	ISMComponent->SetStaticMesh(StaticMesh);
+	ISMComponent->SetMobility((MeshInstance.Desc.Flags & FISMPoolComponentDescription::StaticMobility) != 0 ? EComponentMobility::Static : EComponentMobility::Movable);
+
+	ISMComponent->NumCustomDataFloats = MeshInstance.Desc.NumCustomDataFloats;
+	for (int32 DataIndex = 0; DataIndex < MeshInstance.CustomPrimitiveData.Num(); DataIndex++)
+	{
+		ISMComponent->SetDefaultCustomPrimitiveDataFloat(DataIndex, MeshInstance.CustomPrimitiveData[DataIndex]);
+	}
+
+	const bool bReverseCulling = (MeshInstance.Desc.Flags & FISMPoolComponentDescription::ReverseCulling) != 0;
+	// Instead of reverse culling we put the mirror in the component transform so that PRIMITIVE_SCENE_DATA_FLAG_DETERMINANT_SIGN will be set for use by materials.
+	//ISMComponent->SetReverseCulling(bReverseCulling);
+	const FVector Scale = bReverseCulling ? FVector(-1, 1, 1) : FVector(1, 1, 1);
+
+	if(bOverrideTransformUpdates)
+	{
+		FTransform TempTm = ISMComponent->GetAttachParent() ?
+			ISMComponent->GetAttachParent()->GetComponentToWorld() :
+			FTransform::Identity;
+
+		// Apply above identified scale to the transform directly
+		TempTm.SetScale3D(TempTm.GetScale3D() * Scale);
+
+		ISMComponent->SetComponentToWorld(TempTm);
+		ISMComponent->UpdateComponentTransform(EUpdateTransformFlags::None, ETeleportType::None);
+		ISMComponent->MarkRenderTransformDirty();
+	}
+	else
+	{
+		const FTransform NewRelativeTransform(FQuat::Identity, MeshInstance.Desc.Position, Scale);
+
+		if(!ISMComponent->GetRelativeTransform().Equals(NewRelativeTransform))
+		{
+			// If we're not overriding the transform and need a relative offset, apply that here
+			ISMComponent->SetRelativeTransform(FTransform(FQuat::Identity, MeshInstance.Desc.Position, Scale));
+		}
+	}
+
+	if ((MeshInstance.Desc.Flags & FISMPoolComponentDescription::DistanceCullPrimitive) != 0)
+	{
+		ISMComponent->SetCachedMaxDrawDistance(MeshInstance.Desc.EndCullDistance);
+	}
+
+	ISMComponent->SetCullDistances(MeshInstance.Desc.StartCullDistance, MeshInstance.Desc.EndCullDistance);
+	ISMComponent->SetCastShadow((MeshInstance.Desc.Flags & FISMPoolComponentDescription::AffectShadow) != 0);
+	ISMComponent->bAffectDynamicIndirectLighting = (MeshInstance.Desc.Flags & FISMPoolComponentDescription::AffectDynamicIndirectLighting) != 0;
+	ISMComponent->bAffectDistanceFieldLighting = (MeshInstance.Desc.Flags & FISMPoolComponentDescription::AffectDistanceFieldLighting) != 0;
+	ISMComponent->bCastFarShadow = (MeshInstance.Desc.Flags & FISMPoolComponentDescription::AffectFarShadow) != 0;
+	ISMComponent->bWorldPositionOffsetWritesVelocity = (MeshInstance.Desc.Flags & FISMPoolComponentDescription::WorldPositionOffsetWritesVelocity) != 0;
+	ISMComponent->bEvaluateWorldPositionOffset = (MeshInstance.Desc.Flags & FISMPoolComponentDescription::EvaluateWorldPositionOffset) != 0;
+	ISMComponent->bUseGpuLodSelection = (MeshInstance.Desc.Flags & FISMPoolComponentDescription::GpuLodSelection) != 0;
+	ISMComponent->bOverrideMinLOD = MeshInstance.Desc.MinLod > 0;
+	ISMComponent->MinLOD = MeshInstance.Desc.MinLod;
+	ISMComponent->SetLODDistanceScale(MeshInstance.Desc.LodScale);
+	ISMComponent->SetUseConservativeBounds(true);
+	ISMComponent->bComputeFastLocalBounds = true;
+	ISMComponent->bDisallowNanite = bDisallowNanite;
+	ISMComponent->SetMeshDrawCommandStatsCategory(MeshInstance.Desc.StatsCategory);
+	ISMComponent->ComponentTags = MeshInstance.Desc.Tags;
+
+	// Use a fixed seed to avoid getting a different seed at every run (see UInstancedStaticMeshComponent::OnRegister())
+	// A possible improvement would be to compute a hash from the owner component and use that as the seed.
+	ISMComponent->InstancingRandomSeed = 1;	
+}
+
+FISMPoolInstanceGroups::FInstanceGroupId FISMPoolISM::AddInstanceGroup(int32 InstanceCount, TArrayView<const float> CustomDataFloats)
+{
+	// When adding new group it will always have a single range
+	const FISMPoolInstanceGroups::FInstanceGroupId InstanceGroupIndex = InstanceGroups.AddGroup(InstanceCount);
+	const FISMPoolInstanceGroups::FInstanceGroupRange& NewInstanceGroup = InstanceGroups.GroupRanges[InstanceGroupIndex];
+
+	// Ensure that remapping arrays are big enough to hold any new items.
+	InstanceIds.SetNum(InstanceGroups.GetMaxInstanceIndex(), EAllowShrinking::No);
+
+	FTransform ZeroScaleTransform;
+	ZeroScaleTransform.SetIdentityZeroScale();
+	TArray<FTransform> ZeroScaleTransforms;
+	ZeroScaleTransforms.Init(ZeroScaleTransform, InstanceCount);
+
+	ISMComponent->PreAllocateInstancesMemory(InstanceCount);
+	TArray<FPrimitiveInstanceId> AddedInstanceIds = ISMComponent->AddInstancesById(ZeroScaleTransforms, true, true);
+	for (int32 InstanceIndex = 0; InstanceIndex < InstanceCount; ++InstanceIndex)
+	{
+		InstanceIds[NewInstanceGroup.Start + InstanceIndex] = AddedInstanceIds[InstanceIndex];
+	}
+
+	// Set any custom data.
+	if (CustomDataFloats.Num())
+	{
+		const int32 NumCustomDataFloats = ISMComponent->NumCustomDataFloats;
+		if (ensure(NumCustomDataFloats * InstanceCount == CustomDataFloats.Num()))
+		{
+			for (int32 InstanceIndex = 0; InstanceIndex < InstanceCount; ++InstanceIndex)
+			{
+				ISMComponent->SetCustomDataById(AddedInstanceIds[InstanceIndex], CustomDataFloats.Slice(InstanceIndex * NumCustomDataFloats, NumCustomDataFloats));
+			}
+		}
+	}
+
+	return InstanceGroupIndex;
+}
+
+FISMPool::FISMPool()
+	: bCachedKeepAlive(GComponentKeepAlive)
+	, bCachedRecycle(GComponentRecycle)
+{
+}
+
+FISMPool::FISMIndex FISMPool::GetOrAddISM(UISMPoolComponent* OwningComponent, const FISMPoolStaticMeshInstance& MeshInstance, bool& bOutISMCreated)
+{
+	FISMIndex* ISMIndexPtr = MeshToISMIndex.Find(MeshInstance);
+	if (ISMIndexPtr != nullptr)
+	{
+		bOutISMCreated = false;
+		return *ISMIndexPtr;
+	}
+
+	// Take an ISM from the current FreeLists if available instead of allocating a new slot.
+	FISMIndex ISMIndex = INDEX_NONE;
+	if (FreeListISM.Num())
+	{
+		ISMIndex = FreeListISM.Last();
+		FreeListISM.RemoveAt(FreeListISM.Num() - 1);
+	}
+	else if (FreeList.Num())
+	{
+		ISMIndex = FreeList.Last();
+		FreeList.RemoveAt(FreeList.Num() - 1);
+		ISMs[ISMIndex].CreateISM(OwningComponent);
+	}
+	else
+	{
+		ISMIndex = ISMs.AddDefaulted();
+		ISMs[ISMIndex].CreateISM(OwningComponent);
+	}
+	
+	ISMs[ISMIndex].InitISM(MeshInstance, bCachedKeepAlive, bDisableBoundsAndTransformUpdate);
+	
+	bOutISMCreated = true;
+	MeshToISMIndex.Add(MeshInstance, ISMIndex);
+	return ISMIndex;
+}
+
+FISMPoolMeshInfo FISMPool::AddInstancesToISM(UISMPoolComponent* OwningComponent, const FISMPoolStaticMeshInstance& MeshInstance, int32 InstanceCount, TArrayView<const float> CustomDataFloats)
+{
+	bool bISMCreated = false;
+	FISMPoolMeshInfo Info;
+	Info.ISMIndex = GetOrAddISM(OwningComponent, MeshInstance, bISMCreated);
+	Info.InstanceGroupIndex = ISMs[Info.ISMIndex].AddInstanceGroup(InstanceCount, CustomDataFloats);
+	return Info;
+}
+
+bool FISMPool::BatchUpdateInstancesTransforms(FISMPoolMeshInfo& MeshInfo, int32 StartInstanceIndex, TArrayView<const FTransform> NewInstancesTransforms, bool bWorldSpace, bool bMarkRenderStateDirty, bool bTeleport, bool bAllowPerInstanceRemoval)
+{
+	if (!ISMs.IsValidIndex(MeshInfo.ISMIndex))
+	{
+		UE_LOG(LogChaos, Warning, TEXT("UISMPoolComponent : Invalid ISM Id (%d) when updating the transform "), MeshInfo.ISMIndex);
+		return false;
+	}
+		
+	FISMPoolISM& ISM = ISMs[MeshInfo.ISMIndex];
+	const FISMPoolInstanceGroups::FInstanceGroupRange& InstanceGroup = ISM.InstanceGroups.GroupRanges[MeshInfo.InstanceGroupIndex];
+	// If ISM component has identity transform (the common case) then we can skip world space to component space maths inside BatchUpdateInstancesTransforms()
+	bWorldSpace &= !ISM.ISMComponent->GetComponentTransform().Equals(FTransform::Identity, 0.f);
+
+	// The transform count should fit within the instance group.
+	// Clamp it if it doesn't, but if we hit this ensure we need to investigate why.
+	ensure(StartInstanceIndex + NewInstancesTransforms.Num() <= InstanceGroup.Count);
+	const int32 NumTransforms = FMath::Min(NewInstancesTransforms.Num(), InstanceGroup.Count - StartInstanceIndex);
+
+	// Loop over transforms.
+	// todo: There may be some value in batching InstanceIds and caling one function for each of Add/Remove/Update.
+	// However the ISM batched calls themselves seem to be just simple loops over the single instance calls, so probably no benefit.
+	for (int InstanceIndex = StartInstanceIndex; InstanceIndex < StartInstanceIndex + NumTransforms; ++InstanceIndex)
+	{
+		FPrimitiveInstanceId InstanceId = ISM.InstanceIds[InstanceGroup.Start + InstanceIndex];
+		FTransform const& Transform = NewInstancesTransforms[InstanceIndex];
+
+		if (bAllowPerInstanceRemoval)
+		{
+			if (Transform.GetScale3D().IsZero() && InstanceId.IsValid())
+			{
+				// Zero scale is used to indicate that we should remove the instance from the ISM.
+				ISM.ISMComponent->RemoveInstanceById(InstanceId);
+				ISM.InstanceIds[InstanceGroup.Start + InstanceIndex] = FPrimitiveInstanceId();
+				continue;
+			}
+			else if (!Transform.GetScale3D().IsZero() && !InstanceId.IsValid())
+			{
+				// Re-add the instance to the ISM if the scale becomes non-zero.
+				FPrimitiveInstanceId Id = ISM.ISMComponent->AddInstanceById(Transform, bWorldSpace);
+				ISM.InstanceIds[InstanceGroup.Start + InstanceIndex] = Id;
+
+				if (MeshInfo.CustomData.Num())
+				{
+					ISM.ISMComponent->SetCustomDataById(Id, MeshInfo.CustomDataSlice(InstanceIndex, ISM.ISMComponent->NumCustomDataFloats));
+				}
+				continue;
+			}
+		}
+
+		if (InstanceId.IsValid())
+		{
+			ISM.ISMComponent->UpdateInstanceTransformById(InstanceId, Transform, bWorldSpace, bTeleport);
+		}
+	}
+
+	return true;
+}
+
+void FISMPool::BatchUpdateInstanceCustomData(FISMPoolMeshInfo const& MeshInfo, int32 CustomFloatIndex, float CustomFloatValue)
+{
+	if (!ISMs.IsValidIndex(MeshInfo.ISMIndex))
+	{
+		return;
+	}
+	
+	FISMPoolISM& ISM = ISMs[MeshInfo.ISMIndex];
+	if (!ensure(CustomFloatIndex < ISM.MeshInstance.Desc.NumCustomDataFloats))
+	{
+		return;
+	}
+		
+	const FISMPoolInstanceGroups::FInstanceGroupRange& InstanceGroup = ISM.InstanceGroups.GroupRanges[MeshInfo.InstanceGroupIndex];
+	for (int32 InstanceIndex = 0; InstanceIndex < InstanceGroup.Count; ++InstanceIndex)
+	{
+		const FPrimitiveInstanceId InstanceId = ISM.InstanceIds[InstanceGroup.Start + InstanceIndex];
+		if (InstanceId.IsValid())
+		{
+			ISM.ISMComponent->SetCustomDataValueById(InstanceId, CustomFloatIndex, CustomFloatValue);
+		}
+	}
+}
+
+void FISMPool::RemoveInstancesFromISM(const FISMPoolMeshInfo& MeshInfo)
+{
+	if (ISMs.IsValidIndex(MeshInfo.ISMIndex))
+	{
+		FISMPoolISM& ISM = ISMs[MeshInfo.ISMIndex];
+		const FISMPoolInstanceGroups::FInstanceGroupRange& InstanceGroup = ISM.InstanceGroups.GroupRanges[MeshInfo.InstanceGroupIndex];
+		
+		for (int32 Index = 0; Index < InstanceGroup.Count; ++Index)
+		{
+			FPrimitiveInstanceId InstanceId = ISM.InstanceIds[InstanceGroup.Start + Index];
+			if (InstanceId.IsValid())
+			{
+				// todo: Could RemoveInstanceByIds() instead as long as that function can handle skipping invalid InstanceIds.
+				ISM.ISMComponent->RemoveInstanceById(InstanceId);
+			}
+		}
+		
+#if DO_CHECK
+		// clear the IDs
+		for (int32 Index = 0; Index < InstanceGroup.Count; ++Index)
+		{
+			ISM.InstanceIds[InstanceGroup.Start + Index] = FPrimitiveInstanceId();
+		}
+#endif
+		ISM.InstanceGroups.RemoveGroup(MeshInfo.InstanceGroupIndex);
+	
+		if (ISM.InstanceGroups.IsEmpty())
+		{
+			ensure(ISM.ISMComponent->PerInstanceSMData.Num() == 0);
+
+			// No live instances, so take opportunity to reset indexing.
+			ISM.InstanceGroups.Reset();
+			ISM.InstanceIds.Reset();
+
+			RemoveISM(MeshInfo.ISMIndex, bCachedKeepAlive, bCachedRecycle);
+
+			if (!bCachedKeepAlive)
+			{
+				MeshToISMIndex.Remove(ISM.MeshInstance);
+			}
+		}
+	}
+}
+
+void FISMPool::RemoveISM(FISMIndex ISMIndex, bool bKeepAlive, bool bRecycle)
+{
+	FISMPoolISM& ISM = ISMs[ISMIndex];
+	ensure(ISM.InstanceGroups.IsEmpty());
+	ensure(ISM.InstanceIds.IsEmpty());
+
+	if (bKeepAlive)
+	{
+		// Nothing to do.
+	}
+	else if (bRecycle)
+	{
+		// Recycle to the free list.
+#if WITH_EDITOR
+		ISM.ISMComponent->Rename(nullptr);
+#endif
+		FreeListISM.Add(ISMIndex);
+	}
+	else
+	{
+		// Completely unregister and destroy the component and mark the ISM slot as free.
+		ISM.ISMComponent->DestroyComponent();
+		ISM.ISMComponent = nullptr;
+		
+		FreeList.Add(ISMIndex);
+	}
+}
+
+void FISMPool::Clear()
+{
+	MeshToISMIndex.Reset();
+	PrellocationQueue.Reset();
+	FreeList.Reset();
+	FreeListISM.Reset();
+	if (ISMs.Num() > 0)
+	{
+		if (AActor* OwningActor = ISMs[0].ISMComponent->GetOwner())
+		{
+			for(FISMPoolISM& ISM : ISMs)
+			{
+				ISM.ISMComponent->DestroyComponent();
+			}
+		}
+		ISMs.Reset();
+	}
+}
+
+void FISMPool::RequestPreallocateMeshInstance(const FISMPoolStaticMeshInstance& MeshInstance)
+{
+	// Preallocation only makes sense when we are keeping empty components alive.
+	if (bCachedKeepAlive)
+	{
+		uint32 KeyHash = GetTypeHash(MeshInstance);
+		if (MeshToISMIndex.FindByHash(KeyHash, MeshInstance) == nullptr)
+		{
+			PrellocationQueue.AddByHash(KeyHash, MeshInstance);
+		}
+	}
+}
+
+static bool AreWeakPointersValid(FISMPoolStaticMeshInstance& InMeshInstance)
+{
+	if (!InMeshInstance.StaticMesh.IsValid())
+	{
+		return false;
+	}
+
+	for (TWeakObjectPtr<UMaterialInterface> Material : InMeshInstance.MaterialsOverrides)
+	{
+		if (!Material.IsValid())
+		{
+			return false;
+		}
+	}
+
+	return true;
+}
+
+void FISMPool::ProcessPreallocationRequests(UISMPoolComponent* OwningComponent, int32 MaxPreallocations)
+{
+	int32 NumAdded = 0;
+	for (TSet<FISMPoolStaticMeshInstance>::TIterator It(PrellocationQueue); It; ++It)
+	{
+		bool bISMCreated = false;
+
+		// Objects in the entries of the preallocation queue may no longer be loaded.
+		if (AreWeakPointersValid(*It))
+		{
+			GetOrAddISM(OwningComponent, *It, bISMCreated);
+		}
+
+		It.RemoveCurrent();
+
+		if (bISMCreated)
+		{
+			if (++NumAdded >= MaxPreallocations)
+			{
+				break;
+			}
+		}
+	}
+}
+
+void FISMPool::UpdateAbsoluteTransforms(const FTransform& BaseTransform, EUpdateTransformFlags UpdateTransformFlags, ETeleportType Teleport)
+{
+	for(const FISMPoolISM& GcIsm : ISMs)
+	{
+		const bool bReverseCulling = (GcIsm.MeshInstance.Desc.Flags & FISMPoolComponentDescription::ReverseCulling) != 0;
+		check(GcIsm.MeshInstance.Desc.Position == FVector::ZeroVector);
+		
+		if(UInstancedStaticMeshComponent* Ism = GcIsm.ISMComponent)
+		{
+			if(bReverseCulling)
+			{
+				// As in InitISM we need to apply the inverted X scale for reverse culling.
+				// Just copy the transform and set an inverted scale to apply to the ISM
+				FVector BaseScale = BaseTransform.GetScale3D();
+				BaseScale.X = -BaseScale.X;
+				FTransform Flipped = BaseTransform;
+				Flipped.SetScale3D(BaseScale);
+
+				Ism->SetComponentToWorld(Flipped);
+			}
+			else
+			{
+				Ism->SetComponentToWorld(BaseTransform);
+			}
+
+			Ism->UpdateComponentTransform(UpdateTransformFlags | EUpdateTransformFlags::SkipPhysicsUpdate, Teleport);
+			Ism->MarkRenderTransformDirty();
+		}
+	}
+}
+
+void FISMPool::Tick(UISMPoolComponent* OwningComponent)
+{
+	// Recache component lifecycle state from cvar.
+	const bool bRemovedKeepAlive = bCachedKeepAlive && !GComponentKeepAlive;
+	const bool bRemovedReycle = bCachedRecycle && !GComponentRecycle;
+	bCachedKeepAlive = GComponentKeepAlive;
+	bCachedRecycle = GComponentRecycle;
+
+	// If we disabled keep alive behavior since last update then deal with the zombie components.
+	if (bRemovedKeepAlive)
+	{
+		for (int32 ISMIndex = 0; ISMIndex < ISMs.Num(); ++ISMIndex)
+		{
+			FISMPoolISM& ISM = ISMs[ISMIndex];
+			if (ISM.ISMComponent && ISM.InstanceGroups.IsEmpty())
+			{
+				// Actually release the ISM.
+				RemoveISM(ISMIndex, false, bCachedRecycle);
+				MeshToISMIndex.Remove(ISM.MeshInstance);
+			}
+		}
+	}
+
+	// Process preallocation queue.
+	if (!bCachedKeepAlive)
+	{
+		PrellocationQueue.Reset();
+	}
+	else if (!PrellocationQueue.IsEmpty())
+	{
+		// Preallocate components per tick until the queue is empty.
+		const int32 PreallocateCountPerTick = 2;
+		ProcessPreallocationRequests(OwningComponent, PreallocateCountPerTick);
+	}
+
+	if (FreeListISM.Num() > 0)
+	{
+		// Release components per tick until we reach minimum pool size.
+		const int32 RemoveCountPerTick = 1;
+		const int32 FreeListTargetSize = bRemovedReycle ? 0 : FMath::Max(FMath::Max(FreeListISM.Num() - RemoveCountPerTick, GComponentFreeListTargetSize), 0);
+		while (FreeListISM.Num() > FreeListTargetSize)
+		{
+			const int32 ISMIndex = FreeListISM.Pop(EAllowShrinking::No);
+			RemoveISM(ISMIndex, false, false);
+		}
+	}
+}
+
+
+UISMPoolComponent::UISMPoolComponent(const FObjectInitializer& ObjectInitializer)
+	: NextMeshGroupId(0)
+{
+	PrimaryComponentTick.bCanEverTick = true;
+	PrimaryComponentTick.bStartWithTickEnabled = true;
+	PrimaryComponentTick.bAllowTickOnDedicatedServer = false;
+	PrimaryComponentTick.TickInterval = 0.25f;
+}
+
+void UISMPoolComponent::TickComponent(float DeltaTime, enum ELevelTick TickType, FActorComponentTickFunction* ThisTickFunction)
+{
+	Super::TickComponent(DeltaTime, TickType, ThisTickFunction);
+	Pool.Tick(this);
+}
+
+UISMPoolComponent::FMeshGroupId  UISMPoolComponent::CreateMeshGroup(bool bAllowPerInstanceRemoval)
+{
+	FISMPoolMeshGroup Group;
+	Group.bAllowPerInstanceRemoval = bAllowPerInstanceRemoval;
+	MeshGroups.Add(NextMeshGroupId, Group);
+	return NextMeshGroupId++;
+}
+
+void UISMPoolComponent::DestroyMeshGroup(FMeshGroupId MeshGroupId)
+{
+	if (FISMPoolMeshGroup* MeshGroup = MeshGroups.Find(MeshGroupId))
+	{
+		MeshGroup->RemoveAllMeshes(Pool);
+		MeshGroups.Remove(MeshGroupId);
+	}
+}
+
+UISMPoolComponent::FMeshId UISMPoolComponent::AddMeshToGroup(FMeshGroupId MeshGroupId, const FISMPoolStaticMeshInstance& MeshInstance, int32 InstanceCount, TArrayView<const float> CustomDataFloats)
+{
+	if (FISMPoolMeshGroup* MeshGroup = MeshGroups.Find(MeshGroupId))
+	{
+		const FISMPoolMeshInfo ISMInstanceInfo = Pool.AddInstancesToISM(this, MeshInstance, InstanceCount, CustomDataFloats);
+		return MeshGroup->AddMesh(MeshInstance, InstanceCount, ISMInstanceInfo, CustomDataFloats);
+	}
+	UE_LOG(LogChaos, Warning, TEXT("UISMPoolComponent : Trying to add a mesh to a mesh group (%d) that does not exists"), MeshGroupId);
+	return INDEX_NONE;
+}
+
+bool UISMPoolComponent::BatchUpdateInstancesTransforms(FMeshGroupId MeshGroupId, FMeshId MeshId, int32 StartInstanceIndex, const TArray<FTransform>& NewInstancesTransforms, bool bWorldSpace, bool bMarkRenderStateDirty, bool bTeleport)
+{
+	return BatchUpdateInstancesTransforms(MeshGroupId, MeshId, StartInstanceIndex, MakeArrayView(NewInstancesTransforms), bWorldSpace, bMarkRenderStateDirty, bTeleport);
+}
+
+bool UISMPoolComponent::BatchUpdateInstancesTransforms(FMeshGroupId MeshGroupId, FMeshId MeshId, int32 StartInstanceIndex, TArrayView<const FTransform> NewInstancesTransforms, bool bWorldSpace, bool bMarkRenderStateDirty, bool bTeleport)
+{
+	if (FISMPoolMeshGroup* MeshGroup = MeshGroups.Find(MeshGroupId))
+	{
+		return MeshGroup->BatchUpdateInstancesTransforms(Pool, MeshId, StartInstanceIndex, NewInstancesTransforms, bWorldSpace, bMarkRenderStateDirty, bTeleport);
+	}
+	UE_LOG(LogChaos, Warning, TEXT("UISMPoolComponent : Trying to update instance with mesh group (%d) that not exists"), MeshGroupId);
+	return false;
+}
+
+bool UISMPoolComponent::BatchUpdateInstanceCustomData(FMeshGroupId MeshGroupId, int32 CustomFloatIndex, float CustomFloatValue)
+{
+	if (FISMPoolMeshGroup* MeshGroup = MeshGroups.Find(MeshGroupId))
+	{
+		MeshGroup->BatchUpdateInstanceCustomData(Pool, CustomFloatIndex, CustomFloatValue);
+		return true;
+	}
+	UE_LOG(LogChaos, Warning, TEXT("UISMPoolComponent : Trying to update instance with mesh group (%d) that not exists"), MeshGroupId);
+	return false;
+}
+
+void UISMPoolComponent::PreallocateMeshInstance(const FISMPoolStaticMeshInstance& MeshInstance)
+{
+	Pool.RequestPreallocateMeshInstance(MeshInstance);
+}
+
+void UISMPoolComponent::SetTickablePoolManagement(bool bEnablePoolManagement)
+{
+	if (!bEnablePoolManagement)
+	{
+		// Disable the keep alive and recycle pool management systems.
+		// This also disables preallocation for this pool.
+		Pool.bCachedKeepAlive = false;
+		Pool.bCachedRecycle = false;
+	}
+	// Disable the Tick that is used to manage the pool.
+	PrimaryComponentTick.SetTickFunctionEnable(bEnablePoolManagement);
+}
+
+void UISMPoolComponent::SetOverrideTransformUpdates(bool bOverrideUpdates)
+{
+	Pool.bDisableBoundsAndTransformUpdate = bOverrideUpdates;
+}
+
+void UISMPoolComponent::UpdateAbsoluteTransforms(const FTransform& BaseTransform, EUpdateTransformFlags UpdateTransformFlags, ETeleportType Teleport)
+{
+	Pool.UpdateAbsoluteTransforms(BaseTransform, UpdateTransformFlags, Teleport);
+}
+
+void UISMPoolComponent::GetResourceSizeEx(FResourceSizeEx& CumulativeResourceSize)
+{
+	Super::GetResourceSizeEx(CumulativeResourceSize);
+
+	int32 SizeBytes =
+		MeshGroups.GetAllocatedSize()
+		+ Pool.MeshToISMIndex.GetAllocatedSize()
+		+ Pool.ISMs.GetAllocatedSize()
+		+ Pool.FreeList.GetAllocatedSize()
+		+ Pool.FreeListISM.GetAllocatedSize();
+	
+	for (FISMPoolISM ISM : Pool.ISMs)
+	{
+		SizeBytes += ISM.InstanceIds.GetAllocatedSize()
+			+ ISM.InstanceGroups.GroupRanges.GetAllocatedSize()
+			+ ISM.InstanceGroups.FreeList.GetAllocatedSize();
+	}
+
+	CumulativeResourceSize.AddDedicatedSystemMemoryBytes(SizeBytes);
+}
diff --git a/Engine/Source/Runtime/Experimental/ISMPool/Private/ISMPool/ISMPoolSubSystem.cpp b/Engine/Source/Runtime/Experimental/ISMPool/Private/ISMPool/ISMPoolSubSystem.cpp
new file mode 100644
--- /dev/null
+++ b/Engine/Source/Runtime/Experimental/ISMPool/Private/ISMPool/ISMPoolSubSystem.cpp
@@ -0,0 +1,71 @@
+// Copyright Epic Games, Inc. All Rights Reserved.
+
+#include "ISMPool/ISMPoolSubSystem.h"
+#include "ISMPool/ISMPoolActor.h"
+#include "Engine/World.h"
+
+//----------------------------------------------------------------------//
+// UISMPoolSubSystem
+//----------------------------------------------------------------------//
+UISMPoolSubSystem::UISMPoolSubSystem()
+{
+}
+
+void UISMPoolSubSystem::Initialize(FSubsystemCollectionBase& Collection)
+{
+	Super::Initialize(Collection);
+	Collection.InitializeDependency<UISMPoolSubSystem>();
+}
+
+void UISMPoolSubSystem::Deinitialize()
+{
+	PerLevelISMPoolActors.Reset();
+	Super::Deinitialize();
+}
+
+AISMPoolActor* UISMPoolSubSystem::FindISMPoolActor(ULevel* Level)
+{
+	AISMPoolActor* ISMPoolActor = nullptr;
+
+	// on demand creation of the actor based on level
+	TObjectPtr<AISMPoolActor>* ISMPoolActorInLevel = PerLevelISMPoolActors.Find(Level);
+	if (ISMPoolActorInLevel)
+	{
+		ISMPoolActor = ISMPoolActorInLevel->Get();
+	}
+	else
+	{
+		// we keep it as transient to avoid accumulation of those actors in saved levels
+		FActorSpawnParameters Params;
+		Params.ObjectFlags = EObjectFlags::RF_DuplicateTransient | EObjectFlags::RF_Transient;
+		Params.OverrideLevel = Level;
+		ISMPoolActor = GetWorld()->SpawnActor<AISMPoolActor>(Params);
+		// spawn can still fail if we are in the middle or tearing down the world
+		if (ISMPoolActor)
+		{
+			PerLevelISMPoolActors.Add(Level, ISMPoolActor);
+			// make sure we capture when the actor get removed so we can update our internal structure accordingly 
+			ISMPoolActor->OnEndPlay.AddDynamic(this, &UISMPoolSubSystem::OnActorEndPlay);
+		}
+	}
+	return ISMPoolActor;
+}
+
+void UISMPoolSubSystem::OnActorEndPlay(AActor* InSource, EEndPlayReason::Type Reason)
+{
+	if (ULevel* ActorLevel = InSource->GetLevel())
+	{
+		PerLevelISMPoolActors.Remove(ActorLevel);
+	}
+}
+
+void UISMPoolSubSystem::GetISMPoolActors(TArray<AISMPoolActor*>& OutActors) const
+{
+	for (const auto& MapEntry : PerLevelISMPoolActors)
+	{
+		if (MapEntry.Value != nullptr)
+		{
+			OutActors.Add(MapEntry.Value);
+		}
+	}
+}
diff --git a/Engine/Source/Runtime/Experimental/ISMPool/Public/ISMPool/ISMPoolActor.h b/Engine/Source/Runtime/Experimental/ISMPool/Public/ISMPool/ISMPoolActor.h
new file mode 100644
--- /dev/null
+++ b/Engine/Source/Runtime/Experimental/ISMPool/Public/ISMPool/ISMPoolActor.h
@@ -0,0 +1,28 @@
+// Copyright Epic Games, Inc. All Rights Reserved.
+#pragma once
+
+#include "GameFramework/Actor.h"
+#include "ISMPoolActor.generated.h"
+
+class UISMPoolComponent;
+class UISMPoolDebugDrawComponent;
+
+UCLASS(ConversionRoot, ComponentWrapperClass, MinimalAPI)
+class AISMPoolActor : public AActor
+{
+	GENERATED_UCLASS_BODY()
+
+private:
+	UPROPERTY(Category = ISMPoolActor, VisibleAnywhere, BlueprintReadOnly, meta = (AllowPrivateAccess = "true"))
+	TObjectPtr<class UISMPoolComponent> ISMPoolComp;
+
+	UPROPERTY(Category = ISMPoolActor, VisibleAnywhere, BlueprintReadOnly, meta = (AllowPrivateAccess = "true"))
+	TObjectPtr<class UISMPoolDebugDrawComponent> ISMPoolDebugDrawComp;
+
+public:
+	/** Returns ISMPoolComp subobject **/
+	UISMPoolComponent* GetISMPoolComp() const { return ISMPoolComp; }
+};
+
+
+
diff --git a/Engine/Source/Runtime/Experimental/ISMPool/Public/ISMPool/ISMPoolComponent.h b/Engine/Source/Runtime/Experimental/ISMPool/Public/ISMPool/ISMPoolComponent.h
new file mode 100644
--- /dev/null
+++ b/Engine/Source/Runtime/Experimental/ISMPool/Public/ISMPool/ISMPoolComponent.h
@@ -0,0 +1,392 @@
+// Copyright Epic Games, Inc. All Rights Reserved.
+
+#pragma once
+
+#include "Components/SceneComponent.h"
+#include "Containers/Map.h"
+#include "InstancedStaticMeshDelegates.h"
+#include "Materials/MaterialInterface.h"
+#include "InstanceDataTypes.h"
+
+#include "ISMPoolComponent.generated.h"
+
+class AActor;
+class UISMPoolComponent;
+class UInstancedStaticMeshComponent;
+
+/** 
+ * Structure containing a set of allocated instance ranges in an FISMPoolMeshISM which is the manager for a single ISM component.
+ * The instance ranges don't change once allocated, and aren't the same as the actual render indices in the ISM.
+ * The reason that we don't store the the actual ISM render indices is that ISM component is free to reorder its instances whenever it likes.
+ */
+struct FISMPoolInstanceGroups
+{
+	using FInstanceGroupId = int32;
+
+	/** A single continuous range associated with an FInstanceGroupId. */
+	struct FInstanceGroupRange
+	{
+		FInstanceGroupRange(int32 InStart, int32 InCount)
+			: Start(InStart)
+			, Count(InCount)
+		{
+		}
+
+		int32 Start = 0;
+		int32 Count = 0;
+	};
+
+	/** Reset all contents. */
+	void Reset()
+	{
+		TotalInstanceCount = 0;
+		TotalFreeInstanceCount = 0;
+		GroupRanges.Empty();
+		FreeList.Empty();
+	}
+
+	/** Returns true if no groups ranges are in use. */
+	bool IsEmpty() const
+	{
+		return GroupRanges.Num() == FreeList.Num();
+	}
+
+	/** Returns the maximum instance index that we have allocated in a group. */
+	int32 GetMaxInstanceIndex() const
+	{
+		return TotalInstanceCount;
+	}
+
+	/** Add a new group range. */
+	FInstanceGroupId AddGroup(int32 Count)
+	{
+		// First check to see if we can recycle a group from the free list.
+		for (int32 Index = 0; Index < FreeList.Num(); ++Index)
+		{
+			const FInstanceGroupId GroupId = FreeList[Index];
+			// todo: Could also allow allocating a subrange if Count is less than the group range count.
+			if (Count == GroupRanges[GroupId].Count)
+			{
+				TotalFreeInstanceCount -= Count;
+				FreeList.RemoveAtSwap(Index, EAllowShrinking::No);
+				return GroupId;
+			}
+		}
+		
+		// Create a new group.
+		const int32 StartIndex = TotalInstanceCount;
+		TotalInstanceCount += Count;
+		const FInstanceGroupId GroupId = GroupRanges.Add(FInstanceGroupRange(StartIndex, Count));
+		return GroupId;
+	}
+
+	/** Remove a group range. */
+	void RemoveGroup(FInstanceGroupId GroupId)
+	{
+		// Remove the group by putting on free list for reuse.
+		// Actually removing it would require too much shuffling of the render instance index remapping.
+		TotalFreeInstanceCount += GroupRanges[GroupId].Count;
+		FreeList.Add(GroupId);
+	}
+
+	int32 TotalInstanceCount = 0;
+	int32 TotalFreeInstanceCount = 0;
+	TArray<FInstanceGroupRange> GroupRanges;
+	TArray<FInstanceGroupId> FreeList;
+};
+
+/** 
+ * A description for an ISM component.
+ */
+struct FISMPoolComponentDescription
+{
+	enum EFlags
+	{
+		UseHISM = 1 << 1,							// HISM is no longer supported. This flag is ignored.
+		GpuLodSelection = 1 << 2,
+		ReverseCulling = 1 << 3,
+		StaticMobility = 1 << 4,
+		WorldPositionOffsetWritesVelocity = 1 << 5,
+		EvaluateWorldPositionOffset = 1 << 6,
+		AffectShadow = 1 << 7,
+		AffectDistanceFieldLighting = 1 << 8,
+		AffectDynamicIndirectLighting = 1 << 9,
+		AffectFarShadow = 1 << 10,
+		DistanceCullPrimitive = 1 << 11,
+	};
+
+	uint32 Flags = WorldPositionOffsetWritesVelocity|EvaluateWorldPositionOffset|AffectShadow;
+	int32 NumCustomDataFloats = 0;
+	FVector Position = FVector::ZeroVector;
+	int32 StartCullDistance = 0;
+	int32 EndCullDistance = 0;
+	int32 MinLod = 0;
+	uint32 GroupHash = 0;	// Optional, allows identical SMs to be separated into different groups for finer grained culling
+	float LodScale = 1.f;
+	TArray<FName> Tags;
+	FName StatsCategory;
+
+	bool operator==(const FISMPoolComponentDescription& Other) const
+	{
+		return Flags == Other.Flags &&
+			NumCustomDataFloats == Other.NumCustomDataFloats &&
+			Position == Other.Position &&
+			StartCullDistance == Other.StartCullDistance &&
+			EndCullDistance == Other.EndCullDistance &&
+			MinLod == Other.MinLod &&
+			LodScale == Other.LodScale &&
+			Tags == Other.Tags &&
+			GroupHash == Other.GroupHash &&
+			StatsCategory == Other.StatsCategory;
+	}
+};
+
+FORCEINLINE uint32 GetTypeHash(const FISMPoolComponentDescription& Desc)
+{
+	uint32 Hash = HashCombineFast(GetTypeHash(Desc.Flags), GetTypeHash(Desc.NumCustomDataFloats));
+	Hash = HashCombineFast(Hash, GetTypeHash(Desc.Position));
+	Hash = HashCombineFast(Hash, GetTypeHash(Desc.StartCullDistance));
+	Hash = HashCombineFast(Hash, GetTypeHash(Desc.EndCullDistance));
+	Hash = HashCombineFast(Hash, GetTypeHash(Desc.MinLod));
+	Hash = HashCombineFast(Hash, GetTypeHash(Desc.LodScale));
+	Hash = HashCombineFast(Hash, GetTypeHash(Desc.GroupHash));
+	Hash = HashCombineFast(Hash, GetArrayHash(Desc.Tags.GetData(), Desc.Tags.Num()));
+	return HashCombineFast(Hash, GetTypeHash(Desc.StatsCategory));
+}
+
+/**
+ * A mesh with potentially overriden materials and ISM property description.
+ * We batch instances into ISMs that have equivalent values for this structure.
+ */
+struct FISMPoolStaticMeshInstance
+{
+	TWeakObjectPtr<UStaticMesh> StaticMesh;
+	TArray<TWeakObjectPtr<UMaterialInterface>> MaterialsOverrides;
+	TArray<float> CustomPrimitiveData;
+	FISMPoolComponentDescription Desc;
+
+	bool operator==(const FISMPoolStaticMeshInstance& Other) const 
+	{
+		if (!(Desc == Other.Desc))
+		{
+			return false;
+		}
+		if (!StaticMesh.HasSameIndexAndSerialNumber(Other.StaticMesh))
+		{
+			return false;
+		}
+		if (MaterialsOverrides.Num() != Other.MaterialsOverrides.Num())
+		{
+			return false;
+		}
+		for (int32 MatIndex = 0; MatIndex < MaterialsOverrides.Num(); MatIndex++)
+		{
+			if (!MaterialsOverrides[MatIndex].HasSameIndexAndSerialNumber(Other.MaterialsOverrides[MatIndex]))
+			{
+				return false;
+			}
+		}
+		if (CustomPrimitiveData.Num() != Other.CustomPrimitiveData.Num())
+		{
+			return false;
+		}
+		for (int32 DataIndex = 0; DataIndex < CustomPrimitiveData.Num(); DataIndex++)
+		{
+			if (CustomPrimitiveData[DataIndex] != Other.CustomPrimitiveData[DataIndex])
+			{
+				return false;
+			}
+		}
+
+		return true;
+	}
+};
+
+FORCEINLINE uint32 GetTypeHash(const FISMPoolStaticMeshInstance& MeshInstance)
+{
+	uint32 CombinedHash = GetTypeHash(MeshInstance.StaticMesh);
+	CombinedHash = HashCombineFast(CombinedHash, GetTypeHash(MeshInstance.MaterialsOverrides.Num()));
+	for (const TWeakObjectPtr<UMaterialInterface> Material: MeshInstance.MaterialsOverrides)
+	{
+		CombinedHash = HashCombineFast(CombinedHash, GetTypeHash(Material));
+	}
+	for (const float CustomFloat : MeshInstance.CustomPrimitiveData)
+	{
+		CombinedHash = HashCombineFast(CombinedHash, GetTypeHash(CustomFloat));
+	}
+	CombinedHash = HashCombineFast(CombinedHash, GetTypeHash(MeshInstance.Desc));
+	return CombinedHash;
+}
+
+/** Describes a group of instances within an ISM. */
+struct FISMPoolMeshInfo
+{
+	int32 ISMIndex;
+	FISMPoolInstanceGroups::FInstanceGroupId InstanceGroupIndex;
+
+	TArray<float> CustomData;
+	TArrayView<const float> CustomDataSlice(int32 InstanceIndex, int32 NumCustomDataFloatsPerInstance);
+	void ShadowCopyCustomData(int32 InstanceCount, int32 NumCustomDataFloatsPerInstance, TArrayView<const float> CustomDataFloats);
+};
+
+struct FISMPool;
+
+/**
+ * A mesh group which is a collection of meshes and their related FISMPoolMeshInfo.
+ * We group these with a single handle with the expectation that a client will want to own multiple meshs and release them together.
+ */
+struct FISMPoolMeshGroup
+{
+	using FMeshId = int32;
+
+	/** Adds a new mesh with instance count. We expect to only add a unique mesh instance once to each group. Returns a ID that can be used to update the instances. */
+	FMeshId AddMesh(const FISMPoolStaticMeshInstance& MeshInstance, int32 InstanceCount, const FISMPoolMeshInfo& ISMInstanceInfo, TArrayView<const float> CustomDataFloats);
+	/** Update instance transforms for a group of instances. */
+	bool BatchUpdateInstancesTransforms(FISMPool& ISMPool, FMeshId MeshId, int32 StartInstanceIndex, TArrayView<const FTransform> NewInstancesTransforms, bool bWorldSpace, bool bMarkRenderStateDirty, bool bTeleport);
+	void BatchUpdateInstanceCustomData(FISMPool& ISMPool, int32 CustomFloatIndex, float CustomFloatValue);
+
+	/** Remove all of our managed meshes and associated instances. */
+	void RemoveAllMeshes(FISMPool& ISMPool);
+
+	/** Array of allocated mesh infos. */
+	TArray<FISMPoolMeshInfo> MeshInfos;
+
+	/** Flag for whether we allow removal of instances when transform scale is set to zero. */
+	bool bAllowPerInstanceRemoval = false;
+};
+
+/** Structure containting all info for a single ISM. */
+struct FISMPoolISM
+{
+	/** Create the ISMComponent according to settings on the mesh instance. */
+	void CreateISM(USceneComponent* InOwningComponent);
+	/** Initialize the ISMComponent according to settings on the mesh instance. */
+	void InitISM(const FISMPoolStaticMeshInstance& InMeshInstance, bool bKeepAlive, bool bOverrideTransformUpdates = false);
+	/** Add a group to the ISM. Returns the group index. */
+	FISMPoolInstanceGroups::FInstanceGroupId AddInstanceGroup(int32 InstanceCount, TArrayView<const float> CustomDataFloats);
+
+	/** Unique description of ISM component settings. */
+	FISMPoolStaticMeshInstance MeshInstance;
+	/** Created ISM component. Will be nullptr when this slot has been recycled to FISMPool FreeList. */
+	TObjectPtr<UInstancedStaticMeshComponent> ISMComponent;
+	/** Groups of instances allocated in the ISM. */
+	FISMPoolInstanceGroups InstanceGroups;
+	/** Id of Instance in ISMC */
+	TArray<FPrimitiveInstanceId> InstanceIds;
+};
+
+
+/** A pool of ISMs. */
+struct FISMPool
+{
+	using FISMIndex = int32;
+
+	FISMPool();
+
+	/** Find or add an ISM and return an ISM index handle. */
+	FISMIndex GetOrAddISM(UISMPoolComponent* OwningComponent, const FISMPoolStaticMeshInstance& MeshInstance, bool& bOutISMCreated);
+	/** Remove an ISM. */
+	void RemoveISM(FISMIndex ISMIndex, bool bKeepAlive, bool bRecycle);
+	/** Add instances to ISM and return a mesh info handle. */
+	FISMPoolMeshInfo AddInstancesToISM(UISMPoolComponent* OwningComponent, const FISMPoolStaticMeshInstance& MeshInstance, int32 InstanceCount, TArrayView<const float> CustomDataFloats);
+	/** Remove instances from an ISM. */
+	void RemoveInstancesFromISM(const FISMPoolMeshInfo& MeshInfo);
+	/** Update ISM contents. */
+	bool BatchUpdateInstancesTransforms(FISMPoolMeshInfo& MeshInfo, int32 StartInstanceIndex, TArrayView<const FTransform> NewInstancesTransforms, bool bWorldSpace, bool bMarkRenderStateDirty, bool bTeleport, bool bAllowPerInstanceRemoval);
+	void BatchUpdateInstanceCustomData(FISMPoolMeshInfo const& MeshInfo, int32 CustomFloatIndex, float CustomFloatValue);
+
+	/** Clear all ISM components and associated data. */
+	void Clear();
+
+	/** Tick maintenance of free list and preallocation. */
+	void Tick(UISMPoolComponent* OwningComponent);
+
+	/** Add an ISM description to the preallocation queue. */
+	void RequestPreallocateMeshInstance(const FISMPoolStaticMeshInstance& MeshInstances);
+	/** Process the preallocation queue. Processing is timesliced so that only some of the queue will be processed in every call. */
+	void ProcessPreallocationRequests(UISMPoolComponent* OwningComponent, int32 MaxPreallocations);
+
+	void UpdateAbsoluteTransforms(const FTransform& BaseTransform, EUpdateTransformFlags UpdateTransformFlags, ETeleportType Teleport);
+
+	/** Array of ISM objects. */
+	TArray<FISMPoolISM> ISMs;
+	/** Mapping from mesh description to ISMs array slot. */
+	TMap<FISMPoolStaticMeshInstance, FISMIndex> MeshToISMIndex;
+	
+	/** Set of ISM descriptions that we would like to preallocate. */
+	TSet<FISMPoolStaticMeshInstance> PrellocationQueue;
+
+	/** Free list of indices in ISMs that are empty. */
+	TArray<int32> FreeList;
+	/** Free list of indices in ISMs that have registered ISM components. */
+	TArray<int32> FreeListISM;
+
+	// Cached state of lifecycle cvars from the last Tick()
+	bool bCachedKeepAlive = false;
+	bool bCachedRecycle = false;
+
+	// Whether we force ISMs to use parent bounds and disable transform updates
+	bool bDisableBoundsAndTransformUpdate = false;
+};
+
+
+/**
+* UISMPoolComponent.
+* Component that manages a pool of ISM components in order to allow multiple client components that use the same meshes to the share ISMs.
+*/
+UCLASS(meta = (BlueprintSpawnableComponent), MinimalAPI)
+class UISMPoolComponent: public USceneComponent
+{
+	GENERATED_UCLASS_BODY()
+
+public:
+	using FMeshGroupId = int32;
+	using FMeshId = int32;
+
+	//~ Begin UActorComponent Interface
+	ISMPOOL_API virtual void TickComponent(float DeltaTime, enum ELevelTick TickType, FActorComponentTickFunction* ThisTickFunction) override;
+	ISMPOOL_API virtual void GetResourceSizeEx(FResourceSizeEx& CumulativeResourceSize) override;
+	//~ End UActorComponent Interface
+
+	/** 
+	* Create an Mesh group which represent an arbitrary set of mesh with their instance 
+	* no resources are created until the meshes are added for this group 
+	* return a mesh group Id used to add and update instances
+	*/
+	ISMPOOL_API FMeshGroupId CreateMeshGroup(bool bAllowPerInstanceRemoval = false);
+
+	/** Destroy  a mesh group and its associated resources */
+	ISMPOOL_API void DestroyMeshGroup(FMeshGroupId MeshGroupId);
+
+	/** Add a static mesh for a mesh group */
+	ISMPOOL_API FMeshId AddMeshToGroup(FMeshGroupId MeshGroupId, const FISMPoolStaticMeshInstance& MeshInstance, int32 InstanceCount, TArrayView<const float> CustomDataFloats);
+
+	/** Update transforms for a mesh group */
+	ISMPOOL_API bool BatchUpdateInstancesTransforms(FMeshGroupId MeshGroupId, FMeshId MeshId, int32 StartInstanceIndex, TArrayView<const FTransform> NewInstancesTransforms, bool bWorldSpace = false, bool bMarkRenderStateDirty = false, bool bTeleport = false);
+
+	UE_DEPRECATED(5.3, "BatchUpdateInstancesTransforms Array parameter version is deprecated, use the TArrayView version instead")
+	ISMPOOL_API bool BatchUpdateInstancesTransforms(FMeshGroupId MeshGroupId, FMeshId MeshId, int32 StartInstanceIndex, const TArray<FTransform>& NewInstancesTransforms, bool bWorldSpace = false, bool bMarkRenderStateDirty = false, bool bTeleport = false);
+
+	/** Update a single slot of custom instance data for all instances in a mesh group */
+	ISMPOOL_API bool BatchUpdateInstanceCustomData(FMeshGroupId MeshGroupId, int32 CustomFloatIndex, float CustomFloatValue);
+
+	/** 
+	 * Preallocate an ISM in the pool. 
+	 * Doing this early for known mesh instance descriptions can reduce the component registration cost of AddMeshToGroup() for newly discovered mesh descriptions.
+	 */
+	ISMPOOL_API void PreallocateMeshInstance(const FISMPoolStaticMeshInstance& MeshInstance);
+
+	ISMPOOL_API void SetTickablePoolManagement(bool bEnablePoolManagement);
+
+	ISMPOOL_API void SetOverrideTransformUpdates(bool bOverrideUpdates);
+
+	ISMPOOL_API void UpdateAbsoluteTransforms(const FTransform& BaseTransform, EUpdateTransformFlags UpdateTransformFlags, ETeleportType Teleport);
+
+private:
+	uint32 NextMeshGroupId = 0;
+	TMap<FMeshGroupId, FISMPoolMeshGroup> MeshGroups;
+	FISMPool Pool;
+
+	// Expose internals for debug draw support.
+	friend class UISMPoolDebugDrawComponent;
+};
diff --git a/Engine/Source/Runtime/Experimental/ISMPool/Public/ISMPool/ISMPoolDebugDrawComponent.cpp b/Engine/Source/Runtime/Experimental/ISMPool/Public/ISMPool/ISMPoolDebugDrawComponent.cpp
new file mode 100644
--- /dev/null
+++ b/Engine/Source/Runtime/Experimental/ISMPool/Public/ISMPool/ISMPoolDebugDrawComponent.cpp
@@ -0,0 +1,360 @@
+// Copyright Epic Games, Inc. All Rights Reserved.
+
+#include "ISMPool/ISMPoolDebugDrawComponent.h"
+
+#include "CollisionQueryParams.h"
+#include "ComponentRecreateRenderStateContext.h"
+#include "Components/HierarchicalInstancedStaticMeshComponent.h"
+#include "Engine/StaticMesh.h"
+#include "Engine/World.h"
+#include "EngineUtils.h"
+#include "GameFramework/Pawn.h"
+#include "GameFramework/PlayerController.h"
+#include "ISMPool/ISMPoolActor.h"
+#include "ISMPool/ISMPoolComponent.h"
+#include "Misc/CoreDelegates.h"
+
+#if WITH_EDITOR
+#include "Editor.h"
+#include "EditorViewportClient.h"
+#include "LevelEditorViewport.h"
+#endif
+
+#include UE_INLINE_GENERATED_CPP_BY_NAME(ISMPoolDebugDrawComponent)
+
+static TAutoConsoleVariable<int32> CVarISMPoolStats(
+	TEXT("r.ISMPool.DebugStats"),
+	0,
+	TEXT("Show stats for the ISM pools"),
+	FConsoleVariableDelegate::CreateLambda([](IConsoleVariable* InVariable) { UISMPoolDebugDrawComponent::UpdateAllTickEnabled(); }),
+	ECVF_Default);
+
+static TAutoConsoleVariable<int32> CVarISMPoolDebugDraw(
+	TEXT("r.ISMPool.DebugDraw"),
+	0,
+	TEXT("Show debug drawing for the ISM pools"),
+	FConsoleVariableDelegate::CreateLambda([](IConsoleVariable* InVariable) { UISMPoolDebugDrawComponent::UpdateAllTickEnabled(); FGlobalComponentRecreateRenderStateContext Context; }),
+	ECVF_Default);
+
+
+class FISMPoolDebugDrawSceneProxy final : public FDebugRenderSceneProxy
+{
+public:
+	SIZE_T GetTypeHash() const override
+	{
+		static size_t UniquePointer;
+		return reinterpret_cast<size_t>(&UniquePointer);
+	}
+
+	FISMPoolDebugDrawSceneProxy(UISMPoolDebugDrawComponent const* InComponent)
+		: FDebugRenderSceneProxy(InComponent)
+	{
+		DrawType = FDebugRenderSceneProxy::SolidAndWireMeshes;
+	}
+
+	FPrimitiveViewRelevance GetViewRelevance(const FSceneView* View) const override
+	{
+		FPrimitiveViewRelevance Result;
+		Result.bDrawRelevance = true;
+		Result.bDynamicRelevance = true;
+		Result.bSeparateTranslucency = Result.bNormalTranslucency = true;
+		return Result;
+	}
+};
+
+
+UISMPoolDebugDrawComponent::UISMPoolDebugDrawComponent(const FObjectInitializer& ObjectInitializer)
+	: Super(ObjectInitializer)
+{
+	bIsEditorOnly = false;
+	bSelectable = false;
+
+	PrimaryComponentTick.bCanEverTick = true;
+	PrimaryComponentTick.bStartWithTickEnabled = false;
+	PrimaryComponentTick.bAllowTickOnDedicatedServer = false;
+}
+
+void UISMPoolDebugDrawComponent::UpdateTickEnabled()
+{
+	const bool bForceShowStats = CVarISMPoolStats.GetValueOnAnyThread() != 0;
+	const bool bForceShowBounds = CVarISMPoolDebugDraw.GetValueOnAnyThread() != 0;
+	PrimaryComponentTick.SetTickFunctionEnable(bShowStats || bForceShowStats || bShowBounds || bForceShowBounds);
+}
+
+void UISMPoolDebugDrawComponent::UpdateAllTickEnabled()
+{
+	for (TObjectIterator<UISMPoolDebugDrawComponent> It; It; ++It)
+	{
+		It->UpdateTickEnabled();
+	}
+}
+
+void UISMPoolDebugDrawComponent::BeginPlay()
+{
+	Super::BeginPlay();
+#if UE_ENABLE_DEBUG_DRAWING
+	OnScreenMessagesHandle = FCoreDelegates::OnGetOnScreenMessages.AddUObject(this, &UISMPoolDebugDrawComponent::GetOnScreenMessages);
+#endif
+	UpdateTickEnabled();
+}
+
+void UISMPoolDebugDrawComponent::EndPlay(const EEndPlayReason::Type EndPlayReason)
+{
+	Super::EndPlay(EndPlayReason);
+#if UE_ENABLE_DEBUG_DRAWING
+	FCoreDelegates::OnGetOnScreenMessages.Remove(OnScreenMessagesHandle);
+#endif
+}
+
+FBoxSphereBounds UISMPoolDebugDrawComponent::CalcBounds(const FTransform& LocalToWorld) const
+{
+	return SelectedComponent ? SelectedComponent->CalcBounds(LocalToWorld) : FBox(ForceInitToZero);
+}
+
+#if WITH_EDITOR
+void UISMPoolDebugDrawComponent::PostEditChangeProperty(FPropertyChangedEvent& PropertyChangedEvent)
+{
+	Super::PostEditChangeProperty(PropertyChangedEvent);
+
+	if (PropertyChangedEvent.Property)
+	{
+		if (PropertyChangedEvent.Property->GetFName() == GET_MEMBER_NAME_CHECKED(UISMPoolDebugDrawComponent, bShowStats) ||
+			PropertyChangedEvent.Property->GetFName() == GET_MEMBER_NAME_CHECKED(UISMPoolDebugDrawComponent, bShowBounds))
+		{
+			UpdateTickEnabled();
+		}
+	}
+}
+#endif
+
+void UISMPoolDebugDrawComponent::TickComponent(float DeltaTime, enum ELevelTick TickType, FActorComponentTickFunction* ThisTickFunction)
+{
+	Super::TickComponent(DeltaTime, TickType, ThisTickFunction);
+
+#if UE_ENABLE_DEBUG_DRAWING
+	UInstancedStaticMeshComponent const* FoundISM = nullptr;
+
+	const bool bForceShowStats = CVarISMPoolStats.GetValueOnAnyThread() != 0;
+	const bool bForceShowBounds = CVarISMPoolDebugDraw.GetValueOnAnyThread() != 0;
+	if (bShowStats || bForceShowStats || bShowBounds || bForceShowBounds)
+	{
+#if WITH_EDITOR
+		if (GCurrentLevelEditingViewportClient)
+		{
+			FViewport* Viewport = GCurrentLevelEditingViewportClient->Viewport;
+			HHitProxy* HitResult = Viewport->GetHitProxy(Viewport->GetMouseX(), Viewport->GetMouseY());
+			if (HActor* HitActor = HitProxyCast<HActor>(HitResult))
+			{
+				FoundISM = Cast<UInstancedStaticMeshComponent>(HitActor->PrimComponent);
+			}
+		} 
+#endif
+		if (FoundISM == nullptr)
+		{
+			if (UWorld* World = GetWorld())
+			{
+				APlayerController const* Controller = nullptr;
+				for (FConstPlayerControllerIterator Iterator = World->GetPlayerControllerIterator(); Iterator; ++Iterator)
+				{
+					APlayerController* PC = Iterator->Get();
+					if (PC && PC->IsLocalController())
+					{
+						Controller = PC;
+					}
+				}
+
+				if (Controller != nullptr)
+				{
+					FVector CamLoc;
+					FRotator CamRot;
+					Controller->GetPlayerViewPoint(CamLoc, CamRot);
+					const FVector CamForward = CamRot.Vector();
+					const FVector TraceStart = CamLoc;
+					const FVector TraceEnd = TraceStart + CamForward * 20000;
+
+					FHitResult HitResult(ForceInit);
+					FCollisionQueryParams TraceParams(SCENE_QUERY_STAT(ChaosDebugVisibilityTrace), true, Controller->GetPawn());
+					const bool bHit = World->LineTraceSingleByChannel(HitResult, TraceStart, TraceEnd, ECC_Visibility, TraceParams);
+					if (bHit && HitResult.GetComponent() != nullptr)
+					{
+						if (UInstancedStaticMeshComponent* ISMComponent = Cast<UInstancedStaticMeshComponent>(HitResult.GetComponent()))
+						{
+							FoundISM = ISMComponent;
+						}
+					}
+				}
+			}
+		}
+	}
+
+	if (FoundISM != nullptr && FoundISM->GetOuter() != GetOuter())
+	{
+		FoundISM = nullptr;
+	}
+
+	if (SelectedComponent != FoundISM)
+	{
+		bool bChangeSelection = (SelectedComponent == nullptr) || SelectTimer < 0.f;
+		if (bChangeSelection)
+		{
+#if WITH_EDITOR
+			if (SelectedComponent)
+			{
+				((UInstancedStaticMeshComponent*)SelectedComponent)->PushHoveredToProxy(false);
+			}
+#endif
+			SelectedComponent = FoundISM;
+#if WITH_EDITOR
+			if (SelectedComponent)
+			{
+				((UInstancedStaticMeshComponent*)SelectedComponent)->PushHoveredToProxy(true);
+			}
+#endif
+			SelectTimer = 0.3f;
+			MarkRenderStateDirty();
+		}
+		else
+		{
+			SelectTimer -= DeltaTime;
+		}
+	}
+#endif // UE_ENABLE_DEBUG_DRAWING
+}
+
+#if UE_ENABLE_DEBUG_DRAWING
+
+static void DrawBoxAndChildren(FDebugRenderSceneProxy* DebugProxy, FTransform const& LocalToWorldTransform, TArray<FClusterNode> const& Nodes, int32 NodeIndex, int32 LevelIndex)
+{
+	FClusterNode const& Node = Nodes[NodeIndex];
+
+	const FBox LocalBounds(FBox(Node.BoundMin, Node.BoundMax));
+	
+	static TArray<FColor> Colors = { FColorList::VeryLightGrey, FColorList::CoolCopper, FColorList::GreenYellow, FColorList::CornFlowerBlue, FColorList::DustyRose, FColorList::Red, FColorList::Magenta };
+	const FColor LevelColor = Colors[FMath::Min(LevelIndex, Colors.Num() - 1)];
+	
+	DebugProxy->Boxes.Add(FDebugRenderSceneProxy::FDebugBox(LocalBounds, LevelColor, LocalToWorldTransform));
+
+	if (Node.FirstChild > 0)
+	{
+		for (int32 ChildIndex = Node.FirstChild; ChildIndex <= Node.LastChild; ChildIndex++)
+		{
+			DrawBoxAndChildren(DebugProxy, LocalToWorldTransform, Nodes, ChildIndex, LevelIndex + 1);
+		}
+	}
+}
+
+FDebugRenderSceneProxy* UISMPoolDebugDrawComponent::CreateDebugSceneProxy()
+{
+	const bool bForceShowBounds = CVarISMPoolDebugDraw.GetValueOnAnyThread() != 0;
+	if (!(SelectedComponent && (bShowBounds || bForceShowBounds)))
+	{
+		return nullptr;
+	}
+
+	FDebugRenderSceneProxy* DebugProxy = new FISMPoolDebugDrawSceneProxy(this);
+
+	if (const UHierarchicalInstancedStaticMeshComponent* HISMComponent = Cast<UHierarchicalInstancedStaticMeshComponent>(SelectedComponent))
+	{
+		const FTransform LocalToWorldTransform = HISMComponent->GetComponentTransform();
+		TArray<FClusterNode> TreeNodes;
+		HISMComponent->GetTree(TreeNodes);
+		DrawBoxAndChildren(DebugProxy, LocalToWorldTransform, TreeNodes, 0, 0);
+	}
+	else if(const UInstancedStaticMeshComponent* ISMComponent = Cast<UInstancedStaticMeshComponent>(SelectedComponent))
+	{
+		const FBox LocalBounds = ISMComponent->CalcBounds(FTransform::Identity).GetBox();
+		const FTransform LocalToWorldTransform = ISMComponent->GetComponentTransform();
+		DebugProxy->Boxes.Add(FDebugRenderSceneProxy::FDebugBox(LocalBounds, FColorList::CornFlowerBlue, LocalToWorldTransform));
+	}
+
+	return DebugProxy;
+}
+
+void UISMPoolDebugDrawComponent::GetOnScreenMessages(TMultiMap<FCoreDelegates::EOnScreenMessageSeverity, FText>& OutMessages)
+{
+	const bool bForceShowStats = CVarISMPoolStats.GetValueOnAnyThread() != 0;
+	if (!bShowStats && !bShowGlobalStats && !bForceShowStats)
+	{
+		return;
+	}
+
+	if (bShowGlobalStats || bForceShowStats)
+	{
+		AISMPoolActor const* ISMPoolActor = Cast<AISMPoolActor>(GetOuter());
+		UISMPoolComponent const* ISMPoolComponent = ISMPoolActor ? ISMPoolActor->GetISMPoolComp() : nullptr;
+
+		if (ISMPoolComponent)
+		{
+			int32 NumISMs = 0;
+			int32 NumHISMs = 0;
+			int32 NumInstancesTotal = 0;
+			TArray<int32> NumInstancesPerISM;
+			NumInstancesPerISM.Reserve(ISMPoolComponent->Pool.ISMs.Num());
+
+			for (FISMPoolISM const& ISM : ISMPoolComponent->Pool.ISMs)
+			{
+				if (ISM.ISMComponent != nullptr)
+				{
+					const bool bIsHISM = ISM.ISMComponent->IsA(UHierarchicalInstancedStaticMeshComponent::StaticClass());
+					NumISMs += bIsHISM ? 0 : 1;
+					NumHISMs += bIsHISM ? 1 : 0;
+
+					const int32 NumInstances = ISM.ISMComponent->GetNumRenderInstances();
+					NumInstancesTotal += NumInstances;
+					NumInstancesPerISM.Add(NumInstances);
+				}
+			}
+
+			if (NumInstancesTotal > 0)
+			{
+				OutMessages.Add(FCoreDelegates::EOnScreenMessageSeverity::Warning, FText::FromString(*GetOwner()->GetPathName()));
+				OutMessages.Add(FCoreDelegates::EOnScreenMessageSeverity::Info, FText::FromString(FString::Printf(TEXT("Num ISMS %d"), NumISMs)));
+				OutMessages.Add(FCoreDelegates::EOnScreenMessageSeverity::Info, FText::FromString(FString::Printf(TEXT("Num HISMS %d"), NumHISMs)));
+				OutMessages.Add(FCoreDelegates::EOnScreenMessageSeverity::Info, FText::FromString(FString::Printf(TEXT("Num Instances %d"), NumInstancesTotal)));
+				OutMessages.Add(FCoreDelegates::EOnScreenMessageSeverity::Info, FText::FromString(FString::Printf(TEXT("Average Instances Per ISM %f"), (float)NumInstancesTotal / (float)(NumISMs + NumHISMs))));
+
+				NumInstancesPerISM.Sort();
+				for (int32 Index = 0; Index < 10; ++Index)
+				{
+					const int32 Percentile = (Index + 1) * 10;
+					const int32 PercentileIndex = FMath::Min(NumInstancesPerISM.Num() * (Index + 1) / 10, NumInstancesPerISM.Num() - 1);
+					OutMessages.Add(FCoreDelegates::EOnScreenMessageSeverity::Info, FText::FromString(FString::Printf(TEXT("Num Instances at %d percentile: %d"), Percentile, NumInstancesPerISM[PercentileIndex])));
+				}
+			}
+		}
+	}
+
+	if ((bShowStats || bForceShowStats) && SelectedComponent)
+	{
+		TCHAR const* ISMType = SelectedComponent->IsA(UHierarchicalInstancedStaticMeshComponent::StaticClass()) ? TEXT("HISM") : TEXT("ISM");
+
+		FVector BoundsMin, BoundsMax;
+		SelectedComponent->GetLocalBounds(BoundsMin, BoundsMax);
+		const FVector BoundsSize = BoundsMax - BoundsMin;
+
+		const int32 NumInstances = SelectedComponent->GetNumRenderInstances();
+		const int32 NumCustomDataFloats = SelectedComponent->NumCustomDataFloats;
+		const int32 StartCullDistance = SelectedComponent->InstanceStartCullDistance;
+		const int32 EndCullDistance = SelectedComponent->InstanceEndCullDistance;
+		const int32 NumMaterials = SelectedComponent->GetNumMaterials();
+
+		UStaticMesh const* StaticMesh = SelectedComponent->GetStaticMesh();
+		const int32 NumLods = StaticMesh ? StaticMesh->GetNumLODs() : 0;
+
+		FString ISMDescription = FString::Printf(TEXT("Type=%s   Count=%d   Bounds=(%f., %f., %f.)   CullDistance=(%d - %d)   NumLods= %d    NumMaterials=%d   CustomDataSize=%d"), 
+			ISMType, NumInstances, BoundsSize.X, BoundsSize.Z, BoundsSize.Z, StartCullDistance, EndCullDistance, NumLods, NumMaterials, NumCustomDataFloats);
+
+		for (int32 LodIndex = 0; LodIndex < NumLods; ++LodIndex)
+		{
+			const int32 NumVertices = StaticMesh->GetNumVertices(LodIndex);
+			const int32 NumTriangles = StaticMesh->GetNumTriangles(LodIndex);
+			ISMDescription += FString::Printf(TEXT("\nLod %d: Vertices=%d, Triangles=%d."), LodIndex, NumVertices, NumTriangles);
+		}
+
+		OutMessages.Add(FCoreDelegates::EOnScreenMessageSeverity::Warning, FText::FromString(*GetOwner()->GetPathName()));
+		OutMessages.Add(FCoreDelegates::EOnScreenMessageSeverity::Warning, FText::FromString(*SelectedComponent->GetName()));
+		OutMessages.Add(FCoreDelegates::EOnScreenMessageSeverity::Info, FText::FromString(ISMDescription));
+	}
+}
+
+#endif
diff --git a/Engine/Source/Runtime/Experimental/ISMPool/Public/ISMPool/ISMPoolDebugDrawComponent.h b/Engine/Source/Runtime/Experimental/ISMPool/Public/ISMPool/ISMPoolDebugDrawComponent.h
new file mode 100644
--- /dev/null
+++ b/Engine/Source/Runtime/Experimental/ISMPool/Public/ISMPool/ISMPoolDebugDrawComponent.h
@@ -0,0 +1,55 @@
+// Copyright Epic Games, Inc. All Rights Reserved.
+
+#pragma once
+
+#include "Debug/DebugDrawComponent.h"
+#include "DebugRenderSceneProxy.h"
+#include "ISMPoolDebugDrawComponent.generated.h"
+
+class FISMPoolDebugDrawDelegateHelper;
+class UInstancedStaticMeshComponent;
+
+UCLASS(ClassGroup = Debug)
+class ISMPOOL_API UISMPoolDebugDrawComponent : public UDebugDrawComponent
+{
+	GENERATED_UCLASS_BODY()
+
+	UPROPERTY(EditAnywhere, Category = Debug)
+	bool bShowGlobalStats = false;
+
+	UPROPERTY(EditAnywhere, Category = Debug)
+	bool bShowStats = false;
+
+	UPROPERTY(EditAnywhere, Category = Debug)
+	bool bShowBounds = false;
+
+	UPROPERTY(Transient)
+	TObjectPtr<const UInstancedStaticMeshComponent> SelectedComponent;
+
+	float SelectTimer = 0.f;
+
+protected:
+	virtual void BeginPlay() override;
+	virtual void EndPlay(const EEndPlayReason::Type EndPlayReason) override;
+	virtual void TickComponent(float DeltaTime, enum ELevelTick TickType, FActorComponentTickFunction* ThisTickFunction) override;
+	virtual FBoxSphereBounds CalcBounds(const FTransform& LocalToWorld) const override;
+
+#if WITH_EDITOR
+	virtual void PostEditChangeProperty(FPropertyChangedEvent& PropertyChangedEvent) override;
+#endif
+
+	void UpdateTickEnabled();
+
+#if UE_ENABLE_DEBUG_DRAWING
+  	virtual FDebugRenderSceneProxy* CreateDebugSceneProxy() override;
+	
+	FDebugDrawDelegateHelper DebugDrawDelegateHelper;
+	virtual FDebugDrawDelegateHelper& GetDebugDrawDelegateHelper() override { return DebugDrawDelegateHelper; }
+
+	FDelegateHandle OnScreenMessagesHandle;
+	void GetOnScreenMessages(TMultiMap<FCoreDelegates::EOnScreenMessageSeverity, FText>& OutMessages);
+#endif
+
+public:
+	static void UpdateAllTickEnabled();
+};
diff --git a/Engine/Source/Runtime/Experimental/ISMPool/Public/ISMPool/ISMPoolSubSystem.h b/Engine/Source/Runtime/Experimental/ISMPool/Public/ISMPool/ISMPoolSubSystem.h
new file mode 100644
--- /dev/null
+++ b/Engine/Source/Runtime/Experimental/ISMPool/Public/ISMPool/ISMPoolSubSystem.h
@@ -0,0 +1,38 @@
+// Copyright Epic Games, Inc. All Rights Reserved.
+
+#pragma once
+
+#include "Subsystems/WorldSubsystem.h"
+#include "ISMPoolSubSystem.generated.h"
+
+class AISMPoolActor;
+class ULevel;
+/**
+ * A subsystem managing ISMPool actors.
+ */
+UCLASS(MinimalAPI)
+class UISMPoolSubSystem : public UWorldSubsystem
+{
+	GENERATED_BODY()
+
+public:
+	ISMPOOL_API UISMPoolSubSystem();
+
+	// USubsystem BEGIN
+	ISMPOOL_API virtual void Initialize(FSubsystemCollectionBase& Collection) override;
+	ISMPOOL_API virtual void Deinitialize() override;
+	// USubsystem END
+
+	/** Finds or creates an actor. */
+	ISMPOOL_API AISMPoolActor* FindISMPoolActor(ULevel* Level);
+	
+	/** Get all actors managed by the subsystem. */
+	ISMPOOL_API void GetISMPoolActors(TArray<AISMPoolActor*>& OutActors) const;
+
+protected:
+	UFUNCTION()
+	void OnActorEndPlay(AActor* InSource, EEndPlayReason::Type Reason);
+
+	/** ISMPool are per level **/
+	TMap<TObjectPtr<ULevel>, TObjectPtr<AISMPoolActor> > PerLevelISMPoolActors;
+};
