Constraints scripting: add more debug info

#rnx

diff --git a/Engine/Plugins/Animation/ControlRig/Source/ControlRigEditor/Private/ControlRigSequencerEditorLibrary.cpp b/Engine/Plugins/Animation/ControlRig/Source/ControlRigEditor/Private/ControlRigSequencerEditorLibrary.cpp
--- a/Engine/Plugins/Animation/ControlRig/Source/ControlRigEditor/Private/ControlRigSequencerEditorLibrary.cpp
+++ b/Engine/Plugins/Animation/ControlRig/Source/ControlRigEditor/Private/ControlRigSequencerEditorLibrary.cpp
@@ -482,6 +482,18 @@ UTickableConstraint* UControlRigSequencerEditorLibrary::AddConstraint(UWorld* Wo
 		return nullptr;
 	}
 
+	if (!InChildHandle || !InChildHandle->IsValid())
+    {
+    	UE_LOG(LogControlRig, Error, TEXT("AddConstraint: Need Valid Child Handle"));
+    	return nullptr;
+    }
+
+	if (!InParentHandle || !InParentHandle->IsValid())
+	{
+		UE_LOG(LogControlRig, Error, TEXT("AddConstraint: Need Valid Parent Handle"));
+		return nullptr;
+	}
+
 	UTickableTransformConstraint* Constraint = FTransformConstraintUtils::CreateFromType(World, InType);
 	if (!Constraint)
 	{
diff --git a/Engine/Source/Runtime/Experimental/Animation/Constraints/Private/ConstraintsScripting.cpp b/Engine/Source/Runtime/Experimental/Animation/Constraints/Private/ConstraintsScripting.cpp
--- a/Engine/Source/Runtime/Experimental/Animation/Constraints/Private/ConstraintsScripting.cpp
+++ b/Engine/Source/Runtime/Experimental/Animation/Constraints/Private/ConstraintsScripting.cpp
@@ -23,14 +23,27 @@ UTransformableComponentHandle* UConstraintsScriptingLibrary::CreateTransformable
 
 UTransformableHandle* UConstraintsScriptingLibrary::CreateTransformableHandle(UWorld* InWorld, UObject* InObject, const FName& InAttachmentName)
 {
-	
+	if (!InObject)
+	{
+		if (InAttachmentName == NAME_None)
+		{
+			UE_LOG(LogTemp, Error, TEXT("CreateTransformableHandle: InObject in null."));
+		}
+		else
+		{
+			UE_LOG(LogTemp, Error, TEXT("CreateTransformableHandle ('%s'): InObject in null."), *InAttachmentName.ToString());
+		}
+		return nullptr;	
+	}
+		
 	// look for customized transform handle
 	const FTransformableRegistry& Registry = FTransformableRegistry::Get();
 	if (const FTransformableRegistry::CreateHandleFuncT CreateFunction = Registry.GetCreateFunction(InObject->GetClass()))
 	{
 		return CreateFunction(InObject, InAttachmentName);
 	}
-	
+
+	UE_LOG(LogTemp, Error, TEXT("CreateTransformableHandle: Object Class '%s' not supported."), *InObject->GetClass()->GetName());
 	return nullptr;
 }
 
