Fix Asan crash in Event Manager
Blind fix, I was not able reproduce that crash

diff --git a/Engine/Source/Runtime/Experimental/Chaos/Public/EventManager.h b/Engine/Source/Runtime/Experimental/Chaos/Public/EventManager.h
--- a/Engine/Source/Runtime/Experimental/Chaos/Public/EventManager.h
+++ b/Engine/Source/Runtime/Experimental/Chaos/Public/EventManager.h
@@ -301,7 +301,7 @@ namespace Chaos
 			{
 				return;
 			}
-
+			UnregisterDeferedHandler();
 			HandlerLock.ReadLock();
 			const TMap<IPhysicsProxyBase*, TArray<int32>>* Map = GetProxyToIndexMap(Buffer); // Use t his map to get all proxies used in the event buffer
 			// Only take this path if we have fewer Events than Handlers
@@ -316,7 +316,7 @@ namespace Chaos
 						const IPhysicsProxyBase* Proxy = KeyValue.Get<0>();
 						const UObject* Owner = Proxy->GetOwner();
 						for (TMultiMap<UObject*, FEventHandlerPtr>::TConstKeyIterator It = ProxyOwnerToHandlerMap.CreateConstKeyIterator(Owner); It; ++It)
-						{
+						{ 
 							UniqueHandlers.Add(It.Value());
 						}
 					}
@@ -341,29 +341,36 @@ namespace Chaos
 				}
 			}
 			HandlerLock.ReadUnlock();
-			UnAndRegisterDeferedHandler();
+			RegisterDeferedHandler();
 		}
 
 private:
 
-		void UnAndRegisterDeferedHandler()
+	void UnregisterDeferedHandler()
+	{
+		DeferredHandlerLock.WriteLock();
+		// Move array
+		TArray<const void*> DeferredUnregisterHandlersCopy(MoveTemp(DeferredUnregisterHandlers));
+		check(DeferredUnregisterHandlers.Num() == 0);
+		DeferredHandlerLock.WriteUnlock();
+		for (const void* HandlerPtr : DeferredUnregisterHandlersCopy)
 		{
-			DeferredHandlerLock.WriteLock();
-			// Move array
-			TArray<FEventHandlerPtr> DeferredHandlersCopy(MoveTemp(DeferredHandlers));
-			check(DeferredHandlers.Num() == 0);
-			TArray<const void*> DeferredUnregisterHandlersCopy(MoveTemp(DeferredUnregisterHandlers));
-			check(DeferredUnregisterHandlers.Num() == 0);
-			DeferredHandlerLock.WriteUnlock();
-			for (const FEventHandlerPtr& HandlerPtr : DeferredHandlersCopy)
-			{
-				RegisterHandler(HandlerPtr);
-			}
-			for (const void* HandlerPtr : DeferredUnregisterHandlersCopy)
-			{
-				UnregisterHandler(HandlerPtr);
-			}
+			UnregisterHandler(HandlerPtr);
+		}
+	}
+
+	void RegisterDeferedHandler()
+	{
+		DeferredHandlerLock.WriteLock();
+		// Move array
+		TArray<FEventHandlerPtr> DeferredHandlersCopy(MoveTemp(DeferredHandlers));
+		check(DeferredHandlers.Num() == 0);
+		DeferredHandlerLock.WriteUnlock();
+		for (const FEventHandlerPtr& HandlerPtr : DeferredHandlersCopy)
+		{
+			RegisterHandler(HandlerPtr);
 		}
+	}
 
 		void DeleteHandler(FEventHandlerPtr& HandlerPtr)
 		{
