Chaos: Remove unnecessary conditional check & AccumulatedTime set in Chaos base solver's EnableAsyncMode

Async is initialized with the solver & the execution always resulted in false, so these are no longer necessary; Keeping the logging to stay consistent with disabling async mode for a solver.

#rnx

diff --git a/Engine/Source/Runtime/Experimental/Chaos/Private/Chaos/Framework/PhysicsSolverBase.cpp b/Engine/Source/Runtime/Experimental/Chaos/Private/Chaos/Framework/PhysicsSolverBase.cpp
--- a/Engine/Source/Runtime/Experimental/Chaos/Private/Chaos/Framework/PhysicsSolverBase.cpp
+++ b/Engine/Source/Runtime/Experimental/Chaos/Private/Chaos/Framework/PhysicsSolverBase.cpp
@@ -274,11 +274,7 @@ namespace Chaos
 	void FPhysicsSolverBase::EnableAsyncMode(FReal FixedDt)
 	{
 		AsyncDt = FixedDt;
-		if (AsyncDt != FixedDt)
-		{
-			AccumulatedTime = 0;
-			UE_LOG(LogChaos, Verbose, TEXT("FPhysicsSolverBase::AsyncDt:%f"), IsUsingAsyncResults() ? AsyncDt : -1);
-		}
+		UE_LOG(LogChaos, Verbose, TEXT("FPhysicsSolverBase::AsyncDt:%f"), IsUsingAsyncResults() ? AsyncDt : -1);
 	}
 
 	void FPhysicsSolverBase::DisableAsyncMode()
@@ -287,7 +283,6 @@ namespace Chaos
 		UE_LOG(LogChaos, Verbose, TEXT("FPhysicsSolverBase::AsyncDt:%f"), AsyncDt);
 	}
 
-
 	FPhysicsSolverBase::~FPhysicsSolverBase()
 	{
 		//reset history buffer before freeing any unremoved callback objects
