Use synchronous file I/O when reading/writing to current IAD cache block

#tests Manual
#rnx

diff --git a/Engine/Source/Runtime/Experimental/IoStore/OnDemand/Private/OnDemandInstallCache.cpp b/Engine/Source/Runtime/Experimental/IoStore/OnDemand/Private/OnDemandInstallCache.cpp
--- a/Engine/Source/Runtime/Experimental/IoStore/OnDemand/Private/OnDemandInstallCache.cpp
+++ b/Engine/Source/Runtime/Experimental/IoStore/OnDemand/Private/OnDemandInstallCache.cpp
@@ -26,6 +26,21 @@
 
 namespace UE::IoStore
 {
+
+///////////////////////////////////////////////////////////////////////////////
+namespace CVars
+{
+
+bool bForceSyncIO = false;
+static FAutoConsoleVariableRef CVarForceSyncIO(
+	TEXT("IoStore.OnDemand.ForceSyncIO"),
+	bForceSyncIO,
+	TEXT("Whether to force using synchronous file reads even if cache block is immutable"),
+	ECVF_ReadOnly
+);
+
+}
+
 ///////////////////////////////////////////////////////////////////////////////
 double ToKiB(uint64 Value)
 {
@@ -101,10 +116,12 @@ struct FCas
 	using FLastAccess		= TMap<FCasBlockId, int64>;
 
 	FIoStatus				Initialize(FStringView Directory);
+	FCasLocation			FindChunk(const FIoHash& Hash, bool& bIsLocationInCurrentBlock) const;
 	FCasLocation			FindChunk(const FIoHash& Hash) const;
 	FCasBlockId				CreateBlock();
 	FIoStatus				DeleteBlock(FCasBlockId BlockId, TArray<FCasAddr>& OutAddrs);
 	FString					GetBlockFilename(FCasBlockId BlockId) const;
+	FFileOpenResult 		OpenRead(FCasBlockId BlockId);
 	FSharedAsyncFileHandle	OpenAsyncRead(FCasBlockId  BlockId);
 	FUniqueFileHandle		OpenWrite(FCasBlockId BlockId);
 	void					TrackAccess(FCasBlockId BlockId, int64 UtcTicks);
@@ -153,7 +170,7 @@ FIoStatus FCas::Initialize(FStringView Directory)
 	return EIoErrorCode::Ok;
 };
 
-FCasLocation FCas::FindChunk(const FIoHash& Hash) const
+FCasLocation FCas::FindChunk(const FIoHash& Hash, bool& bIsLocationInCurrentBlock) const
 {
 	const FCasAddr* Addr	= reinterpret_cast<const FCasAddr*>(&Hash);
 	const uint32 TypeHash	= GetTypeHash(*Addr);
@@ -161,6 +178,7 @@ FCasLocation FCas::FindChunk(const FIoHash& Hash) const
 		UE::TUniqueLock Lock(Mutex);
 		if (const FCasLocation* Loc = Lookup.FindByHash(TypeHash, *Addr))
 		{
+			bIsLocationInCurrentBlock = Loc->BlockId == CurrentBlock;
 			return *Loc;
 		}
 	}
@@ -168,6 +186,12 @@ FCasLocation FCas::FindChunk(const FIoHash& Hash) const
 	return FCasLocation{};
 }
 
+FCasLocation FCas::FindChunk(const FIoHash& Hash) const
+{
+	bool bTmp = false;
+	return FindChunk(Hash, bTmp);
+}
+
 FCasBlockId FCas::CreateBlock()
 {
 	IPlatformFile&	Ipf = FPlatformFileManager::Get().GetPlatformFile();
@@ -240,6 +264,14 @@ FString FCas::GetBlockFilename(FCasBlockId BlockId) const
 	return FString(Path.ToView());
 }
 
+FFileOpenResult FCas::OpenRead(FCasBlockId BlockId)
+{
+	const FString	Filename = GetBlockFilename(BlockId);
+	IPlatformFile&	Ipf = FPlatformFileManager::Get().GetPlatformFile();
+
+	return Ipf.OpenRead(*Filename, IPlatformFile::EOpenReadFlags::AllowWrite);
+}
+
 FSharedAsyncFileHandle FCas::OpenAsyncRead(FCasBlockId BlockId)
 {
 	UE::TUniqueLock Lock(Mutex);
@@ -269,8 +301,9 @@ FUniqueFileHandle FCas::OpenWrite(FCasBlockId BlockId)
 	IPlatformFile&	Ipf = FPlatformFileManager::Get().GetPlatformFile();
 	const FString	Filename = GetBlockFilename(BlockId);
 	const bool		bAppend = true;
+	const bool		bAllowRead = true;
 
-	return FUniqueFileHandle(Ipf.OpenWrite(*Filename, bAppend));
+	return FUniqueFileHandle(Ipf.OpenWrite(*Filename, bAppend, bAllowRead));
 }
 
 void FCas::TrackAccess(FCasBlockId BlockId, int64 UtcTicks)
@@ -1143,7 +1176,8 @@ bool FOnDemandInstallCache::Resolve(FIoRequestImpl* Request)
 		return false;
 	}
 
-	const FCasLocation CasLoc = Cas.FindChunk(ChunkInfo.Hash());
+	bool bIsLocationInCurrentBlock = false;
+	const FCasLocation CasLoc = Cas.FindChunk(ChunkInfo.Hash(), bIsLocationInCurrentBlock);
 	if (CasLoc.IsValid() == false)
 	{
 		return false;
@@ -1166,10 +1200,69 @@ bool FOnDemandInstallCache::Resolve(FIoRequestImpl* Request)
 		return false;
 	}
 
+	Cas.TrackAccess(CasLoc.BlockId);
+
+	// Use synchronous file read API when reading from and writing to the same cache block concurrently.
+	const bool bSyncRead = bIsLocationInCurrentBlock || CVars::bForceSyncIO;
+	if (bSyncRead)
+	{
+		// The internal request parameters are attached/owned by the I/O request via
+		// the backend data parameter. The chunk request is deleted in GetCompletedRequests
+		FChunkRequest::Attach(*Request, new FChunkRequest(
+			FSharedAsyncFileHandle(),
+			Request,
+			MoveTemp(ChunkInfo),
+			ChunkRange.ConsumeValueOrDie(),
+			RequestSize));
+
+		UE::Tasks::Launch(UE_SOURCE_LOCATION, [this, Request, CasLoc]
+		{
+			FChunkRequest& ChunkRequest = FChunkRequest::GetRef(*Request);
+			bool bOk = false;
+
+			FFileOpenResult FileOpenResult = Cas.OpenRead(CasLoc.BlockId);
+			if (FileOpenResult.IsValid())
+			{
+				TUniquePtr<IFileHandle> FileHandle = FileOpenResult.StealValue();
+				const int64 CasBlockOffset = CasLoc.BlockOffset + ChunkRequest.ChunkRange.GetOffset();
+				if (Request->IsCancelled() == false && FileHandle->Seek(CasBlockOffset))
+				{
+					bOk = FileHandle->Read(ChunkRequest.EncodedChunk.GetData(), ChunkRequest.EncodedChunk.GetSize());
+					if (bOk == false)
+					{
+						const FString Filename = Cas.GetBlockFilename(CasLoc.BlockId);
+						UE_LOG(LogIoStoreOnDemand, Error, TEXT("Failed to read %llu bytes at offset %lld in CAS block '%s'"),
+							ChunkRequest.EncodedChunk.GetSize(),
+							CasBlockOffset,
+							*Filename);
+					}
+				}
+				else
+				{
+					const FString Filename = Cas.GetBlockFilename(CasLoc.BlockId);
+					UE_LOG(LogIoStoreOnDemand, Error, TEXT("Failed to seek to offset %lld in CAS block '%s'"), CasBlockOffset, *Filename);
+				}
+			}
+			else
+			{
+				const FString Filename = Cas.GetBlockFilename(CasLoc.BlockId);
+				UE_LOG(LogIoStoreOnDemand, Error, TEXT("Failed to open CAS block '%s' for reading, reason '%s'"),
+					*Filename, *FileOpenResult.GetError().GetMessage());
+			}
+
+			const bool bWasCancelled = bOk == false;
+			CompleteRequest(Request, bWasCancelled);
+		});
+
+		return true;
+	}
+
 	FSharedAsyncFileHandle FileHandle = Cas.OpenAsyncRead(CasLoc.BlockId);
 	if (FileHandle.IsValid() == false)
 	{
-		UE_LOG(LogIoStoreOnDemand, Error, TEXT("Failed to CAS block file for reading"));
+		const FString Filename = Cas.GetBlockFilename(CasLoc.BlockId);
+		UE_LOG(LogIoStoreOnDemand, Error, TEXT("Failed to open CAS block '%s' for async reading"), *Filename);
+		TUniquePtr<FChunkRequest> Detached = FChunkRequest::Detach(*Request);
 	}
 
 	// The internal request parameters are attached/owned by the I/O request via
@@ -1189,7 +1282,6 @@ bool FOnDemandInstallCache::Resolve(FIoRequestImpl* Request)
 		});
 	};
 
-	Cas.TrackAccess(CasLoc.BlockId);
 	ChunkRequest.FileReadRequest.Reset(FileHandle->ReadRequest(
 		CasLoc.BlockOffset + ChunkRequest.ChunkRange.GetOffset(),
 		ChunkRequest.ChunkRange.GetLength(),
@@ -1466,7 +1558,6 @@ FIoStatus FOnDemandInstallCache::FlushPendingChunks(FPendingChunks& Chunks)
 			Transaction.BlockCreated(Cas.CurrentBlock);
 		}
 
-		//TODO: Allow reading and writing to the same block? 
 		TUniquePtr<IFileHandle>	CasFileHandle = Cas.OpenWrite(Cas.CurrentBlock);
 		if (CasFileHandle.IsValid() == false)
 		{
