Constraints: invalidate evaluation graphs when registering a new world

#rnx

diff --git a/Engine/Source/Runtime/Experimental/Animation/Constraints/Private/ConstraintSubsystem.cpp b/Engine/Source/Runtime/Experimental/Animation/Constraints/Private/ConstraintSubsystem.cpp
--- a/Engine/Source/Runtime/Experimental/Animation/Constraints/Private/ConstraintSubsystem.cpp
+++ b/Engine/Source/Runtime/Experimental/Animation/Constraints/Private/ConstraintSubsystem.cpp
@@ -4,6 +4,7 @@
 #include "Containers/Ticker.h"
 #include "Engine/Engine.h"
 #include "ConstraintsManager.h"
+#include "Algo/ForEach.h"
 #include "Misc/CoreDelegates.h"
 
 //needs to be static to avoid system getting deleted with dangling handles.
@@ -64,38 +65,34 @@ UConstraintSubsystem* UConstraintSubsystem::Get()
 	return nullptr;
 }
 
-const FConstraintsInWorld* UConstraintSubsystem::ConstraintsInWorldFind(UWorld* InWorld) const
+int32 UConstraintSubsystem::GetConstraintsInWorldIndex(const UWorld* InWorld) const
+{
+	return ConstraintsInWorld.IndexOfByPredicate([InWorld](const FConstraintsInWorld& ConstraintsInWorld)
+	{
+		return ConstraintsInWorld.World.Get() == InWorld; 
+	});
+}
+
+const FConstraintsInWorld* UConstraintSubsystem::ConstraintsInWorldFind(const UWorld* InWorld) const
 {
 	if (bNeedsCleanup)
 	{
 		CleanupInvalidConstraints();
 	}
-	
-	for (const FConstraintsInWorld& CInW : ConstraintsInWorld)
-	{
-		if (CInW.World.Get() == InWorld)
-		{
-			return &CInW;
-		}
-	}
-	return nullptr;
+
+	const int32 Index = GetConstraintsInWorldIndex(InWorld);
+	return Index != INDEX_NONE ? &ConstraintsInWorld[Index] : nullptr;
 }
 
-FConstraintsInWorld* UConstraintSubsystem::ConstraintsInWorldFind(UWorld* InWorld) 
+FConstraintsInWorld* UConstraintSubsystem::ConstraintsInWorldFind(const UWorld* InWorld) 
 {
 	if (bNeedsCleanup)
 	{
 		CleanupInvalidConstraints();
 	}
-	
-	for (FConstraintsInWorld& CInW : ConstraintsInWorld)
-	{
-		if (CInW.World.Get() == InWorld)
-		{
-			return &CInW;
-		}
-	}
-	return nullptr;
+
+	const int32 Index = GetConstraintsInWorldIndex(InWorld);
+	return Index != INDEX_NONE ? &ConstraintsInWorld[Index] : nullptr;
 }
 
 FConstraintsInWorld& UConstraintSubsystem::ConstraintsInWorldFindOrAdd(UWorld* InWorld)
@@ -104,20 +101,25 @@ FConstraintsInWorld& UConstraintSubsystem::ConstraintsInWorldFindOrAdd(UWorld* I
 	{
 		CleanupInvalidConstraints();
 	}
-	
-	for (FConstraintsInWorld& CInW : ConstraintsInWorld)
+
+	const int32 Index = GetConstraintsInWorldIndex(InWorld);
+	if (Index != INDEX_NONE)
 	{
-		if (CInW.World.Get() == InWorld)
-		{
-			return CInW;
-		}
+		return ConstraintsInWorld[Index];
 	}
-	FConstraintsInWorld NewCInW;
-	NewCInW.World = InWorld;
-	return ConstraintsInWorld.Emplace_GetRef(MoveTemp(NewCInW));
+	
+	FConstraintsInWorld& NewConstraintsInWorld = ConstraintsInWorld.Emplace_GetRef();
+	NewConstraintsInWorld.World = InWorld;
+
+	Algo::ForEach(ConstraintsInWorld, [](FConstraintsInWorld& ConstraintsInWorld)
+	{
+		return ConstraintsInWorld.InvalidateGraph();
+	});
+	
+	return NewConstraintsInWorld;
 }
 
-TArray<TWeakObjectPtr<UTickableConstraint>> UConstraintSubsystem::GetConstraints(UWorld* InWorld) const
+TArray<TWeakObjectPtr<UTickableConstraint>> UConstraintSubsystem::GetConstraints(const UWorld* InWorld) const
 {
 	static const TArray< TWeakObjectPtr<UTickableConstraint> > DummyArray;
 	if (const FConstraintsInWorld* Constraints = ConstraintsInWorldFind(InWorld))
@@ -127,7 +129,7 @@ TArray<TWeakObjectPtr<UTickableConstraint>> UConstraintSubsystem::GetConstraints
 	return DummyArray;
 }
 
-const TArray<TWeakObjectPtr<UTickableConstraint>>& UConstraintSubsystem::GetConstraintsArray(UWorld* InWorld) const
+const TArray<TWeakObjectPtr<UTickableConstraint>>& UConstraintSubsystem::GetConstraintsArray(const UWorld* InWorld) const
 {
 	static const TArray< TWeakObjectPtr<UTickableConstraint> > DummyArray;
 	if (const FConstraintsInWorld* Constraints = ConstraintsInWorldFind(InWorld))
@@ -239,27 +241,28 @@ void UConstraintSubsystem::PostEditUndo()
 }
 #endif
 
-void UConstraintSubsystem::OnWorldInit(UWorld* InWorld, const UWorld::InitializationValues IVS)
+void UConstraintSubsystem::OnWorldInit(UWorld* InWorld, const UWorld::InitializationValues)
 {
-	if (UConstraintSubsystem* System = UConstraintSubsystem::Get())
+	if (UConstraintSubsystem* System = Get())
 	{
-		FConstraintsInWorld &ConstraintInWorld = System->ConstraintsInWorldFindOrAdd(InWorld);
-		ConstraintInWorld.Init(InWorld);
+		System->ConstraintsInWorldFindOrAdd(InWorld);
 	}
 }
 
 void UConstraintSubsystem::OnWorldCleanup(UWorld* InWorld, bool bSessionEnded, bool bCleanupResources)
 {
-	if(UConstraintSubsystem* System = UConstraintSubsystem::Get())
+	if (UConstraintSubsystem* System = Get())
 	{
-		for (int32 Index = System->ConstraintsInWorld.Num() - 1; Index >= 0; --Index)
+		const int32 Index = System->GetConstraintsInWorldIndex(InWorld);
+		if (Index != INDEX_NONE)
 		{
-			if (System->ConstraintsInWorld[Index].World.Get() == InWorld)
+			System->ConstraintsInWorld[Index].RemoveConstraints(InWorld);
+			System->ConstraintsInWorld.RemoveAt(Index);
+
+			Algo::ForEach(System->ConstraintsInWorld, [](FConstraintsInWorld& ConstraintsInWorld)
 			{
-				System->ConstraintsInWorld[Index].RemoveConstraints(InWorld);
-				System->ConstraintsInWorld.RemoveAt(Index);
-				break;
-			}
+				return ConstraintsInWorld.InvalidateGraph();
+			});
 		}
 	}
 }
@@ -346,16 +349,6 @@ void FConstraintsInWorld::RemoveConstraints(UWorld* InWorld)
 	InvalidateGraph();
 }
 
-
-void FConstraintsInWorld::Init(UWorld* InWorld)
-{
-	if (InWorld)
-	{
-		World = InWorld;
-	}
-	InvalidateGraph();
-}
-
 FConstraintsEvaluationGraph& FConstraintsInWorld::GetEvaluationGraph()
 {
 	if (!EvaluationGraph)
diff --git a/Engine/Source/Runtime/Experimental/Animation/Constraints/Private/ConstraintSubsystem.h b/Engine/Source/Runtime/Experimental/Animation/Constraints/Private/ConstraintSubsystem.h
--- a/Engine/Source/Runtime/Experimental/Animation/Constraints/Private/ConstraintSubsystem.h
+++ b/Engine/Source/Runtime/Experimental/Animation/Constraints/Private/ConstraintSubsystem.h
@@ -26,7 +26,6 @@ public:
 	UPROPERTY(transient)
 	mutable TArray<TWeakObjectPtr<UTickableConstraint>> Constraints;
 
-	CONSTRAINTS_API void Init(UWorld* World);
 	CONSTRAINTS_API void RemoveConstraints(UWorld* World);
 
 	FConstraintsEvaluationGraph& GetEvaluationGraph();
@@ -58,8 +57,8 @@ public:
 		FConstraintTickFunction* InFunctionToTickBefore,
 		FConstraintTickFunction* InFunctionToTickAfter);
 
-	CONSTRAINTS_API TArray<TWeakObjectPtr<UTickableConstraint>> GetConstraints(UWorld* InWorld) const;
-	CONSTRAINTS_API const TArray<TWeakObjectPtr<UTickableConstraint>>& GetConstraintsArray(UWorld* InWorld) const;
+	CONSTRAINTS_API TArray<TWeakObjectPtr<UTickableConstraint>> GetConstraints(const UWorld* InWorld) const;
+	CONSTRAINTS_API const TArray<TWeakObjectPtr<UTickableConstraint>>& GetConstraintsArray(const UWorld* InWorld) const;
 
 	CONSTRAINTS_API void AddConstraint(UWorld* InWorld, UTickableConstraint* InConstraint);
 	CONSTRAINTS_API void RemoveConstraint(UWorld* InWorld, UTickableConstraint* InConstraint, bool bDoNoCompensate);
@@ -100,9 +99,11 @@ private:
 
 	mutable bool bNeedsCleanup = false;
 
-	const FConstraintsInWorld* ConstraintsInWorldFind(UWorld* InWorld) const;
-	FConstraintsInWorld* ConstraintsInWorldFind(UWorld* InWorld);
+	const FConstraintsInWorld* ConstraintsInWorldFind(const UWorld* InWorld) const;
+	FConstraintsInWorld* ConstraintsInWorldFind(const UWorld* InWorld);
 	FConstraintsInWorld& ConstraintsInWorldFindOrAdd(UWorld* InWorld);
 
+	int32 GetConstraintsInWorldIndex(const UWorld* InWorld) const;
+
 	void CleanupInvalidConstraints() const;
 };
