follow-up to 35970585

#rnx

diff --git a/Engine/Source/Runtime/Experimental/Animation/Constraints/Private/ConstraintSubsystem.cpp b/Engine/Source/Runtime/Experimental/Animation/Constraints/Private/ConstraintSubsystem.cpp
--- a/Engine/Source/Runtime/Experimental/Animation/Constraints/Private/ConstraintSubsystem.cpp
+++ b/Engine/Source/Runtime/Experimental/Animation/Constraints/Private/ConstraintSubsystem.cpp
@@ -67,9 +67,9 @@ UConstraintSubsystem* UConstraintSubsystem::Get()
 
 int32 UConstraintSubsystem::GetConstraintsInWorldIndex(const UWorld* InWorld) const
 {
-	return ConstraintsInWorld.IndexOfByPredicate([InWorld](const FConstraintsInWorld& ConstraintsInWorld)
+	return ConstraintsInWorld.IndexOfByPredicate([InWorld](const FConstraintsInWorld& ConstInWorld)
 	{
-		return ConstraintsInWorld.World.Get() == InWorld; 
+		return ConstInWorld.World.Get() == InWorld; 
 	});
 }
 
@@ -111,9 +111,9 @@ FConstraintsInWorld& UConstraintSubsystem::ConstraintsInWorldFindOrAdd(UWorld* I
 	FConstraintsInWorld& NewConstraintsInWorld = ConstraintsInWorld.Emplace_GetRef();
 	NewConstraintsInWorld.World = InWorld;
 
-	Algo::ForEach(ConstraintsInWorld, [](FConstraintsInWorld& ConstraintsInWorld)
+	Algo::ForEach(ConstraintsInWorld, [](FConstraintsInWorld& ConstInWorld)
 	{
-		return ConstraintsInWorld.InvalidateGraph();
+		return ConstInWorld.InvalidateGraph();
 	});
 	
 	return NewConstraintsInWorld;
@@ -259,9 +259,9 @@ void UConstraintSubsystem::OnWorldCleanup(UWorld* InWorld, bool bSessionEnded, b
 			System->ConstraintsInWorld[Index].RemoveConstraints(InWorld);
 			System->ConstraintsInWorld.RemoveAt(Index);
 
-			Algo::ForEach(System->ConstraintsInWorld, [](FConstraintsInWorld& ConstraintsInWorld)
+			Algo::ForEach(System->ConstraintsInWorld, [](FConstraintsInWorld& ConstInWorld)
 			{
-				return ConstraintsInWorld.InvalidateGraph();
+				return ConstInWorld.InvalidateGraph();
 			});
 		}
 	}
