Build fix - rename to prevent global namespace conflicts

Original CL:
Commit 1 of 3 - Decoupling ISMPool from GeometryCollection

Create a copy of the GC_ISM pool files in its own Runtime/Experimental folder.  This allows our non-GC assets to take advantage of the ISMPooling and custom instance data.

diff --git a/Engine/Source/Runtime/Experimental/GeometryCollectionEngine/Private/GeometryCollection/GeometryCollectionISMPoolComponent.cpp b/Engine/Source/Runtime/Experimental/GeometryCollectionEngine/Private/GeometryCollection/GeometryCollectionISMPoolComponent.cpp
--- a/Engine/Source/Runtime/Experimental/GeometryCollectionEngine/Private/GeometryCollection/GeometryCollectionISMPoolComponent.cpp
+++ b/Engine/Source/Runtime/Experimental/GeometryCollectionEngine/Private/GeometryCollection/GeometryCollectionISMPoolComponent.cpp
@@ -14,8 +14,8 @@
 // This can remove the high cost associated with repeated registration, scene proxy creation and mesh draw command creation.
 // But it can also have a high memory overhead since the ISMs retain hard references to their static meshes.
 static bool GComponentKeepAlive = false; 
-FAutoConsoleVariableRef CVarISMPoolComponentKeepAlive(
-	TEXT("r.ISMPool.ComponentKeepAlive"),
+FAutoConsoleVariableRef CVarGCISMPoolComponentKeepAlive(
+	TEXT("r.GC.ISMPool.ComponentKeepAlive"),
 	GComponentKeepAlive,
 	TEXT("Keep ISM components alive when all their instances are removed."));
 
@@ -26,23 +26,23 @@ FAutoConsoleVariableRef CVarISMPoolComponentKeepAlive(
 // But there is more CPU cost to recycling a component then to simply keeping it alive because scene proxy creation and mesh draw command caching isn't cheap.
 // The component memory cost is kept bounded when compared to keeping components alive.
 static bool GComponentRecycle = true;
-FAutoConsoleVariableRef CVarISMPoolComponentRecycle(
-	TEXT("r.ISMPool.ComponentRecycle"),
+FAutoConsoleVariableRef CVarGCISMPoolComponentRecycle(
+	TEXT("r.GC.ISMPool.ComponentRecycle"),
 	GComponentRecycle,
 	TEXT("Recycle ISM components to a free list for reuse when all their instances are removed."));
 
 // Target free list size when recycling ISM components.
 // We try to maintain a pool of free components for fast allocation, but want to clean up when numbers get too high.
 static int32 GComponentFreeListTargetSize = 50;
-FAutoConsoleVariableRef CVarISMPoolComponentFreeListTargetSize(
-	TEXT("r.ISMPool.ComponentFreeListTargetSize"),
+FAutoConsoleVariableRef CVarGCISMPoolComponentFreeListTargetSize(
+	TEXT("r.GC.ISMPool.ComponentFreeListTargetSize"),
 	GComponentFreeListTargetSize,
 	TEXT("Target size for number of ISM components in the recycling free list."));
 
 // Keep copies of all custom instance data for restoration on readding an instance.
 static bool GShadowCopyCustomData = false;
 FAutoConsoleVariableRef CVarShadowCopyCustomData(
-	TEXT("r.ISMPool.ShadowCopyCustomData"),
+	TEXT("r.GC.ISMPool.ShadowCopyCustomData"),
 	GShadowCopyCustomData,
 	TEXT("Keeps a copy of custom instance data so it can be restored if the instance is removed and readded."));
 
diff --git a/Engine/Source/Runtime/Experimental/ISMPool/Private/ISMPool/ISMPoolComponent.cpp b/Engine/Source/Runtime/Experimental/ISMPool/Private/ISMPool/ISMPoolComponent.cpp
--- a/Engine/Source/Runtime/Experimental/ISMPool/Private/ISMPool/ISMPoolComponent.cpp
+++ b/Engine/Source/Runtime/Experimental/ISMPool/Private/ISMPool/ISMPoolComponent.cpp
@@ -41,7 +41,7 @@ FAutoConsoleVariableRef CVarISMPoolComponentFreeListTargetSize(
 
 // Keep copies of all custom instance data for restoration on readding an instance.
 static bool GShadowCopyCustomData = false;
-FAutoConsoleVariableRef CVarShadowCopyCustomData(
+FAutoConsoleVariableRef CVarISMPoolShadowCopyCustomData(
 	TEXT("r.ISMPool.ShadowCopyCustomData"),
 	GShadowCopyCustomData,
 	TEXT("Keeps a copy of custom instance data so it can be restored if the instance is removed and readded."));
