Dataflow nodes : Vector and Trig math nodes
- Also added support for vector compatible anytypes types

diff --git a/Engine/Source/Runtime/Experimental/Dataflow/Core/Private/Dataflow/DataflowMathNodes.cpp b/Engine/Source/Runtime/Experimental/Dataflow/Core/Private/Dataflow/DataflowMathNodes.cpp
--- a/Engine/Source/Runtime/Experimental/Dataflow/Core/Private/Dataflow/DataflowMathNodes.cpp
+++ b/Engine/Source/Runtime/Experimental/Dataflow/Core/Private/Dataflow/DataflowMathNodes.cpp
@@ -5,13 +5,17 @@
 #include "Dataflow/DataflowNodeFactory.h"
 #include "Dataflow/DataflowNodeColorsRegistry.h"
 
+#include "MathUtil.h"
+
 namespace Dataflow
 {
 	void RegisterDataflowMathNodes()
 	{
+		// scalar
 		DATAFLOW_NODE_REGISTER_CREATION_FACTORY(FDataflowMathAbsNode);
 		DATAFLOW_NODE_REGISTER_CREATION_FACTORY(FDataflowMathAddNode);
 		DATAFLOW_NODE_REGISTER_CREATION_FACTORY(FDataflowMathCeilNode);
+		DATAFLOW_NODE_REGISTER_CREATION_FACTORY(FDataflowMathConstantNode);
 		DATAFLOW_NODE_REGISTER_CREATION_FACTORY(FDataflowMathCubeNode);
 		DATAFLOW_NODE_REGISTER_CREATION_FACTORY(FDataflowMathDivideNode);
 		DATAFLOW_NODE_REGISTER_CREATION_FACTORY(FDataflowMathExpNode);
@@ -24,6 +28,7 @@ namespace Dataflow
 		DATAFLOW_NODE_REGISTER_CREATION_FACTORY(FDataflowMathMinimumNode);
 		DATAFLOW_NODE_REGISTER_CREATION_FACTORY(FDataflowMathMultiplyNode);
 		DATAFLOW_NODE_REGISTER_CREATION_FACTORY(FDataflowMathNegateNode);
+		DATAFLOW_NODE_REGISTER_CREATION_FACTORY(FDataflowMathOneMinusNode);
 		DATAFLOW_NODE_REGISTER_CREATION_FACTORY(FDataflowMathPowNode);
 		DATAFLOW_NODE_REGISTER_CREATION_FACTORY(FDataflowMathReciprocalNode);
 		DATAFLOW_NODE_REGISTER_CREATION_FACTORY(FDataflowMathRoundNode);
@@ -33,6 +38,17 @@ namespace Dataflow
 		DATAFLOW_NODE_REGISTER_CREATION_FACTORY(FDataflowMathSubtractNode);
 		DATAFLOW_NODE_REGISTER_CREATION_FACTORY(FDataflowMathTruncNode);
 
+		// trigonometric
+		DATAFLOW_NODE_REGISTER_CREATION_FACTORY(FDataflowMathCosNode);
+		DATAFLOW_NODE_REGISTER_CREATION_FACTORY(FDataflowMathSinNode);
+		DATAFLOW_NODE_REGISTER_CREATION_FACTORY(FDataflowMathTanNode);
+		DATAFLOW_NODE_REGISTER_CREATION_FACTORY(FDataflowMathArcCosNode);
+		DATAFLOW_NODE_REGISTER_CREATION_FACTORY(FDataflowMathArcSinNode);
+		DATAFLOW_NODE_REGISTER_CREATION_FACTORY(FDataflowMathArcTanNode);
+		DATAFLOW_NODE_REGISTER_CREATION_FACTORY(FDataflowMathArcTan2Node)
+		DATAFLOW_NODE_REGISTER_CREATION_FACTORY(FDataflowMathDegToRadNode);
+		DATAFLOW_NODE_REGISTER_CREATION_FACTORY(FDataflowMathRadToDegNode);
+
 		// Math
 		static constexpr FLinearColor CDefaultMathNodeBodyTintColor(0.f, 0.f, 0.f, 0.5f);
 		DATAFLOW_NODE_REGISTER_CREATION_FACTORY_NODE_COLORS_BY_CATEGORY("Math", FLinearColor(0.f, 0.4f, 0.8f), CDefaultMathNodeBodyTintColor);
@@ -51,8 +67,8 @@ void FDataflowMathOneInputOperatorNode::RegisterInputsAndOutputs()
 	RegisterInputConnection(&A);
 	RegisterOutputConnection(&Result);
 
-	// set the output to double for now so that it is strongly type and easy to connect to the next node
-	// Cnce we can change the outpout type from the UI , this could be removed 
+	// Set the output to double for now so that it is strongly type and easy to connect to the next node
+	// Once we can change the output type from the UI, this could be removed 
 	SetOutputConcreteType(&Result, TDataflowSingleTypePolicy<double>::TypeName);
 }
 
@@ -79,8 +95,8 @@ void FDataflowMathTwoInputsOperatorNode::RegisterInputsAndOutputs()
 	RegisterInputConnection(&B);
 	RegisterOutputConnection(&Result);
 
-	// set the output to double for now so that it is strongly type and easy to connect to the next node
-	// Cnce we can change the outpout type from the UI , this could be removed 
+	// Set the output to double for now so that it is strongly type and easy to connect to the next node
+	// Once we can change the output type from the UI, this could be removed 
 	SetOutputConcreteType(&Result, TDataflowSingleTypePolicy<double>::TypeName);
 }
 
@@ -419,4 +435,182 @@ FDataflowMathSignNode::FDataflowMathSignNode(const Dataflow::FNodeParameters& In
 double FDataflowMathSignNode::ComputeResult(Dataflow::FContext& Context, double InA) const
 {
 	return FMath::Sign(InA);
-}
\ No newline at end of file
+}
+
+//-----------------------------------------------------------------------------------------------
+
+FDataflowMathOneMinusNode::FDataflowMathOneMinusNode(const Dataflow::FNodeParameters& InParam, FGuid InGuid)
+	: FDataflowMathOneInputOperatorNode(InParam, InGuid)
+{
+	RegisterInputsAndOutputs();
+}
+
+double FDataflowMathOneMinusNode::ComputeResult(Dataflow::FContext& Context, double InA) const
+{
+	return (1.0 - InA);
+}
+
+//-----------------------------------------------------------------------------------------------
+
+FDataflowMathConstantNode::FDataflowMathConstantNode(const Dataflow::FNodeParameters& InParam, FGuid InGuid)
+	: FDataflowNode(InParam, InGuid)
+{
+	RegisterOutputConnection(&Result);
+
+	// Set the output to double for now so that it is strongly type and easy to connect to the next node
+	// Once we can change the output type from the UI, this could be removed 
+	SetOutputConcreteType(&Result, TDataflowSingleTypePolicy<double>::TypeName);
+}
+
+double FDataflowMathConstantNode::GetConstant() const
+{
+	switch (Constant)
+	{
+	case EDataflowMathConstantsEnum::Dataflow_Math_Constants_Pi:			return FMathd::Pi;
+	case EDataflowMathConstantsEnum::Dataflow_Math_Constants_HalfPi:		return FMathd::HalfPi;
+	case EDataflowMathConstantsEnum::Dataflow_Math_Constants_TwoPi:			return FMathd::TwoPi;
+	case EDataflowMathConstantsEnum::Dataflow_Math_Constants_FourPi:		return FMathd::FourPi;
+	case EDataflowMathConstantsEnum::Dataflow_Math_Constants_InvPi:			return FMathd::InvPi;
+	case EDataflowMathConstantsEnum::Dataflow_Math_Constants_InvTwoPi:		return FMathd::InvTwoPi;
+	case EDataflowMathConstantsEnum::Dataflow_Math_Constants_Sqrt2:			return FMathd::Sqrt2;
+	case EDataflowMathConstantsEnum::Dataflow_Math_Constants_InvSqrt2:		return FMathd::InvSqrt2;
+	case EDataflowMathConstantsEnum::Dataflow_Math_Constants_Sqrt3:			return FMathd::Sqrt3;
+	case EDataflowMathConstantsEnum::Dataflow_Math_Constants_InvSqrt3:		return FMathd::InvSqrt3;
+	case EDataflowMathConstantsEnum::Dataflow_Math_Constants_E:				return 2.71828182845904523536;
+	case EDataflowMathConstantsEnum::Dataflow_Math_Constants_Gamma:			return 0.577215664901532860606512090082;
+	case EDataflowMathConstantsEnum::Dataflow_Math_Constants_GoldenRatio:	return 1.618033988749894;
+	default:
+		break;
+	}
+	ensureMsgf(false, TEXT("Unexpected constant enum, returning a zero value. Is it missing from the list above ?"));
+	return 0.0;
+}
+
+void FDataflowMathConstantNode::Evaluate(Dataflow::FContext& Context, const FDataflowOutput* Out) const
+{
+	if (Out->IsA(&Result))
+	{
+		SetValue(Context, GetConstant(), &Result);
+	}
+}
+
+//--------------------------------------------------------------------------
+//
+// Trigonometric nodes
+//
+//--------------------------------------------------------------------------
+
+FDataflowMathSinNode::FDataflowMathSinNode(const Dataflow::FNodeParameters& InParam, FGuid InGuid)
+	: FDataflowMathOneInputOperatorNode(InParam, InGuid)
+{
+	RegisterInputsAndOutputs();
+}
+
+double FDataflowMathSinNode::ComputeResult(Dataflow::FContext& Context, double InA) const
+{
+	return FMath::Sin(InA);
+}
+
+//-----------------------------------------------------------------------------------------------
+
+FDataflowMathCosNode::FDataflowMathCosNode(const Dataflow::FNodeParameters& InParam, FGuid InGuid)
+	: FDataflowMathOneInputOperatorNode(InParam, InGuid)
+{
+	RegisterInputsAndOutputs();
+}
+
+double FDataflowMathCosNode::ComputeResult(Dataflow::FContext& Context, double InA) const
+{
+	return FMath::Cos(InA);
+}
+
+//-----------------------------------------------------------------------------------------------
+
+FDataflowMathTanNode::FDataflowMathTanNode(const Dataflow::FNodeParameters& InParam, FGuid InGuid)
+	: FDataflowMathOneInputOperatorNode(InParam, InGuid)
+{
+	RegisterInputsAndOutputs();
+}
+
+double FDataflowMathTanNode::ComputeResult(Dataflow::FContext& Context, double InA) const
+{
+	return FMath::Tan(InA);
+}
+
+//-----------------------------------------------------------------------------------------------
+
+FDataflowMathArcSinNode::FDataflowMathArcSinNode(const Dataflow::FNodeParameters& InParam, FGuid InGuid)
+	: FDataflowMathOneInputOperatorNode(InParam, InGuid)
+{
+	RegisterInputsAndOutputs();
+}
+
+double FDataflowMathArcSinNode::ComputeResult(Dataflow::FContext& Context, double InA) const
+{
+	return FMath::Asin(InA);
+}
+
+//-----------------------------------------------------------------------------------------------
+
+FDataflowMathArcCosNode::FDataflowMathArcCosNode(const Dataflow::FNodeParameters& InParam, FGuid InGuid)
+	: FDataflowMathOneInputOperatorNode(InParam, InGuid)
+{
+	RegisterInputsAndOutputs();
+}
+
+double FDataflowMathArcCosNode::ComputeResult(Dataflow::FContext& Context, double InA) const
+{
+	return FMath::Acos(InA);
+}
+
+//-----------------------------------------------------------------------------------------------
+
+FDataflowMathArcTanNode::FDataflowMathArcTanNode(const Dataflow::FNodeParameters& InParam, FGuid InGuid)
+	: FDataflowMathOneInputOperatorNode(InParam, InGuid)
+{
+	RegisterInputsAndOutputs();
+}
+
+double FDataflowMathArcTanNode::ComputeResult(Dataflow::FContext& Context, double InA) const
+{
+	return FMath::Atan(InA);
+}
+
+//-----------------------------------------------------------------------------------------------
+
+FDataflowMathArcTan2Node::FDataflowMathArcTan2Node(const Dataflow::FNodeParameters& InParam, FGuid InGuid)
+	: FDataflowMathTwoInputsOperatorNode(InParam, InGuid)
+{
+	RegisterInputsAndOutputs();
+}
+
+double FDataflowMathArcTan2Node::ComputeResult(Dataflow::FContext& Context, double InA, double InB) const
+{
+	return FMath::Atan2(InA, InB);
+}
+
+//-----------------------------------------------------------------------------------------------
+
+FDataflowMathDegToRadNode::FDataflowMathDegToRadNode(const Dataflow::FNodeParameters& InParam, FGuid InGuid)
+	: FDataflowMathOneInputOperatorNode(InParam, InGuid)
+{
+	RegisterInputsAndOutputs();
+}
+
+double FDataflowMathDegToRadNode::ComputeResult(Dataflow::FContext& Context, double InA) const
+{
+	return FMath::DegreesToRadians(InA);
+}
+
+//-----------------------------------------------------------------------------------------------
+
+FDataflowMathRadToDegNode::FDataflowMathRadToDegNode(const Dataflow::FNodeParameters& InParam, FGuid InGuid)
+	: FDataflowMathOneInputOperatorNode(InParam, InGuid)
+{
+	RegisterInputsAndOutputs();
+}
+
+double FDataflowMathRadToDegNode::ComputeResult(Dataflow::FContext& Context, double InA) const
+{
+	return FMath::RadiansToDegrees(InA);
+}
diff --git a/Engine/Source/Runtime/Experimental/Dataflow/Core/Private/Dataflow/DataflowNodeFactory.cpp b/Engine/Source/Runtime/Experimental/Dataflow/Core/Private/Dataflow/DataflowNodeFactory.cpp
--- a/Engine/Source/Runtime/Experimental/Dataflow/Core/Private/Dataflow/DataflowNodeFactory.cpp
+++ b/Engine/Source/Runtime/Experimental/Dataflow/Core/Private/Dataflow/DataflowNodeFactory.cpp
@@ -5,6 +5,7 @@
 #include "Dataflow/DataflowNode.h"
 #include "Dataflow/DataflowCoreNodes.h"
 #include "Dataflow/DataflowMathNodes.h"
+#include "Dataflow/DataflowVectorNodes.h"
 #include "Misc/MessageDialog.h"
 #include "Logging/LogMacros.h"
 
@@ -18,6 +19,7 @@ namespace Dataflow
 	{
 		Dataflow::RegisterCoreNodes();
 		Dataflow::RegisterDataflowMathNodes();
+		Dataflow::RegisterDataflowVectorNodes();
 	}
 
 	TSharedPtr<FDataflowNode> FNodeFactory::NewNodeFromRegisteredType(FGraph& Graph, const FNewNodeParameters& Param)
diff --git a/Engine/Source/Runtime/Experimental/Dataflow/Core/Private/Dataflow/DataflowVectorNodes.cpp b/Engine/Source/Runtime/Experimental/Dataflow/Core/Private/Dataflow/DataflowVectorNodes.cpp
new file mode 100644
--- /dev/null
+++ b/Engine/Source/Runtime/Experimental/Dataflow/Core/Private/Dataflow/DataflowVectorNodes.cpp
@@ -0,0 +1,373 @@
+// Copyright Epic Games, Inc. All Rights Reserved.
+
+#include "Dataflow/DataflowVectorNodes.h"
+
+#include "Dataflow/DataflowNodeFactory.h"
+#include "Dataflow/DataflowNodeColorsRegistry.h"
+
+namespace Dataflow
+{
+	void RegisterDataflowVectorNodes()
+	{
+		DATAFLOW_NODE_REGISTER_CREATION_FACTORY(FDataflowVectorMakeVec2Node);
+		DATAFLOW_NODE_REGISTER_CREATION_FACTORY(FDataflowVectorMakeVec3Node);
+		DATAFLOW_NODE_REGISTER_CREATION_FACTORY(FDataflowVectorMakeVec4Node);
+		DATAFLOW_NODE_REGISTER_CREATION_FACTORY(FDataflowVectorBreakNode);
+		DATAFLOW_NODE_REGISTER_CREATION_FACTORY(FDataflowVectorAddNode);
+		DATAFLOW_NODE_REGISTER_CREATION_FACTORY(FDataflowVectorSubtractNode);
+		DATAFLOW_NODE_REGISTER_CREATION_FACTORY(FDataflowVectorDotProductNode);
+		DATAFLOW_NODE_REGISTER_CREATION_FACTORY(FDataflowVectorLengthNode);
+		DATAFLOW_NODE_REGISTER_CREATION_FACTORY(FDataflowVectorSquaredLengthNode);
+		DATAFLOW_NODE_REGISTER_CREATION_FACTORY(FDataflowVectorDistanceNode);
+		DATAFLOW_NODE_REGISTER_CREATION_FACTORY(FDataflowVectorCrossProductNode);
+		DATAFLOW_NODE_REGISTER_CREATION_FACTORY(FDataflowVectorScaleNode);
+		DATAFLOW_NODE_REGISTER_CREATION_FACTORY(FDataflowVectorNormalize);
+	}
+}
+
+
+//-----------------------------------------------------------
+
+FDataflowVectorMakeVec2Node::FDataflowVectorMakeVec2Node(const Dataflow::FNodeParameters& InParam, FGuid InGuid)
+	: FDataflowNode(InParam, InGuid)
+{
+	RegisterInputConnection(&X);
+	RegisterInputConnection(&Y);
+	RegisterOutputConnection(&Vector2D);
+	// set default to 2D vector
+	SetOutputConcreteType<FVector2D>(&Vector2D);
+}
+
+void FDataflowVectorMakeVec2Node::Evaluate(Dataflow::FContext& Context, const FDataflowOutput* Out) const
+{
+	if (Out->IsA(&Vector2D))
+	{
+		const double InX = GetValue(Context, &X);
+		const double InY = GetValue(Context, &Y);
+		SetValue(Context, FVector4(InX, InY, 0, 0), &Vector2D);
+	}
+}
+
+//-----------------------------------------------------------
+
+FDataflowVectorMakeVec3Node::FDataflowVectorMakeVec3Node(const Dataflow::FNodeParameters& InParam, FGuid InGuid)
+	: FDataflowNode(InParam, InGuid)
+{
+	RegisterInputConnection(&X);
+	RegisterInputConnection(&Y);
+	RegisterInputConnection(&Z);
+	RegisterOutputConnection(&Vector3D);
+	// set default to 3D vector
+	SetOutputConcreteType<FVector>(&Vector3D);
+}
+
+void FDataflowVectorMakeVec3Node::Evaluate(Dataflow::FContext& Context, const FDataflowOutput* Out) const
+{
+	if (Out->IsA(&Vector3D))
+	{
+		const double InX = GetValue(Context, &X);
+		const double InY = GetValue(Context, &Y);
+		const double InZ = GetValue(Context, &Z);
+		SetValue(Context, FVector4(InX, InY, InZ, 0), &Vector3D);
+	}
+}
+
+//-----------------------------------------------------------
+
+FDataflowVectorMakeVec4Node::FDataflowVectorMakeVec4Node(const Dataflow::FNodeParameters& InParam, FGuid InGuid)
+	: FDataflowNode(InParam, InGuid)
+{
+	RegisterInputConnection(&X);
+	RegisterInputConnection(&Y);
+	RegisterInputConnection(&Z);
+	RegisterInputConnection(&W);
+	RegisterOutputConnection(&Vector4D);
+	// set default to 4D vector
+	SetOutputConcreteType<FVector4>(&Vector4D);
+}
+
+void FDataflowVectorMakeVec4Node::Evaluate(Dataflow::FContext& Context, const FDataflowOutput* Out) const
+{
+	if (Out->IsA(&Vector4D))
+	{
+		const double InX = GetValue(Context, &X);
+		const double InY = GetValue(Context, &Y);
+		const double InZ = GetValue(Context, &Z);
+		const double InW = GetValue(Context, &W);
+		SetValue(Context, FVector4(InX, InY, InZ, InW), &Vector4D);
+	}
+}
+
+//-----------------------------------------------------------
+
+FDataflowVectorBreakNode::FDataflowVectorBreakNode(const Dataflow::FNodeParameters& InParam, FGuid InGuid)
+	: FDataflowNode(InParam, InGuid)
+{
+	RegisterInputConnection(&V);
+	RegisterOutputConnection(&X);
+	RegisterOutputConnection(&Y);
+	RegisterOutputConnection(&Z);
+	RegisterOutputConnection(&W);
+
+	// set default out types to double
+	SetOutputConcreteType<double>(&X);
+	SetOutputConcreteType<double>(&Y);
+	SetOutputConcreteType<double>(&Z);
+	SetOutputConcreteType<double>(&W);
+}
+
+void FDataflowVectorBreakNode::Evaluate(Dataflow::FContext& Context, const FDataflowOutput* Out) const
+{
+	if (Out->IsA(&X))
+	{
+		SetValue(Context, GetValue(Context, &V).X, &X);
+	}
+	else if (Out->IsA(&Y))
+	{
+		SetValue(Context, GetValue(Context, &V).Y, &Y);
+	}
+	else if (Out->IsA(&Z))
+	{
+		SetValue(Context, GetValue(Context, &V).Z, &Z);
+	}
+	else if (Out->IsA(&W))
+	{
+		SetValue(Context, GetValue(Context, &V).W, &W);
+	}
+}
+
+//-----------------------------------------------------------
+
+FDataflowVectorAddNode::FDataflowVectorAddNode(const Dataflow::FNodeParameters& InParam, FGuid InGuid)
+	: FDataflowNode(InParam, InGuid)
+{
+	RegisterInputConnection(&A);
+	RegisterInputConnection(&B);
+	RegisterOutputConnection(&V);
+	// set default to 4D vector
+	SetOutputConcreteType<FVector4>(&V);
+
+}
+
+void FDataflowVectorAddNode::Evaluate(Dataflow::FContext& Context, const FDataflowOutput* Out) const
+{
+	if (Out->IsA(&V))
+	{
+		const FVector4 InA = GetValue(Context, &A);
+		const FVector4 InB = GetValue(Context, &B);
+		SetValue(Context, (InA + InB), &V);
+	}
+}
+
+//-----------------------------------------------------------
+
+FDataflowVectorSubtractNode::FDataflowVectorSubtractNode(const Dataflow::FNodeParameters& InParam, FGuid InGuid)
+	: FDataflowNode(InParam, InGuid)
+{
+	RegisterInputConnection(&A);
+	RegisterInputConnection(&B);
+	RegisterOutputConnection(&V);
+	// set default to 4D vector
+	SetOutputConcreteType<FVector4>(&V);
+
+}
+
+void FDataflowVectorSubtractNode::Evaluate(Dataflow::FContext& Context, const FDataflowOutput* Out) const
+{
+	if (Out->IsA(&V))
+	{
+		const FVector4 InA = GetValue(Context, &A);
+		const FVector4 InB = GetValue(Context, &B);
+		SetValue(Context, (InA - InB), &V);
+	}
+}
+
+//-----------------------------------------------------------
+
+FDataflowVectorDotProductNode::FDataflowVectorDotProductNode(const Dataflow::FNodeParameters& InParam, FGuid InGuid)
+	: FDataflowNode(InParam, InGuid)
+{
+	RegisterInputConnection(&A);
+	RegisterInputConnection(&B);
+	RegisterOutputConnection(&DotProduct);
+	// set default type to be double
+	SetOutputConcreteType<double>(&DotProduct);
+}
+
+void FDataflowVectorDotProductNode::Evaluate(Dataflow::FContext& Context, const FDataflowOutput* Out) const
+{
+	if (Out->IsA(&DotProduct))
+	{
+		const FVector4 InA = GetValue(Context, &A);
+		const FVector4 InB = GetValue(Context, &B);
+		const double OutDotProduct = ((InA.X * InB.X) + (InA.Y * InB.Y) + (InA.Z * InB.Z) + (InA.W * InB.W));
+		SetValue(Context, OutDotProduct, &DotProduct);
+	}
+}
+
+//-----------------------------------------------------------
+
+FDataflowVectorLengthNode::FDataflowVectorLengthNode(const Dataflow::FNodeParameters& InParam, FGuid InGuid)
+	: FDataflowNode(InParam, InGuid)
+{
+	RegisterInputConnection(&V);
+	RegisterOutputConnection(&Length);
+	// set default type to be double
+	SetOutputConcreteType<double>(&Length);
+}
+
+void FDataflowVectorLengthNode::Evaluate(Dataflow::FContext& Context, const FDataflowOutput* Out) const
+{
+	if (Out->IsA(&Length))
+	{
+		const FVector4 InV = GetValue(Context, &V);
+		const double OutLength = FMath::Sqrt((InV.X * InV.X) + (InV.Y * InV.Y) + (InV.Z * InV.Z) + (InV.W * InV.W));
+		SetValue(Context, OutLength, &Length);
+	}
+}
+
+//-----------------------------------------------------------
+
+FDataflowVectorSquaredLengthNode::FDataflowVectorSquaredLengthNode(const Dataflow::FNodeParameters& InParam, FGuid InGuid)
+	: FDataflowNode(InParam, InGuid)
+{
+	RegisterInputConnection(&V);
+	RegisterOutputConnection(&SquaredLength);
+	// set default type to be double
+	SetOutputConcreteType<double>(&SquaredLength);
+}
+
+void FDataflowVectorSquaredLengthNode::Evaluate(Dataflow::FContext& Context, const FDataflowOutput* Out) const
+{
+	if (Out->IsA(&SquaredLength))
+	{
+		const FVector4 InV = GetValue(Context, &V);
+		const double OutSquaredLength = ((InV.X * InV.X) + (InV.Y * InV.Y) + (InV.Z * InV.Z) + (InV.W * InV.W));
+		SetValue(Context, OutSquaredLength, &SquaredLength);
+	}
+}
+
+//-----------------------------------------------------------
+
+FDataflowVectorDistanceNode::FDataflowVectorDistanceNode(const Dataflow::FNodeParameters& InParam, FGuid InGuid)
+	: FDataflowNode(InParam, InGuid)
+{
+	RegisterInputConnection(&A);
+	RegisterInputConnection(&B);
+	RegisterOutputConnection(&Distance);
+	// set default type to be double
+	SetOutputConcreteType<double>(&Distance);
+}
+
+void FDataflowVectorDistanceNode::Evaluate(Dataflow::FContext& Context, const FDataflowOutput* Out) const
+{
+	if (Out->IsA(&Distance))
+	{
+		const FVector4 InA = GetValue(Context, &A);
+		const FVector4 InB = GetValue(Context, &B);
+		const FVector4 Delta = (InB - InA);
+		const double OutDistance = FMath::Sqrt((Delta.X * Delta.X) + (Delta.Y * Delta.Y) + (Delta.Z * Delta.Z) + (Delta.W * Delta.W));
+		SetValue(Context, OutDistance, &Distance);
+	}
+}
+
+//-----------------------------------------------------------
+
+FDataflowVectorCrossProductNode::FDataflowVectorCrossProductNode(const Dataflow::FNodeParameters& InParam, FGuid InGuid)
+	: FDataflowNode(InParam, InGuid)
+{
+	RegisterInputConnection(&A);
+	RegisterInputConnection(&B);
+	RegisterOutputConnection(&CrossProduct);
+	// set default type to be double
+	SetOutputConcreteType<FVector>(&CrossProduct);
+}
+
+void FDataflowVectorCrossProductNode::Evaluate(Dataflow::FContext& Context, const FDataflowOutput* Out) const
+{
+	if (Out->IsA(&CrossProduct))
+	{
+		// convert to FVector from FVector4
+		const FVector InA(GetValue(Context, &A));
+		const FVector InB(GetValue(Context, &B));
+		const FVector OutCrossProduct = FVector::CrossProduct(InA, InB);
+		SetValue(Context, OutCrossProduct, &CrossProduct);
+	}
+}
+
+//-----------------------------------------------------------
+
+FDataflowVectorScaleNode::FDataflowVectorScaleNode(const Dataflow::FNodeParameters& InParam, FGuid InGuid)
+	: FDataflowNode(InParam, InGuid)
+{
+	Scale.Value = 1.0; // set default
+	RegisterInputConnection(&V);
+	RegisterInputConnection(&Scale);
+	RegisterOutputConnection(&Scaled);
+}
+
+void FDataflowVectorScaleNode::Evaluate(Dataflow::FContext& Context, const FDataflowOutput* Out) const
+{
+	if (Out->IsA(&Scaled))
+	{
+		const FVector4 InV = GetValue(Context, &V);
+		const double InScale = GetValue(Context, &Scale);
+		const FVector4 OutResult = (InV * InScale);
+		SetValue(Context, OutResult, &Scaled);
+	}
+}
+
+bool FDataflowVectorScaleNode::OnInputTypeChanged(const FDataflowInput* Input)
+{
+	if (Input->IsA(&V))
+	{
+		return SetOutputConcreteType(&Scaled, Input->GetType());
+	}
+	return false;
+}
+
+bool FDataflowVectorScaleNode::OnOutputTypeChanged(const FDataflowOutput* Output)
+{
+	if (Output->IsA(&Scaled))
+	{
+		return SetInputConcreteType(&V, Output->GetType());
+	}
+	return false;
+}
+
+//-----------------------------------------------------------
+
+FDataflowVectorNormalize::FDataflowVectorNormalize(const Dataflow::FNodeParameters& InParam, FGuid InGuid)
+	: FDataflowNode(InParam, InGuid)
+{
+	RegisterInputConnection(&V);
+	RegisterOutputConnection(&Normalized);
+}
+
+void FDataflowVectorNormalize::Evaluate(Dataflow::FContext& Context, const FDataflowOutput* Out) const
+{
+	if (Out->IsA(&Normalized))
+	{
+		const FVector4 InV = GetValue(Context, &V);
+		const FVector4 OutResult = InV.GetSafeNormal();
+		SetValue(Context, OutResult, &Normalized);
+	}
+}
+
+bool FDataflowVectorNormalize::OnInputTypeChanged(const FDataflowInput* Input)
+{
+	if (Input->IsA(&V))
+	{
+		return SetOutputConcreteType(&Normalized, Input->GetType());
+	}
+	return false;
+}
+
+bool FDataflowVectorNormalize::OnOutputTypeChanged(const FDataflowOutput* Output)
+{
+	if (Output->IsA(&Normalized))
+	{
+		return SetInputConcreteType(&V, Output->GetType());
+	}
+	return false;
+}
diff --git a/Engine/Source/Runtime/Experimental/Dataflow/Core/Public/Dataflow/DataflowAnyType.h b/Engine/Source/Runtime/Experimental/Dataflow/Core/Public/Dataflow/DataflowAnyType.h
--- a/Engine/Source/Runtime/Experimental/Dataflow/Core/Public/Dataflow/DataflowAnyType.h
+++ b/Engine/Source/Runtime/Experimental/Dataflow/Core/Public/Dataflow/DataflowAnyType.h
@@ -7,6 +7,7 @@
 
 #include "DataflowAnyType.generated.h"
 
+/** Any supported type */
 USTRUCT()
 struct FDataflowAnyType
 {
@@ -17,6 +18,7 @@ struct FDataflowAnyType
 	DATAFLOWCORE_API static const FName TypeName;
 };
 
+/** Any supported type */
 USTRUCT()
 struct FDataflowAllTypes : public FDataflowAnyType
 {
@@ -25,6 +27,10 @@ struct FDataflowAllTypes : public FDataflowAnyType
 	GENERATED_USTRUCT_BODY()
 };
 
+/**
+* Numeric types
+* (double, float, int64, uint64, int32, uint32, int16, uint16, int8, uint8)"
+*/
 USTRUCT()
 struct FDataflowNumericTypes : public FDataflowAnyType
 {
@@ -37,6 +43,23 @@ struct FDataflowNumericTypes : public FDataflowAnyType
 	double Value = 0.0;
 };
 
+/**
+* Vector types
+* (2D, 3D and 4D vector, single and double precision)
+*/
+USTRUCT()
+struct FDataflowVectorTypes : public FDataflowAnyType
+{
+	using FPolicyType = FDataflowVectorTypePolicy;
+	using FStorageType = FVector4;
+
+	GENERATED_USTRUCT_BODY()
+
+	UPROPERTY(EditAnywhere, Category = Value)
+	FVector4 Value = FVector4(0, 0, 0, 0);
+};
+
+/** String types (FString or FName) */
 USTRUCT()
 struct FDataflowStringTypes : public FDataflowAnyType
 {
@@ -49,6 +72,10 @@ struct FDataflowStringTypes : public FDataflowAnyType
 	FString Value;
 };
 
+/**
+* String convertible types
+* (String types, Numeric types, Vector types and Booleans)
+*/
 USTRUCT()
 struct FDataflowStringConvertibleTypes : public FDataflowAnyType
 {
@@ -61,7 +88,7 @@ struct FDataflowStringConvertibleTypes : public FDataflowAnyType
 	FString Value;
 };
 
-
+/** UObject types */
 USTRUCT()
 struct FDataflowUObjectConvertibleTypes : public FDataflowAnyType
 {
diff --git a/Engine/Source/Runtime/Experimental/Dataflow/Core/Public/Dataflow/DataflowMathNodes.h b/Engine/Source/Runtime/Experimental/Dataflow/Core/Public/Dataflow/DataflowMathNodes.h
--- a/Engine/Source/Runtime/Experimental/Dataflow/Core/Public/Dataflow/DataflowMathNodes.h
+++ b/Engine/Source/Runtime/Experimental/Dataflow/Core/Public/Dataflow/DataflowMathNodes.h
@@ -10,6 +10,12 @@
 
 #include "DataflowMathNodes.generated.h"
 
+//--------------------------------------------------------------------------
+//
+// Trigonometric nodes
+//
+//--------------------------------------------------------------------------
+
 #define DATAFLOW_MATH_NODES_CATEGORY "Math|Scalar"
 
 /** One input operators base class */
@@ -365,16 +371,187 @@ public:
 	virtual double ComputeResult(Dataflow::FContext& Context, double InA) const override;
 };
 
-namespace Dataflow
+/** One minus (1 - A) */
+USTRUCT()
+struct FDataflowMathOneMinusNode : public FDataflowMathOneInputOperatorNode
 {
-	void RegisterDataflowMathNodes();
-}
+	GENERATED_USTRUCT_BODY()
+	DATAFLOW_NODE_DEFINE_INTERNAL(FDataflowMathOneMinusNode, "OneMinus", DATAFLOW_MATH_NODES_CATEGORY, "")
+
+public:
+	FDataflowMathOneMinusNode(const Dataflow::FNodeParameters& InParam, FGuid InGuid = FGuid::NewGuid());
+	virtual double ComputeResult(Dataflow::FContext& Context, double InA) const override;
+};
+
+UENUM(BlueprintType)
+enum class EDataflowMathConstantsEnum : uint8
+{
+	Dataflow_Math_Constants_Pi			UMETA(DisplayName = "Pi"),
+	Dataflow_Math_Constants_HalfPi		UMETA(DisplayName = "HalfPi"),
+	Dataflow_Math_Constants_TwoPi		UMETA(DisplayName = "TwoPi"),
+	Dataflow_Math_Constants_FourPi		UMETA(DisplayName = "FourPi"),
+	Dataflow_Math_Constants_InvPi		UMETA(DisplayName = "InvPi"),
+	Dataflow_Math_Constants_InvTwoPi	UMETA(DisplayName = "InvTwoPi"),
+	Dataflow_Math_Constants_Sqrt2		UMETA(DisplayName = "Sqrt2"),
+	Dataflow_Math_Constants_InvSqrt2	UMETA(DisplayName = "InvSqrt2"),
+	Dataflow_Math_Constants_Sqrt3		UMETA(DisplayName = "Sqrt3"),
+	Dataflow_Math_Constants_InvSqrt3	UMETA(DisplayName = "InvSqrt3"),
+	Dataflow_Math_Constants_E			UMETA(DisplayName = "e"),
+	Dataflow_Math_Constants_Gamma		UMETA(DisplayName = "Gamma"),
+	Dataflow_Math_Constants_GoldenRatio	UMETA(DisplayName = "GoldenRatio"),
+	//~~~
+	//256th entry
+	Dataflow_Math_Constants_Max UMETA(Hidden)
+};
+
+/** Math constants ( see EDataflowMathConstantsEnum ) */
+USTRUCT()
+struct FDataflowMathConstantNode : public FDataflowNode
+{
+	GENERATED_USTRUCT_BODY()
+	DATAFLOW_NODE_DEFINE_INTERNAL(FDataflowMathConstantNode, "Constants", DATAFLOW_MATH_NODES_CATEGORY, "")
+
+	/** Math constant to output */
+	UPROPERTY(EditAnywhere, Category = "Constants");
+	EDataflowMathConstantsEnum Constant = EDataflowMathConstantsEnum::Dataflow_Math_Constants_Pi;
+
+	UPROPERTY(meta = (DataflowOutput))
+	FDataflowNumericTypes Result;
+
+public:
+	FDataflowMathConstantNode(const Dataflow::FNodeParameters& InParam, FGuid InGuid = FGuid::NewGuid());
+	virtual void Evaluate(Dataflow::FContext& Context, const FDataflowOutput* Out) const override;
+
+private:
+	double GetConstant() const;
+};
+
+//--------------------------------------------------------------------------
+//
+// Trigonometric nodes
+//
+//--------------------------------------------------------------------------
+
+#define DATAFLOW_MATH_TRIG_NODES_CATEGORY "Math|Trig"
+
+/** Sin(A) with A in radians  */
+USTRUCT()
+struct FDataflowMathSinNode : public FDataflowMathOneInputOperatorNode
+{
+	GENERATED_USTRUCT_BODY()
+	DATAFLOW_NODE_DEFINE_INTERNAL(FDataflowMathSinNode, "Sin", DATAFLOW_MATH_TRIG_NODES_CATEGORY, "")
+
+public:
+	FDataflowMathSinNode(const Dataflow::FNodeParameters& InParam, FGuid InGuid = FGuid::NewGuid());
+	virtual double ComputeResult(Dataflow::FContext& Context, double InA) const override;
+};
+
+/** Cos(A) with A in radians  */
+USTRUCT()
+struct FDataflowMathCosNode : public FDataflowMathOneInputOperatorNode
+{
+	GENERATED_USTRUCT_BODY()
+	DATAFLOW_NODE_DEFINE_INTERNAL(FDataflowMathCosNode, "Cos", DATAFLOW_MATH_TRIG_NODES_CATEGORY, "")
+
+public:
+	FDataflowMathCosNode(const Dataflow::FNodeParameters& InParam, FGuid InGuid = FGuid::NewGuid());
+	virtual double ComputeResult(Dataflow::FContext& Context, double InA) const override;
+};
+
+/** Tan(A) with A in radians  */
+USTRUCT()
+struct FDataflowMathTanNode : public FDataflowMathOneInputOperatorNode
+{
+	GENERATED_USTRUCT_BODY()
+	DATAFLOW_NODE_DEFINE_INTERNAL(FDataflowMathTanNode, "Tan", DATAFLOW_MATH_TRIG_NODES_CATEGORY, "")
+
+public:
+	FDataflowMathTanNode(const Dataflow::FNodeParameters& InParam, FGuid InGuid = FGuid::NewGuid());
+	virtual double ComputeResult(Dataflow::FContext& Context, double InA) const override;
+};
+
+/** ArcSin(A) returns a value in radians  */
+USTRUCT()
+struct FDataflowMathArcSinNode : public FDataflowMathOneInputOperatorNode
+{
+	GENERATED_USTRUCT_BODY()
+	DATAFLOW_NODE_DEFINE_INTERNAL(FDataflowMathArcSinNode, "ArcSin", DATAFLOW_MATH_TRIG_NODES_CATEGORY, "")
+
+public:
+	FDataflowMathArcSinNode(const Dataflow::FNodeParameters& InParam, FGuid InGuid = FGuid::NewGuid());
+	virtual double ComputeResult(Dataflow::FContext& Context, double InA) const override;
+};
+
+/** ArcCos(A) returns a value in radians  */
+USTRUCT()
+struct FDataflowMathArcCosNode : public FDataflowMathOneInputOperatorNode
+{
+	GENERATED_USTRUCT_BODY()
+	DATAFLOW_NODE_DEFINE_INTERNAL(FDataflowMathArcCosNode, "ArcCos", DATAFLOW_MATH_TRIG_NODES_CATEGORY, "")
+
+public:
+	FDataflowMathArcCosNode(const Dataflow::FNodeParameters& InParam, FGuid InGuid = FGuid::NewGuid());
+	virtual double ComputeResult(Dataflow::FContext& Context, double InA) const override;
+};
+
+/** ArcTan(A) returns a value in radians  */
+USTRUCT()
+struct FDataflowMathArcTanNode : public FDataflowMathOneInputOperatorNode
+{
+	GENERATED_USTRUCT_BODY()
+	DATAFLOW_NODE_DEFINE_INTERNAL(FDataflowMathArcTanNode, "ArcTan", DATAFLOW_MATH_TRIG_NODES_CATEGORY, "")
+
+public:
+	FDataflowMathArcTanNode(const Dataflow::FNodeParameters& InParam, FGuid InGuid = FGuid::NewGuid());
+	virtual double ComputeResult(Dataflow::FContext& Context, double InA) const override;
+};
+
+/** ArcTan2(A, B) returns a value in radians  */
+USTRUCT()
+struct FDataflowMathArcTan2Node : public FDataflowMathTwoInputsOperatorNode
+{
+	GENERATED_USTRUCT_BODY()
+	DATAFLOW_NODE_DEFINE_INTERNAL(FDataflowMathArcTan2Node, "ArcTan2", DATAFLOW_MATH_TRIG_NODES_CATEGORY, "")
+
+public:
+	FDataflowMathArcTan2Node(const Dataflow::FNodeParameters& InParam, FGuid InGuid = FGuid::NewGuid());
+	virtual double ComputeResult(Dataflow::FContext& Context, double InA, double InB) const override;
+};
+
+
+/** DegToRad(A) convert degrees to radians */
+USTRUCT()
+struct FDataflowMathDegToRadNode : public FDataflowMathOneInputOperatorNode
+{
+	GENERATED_USTRUCT_BODY()
+	DATAFLOW_NODE_DEFINE_INTERNAL(FDataflowMathDegToRadNode, "DegToRad", DATAFLOW_MATH_TRIG_NODES_CATEGORY, "")
+
+public:
+	FDataflowMathDegToRadNode(const Dataflow::FNodeParameters& InParam, FGuid InGuid = FGuid::NewGuid());
+	virtual double ComputeResult(Dataflow::FContext& Context, double InA) const override;
+};
+
+/** RadToDeg(A) convert radians to degrees */
+USTRUCT()
+struct FDataflowMathRadToDegNode : public FDataflowMathOneInputOperatorNode
+{
+	GENERATED_USTRUCT_BODY()
+	DATAFLOW_NODE_DEFINE_INTERNAL(FDataflowMathRadToDegNode, "RadToDeg", DATAFLOW_MATH_TRIG_NODES_CATEGORY, "")
+
+public:
+	FDataflowMathRadToDegNode(const Dataflow::FNodeParameters& InParam, FGuid InGuid = FGuid::NewGuid());
+	virtual double ComputeResult(Dataflow::FContext& Context, double InA) const override;
+};
+
+//--------------------------------------------------------------------------
+//
+// Trigonometric nodes
+//
+//--------------------------------------------------------------------------
+
 
 /*
 Nodes Left to be converted :
-	FMathConstantsDataflowNode);
-	FOneMinusDataflowNode);
-	
 	FFloatMathExpressionDataflowNode);
 	FMathExpressionDataflowNode);
 	
@@ -385,16 +562,7 @@ Nodes Left to be converted :
 	FLerpDataflowNode);
 	FWrapDataflowNode);
 
-	// trig
-	FSinDataflowNode);
-	FArcSinDataflowNode);
-	FCosDataflowNode);
-	FArcCosDataflowNode);
-	FTanDataflowNode);
-	FArcTanDataflowNode);
-	FArcTan2DataflowNode);
-	FRadiansToDegreesDataflowNode);
-	FDegreesToRadiansDataflowNode);
+
 
 	// vectors - requires Vector any type ? 
 	FNormalizeToRangeDataflowNode);
@@ -412,4 +580,9 @@ Nodes Left to be converted :
 	FRandomUnitVectorDataflowNode);
 	FRandomUnitVectorInConeDataflowNode);
 
-*/
\ No newline at end of file
+*/
+
+namespace Dataflow
+{
+	void RegisterDataflowMathNodes();
+}
diff --git a/Engine/Source/Runtime/Experimental/Dataflow/Core/Public/Dataflow/DataflowTypePolicy.h b/Engine/Source/Runtime/Experimental/Dataflow/Core/Public/Dataflow/DataflowTypePolicy.h
--- a/Engine/Source/Runtime/Experimental/Dataflow/Core/Public/Dataflow/DataflowTypePolicy.h
+++ b/Engine/Source/Runtime/Experimental/Dataflow/Core/Public/Dataflow/DataflowTypePolicy.h
@@ -5,6 +5,7 @@
 #include "CoreMinimal.h"
 #include "Misc/TVariant.h"
 #include "UObject/Object.h"
+#include "Math/MathFwd.h"
 #include <string>
 
 class UObject;
@@ -38,6 +39,12 @@ UE_DATAFLOW_POLICY_DECLARE_TYPENAME(double)
 UE_DATAFLOW_POLICY_DECLARE_TYPENAME(FName)
 UE_DATAFLOW_POLICY_DECLARE_TYPENAME(FString)
 UE_DATAFLOW_POLICY_DECLARE_TYPENAME(TObjectPtr<UObject>)
+UE_DATAFLOW_POLICY_DECLARE_TYPENAME(FVector2D)
+UE_DATAFLOW_POLICY_DECLARE_TYPENAME(FVector)
+UE_DATAFLOW_POLICY_DECLARE_TYPENAME(FVector4)
+UE_DATAFLOW_POLICY_DECLARE_TYPENAME(FVector2f)
+UE_DATAFLOW_POLICY_DECLARE_TYPENAME(FVector3f)
+UE_DATAFLOW_POLICY_DECLARE_TYPENAME(FVector4f)
 
 #undef UE_DATAFLOW_POLICY_DECLARE_TYPENAME
 
@@ -157,6 +164,12 @@ struct FDataflowNumericTypePolicy :
 {
 };
 
+struct FDataflowVectorTypePolicy:
+	public TDataflowMultiTypePolicy<FVector2D, FVector, FVector4, FVector2f, FVector3f, FVector4f>
+{
+};
+
+
 struct FDataflowStringTypePolicy :
 	public TDataflowMultiTypePolicy<FString, FName>
 {
@@ -166,6 +179,7 @@ struct FDataflowStringTypePolicy :
 * string comvertible types
 * - FString / Fname
 * - Numeric types ( see FDataflowNumericTypePolicy )
+* - Vector types ( see FDataflowVectorTypePolicy )
 * - bool
 */
 struct FDataflowStringConvertibleTypePolicy : IDataflowTypePolicy
@@ -179,6 +193,7 @@ struct FDataflowStringConvertibleTypePolicy : IDataflowTypePolicy
 	{
 		return FDataflowStringTypePolicy::SupportsTypeStatic(InType)
 			|| FDataflowNumericTypePolicy::SupportsTypeStatic(InType)
+			|| FDataflowVectorTypePolicy::SupportsTypeStatic(InType)
 			|| TDataflowSingleTypePolicy<bool>::SupportsTypeStatic(InType)
 			;
 	}
@@ -188,6 +203,7 @@ struct FDataflowStringConvertibleTypePolicy : IDataflowTypePolicy
 	{
 		return FDataflowStringTypePolicy::VisitPolicyByType(RequestedType, Visitor)
 			|| FDataflowNumericTypePolicy::VisitPolicyByType(RequestedType, Visitor)
+			|| FDataflowVectorTypePolicy::VisitPolicyByType(RequestedType, Visitor)
 			|| TDataflowSingleTypePolicy<bool>::VisitPolicyByType(RequestedType, Visitor)
 			;
 	}
@@ -267,6 +283,19 @@ struct FDataflowConverter
 	static void To(const T& From, TToType& To) { To = From; }
 };
 
+template<typename T>
+concept HasToStringMethod =
+	requires(T t) {
+		static_cast<FString>(t.ToString());
+};
+
+
+template<typename T>
+concept HasInitFromStringMethod =
+	requires(T t, const FString& s ) {
+		static_cast<bool>(t.InitFromString(s));
+};
+
 template <>
 struct FDataflowConverter<FString>
 {
@@ -285,6 +314,10 @@ struct FDataflowConverter<FString>
 		{
 			To = FString::SanitizeFloat(double(From), 0);
 		}
+		else if constexpr (HasToStringMethod<TFromType>)
+		{
+			To = From.ToString();
+		}
 		else
 		{
 			To = From;
@@ -304,10 +337,68 @@ struct FDataflowConverter<FString>
 		}
 		else if constexpr (std::is_convertible_v<double, TToType>)
 		{
-			double Result = {};
+			double Result = {0};
 			LexTryParseString(Result, *From);
 			To = Result;
 		}
+		else if constexpr (HasInitFromStringMethod<TToType>)
+		{
+			To.InitFromString(From);
+		}
+		else
+		{
+			To = From;
+		}
+	}
+};
+
+template <>
+struct FDataflowConverter<FVector4>
+{
+	template <typename TFromType>
+	static void From(const TFromType& From, FVector4& To)
+	{
+		if constexpr (std::is_same_v<TFromType, FVector2D> || std::is_same_v<TFromType, FVector2f>)
+		{
+			To = FVector4{ (double)From.X, (double)From.Y, 0, 0};
+		}
+		else if constexpr (std::is_same_v<TFromType, FVector> || std::is_same_v<TFromType, FVector3f>)
+		{
+			To = FVector4{ (double)From.X, (double)From.Y, (double)From.Z, 0 };
+		}
+		else if constexpr (std::is_same_v<TFromType, FVector4f>)
+		{
+			To = FVector4{ (double)From.X, (double)From.Y, (double)From.Z, (double)From.W };
+		}
+		else
+		{
+			To = From;
+		}
+	}
+
+	template <typename TToType>
+	static void To(const FVector4& From, TToType& To)
+	{
+		if constexpr (std::is_same_v<TToType, FVector2D>)
+		{
+			To = FVector2D{ From.X, From.Y };
+		}
+		else if constexpr (std::is_same_v<TToType, FVector2f>)
+		{
+			To = FVector2f{ (float)From.X, (float)From.Y };
+		}
+		else if constexpr (std::is_same_v<TToType, FVector>)
+		{
+			To = FVector{ From.X, From.Y, From.Z };
+		}
+		else if constexpr (std::is_same_v<TToType, FVector3f>)
+		{
+			To = FVector3f{ (float)From.X, (float)From.Y, (float)From.Z };
+		}
+		else if constexpr (std::is_same_v<TToType, FVector4f>)
+		{
+			To = FVector4f{ (float)From.X, (float)From.Y, (float)From.Z, (float)From.W };
+		}
 		else
 		{
 			To = From;
diff --git a/Engine/Source/Runtime/Experimental/Dataflow/Core/Public/Dataflow/DataflowVectorNodes.h b/Engine/Source/Runtime/Experimental/Dataflow/Core/Public/Dataflow/DataflowVectorNodes.h
new file mode 100644
--- /dev/null
+++ b/Engine/Source/Runtime/Experimental/Dataflow/Core/Public/Dataflow/DataflowVectorNodes.h
@@ -0,0 +1,374 @@
+// Copyright Epic Games, Inc. All Rights Reserved.
+
+#pragma once
+
+#include "Dataflow/DataflowNode.h"
+
+#include "DataflowVectorNodes.generated.h"
+
+#define DATAFLOW_MATH_VECTOR_NODES_CATEGORY "Math|Vectors"
+
+/** Make a 2D Vector */
+USTRUCT()
+struct FDataflowVectorMakeVec2Node : public FDataflowNode
+{
+	GENERATED_USTRUCT_BODY()
+	DATAFLOW_NODE_DEFINE_INTERNAL(FDataflowVectorMakeVec2Node, "MakeVector2", DATAFLOW_MATH_VECTOR_NODES_CATEGORY, "")
+
+public:
+	FDataflowVectorMakeVec2Node(const Dataflow::FNodeParameters& InParam, FGuid InGuid = FGuid::NewGuid());
+
+private:
+	/** X component */
+	UPROPERTY(EditAnywhere, Category = Operands, meta = (DataflowInput));
+	FDataflowNumericTypes X;
+
+	/** Y component */
+	UPROPERTY(EditAnywhere, Category = Operands, meta = (DataflowInput));
+	FDataflowNumericTypes Y;
+
+	/** 2D Vector {X, Y} */
+	UPROPERTY(meta = (DataflowOutput))
+	FDataflowVectorTypes Vector2D;
+
+	virtual void Evaluate(Dataflow::FContext& Context, const FDataflowOutput* Out) const override;
+};
+
+/** Make a 3D Vector */
+USTRUCT()
+struct FDataflowVectorMakeVec3Node : public FDataflowNode
+{
+	GENERATED_USTRUCT_BODY()
+	DATAFLOW_NODE_DEFINE_INTERNAL(FDataflowVectorMakeVec3Node, "MakeVector3", DATAFLOW_MATH_VECTOR_NODES_CATEGORY, "")
+
+public:
+	FDataflowVectorMakeVec3Node(const Dataflow::FNodeParameters& InParam, FGuid InGuid = FGuid::NewGuid());
+
+private:
+	/** X component */
+	UPROPERTY(EditAnywhere, Category = Operands, meta = (DataflowInput));
+	FDataflowNumericTypes X;
+
+	/** Y component */
+	UPROPERTY(EditAnywhere, Category = Operands, meta = (DataflowInput));
+	FDataflowNumericTypes Y;
+
+	/** Z component */
+	UPROPERTY(EditAnywhere, Category = Operands, meta = (DataflowInput));
+	FDataflowNumericTypes Z;
+
+	/** 3D Vector {X, Y, Z} */
+	UPROPERTY(meta = (DataflowOutput))
+	FDataflowVectorTypes Vector3D;
+
+	virtual void Evaluate(Dataflow::FContext& Context, const FDataflowOutput* Out) const override;
+};
+
+
+/** Make a 4D Vector */
+USTRUCT()
+struct FDataflowVectorMakeVec4Node : public FDataflowNode
+{
+	GENERATED_USTRUCT_BODY()
+	DATAFLOW_NODE_DEFINE_INTERNAL(FDataflowVectorMakeVec4Node, "MakeVector4", DATAFLOW_MATH_VECTOR_NODES_CATEGORY, "")
+
+public:
+	FDataflowVectorMakeVec4Node(const Dataflow::FNodeParameters& InParam, FGuid InGuid = FGuid::NewGuid());
+
+private:
+	/** X component */
+	UPROPERTY(EditAnywhere, Category = Operands, meta = (DataflowInput));
+	FDataflowNumericTypes X;
+
+	/** Y component */
+	UPROPERTY(EditAnywhere, Category = Operands, meta = (DataflowInput));
+	FDataflowNumericTypes Y;
+
+	/** Z component */
+	UPROPERTY(EditAnywhere, Category = Operands, meta = (DataflowInput));
+	FDataflowNumericTypes Z;
+
+	/** W component */
+	UPROPERTY(EditAnywhere, Category = Operands, meta = (DataflowInput));
+	FDataflowNumericTypes W;
+
+	/** 4D Vector {X, Y, Z, W} */
+	UPROPERTY(meta = (DataflowOutput))
+	FDataflowVectorTypes Vector4D;
+
+	virtual void Evaluate(Dataflow::FContext& Context, const FDataflowOutput* Out) const override;
+};
+
+
+/** 
+* Break a vector in 4 components
+* if the input vector is of a lower dimension than 4, the remaining components will be set to zero
+*/
+USTRUCT()
+struct FDataflowVectorBreakNode: public FDataflowNode
+{
+	GENERATED_USTRUCT_BODY()
+	DATAFLOW_NODE_DEFINE_INTERNAL(FDataflowVectorBreakNode, "BreakVector", DATAFLOW_MATH_VECTOR_NODES_CATEGORY, "")
+
+public:
+	FDataflowVectorBreakNode(const Dataflow::FNodeParameters& InParam, FGuid InGuid = FGuid::NewGuid());
+
+private:
+	/** Vector to break into components */
+	UPROPERTY(EditAnywhere, Category = Operands, meta = (DataflowInput));
+	FDataflowVectorTypes V;
+
+	/** X component */
+	UPROPERTY(meta = (DataflowOutput));
+	FDataflowNumericTypes X;
+
+	/** Y component */
+	UPROPERTY(meta = (DataflowOutput));
+	FDataflowNumericTypes Y;
+
+	/** Z component */
+	UPROPERTY(meta = (DataflowOutput));
+	FDataflowNumericTypes Z;
+
+	/** W component */
+	UPROPERTY(meta = (DataflowOutput));
+	FDataflowNumericTypes W;
+
+	virtual void Evaluate(Dataflow::FContext& Context, const FDataflowOutput* Out) const override;
+};
+
+/** Add two vectors component wise : V = (A + B)*/
+USTRUCT()
+struct FDataflowVectorAddNode : public FDataflowNode
+{
+	GENERATED_USTRUCT_BODY()
+	DATAFLOW_NODE_DEFINE_INTERNAL(FDataflowVectorAddNode, "AddVector", DATAFLOW_MATH_VECTOR_NODES_CATEGORY, "")
+
+public:
+	FDataflowVectorAddNode(const Dataflow::FNodeParameters& InParam, FGuid InGuid = FGuid::NewGuid());
+
+private:
+	/** A Vector operand */
+	UPROPERTY(EditAnywhere, Category = Operands, meta = (DataflowInput));
+	FDataflowVectorTypes A;
+
+	/** B Vector operand */
+	UPROPERTY(EditAnywhere, Category = Operands, meta = (DataflowInput));
+	FDataflowVectorTypes B;
+
+	/** Add result V=A+B */
+	UPROPERTY(meta = (DataflowOutput))
+	FDataflowVectorTypes V;
+
+	virtual void Evaluate(Dataflow::FContext& Context, const FDataflowOutput* Out) const override;
+};
+
+/** Subtract two vectors component wise: V = (A - B) */
+USTRUCT()
+struct FDataflowVectorSubtractNode : public FDataflowNode
+{
+	GENERATED_USTRUCT_BODY()
+	DATAFLOW_NODE_DEFINE_INTERNAL(FDataflowVectorSubtractNode, "SubtractVector", DATAFLOW_MATH_VECTOR_NODES_CATEGORY, "")
+
+public:
+	FDataflowVectorSubtractNode(const Dataflow::FNodeParameters& InParam, FGuid InGuid = FGuid::NewGuid());
+
+private:
+	/** A Vector operand */
+	UPROPERTY(EditAnywhere, Category = Operands, meta = (DataflowInput));
+	FDataflowVectorTypes A;
+
+	/** B Vector operand */
+	UPROPERTY(EditAnywhere, Category = Operands, meta = (DataflowInput));
+	FDataflowVectorTypes B;
+
+	/** Add result V=A-B */
+	UPROPERTY(meta = (DataflowOutput))
+	FDataflowVectorTypes V;
+
+	virtual void Evaluate(Dataflow::FContext& Context, const FDataflowOutput* Out) const override;
+};
+
+/** Compute the dot product between two vectors : DotProduct = A.B */
+USTRUCT()
+struct FDataflowVectorDotProductNode : public FDataflowNode
+{
+	GENERATED_USTRUCT_BODY()
+	DATAFLOW_NODE_DEFINE_INTERNAL(FDataflowVectorDotProductNode, "VectorDotProduct", DATAFLOW_MATH_VECTOR_NODES_CATEGORY, "")
+
+public:
+	FDataflowVectorDotProductNode(const Dataflow::FNodeParameters& InParam, FGuid InGuid = FGuid::NewGuid());
+
+private:
+	/** A Vector operand */
+	UPROPERTY(EditAnywhere, Category = Operands, meta = (DataflowInput));
+	FDataflowVectorTypes A;
+
+	/** B Vector operand */
+	UPROPERTY(EditAnywhere, Category = Operands, meta = (DataflowInput));
+	FDataflowVectorTypes B;
+
+	/** Resulting dot product : DotProduct=A.B  */
+	UPROPERTY(meta = (DataflowOutput))
+	FDataflowNumericTypes DotProduct;
+
+	virtual void Evaluate(Dataflow::FContext& Context, const FDataflowOutput* Out) const override;
+};
+
+/** Compute the Length of a vector : Length = |V| */
+USTRUCT()
+struct FDataflowVectorLengthNode : public FDataflowNode
+{
+	GENERATED_USTRUCT_BODY()
+	DATAFLOW_NODE_DEFINE_INTERNAL(FDataflowVectorLengthNode, "VectorLength", DATAFLOW_MATH_VECTOR_NODES_CATEGORY, "")
+
+public:
+	FDataflowVectorLengthNode(const Dataflow::FNodeParameters& InParam, FGuid InGuid = FGuid::NewGuid());
+
+private:
+	/** Vector to get length from */
+	UPROPERTY(EditAnywhere, Category = Operands, meta = (DataflowInput));
+	FDataflowVectorTypes V;
+
+	/** Length of the input vector : Length=|V| */
+	UPROPERTY(meta = (DataflowOutput))
+	FDataflowNumericTypes Length;
+
+	virtual void Evaluate(Dataflow::FContext& Context, const FDataflowOutput* Out) const override;
+};
+
+/** Compute the Squared length of a vector : Length = |V||V| */
+USTRUCT()
+struct FDataflowVectorSquaredLengthNode : public FDataflowNode
+{
+	GENERATED_USTRUCT_BODY()
+	DATAFLOW_NODE_DEFINE_INTERNAL(FDataflowVectorSquaredLengthNode, "VectorSquaredLength", DATAFLOW_MATH_VECTOR_NODES_CATEGORY, "")
+
+public:
+	FDataflowVectorSquaredLengthNode(const Dataflow::FNodeParameters& InParam, FGuid InGuid = FGuid::NewGuid());
+
+private:
+	/** Vector to get squared length from */
+	UPROPERTY(EditAnywhere, Category = Operands, meta = (DataflowInput));
+	FDataflowVectorTypes V;
+
+	/** Length of the input vector : SquaredLength = |V||V| */
+	UPROPERTY(meta = (DataflowOutput))
+	FDataflowNumericTypes SquaredLength;
+
+	virtual void Evaluate(Dataflow::FContext& Context, const FDataflowOutput* Out) const override;
+};
+
+/** Compute the distance between two vectors : Distance = |B-A| */
+USTRUCT()
+struct FDataflowVectorDistanceNode : public FDataflowNode
+{
+	GENERATED_USTRUCT_BODY()
+	DATAFLOW_NODE_DEFINE_INTERNAL(FDataflowVectorDistanceNode, "VectorDistance", DATAFLOW_MATH_VECTOR_NODES_CATEGORY, "")
+
+public:
+	FDataflowVectorDistanceNode(const Dataflow::FNodeParameters& InParam, FGuid InGuid = FGuid::NewGuid());
+
+private:
+	/** A Vector operand */
+	UPROPERTY(EditAnywhere, Category = Operands, meta = (DataflowInput));
+	FDataflowVectorTypes A;
+
+	/** B Vector operand */
+	UPROPERTY(EditAnywhere, Category = Operands, meta = (DataflowInput));
+	FDataflowVectorTypes B;
+
+	/** Distance between A and B : Distance=|B-A| */
+	UPROPERTY(meta = (DataflowOutput))
+	FDataflowNumericTypes Distance;
+
+	virtual void Evaluate(Dataflow::FContext& Context, const FDataflowOutput* Out) const override;
+};
+
+/** 
+* Compute the cross product of two vectors :  CrossProduct = B^A 
+* This node only operates in 3 dimensions, inputs will be converted to a 3D vector internally and result will be a vector with a zero W component
+*/
+USTRUCT()
+struct FDataflowVectorCrossProductNode : public FDataflowNode
+{
+	GENERATED_USTRUCT_BODY()
+	DATAFLOW_NODE_DEFINE_INTERNAL(FDataflowVectorCrossProductNode, "VectorCrossProduct", DATAFLOW_MATH_VECTOR_NODES_CATEGORY, "")
+
+public:
+	FDataflowVectorCrossProductNode(const Dataflow::FNodeParameters& InParam, FGuid InGuid = FGuid::NewGuid());
+
+private:
+	/** A Vector operand */
+	UPROPERTY(EditAnywhere, Category = Operands, meta = (DataflowInput));
+	FDataflowVectorTypes A;
+
+	/** B Vector operand */
+	UPROPERTY(EditAnywhere, Category = Operands, meta = (DataflowInput));
+	FDataflowVectorTypes B;
+
+	/** Resulting cross product of A and B : CrossProduct=B^A */
+	UPROPERTY(meta = (DataflowOutput))
+	FDataflowVectorTypes CrossProduct;
+
+	virtual void Evaluate(Dataflow::FContext& Context, const FDataflowOutput* Out) const override;
+};
+
+/** Scale a vector by a scalar : Scaled = (V * Scale) */
+USTRUCT()
+struct FDataflowVectorScaleNode : public FDataflowNode
+{
+	GENERATED_USTRUCT_BODY()
+	DATAFLOW_NODE_DEFINE_INTERNAL(FDataflowVectorScaleNode, "ScaleVector", DATAFLOW_MATH_VECTOR_NODES_CATEGORY, "")
+
+public:
+	FDataflowVectorScaleNode(const Dataflow::FNodeParameters& InParam, FGuid InGuid = FGuid::NewGuid());
+
+private:
+	/** Vector to scale */
+	UPROPERTY(EditAnywhere, Category = Operands, meta = (DataflowInput));
+	FDataflowVectorTypes V;
+
+	/** Scale factor */
+	UPROPERTY(EditAnywhere, Category = Operands, meta = (DataflowInput));
+	FDataflowNumericTypes Scale;
+
+	/** Scaled vector : Scaled=(V * Scale) */
+	UPROPERTY(meta = (DataflowOutput))
+	FDataflowVectorTypes Scaled;
+
+	virtual void Evaluate(Dataflow::FContext& Context, const FDataflowOutput* Out) const override;
+
+	virtual bool OnInputTypeChanged(const FDataflowInput* Input) override;
+	virtual bool OnOutputTypeChanged(const FDataflowOutput* Output) override;
+};
+
+/** Normalize a vector : Normalized = V/|V| */
+USTRUCT()
+struct FDataflowVectorNormalize : public FDataflowNode
+{
+	GENERATED_USTRUCT_BODY()
+	DATAFLOW_NODE_DEFINE_INTERNAL(FDataflowVectorNormalize, "NormalizeVector", DATAFLOW_MATH_VECTOR_NODES_CATEGORY, "")
+
+public:
+	FDataflowVectorNormalize(const Dataflow::FNodeParameters& InParam, FGuid InGuid = FGuid::NewGuid());
+
+private:
+	/** Vector to normalize */
+	UPROPERTY(EditAnywhere, Category = Operands, meta = (DataflowInput));
+	FDataflowVectorTypes V;
+
+	/** Normalized vector : Normalized=V/|V| */
+	UPROPERTY(meta = (DataflowOutput))
+	FDataflowVectorTypes Normalized;
+
+	virtual void Evaluate(Dataflow::FContext& Context, const FDataflowOutput* Out) const override;
+
+	virtual bool OnInputTypeChanged(const FDataflowInput* Input) override;
+	virtual bool OnOutputTypeChanged(const FDataflowOutput* Output) override;
+};
+
+namespace Dataflow
+{
+	void RegisterDataflowVectorNodes();
+}
