Chaos : fix clustering crash in AdvanceClustering
- this is a fallout from CL 35727536 to fix momentum
- Do not remove the cluster from the children map when processing the islands
- Also made the AdvanceClustering code more robust to missing particles in the map

diff --git a/Engine/Source/Runtime/Experimental/Chaos/Private/Chaos/PBDRigidClustering.cpp b/Engine/Source/Runtime/Experimental/Chaos/Private/Chaos/PBDRigidClustering.cpp
--- a/Engine/Source/Runtime/Experimental/Chaos/Private/Chaos/PBDRigidClustering.cpp
+++ b/Engine/Source/Runtime/Experimental/Chaos/Private/Chaos/PBDRigidClustering.cpp
@@ -1482,7 +1482,8 @@ namespace Chaos
 							// just got re-enabled in RemoveParticlesFromCluster.
 							for (FPBDRigidParticleHandle* Particle : Island)
 							{
-								DisableParticle(Particle);
+								// Particles may be re-enabled again at a later stage  so we do not want to remove them from the children map 
+								DisableParticle(Particle, /*bRemoveFromChildrenMap*/ false);
 							}
 						}
 					}
@@ -1689,28 +1690,33 @@ namespace Chaos
 
 			auto ProcessClusteredParticle = [&ParticlesToProcess, &bPotentialBreak, this](FPBDRigidClusteredParticleHandle* Particle)
 			{
-				TArray<FRigidHandle>& ParentToChildren = MChildren[Particle];
-
 				bool bAddParent = false;
-				for(FRigidHandle Child : ParentToChildren)
+				if (TArray<FRigidHandle>* ParentToChildren = MChildren.Find(Particle))
 				{
-					if(FClusterHandle ClusteredChild = Child->CastToClustered())
+					for (FRigidHandle Child : *ParentToChildren)
 					{
-						if(ClusteredChild->GetInternalStrains() <= 0.f)
-						{
-							bAddParent = true;
-							// #TODO remove need to set this here so we can early out as soon as we
-							// find one child that requires processing for breaks
-							ClusteredChild->CollisionImpulse() = FLT_MAX;
-							MCollisionImpulseArrayDirty = true;
-						}
-						else if(ClusteredChild->GetExternalStrain() > 0 || ClusteredChild->CollisionImpulse() > 0)
+						if (FClusterHandle ClusteredChild = Child->CastToClustered())
 						{
-							bAddParent = true;
-							bPotentialBreak = true;
+							if (ClusteredChild->GetInternalStrains() <= 0.f)
+							{
+								bAddParent = true;
+								// #TODO remove need to set this here so we can early out as soon as we
+								// find one child that requires processing for breaks
+								ClusteredChild->CollisionImpulse() = FLT_MAX;
+								MCollisionImpulseArrayDirty = true;
+							}
+							else if (ClusteredChild->GetExternalStrain() > 0 || ClusteredChild->CollisionImpulse() > 0)
+							{
+								bAddParent = true;
+								bPotentialBreak = true;
+							}
 						}
 					}
 				}
+				else
+				{
+					ensureMsgf(false, TEXT("Could not find an entry in MChildrenMap for a cluster to process, this is not expected. Skipping the processing"));
+				}
 
 				// Ensure we only add the parent once.
 				if(bAddParent)
@@ -2410,7 +2416,7 @@ namespace Chaos
 		ClusteredParticle->ClusterGroupIndex() = 0;
 	}
 
-	void FRigidClustering::DisableParticle(FPBDRigidParticleHandle* ParticleToDisable)
+	void FRigidClustering::DisableParticle(FPBDRigidParticleHandle* ParticleToDisable, bool bRemoveFromChildrenMap)
 	{
 		if (ParticleToDisable)
 		{
@@ -2423,7 +2429,10 @@ namespace Chaos
 			{
 				TopLevelClusterParents.Remove(ClusteredParticle);
 				TopLevelClusterParentsStrained.Remove(ClusteredParticle);
-				MChildren.Remove(ClusteredParticle);
+				if (bRemoveFromChildrenMap)
+				{
+					MChildren.Remove(ClusteredParticle);
+				}
 			}
 		}
 	}
diff --git a/Engine/Source/Runtime/Experimental/Chaos/Public/Chaos/PBDRigidClustering.h b/Engine/Source/Runtime/Experimental/Chaos/Public/Chaos/PBDRigidClustering.h
--- a/Engine/Source/Runtime/Experimental/Chaos/Public/Chaos/PBDRigidClustering.h
+++ b/Engine/Source/Runtime/Experimental/Chaos/Public/Chaos/PBDRigidClustering.h
@@ -497,7 +497,7 @@ public:
 	void ThrottleReleasedParticlesIfNecessary(TArray<FPBDRigidParticleHandle*>& Particles);
 
 	/** disable a particle and make sure the related internal structure are up to date */
-	void DisableParticle(FPBDRigidParticleHandle* ParticleToDisable);
+	void DisableParticle(FPBDRigidParticleHandle* ParticleToDisable, bool bRemoveFromChildrenMap = true);
 
  protected:
 
