6- Chaos rigid dataflow simulation support

#virtualized

diff --git a/Engine/Plugins/Experimental/GeometryCollectionPlugin/Source/GeometryCollectionEditor/Private/GeometryCollection/AssetDefinition_GeometryCollection.cpp b/Engine/Plugins/Experimental/GeometryCollectionPlugin/Source/GeometryCollectionEditor/Private/GeometryCollection/AssetDefinition_GeometryCollection.cpp
--- a/Engine/Plugins/Experimental/GeometryCollectionPlugin/Source/GeometryCollectionEditor/Private/GeometryCollection/AssetDefinition_GeometryCollection.cpp
+++ b/Engine/Plugins/Experimental/GeometryCollectionPlugin/Source/GeometryCollectionEditor/Private/GeometryCollection/AssetDefinition_GeometryCollection.cpp
@@ -75,8 +75,11 @@ EAssetCommandResult UAssetDefinition_GeometryCollection::OpenAssets(const FAsset
 			if (FDataflowEditorToolkit::HasDataflowAsset(GeometryCollection))
 			{
 				UAssetEditorSubsystem* const AssetEditorSubsystem = GEditor->GetEditorSubsystem<UAssetEditorSubsystem>();
+				
 				UDataflowEditor* const AssetEditor = NewObject<UDataflowEditor>(AssetEditorSubsystem, NAME_None, RF_Transient);
-				AssetEditor->Initialize({ GeometryCollection });
+				const TSubclassOf<AActor> ActorClass = StaticLoadClass(AActor::StaticClass(), nullptr,
+					TEXT("/GeometryCollectionPlugin/BP_GeometryCollectionPreview.BP_GeometryCollectionPreview_C"), nullptr, LOAD_None, nullptr);
+				AssetEditor->Initialize({ GeometryCollection }, ActorClass);
 				return EAssetCommandResult::Handled;
 			}
 
diff --git a/Engine/Source/Runtime/Engine/Private/PhysicsEngine/Experimental/PhysScene_Chaos.cpp b/Engine/Source/Runtime/Engine/Private/PhysicsEngine/Experimental/PhysScene_Chaos.cpp
--- a/Engine/Source/Runtime/Engine/Private/PhysicsEngine/Experimental/PhysScene_Chaos.cpp
+++ b/Engine/Source/Runtime/Engine/Private/PhysicsEngine/Experimental/PhysScene_Chaos.cpp
@@ -2021,12 +2021,16 @@ float FPhysScene_Chaos::OnStartFrame(float InDeltaTime)
 	SCOPE_CYCLE_COUNTER(STAT_Scene_StartFrame);
 
 #if WITH_EDITOR
-	if (IsOwningWorldEditor())
+	if (IsOwningWorldEditor() && GetSolver())
 	{
 		// Ensure editor solver is enabled
 		GetSolver()->SetIsPaused_External(false);
 
-		UseDeltaTime = 0.0f;
+		// Only pause the solver if the solver is not a standalone one
+		if(!GetSolver()->IsStandaloneSolver())
+		{
+			UseDeltaTime = 0.0f;
+		}
 	}
 #endif
 	ensure(DeferredCreatePhysicsStateComponents.Num() == 0);
diff --git a/Engine/Source/Runtime/Experimental/Chaos/Private/Chaos/Framework/PhysicsSolverBase.cpp b/Engine/Source/Runtime/Experimental/Chaos/Private/Chaos/Framework/PhysicsSolverBase.cpp
--- a/Engine/Source/Runtime/Experimental/Chaos/Private/Chaos/Framework/PhysicsSolverBase.cpp
+++ b/Engine/Source/Runtime/Experimental/Chaos/Private/Chaos/Framework/PhysicsSolverBase.cpp
@@ -486,6 +486,12 @@ namespace Chaos
 			ExternalSteps = 0;
 		}
 
+		// If standalone solver we are not responsible to spawn tasks
+		if(IsStandaloneSolver())
+		{
+			return {};
+		}
+
 		// Ensures we block on any tasks generated from previous frames
 		FGraphEventRef BlockingTasks = PendingTasks;
 
diff --git a/Engine/Source/Runtime/Experimental/Chaos/Private/GeometryCollectionProxyData.cpp b/Engine/Source/Runtime/Experimental/Chaos/Private/GeometryCollectionProxyData.cpp
--- a/Engine/Source/Runtime/Experimental/Chaos/Private/GeometryCollectionProxyData.cpp
+++ b/Engine/Source/Runtime/Experimental/Chaos/Private/GeometryCollectionProxyData.cpp
@@ -107,7 +107,7 @@ void FTransformDynamicCollection::SetHasParent(int32 Index, bool Value)
 int32 FTransformDynamicCollection::GetParent(int32 Index) const 
 {
 	check(RestCollectionShared);
-	return HasParent[Index] ? RestCollectionShared->Parent[Index] : INDEX_NONE;
+	return (HasParent.IsValidIndex(Index) && HasParent[Index])  ? RestCollectionShared->Parent[Index] : INDEX_NONE;
 }
 
 bool FTransformDynamicCollection::HasChildren(int32 Index) const
diff --git a/Engine/Source/Runtime/Experimental/Chaos/Private/PhysicsProxy/GeometryCollectionPhysicsProxy.cpp b/Engine/Source/Runtime/Experimental/Chaos/Private/PhysicsProxy/GeometryCollectionPhysicsProxy.cpp
--- a/Engine/Source/Runtime/Experimental/Chaos/Private/PhysicsProxy/GeometryCollectionPhysicsProxy.cpp
+++ b/Engine/Source/Runtime/Experimental/Chaos/Private/PhysicsProxy/GeometryCollectionPhysicsProxy.cpp
@@ -3932,7 +3932,7 @@ void FGeometryCollectionPhysicsProxy::BufferPhysicsResults_Internal(Chaos::FPBDR
 				FTransform3f NewParentSpaceTransform3f(ParentSpaceTransform);
 				if (!NewParentSpaceTransform3f.Equals(ParentSpaceTransform3f))
 				{
-					PhysicsThreadCollection.SetTransform(TransformGroupIndex, ParentSpaceTransform3f);
+					PhysicsThreadCollection.SetTransform(TransformGroupIndex, NewParentSpaceTransform3f);
 				}
 				IsObjectDynamic = true;
 			}
diff --git a/Engine/Source/Runtime/Experimental/Chaos/Public/Chaos/Framework/PhysicsSolverBase.h b/Engine/Source/Runtime/Experimental/Chaos/Public/Chaos/Framework/PhysicsSolverBase.h
--- a/Engine/Source/Runtime/Experimental/Chaos/Public/Chaos/Framework/PhysicsSolverBase.h
+++ b/Engine/Source/Runtime/Experimental/Chaos/Public/Chaos/Framework/PhysicsSolverBase.h
@@ -760,7 +760,17 @@ namespace Chaos
 			return NetworkPhysicsPredictionInterpLerp;
 		}
 
+		/** Get the standalone solver flag */
+		bool IsStandaloneSolver() const { return bIsStandaloneSolver;}
+		
+		/** Set the standalone solver flag */
+		void SetStandaloneSolver(const bool bStandaloneSolver) {bIsStandaloneSolver = bStandaloneSolver;}
+
 	protected:
+
+		/** Boolean to check if the solver is a standalone solver responsible to spawn its own task */
+		bool bIsStandaloneSolver = false;
+		
 		/** Mode that the results buffers should be set to (single, double, triple) */
 		EMultiBufferMode BufferMode;
 		
diff --git a/Engine/Source/Runtime/Experimental/ChaosSolverEngine/ChaosSolverEngine.Build.cs b/Engine/Source/Runtime/Experimental/ChaosSolverEngine/ChaosSolverEngine.Build.cs
--- a/Engine/Source/Runtime/Experimental/ChaosSolverEngine/ChaosSolverEngine.Build.cs
+++ b/Engine/Source/Runtime/Experimental/ChaosSolverEngine/ChaosSolverEngine.Build.cs
@@ -14,8 +14,11 @@ namespace UnrealBuildTool.Rules
                     "Engine",
                     "RenderCore",
                     "RHI",
-					"DeveloperSettings"
-                }
+					"DeveloperSettings", 
+					"DataflowCore",
+					"DataflowEngine",
+					"DataflowSimulation"
+				}
 				);
 
 			SetupModulePhysicsSupport(Target);
diff --git a/Engine/Source/Runtime/Experimental/ChaosSolverEngine/Private/Chaos/ChaosSolverActor.cpp b/Engine/Source/Runtime/Experimental/ChaosSolverEngine/Private/Chaos/ChaosSolverActor.cpp
--- a/Engine/Source/Runtime/Experimental/ChaosSolverEngine/Private/Chaos/ChaosSolverActor.cpp
+++ b/Engine/Source/Runtime/Experimental/ChaosSolverEngine/Private/Chaos/ChaosSolverActor.cpp
@@ -13,6 +13,8 @@
 #include "Engine/Texture2D.h"
 #include "UObject/FortniteMainBranchObjectVersion.h"
 #include "PhysicsProxy/SingleParticlePhysicsProxy.h"
+#include "Dataflow/DataflowSimulationManager.h"
+#include "PhysicsEngine/PhysicsSettings.h"
 
 #include UE_INLINE_GENERATED_CPP_BY_NAME(ChaosSolverActor)
 
@@ -176,6 +178,18 @@ private:
 static TUniquePtr<FChaosSolverActorConsoleObjects> ChaosSolverActorConsoleObjects;
 #endif  // #if CHAOS_DEBUG_SUBSTEP
 
+void FDataflowRigidSolverProxy::AdvanceSolverDatas(const float DeltaTime)
+{
+	if(Solver != nullptr)
+	{
+		for(auto* PushData : PushDatas)
+		{
+			Chaos::FSolverTasksPTOnly  SolverTask(*Solver, PushData);
+			SolverTask.AdvanceSolver();
+		}
+	}
+}
+
 AChaosSolverActor::AChaosSolverActor(const FObjectInitializer& ObjectInitializer)
 	: Super(ObjectInitializer)
 	, TimeStepMultiplier_DEPRECATED(1.f)
@@ -192,7 +206,7 @@ AChaosSolverActor::AChaosSolverActor(const FObjectInitializer& ObjectInitializer
 	, FloorHeight(0.f)
 	, ChaosDebugSubstepControl()
 	, PhysScene(nullptr)
-	, Solver(nullptr)
+	, RigidSolverProxy()
 	, Proxy(nullptr)
 {
 	if(!HasAnyFlags(RF_ClassDefaultObject))
@@ -205,7 +219,8 @@ AChaosSolverActor::AChaosSolverActor(const FObjectInitializer& ObjectInitializer
 								  , TEXT("Solver Actor Physics")
 #endif
 		));
-		Solver = PhysScene->GetSolver();
+		RigidSolverProxy.Solver = PhysScene->GetSolver();
+		
 		// Ticking setup for collision/breaking notifies
 		PrimaryActorTick.TickGroup = TG_PostPhysics;
 		PrimaryActorTick.bCanEverTick = true;
@@ -263,23 +278,41 @@ void AChaosSolverActor::PreInitializeComponents()
 	Super::PreInitializeComponents();
 }
 
+void AChaosSolverActor::MigrateSolver() const
+{
+	if(GetSolver())
+	{
+		if(FChaosSolversModule* Module = FChaosSolversModule::GetModule())
+		{
+			if(SimulationAsset.DataflowAsset)
+			{
+				Module->MigrateSolver(GetSolver(), this);
+				GetSolver()->SetStandaloneSolver(true);
+			}
+			else
+			{
+				Module->MigrateSolver(GetSolver(), GetWorld());
+				GetSolver()->SetStandaloneSolver(false);
+			}
+		}
+	}
+	// Make sure that the owning world is correct
+	PhysScene->SetOwningWorld(GetWorld());
+}
+
 void AChaosSolverActor::BeginPlay()
 {
 	Super::BeginPlay();
 
-	if(!Solver)
+	if(!RigidSolverProxy.Solver)
 	{
 		return;
 	}
 
-	// Make sure that the solver is registered in the right world
-	if(FChaosSolversModule* Module = FChaosSolversModule::GetModule())
-	{
-		Module->MigrateSolver(GetSolver(), GetWorld());
-	}
-
-	Solver->EnqueueCommandImmediate(
-		[InSolver = Solver, InProps = Properties]()
+	MigrateSolver();
+	
+	RigidSolverProxy.Solver->EnqueueCommandImmediate(
+		[InSolver = RigidSolverProxy.Solver, InProps = Properties]()
 		{
 			InSolver->ApplyConfig(InProps);
 		});
@@ -307,18 +340,18 @@ void AChaosSolverActor::EndPlay(const EEndPlayReason::Type ReasonEnd)
 {
 	Super::EndPlay(ReasonEnd);
 
-	if(!Solver)
+	if(!RigidSolverProxy.Solver)
 	{
 		return;
 	}
 
 	if(Proxy)
 	{
-		Solver->UnregisterObject(Proxy);
+		RigidSolverProxy.Solver->UnregisterObject(Proxy);
 		Proxy = nullptr;
 	}
 
-	Solver->EnqueueCommandImmediate([InSolver=Solver]()
+	RigidSolverProxy.Solver->EnqueueCommandImmediate([InSolver=RigidSolverProxy.Solver]()
 		{
 			// #TODO BG - We should really reset the solver here but the current reset function
 			// is really heavy handed and clears out absolutely everything. Ideally we want to keep
@@ -384,21 +417,31 @@ void AChaosSolverActor::PostRegisterAllComponents()
 {
 	Super::PostRegisterAllComponents();
 
+	MigrateSolver();
+
 	UWorld* const W = GetWorld(); 
 	if (W && !W->PhysicsScene_Chaos)
 	{
 		SetAsCurrentWorldSolver();
 	}
+
+	// Register the dataflow simulation interface
+	UE::Dataflow::RegisterSimulationInterface(this);
+}
+
+void AChaosSolverActor::PostUnregisterAllComponents()
+{
+	Super::PostUnregisterAllComponents();
+	
+	// Unregister the dataflow simulation interface
+	UE::Dataflow::UnregisterSimulationInterface(this);
 }
 
 void AChaosSolverActor::PostDuplicate(EDuplicateMode::Type DuplicateMode)
 {
 	Super::PostDuplicate(DuplicateMode);
 
-	if(FChaosSolversModule* Module = FChaosSolversModule::GetModule())
-	{
-		Module->MigrateSolver(GetSolver(), GetWorld());
-	}
+	MigrateSolver();
 }
 
 void AChaosSolverActor::MakeFloor()
@@ -419,14 +462,14 @@ void AChaosSolverActor::MakeFloor()
 		FilterData.Word3 = 0xFFFF;
 		FloorParticle->SetShapeSimData(0, FilterData);
 		Proxy = Chaos::FSingleParticlePhysicsProxy::Create(MoveTemp(FloorParticle));
-		Solver->RegisterObject(Proxy);
+		RigidSolverProxy.Solver->RegisterObject(Proxy);
 	}
 }
 
 void AChaosSolverActor::SetAsCurrentWorldSolver()
 {
 	UWorld* const W = GetWorld();
-	if (W)
+	if (W && (GetSolver() && !GetSolver()->IsStandaloneSolver()))
 	{
 		W->PhysicsScene_Chaos = PhysScene;
 	}
@@ -434,9 +477,67 @@ void AChaosSolverActor::SetAsCurrentWorldSolver()
 
 void AChaosSolverActor::SetSolverActive(bool bActive)
 {
-	if(Solver && PhysScene)
+	if(RigidSolverProxy.Solver && PhysScene)
 	{
-		Solver->SetIsPaused_External(!bActive);
+		RigidSolverProxy.Solver->SetIsPaused_External(!bActive);
+	}
+}
+
+void AChaosSolverActor::BuildSimulationProxy()
+{
+	RigidSolverProxy.Solver = PhysScene->GetSolver();
+	
+	// Ticking setup for collision/breaking notifies
+	PrimaryActorTick.TickGroup = TG_PostPhysics;
+	PrimaryActorTick.bCanEverTick = true;
+	PrimaryActorTick.bStartWithTickEnabled = true;
+}
+
+void AChaosSolverActor::ResetSimulationProxy()
+{
+	RigidSolverProxy.Solver = nullptr;
+}
+
+void AChaosSolverActor::WriteToSimulation(const float DeltaTime, const bool bAsyncTask)
+{
+	if (RigidSolverProxy.IsValid())
+	{
+		if(!bAsyncTask)
+		{
+			if(!Proxy && bHasFloor)
+			{
+				MakeFloor();
+			} 
+		
+			// Update gravity in case it changed
+			const FVector DefaultGravity( 0.f, 0.f, GetWorld()->GetGravityZ() );
+		
+			PhysScene->SetUpForFrame(&DefaultGravity, DeltaTime, UPhysicsSettings::Get()->MinPhysicsDeltaTime, UPhysicsSettings::Get()->MaxPhysicsDeltaTime,
+				UPhysicsSettings::Get()->MaxSubstepDeltaTime, UPhysicsSettings::Get()->MaxSubsteps, UPhysicsSettings::Get()->bSubstepping);
+		
+			PhysScene->StartFrame();
+		}
+		else
+		{
+			GetSolver()->AdvanceAndDispatch_External(DeltaTime);
+		}
+
+		RigidSolverProxy.PushDatas.Reset();
+		while(Chaos::FPushPhysicsData* PushData = RigidSolverProxy.Solver->GetMarshallingManager().StepInternalTime_External())
+		{
+			RigidSolverProxy.PushDatas.Add(PushData);
+		}
+	}
+}
+
+void AChaosSolverActor::ReadFromSimulation(const float DeltaTime, const bool bAsyncTask)
+{
+	if(!bAsyncTask)
+	{
+		if (RigidSolverProxy.IsValid())
+		{
+			PhysScene->EndFrame();
+		}
 	}
 }
 
@@ -445,11 +546,11 @@ void AChaosSolverActor::PostEditChangeProperty(struct FPropertyChangedEvent& Pro
 {
 	Super::PostEditChangeProperty(PropertyChangedEvent);
 
-	if(Solver && PropertyChangedEvent.Property)
+	if(RigidSolverProxy.Solver && PropertyChangedEvent.Property)
 	{
 		if(PropertyChangedEvent.MemberProperty->GetFName() == GET_MEMBER_NAME_CHECKED(AChaosSolverActor, Properties))
 		{
-			Solver->EnqueueCommandImmediate([InSolver = Solver, InConfig = Properties]()
+			RigidSolverProxy.Solver->EnqueueCommandImmediate([InSolver = RigidSolverProxy.Solver, InConfig = Properties]()
 			{
 				InSolver->ApplyConfig(InConfig);
 			});
@@ -458,7 +559,7 @@ void AChaosSolverActor::PostEditChangeProperty(struct FPropertyChangedEvent& Pro
 		{
 			if(Proxy)
 			{
-				Solver->UnregisterObject(Proxy);
+				RigidSolverProxy.Solver->UnregisterObject(Proxy);
 				Proxy = nullptr;
 			}
 
@@ -471,7 +572,29 @@ void AChaosSolverActor::PostEditChangeProperty(struct FPropertyChangedEvent& Pro
 				Proxy->GetGameThreadAPI().SetX(FVector(0.0f, 0.0f, FloorHeight));
 			}
 		}
+		else if(PropertyChangedEvent.Property->GetFName() == GET_MEMBER_NAME_CHECKED(AChaosSolverActor, SimulationAsset))
+		{
+			MigrateSolver();
+		}
+	}
+}
+
+bool AChaosSolverActor::CanEditChange(const FProperty* InProperty) const
+{
+	if (!Super::CanEditChange(InProperty))
+	{
+		return false;
 	}
+
+	const FName& Name = InProperty->GetFName();
+
+	if (Name == GET_MEMBER_NAME_CHECKED(ThisClass, SimulationAsset))
+	{
+		static const auto CVarEnableSimulationDataflow = IConsoleManager::Get().FindConsoleVariable(TEXT("p.Dataflow.EnableSimulation"));
+		return CVarEnableSimulationDataflow->GetBool();
+	}
+
+	return true;
 }
 
 #if TODO_REIMPLEMENT_SERIALIZATION_FOR_PERF_TEST
@@ -496,3 +619,5 @@ FAutoConsoleCommand SerializeForPerfTestCommand(TEXT("p.SerializeForPerfTest"),
 #endif
 
 
+
+
diff --git a/Engine/Source/Runtime/Experimental/ChaosSolverEngine/Public/Chaos/ChaosSolverActor.h b/Engine/Source/Runtime/Experimental/ChaosSolverEngine/Public/Chaos/ChaosSolverActor.h
--- a/Engine/Source/Runtime/Experimental/ChaosSolverEngine/Public/Chaos/ChaosSolverActor.h
+++ b/Engine/Source/Runtime/Experimental/ChaosSolverEngine/Public/Chaos/ChaosSolverActor.h
@@ -15,6 +15,7 @@
 #include "Physics/Experimental/PhysScene_Chaos.h"
 #include "ChaosSolverConfiguration.h"
 #include "SolverEventFilters.h"
+#include "Dataflow/Interfaces/DataflowPhysicsSolver.h"
 #include "PhysicsProxy/SingleParticlePhysicsProxyFwd.h"
 
 #include "ChaosSolverActor.generated.h"
@@ -71,8 +72,39 @@ struct FChaosDebugSubstepControl
 #endif
 };
 
+USTRUCT(BlueprintType)
+struct FDataflowRigidSolverProxy : public FDataflowPhysicsSolverProxy
+{
+	GENERATED_USTRUCT_BODY()
+
+	FDataflowRigidSolverProxy() : FDataflowPhysicsSolverProxy()
+	{}
+	virtual ~FDataflowRigidSolverProxy() override = default;
+
+	//~ Begin FPhysicsSolverInterface interface
+	virtual void AdvanceSolverDatas(const float DeltaTime) override;
+	virtual bool IsValid() const override { return Solver != nullptr;}
+	virtual const UScriptStruct* GetScriptStruct() const override
+	{
+		return StaticStruct();
+	}
+	//~ End FPhysicsSolverInterface interface
+
+	/** Chaos deformable solver that will be used in the component */
+	Chaos::FPhysicsSolver* Solver = nullptr;
+
+	/** List of push datas that will be used to advance the solver */
+	TArray<Chaos::FPushPhysicsData*> PushDatas = {};
+};
+
+template <>
+struct TStructOpsTypeTraits<FDataflowRigidSolverProxy> : public TStructOpsTypeTraitsBase2<FDataflowRigidSolverProxy>
+{
+	enum { WithCopy = false };
+};
+
 UCLASS(MinimalAPI)
-class AChaosSolverActor : public AActor
+class AChaosSolverActor : public AActor, public IDataflowPhysicsSolverInterface
 {
 	GENERATED_UCLASS_BODY()
 
@@ -130,6 +162,18 @@ public:
 	UFUNCTION(BlueprintCallable, Category = "ChaosPhysics")
 	CHAOSSOLVERENGINE_API virtual void SetSolverActive(bool bActive);
 
+	//~ Begin IDataflowPhysicsSolverInterface interface
+	virtual FString GetSimulationName() const override {return GetName();};
+	virtual FDataflowSimulationAsset& GetSimulationAsset() override {return SimulationAsset;};
+	virtual const FDataflowSimulationAsset& GetSimulationAsset() const override {return SimulationAsset;};
+	virtual FDataflowSimulationProxy* GetSimulationProxy() override {return &RigidSolverProxy;}
+	virtual const FDataflowSimulationProxy* GetSimulationProxy() const  override {return &RigidSolverProxy;}
+	virtual void BuildSimulationProxy() override;
+	virtual void ResetSimulationProxy() override;
+	virtual void WriteToSimulation(const float DeltaTime, const bool bAsyncTask) override;
+	virtual void ReadFromSimulation(const float DeltaTime, const bool bAsyncTask) override;
+	//~ End IDataflowPhysicsSolverInterface interface
+
 	/*
 	* Display icon in the editor
 	*/
@@ -139,17 +183,19 @@ public:
 	UChaosGameplayEventDispatcher* GetGameplayEventDispatcher() const { return GameplayEventDispatcherComponent; };
 
 	TSharedPtr<FPhysScene_Chaos> GetPhysicsScene() const { return PhysScene; }
-	Chaos::FPhysicsSolver* GetSolver() const { return Solver; }
+	Chaos::FPhysicsSolver* GetSolver() const { return RigidSolverProxy.Solver; }
 
 	CHAOSSOLVERENGINE_API virtual void PostRegisterAllComponents() override;
 	CHAOSSOLVERENGINE_API virtual void PreInitializeComponents() override;
+	CHAOSSOLVERENGINE_API virtual void PostUnregisterAllComponents() override;
 	
 	CHAOSSOLVERENGINE_API virtual void BeginPlay() override;
 	CHAOSSOLVERENGINE_API virtual void EndPlay(const EEndPlayReason::Type ReasonEnd) override;
-
+	
 	/** UObject interface */
 #if WITH_EDITOR
 	CHAOSSOLVERENGINE_API virtual void PostEditChangeProperty(struct FPropertyChangedEvent& PropertyChangedEvent) override;
+	CHAOSSOLVERENGINE_API virtual bool CanEditChange(const FProperty* InProperty) const override;
 #endif
 	CHAOSSOLVERENGINE_API void PostLoad() override;
 	CHAOSSOLVERENGINE_API void Serialize(FArchive& Ar) override;
@@ -158,11 +204,20 @@ public:
 
 private:
 
+	/* Solver dataflow asset used to advance in time */
+	UPROPERTY(EditAnywhere, Category = "Physics", meta=(EditConditionHides), AdvancedDisplay)
+	FDataflowSimulationAsset SimulationAsset;
+
 	/** If floor is enabled, make a particle to represent it */
 	CHAOSSOLVERENGINE_API void MakeFloor();
 
 	TSharedPtr<FPhysScene_Chaos> PhysScene;
-	Chaos::FPhysicsSolver* Solver;
+
+	/** Rigid solver proxy used in dataflow simulation */
+	FDataflowRigidSolverProxy RigidSolverProxy;
+
+	/** Migrate the solver onto the right owner (World vs Actor)*/
+	void MigrateSolver() const;
 
 	/** Component responsible for harvesting and triggering physics-related gameplay events (hits, breaks, etc) */
 	UPROPERTY()
diff --git a/Engine/Source/Runtime/Experimental/GeometryCollectionEngine/Private/GeometryCollection/GeometryCollectionComponent.cpp b/Engine/Source/Runtime/Experimental/GeometryCollectionEngine/Private/GeometryCollection/GeometryCollectionComponent.cpp
--- a/Engine/Source/Runtime/Experimental/GeometryCollectionEngine/Private/GeometryCollection/GeometryCollectionComponent.cpp
+++ b/Engine/Source/Runtime/Experimental/GeometryCollectionEngine/Private/GeometryCollection/GeometryCollectionComponent.cpp
@@ -1891,6 +1891,15 @@ void UGeometryCollectionComponent::PostEditChangeProperty(FPropertyChangedEvent&
 	{
 		OnGeometryCollectionPropertyChanged.Broadcast();
 	}
+
+	if (PropertyChangedEvent.Property && PropertyChangedEvent.Property->GetFName() == GET_MEMBER_NAME_CHECKED(UGeometryCollectionComponent, ChaosSolverActor))
+	{
+		// Reset the dynamic collection 
+		ResetDynamicCollection();
+
+		// Recreate the dynamic state since the dynamic collection has been reset
+		RecreatePhysicsState();
+	}
 }
 #endif
 
@@ -3699,14 +3708,7 @@ void UGeometryCollectionComponent::ResetDynamicCollection()
 {
 	bool bCreateDynamicCollection = true;
 #if WITH_EDITOR
-	bCreateDynamicCollection = false;
-	if (UWorld* World = GetWorld())
-	{
-		if(World->IsGameWorld() || GeometryCollectionCreatePhysicsStateInEditor)
-		{
-			bCreateDynamicCollection = true;
-		}
-	}
+	bCreateDynamicCollection = CanRunSimulationInEditor();
 #endif
 	if (bCreateDynamicCollection && RestCollection && RestCollection->GetGeometryCollection())
 	{
@@ -3769,7 +3771,7 @@ void UGeometryCollectionComponent::OnCreatePhysicsState()
 			RestCollectionMutable->CreateSimulationDataIfNeeded();
 		}
 #endif
-		const bool bValidWorld = GetWorld() && (GetWorld()->IsGameWorld() || GetWorld()->IsPreviewWorld() || GeometryCollectionCreatePhysicsStateInEditor);
+		const bool bValidWorld = CanRunSimulationInEditor() || (GetWorld() && GetWorld()->IsPreviewWorld());
 		const bool bValidCollection = DynamicCollection && DynamicCollection->GetNumTransforms() > 0;
 		if (bValidWorld && bValidCollection)
 		{
@@ -3991,7 +3993,7 @@ PRAGMA_ENABLE_DEPRECATION_WARNINGS
 	PhysicsProxy->SetPostPhysicsSyncCallback([this]() { OnPostPhysicsSync(); });
 	PhysicsProxy->SetPostParticlesCreatedCallback([this]() { OnPostCreateParticles(); });
 
-
+	
 	if (GetIsReplicated())
 	{
 		// using net mode and not local role because at this time in the initialization client and server both have an authority local role
@@ -5016,8 +5018,12 @@ void UGeometryCollectionComponent::SetRestCollection(const UGeometryCollection*
 		// otherwise this may cause mismatch issues and make the geometry collection look wrong 
 		RestTransforms.Reset();
 
+		// Reset the dynamic collection
 		ResetDynamicCollection();
 
+		// Rebuild the physics state since the dynamic collection has changed
+		RecreatePhysicsState();
+
 		if (!IsEmbeddedGeometryValid())
 		{
 			InitializeEmbeddedGeometry();
@@ -5927,6 +5933,27 @@ AChaosSolverActor* UGeometryCollectionComponent::GetPhysicsSolverActor() const
 	}
 }
 
+void UGeometryCollectionComponent::SetSolverActor(AChaosSolverActor* InSolverActor)
+{
+	ChaosSolverActor = InSolverActor;
+
+	// Reset the dynamic collection
+	ResetDynamicCollection();
+
+	// Rebuild the physics state since the proxy could now be registered to another actor
+	RecreatePhysicsState();
+}
+
+bool UGeometryCollectionComponent::CanRunSimulationInEditor() const
+{
+	if (UWorld* World = GetWorld())
+	{
+		Chaos::FPhysicsSolver* CurrSolver = GetSolver(*this);
+		return World->IsGameWorld() || GeometryCollectionCreatePhysicsStateInEditor || (CurrSolver && CurrSolver->IsStandaloneSolver());
+	}
+	return false;
+}
+
 #define GEOMETRY_COLLECTION_CHECK_FOR_NANS_IN_TRANSFORMS 0
 
 #if GEOMETRY_COLLECTION_CHECK_FOR_NANS_IN_TRANSFORMS
diff --git a/Engine/Source/Runtime/Experimental/GeometryCollectionEngine/Public/GeometryCollection/GeometryCollectionComponent.h b/Engine/Source/Runtime/Experimental/GeometryCollectionEngine/Public/GeometryCollection/GeometryCollectionComponent.h
--- a/Engine/Source/Runtime/Experimental/GeometryCollectionEngine/Public/GeometryCollection/GeometryCollectionComponent.h
+++ b/Engine/Source/Runtime/Experimental/GeometryCollectionEngine/Public/GeometryCollection/GeometryCollectionComponent.h
@@ -29,6 +29,7 @@
 #include "Chaos/ChaosSolverComponentTypes.h"
 #include "Chaos/PBDRigidsEvolutionFwd.h"
 #include "EngineDefines.h"
+#include "Dataflow/Interfaces/DataflowPhysicsObject.h"
 #include "Math/MathFwd.h"
 
 #include "GeometryCollectionComponent.generated.h"
@@ -683,9 +684,11 @@ public:
 	//~ End UMeshComponent Interface.
 
 	/** Chaos RBD Solver override. Will use the world's default solver actor if null. */
-	UPROPERTY(EditAnywhere, BlueprintReadWrite, Category = "ChaosPhysics", meta = (DisplayName = "Chaos Solver"))
+	UPROPERTY(EditAnywhere, BlueprintReadWrite, BlueprintGetter=GetSolverActor, BlueprintSetter=SetSolverActor, Category = "ChaosPhysics", meta = (DisplayName = "Chaos Solver"))
 	TObjectPtr<AChaosSolverActor> ChaosSolverActor;
 
+	
+
 	/**
 	* Get local bounds of the geometry collection
 	*/
@@ -1550,7 +1553,7 @@ protected:
 	/** Issue a field command for the physics thread */
 	GEOMETRYCOLLECTIONENGINE_API void DispatchFieldCommand(const FFieldSystemCommand& InCommand);
 
-	GEOMETRYCOLLECTIONENGINE_API Chaos::FPhysicsSolver* GetSolver(const UGeometryCollectionComponent& GeometryCollectionComponent);
+	GEOMETRYCOLLECTIONENGINE_API static Chaos::FPhysicsSolver* GetSolver(const UGeometryCollectionComponent& GeometryCollectionComponent);
 
 	UE_DEPRECATED(5.4, "CalculateLocalBounds is now Deprecated as it does not need to be called anymore, see ComponentSpaceBounds which replace LocalBounds")
 	GEOMETRYCOLLECTIONENGINE_API void CalculateLocalBounds() {};
@@ -1792,6 +1795,17 @@ public:
 
 private:
 
+	/** Return true if the simulation can be run in editor */
+	GEOMETRYCOLLECTIONENGINE_API bool CanRunSimulationInEditor() const;
+
+	/** BP internal function to get the solver actor */
+	UFUNCTION(BlueprintPure, BlueprintInternalUseOnly)
+	GEOMETRYCOLLECTIONENGINE_API AChaosSolverActor* GetSolverActor() const { return ChaosSolverActor; }
+
+	/** BP internal function to set the solver actor */
+	UFUNCTION(BlueprintCallable, BlueprintInternalUseOnly)
+	GEOMETRYCOLLECTIONENGINE_API void SetSolverActor(AChaosSolverActor* InSolverActor);
+
 	GEOMETRYCOLLECTIONENGINE_API void IncrementSleepTimer(float DeltaTime);
 	GEOMETRYCOLLECTIONENGINE_API void IncrementBreakTimer(float DeltaTime);
 	GEOMETRYCOLLECTIONENGINE_API bool CalculateInnerSphere(int32 TransformIndex, UE::Math::TSphere<double>& SphereOut) const;
diff --git a/Engine/Source/Runtime/PhysicsCore/Private/ChaosScene.cpp b/Engine/Source/Runtime/PhysicsCore/Private/ChaosScene.cpp
--- a/Engine/Source/Runtime/PhysicsCore/Private/ChaosScene.cpp
+++ b/Engine/Source/Runtime/PhysicsCore/Private/ChaosScene.cpp
@@ -370,23 +370,19 @@ void FChaosScene::StartFrame()
 		return;
 	}
 
-	const float UseDeltaTime = OnStartFrame(MDeltaTime);;
-
-	TArray<FPhysicsSolverBase*> SolverList;
-	ChaosModule->GetSolversMutable(Owner,SolverList);
-
-	if(FPhysicsSolver* Solver = GetSolver())
-	{
-		// Make sure our solver is in the list
-		SolverList.AddUnique(Solver);
-	}
-
+	const float UseDeltaTime = OnStartFrame(MDeltaTime);
 
+	TArray<FPhysicsSolverBase*> SolverList = GetPhysicsSolvers();
 	for(FPhysicsSolverBase* Solver : SolverList)
 	{
-		CompletionEvents.Add(Solver->AdvanceAndDispatch_External(UseDeltaTime));
+		if(FGraphEventRef SolverEvent = Solver->AdvanceAndDispatch_External(UseDeltaTime))
+		{
+			if(SolverEvent.IsValid())
+			{
+				CompletionEvents.Add(SolverEvent);
+			}
+		}
 	}
-
 }
 
 void FChaosScene::OnSyncBodies(Chaos::FPhysicsSolverBase* Solver)
@@ -481,6 +477,25 @@ void GetAABBTreeStats(Chaos::ISpatialAccelerationCollection<Chaos::FAcceleration
 	}
 }
 
+TArray<Chaos::FPhysicsSolverBase*> FChaosScene::GetPhysicsSolvers() const
+{
+	// Make a list of solvers to process. This is a list of all solvers registered to our world
+	// And our internal base scene solver.
+	TArray<Chaos::FPhysicsSolverBase*> SolverList;
+	if(const Chaos::FPhysicsSolver* Solver = GetSolver())
+	{
+		if(!Solver->IsStandaloneSolver())
+		{
+			// Get all the solvers with the same owner
+			ChaosModule->GetSolversMutable(Owner,SolverList);
+		}
+
+		// Make sure our solver is in the list
+		SolverList.AddUnique(GetSolver());
+	}
+	return SolverList;
+}
+
 void FChaosScene::EndFrame()
 {
 	using namespace Chaos;
@@ -532,13 +547,7 @@ void FChaosScene::EndFrame()
 
 	// Make a list of solvers to process. This is a list of all solvers registered to our world
 	// And our internal base scene solver.
-	TArray<FPhysicsSolverBase*> SolverList;
-	ChaosModule->GetSolversMutable(Owner,SolverList);
-
-	{
-		// Make sure our solver is in the list
-		SolverList.AddUnique(GetSolver());
-	}
+	TArray<FPhysicsSolverBase*> SolverList = GetPhysicsSolvers();
 
 	// Flip the buffers over to the game thread and sync
 	{
diff --git a/Engine/Source/Runtime/PhysicsCore/Public/Chaos/ChaosScene.h b/Engine/Source/Runtime/PhysicsCore/Public/Chaos/ChaosScene.h
--- a/Engine/Source/Runtime/PhysicsCore/Public/Chaos/ChaosScene.h
+++ b/Engine/Source/Runtime/PhysicsCore/Public/Chaos/ChaosScene.h
@@ -184,6 +184,8 @@ protected:
 	UObject* Owner;
 
 private:
+	/** Get all the physics solvers with the same owner */
+	PHYSICSCORE_API TArray<Chaos::FPhysicsSolverBase*> GetPhysicsSolvers() const;
 
 	PHYSICSCORE_API void SetGravity(const Chaos::FVec3& Acceleration);
 
