[ChaosVD] Fix for an issue where destroyed particles migh not be unloaded when scrubbing by a delta of +1 frame

This issue is fallout of a recent change where we added support to record particle data outside of a solver stage. When we scrub with a delta of +1 frame, we need to collapse all the in-between frames from the closest keframe to ensure we don't miss any data.

The change mentioned earlier added a new state for each recorded solver stage, to identify it as explicitly recorded or something we generated on the fly during load. The issue is that the system that collapses frame data was not taking that change into account, so generated collapsed frame were not taken into account as key frames.

This fix adds the missing flag to the generated solver stage

diff --git a/Engine/Source/Runtime/Experimental/ChaosVDData/Private/ChaosVDRecording.cpp b/Engine/Source/Runtime/Experimental/ChaosVDData/Private/ChaosVDRecording.cpp
--- a/Engine/Source/Runtime/Experimental/ChaosVDData/Private/ChaosVDRecording.cpp
+++ b/Engine/Source/Runtime/Experimental/ChaosVDData/Private/ChaosVDRecording.cpp
@@ -410,6 +410,9 @@ void FChaosVDRecording::CollapseSolverFramesRange_AssumesLocked(int32 SolverID,
 				FChaosVDStepData CollapsedStepData;
 				CollapsedStepData.StepName = TEXT("Auto Generated Stage");
 
+				// Although this is a generated on load stage, is based on explicit stage data, therefore it needs to be treated as such
+				EnumAddFlags(CollapsedStepData.StageFlags, EChaosVDSolverStageFlags::ExplicitStage);
+
 				ParticlesOnCurrentGeneratedKeyframe.GenerateValueArray(CollapsedStepData.RecordedParticlesData);
 
 				ParticlesOnCurrentGeneratedKeyframe.Reset();
