Iris - Fixed bug causing alignment to be stored as size.

diff --git a/Engine/Source/Runtime/Experimental/Iris/Core/Private/Iris/Serialization/NetSerializerArrayStorage.cpp b/Engine/Source/Runtime/Experimental/Iris/Core/Private/Iris/Serialization/NetSerializerArrayStorage.cpp
--- a/Engine/Source/Runtime/Experimental/Iris/Core/Private/Iris/Serialization/NetSerializerArrayStorage.cpp
+++ b/Engine/Source/Runtime/Experimental/Iris/Core/Private/Iris/Serialization/NetSerializerArrayStorage.cpp
@@ -34,7 +34,6 @@ void FNetSerializerAlignedStorage::AdjustSize(FNetSerializationContext& Context,
 	// If the allocation isn't properly aligned or if the allocation is too small we make a new allocation.
 	if ((InNum > StorageMaxCapacity) || !IsAligned(Data, InAlignment))
 	{
-		// Can't use realloc as the alignment of the original allocation isn't tracked.
 		void* NewData = Context.GetInternalContext()->Alloc(InNum, InAlignment);
 		FMemory::Memzero(NewData, InNum);
 		// Copy old data
@@ -58,7 +57,8 @@ void FNetSerializerAlignedStorage::AdjustSize(FNetSerializationContext& Context,
 			FMemory::Memzero(static_cast<void*>(Data + InNum), StorageNum - InNum);
 		}
 		StorageNum = InNum;
-		StorageNum = InAlignment;
+		// Use the requested alignment as the StorageAlignment. This allows Clone to allocate using the minimum required alignment.
+		StorageAlignment = InAlignment;
 	}
 }
 
