FCas
- Fix Verify not pruning iostore chunks that reference missing CAS blocks from lookup table.
- Fix CurrentBlock not being reset if the current block is purged which would lead to missing journal transactions and cache corruption.

diff --git a/Engine/Source/Runtime/Experimental/IoStore/OnDemand/Private/OnDemandInstallCache.cpp b/Engine/Source/Runtime/Experimental/IoStore/OnDemand/Private/OnDemandInstallCache.cpp
--- a/Engine/Source/Runtime/Experimental/IoStore/OnDemand/Private/OnDemandInstallCache.cpp
+++ b/Engine/Source/Runtime/Experimental/IoStore/OnDemand/Private/OnDemandInstallCache.cpp
@@ -415,15 +415,6 @@ FIoStatus FCas::Verify(TArray<FCasAddr>& OutAddrs)
 		const FString Filename = GetBlockFilename(BlockId);
 		UE_LOG(LogIoStoreOnDemand, Warning, TEXT("Missing CAS block '%s'"), *Filename);
 
-		for (auto It = Lookup.CreateIterator(); It; ++It)
-		{
-			if (It->Value.BlockId == BlockId)
-			{
-				OutAddrs.Add(It->Key);
-				It.RemoveCurrent();
-			}
-		}
-
 		LastAccess.Remove(BlockId);
 		BlockIt.RemoveCurrent();
 		Status = EIoErrorCode::NotFound;
@@ -448,6 +439,25 @@ FIoStatus FCas::Verify(TArray<FCasAddr>& OutAddrs)
 		}
 	}
 
+	TSet<FString> MissingReferencedBlocks;
+	for (auto It = Lookup.CreateIterator(); It; ++It)
+	{
+		if (!BlockIds.Contains(It->Value.BlockId))
+		{
+			MissingReferencedBlocks.Add(GetBlockFilename(It->Value.BlockId));
+			
+			OutAddrs.Add(It->Key);
+			It.RemoveCurrent();
+
+			Status = EIoErrorCode::NotFound;
+		}
+	}
+
+	for (const FString& Filename : MissingReferencedBlocks)
+	{
+		UE_LOG(LogIoStoreOnDemand, Warning, TEXT("Lookup references missing CAS block '%s'"), *Filename);
+	}
+
 	return Status; 
 }
 
@@ -1468,6 +1478,7 @@ FIoStatus FOnDemandInstallCache::Purge(const FCasBlockInfoMap& BlockInfo, const
 
 	for (const TPair<FCasBlockId, FCasBlockInfo>& Kv : BlockInfo)
 	{
+		const FCasBlockId BlockId = Kv.Key;
 		const FCasBlockInfo& Info = Kv.Value;
 		if (Info.RefCount > 0)
 		{
@@ -1477,18 +1488,23 @@ FIoStatus FOnDemandInstallCache::Purge(const FCasBlockInfoMap& BlockInfo, const
 		FCasJournal::FTransaction	Transaction = FCasJournal::Begin(GetJournalFilename());
 		TArray<FCasAddr>			RemovedChunks;
 
-		if (FIoStatus Status = Cas.DeleteBlock(Kv.Key, RemovedChunks); !Status.IsOk())
+		if (FIoStatus Status = Cas.DeleteBlock(BlockId, RemovedChunks); !Status.IsOk())
 		{
 			return Status;
 		}
 
+		if (Cas.CurrentBlock == BlockId)
+		{
+			Cas.CurrentBlock = FCasBlockId::Invalid;
+		}
+
 		OutTotalPurgedBytes += Info.FileSize;
 
 		for (const FCasAddr& Addr : RemovedChunks)
 		{
 			Transaction.ChunkLocation(FCasLocation::Invalid, Addr);
 		}
-		Transaction.BlockDeleted(Kv.Key);
+		Transaction.BlockDeleted(BlockId);
 
 		if (FIoStatus Status = FCasJournal::Commit(MoveTemp(Transaction)); !Status.IsOk())
 		{
