IAD analytics.

diff --git a/Engine/Source/Runtime/Experimental/IoStore/OnDemand/Public/IO/IoStoreOnDemand.h b/Engine/Source/Runtime/Experimental/IoStore/OnDemand/Public/IO/IoStoreOnDemand.h
--- a/Engine/Source/Runtime/Experimental/IoStore/OnDemand/Public/IO/IoStoreOnDemand.h
+++ b/Engine/Source/Runtime/Experimental/IoStore/OnDemand/Public/IO/IoStoreOnDemand.h
@@ -347,6 +347,32 @@ struct FOnDemandInstallProgress
 	uint64 TotalInstallSize = 0;
 	/** The size currently installed/downloaded (<= TotalInstallSize). */
 	uint64 CurrentInstallSize = 0;
+
+
+	FORCEINLINE uint64 GetTotalDownloadSize() const
+	{
+		return TotalInstallSize;
+	}
+
+	FORCEINLINE uint64 GetAlreadyDownloadedSize() const
+	{
+		return CurrentInstallSize;
+	}
+
+	FORCEINLINE float GetProgress() const 
+	{
+		return float(GetTotalDownloadSize()) / float(GetAlreadyDownloadedSize());
+	}
+
+	FORCEINLINE uint64 GetCachedSize() const 
+	{
+		return TotalContentSize - TotalInstallSize;
+	}
+
+	FORCEINLINE uint64 GetTotalSize() const
+	{
+		return TotalContentSize;
+	}
 };
 
 /** Install Progress callback. */
