Perform 3D distance calcluations in the Iris grid filter to match the same behavior as RepGraph.

diff --git a/Engine/Source/Runtime/Experimental/Iris/Core/Private/Iris/ReplicationSystem/Filtering/NetObjectGridFilter.cpp b/Engine/Source/Runtime/Experimental/Iris/Core/Private/Iris/ReplicationSystem/Filtering/NetObjectGridFilter.cpp
--- a/Engine/Source/Runtime/Experimental/Iris/Core/Private/Iris/ReplicationSystem/Filtering/NetObjectGridFilter.cpp
+++ b/Engine/Source/Runtime/Experimental/Iris/Core/Private/Iris/ReplicationSystem/Filtering/NetObjectGridFilter.cpp
@@ -209,7 +209,7 @@ void UNetObjectGridFilter::Filter(FNetObjectFilteringParams& Params)
 					for (const FReplicationView::FView& View : Params.View.Views)
 					{
 						const double DistSq = PerObjectInfo.GetCullDistanceSq();
-						const double ObjectToViewDistSq = FVector::DistSquared2D(PerObjectInfo.Position, View.Pos);
+						const double ObjectToViewDistSq = FVector::DistSquared(PerObjectInfo.Position, View.Pos);
 
 						if (ObjectToViewDistSq <= DistSq)
 						{
