Iris
* Moved some enums from ReplicateBridge.h to their own file in ReplicationBridgeTypes.h

diff --git a/Engine/Source/Runtime/Experimental/Iris/Core/Private/Iris/ReplicationSystem/ReplicationBridge.cpp b/Engine/Source/Runtime/Experimental/Iris/Core/Private/Iris/ReplicationSystem/ReplicationBridge.cpp
--- a/Engine/Source/Runtime/Experimental/Iris/Core/Private/Iris/ReplicationSystem/ReplicationBridge.cpp
+++ b/Engine/Source/Runtime/Experimental/Iris/Core/Private/Iris/ReplicationSystem/ReplicationBridge.cpp
@@ -1129,88 +1129,3 @@ FString UReplicationBridge::PrintObjectFromNetRefHandle(FNetRefHandle RefHandle)
 {
 	return NetRefHandleManager->PrintObjectFromNetRefHandle(RefHandle);
 }
-
-const TCHAR* LexToString(EReplicationBridgeDestroyInstanceReason Reason)
-{
-	switch (Reason)
-	{
-	case EReplicationBridgeDestroyInstanceReason::DoNotDestroy:
-	{
-		return TEXT("DoNotDestroy");
-	}
-	case EReplicationBridgeDestroyInstanceReason::TearOff:
-	{
-		return TEXT("TearOff");
-	}
-	case EReplicationBridgeDestroyInstanceReason::Destroy:
-	{
-		return TEXT("Destroy");
-	}
-	default:
-	{
-		return TEXT("[Invalid]");
-	}
-	}
-}
-
-const TCHAR* LexToString(EReplicationBridgeDestroyInstanceFlags DestroyFlags)
-{
-	switch (DestroyFlags)
-	{
-	case EReplicationBridgeDestroyInstanceFlags::None:
-	{
-		return TEXT("None");
-	}
-	case EReplicationBridgeDestroyInstanceFlags::AllowDestroyInstanceFromRemote:
-	{
-		return TEXT("AllowDestroyInstanceFromRemote");
-	}
-	default:
-	{
-		return TEXT("[Invalid]");
-	}
-	}
-}
-
-FString LexToString(EEndReplicationFlags EndReplicationFlags)
-{
-	
-	if (EndReplicationFlags == EEndReplicationFlags::None)
-	{
-		return TEXT("None");
-	}
-
-	FString Flags;
-
-	if (EnumHasAnyFlags(EndReplicationFlags, EEndReplicationFlags::Destroy))
-	{
-		Flags += TEXT("Destroy");
-	}
-	if (EnumHasAnyFlags(EndReplicationFlags, EEndReplicationFlags::TearOff))
-	{
-		if (!Flags.IsEmpty()) { Flags += TEXT(','); }
-		Flags += TEXT("TearOff");
-	}
-	if (EnumHasAnyFlags(EndReplicationFlags, EEndReplicationFlags::Flush))
-	{
-		if (!Flags.IsEmpty()) { Flags += TEXT(','); }
-		Flags += TEXT("Flush");
-	}
-	if (EnumHasAnyFlags(EndReplicationFlags, EEndReplicationFlags::DestroyNetHandle))
-	{
-		if (!Flags.IsEmpty()) { Flags += TEXT(','); };
-		Flags += TEXT("DestroyNetHandle");
-	}
-	if (EnumHasAnyFlags(EndReplicationFlags, EEndReplicationFlags::ClearNetPushId))
-	{
-		if (!Flags.IsEmpty()) { Flags += TEXT(','); };
-		Flags += TEXT("ClearNetPushId");
-	}
-	if (EnumHasAnyFlags(EndReplicationFlags, EEndReplicationFlags::SkipPendingEndReplicationValidation))
-	{
-		if (!Flags.IsEmpty()) { Flags += TEXT(','); };
-		Flags += TEXT("SkipPendingEndReplicationValidation");
-	}
-	
-	return Flags;
-}
\ No newline at end of file
diff --git a/Engine/Source/Runtime/Experimental/Iris/Core/Private/Iris/ReplicationSystem/ReplicationBridgeTypes.cpp b/Engine/Source/Runtime/Experimental/Iris/Core/Private/Iris/ReplicationSystem/ReplicationBridgeTypes.cpp
new file mode 100644
--- /dev/null
+++ b/Engine/Source/Runtime/Experimental/Iris/Core/Private/Iris/ReplicationSystem/ReplicationBridgeTypes.cpp
@@ -0,0 +1,89 @@
+// Copyright Epic Games, Inc. All Rights Reserved.
+
+#include "Iris/ReplicationSystem/ReplicationBridgeTypes.h"
+
+
+FString LexToString(EEndReplicationFlags EndReplicationFlags)
+{
+	
+	if (EndReplicationFlags == EEndReplicationFlags::None)
+	{
+		return TEXT("None");
+	}
+
+	FString Flags;
+
+	if (EnumHasAnyFlags(EndReplicationFlags, EEndReplicationFlags::Destroy))
+	{
+		Flags += TEXT("Destroy");
+	}
+	if (EnumHasAnyFlags(EndReplicationFlags, EEndReplicationFlags::TearOff))
+	{
+		if (!Flags.IsEmpty()) { Flags += TEXT(','); }
+		Flags += TEXT("TearOff");
+	}
+	if (EnumHasAnyFlags(EndReplicationFlags, EEndReplicationFlags::Flush))
+	{
+		if (!Flags.IsEmpty()) { Flags += TEXT(','); }
+		Flags += TEXT("Flush");
+	}
+	if (EnumHasAnyFlags(EndReplicationFlags, EEndReplicationFlags::DestroyNetHandle))
+	{
+		if (!Flags.IsEmpty()) { Flags += TEXT(','); };
+		Flags += TEXT("DestroyNetHandle");
+	}
+	if (EnumHasAnyFlags(EndReplicationFlags, EEndReplicationFlags::ClearNetPushId))
+	{
+		if (!Flags.IsEmpty()) { Flags += TEXT(','); };
+		Flags += TEXT("ClearNetPushId");
+	}
+	if (EnumHasAnyFlags(EndReplicationFlags, EEndReplicationFlags::SkipPendingEndReplicationValidation))
+	{
+		if (!Flags.IsEmpty()) { Flags += TEXT(','); };
+		Flags += TEXT("SkipPendingEndReplicationValidation");
+	}
+	
+	return Flags;
+}
+
+const TCHAR* LexToString(EReplicationBridgeDestroyInstanceReason Reason)
+{
+	switch (Reason)
+	{
+	case EReplicationBridgeDestroyInstanceReason::DoNotDestroy:
+	{
+		return TEXT("DoNotDestroy");
+	}
+	case EReplicationBridgeDestroyInstanceReason::TearOff:
+	{
+		return TEXT("TearOff");
+	}
+	case EReplicationBridgeDestroyInstanceReason::Destroy:
+	{
+		return TEXT("Destroy");
+	}
+	default:
+	{
+		return TEXT("[Invalid]");
+	}
+	}
+}
+
+const TCHAR* LexToString(EReplicationBridgeDestroyInstanceFlags DestroyFlags)
+{
+	switch (DestroyFlags)
+	{
+	case EReplicationBridgeDestroyInstanceFlags::None:
+	{
+		return TEXT("None");
+	}
+	case EReplicationBridgeDestroyInstanceFlags::AllowDestroyInstanceFromRemote:
+	{
+		return TEXT("AllowDestroyInstanceFromRemote");
+	}
+	default:
+	{
+		return TEXT("[Invalid]");
+	}
+	}
+}
\ No newline at end of file
diff --git a/Engine/Source/Runtime/Experimental/Iris/Core/Public/Iris/ReplicationSystem/ReplicationBridge.h b/Engine/Source/Runtime/Experimental/Iris/Core/Public/Iris/ReplicationSystem/ReplicationBridge.h
--- a/Engine/Source/Runtime/Experimental/Iris/Core/Public/Iris/ReplicationSystem/ReplicationBridge.h
+++ b/Engine/Source/Runtime/Experimental/Iris/Core/Public/Iris/ReplicationSystem/ReplicationBridge.h
@@ -8,6 +8,7 @@
 #include "Iris/Core/NetObjectReference.h"
 #include "Iris/ReplicationSystem/NetRefHandle.h"
 #include "Iris/ReplicationSystem/NetObjectGroupHandle.h"
+#include "Iris/ReplicationSystem/ReplicationBridgeTypes.h"
 #include "Misc/EnumClassFlags.h"
 #include "Net/Core/NetHandle/NetHandle.h"
 #include "UObject/ObjectKey.h"
@@ -65,58 +66,7 @@ struct FReplicationBridgeSerializationContext
 	bool bIsDestructionInfo;
 };
 
-enum class EEndReplicationFlags : uint32
-{
-	None								= 0U,
-	/** Destroy remote instance. Default for dynamic objects unless they have TearOff flag set. */
-	Destroy								= 1U,				
-	/** Stop replication object without destroying instance on the remote end. */
-	TearOff								= Destroy << 1U,
-	/** Complete replication of pending state to all clients before ending replication. */
-	Flush								= TearOff << 1U,
-	/** Destroy NetHandle if one is associated with the replicated object. This should only be done if the object should not be replicated by any other replication system. */
-	DestroyNetHandle					= Flush << 1U,
-	/** Clear net push ID to prevent this object and its subobjects from being marked as dirty in the networking system. This should only be done if the object should not be replicated by any other replication system. */
-	ClearNetPushId						= DestroyNetHandle << 1U,
-	/** Skip bPendingEndReplication Validation, In some cases we want to allow detaching instance from replicated object on clients, such as when shutting down */
-	SkipPendingEndReplicationValidation = ClearNetPushId << 1U,
-};
-ENUM_CLASS_FLAGS(EEndReplicationFlags);
-
-IRISCORE_API FString LexToString(EEndReplicationFlags EndReplicationFlags);
-
-enum class EReplicationBridgeCreateNetRefHandleResultFlags : unsigned
-{
-	None = 0U,
-	/** Whether the instance may be destroyed due to the remote peer requesting the object to be destroyed. If not then the object itself must not be destroyed. */
-	AllowDestroyInstanceFromRemote = 1U << 0U,
-	ShouldCallSubObjectCreatedFromReplication = AllowDestroyInstanceFromRemote << 1U,
-};
-ENUM_CLASS_FLAGS(EReplicationBridgeCreateNetRefHandleResultFlags);
-
-struct FReplicationBridgeCreateNetRefHandleResult
-{
-	UE::Net::FNetRefHandle NetRefHandle;
-	EReplicationBridgeCreateNetRefHandleResultFlags Flags = EReplicationBridgeCreateNetRefHandleResultFlags::None;
-};
-
-enum class EReplicationBridgeDestroyInstanceReason : unsigned
-{
-	DoNotDestroy,
-	TearOff,
-	Destroy,
-};
-IRISCORE_API const TCHAR* LexToString(EReplicationBridgeDestroyInstanceReason Reason);
-
-enum class EReplicationBridgeDestroyInstanceFlags : unsigned
-{
-	None = 0U,
-	/** Whether the instance may be destroyed when instructed from the remote peer. This flag applies when the destroy reason is TearOff and torn off actors are to be destroyed as well as regular Destroy. */
-	AllowDestroyInstanceFromRemote = 1U << 0U,
-};
-ENUM_CLASS_FLAGS(EReplicationBridgeDestroyInstanceFlags);
 
-IRISCORE_API const TCHAR* LexToString(EReplicationBridgeDestroyInstanceFlags DestroyFlags);
 
 UCLASS(Transient, MinimalAPI)
 class UReplicationBridge : public UObject
diff --git a/Engine/Source/Runtime/Experimental/Iris/Core/Public/Iris/ReplicationSystem/ReplicationBridgeTypes.h b/Engine/Source/Runtime/Experimental/Iris/Core/Public/Iris/ReplicationSystem/ReplicationBridgeTypes.h
new file mode 100644
--- /dev/null
+++ b/Engine/Source/Runtime/Experimental/Iris/Core/Public/Iris/ReplicationSystem/ReplicationBridgeTypes.h
@@ -0,0 +1,71 @@
+// Copyright Epic Games, Inc. All Rights Reserved.
+
+#pragma once
+
+#include "Iris/ReplicationSystem/NetRefHandle.h"
+
+#include "Misc/EnumClassFlags.h"
+
+//------------------------------------------------------------------------
+
+enum class EEndReplicationFlags : uint32
+{
+	None								= 0U,
+	/** Destroy remote instance. Default for dynamic objects unless they have TearOff flag set. */
+	Destroy								= 1U,				
+	/** Stop replication object without destroying instance on the remote end. */
+	TearOff								= Destroy << 1U,
+	/** Complete replication of pending state to all clients before ending replication. */
+	Flush								= TearOff << 1U,
+	/** Destroy NetHandle if one is associated with the replicated object. This should only be done if the object should not be replicated by any other replication system. */
+	DestroyNetHandle					= Flush << 1U,
+	/** Clear net push ID to prevent this object and its subobjects from being marked as dirty in the networking system. This should only be done if the object should not be replicated by any other replication system. */
+	ClearNetPushId						= DestroyNetHandle << 1U,
+	/** Skip bPendingEndReplication Validation, In some cases we want to allow detaching instance from replicated object on clients, such as when shutting down */
+	SkipPendingEndReplicationValidation = ClearNetPushId << 1U,
+};
+ENUM_CLASS_FLAGS(EEndReplicationFlags);
+
+IRISCORE_API FString LexToString(EEndReplicationFlags EndReplicationFlags);
+
+//------------------------------------------------------------------------
+
+enum class EReplicationBridgeCreateNetRefHandleResultFlags : unsigned
+{
+	None = 0U,
+	/** Whether the instance may be destroyed due to the remote peer requesting the object to be destroyed. If not then the object itself must not be destroyed. */
+	AllowDestroyInstanceFromRemote = 1U << 0U,
+	ShouldCallSubObjectCreatedFromReplication = AllowDestroyInstanceFromRemote << 1U,
+};
+ENUM_CLASS_FLAGS(EReplicationBridgeCreateNetRefHandleResultFlags);
+
+//------------------------------------------------------------------------
+
+struct FReplicationBridgeCreateNetRefHandleResult
+{
+	UE::Net::FNetRefHandle NetRefHandle;
+	EReplicationBridgeCreateNetRefHandleResultFlags Flags = EReplicationBridgeCreateNetRefHandleResultFlags::None;
+};
+
+//------------------------------------------------------------------------
+
+enum class EReplicationBridgeDestroyInstanceReason : unsigned
+{
+	DoNotDestroy,
+	TearOff,
+	Destroy,
+};
+IRISCORE_API const TCHAR* LexToString(EReplicationBridgeDestroyInstanceReason Reason);
+
+//------------------------------------------------------------------------
+
+enum class EReplicationBridgeDestroyInstanceFlags : unsigned
+{
+	None = 0U,
+	/** Whether the instance may be destroyed when instructed from the remote peer. This flag applies when the destroy reason is TearOff and torn off actors are to be destroyed as well as regular Destroy. */
+	AllowDestroyInstanceFromRemote = 1U << 0U,
+};
+ENUM_CLASS_FLAGS(EReplicationBridgeDestroyInstanceFlags);
+
+IRISCORE_API const TCHAR* LexToString(EReplicationBridgeDestroyInstanceFlags DestroyFlags);
+
