Replaced the experimental define UE_CHAOS_ASYNC_INITBODY_ENABLED with an experimental console variable p.Chaos.EnableAsyncInitBody (default value is false)
This will allow to enable/disable this experimental option per project or per platform without having to recompile.

diff --git a/Engine/Source/Runtime/Engine/Private/PhysicsEngine/BodyInstance.cpp b/Engine/Source/Runtime/Engine/Private/PhysicsEngine/BodyInstance.cpp
--- a/Engine/Source/Runtime/Engine/Private/PhysicsEngine/BodyInstance.cpp
+++ b/Engine/Source/Runtime/Engine/Private/PhysicsEngine/BodyInstance.cpp
@@ -30,6 +30,7 @@
 #include "Chaos/Capsule.h"
 #include "Chaos/Convex.h"
 #include "Chaos/TriangleMeshImplicitObject.h"
+#include "Chaos/AsyncInitBodyHelper.h"
 #include "Physics/Experimental/ChaosInterfaceUtils.h"
 
 #include "PhysicsEngine/TaperedCapsuleElem.h"
@@ -1395,12 +1396,9 @@ bool TInitBodiesHelperBase<InAllocatorType>::CreateShapesAndActors()
 	const int32 NumBodies = Bodies.Num();
 
 	const bool bIsInGameThread = IsInGameThread();
-	check(UE_CHAOS_ASYNC_INITBODY_ENABLED || bIsInGameThread);
-#if UE_CHAOS_ASYNC_INITBODY_ENABLED
-	// Current case where CreateShapesAndActors is not called from the GameThread is when experimental UE_CHAOS_ASYNC_INITBODY_ENABLED = 1
-	check(bIsInGameThread || (BodySetup->bCreatedPhysicsMeshes || BodySetup->bNeverNeedsCookedCollisionData));
+	// Current case where CreateShapesAndActors is not called from the GameThread is when p.Chaos.EnableAsyncInitBody is true
+	check(bIsInGameThread || (Chaos::CVars::bEnableAsyncInitBody && (BodySetup->bCreatedPhysicsMeshes || BodySetup->bNeverNeedsCookedCollisionData)));
 	if (bIsInGameThread)
-#endif
 	{
 		// Ensure we have the AggGeom inside the body setup so we can calculate the number of shapes
 		BodySetup->CreatePhysicsMeshes();
@@ -1543,7 +1541,7 @@ void TInitBodiesHelperBase<InAllocatorType>::InitBodies()
 	LLM_SCOPE(ELLMTag::ChaosBody);
 
 	const bool bIsInGameThread = IsInGameThread();
-	check(UE_CHAOS_ASYNC_INITBODY_ENABLED || bIsInGameThread);
+	check(Chaos::CVars::bEnableAsyncInitBody || bIsInGameThread);
 	if (bIsInGameThread ? CreateShapesAndActors() : true)
 	{
 		FPhysicsCommand::ExecuteWrite(PhysScene, [&]()
diff --git a/Engine/Source/Runtime/Engine/Private/PhysicsEngine/Experimental/PhysScene_Chaos.cpp b/Engine/Source/Runtime/Engine/Private/PhysicsEngine/Experimental/PhysScene_Chaos.cpp
--- a/Engine/Source/Runtime/Engine/Private/PhysicsEngine/Experimental/PhysScene_Chaos.cpp
+++ b/Engine/Source/Runtime/Engine/Private/PhysicsEngine/Experimental/PhysScene_Chaos.cpp
@@ -28,10 +28,8 @@
 #include "PhysicsProxy/StaticMeshPhysicsProxy.h"
 #include "Chaos/PendingSpatialData.h"
 #include "Chaos/PhysicsSolverBaseImpl.h"
+#include "Chaos/AsyncInitBodyHelper.h"
 #include "Misc/CoreMisc.h"
-#if UE_CHAOS_ASYNC_INITBODY_ENABLED
-#include "Misc/ScopeRWLock.h"
-#endif
 
 #if WITH_EDITOR
 #include "Editor.h"
diff --git a/Engine/Source/Runtime/Engine/Public/Physics/Experimental/PhysScene_Chaos.h b/Engine/Source/Runtime/Engine/Public/Physics/Experimental/PhysScene_Chaos.h
--- a/Engine/Source/Runtime/Engine/Public/Physics/Experimental/PhysScene_Chaos.h
+++ b/Engine/Source/Runtime/Engine/Public/Physics/Experimental/PhysScene_Chaos.h
@@ -14,12 +14,9 @@
 #include "Chaos/ChaosScene.h"
 #include "Chaos/ContactModification.h"
 #include "Chaos/Real.h"
+#include "Chaos/AsyncInitBodyHelper.h"
 #include "UObject/ObjectKey.h"
 
-#if UE_CHAOS_ASYNC_INITBODY_ENABLED
-#include "Misc/ScopeRWLock.h"
-#endif
-
 #ifndef CHAOS_WITH_PAUSABLE_SOLVER
 #define CHAOS_WITH_PAUSABLE_SOLVER 1
 #endif
@@ -451,10 +448,8 @@ private:
 
 	TSet<UPrimitiveComponent*> DeferredCreatePhysicsStateComponents;
 
-#if UE_CHAOS_ASYNC_INITBODY_ENABLED
 	// RWLock for for thread safe access to PhysicsProxyToComponentMap and ComponentToPhysicsProxyMap.
 	mutable FRWLock PhysicsProxyComponentMapsLock;
-#endif
 
 	//Body Instances
 	TUniquePtr<Chaos::TArrayCollectionArray<FBodyInstance*>> BodyInstances;
diff --git a/Engine/Source/Runtime/Experimental/Chaos/Private/Chaos/AsyncInitBodyHelper.cpp b/Engine/Source/Runtime/Experimental/Chaos/Private/Chaos/AsyncInitBodyHelper.cpp
new file mode 100644
--- /dev/null
+++ b/Engine/Source/Runtime/Experimental/Chaos/Private/Chaos/AsyncInitBodyHelper.cpp
@@ -0,0 +1,18 @@
+// Copyright Epic Games, Inc. All Rights Reserved.
+
+#include "Chaos/AsyncInitBodyHelper.h"
+#include "HAL/IConsoleManager.h"
+
+namespace Chaos
+{
+	namespace CVars
+	{
+		CHAOS_API bool bEnableAsyncInitBody = false;
+
+		FAutoConsoleVariableRef CVarEnableAsyncInitBody(
+		TEXT("p.Chaos.EnableAsyncInitBody"), 
+		bEnableAsyncInitBody,
+		TEXT("[Experimental] Allow body instances to be initialized outside of game thread (default is false)."),
+		ECVF_ReadOnly);
+	}
+}
\ No newline at end of file
diff --git a/Engine/Source/Runtime/Experimental/Chaos/Private/Chaos/ChaosMarshallingManager.cpp b/Engine/Source/Runtime/Experimental/Chaos/Private/Chaos/ChaosMarshallingManager.cpp
--- a/Engine/Source/Runtime/Experimental/Chaos/Private/Chaos/ChaosMarshallingManager.cpp
+++ b/Engine/Source/Runtime/Experimental/Chaos/Private/Chaos/ChaosMarshallingManager.cpp
@@ -2,9 +2,7 @@
 
 #include "Chaos/ChaosMarshallingManager.h"
 #include "Chaos/PullPhysicsDataImp.h"
-#if UE_CHAOS_ASYNC_INITBODY_ENABLED
-#include "Misc/ScopeRWLock.h"
-#endif
+#include "Chaos/AsyncInitBodyHelper.h"
 
 namespace Chaos
 {
@@ -51,7 +49,7 @@ void FChaosMarshallingManager::PreparePullData()
 
 void FChaosMarshallingManager::PrepareExternalQueue_External()
 {
-	// Here, we assume that MarshallingManagerLock is locked (when UE_CHAOS_ASYNC_INITBODY_ENABLED = 1)
+	// Here, we assume that MarshallingManagerLock is locked when p.Chaos.EnableAsyncInitBody is true (see FPBDRigidsSolver::PushPhysicsState)
 	if(!PushDataPool.Dequeue(ProducerData))
 	{
 		BackingBuffer.Add(MakeUnique<FPushPhysicsData>());
@@ -63,7 +61,7 @@ void FChaosMarshallingManager::PrepareExternalQueue_External()
 
 void FChaosMarshallingManager::Step_External(FReal ExternalDT, const int32 NumSteps, bool bInSolverSubstepped)
 {
-	// Here, we assume that MarshallingManagerLock is locked (when UE_CHAOS_ASYNC_INITBODY_ENABLED = 1)
+	// Here, we assume that MarshallingManagerLock is locked when p.Chaos.EnableAsyncInitBody is true (see FPBDRigidsSolver::PushPhysicsState)
 	ensure(NumSteps > 0);
 
 	FPushPhysicsData* FirstStepData = nullptr;
diff --git a/Engine/Source/Runtime/Experimental/Chaos/Private/Chaos/PhysicalMaterials.cpp b/Engine/Source/Runtime/Experimental/Chaos/Private/Chaos/PhysicalMaterials.cpp
--- a/Engine/Source/Runtime/Experimental/Chaos/Private/Chaos/PhysicalMaterials.cpp
+++ b/Engine/Source/Runtime/Experimental/Chaos/Private/Chaos/PhysicalMaterials.cpp
@@ -2,9 +2,7 @@
 
 #include "Chaos/PhysicalMaterials.h"
 #include "HAL/LowLevelMemTracker.h"
-#if UE_CHAOS_ASYNC_INITBODY_ENABLED
-#include "Misc/ScopeRWLock.h"
-#endif
+#include "Chaos/AsyncInitBodyHelper.h"
 
 namespace Chaos
 {
@@ -83,13 +81,13 @@ namespace Chaos
 	
 	void FPhysicalMaterialManager::UpdateMaterial(FMaterialHandle InHandle)
 	{
-		check(UE_CHAOS_ASYNC_INITBODY_ENABLED || IsInGameThread());
+		check(Chaos::CVars::bEnableAsyncInitBody || IsInGameThread());
 		OnMaterialUpdated.Broadcast(InHandle);
 	}
 
 	void FPhysicalMaterialManager::UpdateMaterialMask(FMaterialMaskHandle InHandle)
 	{
-		check(UE_CHAOS_ASYNC_INITBODY_ENABLED || IsInGameThread());
+		check(Chaos::CVars::bEnableAsyncInitBody || IsInGameThread());
 		OnMaterialMaskUpdated.Broadcast(InHandle);
 	}
 
@@ -117,7 +115,7 @@ namespace Chaos
 	{
 		LLM_SCOPE(ELLMTag::ChaosMaterial);
 
-		check(UE_CHAOS_ASYNC_INITBODY_ENABLED || IsInGameThread());
+		check(Chaos::CVars::bEnableAsyncInitBody || IsInGameThread());
 		FMaterialHandle OutHandle;
 		{
 			UE_CHAOS_ASYNC_INITBODY_WRITESCOPELOCK(MaterialsLock);
@@ -130,7 +128,7 @@ namespace Chaos
 
 	FMaterialMaskHandle FPhysicalMaterialManager::CreateMask()
 	{
-		check(UE_CHAOS_ASYNC_INITBODY_ENABLED || IsInGameThread());
+		check(Chaos::CVars::bEnableAsyncInitBody || IsInGameThread());
 		FMaterialMaskHandle OutHandle;
 		{
 			UE_CHAOS_ASYNC_INITBODY_WRITESCOPELOCK(MaterialMasksLock);
@@ -145,7 +143,7 @@ namespace Chaos
 	{
 		LLM_SCOPE(ELLMTag::ChaosMaterial);
 
-		check(UE_CHAOS_ASYNC_INITBODY_ENABLED || IsInGameThread());
+		check(Chaos::CVars::bEnableAsyncInitBody || IsInGameThread());
 		if(InHandle.InnerHandle.IsValid())
 		{
 			OnMaterialDestroyed.Broadcast(InHandle);
@@ -158,7 +156,7 @@ namespace Chaos
 
 	void FPhysicalMaterialManager::Destroy(FMaterialMaskHandle InHandle)
 	{
-		check(UE_CHAOS_ASYNC_INITBODY_ENABLED || IsInGameThread());
+		check(Chaos::CVars::bEnableAsyncInitBody || IsInGameThread());
 		if (InHandle.InnerHandle.IsValid())
 		{
 			OnMaterialMaskDestroyed.Broadcast(InHandle);
diff --git a/Engine/Source/Runtime/Experimental/Chaos/Private/PBDRigidsSolver.cpp b/Engine/Source/Runtime/Experimental/Chaos/Private/PBDRigidsSolver.cpp
--- a/Engine/Source/Runtime/Experimental/Chaos/Private/PBDRigidsSolver.cpp
+++ b/Engine/Source/Runtime/Experimental/Chaos/Private/PBDRigidsSolver.cpp
@@ -28,9 +28,7 @@
 #include "Chaos/PullPhysicsDataImp.h"
 #include "Chaos/PhysicsSolverBaseImpl.h"
 #include "Chaos/ConvexOptimizer.h"
-#if UE_CHAOS_ASYNC_INITBODY_ENABLED
-#include "Misc/ScopeRWLock.h"
-#endif
+#include "Chaos/AsyncInitBodyHelper.h"
 
 #include "ChaosDebugDraw/ChaosDDContext.h"
 #include "ChaosDebugDraw/ChaosDDScene.h"
@@ -749,7 +747,7 @@ namespace Chaos
 
 	void FPBDRigidsSolver::RegisterObject(FSingleParticlePhysicsProxy* Proxy)
 	{
-		UE_CHAOS_ASYNC_INITBODY_PHYSICSSCENE_WRITESCOPELOCK(GetExternalDataLock_External());
+		UE_CHAOS_ASYNC_INITBODY_WRITESCOPELOCK(GetExternalDataLock_External());
 		LLM_SCOPE(ELLMTag::ChaosBody);
 
 		UE_LOG(LogPBDRigidsSolver, Verbose, TEXT("FPBDRigidsSolver::RegisterObject()"));
@@ -804,7 +802,7 @@ namespace Chaos
 
 	void FPBDRigidsSolver::UnregisterObject(FSingleParticlePhysicsProxy* Proxy)
 	{
-		UE_CHAOS_ASYNC_INITBODY_PHYSICSSCENE_WRITESCOPELOCK(GetExternalDataLock_External());
+		UE_CHAOS_ASYNC_INITBODY_WRITESCOPELOCK(GetExternalDataLock_External());
 		UE_LOG(LogPBDRigidsSolver, Verbose, TEXT("FPBDRigidsSolver::UnregisterObject()"));
 
 		PullResultsManager->RemoveProxy_External(Proxy);
diff --git a/Engine/Source/Runtime/Experimental/Chaos/Public/Chaos/AsyncInitBodyHelper.h b/Engine/Source/Runtime/Experimental/Chaos/Public/Chaos/AsyncInitBodyHelper.h
new file mode 100644
--- /dev/null
+++ b/Engine/Source/Runtime/Experimental/Chaos/Public/Chaos/AsyncInitBodyHelper.h
@@ -0,0 +1,66 @@
+// Copyright Epic Games, Inc. All Rights Reserved.
+
+#pragma once
+
+#include "CoreTypes.h"
+#include "Misc/ScopeExit.h"
+
+namespace Chaos
+{
+	namespace CVars
+	{
+		extern CHAOS_API bool bEnableAsyncInitBody;
+	}
+
+	namespace Private
+	{
+		template<typename T>
+		struct FChaosAsyncInitBodyWriteScopeLock
+		{
+			inline FChaosAsyncInitBodyWriteScopeLock(T& InLock)
+				: Lock(Chaos::CVars::bEnableAsyncInitBody ? &InLock : nullptr)
+			{
+				if (Lock)
+				{
+					Lock->WriteLock();
+				}
+			}
+
+			inline ~FChaosAsyncInitBodyWriteScopeLock()
+			{
+				if (Lock)
+				{
+					Lock->WriteUnlock();
+				}
+			}
+
+			T* Lock;
+		};
+
+		template<typename T>
+		struct FChaosAsyncInitBodyReadScopeLock
+		{
+			inline FChaosAsyncInitBodyReadScopeLock(T& InLock)
+				: Lock(Chaos::CVars::bEnableAsyncInitBody ? &InLock : nullptr)
+			{
+				if (Lock)
+				{
+					Lock->ReadLock();
+				}
+			}
+
+			inline ~FChaosAsyncInitBodyReadScopeLock()
+			{
+				if (Lock)
+				{
+					Lock->ReadUnlock();
+				}
+			}
+
+			T* Lock;
+		};
+	}
+}
+
+#define UE_CHAOS_ASYNC_INITBODY_WRITESCOPELOCK(x) Chaos::Private::FChaosAsyncInitBodyWriteScopeLock WriteScopeLock(x)
+#define UE_CHAOS_ASYNC_INITBODY_READSCOPELOCK(x)  Chaos::Private::FChaosAsyncInitBodyReadScopeLock ReadScopeLock(x)
diff --git a/Engine/Source/Runtime/Experimental/Chaos/Public/Chaos/ChaosMarshallingManager.h b/Engine/Source/Runtime/Experimental/Chaos/Public/Chaos/ChaosMarshallingManager.h
--- a/Engine/Source/Runtime/Experimental/Chaos/Public/Chaos/ChaosMarshallingManager.h
+++ b/Engine/Source/Runtime/Experimental/Chaos/Public/Chaos/ChaosMarshallingManager.h
@@ -8,9 +8,7 @@
 #include "Chaos/ParallelFor.h"
 #include "PhysicsProxy/SingleParticlePhysicsProxyFwd.h"
 #include "Chaos/SimCallbackObject.h"
-#if UE_CHAOS_ASYNC_INITBODY_ENABLED
-#include "Misc/ScopeRWLock.h"
-#endif
+#include "Chaos/AsyncInitBodyHelper.h"
 
 class FGeometryCollectionResults;
 
@@ -298,9 +296,7 @@ public:
 
 	const FDirtyProxiesBucketInfo& GetDirtyProxyBucketInfo_External()
 	{
-#if UE_CHAOS_ASYNC_INITBODY_ENABLED
-		ensureAlwaysMsgf(false, TEXT("This method is not safe when used with UE_CHAOS_ASYNC_INITBODY_ENABLED = 1"));
-#endif
+		ensureAlwaysMsgf(!Chaos::CVars::bEnableAsyncInitBody, TEXT("This method is not safe when p.Chaos.EnableAsyncInitBody is true"));
 		return ProducerData->DirtyProxiesDataBuffer.GetDirtyProxyBucketInfo();
 	}
 
@@ -403,10 +399,8 @@ public:
 	/** Return the size of the history queue */
 	int32 GetNumHistory_Internal() const {return HistoryQueue_Internal.Num();}
 
-#if UE_CHAOS_ASYNC_INITBODY_ENABLED
 	/** Used for multithreaded access */
 	CHAOS_API FRWLock& GetMarshallingManagerLock() { return MarshallingManagerLock; }
-#endif
 
 private:
 	std::atomic<FReal> ExternalTime_External;	//the global time external thread is currently at
@@ -434,9 +428,7 @@ private:
 	CHAOS_API void PrepareExternalQueue_External();
 	CHAOS_API void PreparePullData();
 
-#if UE_CHAOS_ASYNC_INITBODY_ENABLED
 	/** Used for multithreaded access */
 	FRWLock MarshallingManagerLock;
-#endif
 };
 }; // namespace Chaos
diff --git a/Engine/Source/Runtime/Experimental/Chaos/Public/Chaos/Defines.h b/Engine/Source/Runtime/Experimental/Chaos/Public/Chaos/Defines.h
--- a/Engine/Source/Runtime/Experimental/Chaos/Public/Chaos/Defines.h
+++ b/Engine/Source/Runtime/Experimental/Chaos/Public/Chaos/Defines.h
@@ -19,18 +19,6 @@
 #endif
 #endif
 
-// Experimental UE_CHAOS_ASYNC_INITBODY_ENABLED option: Set to 1 to allow InitBody to be called outside of GameThread
-#define UE_CHAOS_ASYNC_INITBODY_ENABLED 0 
-#if UE_CHAOS_ASYNC_INITBODY_ENABLED
-#define UE_CHAOS_ASYNC_INITBODY_WRITESCOPELOCK(x) FWriteScopeLock WriteLock(x)
-#define UE_CHAOS_ASYNC_INITBODY_READSCOPELOCK(x) FReadScopeLock ReadLock(x)
-#define UE_CHAOS_ASYNC_INITBODY_PHYSICSSCENE_WRITESCOPELOCK(x) Chaos::FPhysicsSceneGuardScopedWrite WriteLock(x)
-#else
-#define UE_CHAOS_ASYNC_INITBODY_WRITESCOPELOCK(x)
-#define UE_CHAOS_ASYNC_INITBODY_READSCOPELOCK(x)
-#define UE_CHAOS_ASYNC_INITBODY_PHYSICSSCENE_WRITESCOPELOCK(x)
-#endif
-
 #if COMPILE_WITHOUT_UNREAL_SUPPORT
 	#include <stdint.h>
 #else
diff --git a/Engine/Source/Runtime/Experimental/Chaos/Public/Chaos/Framework/PhysicsSolverBase.h b/Engine/Source/Runtime/Experimental/Chaos/Public/Chaos/Framework/PhysicsSolverBase.h
--- a/Engine/Source/Runtime/Experimental/Chaos/Public/Chaos/Framework/PhysicsSolverBase.h
+++ b/Engine/Source/Runtime/Experimental/Chaos/Public/Chaos/Framework/PhysicsSolverBase.h
@@ -11,6 +11,7 @@
 #include "Async/ParallelFor.h"
 #include "Containers/Queue.h"
 #include "Chaos/ChaosMarshallingManager.h"
+#include "Chaos/AsyncInitBodyHelper.h"
 #include "Stats/Stats2.h"
 #include "ChaosSolversModule.h"
 
@@ -423,7 +424,7 @@ namespace Chaos
 		void EnqueueCommandImmediate(Lambda&& Func)
 		{
 			//TODO: remove this check. Need to rename with _External
-			check(UE_CHAOS_ASYNC_INITBODY_ENABLED || IsInGameThread());
+			check(Chaos::CVars::bEnableAsyncInitBody || IsInGameThread());
 			RegisterSimOneShotCallback(MoveTemp(Func));
 		}
 
diff --git a/Engine/Source/Runtime/Experimental/Chaos/Public/Chaos/ParticleHandle.h b/Engine/Source/Runtime/Experimental/Chaos/Public/Chaos/ParticleHandle.h
--- a/Engine/Source/Runtime/Experimental/Chaos/Public/Chaos/ParticleHandle.h
+++ b/Engine/Source/Runtime/Experimental/Chaos/Public/Chaos/ParticleHandle.h
@@ -12,6 +12,7 @@
 #include "Chaos/Collision/CollisionFilterBits.h"
 #include "ChaosCheck.h"
 #include "Chaos/ChaosDebugDrawDeclares.h"
+#include "Chaos/AsyncInitBodyHelper.h"
 #include "ChaosVisualDebugger/ChaosVisualDebuggerTrace.h"
 #include "PhysicsInterfaceTypesCore.h"
 #if CHAOS_DEBUG_DRAW
@@ -3090,7 +3091,7 @@ protected:
 	template <typename Lambda>
 	void ModifyGeometry(EGeometryAccess AccessType, const Lambda& Func)
 	{
-		ensure(UE_CHAOS_ASYNC_INITBODY_ENABLED || IsInGameThread());
+		ensure(Chaos::CVars::bEnableAsyncInitBody || IsInGameThread());
 		FPhysicsSolverBase* Solver = Proxy ? Proxy->GetSolverBase() : nullptr;
 
 		if(Solver == nullptr)
@@ -3857,7 +3858,7 @@ FORCEINLINE_DEBUGGABLE FAccelerationStructureHandle::FAccelerationStructureHandl
 	if (InGeometryParticle)
 	{
 		ensure(CachedUniqueIdx.IsValid());
-		ensure(UE_CHAOS_ASYNC_INITBODY_ENABLED || IsInGameThread());
+		ensure(Chaos::CVars::bEnableAsyncInitBody || IsInGameThread());
 		if (bUsePrefiltering)
 		{
 			UpdatePrePreFilter(*InGeometryParticle);
diff --git a/Engine/Source/Runtime/Experimental/Chaos/Public/Chaos/PhysicalMaterials.h b/Engine/Source/Runtime/Experimental/Chaos/Public/Chaos/PhysicalMaterials.h
--- a/Engine/Source/Runtime/Experimental/Chaos/Public/Chaos/PhysicalMaterials.h
+++ b/Engine/Source/Runtime/Experimental/Chaos/Public/Chaos/PhysicalMaterials.h
@@ -168,10 +168,8 @@ namespace Chaos
 		/** Handle-managed array of global material masks. This is pushed to all solvers who all maintain a copy */
 		THandleArray<FChaosPhysicsMaterialMask> MaterialMasks;
 
-#if UE_CHAOS_ASYNC_INITBODY_ENABLED
 		mutable FRWLock MaterialsLock;
 		mutable FRWLock MaterialMasksLock;
-#endif
 	};
 }
 
diff --git a/Engine/Source/Runtime/PhysicsCore/Private/ChaosScene.cpp b/Engine/Source/Runtime/PhysicsCore/Private/ChaosScene.cpp
--- a/Engine/Source/Runtime/PhysicsCore/Private/ChaosScene.cpp
+++ b/Engine/Source/Runtime/PhysicsCore/Private/ChaosScene.cpp
@@ -35,6 +35,7 @@
 #include "RewindData.h"
 #include "PhysicsSettingsCore.h"
 #include "Chaos/PhysicsSolverBaseImpl.h"
+#include "Chaos/AsyncInitBodyHelper.h"
 
 #include "ChaosVisualDebugger/ChaosVisualDebuggerTrace.h"
 
@@ -281,7 +282,7 @@ void FChaosScene::AddActorsToScene_AssumesLocked(TArray<FPhysicsActorHandle>& In
 	TRACE_CPUPROFILER_EVENT_SCOPE(FChaosScene::AddActorsToScene_AssumesLocked)
 
 	Chaos::FPhysicsSolver* Solver = GetSolver();
-	UE_CHAOS_ASYNC_INITBODY_PHYSICSSCENE_WRITESCOPELOCK(Solver->GetExternalDataLock_External());
+	UE_CHAOS_ASYNC_INITBODY_WRITESCOPELOCK(Solver->GetExternalDataLock_External());
 	Chaos::ISpatialAcceleration<Chaos::FAccelerationStructureHandle,Chaos::FReal,3>* SpatialAcceleration = GetSpacialAcceleration();
 	for(FPhysicsActorHandle& Handle : InHandles)
 	{
