[ChaosVD] Fix an issue where skeletal mesh bones movement are not recorded every frame

The deferred skeletal mesh update logic moves the bones by setting the new transform directly (performing a teleport). This means the particle is not added to the dirty particle array which is what CVD uses to determine what to record during a delta frame.

The fix is to ensure we record the particle when a new transform is set directly via the Evolution API

diff --git a/Engine/Source/Runtime/Experimental/Chaos/Private/Chaos/PBDRigidsEvolutionGBF.cpp b/Engine/Source/Runtime/Experimental/Chaos/Private/Chaos/PBDRigidsEvolutionGBF.cpp
--- a/Engine/Source/Runtime/Experimental/Chaos/Private/Chaos/PBDRigidsEvolutionGBF.cpp
+++ b/Engine/Source/Runtime/Experimental/Chaos/Private/Chaos/PBDRigidsEvolutionGBF.cpp
@@ -1420,6 +1420,10 @@ void FPBDRigidsEvolutionGBF::SetParticleTransform(FGeometryParticleHandle* InPar
 
 	OnParticleMoved(InParticle, PrevX, PrevR, bIsTeleport);
 
+	// This is a bit overkill, if this becomes an issue we will need to look into adding support
+	// to record partial updates of particles, only the transform in this case
+	CVD_TRACE_PARTICLE(InParticle);
+
 #if CHAOS_EVOLUTION_COLLISION_TESTMODE
 	{
 		// Update the test mode cache so we can move particles in PIE to test collisions
