ChaosFlesh:: fix jira Clang compiler issue

diff --git a/Engine/Source/Runtime/Experimental/Chaos/Public/Chaos/Deformable/GaussSeidelWeakConstraints.h b/Engine/Source/Runtime/Experimental/Chaos/Public/Chaos/Deformable/GaussSeidelWeakConstraints.h
--- a/Engine/Source/Runtime/Experimental/Chaos/Public/Chaos/Deformable/GaussSeidelWeakConstraints.h
+++ b/Engine/Source/Runtime/Experimental/Chaos/Public/Chaos/Deformable/GaussSeidelWeakConstraints.h
@@ -109,12 +109,12 @@ namespace Chaos::Softs
 		TGaussSeidelWeakConstraintData& operator=(TGaussSeidelWeakConstraintData<T>&& Other)
 		{
 			MIndices = MoveTemp(Other.MIndices);
-			MIndices = MoveTemp(Other.MSecondIndices);
-			MIndices = MoveTemp(Other.MWeights);
-			MIndices = MoveTemp(Other.MSecondWeights);
-			MIndices = MoveTemp(Other.MStiffness);
-			MIndices = MoveTemp(Other.MIsAnisotropic);
-			MIndices = MoveTemp(Other.MNormals);
+			MSecondIndices = MoveTemp(Other.MSecondIndices);
+			MWeights = MoveTemp(Other.MWeights);
+			MSecondWeights = MoveTemp(Other.MSecondWeights);
+			MStiffness = MoveTemp(Other.MStiffness);
+			MIsAnisotropic = MoveTemp(Other.MIsAnisotropic);
+			MNormals = MoveTemp(Other.MNormals);
 			ResizeHelper(Other.Size());
 			Other.MSize = 0;
 			return *this;
