Chaos::Fix bug in FPBDJointUtilities::ConditionInverseMassAndInertia

diff --git a/Engine/Source/Runtime/Experimental/Chaos/Public/Chaos/PBDRigidParticles.h b/Engine/Source/Runtime/Experimental/Chaos/Public/Chaos/PBDRigidParticles.h
--- a/Engine/Source/Runtime/Experimental/Chaos/Public/Chaos/PBDRigidParticles.h
+++ b/Engine/Source/Runtime/Experimental/Chaos/Public/Chaos/PBDRigidParticles.h
@@ -179,7 +179,7 @@ class TPBDRigidParticles : public TRigidParticles<T, d>
 		{
 			// Transitioning from kinematic or static to dynamic, compute the inverses.
 			this->InvM(Index) = FMath::IsNearlyZero(this->M(Index)) ? 0.0f : 1.f / this->M(Index);
-			this->InvI(Index) = this->I(Index).IsNearlyZero() ? TVec3<FRealSingle>::ZeroVector :
+			this->InvI(Index) = (this->I(Index).IsNearlyZero() || FMath::IsNearlyZero(this->M(Index))) ? TVec3<FRealSingle>::ZeroVector :
 				TVec3<FRealSingle>(
 				1.f / this->I(Index)[0],
 				1.f / this->I(Index)[1],
