[ChaosVD] Fix serialization errors introduced a few days ago in CL

Data was added to one of the data wrappers structures without a version bump.

This CL means recordings made in between will fail to load, but that is an acceptable trade off to stop the bleeding now (CVD recording are not assets, they are debug files. But backwards compat is important so we can help other teams or licencees working in previous versions of the engine without having to sync to a specific brach or CL)

diff --git a/Engine/Source/Runtime/Core/Public/UObject/FortniteMainBranchObjectVersion.h b/Engine/Source/Runtime/Core/Public/UObject/FortniteMainBranchObjectVersion.h
--- a/Engine/Source/Runtime/Core/Public/UObject/FortniteMainBranchObjectVersion.h
+++ b/Engine/Source/Runtime/Core/Public/UObject/FortniteMainBranchObjectVersion.h
@@ -524,6 +524,9 @@ struct FFortniteMainBranchObjectVersion
 		// Subsurface profile now has a guid to be able to select one of many in a Substrate material.
 		SubsurfaceProfileGuid,
 
+		// Added support for to record the new solver iteration settings in CVD
+		SolverIterationsDataSupportInChaosVisualDebugger,
+
 		// -----<new versions can be added above this line>-------------------------------------------------
 		VersionPlusOne,
 		LatestVersion = VersionPlusOne - 1
diff --git a/Engine/Source/Runtime/Experimental/ChaosVisualDebugger/Private/DataWrappers/ChaosVDParticleDataWrapper.cpp b/Engine/Source/Runtime/Experimental/ChaosVisualDebugger/Private/DataWrappers/ChaosVDParticleDataWrapper.cpp
--- a/Engine/Source/Runtime/Experimental/ChaosVisualDebugger/Private/DataWrappers/ChaosVDParticleDataWrapper.cpp
+++ b/Engine/Source/Runtime/Experimental/ChaosVisualDebugger/Private/DataWrappers/ChaosVDParticleDataWrapper.cpp
@@ -2,6 +2,7 @@
 
 #include "DataWrappers/ChaosVDParticleDataWrapper.h"
 
+#include "UObject/FortniteMainBranchObjectVersion.h"
 #include "UObject/FortniteSeasonBranchObjectVersion.h"
 
 // @note: Tracing an scene with 1000 particles moving, Manually serializing the structs is ~20% faster
@@ -27,9 +28,15 @@ bool FChaosVDFRigidParticleControlFlags::Serialize(FArchive& Ar)
 	Ar << bInertiaConditioningEnabled;
 	Ar << GravityGroupIndex;
 	Ar << bMACDEnabled;
-	Ar << PositionSolverIterationCount;
-	Ar << VelocitySolverIterationCount;
-	Ar << ProjectionSolverIterationCount;
+
+	Ar.UsingCustomVersion(FFortniteMainBranchObjectVersion::GUID);
+
+	if (Ar.CustomVer(FFortniteMainBranchObjectVersion::GUID) >= FFortniteMainBranchObjectVersion::SolverIterationsDataSupportInChaosVisualDebugger)
+	{
+		Ar << PositionSolverIterationCount;
+		Ar << VelocitySolverIterationCount;
+		Ar << ProjectionSolverIterationCount;
+	}
 
 	return !Ar.IsError();
 }
