Second fix shift of parts on dynamic cluster union
This is fix of the previous fix : 36098349
- this time the skip flag is properly captured by the physics commands to make sure it is properly considered on the physics thread

diff --git a/Engine/Source/Runtime/Experimental/Chaos/Private/PhysicsProxy/GeometryCollectionPhysicsProxy.cpp b/Engine/Source/Runtime/Experimental/Chaos/Private/PhysicsProxy/GeometryCollectionPhysicsProxy.cpp
--- a/Engine/Source/Runtime/Experimental/Chaos/Private/PhysicsProxy/GeometryCollectionPhysicsProxy.cpp
+++ b/Engine/Source/Runtime/Experimental/Chaos/Private/PhysicsProxy/GeometryCollectionPhysicsProxy.cpp
@@ -3190,10 +3190,13 @@ void FGeometryCollectionPhysicsProxy::SetWorldTransform_External(const FTransfor
 		PreviousWorldTransform_External = WorldTransform_External;
 		WorldTransform_External = WorldTransform;
 
+		// Because SetWorldTransform_Internal is running on the physics thread we need to capture this parameter 
+		// (it is transient and may changed after calling this function)
+		const bool bSkipChildToParentUpdate = GetSkipChildToParentUpdateWhenInClusterUnion();
 		ExecuteOnPhysicsThread(*this,
-			[this, WorldTransform]()
+			[this, WorldTransform, bSkipChildToParentUpdate]()
 			{	
-				SetWorldTransform_Internal(WorldTransform);
+				SetWorldTransform_Internal(WorldTransform, bSkipChildToParentUpdate);
 			});
 	}
 }
@@ -3221,7 +3224,7 @@ void FGeometryCollectionPhysicsProxy::ScaleClusterGeometry_Internal(const FVecto
 	}
 }
 
-void FGeometryCollectionPhysicsProxy::SetWorldTransform_Internal(const FTransform& InWorldTransform)
+void FGeometryCollectionPhysicsProxy::SetWorldTransform_Internal(const FTransform& InWorldTransform, bool bInSkipChildToParentUpdateWhenInClusterUnion)
 {
 	using namespace Chaos;
 
@@ -3303,7 +3306,7 @@ void FGeometryCollectionPhysicsProxy::SetWorldTransform_Internal(const FTransfor
 		// This should only happen on the server otherwise the client may override the replicated child to parent before it's even set.
 		if (bIsAuthority && ClusterUnionIndex != INDEX_NONE && !DeferredClusterUnionParticleUpdates.IsEmpty() && !DeferredClusterUnionChildToParentUpdates.IsEmpty())
 		{
-			if (!bSkipChildToParentUpdateWhenInClusterUnion)
+			if (!bInSkipChildToParentUpdateWhenInClusterUnion)
 			{
 				ClusterUnionManager.UpdateClusterUnionParticlesChildToParent(ClusterUnionIndex, DeferredClusterUnionParticleUpdates, DeferredClusterUnionChildToParentUpdates, false);
 			}
diff --git a/Engine/Source/Runtime/Experimental/Chaos/Public/PhysicsProxy/GeometryCollectionPhysicsProxy.h b/Engine/Source/Runtime/Experimental/Chaos/Public/PhysicsProxy/GeometryCollectionPhysicsProxy.h
--- a/Engine/Source/Runtime/Experimental/Chaos/Public/PhysicsProxy/GeometryCollectionPhysicsProxy.h
+++ b/Engine/Source/Runtime/Experimental/Chaos/Public/PhysicsProxy/GeometryCollectionPhysicsProxy.h
@@ -623,7 +623,7 @@ protected:
 	/** Scale the cluster particles geometry (creates if necessary an additional TImplicitObjectScaled object into the implicits hierarchy) */
 	CHAOS_API void ScaleClusterGeometry_Internal(const FVector& WorldScale);
 
-	CHAOS_API void SetWorldTransform_Internal(const FTransform& WorldTransform);
+	CHAOS_API void SetWorldTransform_Internal(const FTransform& WorldTransform, bool bInSkipChildToParentUpdateWhenInClusterUnion = false);
 	CHAOS_API void SetFilterData_Internal(const FCollisionFilterData& NewSimFilter, const FCollisionFilterData& NewQueryFilter);
 	CHAOS_API void SetPerParticleFilterData_Internal(const TArray<FParticleCollisionFilterData>& PerParticleData);
 	CHAOS_API void SetDamagePropagationData_Internal(bool bEnabled, float BreakDamagePropagationFactor, float ShockDamagePropagationFactor);
