Experimental MLLevelSet is added as a new physics assest. This new neural network based collision assest only assumes Batched version of PhiWithNormal. It depends on NNE module, and currently supports only MLP type of networks. MLLevelSet assest can be loaded via importing the details from a datatable. At the moment it is assumed that training is done outside of the engine.

diff --git a/Engine/Plugins/ChaosCloth/Source/ChaosCloth/Private/ChaosCloth/ChaosClothConstraints.cpp b/Engine/Plugins/ChaosCloth/Source/ChaosCloth/Private/ChaosCloth/ChaosClothConstraints.cpp
--- a/Engine/Plugins/ChaosCloth/Source/ChaosCloth/Private/ChaosCloth/ChaosClothConstraints.cpp
+++ b/Engine/Plugins/ChaosCloth/Source/ChaosCloth/Private/ChaosCloth/ChaosClothConstraints.cpp
@@ -43,6 +43,9 @@ extern TAutoConsoleVariable<float> CVarGravityMultiplier;
 }
 
 bool bEnableGS = false;
+bool bDoPlanarCollisionConstraints = false;
+Softs::FSolverReal PlanarCollisionTol = 0.f;
+Softs::FSolverReal PlanarCollisionStiffness = 1e10;
 
 #if !UE_BUILD_SHIPPING
     static FAutoConsoleVariableRef CVarClothbEnableGS(TEXT("p.Chaos.Cloth.EnableGaussSeidel"), bEnableGS, TEXT("Use Gauss Seidel constraints instead of XPBD [def: false]"));
@@ -99,6 +102,11 @@ bool bEnableGS = false;
 
 	static FAutoConsoleVariableRef CVarClothSOROmega(TEXT("p.Chaos.Cloth.SOROmega"), SOROmega, TEXT("SOR omega coefficient for acceleration [def: 1.7]"));
 
+	static FAutoConsoleVariableRef CVarClothDoPlanarCollisionConstraints(TEXT("p.Chaos.Cloth.DoPlanarCollisionConstraints"), bDoPlanarCollisionConstraints, TEXT("Do body collision at the beginning of timestep only and then do xpbd planar constraints [def: true]"));
+
+	static FAutoConsoleVariableRef CVarClothPlanarCollisionTol(TEXT("p.Chaos.Cloth.PlanarCollisionTol"), PlanarCollisionTol, TEXT("Collision tolerance for XPBD planar constraint [def: 0.]"));
+
+	static FAutoConsoleVariableRef CVarClothPlanarCollisionStiffness(TEXT("p.Chaos.Cloth.PlanarCollisionStiffness"), PlanarCollisionStiffness, TEXT("Collision stiffness for XPBD planar constraint [def: 1e10.]"));
 #endif
 
 class FClothConstraints::FRuleCreator
@@ -1393,6 +1401,17 @@ void FClothConstraints::CreateCollisionConstraint(
 			);
 		++NumCollisionConstraintRules;
 		++NumUpdateLinearSystemCollisionsRules;
+		if (bDoPlanarCollisionConstraints) 
+		{
+			CollisionConstraint->SetDoCollisionPerTimestep(true);
+			XPlanarConstraints = MakeShared<Softs::TXPBDPlanarConstraints<FSolverReal>>();
+			XPlanarConstraints->SetTolerance(PlanarCollisionTol);
+			XPlanarConstraints->SetStiffness(PlanarCollisionStiffness);
+			++NumConstraintRules;
+			++NumConstraintInits;
+			Evolution->SetDoCollisionPerTimestep(true);
+		}
+
 	}
 }
 
@@ -1744,6 +1763,20 @@ void FClothConstraints::CreateForceBasedRules(const TSharedPtr<Softs::FMultiResC
 			}
 		});
 	}
+	if (bDoPlanarCollisionConstraints)
+	{
+		RuleCreator.AddPerIterationPBDConstraintRule(
+			[this](Softs::FSolverParticlesRange& Particles, const Softs::FSolverReal Dt)
+		{
+			XPlanarConstraints->Apply(Particles, Dt, CollisionConstraint->GetCollisionIndices(), CollisionConstraint->GetCollisionTargets(), CollisionConstraint->GetCollisionNormalsOutput());
+		});
+
+		RuleCreator.AddPostInitialGuessParallelInitRule(
+			[this](const Softs::FSolverParticlesRange& Particles, const Softs::FSolverReal Dt, const Softs::ESolverMode SolverMode)
+		{
+			XPlanarConstraints->Init(Particles, Dt, CollisionConstraint->GetCollisionIndices().Num());
+		});
+	}
 }
 
 
diff --git a/Engine/Plugins/ChaosCloth/Source/ChaosCloth/Private/ChaosCloth/ChaosClothVisualization.cpp b/Engine/Plugins/ChaosCloth/Source/ChaosCloth/Private/ChaosCloth/ChaosClothVisualization.cpp
--- a/Engine/Plugins/ChaosCloth/Source/ChaosCloth/Private/ChaosCloth/ChaosClothVisualization.cpp
+++ b/Engine/Plugins/ChaosCloth/Source/ChaosCloth/Private/ChaosCloth/ChaosClothVisualization.cpp
@@ -38,6 +38,7 @@
 #include "Misc/LazySingleton.h"
 #include "SceneManagement.h"
 #include "SceneView.h"
+#include "Chaos/MLLevelset.h"
 #if WITH_EDITOR
 #include "CanvasItem.h"     //
 #include "Engine/Canvas.h"  // For draw text
@@ -52,6 +53,9 @@ namespace Chaos
 	{
 		static int DrawSkinnedLattice = 0;
 		static FAutoConsoleVariableRef CVarClothVizDrawSkinnedLattice(TEXT("p.ChaosClothVisualization.DrawSkinnedLattice"), DrawSkinnedLattice, TEXT("Draw skinned lattice, 0 = none, 1 = filled, 2 = empty, 3 = both"));
+		
+		static int DrawMLLevelSetFlag = 0;
+		static FAutoConsoleVariableRef CVarChaosDrawMLLevelSet(TEXT("p.ChaosClothVisualization.DrawMLLevelSetFlag"), DrawMLLevelSetFlag, TEXT("Default: 0 -> Draws nothing. 1,3,5 : Draws MLLevelSet Isocounters. 2,3 : Draws Bounding Box. 4, 5: Draws Training Bounding Box Axes."));
 
 		// TODO: move these options to be somewhere the new cloth editor visualization can use.
 		enum class EBendingDrawMode : int
@@ -850,6 +854,118 @@ namespace Chaos
 		}
 	}
 
+	static void DrawMLLevelSet(FPrimitiveDrawInterface* PDI, const FTransform& Transform, const FMaterialRenderProxy* MaterialRenderProxy, const FMLLevelSet& MLLevelSet)
+	{
+#if WITH_EDITOR
+		if (PDI && MaterialRenderProxy)
+		{
+			//Draw MLLevelSet Isosurface (this takes significant amount of time. Only use for debugging)
+			if (Private::DrawMLLevelSetFlag%2 == 1)
+			{
+				TArray<FVector3f> Vertices;
+				TArray<FIntVector> Tris;
+				MLLevelSet.GetZeroIsosurfaceGridCellFaces(Vertices, Tris);
+
+				FDynamicMeshBuilder MeshBuilder(PDI->View->GetFeatureLevel());
+				for (const FVector3f& V : Vertices)
+				{
+					MeshBuilder.AddVertex(FDynamicMeshVertex(V));
+				}
+				for (const FIntVector& T : Tris)
+				{
+					MeshBuilder.AddTriangle(T[0], T[1], T[2]);
+				}
+
+				MeshBuilder.Draw(PDI, Transform.ToMatrixWithScale(), MaterialRenderProxy, SDPG_World, false, false);
+			}
+
+			// Draw the Local grid that MLLevelSet lives in.
+			if (Private::DrawMLLevelSetFlag == 2 || Private::DrawMLLevelSetFlag == 3)
+			{
+				TArray<FVector> LocalGridCorners;
+				LocalGridCorners.SetNum(8);
+				TUniformGrid<FReal, 3> MGrid = MLLevelSet.GetGrid();
+				TArray<FVector> MinMaxCorners;
+				MinMaxCorners.SetNum(2);
+				MinMaxCorners[0] = MGrid.MinCorner();
+				MinMaxCorners[1] = MGrid.MaxCorner();
+
+				for (int32 i1 = 0; i1 < 2; i1++)
+				{
+					for (int32 i2 = 0; i2 < 2; i2++)
+					{
+						for (int32 i3 = 0; i3 < 2; i3++)
+						{
+							int32 i = 4 * i1 + 2 * i2 + i3;
+							LocalGridCorners[i][0] = MinMaxCorners[i1][0];
+							LocalGridCorners[i][1] = MinMaxCorners[i2][1];
+							LocalGridCorners[i][2] = MinMaxCorners[i3][2];
+						}
+					}
+				}
+
+				const FVector P000 = Transform.TransformPosition(LocalGridCorners[0]);
+				const FVector P001 = Transform.TransformPosition(LocalGridCorners[1]);
+				const FVector P010 = Transform.TransformPosition(LocalGridCorners[2]);
+				const FVector P011 = Transform.TransformPosition(LocalGridCorners[3]);
+				const FVector P100 = Transform.TransformPosition(LocalGridCorners[4]);
+				const FVector P101 = Transform.TransformPosition(LocalGridCorners[5]);
+				const FVector P110 = Transform.TransformPosition(LocalGridCorners[6]);
+				const FVector P111 = Transform.TransformPosition(LocalGridCorners[7]);
+
+				PDI->AddReserveLines(SDPG_World, 12);
+				const FColor& Color = FColor::Purple;
+				const FColor& ColorZ = FColor::Blue;
+				const FColor& ColorY = FColor::Green;
+				const FColor& ColorX = FColor::Red;
+				PDI->DrawLine(P000, P001, ColorZ, SDPG_World);
+				PDI->DrawLine(P000, P010, ColorY, SDPG_World);
+				PDI->DrawLine(P000, P100, ColorX, SDPG_World);
+				PDI->DrawLine(P001, P011, Color, SDPG_World);
+				PDI->DrawLine(P001, P101, Color, SDPG_World);
+				PDI->DrawLine(P010, P011, Color, SDPG_World);
+				PDI->DrawLine(P010, P110, Color, SDPG_World);
+				PDI->DrawLine(P011, P111, Color, SDPG_World);
+				PDI->DrawLine(P100, P101, Color, SDPG_World);
+				PDI->DrawLine(P100, P110, Color, SDPG_World);
+				PDI->DrawLine(P101, P111, Color, SDPG_World);
+				PDI->DrawLine(P110, P111, Color, SDPG_World);
+			}
+
+			//Draw the training grid.
+			if (Private::DrawMLLevelSetFlag == 4 || Private::DrawMLLevelSetFlag == 5)
+			{
+				FVector3f LocalGridCorner = MLLevelSet.GetTrainingGridMin();
+
+				TArray<FVector3f> CornersXYZ;
+				CornersXYZ.SetNum(3);
+				for (int32 i = 0; i < 3; i++)
+				{
+					CornersXYZ[i] = LocalGridCorner + MLLevelSet.GetTrainingGridVector(i);
+				}
+
+				const FVector P000 = Transform.TransformPosition((FVector)LocalGridCorner);
+				const FVector P001 = Transform.TransformPosition((FVector)CornersXYZ[2]);
+				const FVector P010 = Transform.TransformPosition((FVector)CornersXYZ[1]);
+				const FVector P100 = Transform.TransformPosition((FVector)CornersXYZ[0]);
+
+				PDI->AddReserveLines(SDPG_World, 3);
+				const FColor& Color = FColor::Purple;
+				const FColor& ColorZ = FColor::Blue;
+				const FColor& ColorY = FColor::Green;
+				const FColor& ColorX = FColor::Red;
+				PDI->DrawLine(P000, P001, ColorZ, SDPG_World);
+				PDI->DrawLine(P000, P010, ColorY, SDPG_World);
+				PDI->DrawLine(P000, P100, ColorX, SDPG_World);
+			}
+		}
+		else
+#endif
+		{
+			DrawCoordinateSystem(PDI, Transform.GetRotation(), Transform.GetTranslation());
+		}
+	}
+
 	void FClothVisualizationNoGC::DrawBounds(FPrimitiveDrawInterface* PDI) const
 	{
 		if (!Solver)
@@ -1436,6 +1552,19 @@ namespace Chaos
 
 							}
 							break;
+						case (ImplicitObjectType::MLLevelSet):
+							{
+								const FTransform CombinedTransform = FTransform(Rotation, Position);
+								const FMLLevelSet& MLLevelSet = Object->GetObjectChecked<FMLLevelSet>();
+								const FMaterialRenderProxy* MaterialRenderProxy =
+#if WITH_EDITOR
+									FMaterials::GetInstance().GetCollisionMaterial() ? FMaterials::GetInstance().GetCollisionMaterial()->GetRenderProxy() : nullptr;
+#else
+									nullptr;
+#endif
+								DrawMLLevelSet(PDI, CombinedTransform, MaterialRenderProxy, MLLevelSet);
+							}
+							break;
 						default:
 							DrawCoordinateSystem(PDI, Rotation, Position);  // Draw everything else as a coordinate for now
 							break;
diff --git a/Engine/Plugins/ChaosCloth/Source/ChaosCloth/Private/ChaosCloth/ChaosClothingSimulationCollider.cpp b/Engine/Plugins/ChaosCloth/Source/ChaosCloth/Private/ChaosCloth/ChaosClothingSimulationCollider.cpp
--- a/Engine/Plugins/ChaosCloth/Source/ChaosCloth/Private/ChaosCloth/ChaosClothingSimulationCollider.cpp
+++ b/Engine/Plugins/ChaosCloth/Source/ChaosCloth/Private/ChaosCloth/ChaosClothingSimulationCollider.cpp
@@ -21,6 +21,7 @@
 #include "HAL/PlatformMath.h"
 #include "Containers/ArrayView.h"
 #include "Containers/BitArray.h"
+#include "Chaos/MLLevelset.h"
 
 DECLARE_CYCLE_STAT(TEXT("Chaos Cloth Update Collider"), STAT_ChaosClothingSimulationColliderUpdate, STATGROUP_ChaosCloth);
 
@@ -46,6 +47,7 @@ struct FClothingSimulationCollider::FLODData
 		const FClothCollisionData& InClothCollisionData,
 		const TArray<FLevelSetCollisionData>& InLevelSetCollisionData,
 		const TArray<FSkinnedLevelSetCollisionData>& InSkinnedLevelSetCollisionData,
+		const TArray<FMLLevelSetCollisionData>& InMLLevelSetCollisionData,
 		const FReal InScale = 1.f,
 		const TArray<int32>& UsedBoneIndices = TArray<int32>());
 	void Remove(FClothingSimulationSolver* Solver, FClothingSimulationCloth* Cloth);
@@ -68,6 +70,7 @@ void FClothingSimulationCollider::FLODData::Add(
 	const FClothCollisionData& InClothCollisionData,
 	const TArray<FLevelSetCollisionData>& InLevelSetCollisionData,
 	const TArray<FSkinnedLevelSetCollisionData>& InSkinnedLevelSetCollisionData,
+	const TArray<FMLLevelSetCollisionData>& InMLLevelSetCollisionData,
 	const FReal InComponentScale,
 	const TArray<int32>& UsedBoneIndices)
 {
@@ -107,12 +110,19 @@ void FClothingSimulationCollider::FLODData::Add(
 	const int32 NumLevelSets = InLevelSetCollisionData.Num();
 	const int32 NumSkinnedLevelSets = InSkinnedLevelSetCollisionData.Num();
 	int32 NumSkinnedLevelSetBones = 0;
+	const int32 NumMLLevelSets = InMLLevelSetCollisionData.Num();
+	int32 NumMLLevelSetProxyBones = 0;
 	for (const FSkinnedLevelSetCollisionData& SkinnedLevelSetData : InSkinnedLevelSetCollisionData)
 	{
 		NumSkinnedLevelSetBones += SkinnedLevelSetData.MappedSkinnedBones.Num();
 	}
 
-	NumGeometries = NumSpheres + NumCapsules + NumConvexes + NumBoxes + NumLevelSets + NumSkinnedLevelSets + NumSkinnedLevelSetBones;
+	for (const FMLLevelSetCollisionData& MLLevelSetData : InMLLevelSetCollisionData)
+	{
+		NumMLLevelSetProxyBones += MLLevelSetData.MappedActiveBoneIndices.Num();
+	}
+
+	NumGeometries = NumSpheres + NumCapsules + NumConvexes + NumBoxes + NumLevelSets + NumSkinnedLevelSets + NumSkinnedLevelSetBones + NumMLLevelSets + NumMLLevelSetProxyBones;
 	if(NumGeometries == 0)
 	{
 		return;
@@ -335,7 +345,7 @@ void FClothingSimulationCollider::FLODData::Add(
 			UE_LOG(LogChaosCloth, VeryVerbose, TEXT("Found collision level set on bone index %d."), BoneIndices[Index]);
 
 			// Setup the collision particle geometry
-			FImplicitObjectPtr TransformedLevelSet = MakeImplicitObjectPtr<TImplicitObjectTransformed<FReal, 3>>(InLevelSetCollisionData[Index].LevelSet->DeepCopyGeometry(), 
+			FImplicitObjectPtr TransformedLevelSet = MakeImplicitObjectPtr<TImplicitObjectTransformed<FReal, 3>>(InLevelSetCollisionData[LevelSetIndex].LevelSet->DeepCopyGeometry(),
 				TRigidTransform<FReal, 3>(InLevelSetCollisionData[LevelSetIndex].Transform));
 			Solver->SetCollisionGeometry(CollisionRangeId, Index, MoveTemp(TransformedLevelSet));
 		}
@@ -381,6 +391,48 @@ void FClothingSimulationCollider::FLODData::Add(
 			++Index;
 		}
 	}
+
+	const int32 MLLevelSetIndexOffset = SkinnedLevelSetIndexOffset + NumSkinnedLevelSets + NumSkinnedLevelSetBones; 
+	if (NumMLLevelSets != 0)
+	{
+		int32 Index = MLLevelSetIndexOffset;
+		for (const FMLLevelSetCollisionData& MLLevelSetCollisionData : InMLLevelSetCollisionData)
+		{
+			TArray<int32> SolverActiveBoneIndices;
+			SolverActiveBoneIndices.Reserve(MLLevelSetCollisionData.MappedActiveBoneIndices.Num());
+			const int32 GlobalBoneOffset = Solver->IsLegacySolver() ? CollisionRangeId : 0;
+			
+			for (int32 MappedActiveBoneIndex : MLLevelSetCollisionData.MappedActiveBoneIndices)
+			{
+				// Always initialize the collision particle transforms before setting any geometry as otherwise NaNs gets detected during the bounding box updates
+				BaseTransforms[Index] = Softs::FSolverRigidTransform3::Identity;
+
+				BoneIndices[Index] = GetMappedBoneIndex(UsedBoneIndices, MappedActiveBoneIndex);
+				UE_LOG(LogChaosCloth, VeryVerbose, TEXT("Found collision ml level set active bone on bone index %d."), BoneIndices[Index]);
+
+				FImplicitObjectPtr BoneProxy = MakeImplicitObjectPtr<FMLLevelSetBoneProxy>();
+				Solver->SetCollisionGeometry(CollisionRangeId, Index, MoveTemp(BoneProxy));
+
+				SolverActiveBoneIndices.Add(GlobalBoneOffset + Index);
+
+				++Index;
+			}
+
+			// Always initialize the collision particle transforms before setting any geometry as otherwise NaNs gets detected during the bounding box updates
+			BaseTransforms[Index] = Softs::FSolverRigidTransform3::Identity;
+
+			BoneIndices[Index] = GetMappedBoneIndex(UsedBoneIndices, MLLevelSetCollisionData.ParentBoneIndex);
+			UE_LOG(LogChaosCloth, VeryVerbose, TEXT("Found collision ml level set on bone index %d."), BoneIndices[Index]);
+			
+			FImplicitObjectPtr MLLevelSet = MLLevelSetCollisionData.MLLevelSet->CopyGeometry();
+			MLLevelSet->GetObjectChecked<FMLLevelSet>().SetSolverActiveBoneIndices(MoveTemp(SolverActiveBoneIndices));
+			MLLevelSet->GetObjectChecked<FMLLevelSet>().SetSolverParentBoneIndex(GlobalBoneOffset + Index);
+
+			Solver->SetCollisionGeometry(CollisionRangeId, Index, MoveTemp(MLLevelSet));
+			++Index;
+		}
+	}
+
 	UE_LOG(LogChaosCloth, VeryVerbose, TEXT("Added collisions: %d spheres, %d capsules, %d convexes, %d boxes, %d level sets."), NumSpheres, NumCapsules, NumConvexes, NumBoxes, NumLevelSets);
 }
 
@@ -499,6 +551,21 @@ void FClothingSimulationCollider::ExtractPhysicsAssetCollision(
 	TArray<int32>& UsedBoneIndices,
 	bool bUseSphylsOnly,
 	bool bSkipMissingBones)
+{
+	TArray<FMLLevelSetCollisionData> MLLevelSetCollisions = {};
+	ExtractPhysicsAssetCollision(PhysicsAsset, ReferenceSkeleton, ClothCollisionData, LevelSetCollisions, SkinnedLevelSetCollisions, MLLevelSetCollisions, UsedBoneIndices, bUseSphylsOnly, bSkipMissingBones);
+}
+
+void FClothingSimulationCollider::ExtractPhysicsAssetCollision(
+	const UPhysicsAsset* PhysicsAsset,
+	const FReferenceSkeleton* ReferenceSkeleton,
+	FClothCollisionData& ClothCollisionData,
+	TArray<FLevelSetCollisionData>& LevelSetCollisions,
+	TArray<FSkinnedLevelSetCollisionData>& SkinnedLevelSetCollisions,
+	TArray<FMLLevelSetCollisionData>& MLLevelSetCollisions,
+	TArray<int32>& UsedBoneIndices,
+	bool bUseSphylsOnly,
+	bool bSkipMissingBones)
 {
 	ClothCollisionData.Reset();
 	UsedBoneIndices.Reset();
@@ -687,6 +754,24 @@ void FClothingSimulationCollider::ExtractPhysicsAssetCollision(
 					SkinnedLevelSetCollisions.Add({ SkinnedLevelSetElem.WeightedLevelSet(), MappedBoneIndex, MoveTemp(MappedSkinnedBones)});
 				}
 			}
+
+			for (const FKMLLevelSetElem& MLLevelSetElem : AggGeom.MLLevelSetElems)
+			{				
+				if (MLLevelSetElem.GetMLLevelSet().IsValid())
+				{
+					FName ParentBoneName = MLLevelSetElem.GetMLLevelSet()->GetParentBoneName();
+					TArray<int32> MappedActiveBones;
+					const TArray<FName>& ActiveBoneNames = MLLevelSetElem.GetMLLevelSet()->GetActiveBoneNames();
+					MappedActiveBones.Reserve(ActiveBoneNames.Num());
+					for (const FName& ActiveBoneName : ActiveBoneNames)
+					{
+						const int32 ActiveBoneIndex = ReferenceSkeleton ? ReferenceSkeleton->FindBoneIndex(ActiveBoneName) : INDEX_NONE;
+						const int32 MappedActiveBoneIndex = UsedBoneIndices.Add(ActiveBoneIndex);
+						MappedActiveBones.Add(MappedActiveBoneIndex);
+					}
+					MLLevelSetCollisions.Add({ MLLevelSetElem.GetMLLevelSet(), MappedBoneIndex, MoveTemp(MappedActiveBones)});
+				}
+			}
 		}  // End for PhysAsset->SkeletalBodySetups
 	}  // End if PhysicsAsset
 }
@@ -761,10 +846,11 @@ PRAGMA_ENABLE_DEPRECATION_WARNINGS
 	FClothCollisionData PhysicsAssetCollisionData;
 	TArray<FLevelSetCollisionData> LevelSetCollisions;
 	TArray<FSkinnedLevelSetCollisionData> SkinnedLevelSetCollisions;
+	TArray<FMLLevelSetCollisionData> MLLevelSetCollisions;
 	TArray<int32> UsedBoneIndices;
-	ExtractPhysicsAssetCollision(PhysicsAsset, ReferenceSkeleton, PhysicsAssetCollisionData, LevelSetCollisions, SkinnedLevelSetCollisions, UsedBoneIndices);
+	ExtractPhysicsAssetCollision(PhysicsAsset, ReferenceSkeleton, PhysicsAssetCollisionData, LevelSetCollisions, SkinnedLevelSetCollisions, MLLevelSetCollisions, UsedBoneIndices);
 
-	LODData[(int32)ECollisionDataType::LODless]->Add(Solver, Cloth, PhysicsAssetCollisionData, LevelSetCollisions, SkinnedLevelSetCollisions, Scale, UsedBoneIndices);
+	LODData[(int32)ECollisionDataType::LODless]->Add(Solver, Cloth, PhysicsAssetCollisionData, LevelSetCollisions, SkinnedLevelSetCollisions, MLLevelSetCollisions, Scale, UsedBoneIndices);
 }
 
 void FClothingSimulationCollider::Remove(FClothingSimulationSolver* Solver, FClothingSimulationCloth* Cloth)
@@ -792,8 +878,9 @@ void FClothingSimulationCollider::PreUpdate(FClothingSimulationSolver* Solver, F
 	// TODO: Get level sets?
 	const TArray<FLevelSetCollisionData> LevelSetCollisions;
 	const TArray<FSkinnedLevelSetCollisionData> SkinnedLevelSetCollisions;
+	const TArray<FMLLevelSetCollisionData> MLLevelSetCollisions;
 
-	LODData[(int32)ECollisionDataType::External]->Add(Solver, Cloth, CollisionData ? *CollisionData : FClothCollisionData(), LevelSetCollisions, SkinnedLevelSetCollisions);
+	LODData[(int32)ECollisionDataType::External]->Add(Solver, Cloth, CollisionData ? *CollisionData : FClothCollisionData(), LevelSetCollisions, SkinnedLevelSetCollisions, MLLevelSetCollisions);
 	
 	// TODO: Find a better way in case the same number but different collisions are being re-added (hash collision data? Provide user dirty function?)
 	bHasExternalCollisionChanged =
diff --git a/Engine/Plugins/ChaosCloth/Source/ChaosCloth/Private/ChaosCloth/ChaosClothingSimulationSolver.cpp b/Engine/Plugins/ChaosCloth/Source/ChaosCloth/Private/ChaosCloth/ChaosClothingSimulationSolver.cpp
--- a/Engine/Plugins/ChaosCloth/Source/ChaosCloth/Private/ChaosCloth/ChaosClothingSimulationSolver.cpp
+++ b/Engine/Plugins/ChaosCloth/Source/ChaosCloth/Private/ChaosCloth/ChaosClothingSimulationSolver.cpp
@@ -11,6 +11,7 @@
 #include "Chaos/PBDFlatWeightMap.h"
 #include "Chaos/WeightedLatticeImplicitObject.h"
 #include "Chaos/Levelset.h"
+#include "Chaos/MLLevelset.h" 
 #include "GeometryCollection/ManagedArrayCollection.h"
 #if INTEL_ISPC
 #include "ChaosClothingSimulationSolver.ispc.generated.h"
@@ -265,6 +266,23 @@ FClothingSimulationSolver::FClothingSimulationSolver(FClothingSimulationConfig*
 					SkinnedLevelSet->DeformPoints(SubBoneTransforms);
 					SkinnedLevelSet->UpdateSpatialHierarchy();
 				}
+				
+				if (FMLLevelSet* MLLevelSet =
+					const_cast<FImplicitObject*>(ParticlesInput.GetGeometry(Index).GetReference())->GetObject<FMLLevelSet>())
+				{
+					const TArray<int32>& ActiveBoneIndices = MLLevelSet->GetSolverActiveBoneIndices();
+					TArray<FTransform> ActiveBoneTransforms;
+					ActiveBoneTransforms.SetNum(ActiveBoneIndices.Num());
+					const FTransform RootTransform = TRigidTransform<FReal, 3>(X[Index], R[Index]);
+					for (int32 ActiveBoneIndex = 0; ActiveBoneIndex < ActiveBoneIndices.Num(); ++ActiveBoneIndex)
+					{
+						const int32 ActiveBoneIndexLocal = ActiveBoneIndices[ActiveBoneIndex] - ParticlesInput.GetOffset();
+						checkSlow(ActiveBoneIndexLocal < Index);
+						const FTransform ActiveBoneTransform = TRigidTransform<FReal, 3>(X[ActiveBoneIndexLocal], R[ActiveBoneIndexLocal]);
+						ActiveBoneTransforms[ActiveBoneIndex] = ActiveBoneTransform.GetRelativeTransform(RootTransform);
+					}
+					MLLevelSet->UpdateActiveBonesRelativeTransforms(ActiveBoneTransforms);
+				}
 			}
 		});
 	}
diff --git a/Engine/Plugins/ChaosCloth/Source/ChaosCloth/Public/ChaosCloth/ChaosClothConstraints.h b/Engine/Plugins/ChaosCloth/Source/ChaosCloth/Public/ChaosCloth/ChaosClothConstraints.h
--- a/Engine/Plugins/ChaosCloth/Source/ChaosCloth/Public/ChaosCloth/ChaosClothConstraints.h
+++ b/Engine/Plugins/ChaosCloth/Source/ChaosCloth/Public/ChaosCloth/ChaosClothConstraints.h
@@ -7,6 +7,7 @@
 #include "Chaos/PBDSoftsEvolutionFwd.h"
 #include "Chaos/PBDLongRangeConstraints.h"
 #include "Chaos/PBDCollisionSpringConstraintsBase.h"
+#include "Chaos/XPBDPlanarConstraints.h"
 #include "Chaos/Deformable/GaussSeidelMainConstraint.h"
 #include "Chaos/Deformable/GaussSeidelCorotatedCodimensionalConstraints.h"
 
@@ -222,6 +223,7 @@ namespace Chaos
 		TSharedPtr<Softs::FPBDSelfCollisionSphereConstraints> SelfCollisionSphereConstraints;
 		TSharedPtr<Softs::FGaussSeidelMainConstraint<Softs::FSolverReal, Softs::FSolverParticles>> GSMainConstraint;
 		TSharedPtr<Softs::FGaussSeidelCorotatedCodimensionalConstraints<Softs::FSolverReal, Softs::FSolverParticles>> GSCorotatedCodimensionalConstraint;
+		TSharedPtr<Softs::TXPBDPlanarConstraints<Softs::FSolverReal>> XPlanarConstraints;
 		TSharedPtr<Softs::FMultiResConstraints> MultiResConstraints;
 		//~ Begin Force-based solver only constraints
 		Softs::FSolverVec3 SolverWindVelocity; // Set from solver and added to wind from the config
diff --git a/Engine/Plugins/ChaosCloth/Source/ChaosCloth/Public/ChaosCloth/ChaosClothingSimulationCollider.h b/Engine/Plugins/ChaosCloth/Source/ChaosCloth/Public/ChaosCloth/ChaosClothingSimulationCollider.h
--- a/Engine/Plugins/ChaosCloth/Source/ChaosCloth/Public/ChaosCloth/ChaosClothingSimulationCollider.h
+++ b/Engine/Plugins/ChaosCloth/Source/ChaosCloth/Public/ChaosCloth/ChaosClothingSimulationCollider.h
@@ -45,6 +45,13 @@ namespace Chaos
 			TArray<int32> MappedSkinnedBones;
 		};
 
+		struct FMLLevelSetCollisionData
+		{
+			const TSharedPtr<Chaos::FMLLevelSet> MLLevelSet;
+			int32 ParentBoneIndex;
+			TArray<int32> MappedActiveBoneIndices;
+		};
+
 		CHAOSCLOTH_API FClothingSimulationCollider(const UPhysicsAsset* InPhysicsAsset, const FReferenceSkeleton* InReferenceSkeleton);
 		CHAOSCLOTH_API ~FClothingSimulationCollider();
 
@@ -93,6 +100,7 @@ namespace Chaos
 		// ---- End of the debugging and visualization functions ----
 
 		// Extract all collisions from the physics asset into a more cloth friendly format
+		UE_DEPRECATED(5.5, "Use the version that includes MLLevelSetCollisions.")
 		CHAOSCLOTH_API static void ExtractPhysicsAssetCollision(
 			const UPhysicsAsset* PhysicsAsset,
 			const FReferenceSkeleton* ReferenceSkeleton,
@@ -103,7 +111,19 @@ namespace Chaos
 			bool bUseSphylOnly = false,
 			bool bSkipMissingBones = false);
 
-	private:
+		// Extract all collisions from the physics asset into a more cloth friendly format
+		CHAOSCLOTH_API static void ExtractPhysicsAssetCollision(
+			const UPhysicsAsset* PhysicsAsset,
+			const FReferenceSkeleton* ReferenceSkeleton,
+			FClothCollisionData& ClothCollisionData,
+			TArray<FLevelSetCollisionData>& LevelSetCollisions,
+			TArray<FSkinnedLevelSetCollisionData>& SkinnedLevelSetCollisions,
+			TArray<FMLLevelSetCollisionData>& MLLevelSetCollisions,
+			TArray<int32>& UsedBoneIndices,
+			bool bUseSphylOnly = false,
+			bool bSkipMissingBones = false);
+
+	private:	
 		CHAOSCLOTH_API int32 GetNumGeometries(int32 InSlotIndex) const;
 
 		// Return the collision particle range ID for the specified slot being LODLess, external, or any of the LODs collision.
diff --git a/Engine/Plugins/ChaosClothAsset/Source/ChaosClothAssetEngine/Private/ChaosClothAsset/CollisionSources.cpp b/Engine/Plugins/ChaosClothAsset/Source/ChaosClothAssetEngine/Private/ChaosClothAsset/CollisionSources.cpp
--- a/Engine/Plugins/ChaosClothAsset/Source/ChaosClothAssetEngine/Private/ChaosClothAsset/CollisionSources.cpp
+++ b/Engine/Plugins/ChaosClothAsset/Source/ChaosClothAssetEngine/Private/ChaosClothAsset/CollisionSources.cpp
@@ -162,6 +162,7 @@ namespace UE::Chaos::ClothAsset
 					// Extract collisions
 					TArray<::Chaos::FClothingSimulationCollider::FLevelSetCollisionData> LevelSetCollisions;
 					TArray<::Chaos::FClothingSimulationCollider::FSkinnedLevelSetCollisionData> SkinnedLevelSetCollisions;
+					TArray<::Chaos::FClothingSimulationCollider::FMLLevelSetCollisionData> MLLevelSetCollisions;
 
 					constexpr bool bSkipMissingBones = true;
 					::Chaos::FClothingSimulationCollider::ExtractPhysicsAssetCollision(
@@ -170,6 +171,7 @@ namespace UE::Chaos::ClothAsset
 						CachedCollisionData,
 						LevelSetCollisions,
 						SkinnedLevelSetCollisions,
+						MLLevelSetCollisions,
 						CachedUsedBoneIndices,
 						bUseSphylsOnly,
 						bSkipMissingBones);
diff --git a/Engine/Source/Developer/PhysicsUtilities/Public/PhysicsAssetUtils.h b/Engine/Source/Developer/PhysicsUtilities/Public/PhysicsAssetUtils.h
--- a/Engine/Source/Developer/PhysicsUtilities/Public/PhysicsAssetUtils.h
+++ b/Engine/Source/Developer/PhysicsUtilities/Public/PhysicsAssetUtils.h
@@ -24,6 +24,7 @@ enum EPhysAssetFitGeomType : int
 	EFG_MultiConvexHull		UMETA(DisplayName="Multi Convex Hull"),
 	EFG_LevelSet			UMETA(DisplayName="Level Set"),
 	EFG_SkinnedLevelSet		UMETA(DisplayName="Skinned Level Set"),
+	EFG_MLLevelSet			UMETA(DisplayName="ML Level Set"),
 };
 
 UENUM()
diff --git a/Engine/Source/Editor/Persona/Private/PhysicsAssetRenderUtils.cpp b/Engine/Source/Editor/Persona/Private/PhysicsAssetRenderUtils.cpp
--- a/Engine/Source/Editor/Persona/Private/PhysicsAssetRenderUtils.cpp
+++ b/Engine/Source/Editor/Persona/Private/PhysicsAssetRenderUtils.cpp
@@ -425,6 +425,7 @@ namespace PhysicsAssetRender
 				CollectDrawElements(DrawElements, EAggCollisionShape::TaperedCapsule, AggGeom->TaperedCapsuleElems, BoneTM, Scale, i, RenderSettings, PhysicsAsset, PDI, GetPrimitiveColor, GetPrimitiveMaterial, GetPrimitiveTransform, CreateHitProxy);
 				CollectDrawElements(DrawElements, EAggCollisionShape::LevelSet, AggGeom->LevelSetElems, BoneTM, Scale, i, RenderSettings, PhysicsAsset, PDI, GetPrimitiveColor, GetPrimitiveMaterial, GetPrimitiveTransform, CreateHitProxy);
 				CollectDrawElements(DrawElements, EAggCollisionShape::SkinnedLevelSet, AggGeom->SkinnedLevelSetElems, BoneTM, Scale, i, RenderSettings, PhysicsAsset, PDI, GetPrimitiveColor, GetPrimitiveMaterial, GetPrimitiveTransform, CreateHitProxy);
+				CollectDrawElements(DrawElements, EAggCollisionShape::MLLevelSet, AggGeom->MLLevelSetElems, BoneTM, Scale, i, RenderSettings, PhysicsAsset, PDI, GetPrimitiveColor, GetPrimitiveMaterial, GetPrimitiveTransform, CreateHitProxy);
 
 				// Sort elements
 				DrawElements.Sort([](const DrawElement& DE1, const DrawElement& DE2) {
diff --git a/Engine/Source/Editor/PhysicsAssetEditor/Private/PhysicsAssetEditor.cpp b/Engine/Source/Editor/PhysicsAssetEditor/Private/PhysicsAssetEditor.cpp
--- a/Engine/Source/Editor/PhysicsAssetEditor/Private/PhysicsAssetEditor.cpp
+++ b/Engine/Source/Editor/PhysicsAssetEditor/Private/PhysicsAssetEditor.cpp
@@ -90,6 +90,7 @@
 #include "Widgets/Input/SNumericEntryBox.h"
 #include "AnimationEditorPreviewActor.h"
 #include "Preferences/PersonaOptions.h"
+#include "PhysicsEngine/MLLevelSetModelAndBonesBinningInfo.h"
 
 const FName PhysicsAssetEditorModes::PhysicsAssetEditorMode("PhysicsAssetEditorMode");
 
@@ -764,6 +765,12 @@ void FPhysicsAssetEditor::ExtendMenu()
 			Section.AddMenuEntry(Commands.SelectShapesQueryAndProbe);
 			Section.AddMenuEntry(Commands.SelectShapesProbeOnly);
 		}
+		
+		if(bEnableMLLevelSet)
+		{
+			FToolMenuSection& Section = EditMenu->AddSection("Import", LOCTEXT("PhatEditImport", "Import"));
+			Section.AddMenuEntry(Commands.ImportMLLevelSet);
+		}
 	}
 }
 
@@ -1258,6 +1265,11 @@ void FPhysicsAssetEditor::BindCommands()
 		FExecuteAction::CreateSP(this, &FPhysicsAssetEditor::OnSelectShapes, ECollisionEnabled::ProbeOnly),
 		FCanExecuteAction::CreateSP(this, &FPhysicsAssetEditor::IsNotSimulation));
 
+	ToolkitCommands->MapAction(
+		Commands.ImportMLLevelSet,
+		FExecuteAction::CreateSP(this, &FPhysicsAssetEditor::ImportMLLevelSetFromDataTable),
+		FCanExecuteAction::CreateSP(this, &FPhysicsAssetEditor::IsNotSimulation));
+
 	ToolkitCommands->MapAction(
 		Commands.SelectAllConstraints,
 		FExecuteAction::CreateSP(this, &FPhysicsAssetEditor::OnSelectAllConstraints),
@@ -1886,6 +1898,7 @@ void FPhysicsAssetEditor::BuildMenuWidgetSelection(FMenuBuilder& InMenuBuilder)
 		}
 		InMenuBuilder.AddMenuEntry( Commands.SelectShapesQueryAndProbe);
 		InMenuBuilder.AddMenuEntry( Commands.SelectShapesProbeOnly);
+		InMenuBuilder.AddMenuEntry( Commands.ImportMLLevelSet);
 		InMenuBuilder.EndSection();
 	}
 	InMenuBuilder.PopCommandList();
@@ -3537,6 +3550,313 @@ void FPhysicsAssetEditor::OnSelectShapes(const ECollisionEnabled::Type Collision
 	SharedData->SetSelectedBodiesPrimitivesWithCollisionType(SelectedBodyIndices.Array(), CollisionEnabled, true);
 }
 
+bool FPhysicsAssetEditor::LoadMLLevelSetFromDataTable(UDataTable* DataTableMLLevelSetModelBoneBinningInfo, TArray<FString>& ErrorMessages)
+{
+	USkeletalMesh* EditorSkelMesh = GetSharedData()->PhysicsAsset->GetPreviewMesh();
+	bool bCanLoadAll=true;
+	if (DataTableMLLevelSetModelBoneBinningInfo && EditorSkelMesh != nullptr)
+	{
+		for (FName RowName : DataTableMLLevelSetModelBoneBinningInfo->GetRowNames())
+		{
+			bool bCanLoadRow = true;
+			FMLLevelSetModelAndBonesBinningInfo* Row = DataTableMLLevelSetModelBoneBinningInfo->FindRow<FMLLevelSetModelAndBonesBinningInfo>(RowName, TEXT("GENERAL"));
+			if (Row)
+			{
+				const FName ParentBoneName = FName(*(Row->ParentBoneName));
+				const FString ActiveBoneNamesStringUnited = (*(Row->ActiveBoneNames));
+				TArray<FString> ActiveBoneNamesStringArray;
+				ActiveBoneNamesStringUnited.ParseIntoArray(ActiveBoneNamesStringArray, TEXT(","), true);
+				TArray<FName> ActiveBoneNames;
+				for (int32 i = 0; i < ActiveBoneNamesStringArray.Num(); i++)
+				{
+					ActiveBoneNames.Add(FName(ActiveBoneNamesStringArray[i]));
+				}
+				const int NumberOfActiveJoints = ActiveBoneNames.Num();
+
+				if (NumberOfActiveJoints == 0)
+				{
+					bCanLoadRow = false;
+					FString ErrorMessage = FString("Error in Row ") + RowName.ToString() + FString(". There must be at least one Active Joint");
+					ErrorMessages.Add(ErrorMessage);
+				}
+
+				if (Row->DebugGridResolution.Num() != 3)
+				{
+					bCanLoadRow = false;
+					FString ErrorMessage = FString("Error in Row ") + RowName.ToString() + FString(". DebugGridResolution must be TArray of 3 int values.");
+					ErrorMessages.Add(ErrorMessage);
+				}
+
+				if (Row->DebugGridResolution[0] < 5 || Row->DebugGridResolution[1] < 5 || Row->DebugGridResolution[2] < 5)
+				{
+					bCanLoadRow = false;
+					FString ErrorMessage = FString("Error in Row ") + RowName.ToString() + FString(". DebugGridResolution must be at least 5 in each component.");
+					ErrorMessages.Add(ErrorMessage);
+				}
+
+				if (Row->SignedDistanceScaling <= 0.0)
+				{
+					bCanLoadRow = false;
+					FString ErrorMessage = FString("Error in Row ") + RowName.ToString() + FString(". SignedDistanceScaling must be a positive number.");
+					ErrorMessages.Add(ErrorMessage);
+				}
+
+				if (Row->TrainingGridOrigin.Num() != 3)
+				{
+					bCanLoadRow = false;
+					FString ErrorMessage = FString("Error in Row ") + RowName.ToString() + FString(". TrainingGridOrigin cannot be imported. It must contain exactly 3 values.");
+					ErrorMessages.Add(ErrorMessage);
+				}
+
+				if (Row->TrainingGridAxisX.Num() != 3)
+				{
+					bCanLoadRow = false;
+					FString ErrorMessage = FString("Error in Row ") + RowName.ToString() + FString(". TrainingGridAxisX cannot be imported. It must contain exactly 3 values.");
+					ErrorMessages.Add(ErrorMessage);
+				}
+
+				if (Row->TrainingGridAxisY.Num() != 3)
+				{
+					bCanLoadRow = false;
+					FString ErrorMessage = FString("Error in Row ") + RowName.ToString() + FString(". TrainingGridAxisY cannot be imported. It must contain exactly 3 values.");
+					ErrorMessages.Add(ErrorMessage);
+				}
+
+				if (Row->TrainingGridAxisZ.Num() != 3)
+				{
+					bCanLoadRow = false;
+					FString ErrorMessage = FString("Error in Row ") + RowName.ToString() + FString(". TrainingGridAxisZ cannot be imported. It must contain exactly 3 values.");
+					ErrorMessages.Add(ErrorMessage);
+				}
+
+				//Check the initial joint rotations and translations matches with the number of active joints
+				if (NumberOfActiveJoints * 3 != Row->ReferenceBoneRotations.Num())
+				{
+					bCanLoadRow = false;
+					FString ErrorMessage = FString("Error in Row ") + RowName.ToString() + FString(". ReferenceBoneRotations cannot be imported. It must contain exactly 3 values per active joint.");
+					ErrorMessages.Add(ErrorMessage);
+				}
+
+				if (NumberOfActiveJoints * 3 != Row->ReferenceBoneTranslations.Num())
+				{
+					bCanLoadRow = false;
+					FString ErrorMessage = FString("Error in Row ") + RowName.ToString() + FString(". ReferenceBoneTranslations cannot be imported. It must contain exactly 3 values per active joint.");
+					ErrorMessages.Add(ErrorMessage);
+				}
+
+				if (NumberOfActiveJoints != Row->NumberOfRotationComponentsPerBone.Num())
+				{
+					bCanLoadRow = false;
+					FString ErrorMessage = FString("Error in Row ") + RowName.ToString() + FString(". NumberOfRotationComponentsPerBone cannot be imported. It must contain exactly 1 value per active joint.");
+					ErrorMessages.Add(ErrorMessage);
+				}
+
+				int32 TotalNumberOfActiveRotationComponents = 0;
+				for (int32 i = 0; i < Row->NumberOfRotationComponentsPerBone.Num(); i++)
+				{
+					TotalNumberOfActiveRotationComponents += Row->NumberOfRotationComponentsPerBone[i];
+				}
+
+				if (TotalNumberOfActiveRotationComponents != Row->RotationComponentIndexes.Num())
+				{
+					bCanLoadRow = false;
+					FString ErrorMessage = FString("Error in Row ") + RowName.ToString() + FString(". RotationComponentIndexes does not match with NumberOfRotationComponentsPerBone.");
+					ErrorMessages.Add(ErrorMessage);
+				}
+
+				if(bCanLoadRow)
+				{					
+					const double SignedDistanceScaling = Row->SignedDistanceScaling;
+					FIntVector DebugGridResolution = { Row->DebugGridResolution[0],Row->DebugGridResolution[1], Row->DebugGridResolution[2] };
+					const FVector3f TrainingGridMin = { Row->TrainingGridOrigin[0], Row->TrainingGridOrigin[1], Row->TrainingGridOrigin[2] };
+					TArray<FVector3f> TrainingGridAxes;
+					TrainingGridAxes.SetNum(3);
+					TrainingGridAxes[0] = FVector3f(Row->TrainingGridAxisX[0], Row->TrainingGridAxisX[1], Row->TrainingGridAxisX[2]);
+					TrainingGridAxes[1] = FVector3f(Row->TrainingGridAxisY[0], Row->TrainingGridAxisY[1], Row->TrainingGridAxisY[2]);
+					TrainingGridAxes[2] = FVector3f(Row->TrainingGridAxisZ[0], Row->TrainingGridAxisZ[1], Row->TrainingGridAxisZ[2]);
+
+					TArray<FVector3f> ReferenceBoneRotations;
+					TArray<FVector3f> ReferenceBoneTranslations;
+					ReferenceBoneRotations.SetNum(NumberOfActiveJoints);
+					ReferenceBoneTranslations.SetNum(NumberOfActiveJoints);
+					for (int32 i = 0; i < NumberOfActiveJoints; i++)
+					{
+						ReferenceBoneRotations[i] = FVector3f(Row->ReferenceBoneRotations[3 * i], Row->ReferenceBoneRotations[3 * i + 1], Row->ReferenceBoneRotations[3 * i + 2]);
+						ReferenceBoneTranslations[i] = FVector3f(Row->ReferenceBoneTranslations[3 * i], Row->ReferenceBoneTranslations[3 * i + 1], Row->ReferenceBoneTranslations[3 * i + 2]);
+					}
+					
+					TArray<TArray<int32>> ActiveBonesJointRotationComponents;
+					ActiveBonesJointRotationComponents.SetNum(NumberOfActiveJoints);
+
+					int32 ActiveRotationComponentIndexHelper = 0;
+					for (int32 i = 0; i < NumberOfActiveJoints; i++)
+					{
+						ActiveBonesJointRotationComponents[i].SetNum(Row->NumberOfRotationComponentsPerBone[i]);
+						for (int32 j = 0; j < ActiveBonesJointRotationComponents[i].Num(); j++)
+						{
+							ActiveBonesJointRotationComponents[i][j] = Row->RotationComponentIndexes[ActiveRotationComponentIndexHelper + j];
+						}
+						ActiveRotationComponentIndexHelper += Row->NumberOfRotationComponentsPerBone[i];
+					}
+					
+					UDataTable* MLLevelSetModelInferenceInfo = LoadObject<UDataTable>(nullptr, *(Row->MLModelInferenceInfoDataTablePath));
+					if (MLLevelSetModelInferenceInfo)
+					{
+						FMLLevelSetModelInferenceInfo* RowModelInferenceInfo = MLLevelSetModelInferenceInfo->FindRow<FMLLevelSetModelInferenceInfo>(FName(Row->MLModelInferenceInfoDataTableIndex), TEXT("GENERAL"));
+						if (RowModelInferenceInfo) 
+						{
+							const FString NNIModelDir = RowModelInferenceInfo->NNEModelPath;
+							const FString MLModelWeightsString = RowModelInferenceInfo->MLModelWeights;
+
+							const int32 ParentBoneIndex = EditorSkelMesh->GetRefSkeleton().FindBoneIndex(ParentBoneName);
+							if (ParentBoneIndex == INDEX_NONE)
+							{
+								UE_LOG(LogChaos, Error, TEXT("(Parent) Bone Index index with name %s is not found for the MLLevelSet."), *ParentBoneName.ToString());
+								bCanLoadRow = false;
+								FString ErrorMessage = FString("Error in Row ") + RowName.ToString() + FString(". ParentBoneIndex is not found.");
+								ErrorMessages.Add(ErrorMessage);
+							}
+							else 
+							{								
+								int32 ParentBodyIndex = GetSharedData()->PhysicsAsset->FindBodyIndex(ParentBoneName);	
+								
+								// Body associated with parent body does not exist. Create new body.
+								if(ParentBodyIndex == INDEX_NONE)
+								{									
+									FPhysAssetCreateParams NewBodyData = GetDefault<UPhysicsAssetGenerationSettings>()->CreateParams;
+									NewBodyData.GeomType = EPhysAssetFitGeomType::EFG_MLLevelSet;
+									ParentBodyIndex = FPhysicsAssetUtils::CreateNewBody(SharedData->PhysicsAsset,ParentBoneName, NewBodyData);
+								}
+
+								if (ParentBodyIndex == INDEX_NONE)
+								{
+									UE_LOG(LogChaos, Error, TEXT("Controlling Body with index %d is not found or cannot be created for MLLevelSet."), ParentBodyIndex);
+									bCanLoadRow = false;
+									FString ErrorMessage = FString("Error in Row ") + RowName.ToString() + FString(". Controlling Body is not found or cannot be created for MLLevelSet.");
+									ErrorMessages.Add(ErrorMessage);
+								}
+								else
+								{
+									const FReferenceSkeleton& RefSkeleton = EditorSkelMesh->GetRefSkeleton();
+									UE_LOG(LogChaos, Display, TEXT("MLLevelSet is imported and attached to ParentBone =%s with index %d"), *ParentBoneName.ToString(), ParentBoneIndex);
+									UBodySetup* DestBody = GetSharedData()->PhysicsAsset->SkeletalBodySetups[ParentBodyIndex];
+									DestBody->Modify(); 
+									FKMLLevelSetElem MLLevelsetElem;
+									
+									FMLLevelSetModelInferenceInfo* RowModelInferenceForIncorrectZoneInfo = MLLevelSetModelInferenceInfo->FindRow<FMLLevelSetModelInferenceInfo>(FName(Row->MLModelInferenceForIncorrectZoneInfoDataTableIndex), TEXT("GENERAL"));
+									if (RowModelInferenceForIncorrectZoneInfo)
+									{
+										TArray<Chaos::FMLLevelSetNNEModelData> NNEModelDataArr;
+										NNEModelDataArr.Add({ RowModelInferenceInfo->ModelArchitectureActivationNodeSizes, MLModelWeightsString, NNIModelDir, nullptr});
+										const FString NNIModelIncorrectZoneDir = RowModelInferenceForIncorrectZoneInfo->NNEModelPath;
+										const FString MLModelIncorrectZoneWeightsString = RowModelInferenceForIncorrectZoneInfo->MLModelWeights;
+										NNEModelDataArr.Add({ RowModelInferenceForIncorrectZoneInfo->ModelArchitectureActivationNodeSizes, MLModelIncorrectZoneWeightsString, NNIModelIncorrectZoneDir, nullptr });
+
+										Chaos::FMLLevelSetImportData MLLevelSetImportData = { ParentBoneName,ActiveBoneNames, MoveTemp(NNEModelDataArr),
+											MoveTemp(ReferenceBoneRotations), MoveTemp(ReferenceBoneTranslations),
+											SignedDistanceScaling, MoveTemp(ActiveBonesJointRotationComponents),TrainingGridMin, TrainingGridAxes, DebugGridResolution };
+										
+										MLLevelsetElem.BuildMLLevelSet(MoveTemp(MLLevelSetImportData));
+									}
+									// No Incorrect Zone Model.
+									else 
+									{
+										TArray<Chaos::FMLLevelSetNNEModelData> NNEModelDataArr;
+										NNEModelDataArr.Add({ RowModelInferenceInfo->ModelArchitectureActivationNodeSizes, MLModelWeightsString, NNIModelDir, nullptr });
+										
+										Chaos::FMLLevelSetImportData MLLevelSetImportData = { ParentBoneName,ActiveBoneNames, MoveTemp(NNEModelDataArr),
+											MoveTemp(ReferenceBoneRotations), MoveTemp(ReferenceBoneTranslations),
+											SignedDistanceScaling, MoveTemp(ActiveBonesJointRotationComponents),TrainingGridMin, TrainingGridAxes, DebugGridResolution };
+
+										MLLevelsetElem.BuildMLLevelSet(MoveTemp(MLLevelSetImportData));
+									}
+									DestBody->AggGeom.MLLevelSetElems.Add(MLLevelsetElem);	
+									RecreatePhysicsState();
+									RefreshHierachyTree();
+									RefreshPreviewViewport();
+								}
+							}
+						}
+						else
+						{
+							bCanLoadRow = false;
+							FString ErrorMessage = FString("Error in Row ") + RowName.ToString() + FString(". Data Table for MLModelInferenceInfo could not be loaded.");
+							ErrorMessages.Add(ErrorMessage);
+						}
+					}
+					else
+					{
+						bCanLoadRow = false;
+						FString ErrorMessage = FString("Error in Row ") + RowName.ToString() + FString(". Data Table for MLModelInferenceInfo could not be loaded. See MLModelInferenceInfoDataTablePath.");
+						ErrorMessages.Add(ErrorMessage);
+					}
+				}			
+			}
+			else
+			{
+				bCanLoadRow = false;
+				FString ErrorMessage = FString("Row ") + RowName.ToString() + FString(" cannot be imported. Wrong Data Table type imported. Data Table has to be FMLLevelSetModelAndBonesBinningInfo");
+				ErrorMessages.Add(ErrorMessage);
+			}
+			bCanLoadAll = bCanLoadAll ? bCanLoadRow : false;
+		} 
+	}
+	else
+	{
+		ErrorMessages.Add(FString("EditorSkelMesh could not be loaded."));
+		bCanLoadAll = false;
+	}
+
+	return bCanLoadAll;
+}
+
+void FPhysicsAssetEditor::ImportMLLevelSetFromDataTable()
+{
+	FContentBrowserModule& ContentBrowserModule = FModuleManager::LoadModuleChecked<FContentBrowserModule>(TEXT("ContentBrowser"));
+	FOpenAssetDialogConfig OpenAssetDialogConfig;
+	OpenAssetDialogConfig.DialogTitleOverride = FText::FromString(TEXT("Select DataTable for MLLevelSet Import"));
+	OpenAssetDialogConfig.AssetClassNames.Add(UDataTable::StaticClass()->GetClassPathName());
+	OpenAssetDialogConfig.bAllowMultipleSelection = false;
+	
+	TArray<FAssetData> SelectedDataTable = ContentBrowserModule.Get().CreateModalOpenAssetDialog(OpenAssetDialogConfig);
+	TArray<FString> ErrorMessages;
+
+	if (SelectedDataTable.Num() > 0)
+	{
+		FAssetData SelectedAsset = SelectedDataTable[0];
+		if (SelectedDataTable[0].IsValid())
+		{
+			UDataTable* DataTableMLLevelSetModelBoneBinningInfo = LoadObject<UDataTable>(nullptr, *SelectedDataTable[0].GetObjectPathString());
+			if (DataTableMLLevelSetModelBoneBinningInfo)
+			{
+				LoadMLLevelSetFromDataTable(DataTableMLLevelSetModelBoneBinningInfo, ErrorMessages);
+			}
+			else
+			{
+				ErrorMessages.Add(FString("Data Table could not be loaded."));
+			}
+		}
+		else
+		{
+			ErrorMessages.Add(FString("Selected Data Table is not valid."));
+		}
+	}
+	else
+	{
+		ErrorMessages.Add(FString("Data Table could not be selected."));
+	}
+
+	if (ErrorMessages.Num() >= 1)
+	{
+		FString ErrorMessagesStr = "";
+		for (FString ErrorMessage : ErrorMessages)
+		{
+			ErrorMessagesStr += ErrorMessage + FString("\n");
+		}
+		FMessageDialog::Open(EAppMsgType::Ok, FText::Format(LOCTEXT("MLLevelSetDataTableErrorMessage", "{0}"), FText::FromString(ErrorMessagesStr)));
+	}
+}
+
 void FPhysicsAssetEditor::OnSelectAllConstraints()
 {
 	UPhysicsAsset * const PhysicsAsset = SharedData->EditorSkelComp->GetPhysicsAsset();
diff --git a/Engine/Source/Editor/PhysicsAssetEditor/Private/PhysicsAssetEditor.h b/Engine/Source/Editor/PhysicsAssetEditor/Private/PhysicsAssetEditor.h
--- a/Engine/Source/Editor/PhysicsAssetEditor/Private/PhysicsAssetEditor.h
+++ b/Engine/Source/Editor/PhysicsAssetEditor/Private/PhysicsAssetEditor.h
@@ -21,6 +21,7 @@
 #include "BodySetupEnums.h"
 #include "Containers/ArrayView.h"
 #include "GraphEditor.h"
+#include "Engine/DataTable.h"
 
 class IDetailLayoutBuilder;
 struct FAssetData;
@@ -328,6 +329,8 @@ private:
 	void OnShowAll();
 	void OnHideAll();
 	void OnDeselectAll();
+	bool LoadMLLevelSetFromDataTable(UDataTable* DataTableMLLevelSetModelBoneBinningInfo, TArray<FString>& ErrorMessages);
+	void ImportMLLevelSetFromDataTable();
 
 	FText GetRepeatLastSimulationToolTip() const;
 	FSlateIcon GetRepeatLastSimulationIcon() const;
diff --git a/Engine/Source/Editor/PhysicsAssetEditor/Private/PhysicsAssetEditorActions.cpp b/Engine/Source/Editor/PhysicsAssetEditor/Private/PhysicsAssetEditorActions.cpp
--- a/Engine/Source/Editor/PhysicsAssetEditor/Private/PhysicsAssetEditorActions.cpp
+++ b/Engine/Source/Editor/PhysicsAssetEditor/Private/PhysicsAssetEditorActions.cpp
@@ -110,6 +110,7 @@ void FPhysicsAssetEditorCommands::RegisterCommands()
 	UI_COMMAND(SelectShapesPhysicsOnly, "Select PhysicsOnly Shapes", "Select Kinematic Bodies", EUserInterfaceActionType::Button, FInputChord());
 	UI_COMMAND(SelectShapesQueryAndProbe, "Select QueryAndProbe Shapes", "Select Kinematic Bodies", EUserInterfaceActionType::Button, FInputChord());
 	UI_COMMAND(SelectShapesProbeOnly, "Select ProbeOnly Shapes", "Select Kinematic Bodies", EUserInterfaceActionType::Button, FInputChord());
+	UI_COMMAND(ImportMLLevelSet, "Import MLLevelset Data", "Import MLLevelset Physical Asset", EUserInterfaceActionType::Button, FInputChord());
 	UI_COMMAND(SelectAllConstraints, "Select All Constraints", "Select All Constraints", EUserInterfaceActionType::Button, FInputChord(EModifierKey::Control | EModifierKey::Shift, EKeys::A));
 	UI_COMMAND(ToggleSelectionType, "Toggle Selection Type", "Select bodies from constraints, or constraints from bodies, depending on selection.", EUserInterfaceActionType::Button, FInputChord(EModifierKey::Control, EKeys::T));
 	UI_COMMAND(ToggleSelectionTypeWithUserConstraints, "Toggle Selection Type (With User Constraints)", "Select bodies from constraints, or constraints from bodies, depending on selection (account for user constraints).", EUserInterfaceActionType::Button, FInputChord(EModifierKey::Control | EModifierKey::Shift, EKeys::T));
diff --git a/Engine/Source/Editor/PhysicsAssetEditor/Private/PhysicsAssetEditorActions.h b/Engine/Source/Editor/PhysicsAssetEditor/Private/PhysicsAssetEditorActions.h
--- a/Engine/Source/Editor/PhysicsAssetEditor/Private/PhysicsAssetEditorActions.h
+++ b/Engine/Source/Editor/PhysicsAssetEditor/Private/PhysicsAssetEditorActions.h
@@ -124,6 +124,7 @@ public:
 	TSharedPtr<FUICommandInfo> SelectShapesPhysicsOnly;
 	TSharedPtr<FUICommandInfo> SelectShapesQueryAndProbe;
 	TSharedPtr<FUICommandInfo> SelectShapesProbeOnly;
+	TSharedPtr<FUICommandInfo> ImportMLLevelSet; 
 	TSharedPtr<FUICommandInfo> SelectAllConstraints;
 	TSharedPtr<FUICommandInfo> ToggleSelectionType;
 	TSharedPtr<FUICommandInfo> ToggleSelectionTypeWithUserConstraints;
diff --git a/Engine/Source/Editor/PhysicsAssetEditor/Private/PhysicsAssetEditorSkeletalMeshComponent.cpp b/Engine/Source/Editor/PhysicsAssetEditor/Private/PhysicsAssetEditorSkeletalMeshComponent.cpp
--- a/Engine/Source/Editor/PhysicsAssetEditor/Private/PhysicsAssetEditorSkeletalMeshComponent.cpp
+++ b/Engine/Source/Editor/PhysicsAssetEditor/Private/PhysicsAssetEditorSkeletalMeshComponent.cpp
@@ -23,6 +23,7 @@
 #include "Styling/AppStyle.h"
 #include "Chaos/WeightedLatticeImplicitObject.h"
 #include "Chaos/Levelset.h"
+#include "Chaos/MLLevelset.h"
 
 namespace
 {
@@ -248,6 +249,12 @@ FTransform UPhysicsAssetEditorSkeletalMeshComponent::GetPrimitiveTransform(const
 		PrimTM.ScaleTranslation(Scale3D);
 		return PrimTM * BoneTM;
 	}
+	else if (PrimType == EAggCollisionShape::MLLevelSet)
+	{
+		FTransform PrimTM = ManTM * SharedBodySetup->AggGeom.MLLevelSetElems[PrimIndex].GetTransform();
+		PrimTM.ScaleTranslation(Scale3D);
+		return PrimTM * BoneTM;
+	}
 
 	// Should never reach here
 	check(0);
@@ -372,6 +379,7 @@ void UPhysicsAssetEditorSkeletalMeshComponent::RefreshBoneTransforms(FActorCompo
 		bNeedToFlipSpaceBaseBuffers = true;
 	}
 	UpdateSkinnedLevelSets();
+	UpdateMLLevelSets();
 }
 
 void UPhysicsAssetEditorSkeletalMeshComponent::AddImpulseAtLocation(FVector Impulse, FVector Location, FName BoneName)
@@ -482,4 +490,47 @@ void UPhysicsAssetEditorSkeletalMeshComponent::UpdateSkinnedLevelSets()
 			}
 		}
 	}
+}
+
+void UPhysicsAssetEditorSkeletalMeshComponent::UpdateMLLevelSets()
+{
+	UPhysicsAsset* const PhysicsAsset = GetPhysicsAsset();
+	if (!PhysicsAsset)
+	{
+		return;
+	}
+	for (int32 i = 0; i < PhysicsAsset->SkeletalBodySetups.Num(); ++i)
+	{
+		const int32 BoneIndex = GetBoneIndex(PhysicsAsset->SkeletalBodySetups[i]->BoneName);
+		if (BoneIndex != INDEX_NONE)
+		{
+			FKAggregateGeom* const AggGeom = &PhysicsAsset->SkeletalBodySetups[i]->AggGeom;
+			if (AggGeom)
+			{
+				for (FKMLLevelSetElem& MLLevelsetElem : AggGeom->MLLevelSetElems)
+				{
+					if (MLLevelsetElem.GetMLLevelSet().IsValid())
+					{
+						const TArray<FName>& ActiveBoneNames = MLLevelsetElem.GetMLLevelSet()->GetActiveBoneNames();
+						TArray<FTransform> RelativeActiveBoneTransforms;
+						RelativeActiveBoneTransforms.SetNum(ActiveBoneNames.Num());
+						const FTransform ParentRootTransformInv = GetBoneTransform(BoneIndex, FTransform::Identity).Inverse();
+						for (int32 ActiveBoneIndex = 0; ActiveBoneIndex < ActiveBoneNames.Num(); ActiveBoneIndex++)
+						{
+							int32 LocalActiveBoneIndex = GetBoneIndex(ActiveBoneNames[ActiveBoneIndex]);
+							if (LocalActiveBoneIndex != INDEX_NONE)
+							{
+								RelativeActiveBoneTransforms[ActiveBoneIndex] = GetBoneTransform(LocalActiveBoneIndex, ParentRootTransformInv);							
+							}
+							else
+							{
+								RelativeActiveBoneTransforms[ActiveBoneIndex] = ParentRootTransformInv;
+							}							
+						}
+						MLLevelsetElem.GetMLLevelSet()->UpdateActiveBonesRelativeTransformsAndUpdateDebugPhi(RelativeActiveBoneTransforms);
+					}
+				}
+			}
+		}
+	}
 }
\ No newline at end of file
diff --git a/Engine/Source/Editor/PhysicsAssetEditor/Private/PhysicsAssetEditorSkeletalMeshComponent.h b/Engine/Source/Editor/PhysicsAssetEditor/Private/PhysicsAssetEditorSkeletalMeshComponent.h
--- a/Engine/Source/Editor/PhysicsAssetEditor/Private/PhysicsAssetEditorSkeletalMeshComponent.h
+++ b/Engine/Source/Editor/PhysicsAssetEditor/Private/PhysicsAssetEditorSkeletalMeshComponent.h
@@ -81,4 +81,5 @@ public:
 
 private:
 	void UpdateSkinnedLevelSets();
+	void UpdateMLLevelSets();
 };
diff --git a/Engine/Source/Programs/UnrealBuildTool/Configuration/ModuleRules.cs b/Engine/Source/Programs/UnrealBuildTool/Configuration/ModuleRules.cs
--- a/Engine/Source/Programs/UnrealBuildTool/Configuration/ModuleRules.cs
+++ b/Engine/Source/Programs/UnrealBuildTool/Configuration/ModuleRules.cs
@@ -1659,6 +1659,7 @@ namespace UnrealBuildTool
 					"PhysicsCore",
 					"ChaosCore",
 					"Chaos",
+					"NNE",
 					}
 				);
 
diff --git a/Engine/Source/Runtime/Engine/Classes/PhysicsEngine/AggregateGeom.h b/Engine/Source/Runtime/Engine/Classes/PhysicsEngine/AggregateGeom.h
--- a/Engine/Source/Runtime/Engine/Classes/PhysicsEngine/AggregateGeom.h
+++ b/Engine/Source/Runtime/Engine/Classes/PhysicsEngine/AggregateGeom.h
@@ -6,6 +6,7 @@
 #include "UObject/ObjectMacros.h"
 #include "PhysicsEngine/ConvexElem.h"
 #include "PhysicsEngine/LevelSetElem.h"
+#include "PhysicsEngine/MLLevelSetElem.h"
 #include "PhysicsEngine/BoxElem.h"
 #include "PhysicsEngine/SkinnedLevelSetElem.h"
 #include "PhysicsEngine/SphereElem.h"
@@ -43,6 +44,9 @@ struct FKAggregateGeom
 	UPROPERTY(EditAnywhere, editfixedsize, Category = "Aggregate Geometry", meta = (DisplayName = "(Experimental) Skinned Level Sets", TitleProperty = "Name"), Experimental)
 	TArray<FKSkinnedLevelSetElem> SkinnedLevelSetElems;
 
+	UPROPERTY(EditAnywhere, editfixedsize, Category = "Aggregate Geometry", meta = (DisplayName = "(Experimental) ML Level Sets"))
+	TArray<FKMLLevelSetElem> MLLevelSetElems;
+
 	FKAggregateGeom()
 		: RenderInfoPtr(nullptr)
 	{
@@ -63,12 +67,12 @@ struct FKAggregateGeom
 
 	int32 GetElementCount() const
 	{
-		return SphereElems.Num() + SphylElems.Num() + BoxElems.Num() + ConvexElems.Num() + TaperedCapsuleElems.Num() + LevelSetElems.Num() + SkinnedLevelSetElems.Num();
+		return SphereElems.Num() + SphylElems.Num() + BoxElems.Num() + ConvexElems.Num() + TaperedCapsuleElems.Num() + LevelSetElems.Num() + SkinnedLevelSetElems.Num() + MLLevelSetElems.Num();
 	}
 
 	ENGINE_API int32 GetElementCount(EAggCollisionShape::Type Type) const;
 
-	SIZE_T GetAllocatedSize() const { return SphereElems.GetAllocatedSize() + SphylElems.GetAllocatedSize() + BoxElems.GetAllocatedSize() + ConvexElems.GetAllocatedSize() + TaperedCapsuleElems.GetAllocatedSize() + LevelSetElems.GetAllocatedSize() + SkinnedLevelSetElems.GetAllocatedSize(); }
+	SIZE_T GetAllocatedSize() const { return SphereElems.GetAllocatedSize() + SphylElems.GetAllocatedSize() + BoxElems.GetAllocatedSize() + ConvexElems.GetAllocatedSize() + TaperedCapsuleElems.GetAllocatedSize() + LevelSetElems.GetAllocatedSize() + SkinnedLevelSetElems.GetAllocatedSize() + MLLevelSetElems.GetAllocatedSize(); }
 
 	FKShapeElem* GetElement(const EAggCollisionShape::Type Type, const int32 Index)
 	{
@@ -88,6 +92,8 @@ struct FKAggregateGeom
 			if (ensure(LevelSetElems.IsValidIndex(Index))) { return &LevelSetElems[Index]; }
 		case EAggCollisionShape::SkinnedLevelSet:
 			if (ensure(SkinnedLevelSetElems.IsValidIndex(Index))) { return &SkinnedLevelSetElems[Index]; }
+		case EAggCollisionShape::MLLevelSet:
+			if (ensure(MLLevelSetElems.IsValidIndex(Index))) { return &MLLevelSetElems[Index]; }
 		default:
 			ensure(false);
 			return nullptr;
@@ -130,6 +136,8 @@ struct FKAggregateGeom
 		if (Index < LevelSetElems.Num()) { return &LevelSetElems[Index]; }
 		Index -= LevelSetElems.Num();
 		if (Index < SkinnedLevelSetElems.Num()) { return &SkinnedLevelSetElems[Index]; }
+		Index -= SkinnedLevelSetElems.Num();
+		if (Index < MLLevelSetElems.Num()) { return &MLLevelSetElems[Index]; }
 		ensure(false);
 		return nullptr;
 	}
@@ -164,6 +172,10 @@ struct FKAggregateGeom
 		{
 			return FoundSkinnedLevelSetElem;
 		}
+		else if (const FKShapeElem* FoundMLLevelSetElem = GetElementByName<FKMLLevelSetElem>(MakeArrayView(MLLevelSetElems), InName))
+		{
+			return FoundMLLevelSetElem;
+		}
 
 		return nullptr;
 	}
@@ -219,6 +231,12 @@ struct FKAggregateGeom
 			return FoundIndex + StartIndex;
 		}
 
+		FoundIndex = GetElementIndexByName<FKMLLevelSetElem>(MakeArrayView(MLLevelSetElems), InName);
+		if (FoundIndex != INDEX_NONE)
+		{
+			return FoundIndex;
+		}
+
 		return INDEX_NONE;
 	}
 
@@ -253,7 +271,7 @@ struct FKAggregateGeom
 		TaperedCapsuleElems.Empty();
 		LevelSetElems.Empty();
 		SkinnedLevelSetElems.Empty();
-
+		MLLevelSetElems.Empty();
 		FreeRenderInfo();
 	}
 
@@ -298,6 +316,7 @@ private:
 		TaperedCapsuleElems = Other.TaperedCapsuleElems;
 		LevelSetElems = Other.LevelSetElems;
 		SkinnedLevelSetElems = Other.SkinnedLevelSetElems;
+		MLLevelSetElems = Other.MLLevelSetElems;
 	}
 
 	template <class T>
diff --git a/Engine/Source/Runtime/Engine/Classes/PhysicsEngine/MLLevelSetElem.h b/Engine/Source/Runtime/Engine/Classes/PhysicsEngine/MLLevelSetElem.h
new file mode 100644
--- /dev/null
+++ b/Engine/Source/Runtime/Engine/Classes/PhysicsEngine/MLLevelSetElem.h
@@ -0,0 +1,84 @@
+// Copyright Epic Games, Inc. All Rights Reserved.
+#pragma once
+
+#include "CoreMinimal.h"
+#include "UObject/ObjectMacros.h"
+#include "PhysicsEngine/ShapeElem.h"
+#include "PhysicsEngine/MLLevelSetModelAndBonesBinningInfo.h"
+#include "NNEModelData.h"  
+#include "NNERuntimeCPU.h"
+#include "Chaos/MLLevelset.h" 
+#include "MLLevelSetElem.generated.h"
+
+extern ENGINE_API bool bEnableMLLevelSet;
+
+USTRUCT()
+struct FKMLLevelSetElem : public FKShapeElem
+{
+	GENERATED_BODY()
+
+	FKMLLevelSetElem() :
+		FKShapeElem(EAggCollisionShape::MLLevelSet)
+	{}
+
+	FKMLLevelSetElem(const FKMLLevelSetElem& Other)
+	{
+		CloneElem(Other);
+	}
+
+	const FKMLLevelSetElem& operator=(const FKMLLevelSetElem& Other)
+	{
+		CloneElem(Other);
+		return *this;
+	}
+
+	ENGINE_API void BuildMLLevelSet(Chaos::FMLLevelSetImportData&& MLLevelSetImportData);
+
+	FTransform GetTransform() const { return FTransform(); }
+
+	// Draw helpers
+	/** Get geometry of all cells where the level set function is less than or equal to InteriorThreshold */
+	ENGINE_API void GetInteriorGridCells(TArray<FBox>& CellBoxes, double InteriorThreshold = 0.0) const;
+
+	/** Get geometry of all cell faces where level set function changes sign */
+	ENGINE_API void GetZeroIsosurfaceGridCellFaces(TArray<FVector3f>& Vertices, TArray<FIntVector>& Tris) const;
+
+	// Draw functions
+	ENGINE_API virtual void DrawElemWire(class FPrimitiveDrawInterface* PDI, const FTransform& ElemTM, const float Scale, const FColor Color) const override;
+	ENGINE_API virtual void DrawElemSolid(class FPrimitiveDrawInterface* PDI, const FTransform& ElemTM, const float Scale, const FMaterialRenderProxy* MaterialRenderProxy) const override;
+	ENGINE_API void GetElemSolid(const FTransform& ElemTM, const FVector& Scale3D, const FMaterialRenderProxy* MaterialRenderProxy, int32 ViewIndex, class FMeshElementCollector& Collector) const;
+
+	ENGINE_API FBox CalcAABB(const FTransform& BoneTM, const FVector& Scale3D) const;
+	ENGINE_API FBox UntransformedAABB() const;
+	ENGINE_API FIntVector3 GridResolution() const;
+
+	bool Serialize(FArchive& Ar);
+
+	const TSharedPtr<Chaos::FMLLevelSet, ESPMode::ThreadSafe> GetMLLevelSet() const
+	{
+		return MLLevelSet;
+	}
+
+private:
+	TSharedPtr<Chaos::FMLLevelSet, ESPMode::ThreadSafe> MLLevelSet;
+
+	UPROPERTY()
+	TObjectPtr<UNNEModelData> NNESignedDistanceModelData;
+
+	UPROPERTY()
+	TObjectPtr<UNNEModelData> NNEIncorrectZoneModelData;
+
+	/** Helper function to safely copy instances of this shape*/
+	ENGINE_API void CloneElem(const FKMLLevelSetElem& Other);
+};
+
+/* Enable our own serialization function to handle FMLLevelSet.*/
+template<>
+struct TStructOpsTypeTraits<FKMLLevelSetElem> : public TStructOpsTypeTraitsBase2<FKMLLevelSetElem>
+{
+	enum
+	{
+		WithSerializer = true
+	};
+	static constexpr EPropertyObjectReferenceType WithSerializerObjectReferences = EPropertyObjectReferenceType::None;
+};
diff --git a/Engine/Source/Runtime/Engine/Classes/PhysicsEngine/MLLevelSetModelAndBonesBinningInfo.h b/Engine/Source/Runtime/Engine/Classes/PhysicsEngine/MLLevelSetModelAndBonesBinningInfo.h
new file mode 100644
--- /dev/null
+++ b/Engine/Source/Runtime/Engine/Classes/PhysicsEngine/MLLevelSetModelAndBonesBinningInfo.h
@@ -0,0 +1,93 @@
+// Copyright Epic Games, Inc. All Rights Reserved.
+#pragma once
+
+#include "CoreMinimal.h"
+#include "Engine/DataTable.h"
+#include "MLLevelSetModelAndBonesBinningInfo.generated.h"
+
+USTRUCT(BlueprintType)
+struct FMLLevelSetModelAndBonesBinningInfo : public FTableRowBase
+{
+	GENERATED_BODY();
+
+	/* The bone that MLLevelSet is attached to. Note that the deformations near this joint is not trained*/
+	UPROPERTY(EditAnywhere, BlueprintReadWrite, Category = MLModelBone)
+	FString ParentBoneName = "";
+
+	/* The bones that are trained for deformation. We suggest only train one bone per MLLevelSet for efficiency*/
+	UPROPERTY(EditAnywhere, BlueprintReadWrite, Category = MLModelBone)
+	FString ActiveBoneNames = "";
+
+	/* The path to the DataTable that includes information about NNE Model */
+	UPROPERTY(EditAnywhere, BlueprintReadWrite, Category = MLModelBone)
+	FString MLModelInferenceInfoDataTablePath = "";
+
+	/* The index of the DataTable that includes information about NNE Model for signed distance */
+	UPROPERTY(EditAnywhere, BlueprintReadWrite, Category = MLModelBone)
+	FString MLModelInferenceInfoDataTableIndex = "";
+
+	/* The index of the DataTable that includes information about NNE Model for incorrect (safe-danger) zone. If uninitialized, this means no safe zone is used. */
+	UPROPERTY(EditAnywhere, BlueprintReadWrite, Category = MLModelBone)
+	FString MLModelInferenceForIncorrectZoneInfoDataTableIndex = "";
+
+	/* Model is trained for the subset of the rotations for each active bone. */
+	/* E.g. Use {1,2} if ActiveBone1 has one and ActiveBone2 has two active rotations. */
+	UPROPERTY(EditAnywhere, BlueprintReadWrite, Category = MLModelBone)
+	TArray<int32> NumberOfRotationComponentsPerBone = {};
+
+	/* The indices of the rotation components. E.g. use {1,1,2} if ActiveBone1 uses Rot.Y and ActiveBone2 uses Rot.Y and Rot.Z. */
+	UPROPERTY(EditAnywhere, BlueprintReadWrite, Category = MLModelBone)
+	TArray<int32> RotationComponentIndexes = {};
+
+	/* Usually MLModels for SDF are trained so that output lies in [-1,1]. To do so the SignedDistances (in the dataset) are divided by SignedDistanceScaling. 
+	SignedDistanceScaling is generally defined to be the max length of the training bounding box. */
+	UPROPERTY(EditAnywhere, BlueprintReadWrite, Category = MLModelBone)
+	double SignedDistanceScaling = -1.0;
+
+	/* Resolution of the Grid that is used for debug visualization */
+	UPROPERTY(EditAnywhere, BlueprintReadWrite, Category = MLModelBone)
+	TArray<int32> DebugGridResolution = {50,50,50};
+
+	/* Reference Rotations for the Active Bones. */
+	UPROPERTY(EditAnywhere, BlueprintReadWrite, Category = MLModelBone)
+	TArray<double> ReferenceBoneRotations = {};
+	
+	/* Reference Translations for the Active Bones. */
+	UPROPERTY(EditAnywhere, BlueprintReadWrite, Category = MLModelBone)
+	TArray<double> ReferenceBoneTranslations = {};
+
+	/* Min Corner of bounding box that MLLevelSet is trained on. */
+	UPROPERTY(EditAnywhere, BlueprintReadWrite, Category = MLModelBone)
+	TArray<float> TrainingGridOrigin = {};
+
+	/* First Edge of the bounding box that MLLevelSet is trained on. */
+	UPROPERTY(EditAnywhere, BlueprintReadWrite, Category = MLModelBone)
+	TArray<float> TrainingGridAxisX = {};
+
+	/* Second Edge of the bounding box that MLLevelSet is trained on. */
+	UPROPERTY(EditAnywhere, BlueprintReadWrite, Category = MLModelBone)
+	TArray<float> TrainingGridAxisY = {};
+
+	/* Third Edge of the bounding box that MLLevelSet is trained on. */
+	UPROPERTY(EditAnywhere, BlueprintReadWrite, Category = MLModelBone)
+	TArray<float> TrainingGridAxisZ = {};
+};
+
+USTRUCT(BlueprintType)
+struct FMLLevelSetModelInferenceInfo : public FTableRowBase
+{
+	GENERATED_BODY();
+
+	/* The path to the NNE model. At the moment MLLevelSet Asset accepts only mlir.tosa models. */
+	UPROPERTY(EditAnywhere, BlueprintReadWrite, Category = MLModelBone)
+	FString NNEModelPath = "";
+
+	/* Model Architecture. Number of neurons in each layer, including input layer and the output layer. */
+	/* In most cases input layer is 3 (location input, X,Y,Z coordinates of the query), and output layer is 1 (Scaled Signed Distance or Incorrect Zone Indicator)*/
+	UPROPERTY(EditAnywhere, BlueprintReadWrite, Category = MLModelBone)
+	TArray<int32> ModelArchitectureActivationNodeSizes = {};
+
+	/* Model weights {W1,..,Wn} is tokanized into FString as "W1_0,W1_1,...,W1_k1|W2_0,W2_1,...,W2_k2|...|Wn_0,Wn_1,...,Wn_kn" */
+	UPROPERTY(EditAnywhere, BlueprintReadWrite, Category = MLModelBone)
+	FString MLModelWeights = "";
+};
\ No newline at end of file
diff --git a/Engine/Source/Runtime/Engine/Classes/PhysicsEngine/ShapeElem.h b/Engine/Source/Runtime/Engine/Classes/PhysicsEngine/ShapeElem.h
--- a/Engine/Source/Runtime/Engine/Classes/PhysicsEngine/ShapeElem.h
+++ b/Engine/Source/Runtime/Engine/Classes/PhysicsEngine/ShapeElem.h
@@ -21,6 +21,7 @@ namespace EAggCollisionShape
 		TaperedCapsule,
 		LevelSet,
 		SkinnedLevelSet,
+		MLLevelSet,
 
 		Unknown
 	};
diff --git a/Engine/Source/Runtime/Engine/Private/Collision/KAggregateGeom.cpp b/Engine/Source/Runtime/Engine/Private/Collision/KAggregateGeom.cpp
--- a/Engine/Source/Runtime/Engine/Private/Collision/KAggregateGeom.cpp
+++ b/Engine/Source/Runtime/Engine/Private/Collision/KAggregateGeom.cpp
@@ -20,11 +20,15 @@
 #include "PhysicsEngine/TaperedCapsuleElem.h"
 #include "Chaos/WeightedLatticeImplicitObject.h"
 #include "AutoRTFM/AutoRTFM.h"
+#include "Chaos/MLLevelset.h"
+#include "UObject/Package.h"
+
 #if INTEL_ISPC
 #include "KAggregateGeom.ispc.generated.h"
 #endif
 
-
+bool bEnableMLLevelSet = false;
+static FAutoConsoleVariableRef CVarShowExperimentalImportMLLevelSetButton(TEXT("p.Experimental.EnableMLLevelSet"), bEnableMLLevelSet, TEXT("Shows Import MLLevelSet button. [def: false]"));
 
 #define MIN_HULL_VERT_DISTANCE		(0.1f)
 #define MIN_HULL_VALID_DIMENSION	(0.5f)
@@ -121,6 +125,11 @@ FBox FKAggregateGeom::CalcAABB(const FTransform& Transform) const
 		Box += SkinnedLevelSetElems[i].CalcAABB(BoneTM, (FVector)Scale3D);
 	}
 
+	for (int32 i = 0; i < MLLevelSetElems.Num(); i++)
+	{
+		Box += MLLevelSetElems[i].CalcAABB(BoneTM, (FVector)Scale3D);
+	}
+
 	return Box;
 
 }
@@ -688,6 +697,93 @@ FTransform FKSkinnedLevelSetElem::GetTransform() const
 	return FTransform();
 }
 
+///////////////////////////////////////
+//////// FKMLLevelSetElem ///////////////
+///////////////////////////////////////
+
+
+FBox FKMLLevelSetElem::CalcAABB(const FTransform& BoneTM, const FVector& Scale3D) const
+{
+	FBox Box(ForceInit);
+
+	if (MLLevelSet.IsValid())
+	{
+		Box = FBox(MLLevelSet->BoundingBox().Min(), MLLevelSet->BoundingBox().Max());
+		const FTransform LocalToWorld = FTransform(FQuat::Identity, FVector::ZeroVector, Scale3D) * BoneTM;
+		return Box.TransformBy(LocalToWorld);
+	}
+
+	return Box;
+}
+
+FBox FKMLLevelSetElem::UntransformedAABB() const
+{
+	FBox Box(ForceInit);
+
+	if (MLLevelSet.IsValid())
+	{
+		Box = FBox(MLLevelSet->BoundingBox().Min(), MLLevelSet->BoundingBox().Max());
+		return Box;
+	}
+
+	return Box;
+}
+
+FIntVector3 FKMLLevelSetElem::GridResolution() const
+{
+	if (MLLevelSet.IsValid())
+	{
+		const Chaos::TVec3<int32>& Dim = MLLevelSet->GetGrid().Counts();
+		return FIntVector3(Dim[0], Dim[1], Dim[2]);
+	}
+
+	return FIntVector3(0, 0, 0);
+}
+
+void FKMLLevelSetElem::BuildMLLevelSet(Chaos::FMLLevelSetImportData&& MLLevelSetImportData)
+{
+	NNESignedDistanceModelData = LoadObject<UNNEModelData>(GetTransientPackage(), *MLLevelSetImportData.NNEModelDataArr[0].NNEModelPath);
+	NNEIncorrectZoneModelData = LoadObject<UNNEModelData>(GetTransientPackage(), *MLLevelSetImportData.NNEModelDataArr[1].NNEModelPath);
+		
+	MLLevelSetImportData.NNEModelDataArr[0].NNEModelData = NNESignedDistanceModelData;
+	MLLevelSetImportData.NNEModelDataArr[1].NNEModelData = NNEIncorrectZoneModelData;
+
+	MLLevelSet = TSharedPtr<Chaos::FMLLevelSet>(new Chaos::FMLLevelSet(MoveTemp(MLLevelSetImportData)));
+}
+
+void FKMLLevelSetElem::GetInteriorGridCells(TArray<FBox>& CellBoxes, double InteriorThreshold) const
+{
+	if (MLLevelSet.IsValid())
+	{
+		const Chaos::TUniformGrid<Chaos::FReal, 3>& Grid = MLLevelSet->GetGrid();
+		const Chaos::TVector<int32, 3> Cells = Grid.Counts();
+		for (int i = 0; i < Cells.X; ++i)
+		{
+			for (int j = 0; j < Cells.Y; ++j)
+			{
+				for (int k = 0; k < Cells.Z; ++k)
+				{
+					const double Value = MLLevelSet->GetPhiArray()(i, j, k);
+
+					if (Value <= InteriorThreshold)
+					{
+						const FVector3d CellMin = Grid.MinCorner() + Grid.Dx() * FVector3d(i, j, k);
+						const FVector3d CellMax = CellMin + Grid.Dx() * FVector3d(1, 1, 1);
+						CellBoxes.Emplace(FBox(CellMin, CellMax));
+					}
+				}
+			}
+		}
+	}
+}
+
+void FKMLLevelSetElem::GetZeroIsosurfaceGridCellFaces(TArray<FVector3f>& Vertices, TArray<FIntVector>& Tris) const
+{
+	if (MLLevelSet.IsValid())
+	{
+		MLLevelSet->GetZeroIsosurfaceGridCellFaces(Vertices, Tris);
+	}
+}
 
 static const float DIST_COMPARE_THRESH = 0.1f;
 static const float DIR_COMPARE_THRESH = 0.0003f; // about 1 degree
@@ -1031,3 +1127,21 @@ bool FKSkinnedLevelSetElem::Serialize(FArchive& Ar)
 	return true;
 }
 
+bool FKMLLevelSetElem::Serialize(FArchive& Ar)
+{
+	Ar << NNESignedDistanceModelData;
+	Ar << NNEIncorrectZoneModelData;
+
+	if (Ar.IsLoading())
+	{
+		MLLevelSet = TSharedPtr<Chaos::FMLLevelSet>(new Chaos::FMLLevelSet(NNESignedDistanceModelData, NNEIncorrectZoneModelData));
+	}
+
+	if (MLLevelSet.IsValid())
+	{
+		MLLevelSet->Serialize(Ar);
+	}
+
+	return true;
+}
+
diff --git a/Engine/Source/Runtime/Engine/Private/PhysicsEngine/BodySetup.cpp b/Engine/Source/Runtime/Engine/Private/PhysicsEngine/BodySetup.cpp
--- a/Engine/Source/Runtime/Engine/Private/PhysicsEngine/BodySetup.cpp
+++ b/Engine/Source/Runtime/Engine/Private/PhysicsEngine/BodySetup.cpp
@@ -39,6 +39,7 @@
 #include "PhysicsEngine/Experimental/ChaosCooking.h"
 #include "PhysicsEngine/SphylElem.h"
 #include "PhysicsEngine/TaperedCapsuleElem.h"
+#include "PhysicsEngine/MLLevelSetElem.h"
 
 #if WITH_EDITOR
 #include "Misc/DataValidation.h"
@@ -184,6 +185,7 @@ void UBodySetup::AddCollisionFrom(const FKAggregateGeom& FromAggGeom)
 	}
 
 	AggGeom.LevelSetElems.Append(FromAggGeom.LevelSetElems);
+	AggGeom.MLLevelSetElems.Append(FromAggGeom.MLLevelSetElems);
 }
 
 namespace
@@ -2327,6 +2329,14 @@ void FKSkinnedLevelSetElem::CloneElem(const FKSkinnedLevelSetElem& Other)
 	WeightedLatticeLevelSet = Other.WeightedLatticeLevelSet;
 }
 
+void FKMLLevelSetElem::CloneElem(const FKMLLevelSetElem& Other)
+{
+	Super::CloneElem(Other);
+	MLLevelSet = Other.MLLevelSet;
+	NNESignedDistanceModelData = Other.NNESignedDistanceModelData;
+	NNEIncorrectZoneModelData = Other.NNEIncorrectZoneModelData;
+}
+
 #if WITH_EDITOR
 
 bool FBodySetupObjectTextFactory::CanCreateClass(UClass* InObjectClass, bool& bOmitSubObjs) const
diff --git a/Engine/Source/Runtime/Engine/Private/PhysicsEngine/PhysDrawing.cpp b/Engine/Source/Runtime/Engine/Private/PhysicsEngine/PhysDrawing.cpp
--- a/Engine/Source/Runtime/Engine/Private/PhysicsEngine/PhysDrawing.cpp
+++ b/Engine/Source/Runtime/Engine/Private/PhysicsEngine/PhysDrawing.cpp
@@ -1034,6 +1034,72 @@ void FKSkinnedLevelSetElem::GetElemSolid(const FTransform& ElemTM, const FVector
 	}
 }
 
+/////////////////////////////////////////////////////////////////////////////////////
+// FKMLLevelSetElem
+/////////////////////////////////////////////////////////////////////////////////////
+
+void FKMLLevelSetElem::DrawElemWire(class FPrimitiveDrawInterface* PDI, const FTransform& ElemTM, float Scale, const FColor Color) const
+{
+	TArray<FBox> Boxes;
+
+	if (GridResolution().GetMax() < 30)	// Don't render individual cells if the resolution is too high
+	{
+		// Cells with negative Phi values
+		const double Threshold = UE_KINDA_SMALL_NUMBER;		// allow slightly greater than zero for visualization purposes
+		GetInteriorGridCells(Boxes, UE_KINDA_SMALL_NUMBER);
+	}
+	else
+	{
+		// Just render the bounding box of the grid itself
+		Boxes.Add(UntransformedAABB());
+	}
+
+	PDI->AddReserveLines(SDPG_World, Boxes.Num() * 12);
+	for (const FBox& Box : Boxes)
+	{
+		constexpr float Thickness = 0.f;
+		DrawWireBox(PDI, ElemTM.ToMatrixWithScale(), Box, Color, SDPG_World, Thickness);
+	}
+}
+
+void FKMLLevelSetElem::DrawElemSolid(class FPrimitiveDrawInterface* PDI, const FTransform& ElemTM, float Scale, const FMaterialRenderProxy* MaterialRenderProxy) const
+{
+	TArray<FVector3f> Vertices;
+	TArray<FIntVector> Tris;
+	GetZeroIsosurfaceGridCellFaces(Vertices, Tris);
+
+	FDynamicMeshBuilder MeshBuilder(PDI->View->GetFeatureLevel());
+	for (const FVector3f& V : Vertices)
+	{
+		MeshBuilder.AddVertex(FDynamicMeshVertex(V));
+	}
+	for (const FIntVector& T : Tris)
+	{
+		MeshBuilder.AddTriangle(T[0], T[1], T[2]);
+	}
+
+	MeshBuilder.Draw(PDI, ElemTM.ToMatrixWithScale(), MaterialRenderProxy, SDPG_World, 0.f);
+}
+
+void FKMLLevelSetElem::GetElemSolid(const FTransform& ElemTM, const FVector& Scale3D, const FMaterialRenderProxy* MaterialRenderProxy, int32 ViewIndex, class FMeshElementCollector& Collector) const
+{
+	TArray<FVector3f> Vertices;
+	TArray<FIntVector> Tris;
+	GetZeroIsosurfaceGridCellFaces(Vertices, Tris);
+
+	FDynamicMeshBuilder MeshBuilder(Collector.GetFeatureLevel());
+	for (const FVector3f& V : Vertices)
+	{
+		MeshBuilder.AddVertex(FDynamicMeshVertex(V));
+	}
+	for (const FIntVector& T : Tris)
+	{
+		MeshBuilder.AddTriangle(T[0], T[1], T[2]);
+	}
+
+	MeshBuilder.GetMesh(ElemTM.ToMatrixWithScale(), MaterialRenderProxy, SDPG_World, false, false, ViewIndex, Collector);
+}
+
 /////////////////////////////////////////////////////////////////////////////////////
 // FKConvexGeomRenderInfo
 /////////////////////////////////////////////////////////////////////////////////////
diff --git a/Engine/Source/Runtime/Experimental/Chaos/Chaos.Build.cs b/Engine/Source/Runtime/Experimental/Chaos/Chaos.Build.cs
--- a/Engine/Source/Runtime/Experimental/Chaos/Chaos.Build.cs
+++ b/Engine/Source/Runtime/Experimental/Chaos/Chaos.Build.cs
@@ -19,9 +19,10 @@ namespace UnrealBuildTool.Rules
                 "IntelISPC",
 				"TraceLog",
                 "Voronoi",
-				"GeometryCore"
+				"GeometryCore",
+				"NNE",
 				}
-            );
+			);
 			PrivateDependencyModuleNames.AddRange(
 				new string[] {
 				"Eigen"
diff --git a/Engine/Source/Runtime/Experimental/Chaos/Private/Chaos/ImplicitObject.cpp b/Engine/Source/Runtime/Experimental/Chaos/Private/Chaos/ImplicitObject.cpp
--- a/Engine/Source/Runtime/Experimental/Chaos/Private/Chaos/ImplicitObject.cpp
+++ b/Engine/Source/Runtime/Experimental/Chaos/Private/Chaos/ImplicitObject.cpp
@@ -18,6 +18,7 @@
 #include "HAL/IConsoleManager.h"
 #include "UObject/DestructionObjectVersion.h"
 #include "UObject/ReleaseObjectVersion.h"
+#include "Chaos/MLLevelset.h" 
 
 namespace Chaos
 {
@@ -370,6 +371,7 @@ namespace Chaos
 		static const FName TransformedName = TEXT("Transformed");
 		static const FName UnionName = TEXT("Union");
 		static const FName LevelSetName = TEXT("LevelSet");
+		static const FName MLLevelSetName = TEXT("MLLevelSet");
 		static const FName UnknownName = TEXT("Unknown");
 		static const FName ConvexName = TEXT("Convex");
 		static const FName TaperedCylinderName = TEXT("TaperedCylinder");
@@ -388,6 +390,7 @@ namespace Chaos
 		case ImplicitObjectType::Transformed: return TransformedName;
 		case ImplicitObjectType::Union: return UnionName;
 		case ImplicitObjectType::LevelSet: return LevelSetName;
+		case ImplicitObjectType::MLLevelSet: return MLLevelSetName;
 		case ImplicitObjectType::Unknown: return UnknownName;
 		case ImplicitObjectType::Convex: return ConvexName;
 		case ImplicitObjectType::TaperedCylinder: return TaperedCylinderName;
diff --git a/Engine/Source/Runtime/Experimental/Chaos/Private/Chaos/MLLevelSetNeuralInference.cpp b/Engine/Source/Runtime/Experimental/Chaos/Private/Chaos/MLLevelSetNeuralInference.cpp
new file mode 100644
--- /dev/null
+++ b/Engine/Source/Runtime/Experimental/Chaos/Private/Chaos/MLLevelSetNeuralInference.cpp
@@ -0,0 +1,103 @@
+#include "Chaos/MLLevelSetNeuralInference.h"
+
+namespace Chaos
+{
+
+CHAOS_API FMLLevelSetNeuralInference::FMLLevelSetNeuralInference(TSharedPtr<UE::NNE::IModelCPU> NNEModel, TArray<TArray<int32>>& InModelWeightsShapes)
+	: NNEModel(NNEModel)
+	, ModelWeightsShapes(InModelWeightsShapes)
+{
+	ModelInstance = NNEModel->CreateModelInstanceCPU();
+}
+
+CHAOS_API FMLLevelSetNeuralInference::FMLLevelSetNeuralInference(TSharedPtr<UE::NNE::IModelCPU> NNEModel)
+	: NNEModel(NNEModel)
+{
+	ModelInstance = NNEModel->CreateModelInstanceCPU();
+}
+
+CHAOS_API FMLLevelSetNeuralInference::FMLLevelSetNeuralInference(FMLLevelSetNeuralInference&& Other)
+	: NNEModel(Other.NNEModel)
+	, ModelInstance(Other.ModelInstance)
+	, ModelWeightsShapes(MoveTemp(Other.ModelWeightsShapes))
+{}
+
+CHAOS_API FMLLevelSetNeuralInference::FMLLevelSetNeuralInference(const FMLLevelSetNeuralInference& Other)
+	: NNEModel(Other.NNEModel)
+	, ModelWeightsShapes(Other.ModelWeightsShapes)
+{
+	ModelInstance = NNEModel->CreateModelInstanceCPU();
+}
+
+FMLLevelSetNeuralInference& FMLLevelSetNeuralInference::operator=(const FMLLevelSetNeuralInference& Other)
+{
+	NNEModel = Other.NNEModel;
+	ModelWeightsShapes = Other.ModelWeightsShapes;
+	ModelInstance = NNEModel->CreateModelInstanceCPU();
+	return *this;
+}
+
+FMLLevelSetNeuralInference* FMLLevelSetNeuralInference::Copy() const
+{
+	FMLLevelSetNeuralInference* CopyInstance = new FMLLevelSetNeuralInference();
+	CopyInstance->NNEModel = NNEModel;
+	CopyInstance->ModelWeightsShapes = ModelWeightsShapes;
+	CopyInstance->ModelInstance = NNEModel->CreateModelInstanceCPU();
+	return CopyInstance;
+}
+
+void FMLLevelSetNeuralInference::RunInference(TArray<float, TAlignedHeapAllocator<64>>& InputData, TArray<float, TAlignedHeapAllocator<64>>& OutputData, 
+	uint32 SingleInputSize, uint32 SingleOutputSize, TArray<TArray<float, TAlignedHeapAllocator<64>>>& ModelWeightsIn) const
+{
+	TRACE_CPUPROFILER_EVENT_SCOPE(FMLLevelSetNeuralInference_RunInference);
+
+	if (ModelInstance.IsValid())
+	{
+		TRACE_CPUPROFILER_EVENT_SCOPE(FMLLevelSetNeuralInference_RunInference_Run);
+		const uint32 NumberOfInputs = uint32(InputData.Num()) / SingleInputSize;
+		TArray<uint32> InputShapeData = { NumberOfInputs , SingleInputSize };
+		TArray< UE::NNE::FTensorShape> InputTensorShapes;
+		InputTensorShapes.SetNum(ModelWeightsShapes.Num() + 1); //Including Input
+		InputTensorShapes[0] = UE::NNE::FTensorShape::Make(InputShapeData);
+		for (int32 Idx = 0; Idx < ModelWeightsShapes.Num(); Idx++)
+		{
+			InputTensorShapes[Idx + 1] = UE::NNE::FTensorShape::Make(ModelWeightsShapes[Idx]);
+		}
+		ModelInstance->SetInputTensorShapes(InputTensorShapes);
+
+		TArray<uint32> OutputShapeData = { NumberOfInputs, SingleOutputSize };
+		TArray< UE::NNE::FTensorShape > OutputTensorShapes = { UE::NNE::FTensorShape::Make(OutputShapeData) };
+
+		TArray<UE::NNE::FTensorBindingCPU> InputBindings;
+		TArray<UE::NNE::FTensorBindingCPU> OutputBindings;
+
+		InputBindings.SetNumZeroed(ModelWeightsIn.Num() + 1);
+		InputBindings[0].Data = InputData.GetData();
+		InputBindings[0].SizeInBytes = InputData.Num() * sizeof(float);
+		for (int32 Index = 0; Index < ModelWeightsIn.Num(); Index++)
+		{
+			InputBindings[Index + 1].Data = ModelWeightsIn[Index].GetData();
+			InputBindings[Index + 1].SizeInBytes = ModelWeightsIn[Index].Num() * sizeof(float);
+		}
+		OutputBindings.SetNumZeroed(1);
+		OutputBindings[0].Data = OutputData.GetData();
+		OutputBindings[0].SizeInBytes = OutputData.Num() * sizeof(float);
+
+		{
+			TRACE_CPUPROFILER_EVENT_SCOPE(FMLLevelSetNeuralInference_RunInference_ModelInstanceRunSync);
+			ModelInstance->RunSync(InputBindings, OutputBindings);
+		}
+	}
+	else
+	{
+		UE_LOG(LogTemp, Error, TEXT("FMLLevelSetNeuralInference::RunInference() - ModelInstance us not valid"));
+	}
+}
+
+bool FMLLevelSetNeuralInference::Serialize(FArchive& Ar)
+{
+	Ar << ModelWeightsShapes;
+	return true;
+}
+
+}//end Chaos
\ No newline at end of file
diff --git a/Engine/Source/Runtime/Experimental/Chaos/Private/Chaos/MLLevelset.cpp b/Engine/Source/Runtime/Experimental/Chaos/Private/Chaos/MLLevelset.cpp
new file mode 100644
--- /dev/null
+++ b/Engine/Source/Runtime/Experimental/Chaos/Private/Chaos/MLLevelset.cpp
@@ -0,0 +1,951 @@
+// Copyright Epic Games, Inc. All Rights Reserved.
+#include "Chaos/MLLevelset.h"
+#include "Misc/Paths.h"
+#include "Async/ParallelFor.h"
+#include <algorithm>
+#include <vector>
+#include "Chaos/PBDSoftsSolverParticles.h" 
+#include "Containers/ContainerAllocationPolicies.h"
+#include "NNE.h"  
+
+namespace Chaos
+{
+static int MLLevelSetUpdatePhiFlag = 0;
+static FAutoConsoleVariableRef CVarChaosMLLevelSetUpdatePhiFlag(TEXT("p.MLLevelSet.MLLevelSetUpdatePhiFlag"), MLLevelSetUpdatePhiFlag, TEXT("0(No): Default. 1(Yes): Updates DebugPhi Array for MLLevelset Isocounter Visualization. "));
+
+FMLLevelSet::FMLLevelSet(FMLLevelSetImportData&& MLLevelSetImportData)
+	: FImplicitObject(EImplicitObject::HasBoundingBox, ImplicitObjectType::MLLevelSet)
+	, ParentBoneName(MLLevelSetImportData.ParentBoneName)
+	, ActiveBoneNames(MoveTemp(MLLevelSetImportData.ActiveBoneNames))
+	, SignedDistanceScaling(MLLevelSetImportData.SignedDistanceScaling)
+	, TrainingGridMin(MLLevelSetImportData.TrainingGridMin)
+	, TrainingGridUnitAxesXYZ(MoveTemp(MLLevelSetImportData.TrainingGridAxesXYZ))
+	, ActiveBonesRotationComponents(MoveTemp(MLLevelSetImportData.ActiveBonesRotationComponents))
+	, ActiveBonesReferenceRotations(MoveTemp(MLLevelSetImportData.ActiveBonesReferenceRotations))
+	, ActiveBonesReferenceTranslations(MoveTemp(MLLevelSetImportData.ActiveBonesReferenceTranslations))
+{
+	TotalNumberOfRotationComponents = 0;
+	for (int32 RotationComponentIndex = 0; RotationComponentIndex < ActiveBonesRotationComponents.Num(); RotationComponentIndex++)
+	{
+		TotalNumberOfRotationComponents += ActiveBonesRotationComponents[RotationComponentIndex].Num();
+	}
+	ActiveBonesRelativeTransforms.SetNum(ActiveBoneNames.Num());
+	check(MLLevelSetImportData.NNEModelDataArr.Num() > 0);
+	BuildNNEModel(MLLevelSetImportData.NNEModelDataArr[0].ModelArchitectureActivationNodeSizes, MLLevelSetImportData.NNEModelDataArr[0].NNEModelData, MLLevelSetImportData.NNEModelDataArr[0].MLModelWeightsString,
+		NNESignedDistanceModel, SignedDistanceNeuralInferences, SignedDistanceModelWeightsShapes, SignedDistanceModelWeights);
+	bUseIncorrectZoneModel = MLLevelSetImportData.NNEModelDataArr.Num() > 1;
+	if (bUseIncorrectZoneModel)
+	{
+		BuildNNEModel(MLLevelSetImportData.NNEModelDataArr[1].ModelArchitectureActivationNodeSizes, MLLevelSetImportData.NNEModelDataArr[1].NNEModelData, MLLevelSetImportData.NNEModelDataArr[1].MLModelWeightsString,
+			NNEIncorrectZoneModel, IncorrectZoneNeuralInferences, IncorrectZoneModelWeightsShapes, IncorrectZoneModelWeights);
+	}
+
+	FVec3 TrainingGridMax = (FVec3)(TrainingGridUnitAxesXYZ[0] + TrainingGridUnitAxesXYZ[1] + TrainingGridUnitAxesXYZ[2] + TrainingGridMin);
+	LocalBoundingBox = Chaos::FAABB3(Chaos::FVec3(TrainingGridMin), TrainingGridMax);
+	for (int32 CornerIndex = 0; CornerIndex < 3; CornerIndex++) 
+	{
+		LocalBoundingBox.GrowToInclude((FVec3)TrainingGridUnitAxesXYZ[CornerIndex] + TrainingGridMin);
+		LocalBoundingBox.GrowToInclude(TrainingGridMax - (FVec3)TrainingGridUnitAxesXYZ[CornerIndex]);
+	}
+	ProcessTrainingGridAxesVectors();
+
+	//Create Debug Phi and Debug Grid
+	const Chaos::TVec3<int32> DebugChaosGridDim(MLLevelSetImportData.DebugGridResolution[0], MLLevelSetImportData.DebugGridResolution[1], MLLevelSetImportData.DebugGridResolution[2]);
+	DebugGrid = TUniformGrid<Chaos::FReal, 3>(LocalBoundingBox.Min(), LocalBoundingBox.Max(), DebugChaosGridDim);
+	DebugPhi = TArrayND<Chaos::FReal, 3>(DebugChaosGridDim);
+}
+
+CHAOS_API FMLLevelSet::FMLLevelSet(TObjectPtr<UNNEModelData> InNNESignedDistanceModelData, TObjectPtr<UNNEModelData> InNNEIncorrectZoneModel)
+	: FImplicitObject(EImplicitObject::HasBoundingBox, ImplicitObjectType::MLLevelSet)
+{
+	BuildNNEModel(InNNESignedDistanceModelData, NNESignedDistanceModel);
+	BuildNNEModel(InNNEIncorrectZoneModel, NNEIncorrectZoneModel);
+}
+
+FMLLevelSet::FMLLevelSet(FMLLevelSet&& Other)
+	: FImplicitObject(EImplicitObject::HasBoundingBox, ImplicitObjectType::MLLevelSet)
+	, ParentBoneName(Other.ParentBoneName)
+	, ActiveBoneNames(MoveTemp(Other.ActiveBoneNames))
+	, SolverActiveBoneIndices(MoveTemp(Other.SolverActiveBoneIndices))
+	, SolverParentBoneIndex(Other.SolverParentBoneIndex)
+	, ActiveBonesRelativeTransforms(MoveTemp(Other.ActiveBonesRelativeTransforms))
+
+	, NNESignedDistanceModel(Other.NNESignedDistanceModel)
+	, SignedDistanceNeuralInferences(MoveTemp(Other.SignedDistanceNeuralInferences))
+	, SignedDistanceModelWeights(MoveTemp(Other.SignedDistanceModelWeights))
+	, SignedDistanceModelWeightsShapes(MoveTemp(Other.SignedDistanceModelWeightsShapes))
+
+	, NNEIncorrectZoneModel(Other.NNEIncorrectZoneModel)
+	, IncorrectZoneNeuralInferences(Other.IncorrectZoneNeuralInferences)
+	, IncorrectZoneModelWeights(MoveTemp(Other.IncorrectZoneModelWeights))
+	, IncorrectZoneModelWeightsShapes(MoveTemp(Other.IncorrectZoneModelWeightsShapes))
+	, bUseIncorrectZoneModel(Other.bUseIncorrectZoneModel)
+
+	, LocalBoundingBox(MoveTemp(Other.LocalBoundingBox))
+	, SignedDistanceScaling(Other.SignedDistanceScaling)
+	, TrainingGridMin(Other.TrainingGridMin)
+	, TrainingGridUnitAxesXYZ(MoveTemp(Other.TrainingGridUnitAxesXYZ))
+	, TrainingGridAxesLengthsXYZ(MoveTemp(Other.TrainingGridAxesLengthsXYZ))
+	, TotalNumberOfRotationComponents(Other.TotalNumberOfRotationComponents)
+	, ActiveBonesRotationComponents(MoveTemp(Other.ActiveBonesRotationComponents))
+	, ActiveBonesReferenceRotations(MoveTemp(Other.ActiveBonesReferenceRotations))
+	, ActiveBonesReferenceTranslations(MoveTemp(Other.ActiveBonesReferenceTranslations))
+
+	, DebugGrid(MoveTemp(Other.DebugGrid))
+	, DebugPhi(MoveTemp(Other.DebugPhi))
+{}
+
+FMLLevelSet::~FMLLevelSet()
+{}
+
+
+Chaos::FImplicitObjectPtr FMLLevelSet::CopyGeometry() const
+{
+	FMLLevelSet* Copy = new FMLLevelSet();
+
+	Copy->ParentBoneName = ParentBoneName;
+	Copy->ActiveBoneNames = ActiveBoneNames;
+	Copy->SolverActiveBoneIndices = SolverActiveBoneIndices;
+	Copy->ActiveBonesRelativeTransforms = ActiveBonesRelativeTransforms;
+	Copy->SolverParentBoneIndex = SolverParentBoneIndex;
+
+	Copy->NNESignedDistanceModel = NNESignedDistanceModel;
+	Copy->SignedDistanceNeuralInferences = SignedDistanceNeuralInferences;
+	Copy->SignedDistanceModelWeights = SignedDistanceModelWeights;
+	Copy->SignedDistanceModelWeightsShapes = SignedDistanceModelWeightsShapes;
+
+	Copy->NNEIncorrectZoneModel = NNEIncorrectZoneModel;
+	Copy->IncorrectZoneNeuralInferences = IncorrectZoneNeuralInferences;
+	Copy->IncorrectZoneModelWeights = IncorrectZoneModelWeights;
+	Copy->IncorrectZoneModelWeightsShapes = IncorrectZoneModelWeightsShapes;
+	Copy->bUseIncorrectZoneModel = bUseIncorrectZoneModel;
+
+	Copy->LocalBoundingBox = LocalBoundingBox;
+	Copy->SignedDistanceScaling = SignedDistanceScaling;
+	Copy->TrainingGridMin = TrainingGridMin;
+	Copy->TrainingGridUnitAxesXYZ = TrainingGridUnitAxesXYZ;
+	Copy->TrainingGridAxesLengthsXYZ = TrainingGridAxesLengthsXYZ;
+	Copy->TotalNumberOfRotationComponents = TotalNumberOfRotationComponents;
+	Copy->ActiveBonesRotationComponents = ActiveBonesRotationComponents;
+	Copy->ActiveBonesReferenceRotations = ActiveBonesReferenceRotations;
+	Copy->ActiveBonesReferenceTranslations = ActiveBonesReferenceTranslations;
+
+	Copy->DebugGrid = DebugGrid;
+	Copy->DebugPhi.Copy(DebugPhi);
+
+	return Chaos::FImplicitObjectPtr(Copy);
+}
+
+Chaos::FImplicitObjectPtr FMLLevelSet::CopyGeometryWithScale(const FVec3& Scale) const
+{
+	FMLLevelSet* Copy = new FMLLevelSet();
+
+	Copy->ParentBoneName = ParentBoneName;
+	Copy->ActiveBoneNames = ActiveBoneNames;
+	Copy->SolverActiveBoneIndices = SolverActiveBoneIndices;
+	Copy->ActiveBonesRelativeTransforms = ActiveBonesRelativeTransforms;
+	Copy->SolverParentBoneIndex = SolverParentBoneIndex;
+
+	Copy->NNESignedDistanceModel = NNESignedDistanceModel;
+	Copy->SignedDistanceNeuralInferences = SignedDistanceNeuralInferences;
+	Copy->SignedDistanceModelWeights = SignedDistanceModelWeights;
+	Copy->SignedDistanceModelWeightsShapes = SignedDistanceModelWeightsShapes;
+
+	Copy->NNEIncorrectZoneModel = NNEIncorrectZoneModel;
+	Copy->IncorrectZoneNeuralInferences = IncorrectZoneNeuralInferences;
+	Copy->IncorrectZoneModelWeights = IncorrectZoneModelWeights;
+	Copy->IncorrectZoneModelWeightsShapes = IncorrectZoneModelWeightsShapes;
+	Copy->bUseIncorrectZoneModel = bUseIncorrectZoneModel;
+
+	Copy->LocalBoundingBox = LocalBoundingBox;
+	Copy->SignedDistanceScaling = SignedDistanceScaling;
+	Copy->TrainingGridMin = TrainingGridMin;
+	Copy->TrainingGridUnitAxesXYZ = TrainingGridUnitAxesXYZ;
+	Copy->TrainingGridAxesLengthsXYZ = TrainingGridAxesLengthsXYZ;
+	Copy->TotalNumberOfRotationComponents = TotalNumberOfRotationComponents;
+	Copy->ActiveBonesRotationComponents = ActiveBonesRotationComponents;
+	Copy->ActiveBonesReferenceRotations = ActiveBonesReferenceRotations;
+	Copy->ActiveBonesReferenceTranslations = ActiveBonesReferenceTranslations;
+
+	Copy->DebugGrid = DebugGrid;
+	Copy->DebugPhi.Copy(DebugPhi);
+
+	return MakeImplicitObjectPtr<TImplicitObjectScaled<FMLLevelSet>>(Copy, Scale);
+}
+
+/** This function uses DebugPhi. Is there a better way to do this, which uses the network directly? */
+void FMLLevelSet::GetZeroIsosurfaceGridCellFaces(TArray<FVector3f>& Vertices, TArray<FIntVector>& Tris) const
+{
+	const TVector<int32, 3> Cells = DebugGrid.Counts();
+	const FVector3d Dx(DebugGrid.Dx());
+
+	for (int i = 0; i < Cells.X - 1; ++i)
+	{
+		for (int j = 0; j < Cells.Y - 1; ++j)
+		{
+			for (int k = 0; k < Cells.Z - 1; ++k)
+			{
+				const double Sign = FMath::Sign(DebugPhi(i, j, k));
+				const double SignNextI = FMath::Sign(DebugPhi(i + 1, j, k));
+				const double SignNextJ = FMath::Sign(DebugPhi(i, j + 1, k));
+				const double SignNextK = FMath::Sign(DebugPhi(i, j, k + 1));
+
+				const FVector3d CellMin = DebugGrid.MinCorner() + Dx * FVector3d(i, j, k);
+
+				if (Sign > SignNextI)
+				{
+					const int32 V0 = Vertices.Emplace(CellMin + Dx * FVector3d(1, 0, 0));
+					const int32 V1 = Vertices.Emplace(CellMin + Dx * FVector3d(1, 1, 0));
+					const int32 V2 = Vertices.Emplace(CellMin + Dx * FVector3d(1, 1, 1));
+					const int32 V3 = Vertices.Emplace(CellMin + Dx * FVector3d(1, 0, 1));
+					Tris.Emplace(FIntVector(V0, V1, V2));
+					Tris.Emplace(FIntVector(V2, V3, V0));
+				}
+				else if (Sign < SignNextI)
+				{
+					const int32 V0 = Vertices.Emplace(CellMin + Dx * FVector3d(1, 0, 0));
+					const int32 V1 = Vertices.Emplace(CellMin + Dx * FVector3d(1, 1, 0));
+					const int32 V2 = Vertices.Emplace(CellMin + Dx * FVector3d(1, 1, 1));
+					const int32 V3 = Vertices.Emplace(CellMin + Dx * FVector3d(1, 0, 1));
+					Tris.Emplace(FIntVector(V0, V2, V1));
+					Tris.Emplace(FIntVector(V2, V0, V3));
+				}
+
+
+				if (Sign > SignNextJ)
+				{
+					const int32 V0 = Vertices.Emplace(CellMin + Dx * FVector3d(0, 1, 0));
+					const int32 V1 = Vertices.Emplace(CellMin + Dx * FVector3d(1, 1, 0));
+					const int32 V2 = Vertices.Emplace(CellMin + Dx * FVector3d(1, 1, 1));
+					const int32 V3 = Vertices.Emplace(CellMin + Dx * FVector3d(0, 1, 1));
+					Tris.Emplace(FIntVector(V0, V2, V1));
+					Tris.Emplace(FIntVector(V2, V0, V3));
+				}
+				else if (Sign < SignNextJ)
+				{
+					const int32 V0 = Vertices.Emplace(CellMin + Dx * FVector3d(0, 1, 0));
+					const int32 V1 = Vertices.Emplace(CellMin + Dx * FVector3d(1, 1, 0));
+					const int32 V2 = Vertices.Emplace(CellMin + Dx * FVector3d(1, 1, 1));
+					const int32 V3 = Vertices.Emplace(CellMin + Dx * FVector3d(0, 1, 1));
+					Tris.Emplace(FIntVector(V0, V1, V2));
+					Tris.Emplace(FIntVector(V2, V3, V0));
+				}
+
+				if (Sign > SignNextK)
+				{
+					const int32 V0 = Vertices.Emplace(CellMin + Dx * FVector3d(0, 0, 1));
+					const int32 V1 = Vertices.Emplace(CellMin + Dx * FVector3d(1, 0, 1));
+					const int32 V2 = Vertices.Emplace(CellMin + Dx * FVector3d(1, 1, 1));
+					const int32 V3 = Vertices.Emplace(CellMin + Dx * FVector3d(0, 1, 1));
+					Tris.Emplace(FIntVector(V0, V1, V2));
+					Tris.Emplace(FIntVector(V2, V3, V0));
+				}
+				else if (Sign < SignNextK)
+				{
+					const int32 V0 = Vertices.Emplace(CellMin + Dx * FVector3d(0, 0, 1));
+					const int32 V1 = Vertices.Emplace(CellMin + Dx * FVector3d(1, 0, 1));
+					const int32 V2 = Vertices.Emplace(CellMin + Dx * FVector3d(1, 1, 1));
+					const int32 V3 = Vertices.Emplace(CellMin + Dx * FVector3d(0, 1, 1));
+					Tris.Emplace(FIntVector(V0, V2, V1));
+					Tris.Emplace(FIntVector(V2, V0, V3));
+				}
+			}
+		}
+	}
+}
+
+/** This function uses DebugPhi. Is there a better way to do this, which uses the network directly? */
+void FMLLevelSet::GetInteriorCells(TArray<TVec3<int32>>& InteriorCells, const FReal InteriorThreshold) const
+{
+	InteriorCells.Reset();
+
+	const TVector<int32, 3> Cells = DebugGrid.Counts();
+
+	for (int i = 0; i < Cells.X; ++i)
+	{
+		for (int j = 0; j < Cells.Y; ++j)
+		{
+			for (int k = 0; k < Cells.Z; ++k)
+			{
+				const FReal Value = DebugPhi(i, j, k);
+				if (Value < InteriorThreshold)
+				{
+					InteriorCells.Emplace(i, j, k);
+				}
+			}
+		}
+	}
+}
+
+/** This function should not be called. MLLevelSets are only BatchPhiWithNormal friendly */
+FReal FMLLevelSet::SignedDistance(const FVec3& x) const
+{
+	UE_LOG(LogChaos, Error, TEXT("FMLLevelSet::PhiWithNormal cannot be used for single queries. Use FMLLevelSet::BatchPhiWithNormal() instead."));
+	return UE_DOUBLE_BIG_NUMBER;
+}
+
+void FMLLevelSet::CreatePhiFromMLModel()
+{	
+	const TVector<int32, 3> Cells = DebugGrid.Counts();
+	TArray<float, TAlignedHeapAllocator<64>> OutputDataSignedDistances;
+	TArray<float, TAlignedHeapAllocator<64>> InputDataLocationsMS;
+	const uint32 ModelInputShapeSize = 3;
+	const uint32 ModelOutputShapeSize = 1;
+	const uint32 ModelOutputSignedDistanceIndex = 0;
+
+	TArray<TArray<float, TAlignedHeapAllocator<64>>> MLWeightsIn;
+	ComputeSignedDistanceNetworkWeightsInput(MLWeightsIn);
+
+	InputDataLocationsMS.SetNum(ModelInputShapeSize * DebugGrid.GetNumCells());
+	OutputDataSignedDistances.SetNum(ModelOutputShapeSize * DebugGrid.GetNumCells());
+
+	TArray<FVector3f> TrainingGridVectorsScaled;
+	TrainingGridVectorsScaled.SetNum(3);
+	for (int32 Index = 0; Index < 3; Index++)
+	{
+		TrainingGridVectorsScaled[Index] = TrainingGridUnitAxesXYZ[Index] / TrainingGridAxesLengthsXYZ[Index];
+	}
+
+	for (int32 i = 0; i < DebugGrid.Counts()[0]; i++)
+	{
+		for (int32 j = 0; j < DebugGrid.Counts()[1]; j++)
+		{
+			for (int32 k = 0; k < DebugGrid.Counts()[2]; k++)
+			{
+				TVector<int32, 3> Index = { i, j, k };
+				const int32 ClosestActiveBoneIndex = GetClosestActiveBoneIndex(DebugGrid.Location(Index));
+				const FVector3f LocalGridCornerShift = (FVector3f)ActiveBonesRelativeTransforms[ClosestActiveBoneIndex].GetTranslation() - (FVector3f)ActiveBonesReferenceTranslations[ClosestActiveBoneIndex];
+				const FVector3f LocalGridCornerShifted = TrainingGridMin + LocalGridCornerShift;
+
+				TVector<FRealSingle, 3> LocationMSShifted = DebugGrid.Location(Index) - LocalGridCornerShifted;
+				for (int32 ModelInputLocationIndex = 0; ModelInputLocationIndex < ModelInputShapeSize; ModelInputLocationIndex++)
+				{
+					InputDataLocationsMS[ModelInputShapeSize * DebugGrid.FlatIndex(Index) + ModelInputLocationIndex] = TrainingGridVectorsScaled[ModelInputLocationIndex][0] * LocationMSShifted[0] + TrainingGridVectorsScaled[ModelInputLocationIndex][1] * LocationMSShifted[1] + TrainingGridVectorsScaled[ModelInputLocationIndex][2] * LocationMSShifted[2];
+				}
+			}
+		}
+	}
+
+	if (SignedDistanceNeuralInferences[0].IsValid())
+	{
+		SignedDistanceNeuralInferences[0].RunInference(InputDataLocationsMS, OutputDataSignedDistances, ModelInputShapeSize, ModelOutputShapeSize, MLWeightsIn);
+	}
+
+	for (int32 i = 0; i < DebugGrid.Counts()[0]; i++)
+	{
+		for (int32 j = 0; j < DebugGrid.Counts()[1]; j++)
+		{
+			for (int32 k = 0; k < DebugGrid.Counts()[2]; k++)
+			{
+				TVector<int32, 3> Index = { i, j, k };
+				int32 FlatIndex = DebugGrid.FlatIndex(Index);
+				DebugPhi(i, j, k) = SignedDistanceScaling * OutputDataSignedDistances[ModelOutputShapeSize * FlatIndex + ModelOutputSignedDistanceIndex];
+			}
+		}
+	}
+}
+
+
+void FMLLevelSet::UpdateActiveBonesRelativeTransforms(TArray<FTransform>& InActiveBonesRelativeTransforms)
+{
+	bool bIsFTransformArraysDifferent = IsFTransformArraysDifferent(ActiveBonesRelativeTransforms, InActiveBonesRelativeTransforms);
+	
+	if (bIsFTransformArraysDifferent)
+	{
+		ActiveBonesRelativeTransforms = MoveTemp(InActiveBonesRelativeTransforms);
+		
+		// Phi is created for each timestep for debug drawing
+		if (MLLevelSetUpdatePhiFlag == 1)
+		{
+			CreatePhiFromMLModel();
+		}
+	}
+}
+
+void FMLLevelSet::UpdateActiveBonesRelativeTransformsAndUpdateDebugPhi(TArray<FTransform>& InActiveBonesRelativeTransforms)
+{
+	bool bIsFTransformDifferent = IsFTransformArraysDifferent(ActiveBonesRelativeTransforms, InActiveBonesRelativeTransforms);
+	ActiveBonesRelativeTransforms = MoveTemp(InActiveBonesRelativeTransforms);
+
+	// Update DebugPhi only if FTransform is updated
+	if (bIsFTransformDifferent)
+	{
+		CreatePhiFromMLModel();
+	}
+}
+
+// Adds more MLLevelSetNeuralInference to allow multi-treating.
+void FMLLevelSet::UpdateNeuralInferencesNumber(const int32 InNeuralInferencesNumber)
+{
+	if (SignedDistanceNeuralInferences.Num() < InNeuralInferencesNumber)
+	{
+		for (int32 Index = SignedDistanceNeuralInferences.Num(); Index < InNeuralInferencesNumber; Index++)
+		{
+			SignedDistanceNeuralInferences.Add(FMLLevelSetNeuralInference(NNESignedDistanceModel, SignedDistanceModelWeightsShapes));
+			if (bUseIncorrectZoneModel)
+			{
+				IncorrectZoneNeuralInferences.Add(FMLLevelSetNeuralInference(NNEIncorrectZoneModel, IncorrectZoneModelWeightsShapes));
+			}
+		}
+	}
+}
+
+// Body Part is trained over a grid/BoundingBox created outside of the engine
+// This training grid might not align with the coordinate system axes of the engine
+// TrainingGridVectors are used for mapping from training coordinate system to engine coordinate system
+// Important Note: If training uses MAYA for generating the dataset, do not forget to negate the Y-axis coordinate before importing the DataTable
+// I.e., UE.Y = - Maya.Y
+void FMLLevelSet::ProcessTrainingGridAxesVectors()
+{
+	//check if training grid vectors are initialized
+	check(TrainingGridUnitAxesXYZ.Num() == 3);
+
+	TrainingGridAxesLengthsXYZ.Empty();
+	TrainingGridAxesLengthsXYZ.SetNum(3);
+
+	for (int32 Index = 0; Index < 3; Index++)
+	{
+		TrainingGridAxesLengthsXYZ[Index] = float(TrainingGridUnitAxesXYZ[Index].Length());
+		TrainingGridUnitAxesXYZ[Index] = TrainingGridUnitAxesXYZ[Index] / TrainingGridAxesLengthsXYZ[Index];
+	}
+}
+
+// Loads MLModelWeights From an FString
+// TArray<TArray<float>> MLModelWeights = {W1,..,Wn} is tokanized into FString as 
+// "W1_0,W1_1,...,W1_k1|W2_0,W2_1,...,W2_k2|...|Wn_0,Wn_1,...,Wn_kn" 
+void FMLLevelSet::LoadMLModelWeightsFromString(const FString& MLModelWeightsString, TArray<TArray<float>>& ModelWeightArray)
+{
+	TArray<FString> MLModelWeightsArrayStringsSeperatedByBar;
+	MLModelWeightsString.ParseIntoArray(MLModelWeightsArrayStringsSeperatedByBar, TEXT("|"), true);
+
+	ModelWeightArray.Empty();
+	ModelWeightArray.SetNum(MLModelWeightsArrayStringsSeperatedByBar.Num());
+	for (int32 i = 0; i < MLModelWeightsArrayStringsSeperatedByBar.Num(); i++)
+	{
+		TArray<FString> MLModelWeightsStringsSeperatedByComma;
+		MLModelWeightsArrayStringsSeperatedByBar[i].ParseIntoArray(MLModelWeightsStringsSeperatedByComma,TEXT(","),true);
+		ModelWeightArray[i].Empty();
+		ModelWeightArray[i].SetNum(MLModelWeightsStringsSeperatedByComma.Num());
+		
+		for (int32 j = 0; j < ModelWeightArray[i].Num(); j++)
+		{
+			ModelWeightArray[i][j]= FCString::Atof(*MLModelWeightsStringsSeperatedByComma[j]);
+		}
+	}
+}
+
+void FMLLevelSet::BuildNNEModel(const TArray<int32>& ModelArchitectureActivationNodeSizes, TObjectPtr<UNNEModelData> NNEModelData, const FString& ModelWeightsString,
+	TSharedPtr<UE::NNE::IModelCPU>& NNEModel, TArray<FMLLevelSetNeuralInference>& NeuralInferences, TArray<TArray<int32>>& ModelWeightsShapes, TArray<TArray<float>>& ModelWeightArray)
+{
+	// This is specific to weight-updated MLP types of networks.
+	// (BigToDo) Update this part to accept different types of networks.
+	ModelWeightsShapes.Empty();
+	for (int32 i = 0; i < ModelArchitectureActivationNodeSizes.Num() - 1; i++)
+	{
+		ModelWeightsShapes.Add(TArray<int32>({ ModelArchitectureActivationNodeSizes[i],ModelArchitectureActivationNodeSizes[i + 1] }));
+		ModelWeightsShapes.Add(TArray<int32>({ ModelArchitectureActivationNodeSizes[i + 1] }));
+	}
+
+	//Create Model Instance and NNE Neural Inferences
+	TWeakInterfacePtr<INNERuntimeCPU> Runtime = UE::NNE::GetRuntime<INNERuntimeCPU>(FString("NNERuntimeIREECpu"));
+	if (Runtime.IsValid())
+	{
+		NNEModel = Runtime->CreateModelCPU(NNEModelData);
+		if (NNEModel.IsValid())
+		{
+			NeuralInferences.Empty();
+			NeuralInferences.Add(FMLLevelSetNeuralInference(NNEModel, ModelWeightsShapes));
+		}
+		else
+		{
+			UE_LOG(LogChaos, Error, TEXT("MLLevelSet::BuildNNEModel - Model is not valid."));
+		}
+	}
+	else
+	{
+		UE_LOG(LogChaos, Error, TEXT("MLLevelSet::BuildNNEModel - NNE Runtime is NOT valid."));
+	}
+
+	LoadMLModelWeightsFromString(ModelWeightsString, ModelWeightArray);
+}
+
+// Only Build NNE Model. Used for serialization
+void FMLLevelSet::BuildNNEModel(TObjectPtr<UNNEModelData> InNNEModelData, TSharedPtr<UE::NNE::IModelCPU>& NNEModel)
+{
+	TWeakInterfacePtr<INNERuntimeCPU> Runtime = UE::NNE::GetRuntime<INNERuntimeCPU>(FString("NNERuntimeIREECpu"));
+	if (Runtime.IsValid())
+	{
+		NNEModel = Runtime->CreateModelCPU(InNNEModelData);
+	}
+	else
+	{
+		UE_LOG(LogChaos, Error, TEXT("MLLevelSet::BuildNNEModel - NNE Runtime is NOT valid."));
+	}
+}
+
+const bool FMLLevelSet::IsFTransformArraysDifferent(TArray<FTransform>& FTransformArr1, TArray<FTransform>& FTransformArr2, FReal Tol) const
+{
+	check(FTransformArr1.Num() == FTransformArr2.Num());
+	for (int32 Index = 0; Index < FTransformArr1.Num(); Index++)
+	{
+		FVector TranslationDelta = FTransformArr1[Index].GetTranslation() - FTransformArr2[Index].GetTranslation() ;
+		if (TranslationDelta.GetAbsMax() > Tol)
+		{
+			return true;
+		}
+		FRotator RotationDelta = FRotator(FTransformArr1[Index].GetRotation()) - FRotator(FTransformArr2[Index].GetRotation());
+		if (!RotationDelta.IsNearlyZero(Tol))
+		{
+			return true;
+		}
+	}
+	return false;
+}
+
+FReal FMLLevelSet::PhiWithNormal(const FVec3& x, FVec3& Normal) const
+{
+	UE_LOG(LogChaos, Error, TEXT("FMLLevelSet::PhiWithNormal() cannot be used for single queries. Use FMLLevelSet::BatchPhiWithNormal() instead."));
+	return UE_DOUBLE_BIG_NUMBER;
+}
+
+void FMLLevelSet::ComputeSignedDistanceNetworkWeightsInput(TArray<TArray<float, TAlignedHeapAllocator<64>>>& NetworkWeightsInput) const
+{
+	ComputeWeightsInput(ActiveBonesRelativeTransforms, SignedDistanceModelWeights, SignedDistanceModelWeightsShapes, NetworkWeightsInput);
+}
+
+void FMLLevelSet::ComputeIncorrectZoneNetworkWeightsInput(TArray<TArray<float, TAlignedHeapAllocator<64>>>& NetworkWeightsInput) const
+{
+	ComputeWeightsInput(ActiveBonesRelativeTransforms, IncorrectZoneModelWeights, IncorrectZoneModelWeightsShapes, NetworkWeightsInput);
+}
+
+
+void FMLLevelSet::ComputeWeightsInput(const TArray<FTransform>& RelativeBoneTransformationsInput, const TArray<TArray<float>>& NetworkWeights, const TArray<TArray<int32>>& NetworkWeightsShapes, TArray<TArray<float, TAlignedHeapAllocator<64>>>& NetworkWeightsInput) const
+{
+	checkSlow(RelativeBoneTransformationsInput.Num() == GetNumberOfActiveBones());
+	NetworkWeightsInput.Empty();
+	NetworkWeightsInput.SetNum(NetworkWeights.Num());
+	float HalfRotatationAngle = 180.0;
+
+	TArray<TArray<float>> JointRotationInputAnglesMod360;
+	JointRotationInputAnglesMod360.SetNum(RelativeBoneTransformationsInput.Num());
+	for (int32 RotationIndex = 0; RotationIndex < JointRotationInputAnglesMod360.Num(); RotationIndex++)
+	{
+		FRotator Rot(RelativeBoneTransformationsInput[RotationIndex].GetRotation());
+		JointRotationInputAnglesMod360[RotationIndex] = { (float)(Rot.Roll - ActiveBonesReferenceRotations[RotationIndex].X), (float)(Rot.Pitch - ActiveBonesReferenceRotations[RotationIndex].Y), (float)(Rot.Yaw - ActiveBonesReferenceRotations[RotationIndex].Z)};
+		// Put Joint Rotation Angles in (-180,180] in mod 360
+		for (int i = 0; i < 3; i++)
+		{
+			while (JointRotationInputAnglesMod360[RotationIndex][i] <= -HalfRotatationAngle)
+			{
+				JointRotationInputAnglesMod360[RotationIndex][i] += 2*HalfRotatationAngle;
+			}
+			while (JointRotationInputAnglesMod360[RotationIndex][i] > HalfRotatationAngle)
+			{
+				JointRotationInputAnglesMod360[RotationIndex][i] -= 2 * HalfRotatationAngle;
+			}
+		}
+	}	
+
+	TArray<float> MLRotationAngleInput;
+	MLRotationAngleInput.Reserve(TotalNumberOfRotationComponents);
+
+	for (int32 RotIndex = 0; RotIndex < ActiveBonesRotationComponents.Num(); RotIndex++)
+	{
+		for (int32 RotationComponent : ActiveBonesRotationComponents[RotIndex])
+		{
+			MLRotationAngleInput.Add(float(JointRotationInputAnglesMod360[RotIndex][RotationComponent] / 360.0));
+		}
+	}
+	
+	const int32 NumberOfAngleParameters = MLRotationAngleInput.Num();
+	const int32 NumberOfWeightVariables = NumberOfAngleParameters + 1;
+	for (int32 WI = 0; WI < NetworkWeights.Num(); WI++)
+	{
+		NetworkWeightsInput[WI].SetNum(NetworkWeights[WI].Num() / NumberOfWeightVariables);
+		for (int32 WJ = 0; WJ < NetworkWeightsInput[WI].Num(); WJ++)
+		{
+			int32 WidthSize = NetworkWeightsShapes[WI][NetworkWeightsShapes[WI].Num() - 1];
+			int32 WJ1 = WJ / WidthSize;
+			int32 WJ2 = WJ % WidthSize;
+			NetworkWeightsInput[WI][WJ] = NetworkWeights[WI][WJ1 * WidthSize * NumberOfWeightVariables + WJ2 + (NumberOfWeightVariables - 1) * WidthSize];
+			for (int32 InputWeightIndex = 0; InputWeightIndex < NumberOfAngleParameters; InputWeightIndex++)
+			{
+				NetworkWeightsInput[WI][WJ] += MLRotationAngleInput[InputWeightIndex] * NetworkWeights[WI][WJ1 * WidthSize * NumberOfWeightVariables + WJ2 + InputWeightIndex * WidthSize];
+			}
+		}
+	}
+}
+
+void FMLLevelSet::BatchPhiWithNormal(const TConstArrayView<Softs::FPAndInvM> PAndInvMArray,
+	const Softs::FSolverRigidTransform3& SolverToThis, TArray<Softs::FSolverReal>& OutBatchPhis, TArray<Softs::FSolverVec3>& OutBatchNormals,
+	const Chaos::Softs::FSolverReal CollisionThickness, const int32 MLLevelsetThread, const int32 BatchBegin, const int32 BatchEnd) const
+{
+	TRACE_CPUPROFILER_EVENT_SCOPE(FMLLevelSet_BatchPhiWithNormal);
+	const uint32 ModelInputShapeSize = 3;  // XYZ coordinates wrt local grid
+	const uint32 ModelOutputShapeSize = 1; // SignedDistance Output
+	const Chaos::Softs::FSolverReal MaxPhiValue = 2 * SignedDistanceScaling;
+	const float DeltaForFiniteDifferenceForNormalCalc = SignedDistanceScaling / float(100); //delta for finite differencing
+
+	const int32 NumParticles = BatchEnd - BatchBegin;
+	const Softs::FPAndInvM* const PAndInvM = PAndInvMArray.GetData();
+	Softs::FSolverReal* const Phis = OutBatchPhis.GetData();
+	Softs::FSolverVec3* const NormalsOut = OutBatchNormals.GetData();
+
+	//Query only the phi values on the cloth
+	TArray<float, TAlignedHeapAllocator<64>> MLInputForPhis;
+	MLInputForPhis.SetNum(NumParticles * ModelInputShapeSize);
+	TArray<float, TAlignedHeapAllocator<64>> MLOutputForPhis;
+
+	TArray<int32> ParticlesInsideTheGridIndexes;
+	ParticlesInsideTheGridIndexes.Reserve(NumParticles);
+	TArray<TArray<float, TAlignedHeapAllocator<64>>> MLWeightsIn;
+	ComputeSignedDistanceNetworkWeightsInput(MLWeightsIn);
+
+	TArray<FVector3f> TrainingGridVectorsScaled;
+	TrainingGridVectorsScaled.SetNum(3);
+	for (int32 Index = 0; Index < 3; Index++)
+	{
+		TrainingGridVectorsScaled[Index] = TrainingGridUnitAxesXYZ[Index] / TrainingGridAxesLengthsXYZ[Index];
+	}
+
+	int32 InteriorParticlesIndex = 0;
+	for (int32 Index = BatchBegin; Index < BatchEnd; Index++)
+	{
+		OutBatchPhis[Index] = MaxPhiValue;
+
+		if (PAndInvM[Index].InvM == (Softs::FSolverReal)0.)
+		{
+			continue;
+		}
+
+		const FVec3 RigidSpacePosition(SolverToThis.TransformPositionNoScale(PAndInvM[Index].P));
+
+		// only query points inside of the grid 
+		if (LocalBoundingBox.Contains(RigidSpacePosition))
+		{
+			const int32 ClosestActiveBoneIndex = GetClosestActiveBoneIndex(RigidSpacePosition);
+			const FVector3f LocalGridCornerShift = (FVector3f)ActiveBonesRelativeTransforms[ClosestActiveBoneIndex].GetTranslation() - (FVector3f)ActiveBonesReferenceTranslations[ClosestActiveBoneIndex];
+			const FVector3f LocalGridCornerShifted = TrainingGridMin + LocalGridCornerShift;
+
+			TVector<FRealSingle, 3> LocationMSShifted = RigidSpacePosition - LocalGridCornerShifted;
+			for (int32 j = 0; j < 3; j++)
+			{
+				MLInputForPhis[InteriorParticlesIndex * ModelInputShapeSize + j] = TrainingGridVectorsScaled[j][0] * LocationMSShifted[0] + TrainingGridVectorsScaled[j][1] * LocationMSShifted[1] + TrainingGridVectorsScaled[j][2] * LocationMSShifted[2];
+			}
+			ParticlesInsideTheGridIndexes[InteriorParticlesIndex] = Index;
+			InteriorParticlesIndex++;
+		}
+	}
+
+	MLInputForPhis.SetNum(InteriorParticlesIndex * ModelInputShapeSize);
+	MLOutputForPhis.SetNum(InteriorParticlesIndex * ModelOutputShapeSize);
+	ParticlesInsideTheGridIndexes.SetNum(InteriorParticlesIndex);
+
+	if (SignedDistanceNeuralInferences[MLLevelsetThread].IsValid())
+	{
+		SignedDistanceNeuralInferences[MLLevelsetThread].RunInference(MLInputForPhis, MLOutputForPhis, ModelInputShapeSize, ModelOutputShapeSize, MLWeightsIn);
+	}
+
+	TArray<int32> NegativePhiValueIndices;
+	NegativePhiValueIndices.Init(-1, NumParticles);
+
+	// Only Calculate Normals of those particles which are inside of the Body
+	int32 NumberOfNegativePhis = 0;
+	for (int32 IndexML = 0; IndexML < ParticlesInsideTheGridIndexes.Num(); IndexML++)
+	{
+		int32 Index = ParticlesInsideTheGridIndexes[IndexML];
+		Phis[Index] = SignedDistanceScaling * MLOutputForPhis[IndexML * ModelOutputShapeSize];
+		if (Phis[Index] < CollisionThickness)
+		{
+			NegativePhiValueIndices[NumberOfNegativePhis] = Index;
+			NumberOfNegativePhis++;
+		}
+	}
+	NegativePhiValueIndices.SetNum(NumberOfNegativePhis);
+
+	TArray<int32> CorrectZoneIndices;
+	int32 NumberOfCorrectZoneIndices = 0;
+
+	if (bUseIncorrectZoneModel)
+	{
+		TArray<float, TAlignedHeapAllocator<64>> MLInputForIncorrectZone;
+		TArray<float, TAlignedHeapAllocator<64>> MLOutputForIncorrectZone;
+		const int32 ModelOutputShapeSizeForIncorrectZone = 1;
+
+		MLInputForIncorrectZone.SetNumUninitialized(NegativePhiValueIndices.Num() * ModelInputShapeSize);
+		MLOutputForIncorrectZone.SetNumUninitialized(NegativePhiValueIndices.Num() * ModelOutputShapeSizeForIncorrectZone);
+		CorrectZoneIndices.SetNumUninitialized(NegativePhiValueIndices.Num());
+
+		for (int32 NegativePhiIndex = 0; NegativePhiIndex < NegativePhiValueIndices.Num(); NegativePhiIndex++)
+		{
+			int32 Index = NegativePhiValueIndices[NegativePhiIndex];
+			const FVec3 PositionLocal(SolverToThis.TransformPositionNoScale(PAndInvM[Index].P));
+			const int32 ClosestBoneJointIndex = GetClosestActiveBoneIndex(PositionLocal);
+			const FVector3f LocalGridCornerShift = (FVector3f)ActiveBonesRelativeTransforms[ClosestBoneJointIndex].GetTranslation() - (FVector3f)ActiveBonesReferenceTranslations[ClosestBoneJointIndex];
+			const FVector3f LocalGridCornerShifted = TrainingGridMin + LocalGridCornerShift;
+
+			TVector<FRealSingle, 3> LocationMSShifted = PositionLocal - LocalGridCornerShifted;
+			for (int32 j = 0; j < 3; j++)
+			{
+				MLInputForIncorrectZone[NegativePhiIndex * ModelInputShapeSize + j] = TrainingGridVectorsScaled[j][0] * LocationMSShifted[0] + TrainingGridVectorsScaled[j][1] * LocationMSShifted[1] + TrainingGridVectorsScaled[j][2] * LocationMSShifted[2];
+			}
+		}
+		TArray<TArray<float, TAlignedHeapAllocator<64>>> NetworkWeightsInputForIncorrectZone;
+		ComputeIncorrectZoneNetworkWeightsInput(NetworkWeightsInputForIncorrectZone);
+		IncorrectZoneNeuralInferences[MLLevelsetThread].RunInference(MLInputForIncorrectZone, MLOutputForIncorrectZone, ModelInputShapeSize, ModelOutputShapeSizeForIncorrectZone, NetworkWeightsInputForIncorrectZone);
+
+		for (int32 NegativePhiIndex = 0; NegativePhiIndex < NegativePhiValueIndices.Num(); NegativePhiIndex++)
+		{
+			int32 Index = NegativePhiValueIndices[NegativePhiIndex];
+			if (MLOutputForIncorrectZone[NegativePhiIndex] > 0)
+			{
+				CorrectZoneIndices[NumberOfCorrectZoneIndices] = Index;
+				NumberOfCorrectZoneIndices++;
+			}
+			else
+			{
+				Phis[Index] = MaxPhiValue;
+			}
+		}
+		CorrectZoneIndices.SetNum(NumberOfCorrectZoneIndices);
+	}
+	else
+	{
+		CorrectZoneIndices = MoveTemp(NegativePhiValueIndices);
+	}
+
+	TArray<float, TAlignedHeapAllocator<64>> MLInputForNormals;
+	MLInputForNormals.SetNum(CorrectZoneIndices.Num() * ModelInputShapeSize * ModelInputShapeSize);
+	TArray<float, TAlignedHeapAllocator<64>> MLOutputForNormals;
+	MLOutputForNormals.SetNum(CorrectZoneIndices.Num() * ModelInputShapeSize * ModelOutputShapeSize);
+
+	for (int32 CorrectZoneIndex = 0; CorrectZoneIndex < CorrectZoneIndices.Num(); CorrectZoneIndex++)
+	{
+		int32 Index = CorrectZoneIndices[CorrectZoneIndex];
+
+		const FVec3 RigidSpacePosition(SolverToThis.TransformPositionNoScale(PAndInvM[Index].P));
+		const int32 ClosestBoneJointIndex = GetClosestActiveBoneIndex(RigidSpacePosition);
+		const FVector3f LocalGridCornerShift = (FVector3f)ActiveBonesRelativeTransforms[ClosestBoneJointIndex].GetTranslation() - (FVector3f)ActiveBonesReferenceTranslations[ClosestBoneJointIndex];
+		const FVector3f LocalGridCornerShifted = TrainingGridMin + LocalGridCornerShift;
+		TVector<FRealSingle, 3> LocationMSShifted = RigidSpacePosition - LocalGridCornerShifted;
+
+		for (uint32 NormalDirectionIndex = 0; NormalDirectionIndex < 3; NormalDirectionIndex++)
+		{
+			for (uint32 i = 0; i < ModelInputShapeSize; i++)
+			{
+				MLInputForNormals[CorrectZoneIndex * ModelInputShapeSize * ModelInputShapeSize + NormalDirectionIndex * ModelInputShapeSize + i] =
+					TrainingGridUnitAxesXYZ[i].Dot(LocationMSShifted) / TrainingGridAxesLengthsXYZ[i];
+				MLInputForNormals[CorrectZoneIndex * ModelInputShapeSize * ModelInputShapeSize + NormalDirectionIndex * ModelInputShapeSize + i] +=
+					DeltaForFiniteDifferenceForNormalCalc * float(TrainingGridUnitAxesXYZ[i][NormalDirectionIndex]) / TrainingGridAxesLengthsXYZ[i];
+			}
+		}
+	}
+
+	if (SignedDistanceNeuralInferences[MLLevelsetThread].IsValid())
+	{
+		SignedDistanceNeuralInferences[MLLevelsetThread].RunInference(MLInputForNormals, MLOutputForNormals, ModelInputShapeSize, ModelOutputShapeSize, MLWeightsIn);
+	}
+
+	for (int32 CorrectZoneIndex = 0; CorrectZoneIndex < CorrectZoneIndices.Num(); CorrectZoneIndex++)
+	{
+		int32 Index = CorrectZoneIndices[CorrectZoneIndex];
+		FVec3 Normal;
+		for (int32 i = 0; i < 3; i++)
+		{
+			Normal[i] = (SignedDistanceScaling * MLOutputForNormals[CorrectZoneIndex * ModelInputShapeSize * ModelOutputShapeSize + ModelOutputShapeSize * i] - Phis[Index]) / DeltaForFiniteDifferenceForNormalCalc;
+		}
+
+		constexpr FReal EpsilonForSafeNormalization = 1.0e-8;
+		Normal.SafeNormalize(EpsilonForSafeNormalization);
+		NormalsOut[Index] = (Softs::FSolverVec3)Normal;
+	}	
+}
+
+void FMLLevelSet::BatchPhiWithNormalWithGreatestInfluenceBone(const TConstArrayView<Softs::FPAndInvM> PAndInvMArray,
+	const Softs::FSolverRigidTransform3& SolverToThis, TArray<Softs::FSolverReal>& OutBatchPhis, TArray<Softs::FSolverVec3>& OutBatchNormals, 
+	TArray<int32>& OutBatchGreatestInfluenceSolverBone, const Chaos::Softs::FSolverReal CollisionThickness, 
+	int32 MLLevelsetThread, int32 BatchBegin, int32 BatchEnd) const
+{	
+	TRACE_CPUPROFILER_EVENT_SCOPE(FMLLevelSet_BatchPhiWithNormal);
+	const uint32 ModelInputShapeSize = 3;  // XYZ coordinates wrt local bounding box
+	const uint32 ModelOutputShapeSize = 1; // At the moment there is a single output (SignedDistance).
+	const Chaos::Softs::FSolverReal MaxPhiValue = 2 * SignedDistanceScaling;
+	const float DeltaForFiniteDifferenceForNormalCalc = SignedDistanceScaling / float(100.0); //delta for finite differencing
+
+	const int32 NumParticles = BatchEnd - BatchBegin;
+	const Softs::FPAndInvM* const PAndInvM = PAndInvMArray.GetData();
+	Softs::FSolverReal* const Phis = OutBatchPhis.GetData();
+	Softs::FSolverVec3* const NormalsOut = OutBatchNormals.GetData();
+	int32* const GreatestInfluenceSolverBones = OutBatchGreatestInfluenceSolverBone.GetData();
+
+	//Query the phi/signed distance values on the cloth
+	TArray<float, TAlignedHeapAllocator<64>> MLInputForPhis;
+	MLInputForPhis.SetNum(NumParticles * ModelInputShapeSize);
+	TArray<float, TAlignedHeapAllocator<64>> MLOutputForPhis;
+
+	TArray<int32> ParticlesInsideTheGridIndexes;
+	ParticlesInsideTheGridIndexes.SetNum(NumParticles);
+	TArray<TArray<float, TAlignedHeapAllocator<64>>> MLWeightsIn;
+	ComputeSignedDistanceNetworkWeightsInput(MLWeightsIn);
+		
+	TArray<FVector3f> TrainingGridVectorsScaled;
+	TrainingGridVectorsScaled.SetNum(3);
+	for (int32 Index = 0; Index < 3; Index++)
+	{
+		TrainingGridVectorsScaled[Index] = TrainingGridUnitAxesXYZ[Index] / TrainingGridAxesLengthsXYZ[Index];
+	}
+
+	int32 InteriorParticlesIndex = 0;
+	for (int32 Index = BatchBegin; Index < BatchEnd; Index++)
+	{
+		OutBatchPhis[Index] = MaxPhiValue;
+
+		if (PAndInvM[Index].InvM == (Softs::FSolverReal)0.)
+		{
+			continue;
+		}
+
+		const FVec3 RigidSpacePosition(SolverToThis.TransformPositionNoScale(PAndInvM[Index].P));
+
+		// only query points inside of the grid
+		if (LocalBoundingBox.Contains(RigidSpacePosition))
+		{
+			const int32 ClosestBoneJointIndex = GetClosestActiveBoneIndex(RigidSpacePosition);
+			const FVector3f LocalGridCornerShift = (FVector3f)ActiveBonesRelativeTransforms[ClosestBoneJointIndex].GetTranslation() - (FVector3f)ActiveBonesReferenceTranslations[ClosestBoneJointIndex];
+			const FVector3f LocalGridCornerShifted = TrainingGridMin + LocalGridCornerShift;
+
+			TVector<FRealSingle, 3> LocationMSShifted = RigidSpacePosition - LocalGridCornerShifted;
+			for (int32 j = 0; j < 3; j++)
+			{
+				MLInputForPhis[InteriorParticlesIndex * ModelInputShapeSize + j] = TrainingGridVectorsScaled[j][0] * LocationMSShifted[0] + TrainingGridVectorsScaled[j][1] * LocationMSShifted[1] + TrainingGridVectorsScaled[j][2] * LocationMSShifted[2];
+			}
+			ParticlesInsideTheGridIndexes[InteriorParticlesIndex] = Index;
+			InteriorParticlesIndex++;
+		}
+		GreatestInfluenceSolverBones[Index] = GetGreatestInfluenceBoneSolverIndex();;
+	}
+
+	MLInputForPhis.SetNum(InteriorParticlesIndex * ModelInputShapeSize);
+	MLOutputForPhis.SetNum(InteriorParticlesIndex * ModelOutputShapeSize);
+	ParticlesInsideTheGridIndexes.SetNum(InteriorParticlesIndex);
+
+	if (SignedDistanceNeuralInferences[MLLevelsetThread].IsValid())
+	{
+		SignedDistanceNeuralInferences[MLLevelsetThread].RunInference(MLInputForPhis, MLOutputForPhis, ModelInputShapeSize, ModelOutputShapeSize, MLWeightsIn);
+	}	
+
+	TArray<int32> NegativePhiValueIndices;
+	NegativePhiValueIndices.SetNumUninitialized(NumParticles);
+
+	// Only Calculate Normals of those particles which are inside of the Body
+	int32 NumberOfNegativePhis = 0;
+	for (int32 IndexML = 0; IndexML < ParticlesInsideTheGridIndexes.Num(); IndexML++)
+	{
+		int32 Index = ParticlesInsideTheGridIndexes[IndexML];
+		Phis[Index] = SignedDistanceScaling * MLOutputForPhis[IndexML * ModelOutputShapeSize];
+		if (Phis[Index] < CollisionThickness)
+		{
+			NegativePhiValueIndices[NumberOfNegativePhis] = Index;
+			NumberOfNegativePhis++;
+		}
+	}
+	NegativePhiValueIndices.SetNum(NumberOfNegativePhis);
+
+	TArray<int32> CorrectZoneIndices;
+	int32 NumberOfCorrectZoneIndices = 0;
+
+	if (bUseIncorrectZoneModel)
+	{
+		TArray<float, TAlignedHeapAllocator<64>> MLInputForIncorrectZone;
+		TArray<float, TAlignedHeapAllocator<64>> MLOutputForIncorrectZone;
+		const int32 ModelOutputShapeSizeForIncorrectZone = 1;
+
+		MLInputForIncorrectZone.SetNumUninitialized(NegativePhiValueIndices.Num() * ModelInputShapeSize);
+		MLOutputForIncorrectZone.SetNumUninitialized(NegativePhiValueIndices.Num() * ModelOutputShapeSizeForIncorrectZone);
+		CorrectZoneIndices.SetNumUninitialized(NegativePhiValueIndices.Num());
+
+		for (int32 NegativePhiIndex = 0; NegativePhiIndex < NegativePhiValueIndices.Num(); NegativePhiIndex++)
+		{
+			int32 Index = NegativePhiValueIndices[NegativePhiIndex];
+			const FVec3 PositionLocal(SolverToThis.TransformPositionNoScale(PAndInvM[Index].P));
+			const int32 ClosestBoneJointIndex = GetClosestActiveBoneIndex(PositionLocal);
+			const FVector3f LocalGridCornerShift = (FVector3f)ActiveBonesRelativeTransforms[ClosestBoneJointIndex].GetTranslation() - (FVector3f)ActiveBonesReferenceTranslations[ClosestBoneJointIndex];
+			const FVector3f LocalGridCornerShifted = TrainingGridMin + LocalGridCornerShift;
+
+			TVector<FRealSingle, 3> LocationMSShifted = PositionLocal - LocalGridCornerShifted;
+			for (int32 j = 0; j < 3; j++)
+			{
+				MLInputForIncorrectZone[NegativePhiIndex * ModelInputShapeSize + j] = TrainingGridVectorsScaled[j][0] * LocationMSShifted[0] + TrainingGridVectorsScaled[j][1] * LocationMSShifted[1] + TrainingGridVectorsScaled[j][2] * LocationMSShifted[2];
+			}
+		}
+		TArray<TArray<float, TAlignedHeapAllocator<64>>> NetworkWeightsInputForIncorrectZone;
+		ComputeIncorrectZoneNetworkWeightsInput(NetworkWeightsInputForIncorrectZone);
+		IncorrectZoneNeuralInferences[MLLevelsetThread].RunInference(MLInputForIncorrectZone, MLOutputForIncorrectZone, ModelInputShapeSize, ModelOutputShapeSizeForIncorrectZone, NetworkWeightsInputForIncorrectZone);
+		
+		for (int32 NegativePhiIndex = 0; NegativePhiIndex < NegativePhiValueIndices.Num(); NegativePhiIndex++)
+		{
+			int32 Index = NegativePhiValueIndices[NegativePhiIndex];
+			if (MLOutputForIncorrectZone[NegativePhiIndex] > 0)
+			{
+				CorrectZoneIndices[NumberOfCorrectZoneIndices] = Index;
+				NumberOfCorrectZoneIndices++;
+			}
+			else
+			{
+				Phis[Index] = MaxPhiValue;
+			}
+		}
+		CorrectZoneIndices.SetNum(NumberOfCorrectZoneIndices);
+	}
+	else
+	{
+		CorrectZoneIndices = MoveTemp(NegativePhiValueIndices);
+	}
+	
+	TArray<float, TAlignedHeapAllocator<64>> MLInputForNormals;
+	MLInputForNormals.SetNum(CorrectZoneIndices.Num() * ModelInputShapeSize * ModelInputShapeSize);
+	TArray<float, TAlignedHeapAllocator<64>> MLOutputForNormals;
+	MLOutputForNormals.SetNum(CorrectZoneIndices.Num() * ModelInputShapeSize * ModelOutputShapeSize);
+	
+	for (int32 CorrectZoneIndex = 0; CorrectZoneIndex < CorrectZoneIndices.Num(); CorrectZoneIndex++)
+	{
+		int32 Index = CorrectZoneIndices[CorrectZoneIndex];
+
+		const FVec3 RigidSpacePosition(SolverToThis.TransformPositionNoScale(PAndInvM[Index].P));
+		const int32 ClosestBoneJointIndex = GetClosestActiveBoneIndex(RigidSpacePosition);
+		const FVector3f LocalGridCornerShift = (FVector3f)ActiveBonesRelativeTransforms[ClosestBoneJointIndex].GetTranslation() - (FVector3f)ActiveBonesReferenceTranslations[ClosestBoneJointIndex];
+		const FVector3f LocalGridCornerShifted = TrainingGridMin + LocalGridCornerShift;
+		TVector<FRealSingle, 3> LocationMSShifted = RigidSpacePosition - LocalGridCornerShifted;
+
+		for (uint32 NormalDirectionIndex = 0; NormalDirectionIndex < 3; NormalDirectionIndex++)
+		{
+			for (uint32 i = 0; i < ModelInputShapeSize; i++)
+			{
+				MLInputForNormals[CorrectZoneIndex * ModelInputShapeSize * ModelInputShapeSize + NormalDirectionIndex * ModelInputShapeSize + i] =
+					TrainingGridUnitAxesXYZ[i].Dot(LocationMSShifted) / TrainingGridAxesLengthsXYZ[i];
+				MLInputForNormals[CorrectZoneIndex * ModelInputShapeSize * ModelInputShapeSize + NormalDirectionIndex * ModelInputShapeSize + i] +=
+					DeltaForFiniteDifferenceForNormalCalc * float(TrainingGridUnitAxesXYZ[i][NormalDirectionIndex]) / TrainingGridAxesLengthsXYZ[i];
+			}
+		}
+	}
+
+	if (SignedDistanceNeuralInferences[MLLevelsetThread].IsValid())
+	{
+		SignedDistanceNeuralInferences[MLLevelsetThread].RunInference(MLInputForNormals, MLOutputForNormals, ModelInputShapeSize, ModelOutputShapeSize, MLWeightsIn);
+	}
+
+	for (int32 CorrectZoneIndex = 0; CorrectZoneIndex < CorrectZoneIndices.Num(); CorrectZoneIndex++)
+	{
+		int32 Index = CorrectZoneIndices[CorrectZoneIndex];
+		FVec3 Normal;
+		for (int32 i = 0; i < 3; i++)
+		{
+			Normal[i] = (SignedDistanceScaling * MLOutputForNormals[CorrectZoneIndex * ModelInputShapeSize * ModelOutputShapeSize + ModelOutputShapeSize * i] - Phis[Index]) / DeltaForFiniteDifferenceForNormalCalc;
+		}
+
+		constexpr FReal EpsilonForSafeNormalization = 1.0e-8;
+		Normal.SafeNormalize(EpsilonForSafeNormalization);
+		NormalsOut[Index] = (Softs::FSolverVec3)Normal;
+	}
+}
+}
diff --git a/Engine/Source/Runtime/Experimental/Chaos/Private/Chaos/PBDSoftBodyCollisionConstraint.cpp b/Engine/Source/Runtime/Experimental/Chaos/Private/Chaos/PBDSoftBodyCollisionConstraint.cpp
--- a/Engine/Source/Runtime/Experimental/Chaos/Private/Chaos/PBDSoftBodyCollisionConstraint.cpp
+++ b/Engine/Source/Runtime/Experimental/Chaos/Private/Chaos/PBDSoftBodyCollisionConstraint.cpp
@@ -2,6 +2,7 @@
 
 #include "Chaos/PBDSoftBodyCollisionConstraint.h"
 #include "Chaos/Levelset.h"
+#include "Chaos/MLLevelset.h"
 #include "Chaos/SoftsEvolutionLinearSystem.h"
 #include "Chaos/WeightedLatticeImplicitObject.h"
 #include "ChaosStats.h"
@@ -30,6 +31,8 @@ static_assert(ispc::ImplicitObjectType::TaperedCapsule == Chaos::ImplicitObjectT
 static_assert(ispc::ImplicitObjectType::Convex == Chaos::ImplicitObjectType::Convex);
 static_assert(ispc::ImplicitObjectType::WeightedLatticeBone == Chaos::ImplicitObjectType::WeightedLatticeBone);
 static_assert(ispc::ImplicitObjectType::IsWeightedLattice == Chaos::ImplicitObjectType::IsWeightedLattice);
+static_assert(ispc::ImplicitObjectType::MLLevelSet == Chaos::ImplicitObjectType::MLLevelSet);
+static_assert(ispc::ImplicitObjectType::MLLevelSetBone == Chaos::ImplicitObjectType::MLLevelSetBone);
 static_assert(ispc::ImplicitObjectType::WeightedLatticeLevelSetType == (Chaos::ImplicitObjectType::IsWeightedLattice | Chaos::ImplicitObjectType::LevelSet));
 static_assert(sizeof(ispc::TArray) == sizeof(TArray<int>));
 // Sphere
@@ -116,9 +119,11 @@ static FAutoConsoleVariableRef CVarChaosSoftBodyCollisionISPCEnabled(TEXT("p.Cha
 #endif
 #endif
 
+static int MLLevelSetParallelize = 0;
 static int32 Chaos_SoftBodyCollision_ISPC_ParallelBatchSize = 128;
 #if !UE_BUILD_SHIPPING
 FAutoConsoleVariableRef CVarChaosSoftBodyCollisionISPCParallelBatchSize(TEXT("p.Chaos.SoftBodyCollision.ISPC.ParallelBatchSize"), Chaos_SoftBodyCollision_ISPC_ParallelBatchSize, TEXT("Parallel batch size for ISPC"));
+static FAutoConsoleVariableRef CVarChaosMLLevelSetParallelize(TEXT("p.Chaos.Cloth.ParallelizeMLLevelSet"), MLLevelSetParallelize, TEXT("Parallelize MLLevelSet Call: 0 = No, 1 = Yes"));
 #endif
 
 namespace Chaos::Softs {
@@ -126,8 +131,7 @@ namespace Chaos::Softs {
 // Helper function to call PhiWithNormal and return data to ISPC
 extern "C" void GetPhiWithNormalCollisionParticleRange(const uint8 * CollisionParticlesRange, const FSolverReal * InV, FSolverReal * Normal, FSolverReal * Phi, const int32 i, const int32 ProgramCount, const int32 Mask)
 {
-	const FSolverCollisionParticlesRange& C = *(const FSolverCollisionParticlesRange*)CollisionParticlesRange;
-
+	const FSolverCollisionParticlesRange& C = *(const FSolverCollisionParticlesRange*)CollisionParticlesRange;	
 	for (int32 Index = 0; Index < ProgramCount; ++Index)
 	{
 		if (Mask & (1 << Index))
@@ -239,7 +243,11 @@ void FPBDSoftBodyCollisionConstraintBase::Apply(FSolverParticlesRange& Particles
 	{
 		if (bWithFriction)
 		{
-			if (bRealTypeCompatibleWithISPC && bChaos_SoftBodyCollision_ISPC_Enabled)
+			if (bDoCollisionPerTimestep)
+			{	
+				ApplyInternalISPCPerTimestep<true>(Particles, Dt, CollisionParticles);				
+			}
+			else if (bRealTypeCompatibleWithISPC && bChaos_SoftBodyCollision_ISPC_Enabled)
 			{
 				ApplyInternalISPC<true>(Particles, Dt, CollisionParticles);
 			}
@@ -254,7 +262,11 @@ void FPBDSoftBodyCollisionConstraintBase::Apply(FSolverParticlesRange& Particles
 		}
 		else
 		{
-			if (bRealTypeCompatibleWithISPC && bChaos_SoftBodyCollision_ISPC_Enabled)
+			if (bDoCollisionPerTimestep)
+			{
+				ApplyInternalISPCPerTimestep<false>(Particles, Dt, CollisionParticles);				
+			}
+			else if (bRealTypeCompatibleWithISPC && bChaos_SoftBodyCollision_ISPC_Enabled)
 			{
 				ApplyInternalISPC<false>(Particles, Dt, CollisionParticles);
 			}
@@ -287,36 +299,25 @@ void FPBDSoftBodyCollisionConstraintBase::ApplyInternal(FSolverParticlesRange& P
 		{
 			for(int32 CollisionIndex = 0; CollisionIndex < CollisionParticlesRange.GetRangeSize(); ++CollisionIndex)
 			{
-				if (CollisionParticlesRange.GetGeometry(CollisionIndex)->GetType() == Chaos::ImplicitObjectType::WeightedLatticeBone)
+				const Chaos::EImplicitObjectType CollisionType = CollisionParticlesRange.GetGeometry(CollisionIndex)->GetType();
+				if ( CollisionType == Chaos::ImplicitObjectType::WeightedLatticeBone || CollisionType == TWeightedLatticeImplicitObject<FLevelSet>::StaticType())
 				{
+					// Skip skinned levelsets here. We will do a batch query for those in a second pass to match ISPC.
+					continue;
+				}
+
+				if (CollisionType == Chaos::ImplicitObjectType::MLLevelSetBone || CollisionType == Chaos::ImplicitObjectType::MLLevelSet)
+				{
+					// Skip  mllevelsets here. 
 					continue;
 				}
 
 				const FSolverRigidTransform3 Frame(CollisionParticlesRange.GetX(CollisionIndex), CollisionParticlesRange.R(CollisionIndex));
 				const FVec3 RigidSpacePosition(Frame.InverseTransformPosition(PAndInvM[Index].P));  // PhiWithNormal requires FReal based arguments
 				FVec3 ImplicitNormal;                                                                // since implicits don't use FSolverReal
-				FSolverReal Phi;
-				int32 VelocityBone = CollisionIndex;
-				FSolverReal Penetration;
-				if (const TWeightedLatticeImplicitObject<FLevelSet>* LevelSet = CollisionParticlesRange.GetGeometry(CollisionIndex)->GetObject< TWeightedLatticeImplicitObject<FLevelSet> >())
-				{
-					FWeightedLatticeImplicitObject::FEmbeddingCoordinate SurfaceCoord;
-					Phi = (FSolverReal)LevelSet->PhiWithNormalAndSurfacePoint(RigidSpacePosition, ImplicitNormal, SurfaceCoord);
-					Penetration = CollisionThickness - Phi; // This is related to the Normal impulse
-					if (bWithFriction && Penetration > (FSolverReal)0.)
-					{
-						const int32 StrongestBone = SurfaceCoord.GreatestInfluenceBone(LevelSet->GetBoneData());
-						if (StrongestBone != INDEX_NONE)
-						{
-							VelocityBone = LevelSet->GetSolverBoneIndices()[StrongestBone];
-						}
-					}
-				}
-				else
-				{
-					Phi = (FSolverReal)CollisionParticlesRange.GetGeometry(CollisionIndex)->PhiWithNormal(RigidSpacePosition, ImplicitNormal);
-					Penetration = CollisionThickness - Phi; // This is related to the Normal impulse
-				}
+				const FSolverReal Phi = (FSolverReal)CollisionParticlesRange.GetGeometry(CollisionIndex)->PhiWithNormal(RigidSpacePosition, ImplicitNormal);
+				const int32 VelocityBone = CollisionIndex;
+				const FSolverReal Penetration = CollisionThickness - Phi; // This is related to the Normal impulse
 				const FSolverVec3 Normal(ImplicitNormal);
 
 				if (Penetration > (FSolverReal)0.)
@@ -354,6 +355,74 @@ void FPBDSoftBodyCollisionConstraintBase::ApplyInternal(FSolverParticlesRange& P
 			}
 		}
 	});
+
+	// Skinned levelset passes
+	for (const FSolverCollisionParticlesRange& CollisionParticlesRange : CollisionParticles)
+	{
+		for (int32 CollisionIndex = 0; CollisionIndex < CollisionParticlesRange.GetRangeSize(); ++CollisionIndex)
+		{
+			if (const TWeightedLatticeImplicitObject<FLevelSet>* const LevelSet = CollisionParticlesRange.GetGeometry(CollisionIndex)->GetObject< TWeightedLatticeImplicitObject<FLevelSet> >())
+			{
+				const FSolverRigidTransform3 Frame(CollisionParticlesRange.X(CollisionIndex), CollisionParticlesRange.R(CollisionIndex));
+				if constexpr (bWithFriction)
+				{
+					LevelSet->BatchPhiWithNormalAndGreatestInfluenceBone(Particles.GetPAndInvM(), Frame.Inverse(), CollisionThickness, BatchPhis, BatchNormals, BatchVelocityBones);
+					check(BatchPhis.Num() == Particles.GetRangeSize());
+					check(BatchNormals.Num() == Particles.GetRangeSize());
+					check(BatchVelocityBones.Num() == Particles.GetRangeSize());
+				}
+				else
+				{
+					LevelSet->BatchPhiWithNormal(Particles.GetPAndInvM(), Frame.Inverse(), BatchPhis, BatchNormals);
+					check(BatchPhis.Num() == Particles.GetRangeSize());
+					check(BatchNormals.Num() == Particles.GetRangeSize());
+				}
+				const FSolverReal* const CollisionPhis = BatchPhis.GetData();
+				const FSolverVec3* const CollisionNormals = BatchNormals.GetData();
+				const int32* const VelocityBones = bWithFriction ? BatchVelocityBones.GetData() : nullptr;
+				PhysicsParallelFor(Particles.GetRangeSize(), [this, Dt, &CollisionParticlesRange, CollisionIndex, &PAndInvM, &X, &Frame, &CollisionPhis, &CollisionNormals, &VelocityBones](int32 Index)
+				{
+					if (PAndInvM[Index].InvM == (FSolverReal)0.)
+					{
+						return;
+					}
+
+					const FSolverReal Penetration = CollisionThickness - CollisionPhis[Index];
+					if (Penetration > (FSolverReal)0.)
+					{
+						const FSolverVec3 NormalWorld = Frame.TransformVector(CollisionNormals[Index]);
+						if constexpr (bLockAndWriteContacts)
+						{
+							checkSlow(Contacts);
+							checkSlow(Normals);
+							checkSlow(Phis);
+							FScopeLock Lock(&DebugMutex);
+							CollisionParticlesRange.GetArrayView(*CollisionParticleCollided)[CollisionIndex] = true;
+							Contacts->Emplace(PAndInvM[Index].P);
+							Normals->Emplace(NormalWorld);
+							Phis->Emplace(CollisionPhis[Index]);
+						}
+
+						PAndInvM[Index].P = Penetration * NormalWorld;
+						if constexpr (bWithFriction)
+						{
+							const int32 VelocityBone = CollisionParticlesRange.IsValidIndex(VelocityBones[Index]) ? VelocityBones[Index] : CollisionIndex;
+							const FSolverVec3 VectorToPoint = PAndInvM[Index].P - CollisionParticlesRange.X(VelocityBone);
+							const FSolverVec3 RelativeDisplacement = (PAndInvM[Index].P - X[Index]) - (CollisionParticlesRange.V(VelocityBone) + FSolverVec3::CrossProduct(CollisionParticlesRange.W(VelocityBone), VectorToPoint)) * Dt; // This corresponds to the tangential velocity multiplied by dt (friction will drive this to zero if it is high enough)
+							const FSolverVec3 RelativeDisplacementTangent = RelativeDisplacement - FSolverVec3::DotProduct(RelativeDisplacement, NormalWorld) * NormalWorld; // Project displacement into the tangential plane
+							const FSolverReal RelativeDisplacementTangentLength = RelativeDisplacementTangent.Size();
+							if (RelativeDisplacementTangentLength >= UE_SMALL_NUMBER)
+							{
+								const FSolverReal PositionCorrection = FMath::Min<FSolverReal>(Penetration * FrictionCoefficient, RelativeDisplacementTangentLength);
+								const FSolverReal CorrectionRatio = PositionCorrection / RelativeDisplacementTangentLength;
+								PAndInvM[Index].P -= CorrectionRatio * RelativeDisplacementTangent;
+							}
+						}
+					}
+				});
+			}
+		}
+	}
 }
 
 template<bool bLockAndWriteContacts, bool bWithFriction>
@@ -379,6 +448,11 @@ void FPBDSoftBodyCollisionConstraintBase::ApplyInternalCCD(FSolverParticlesRange
 					continue;
 				}
 
+				if (CollisionParticlesRange.GetGeometry(CollisionIndex)->GetType() == Chaos::ImplicitObjectType::MLLevelSetBone)
+				{
+					continue;
+				}
+
 				const FSolverRigidTransform3 Frame(CollisionParticlesRange.GetX(CollisionIndex), CollisionParticlesRange.R(CollisionIndex));
 
 				const Pair<FVec3, bool> PointPair = CollisionParticlesRange.GetGeometry(CollisionIndex)->FindClosestIntersection(  // Geometry operates in FReal
@@ -458,6 +532,7 @@ void FPBDSoftBodyCollisionConstraintBase::ApplyInternalISPC(FSolverParticlesRang
 #if INTEL_ISPC
 	check(bRealTypeCompatibleWithISPC);
 	const int32 NumBatches = FMath::CeilToInt((FSolverReal)(Particles.GetRangeSize()) / (FSolverReal)Chaos_SoftBodyCollision_ISPC_ParallelBatchSize);
+	
 	PhysicsParallelFor(NumBatches, [this, &Particles, Dt, &CollisionParticles](int32 BatchNumber)
 	{
 		const int32 BatchBegin = (Chaos_SoftBodyCollision_ISPC_ParallelBatchSize * BatchNumber);
@@ -467,6 +542,7 @@ void FPBDSoftBodyCollisionConstraintBase::ApplyInternalISPC(FSolverParticlesRang
 		{
 			if constexpr (bWithFriction)
 			{
+				//This function skips Skinned LevelSets, MLLevelSets, and their proxy bones.
 				ispc::ApplyPerParticleCollisionFastFrictionNoGroupCheck(
 					(ispc::FVector4f*)Particles.GetPAndInvM().GetData(),
 					(const ispc::FVector3f*)Particles.XArray().GetData(),
@@ -488,6 +564,7 @@ void FPBDSoftBodyCollisionConstraintBase::ApplyInternalISPC(FSolverParticlesRang
 			}
 			else
 			{
+				//This function skips Skinned LevelSets, MLLevelSets, and their proxy bones.
 				ispc::ApplyPerParticleCollisionNoFrictionNoGroupCheck(
 					(ispc::FVector4f*)Particles.GetPAndInvM().GetData(),
 					(const ispc::FVector3f*)Particles.XArray().GetData(),
@@ -508,9 +585,712 @@ void FPBDSoftBodyCollisionConstraintBase::ApplyInternalISPC(FSolverParticlesRang
 			}
 		}
 	});
+	
+
+	// Skinned levelset passes
+	for (const FSolverCollisionParticlesRange& CollisionParticlesRange : CollisionParticles)
+	{
+		for (int32 CollisionIndex = 0; CollisionIndex < CollisionParticlesRange.GetRangeSize(); ++CollisionIndex)
+		{
+			if (const TWeightedLatticeImplicitObject<FLevelSet>* const LevelSet = CollisionParticlesRange.GetGeometry(CollisionIndex)->GetObject< TWeightedLatticeImplicitObject<FLevelSet> >())
+			{
+				TRACE_CPUPROFILER_EVENT_SCOPE(FPBDSoftBodyCollisionConstraint_ApplyInternalISPC_LevelSet);
+				const FSolverRigidTransform3 Frame = FSolverRigidTransform3(CollisionParticlesRange.X(CollisionIndex), CollisionParticlesRange.R(CollisionIndex)).Inverse();
+				if constexpr (bWithFriction)
+				{
+					LevelSet->BatchPhiWithNormalAndGreatestInfluenceBone(Particles.GetPAndInvM(), Frame, CollisionThickness, BatchPhis, BatchNormals, BatchVelocityBones);
+					check(BatchPhis.Num() == Particles.GetRangeSize());
+					check(BatchNormals.Num() == Particles.GetRangeSize());
+					check(BatchVelocityBones.Num() == Particles.GetRangeSize());
+
+
+					const FSolverRotation3 CPR_R = CollisionParticlesRange.R(CollisionIndex);
+					ispc::ApplyPerParticleCollisionFastFrictionBatchedPhiAndNormal(
+						(ispc::FVector4f*)Particles.GetPAndInvM().GetData(),
+						(const ispc::FVector3f*)Particles.XArray().GetData(),
+						CollisionIndex,
+						BatchPhis.GetData(),
+						(const ispc::FVector3f*)BatchNormals.GetData(),
+						BatchVelocityBones.GetData(),
+						(const ispc::FVector3f*)CollisionParticlesRange.GetV().GetData(),
+						(const ispc::FVector3f*)CollisionParticlesRange.XArray().GetData(),
+						(const ispc::FVector3f*)CollisionParticlesRange.GetW().GetData(),
+						reinterpret_cast<const ispc::FVector4f&>(CPR_R),
+						//(const ispc::FVector4f&)CollisionParticlesRange.R(CollisionIndex),
+						FrictionCoefficient,
+						CollisionThickness,
+						Dt,
+						CollisionParticlesRange.GetRangeSize(),
+						0,
+						Particles.GetRangeSize());
+				}
+				else
+				{
+					LevelSet->BatchPhiWithNormal(Particles.GetPAndInvM(), Frame, BatchPhis, BatchNormals);
+					check(BatchPhis.Num() == Particles.GetRangeSize());
+					check(BatchNormals.Num() == Particles.GetRangeSize());
+					const FSolverRotation3 CPR_R = CollisionParticlesRange.R(CollisionIndex);
+					ispc::ApplyPerParticleCollisionNoFrictionBatchedPhiAndNormal(
+						(ispc::FVector4f*)Particles.GetPAndInvM().GetData(),
+						BatchPhis.GetData(),
+						(const ispc::FVector3f*)BatchNormals.GetData(),
+						reinterpret_cast<const ispc::FVector4f&>(CPR_R),
+						//(const ispc::FVector4f&)CollisionParticlesRange.R(CollisionIndex),
+						CollisionThickness,
+						Dt,
+						0,
+						Particles.GetRangeSize());
+				}
+			}
+		}
+	}
+
+	// MLLevelset passes ApplyInternalISPC
+	for (const FSolverCollisionParticlesRange& CollisionParticlesRange : CollisionParticles)
+	{
+		const int32 NumberOfParticles = Particles.GetRangeSize();
+		FPAndInvM* const PAndInvM = Particles.GetPAndInvM().GetData();
+
+		if (MLLevelSetParallelize == 0)
+		{
+			for (int32 CollisionIndex = 0; CollisionIndex < CollisionParticlesRange.GetRangeSize(); ++CollisionIndex)
+			{
+				const FSolverRigidTransform3 Frame = FSolverRigidTransform3(CollisionParticlesRange.X(CollisionIndex), CollisionParticlesRange.R(CollisionIndex)).Inverse();
+				if (FMLLevelSet* MLLevelSet = CollisionParticlesRange.GetGeometry(CollisionIndex)->GetObject<FMLLevelSet>())
+				{
+					if constexpr (bWithFriction)
+					{
+						BatchPhis.SetNum(Particles.GetRangeSize());
+						BatchNormals.SetNum(Particles.GetRangeSize());
+						BatchVelocityBones.SetNum(Particles.GetRangeSize());
+
+						const int32 MLLevelSetThreadNum = 0;
+						const int32 BatchBegin = 0;
+						const int32 BatchEnd = Particles.GetRangeSize();
+
+						MLLevelSet->BatchPhiWithNormalWithGreatestInfluenceBone(Particles.GetPAndInvM(), Frame,
+							BatchPhis, BatchNormals, BatchVelocityBones,
+							CollisionThickness, MLLevelSetThreadNum, BatchBegin, BatchEnd);
+
+						check(BatchPhis.Num() == Particles.GetRangeSize());
+						check(BatchNormals.Num() == Particles.GetRangeSize());
+						check(BatchVelocityBones.Num() == Particles.GetRangeSize());
+
+						const FSolverRotation3 CPR_R = CollisionParticlesRange.R(CollisionIndex);
+						ispc::ApplyPerParticleCollisionFastFrictionBatchedPhiAndNormal(
+							(ispc::FVector4f*)Particles.GetPAndInvM().GetData(),
+							(const ispc::FVector3f*)Particles.XArray().GetData(),
+							CollisionIndex,
+							BatchPhis.GetData(),
+							(const ispc::FVector3f*)BatchNormals.GetData(),
+							BatchVelocityBones.GetData(),
+							(const ispc::FVector3f*)CollisionParticlesRange.GetV().GetData(),
+							(const ispc::FVector3f*)CollisionParticlesRange.XArray().GetData(),
+							(const ispc::FVector3f*)CollisionParticlesRange.GetW().GetData(),
+							reinterpret_cast<const ispc::FVector4f&>(CPR_R),
+							//(const ispc::FVector4f&)CollisionParticlesRange.R(CollisionIndex),
+							FrictionCoefficient,
+							CollisionThickness,
+							Dt,
+							CollisionParticlesRange.GetRangeSize(),
+							BatchBegin,
+							BatchEnd);
+					}
+					else
+					{
+						BatchPhis.SetNum(Particles.GetRangeSize());
+						BatchNormals.SetNum(Particles.GetRangeSize());
+
+						const int32 MLLevelSetThreadNum = 0;
+						const int32 BatchBegin = 0;
+						const int32 BatchEnd = Particles.GetRangeSize();
+
+						MLLevelSet->BatchPhiWithNormalWithGreatestInfluenceBone(Particles.GetPAndInvM(), Frame,
+							BatchPhis, BatchNormals, BatchVelocityBones,
+							CollisionThickness, MLLevelSetThreadNum, BatchBegin, BatchEnd);
+
+						check(BatchPhis.Num() == Particles.GetRangeSize());
+						check(BatchNormals.Num() == Particles.GetRangeSize());
+
+						const FSolverRotation3 CPR_R = CollisionParticlesRange.R(CollisionIndex);
+						ispc::ApplyPerParticleCollisionNoFrictionBatchedPhiAndNormal(
+							(ispc::FVector4f*)Particles.GetPAndInvM().GetData(),
+							BatchPhis.GetData(),
+							(const ispc::FVector3f*)BatchNormals.GetData(),
+							reinterpret_cast<const ispc::FVector4f&>(CPR_R),
+							//(const ispc::FVector4f&)(CollisionParticlesRange.R(CollisionIndex)),
+							CollisionThickness,
+							Dt,
+							BatchBegin,
+							BatchEnd);
+					}
+				}
+			}
+		}
+
+		if (MLLevelSetParallelize == 1)
+		{
+			//Find MLLevelSet CollisionIndices for Parallelization
+			TArray<int32> MLLevelSetIndexes;
+			for (int32 CollisionIndex = 0; CollisionIndex < CollisionParticlesRange.GetRangeSize(); ++CollisionIndex)
+			{
+				if (FMLLevelSet* MLLevelSet = CollisionParticlesRange.GetGeometry(CollisionIndex)->GetObject<FMLLevelSet>())
+				{
+					MLLevelSetIndexes.Add(CollisionIndex);
+				}
+			}
+
+			const int32 NumberOfMLLevelSets = MLLevelSetIndexes.Num();
+			if (NumberOfMLLevelSets > 0)
+			{
+				int32 NumberOfParallelThreadsPerMLLevelset = 1;
+
+				for (int32 MLCollisionIndex : MLLevelSetIndexes)
+				{
+					FMLLevelSet* MLLevelSet = CollisionParticlesRange.GetGeometry(MLCollisionIndex)->GetObject<FMLLevelSet>();
+					MLLevelSet->UpdateNeuralInferencesNumber(NumberOfParallelThreadsPerMLLevelset);
+				}
+
+				const int32 NumberOfParallelCalls = NumberOfParallelThreadsPerMLLevelset * NumberOfMLLevelSets;
+				TArray<TArray<FSolverReal>> BatchPhisArr;
+				TArray<TArray<FSolverVec3>> BatchNormalsArr;
+				TArray<TArray<int32>> BatchVelocityBonesArr;
+
+				BatchPhisArr.SetNum(MLLevelSetIndexes.Num());
+				BatchNormalsArr.SetNum(MLLevelSetIndexes.Num());
+				BatchVelocityBonesArr.SetNum(MLLevelSetIndexes.Num());
+
+				// Allocate TArray's for Parallel MLLevelSet Calls.
+				for (int32 MLLevelSetIndex = 0; MLLevelSetIndex < NumberOfMLLevelSets; MLLevelSetIndex++)
+				{
+					BatchPhisArr[MLLevelSetIndex].SetNum(NumberOfParticles);
+					BatchNormalsArr[MLLevelSetIndex].SetNum(NumberOfParticles);
+					BatchVelocityBonesArr[MLLevelSetIndex].SetNum(NumberOfParticles);
+				}
+
+				PhysicsParallelFor(MLLevelSetIndexes.Num(), [this, &Particles, Dt, &MLLevelSetIndexes, &CollisionParticlesRange,
+					&PAndInvM, &BatchVelocityBonesArr, &BatchNormalsArr, &BatchPhisArr, NumberOfParallelThreadsPerMLLevelset](int32 IndexML)
+					{
+						const int32 MLLevelSetIndex = IndexML / NumberOfParallelThreadsPerMLLevelset;
+						const int32 CollisionIndex = MLLevelSetIndexes[MLLevelSetIndex];
+
+						const int32 MLLevelSetThread = IndexML % NumberOfParallelThreadsPerMLLevelset;
+						const int32 MLLevelSetThreadBlock = Particles.GetRangeSize() / NumberOfParallelThreadsPerMLLevelset;
+						int32 OuterParticleBatchBegin;
+						int32 OuterParticleBatchEnd;
+
+						if (MLLevelSetThread < NumberOfParallelThreadsPerMLLevelset - 1)
+						{
+							OuterParticleBatchBegin = MLLevelSetThread * MLLevelSetThreadBlock;
+							OuterParticleBatchEnd = (MLLevelSetThread + 1) * MLLevelSetThreadBlock;
+						}
+						else
+						{
+							OuterParticleBatchBegin = (NumberOfParallelThreadsPerMLLevelset - 1) * MLLevelSetThreadBlock;;
+							OuterParticleBatchEnd = Particles.GetRangeSize() - 1;
+						}
+
+						FMLLevelSet* MLLevelSet = CollisionParticlesRange.GetGeometry(CollisionIndex)->GetObject<FMLLevelSet>();
+						FSolverRigidTransform3 Frame = FSolverRigidTransform3(CollisionParticlesRange.X(CollisionIndex), CollisionParticlesRange.R(CollisionIndex)).Inverse();
+
+						if constexpr (bWithFriction)
+						{
+							MLLevelSet->BatchPhiWithNormalWithGreatestInfluenceBone(Particles.GetPAndInvM(), Frame,
+								BatchPhisArr[MLLevelSetIndex], BatchNormalsArr[MLLevelSetIndex], BatchVelocityBonesArr[MLLevelSetIndex],
+								CollisionThickness, MLLevelSetThread, OuterParticleBatchBegin, OuterParticleBatchEnd);
+
+							check(BatchPhisArr[MLLevelSetIndex].Num() == Particles.GetRangeSize());
+							check(BatchNormalsArr[MLLevelSetIndex].Num() == Particles.GetRangeSize());
+							check(BatchVelocityBonesArr[MLLevelSetIndex].Num() == Particles.GetRangeSize());
+
+							const FSolverRotation3 CPR_R = CollisionParticlesRange.R(CollisionIndex);
+							ispc::ApplyPerParticleCollisionFastFrictionBatchedPhiAndNormal(
+								(ispc::FVector4f*)Particles.GetPAndInvM().GetData(),
+								(const ispc::FVector3f*)Particles.XArray().GetData(),
+								CollisionIndex,
+								BatchPhisArr[MLLevelSetIndex].GetData(),
+								(const ispc::FVector3f*)BatchNormalsArr[MLLevelSetIndex].GetData(),
+								BatchVelocityBonesArr[MLLevelSetIndex].GetData(),
+								(const ispc::FVector3f*)CollisionParticlesRange.GetV().GetData(),
+								(const ispc::FVector3f*)CollisionParticlesRange.XArray().GetData(),
+								(const ispc::FVector3f*)CollisionParticlesRange.GetW().GetData(),
+								reinterpret_cast<const ispc::FVector4f&>(CPR_R),
+								//(const ispc::FVector4f&)CollisionParticlesRange.R(CollisionIndex),
+								FrictionCoefficient,
+								CollisionThickness,
+								Dt,
+								CollisionParticlesRange.GetRangeSize(),
+								OuterParticleBatchBegin,
+								OuterParticleBatchEnd);
+						}
+						else
+						{
+							MLLevelSet->BatchPhiWithNormal(Particles.GetPAndInvM(), Frame,
+								BatchPhisArr[MLLevelSetIndex], BatchNormalsArr[MLLevelSetIndex],
+								CollisionThickness, MLLevelSetThread, OuterParticleBatchBegin, OuterParticleBatchEnd);
+
+							check(BatchPhisArr[MLLevelSetIndex].Num() == Particles.GetRangeSize());
+							check(BatchNormalsArr[MLLevelSetIndex].Num() == Particles.GetRangeSize());
+
+							const FSolverRotation3 CPR_R = CollisionParticlesRange.R(CollisionIndex);
+							ispc::ApplyPerParticleCollisionNoFrictionBatchedPhiAndNormal(
+								(ispc::FVector4f*)Particles.GetPAndInvM().GetData(),
+								BatchPhisArr[MLLevelSetIndex].GetData(),
+								(const ispc::FVector3f*)BatchNormalsArr[MLLevelSetIndex].GetData(),
+								reinterpret_cast<const ispc::FVector4f&>(CPR_R),
+								//(const ispc::FVector4f&)CollisionParticlesRange.R(CollisionIndex),
+								CollisionThickness,
+								Dt,
+								OuterParticleBatchBegin,
+								OuterParticleBatchEnd);
+						}
+					});
+			}		
+		}		
+	}
+#endif  // #if INTEL_ISPC
+}
+
+template<bool bWithFriction>
+void FPBDSoftBodyCollisionConstraintBase::ApplyInternalISPCPerTimestep(FSolverParticlesRange& Particles, const FSolverReal Dt, const TArray<FSolverCollisionParticlesRange>& CollisionParticles) const
+{
+#if INTEL_ISPC
+	check(bRealTypeCompatibleWithISPC);
+	const int32 NumBatches = FMath::CeilToInt((FSolverReal)(Particles.GetRangeSize()) / (FSolverReal)Chaos_SoftBodyCollision_ISPC_ParallelBatchSize);
+	{
+		PhysicsParallelFor(NumBatches, [this, &Particles, Dt, &CollisionParticles](int32 BatchNumber)
+		{
+			const int32 BatchBegin = (Chaos_SoftBodyCollision_ISPC_ParallelBatchSize * BatchNumber);
+			const int32 BatchEnd = FMath::Min(Particles.GetRangeSize(), BatchBegin + Chaos_SoftBodyCollision_ISPC_ParallelBatchSize);
+
+			for (const FSolverCollisionParticlesRange& CollisionParticlesRange : CollisionParticles)
+			{
+				if constexpr (bWithFriction)
+				{
+					ispc::ApplyPerParticleCollisionFastFrictionNoGroupCheck(
+						(ispc::FVector4f*)Particles.GetPAndInvM().GetData(),
+						(const ispc::FVector3f*)Particles.XArray().GetData(),
+						(const ispc::FVector3f*)CollisionParticlesRange.GetV().GetData(),
+						(const ispc::FVector3f*)CollisionParticlesRange.XArray().GetData(),
+						(const ispc::FVector3f*)CollisionParticlesRange.GetW().GetData(),
+						(const ispc::FVector4f*)CollisionParticlesRange.GetR().GetData(),
+						FrictionCoefficient,
+						CollisionThickness,
+						(const uint8*)&CollisionParticlesRange,
+						(const uint8*)CollisionParticlesRange.GetAllGeometry().GetData(),
+						sizeof(FImplicitObject),
+						FImplicitObject::GetOffsetOfType(),
+						FImplicitObject::GetOffsetOfMargin(),
+						Dt,
+						CollisionParticlesRange.GetRangeSize(),
+						BatchBegin,
+						BatchEnd);
+				}
+				else
+				{
+					ispc::ApplyPerParticleCollisionNoFrictionNoGroupCheck(
+						(ispc::FVector4f*)Particles.GetPAndInvM().GetData(),
+						(const ispc::FVector3f*)Particles.XArray().GetData(),
+						(const ispc::FVector3f*)CollisionParticlesRange.GetV().GetData(),
+						(const ispc::FVector3f*)CollisionParticlesRange.XArray().GetData(),
+						(const ispc::FVector3f*)CollisionParticlesRange.GetW().GetData(),
+						(const ispc::FVector4f*)CollisionParticlesRange.GetR().GetData(),
+						CollisionThickness,
+						(const uint8*)&CollisionParticlesRange,
+						(const uint8*)CollisionParticlesRange.GetAllGeometry().GetData(),
+						sizeof(FImplicitObject),
+						FImplicitObject::GetOffsetOfType(),
+						FImplicitObject::GetOffsetOfMargin(),
+						Dt,
+						CollisionParticlesRange.GetRangeSize(),
+						BatchBegin,
+						BatchEnd);
+				}
+			}
+		});
+	}
+
+	FPAndInvM* const PAndInvM = Particles.GetPAndInvM().GetData();
+
+	CollisionIndices.Empty();
+	CollisionTargets.Empty();
+	CollisionNormalsOutput.Empty();
+	TSet<int32> CollisionIndicesSet;
+
+	// Skinned levelset passes
+	for (const FSolverCollisionParticlesRange& CollisionParticlesRange : CollisionParticles)
+	{
+		for (int32 CollisionIndex = 0; CollisionIndex < CollisionParticlesRange.GetRangeSize(); ++CollisionIndex)
+		{
+			//This part is updated to visualize Normals. This part will be returned back.
+			if (const TWeightedLatticeImplicitObject<FLevelSet>* const LevelSet = CollisionParticlesRange.GetGeometry(CollisionIndex)->GetObject< TWeightedLatticeImplicitObject<FLevelSet> >())
+			{
+				TRACE_CPUPROFILER_EVENT_SCOPE(FPBDSoftBodyCollisionConstraint_ApplyInternalISPC_LevelSet);
+				const FSolverRigidTransform3 Frame = FSolverRigidTransform3(CollisionParticlesRange.X(CollisionIndex), CollisionParticlesRange.R(CollisionIndex)).Inverse();
+				if constexpr (bWithFriction)
+				{
+					LevelSet->BatchPhiWithNormalAndGreatestInfluenceBone(Particles.GetPAndInvM(), Frame, CollisionThickness, BatchPhis, BatchNormals, BatchVelocityBones);
+					check(BatchPhis.Num() == Particles.GetRangeSize());
+					check(BatchNormals.Num() == Particles.GetRangeSize());
+					check(BatchVelocityBones.Num() == Particles.GetRangeSize());
+
+					const FSolverRotation3 CPR_R = CollisionParticlesRange.R(CollisionIndex);
+					ispc::ApplyPerParticleCollisionFastFrictionBatchedPhiAndNormal(
+						(ispc::FVector4f*)Particles.GetPAndInvM().GetData(),
+						(const ispc::FVector3f*)Particles.XArray().GetData(),
+						CollisionIndex,
+						BatchPhis.GetData(),
+						(const ispc::FVector3f*)BatchNormals.GetData(),
+						BatchVelocityBones.GetData(),
+						(const ispc::FVector3f*)CollisionParticlesRange.GetV().GetData(),
+						(const ispc::FVector3f*)CollisionParticlesRange.XArray().GetData(),
+						(const ispc::FVector3f*)CollisionParticlesRange.GetW().GetData(),
+						reinterpret_cast<const ispc::FVector4f&>(CPR_R),
+						//(const ispc::FVector4f*)CollisionParticlesRange.R(CollisionIndex),
+						FrictionCoefficient,
+						CollisionThickness,
+						Dt,
+						CollisionParticlesRange.GetRangeSize(),
+						0,
+						Particles.GetRangeSize());
+				}
+				else
+				{
+					LevelSet->BatchPhiWithNormal(Particles.GetPAndInvM(), Frame, BatchPhis, BatchNormals);
+					check(BatchPhis.Num() == Particles.GetRangeSize());
+					check(BatchNormals.Num() == Particles.GetRangeSize());
+
+					const FSolverRotation3 CPR_R = CollisionParticlesRange.R(CollisionIndex);
+					ispc::ApplyPerParticleCollisionNoFrictionBatchedPhiAndNormal(
+						(ispc::FVector4f*)Particles.GetPAndInvM().GetData(),
+						BatchPhis.GetData(),
+						(const ispc::FVector3f*)BatchNormals.GetData(),
+						reinterpret_cast<const ispc::FVector4f&>(CPR_R),
+						//(const ispc::FVector4f&)CollisionParticlesRange.R(CollisionIndex),
+						CollisionThickness,
+						Dt,
+						0,
+						Particles.GetRangeSize());
+				}
+
+				//obtain data for planar constraints 
+				const FSolverReal* const CollisionPhis = BatchPhis.GetData();
+				const FSolverVec3* const CollisionNormals = BatchNormals.GetData();
+				for (int32 Index = 0; Index < BatchPhis.Num(); Index++)
+				{
+					if (PAndInvM[Index].InvM == (FSolverReal)0.)
+					{
+						continue;
+					}
+
+					const FSolverReal Penetration = CollisionThickness - CollisionPhis[Index];
+
+					if (Penetration > (FSolverReal)0.)
+					{
+						const FSolverVec3 NormalWorld = Frame.TransformVector(CollisionNormals[Index]);
+						if (!CollisionIndicesSet.Contains(Index))
+						{
+							CollisionIndicesSet.Add(Index);
+							CollisionIndices.Emplace(Index);
+							CollisionTargets.Emplace(PAndInvM[Index].P);
+							CollisionNormalsOutput.Emplace(NormalWorld);
+						}
+					}
+				}
+
+			}
+		}
+	}
+
+	// MLLevelset passes in ApplyInternalISPCPerTimestep
+	for (const FSolverCollisionParticlesRange& CollisionParticlesRange : CollisionParticles)
+	{
+		TRACE_CPUPROFILER_EVENT_SCOPE(PBDSoftBody_ApplyInternalISPCPerTimestep_FMLLevelset);		
+		const int32 NumberOfParticles = Particles.GetRangeSize();
+
+		if (MLLevelSetParallelize == 0)
+		{			
+			for (int32 CollisionIndex = 0; CollisionIndex < CollisionParticlesRange.GetRangeSize(); ++CollisionIndex)
+			{
+				const FSolverRigidTransform3 Frame = FSolverRigidTransform3(CollisionParticlesRange.X(CollisionIndex), CollisionParticlesRange.R(CollisionIndex)).Inverse();
+				
+				if (const FMLLevelSet* const MLLevelSet = CollisionParticlesRange.GetGeometry(CollisionIndex)->GetObject<FMLLevelSet>())
+				{
+					BatchPhis.Empty();
+					BatchNormals.Empty();
+					BatchVelocityBones.Empty();
+					if constexpr (bWithFriction)
+					{
+						BatchPhis.SetNum(Particles.GetRangeSize());
+						BatchNormals.SetNum(Particles.GetRangeSize());
+						BatchVelocityBones.SetNum(Particles.GetRangeSize());
+						
+						const int32 MLLevelSetThreadNum = 0;
+						const int32 BatchBegin = 0;
+						const int32 BatchEnd = Particles.GetRangeSize();
+
+						MLLevelSet->BatchPhiWithNormalWithGreatestInfluenceBone(Particles.GetPAndInvM(), Frame, 
+							BatchPhis, BatchNormals, BatchVelocityBones, 
+							CollisionThickness, MLLevelSetThreadNum, BatchBegin, BatchEnd);
+
+						check(BatchPhis.Num() == Particles.GetRangeSize());
+						check(BatchNormals.Num() == Particles.GetRangeSize());
+						check(BatchVelocityBones.Num() == Particles.GetRangeSize());
+
+						const FSolverRotation3 CPR_R = CollisionParticlesRange.R(CollisionIndex);
+						ispc::ApplyPerParticleCollisionFastFrictionBatchedPhiAndNormal(
+							(ispc::FVector4f*)Particles.GetPAndInvM().GetData(),
+							(const ispc::FVector3f*)Particles.XArray().GetData(),
+							CollisionIndex,
+							BatchPhis.GetData(),
+							(const ispc::FVector3f*)BatchNormals.GetData(),
+							BatchVelocityBones.GetData(),
+							(const ispc::FVector3f*)CollisionParticlesRange.GetV().GetData(),
+							(const ispc::FVector3f*)CollisionParticlesRange.XArray().GetData(),
+							(const ispc::FVector3f*)CollisionParticlesRange.GetW().GetData(),
+							reinterpret_cast<const ispc::FVector4f&>(CPR_R),
+							//(const ispc::FVector4f&)CollisionParticlesRange.R(CollisionIndex),
+							FrictionCoefficient,
+							CollisionThickness,
+							Dt,
+							CollisionParticlesRange.GetRangeSize(),
+							0,
+							Particles.GetRangeSize());
+					}
+					else
+					{
+						BatchPhis.SetNum(Particles.GetRangeSize());
+						BatchNormals.SetNum(Particles.GetRangeSize());
+						const int32 MLLevelSetThreadNum = 0;
+						const int32 BatchBegin = 0;
+						const int32 BatchEnd = Particles.GetRangeSize();
+
+						MLLevelSet->BatchPhiWithNormal(Particles.GetPAndInvM(), Frame, BatchPhis, BatchNormals, 
+							CollisionThickness, MLLevelSetThreadNum, BatchBegin, BatchEnd);
+
+						check(BatchPhis.Num() == Particles.GetRangeSize());
+						check(BatchNormals.Num() == Particles.GetRangeSize());
+
+						const FSolverRotation3 CPR_R = CollisionParticlesRange.R(CollisionIndex);
+						ispc::ApplyPerParticleCollisionNoFrictionBatchedPhiAndNormal(
+							(ispc::FVector4f*)Particles.GetPAndInvM().GetData(),
+							BatchPhis.GetData(),
+							(const ispc::FVector3f*)BatchNormals.GetData(),
+							reinterpret_cast<const ispc::FVector4f&>(CPR_R),
+							//(const ispc::FVector4f&)CollisionParticlesRange.R(CollisionIndex),
+							CollisionThickness,
+							Dt,
+							0,
+							Particles.GetRangeSize());
+					}
+
+
+					const FSolverReal* const CollisionPhis = BatchPhis.GetData();
+					const FSolverVec3* const CollisionNormals = BatchNormals.GetData();
+					
+					for (int32 Index = 0; Index < BatchPhis.Num(); Index++)
+					{
+						if (PAndInvM[Index].InvM == (FSolverReal)0.)
+						{
+							continue;
+						}
+
+						const FSolverReal Penetration = CollisionThickness - CollisionPhis[Index];
+
+						if (Penetration > (FSolverReal)0.)
+						{
+							if (!CollisionIndicesSet.Contains(Index))
+							{
+								const FSolverVec3 NormalWorld = Frame.TransformVector(CollisionNormals[Index]);
+
+								CollisionIndicesSet.Add(Index);
+								CollisionIndices.Emplace(Index);
+								CollisionTargets.Emplace(PAndInvM[Index].P);
+								CollisionNormalsOutput.Emplace(NormalWorld);
+							}
+
+						}
+					}
+					
+				}
+			}
+		}
+
+		if (MLLevelSetParallelize == 1)
+		{
+			//Find MLLevelSet CollisionIndices for Parallelization
+			TArray<int32> MLLevelSetIndexes;
+			for (int32 CollisionIndex = 0; CollisionIndex < CollisionParticlesRange.GetRangeSize(); ++CollisionIndex)
+			{
+				if (FMLLevelSet* MLLevelSet = CollisionParticlesRange.GetGeometry(CollisionIndex)->GetObject<FMLLevelSet>())
+				{
+					MLLevelSetIndexes.Add(CollisionIndex);
+				}
+			}			
+			
+			const int32 NumberOfMLLevelSets = MLLevelSetIndexes.Num();
+			if (NumberOfMLLevelSets > 0)
+			{
+				int32 NumberOfParallelThreadsPerMLLevelset = 1;
+
+				for (int32 MLCollisionIndex : MLLevelSetIndexes)
+				{
+					FMLLevelSet* MLLevelSet = CollisionParticlesRange.GetGeometry(MLCollisionIndex)->GetObject<FMLLevelSet>();
+					MLLevelSet->UpdateNeuralInferencesNumber(NumberOfParallelThreadsPerMLLevelset);
+				}
+
+				const int32 NumberOfParallelCalls = NumberOfParallelThreadsPerMLLevelset * NumberOfMLLevelSets;
+				TArray<TArray<FSolverReal>> BatchPhisArr;
+				TArray<TArray<FSolverVec3>> BatchNormalsArr;
+				TArray<TArray<int32>> BatchVelocityBonesArr;
+
+				BatchPhisArr.SetNum(MLLevelSetIndexes.Num());
+				BatchNormalsArr.SetNum(MLLevelSetIndexes.Num());
+				BatchVelocityBonesArr.SetNum(MLLevelSetIndexes.Num());
+
+				// Allocate TArray's for Parallel MLLevelSet Calls.
+				for (int32 MLLevelSetIndex = 0; MLLevelSetIndex < NumberOfMLLevelSets; MLLevelSetIndex++)
+				{
+					BatchPhisArr[MLLevelSetIndex].SetNum(NumberOfParticles);
+					BatchNormalsArr[MLLevelSetIndex].SetNum(NumberOfParticles);
+					BatchVelocityBonesArr[MLLevelSetIndex].SetNum(NumberOfParticles);
+				}
+
+				// Data for the Planar Constraints.
+				TArray<TArray<int32>> BatchCollisionIndices;
+				TArray<TArray<FSolverVec3>> BatchCollisionTargets;
+				TArray<TArray<FSolverVec3>> BatchCollisionNormalsOutput;
+
+				BatchCollisionIndices.SetNum(NumberOfParallelCalls);
+				BatchCollisionTargets.SetNum(NumberOfParallelCalls);
+				BatchCollisionNormalsOutput.SetNum(NumberOfParallelCalls);
+
+				PhysicsParallelFor(NumberOfParallelThreadsPerMLLevelset* MLLevelSetIndexes.Num(), [this, &Particles, Dt, &MLLevelSetIndexes, &CollisionParticlesRange,
+					&PAndInvM, &BatchVelocityBonesArr, &BatchNormalsArr, &BatchPhisArr, NumberOfParallelThreadsPerMLLevelset,
+					&BatchCollisionIndices, &BatchCollisionTargets, &BatchCollisionNormalsOutput](int32 IndexML)
+					{
+						const int32 MLLevelSetIndex = IndexML / NumberOfParallelThreadsPerMLLevelset;
+						const int32 CollisionIndex = MLLevelSetIndexes[MLLevelSetIndex];
+
+						const int32 MLLevelSetThread = IndexML % NumberOfParallelThreadsPerMLLevelset;
+						const int32 MLLevelSetThreadBlock = Particles.GetRangeSize() / NumberOfParallelThreadsPerMLLevelset;
+						int32 OuterParticleBatchBegin;
+						int32 OuterParticleBatchEnd;
+
+						if (MLLevelSetThread < NumberOfParallelThreadsPerMLLevelset - 1)
+						{
+							OuterParticleBatchBegin = MLLevelSetThread * MLLevelSetThreadBlock;
+							OuterParticleBatchEnd = (MLLevelSetThread + 1) * MLLevelSetThreadBlock;
+						}
+						else
+						{
+							OuterParticleBatchBegin = (NumberOfParallelThreadsPerMLLevelset - 1) * MLLevelSetThreadBlock;;
+							OuterParticleBatchEnd = Particles.GetRangeSize() - 1;
+						}
+
+						FMLLevelSet* MLLevelSet = CollisionParticlesRange.GetGeometry(CollisionIndex)->GetObject<FMLLevelSet>();
+						FSolverRigidTransform3 Frame = FSolverRigidTransform3(CollisionParticlesRange.X(CollisionIndex), CollisionParticlesRange.R(CollisionIndex)).Inverse();
+
+						if constexpr (bWithFriction)
+						{
+							MLLevelSet->BatchPhiWithNormalWithGreatestInfluenceBone(Particles.GetPAndInvM(), Frame,
+								BatchPhisArr[MLLevelSetIndex], BatchNormalsArr[MLLevelSetIndex], BatchVelocityBonesArr[MLLevelSetIndex],
+								CollisionThickness, MLLevelSetThread, OuterParticleBatchBegin, OuterParticleBatchEnd);
+
+							check(BatchPhisArr[MLLevelSetIndex].Num() == Particles.GetRangeSize());
+							check(BatchNormalsArr[MLLevelSetIndex].Num() == Particles.GetRangeSize());
+							check(BatchVelocityBonesArr[MLLevelSetIndex].Num() == Particles.GetRangeSize());
+
+							const FSolverRotation3 CPR_R = CollisionParticlesRange.R(CollisionIndex);
+							ispc::ApplyPerParticleCollisionFastFrictionBatchedPhiAndNormal(
+								(ispc::FVector4f*)Particles.GetPAndInvM().GetData(),
+								(const ispc::FVector3f*)Particles.XArray().GetData(),
+								CollisionIndex,
+								BatchPhisArr[MLLevelSetIndex].GetData(),
+								(const ispc::FVector3f*)BatchNormalsArr[MLLevelSetIndex].GetData(),
+								BatchVelocityBonesArr[MLLevelSetIndex].GetData(),
+								(const ispc::FVector3f*)CollisionParticlesRange.GetV().GetData(),
+								(const ispc::FVector3f*)CollisionParticlesRange.XArray().GetData(),
+								(const ispc::FVector3f*)CollisionParticlesRange.GetW().GetData(),
+								reinterpret_cast<const ispc::FVector4f&>(CPR_R),
+								//(const ispc::FVector4f&)CollisionParticlesRange.R(CollisionIndex),
+								FrictionCoefficient,
+								CollisionThickness,
+								Dt,
+								CollisionParticlesRange.GetRangeSize(),
+								OuterParticleBatchBegin,
+								OuterParticleBatchEnd);
+						}
+						else
+						{
+							MLLevelSet->BatchPhiWithNormal(Particles.GetPAndInvM(), Frame,
+								BatchPhisArr[MLLevelSetIndex], BatchNormalsArr[MLLevelSetIndex],
+								CollisionThickness, MLLevelSetThread, OuterParticleBatchBegin, OuterParticleBatchEnd);
+
+							check(BatchPhisArr[MLLevelSetIndex].Num() == Particles.GetRangeSize());
+							check(BatchNormalsArr[MLLevelSetIndex].Num() == Particles.GetRangeSize());
+
+							const FSolverRotation3 CPR_R = CollisionParticlesRange.R(CollisionIndex);
+							ispc::ApplyPerParticleCollisionNoFrictionBatchedPhiAndNormal(
+								(ispc::FVector4f*)Particles.GetPAndInvM().GetData(),
+								BatchPhisArr[MLLevelSetIndex].GetData(),
+								(const ispc::FVector3f*)BatchNormalsArr[MLLevelSetIndex].GetData(),
+								reinterpret_cast<const ispc::FVector4f&>(CPR_R),
+								//(const ispc::FVector4f&)CollisionParticlesRange.R(CollisionIndex),
+								CollisionThickness,
+								Dt,
+								OuterParticleBatchBegin,
+								OuterParticleBatchEnd);
+						}
+
+						//obtain data for planar constraints 
+						const FSolverReal* const CollisionPhis = BatchPhisArr[MLLevelSetIndex].GetData();
+						const FSolverVec3* const CollisionNormals = BatchNormalsArr[MLLevelSetIndex].GetData();
+						for (int32 Index = OuterParticleBatchBegin; Index < OuterParticleBatchEnd; Index++)
+						{
+							if (PAndInvM[Index].InvM == (FSolverReal)0.)
+							{
+								continue;
+							}
+
+							const FSolverReal Penetration = CollisionThickness - CollisionPhis[Index];
+
+							if (Penetration > (FSolverReal)0.)
+							{
+								const FSolverVec3 NormalWorld = Frame.TransformVector(CollisionNormals[Index]);
+								BatchCollisionIndices[IndexML].Emplace(Index);
+								BatchCollisionTargets[IndexML].Emplace(PAndInvM[Index].P);
+								BatchCollisionNormalsOutput[IndexML].Emplace(NormalWorld);
+							}
+						}
+					});
+
+				for (int32 MLLevelSetIndex = 0; MLLevelSetIndex < NumberOfMLLevelSets; MLLevelSetIndex++)
+				{
+					for (int32 MLLevelSetThread = 0; MLLevelSetThread < NumberOfParallelThreadsPerMLLevelset; MLLevelSetThread++)
+					{
+						int32 IndexML = NumberOfParallelThreadsPerMLLevelset * MLLevelSetIndex + MLLevelSetThread;
+
+						for (int32 IndexIndex = 0; IndexIndex < BatchCollisionIndices[IndexML].Num(); IndexIndex++)
+						{
+							int32 Index = BatchCollisionIndices[IndexML][IndexIndex];
+							if (!CollisionIndicesSet.Contains(Index))
+							{
+								CollisionIndicesSet.Add(Index);
+								CollisionIndices.Emplace(Index);
+								CollisionTargets.Emplace(BatchCollisionTargets[IndexML][IndexIndex]);
+								CollisionNormalsOutput.Emplace(BatchCollisionNormalsOutput[IndexML][IndexIndex]);
+							}
+						}
+					}
+				}				
+			}			
+		}
+	}
 #endif  // #if INTEL_ISPC
 }
 
+
 void FPBDSoftBodyCollisionConstraintBase::UpdateLinearSystem(const FSolverParticlesRange& Particles, const FSolverReal Dt, const TArray<FSolverCollisionParticlesRange>& CollisionParticles, FEvolutionLinearSystem& LinearSystem) const
 {
 	TRACE_CPUPROFILER_EVENT_SCOPE(FPBDSoftBodyCollisionConstraint_UpdateLinearSystem);
@@ -545,7 +1325,12 @@ void FPBDSoftBodyCollisionConstraintBase::UpdateLinearSystem(const FSolverPartic
 					continue;
 				}
 
-				const FSolverRigidTransform3 Frame(CollisionParticlesRange.GetX(CollisionIndex), CollisionParticlesRange.R(CollisionIndex));
+				if (CollisionParticlesRange.GetGeometry(CollisionIndex)->GetType() == Chaos::ImplicitObjectType::MLLevelSetBone)
+				{
+					continue;
+				}
+
+				const FSolverRigidTransform3 Frame(CollisionParticlesRange.X(CollisionIndex), CollisionParticlesRange.R(CollisionIndex));
 				const FVec3 RigidSpacePosition(Frame.InverseTransformPosition(PAndInvM[Index].P));  // PhiWithNormal requires FReal based arguments
 				FVec3 ImplicitNormal;                                                                // since implicits don't use FSolverReal
 				const FSolverReal Phi = (FSolverReal)CollisionParticlesRange.GetGeometry(CollisionIndex)->PhiWithNormal(RigidSpacePosition, ImplicitNormal);
diff --git a/Engine/Source/Runtime/Experimental/Chaos/Private/Chaos/PerParticlePBDCollisionConstraint.ispc b/Engine/Source/Runtime/Experimental/Chaos/Private/Chaos/PerParticlePBDCollisionConstraint.ispc
--- a/Engine/Source/Runtime/Experimental/Chaos/Private/Chaos/PerParticlePBDCollisionConstraint.ispc
+++ b/Engine/Source/Runtime/Experimental/Chaos/Private/Chaos/PerParticlePBDCollisionConstraint.ispc
@@ -27,7 +27,10 @@ enum ImplicitObjectType
 	UnionClustered,
 	TaperedCapsule,
 	WeightedLatticeBone,
-	WeightedLatticeLevelSetType = 38, // == LevelSet | IsWeightedLattice (cannot use expressions like this in ISPC)
+	MLLevelSet,
+	MLLevelSetBone,
+	WeightedLatticeLevelSetType = 38, //== LevelSet | IsWeightedLattice (cannot use expressions like this in ISPC)
+
 
 	//Add entries above this line for serialization
 	IsWeightedLattice = 1 << 5,
@@ -714,6 +717,21 @@ export void ApplyPerParticleCollisionFastFriction(uniform FVector4f ParticlesPan
 					GetPhiWithNormalAndVelocityBone(CollisionParticles, (const uniform float *uniform)&PhiWithNormalInputFloat, (uniform float *uniform)&Normal, (uniform float *uniform)&Phi, (uniform int *uniform)&VelocityBone, i, PerGroupThickness, programCount, IMMask);
 					break;
 				}
+			case MLLevelSetBone:
+				{
+					Phi = PerGroupThickness + 10;
+					break;
+				}
+			case MLLevelSet:
+				{
+					uniform int IMMask;
+					if(IM != 0)
+					{
+						IMMask = lanemask();
+					}
+					GetPhiWithNormalAndVelocityBone(CollisionParticles, (const uniform float *uniform)&PhiWithNormalInputFloat, (uniform float *uniform)&Normal, (uniform float *uniform)&Phi, (uniform int *uniform)&VelocityBone, i, PerGroupThickness, programCount, IMMask);
+					break;
+				}
 			default:
 				{
 					uniform int IMMask;
@@ -808,9 +826,20 @@ export void ApplyPerParticleCollisionFastFrictionNoGroupCheck(uniform FVector4f
 
 		for(uniform int i = 0; i < InnerCollisionCount; i++) // SequentialFor
 		{
+
+			const uniform SerializedPtr *uniform GeometryPtr = (const uniform SerializedPtr *uniform)Geometry;
+			const uniform uint8 *uniform Ptr = GeometryPtr[i].Ptr;
+			const uniform uint8 *uniform TypedPtr = Ptr + SizeofFImplicitObject;
+			const uniform uint8 Type = Ptr[OffsetofGeometryType] & 0x3F;
+
+			if(Type == WeightedLatticeBone || Type == WeightedLatticeLevelSetType || Type == MLLevelSet || Type == MLLevelSetBone)
+			{
+				continue;
+			}
+
 			// InverseTransformPosition. Can skip Scale calculation because it is 1 for xyz.
 			const FVector3f TranslatedVec = P - CollisionX[i];
-			const FVector4f InverseR = QuatInverse(CollisionR[i]);
+			const uniform FVector4f InverseR = QuatInverse(CollisionR[i]);
 			const FVector3f PhiWithNormalInputFloat = VectorQuaternionRotateVector(InverseR, TranslatedVec);
 			const FVector PhiWithNormalInput = ConvertVector3fTo3Native(PhiWithNormalInputFloat);
 			FVector NormalReal;
@@ -818,12 +847,6 @@ export void ApplyPerParticleCollisionFastFrictionNoGroupCheck(uniform FVector4f
 			FVector3f Normal;
 			float Phi;
 
-			const uniform SerializedPtr *uniform GeometryPtr = (const uniform SerializedPtr *uniform)Geometry;
-			const uniform uint8 *uniform Ptr = GeometryPtr[i].Ptr;
-			const uniform uint8 *uniform TypedPtr = Ptr + SizeofFImplicitObject;
-			const uniform uint8 Type = Ptr[OffsetofGeometryType] & 0x3F;
-
-			int32 VelocityBone = i;
 			switch(Type)
 			{
 			case Sphere:
@@ -862,18 +885,15 @@ export void ApplyPerParticleCollisionFastFrictionNoGroupCheck(uniform FVector4f
 					break;
 				}
 			case WeightedLatticeBone:
+			case WeightedLatticeLevelSetType: // Skip these here. They're done in a separate pass with batched calls to PhiAndNormal
 				{
 					Phi = PerGroupThickness + 10;
 					break;
 				}
-			case WeightedLatticeLevelSetType:
+			case MLLevelSet:
+			case MLLevelSetBone: // Skip these here. They're done in a separate pass with batched calls to PhiAndNormal
 				{
-					uniform int IMMask;
-					if(IM != 0)
-					{
-						IMMask = lanemask();
-					}
-					GetPhiWithNormalAndVelocityBoneCollisionParticleRange(CollisionParticlesRange, (const uniform float *uniform)&PhiWithNormalInputFloat, (uniform float *uniform)&Normal, (uniform float *uniform)&Phi, (uniform int *uniform)&VelocityBone, i, PerGroupThickness, programCount, IMMask);
+					Phi = PerGroupThickness + 10;
 					break;
 				}
 			default:
@@ -895,39 +915,16 @@ export void ApplyPerParticleCollisionFastFrictionNoGroupCheck(uniform FVector4f
 				// TransformVector. Scale is 1 so that part can be removed.
 				const FVector3f NormalWorld = VectorQuaternionRotateVector(CollisionR[i], Normal);
 				P = P + (Penetration * NormalWorld);
-				if(Type == WeightedLatticeLevelSetType)
-				{
-					#pragma ignore warning(perf)
-					const FVector3f CollisionXVal = CollisionX[VelocityBone];
-					const FVector3f VectorToPoint = P - CollisionXVal;
-					const FVector3f X = VectorLoad(&ParticlesX[extract(Index, 0)]);
-					
-					#pragma ignore warning(perf)
-					const FVector3f CollisionVVal = CollisionV[VelocityBone];
-					#pragma ignore warning(perf)
-					const FVector3f CollisionWVal = CollisionW[VelocityBone];
-					const FVector3f RelativeDisplacement = (P - X) - (CollisionVVal + VectorCross(CollisionWVal, VectorToPoint)) * Dt; // This corresponds to the tangential velocity multiplied by dt (friction will drive this to zero if it is high enough)
-					const FVector3f RelativeDisplacementTangent = RelativeDisplacement - VectorDot(RelativeDisplacement, NormalWorld) * NormalWorld; // Project displacement into the tangential plane
-					const float RelativeDisplacementTangentLength = VectorSize(RelativeDisplacementTangent);
+				const FVector3f VectorToPoint = P - CollisionX[i];
+				const FVector3f X = VectorLoad(&ParticlesX[extract(Index, 0)]);
 
-					const float PositionCorrection = min(Penetration * PerGroupFriction, RelativeDisplacementTangentLength);
-					const float CorrectionRatio = select(RelativeDisplacementTangentLength < FLOAT_SMALL_NUMBER, 0.0f, PositionCorrection / RelativeDisplacementTangentLength);
-					P = P - (CorrectionRatio * RelativeDisplacementTangent);
-				
-				}
-				else
-				{
-					const FVector3f VectorToPoint = P - CollisionX[i];
-					const FVector3f X = VectorLoad(&ParticlesX[extract(Index, 0)]);
+				const FVector3f RelativeDisplacement = (P - X) - (CollisionV[i] + VectorCross(CollisionW[i], VectorToPoint)) * Dt; // This corresponds to the tangential velocity multiplied by dt (friction will drive this to zero if it is high enough)
+				const FVector3f RelativeDisplacementTangent = RelativeDisplacement - VectorDot(RelativeDisplacement, NormalWorld) * NormalWorld; // Project displacement into the tangential plane
+				const float RelativeDisplacementTangentLength = VectorSize(RelativeDisplacementTangent);
 
-					const FVector3f RelativeDisplacement = (P - X) - (CollisionV[i] + VectorCross(CollisionW[i], VectorToPoint)) * Dt; // This corresponds to the tangential velocity multiplied by dt (friction will drive this to zero if it is high enough)
-					const FVector3f RelativeDisplacementTangent = RelativeDisplacement - VectorDot(RelativeDisplacement, NormalWorld) * NormalWorld; // Project displacement into the tangential plane
-					const float RelativeDisplacementTangentLength = VectorSize(RelativeDisplacementTangent);
-
-					const float PositionCorrection = min(Penetration * PerGroupFriction, RelativeDisplacementTangentLength);
-					const float CorrectionRatio = select(RelativeDisplacementTangentLength < FLOAT_SMALL_NUMBER, 0.0f, PositionCorrection / RelativeDisplacementTangentLength);
-					P = P - (CorrectionRatio * RelativeDisplacementTangent);
-				}
+				const float PositionCorrection = min(Penetration * PerGroupFriction, RelativeDisplacementTangentLength);
+				const float CorrectionRatio = select(RelativeDisplacementTangentLength < FLOAT_SMALL_NUMBER, 0.0f, PositionCorrection / RelativeDisplacementTangentLength);
+				P = P - (CorrectionRatio * RelativeDisplacementTangent);
 			}
 		}
 
@@ -938,6 +935,65 @@ export void ApplyPerParticleCollisionFastFrictionNoGroupCheck(uniform FVector4f
 	}
 }
 
+export void ApplyPerParticleCollisionFastFrictionBatchedPhiAndNormal(uniform FVector4f ParticlesPandInvM[],
+													const uniform FVector3f ParticlesX[],
+													const uniform int CollisionIndex,
+													const uniform float BatchPhis[],
+													const uniform FVector3f BatchNormals[],
+													const uniform int BatchVelocityBones[],
+													const uniform FVector3f CollisionV[],
+													const uniform FVector3f CollisionX[],
+													const uniform FVector3f CollisionW[],
+													const uniform FVector4f& CollisionR,
+													const uniform float PerGroupFriction,
+													const uniform float PerGroupThickness,
+													const uniform float Dt,
+													const uniform int InnerCollisionCount,
+													const uniform int OuterParticleBatchBegin,
+													const uniform int OuterParticleBatchEnd)
+{
+	foreach(Index = OuterParticleBatchBegin ... OuterParticleBatchEnd)
+	{
+		const FVector4f PandInvM = VectorLoad(&ParticlesPandInvM[extract(Index, 0)]);
+		varying FVector3f P;
+		varying float IM;
+		UnzipPandInvM(PandInvM, P, IM);
+
+		const float Phi = BatchPhis[Index];
+
+		const float Penetration = PerGroupThickness - Phi; // This is related to the Normal impulse
+		if(IM != 0 && Penetration > 0.0f)
+		{
+			const FVector3f Normal = VectorLoad(&BatchNormals[extract(Index,0)]);
+
+			// TransformVector. Scale is 1 so that part can be removed.
+			const FVector3f NormalWorld = VectorQuaternionRotateVector(CollisionR, Normal);
+			P = P + (Penetration * NormalWorld);
+
+			const int VelocityBone = BatchVelocityBones[Index] >= 0 && BatchVelocityBones[Index] < InnerCollisionCount ? BatchVelocityBones[Index] : CollisionIndex;
+
+			#pragma ignore warning(perf)
+			const FVector3f CollisionXVal = CollisionX[VelocityBone];
+			const FVector3f VectorToPoint = P - CollisionXVal;
+			const FVector3f X = VectorLoad(&ParticlesX[extract(Index, 0)]);
+					
+			#pragma ignore warning(perf)
+			const FVector3f CollisionVVal = CollisionV[VelocityBone];
+			#pragma ignore warning(perf)
+			const FVector3f CollisionWVal = CollisionW[VelocityBone];
+			const FVector3f RelativeDisplacement = (P - X) - (CollisionVVal + VectorCross(CollisionWVal, VectorToPoint)) * Dt; // This corresponds to the tangential velocity multiplied by dt (friction will drive this to zero if it is high enough)
+			const FVector3f RelativeDisplacementTangent = RelativeDisplacement - VectorDot(RelativeDisplacement, NormalWorld) * NormalWorld; // Project displacement into the tangential plane
+			const float RelativeDisplacementTangentLength = VectorSize(RelativeDisplacementTangent);
+
+			const float PositionCorrection = min(Penetration * PerGroupFriction, RelativeDisplacementTangentLength);
+			const float CorrectionRatio = select(RelativeDisplacementTangentLength < FLOAT_SMALL_NUMBER, 0.0f, PositionCorrection / RelativeDisplacementTangentLength);
+			P = P - (CorrectionRatio * RelativeDisplacementTangent);
+			
+			VectorStore(&ParticlesPandInvM[extract(Index, 0)], SetVector4(P, IM));
+		}
+	}
+}
+
 export void ApplyPerParticleCollisionNoFriction(uniform FVector4f ParticlesPandInvM[],
 												const uniform FVector3f ParticlesX[],
 												const uniform FVector3f CollisionV[],
@@ -1114,6 +1170,16 @@ export void ApplyPerParticleCollisionNoFrictionNoGroupCheck(uniform FVector4f Pa
 
 		for(uniform int i = 0; i < InnerCollisionCount; i++) // SequentialFor
 		{
+			const uniform SerializedPtr *uniform GeometryPtr = (const uniform SerializedPtr *uniform)Geometry;
+			const uniform uint8 *uniform Ptr = GeometryPtr[i].Ptr;
+			const uniform uint8 *uniform TypedPtr = Ptr + SizeofFImplicitObject;
+			const uniform uint8 Type = Ptr[OffsetofGeometryType] & 0x3F;
+
+			if(Type == WeightedLatticeBone || Type == WeightedLatticeLevelSetType || Type == MLLevelSet || Type == MLLevelSetBone)
+			{
+				continue;
+			}
+
 			// InverseTransformPosition. Can skip Scale calculation because it is 1 for xyz.
 			const FVector3f TranslatedVec = P - CollisionX[i];
 			const FVector4f InverseR = QuatInverse(CollisionR[i]);
@@ -1124,10 +1190,6 @@ export void ApplyPerParticleCollisionNoFrictionNoGroupCheck(uniform FVector4f Pa
 			FVector3f Normal;
 			float Phi;
 
-			const uniform SerializedPtr *uniform GeometryPtr = (const uniform SerializedPtr *uniform)Geometry;
-			const uniform uint8 *uniform Ptr = GeometryPtr[i].Ptr;
-			const uniform uint8 *uniform TypedPtr = Ptr + SizeofFImplicitObject;
-			const uniform uint8 Type = Ptr[OffsetofGeometryType] & 0x3F;
 
 			switch(Type)
 			{
@@ -1167,6 +1229,13 @@ export void ApplyPerParticleCollisionNoFrictionNoGroupCheck(uniform FVector4f Pa
 					break;
 				}
 			case WeightedLatticeBone:
+			case WeightedLatticeLevelSetType: // Skip these here. They're done in a separate pass with batched calls to PhiAndNormal
+				{
+					Phi = PerGroupThickness + 10;
+					break;
+				}
+			case MLLevelSet:
+			case MLLevelSetBone: // Skip these here. They're done in a separate pass with batched calls to PhiAndNormal
 				{
 					Phi = PerGroupThickness + 10;
 					break;
@@ -1198,4 +1267,36 @@ export void ApplyPerParticleCollisionNoFrictionNoGroupCheck(uniform FVector4f Pa
 			VectorStore(&ParticlesPandInvM[extract(Index, 0)], SetVector4(P, IM));
 		}
 	}
-}
\ No newline at end of file
+}
+
+export void ApplyPerParticleCollisionNoFrictionBatchedPhiAndNormal(uniform FVector4f ParticlesPandInvM[],
+												const uniform float BatchPhis[],
+												const uniform FVector3f BatchNormals[],
+												const uniform FVector4f& CollisionR,
+												const uniform float PerGroupThickness,
+												const uniform float Dt,
+												const uniform int OuterParticleBatchBegin,
+												const uniform int OuterParticleBatchEnd)
+{
+
+	foreach(Index = OuterParticleBatchBegin ... OuterParticleBatchEnd)
+	{
+		const FVector4f PandInvM = VectorLoad(&ParticlesPandInvM[extract(Index, 0)]);
+		varying FVector3f P;
+		varying float IM;
+		UnzipPandInvM(PandInvM, P, IM);
+
+		const float Phi = BatchPhis[Index];
+
+		const float Penetration = PerGroupThickness - Phi; // This is related to the Normal impulse
+		if(IM != 0 && Penetration > 0.0f)
+		{
+			const FVector3f Normal = VectorLoad(&BatchNormals[extract(Index,0)]);
+
+			// TransformVector. Scale is 1 so that part can be removed.
+			const FVector3f NormalWorld = VectorQuaternionRotateVector(CollisionR, Normal);
+			P = P + (Penetration * NormalWorld);
+			VectorStore(&ParticlesPandInvM[extract(Index, 0)], SetVector4(P, IM));
+		}
+	}
+}
diff --git a/Engine/Source/Runtime/Experimental/Chaos/Private/Chaos/SoftsEvolution.cpp b/Engine/Source/Runtime/Experimental/Chaos/Private/Chaos/SoftsEvolution.cpp
--- a/Engine/Source/Runtime/Experimental/Chaos/Private/Chaos/SoftsEvolution.cpp
+++ b/Engine/Source/Runtime/Experimental/Chaos/Private/Chaos/SoftsEvolution.cpp
@@ -6,6 +6,12 @@
 
 namespace Chaos::Softs {
 
+bool bPostTimestepCollisionUpdate = false;
+#if !UE_BUILD_SHIPPING
+static FAutoConsoleVariableRef CVarClothbPostTimestepCollisionUpdate(TEXT("p.Chaos.Cloth.PostTimestepCollisionUpdate"), bPostTimestepCollisionUpdate, TEXT("Do collision update at the end of the timestep instead of the beginning [def: false]"));
+
+#endif
+
 namespace Private
 {
 namespace EvolutionSolverDefault
@@ -508,6 +514,20 @@ void FEvolution::AdvanceOneTimeStepInternal(const FSolverReal Dt, const int32 Ti
 				}
 			}
 		}
+
+		{
+			if (bDoCollisionPerTimestep && !bPostTimestepCollisionUpdate)		
+			{	
+				// Presubstep Collision Rules
+				TRACE_CPUPROFILER_EVENT_SCOPE(FSoftsEvolution_AdvanceOneTimeStepInternal_PreSubstepCollisionRules);
+				TConstArrayView<PBDCollisionConstraintRuleFunc> PerIterationCollisionConstraintRules(SoftBodies.PerIterationCollisionPBDConstraintRules[SoftBodyId]);
+				for (const PBDCollisionConstraintRuleFunc& ConstraintRule : PerIterationCollisionConstraintRules)
+				{
+					ConstraintRule(SoftBodies.ParticleRanges[SoftBodyId], Dt, ActiveCollisionRanges);
+				}		
+			}
+		}
+
 		{
 			// Parallel Inits
 			TRACE_CPUPROFILER_EVENT_SCOPE(FSoftsEvolution_AdvanceOneTimeStepInternal_PostInitialGuessParallelInits);
@@ -540,10 +560,13 @@ void FEvolution::AdvanceOneTimeStepInternal(const FSolverReal Dt, const int32 Ti
 						ConstraintRule(SoftBodies.ParticleRanges[SoftBodyId], Dt);
 					}
 
-					TConstArrayView<PBDCollisionConstraintRuleFunc> PerIterationCollisionConstraintRules(SoftBodies.PerIterationCollisionPBDConstraintRules[SoftBodyId]);
-					for (const PBDCollisionConstraintRuleFunc& ConstraintRule : PerIterationCollisionConstraintRules)
+					if (!bDoCollisionPerTimestep)
 					{
-						ConstraintRule(SoftBodies.ParticleRanges[SoftBodyId], Dt, ActiveCollisionRanges);
+						TConstArrayView<PBDCollisionConstraintRuleFunc> PerIterationCollisionConstraintRules(SoftBodies.PerIterationCollisionPBDConstraintRules[SoftBodyId]);
+						for (const PBDCollisionConstraintRuleFunc& ConstraintRule : PerIterationCollisionConstraintRules)
+						{
+							ConstraintRule(SoftBodies.ParticleRanges[SoftBodyId], Dt, ActiveCollisionRanges);
+						}
 					}
 
 					TConstArrayView<PBDConstraintRuleFunc> PerIterationPostCollisionsConstraintRules(SoftBodies.PerIterationPostCollisionsPBDConstraintRules[SoftBodyId]);
@@ -570,6 +593,18 @@ void FEvolution::AdvanceOneTimeStepInternal(const FSolverReal Dt, const int32 Ti
 		{
 			Private::CopyVtoVPrev(SoftBodies.ParticleRanges[SoftBodyId]);
 		}
+
+		{
+			if (bDoCollisionPerTimestep && bPostTimestepCollisionUpdate)		
+			{	
+				TRACE_CPUPROFILER_EVENT_SCOPE(FSoftsEvolution_AdvanceOneTimeStepInternal_PostSubstepCollisionRules);
+				TConstArrayView<PBDCollisionConstraintRuleFunc> PerIterationCollisionConstraintRules(SoftBodies.PerIterationCollisionPBDConstraintRules[SoftBodyId]);
+				for (const PBDCollisionConstraintRuleFunc& ConstraintRule : PerIterationCollisionConstraintRules)
+				{
+					ConstraintRule(SoftBodies.ParticleRanges[SoftBodyId], Dt, ActiveCollisionRanges);
+				}		
+			}
+		}
 		
 		// Linear system update... currently we don't have any cross-softbody constraints. 
 		// If we did, the linear system would need to live at the group level.
diff --git a/Engine/Source/Runtime/Experimental/Chaos/Private/Chaos/WeightedLatticeImplicitObject.cpp b/Engine/Source/Runtime/Experimental/Chaos/Private/Chaos/WeightedLatticeImplicitObject.cpp
--- a/Engine/Source/Runtime/Experimental/Chaos/Private/Chaos/WeightedLatticeImplicitObject.cpp
+++ b/Engine/Source/Runtime/Experimental/Chaos/Private/Chaos/WeightedLatticeImplicitObject.cpp
@@ -2,10 +2,16 @@
 #include "Chaos/WeightedLatticeImplicitObject.h"
 #include "Chaos/Levelset.h"
 #include "Chaos/Tetrahedron.h"
+#include "Chaos/PBDSoftsSolverParticles.h"
 
 namespace Chaos
 {
 
+static int32 Chaos_WeightedLattice_PhiAndNormalBatchSize = 16;
+#if !UE_BUILD_SHIPPING
+FAutoConsoleVariableRef CVarChaosWeightedLatticePhiAndNormalBatchSize(TEXT("p.Chaos.WeightedLattice.PhiAndNormal.ParallelBatchSize"), Chaos_WeightedLattice_PhiAndNormalBatchSize, TEXT("Parallel batch size PhiAndNormal queries"));
+#endif
+
 FWeightedLatticeImplicitObject::FEmbeddingCoordinate::FEmbeddingCoordinate(const TVec3<int32>& InCellIndex, const FVec3& TrilinearCoordinate)
 	:CellIndex(InCellIndex)
 {
@@ -648,7 +654,7 @@ template<typename TConcrete>
 FReal TWeightedLatticeImplicitObject<TConcrete>::PhiWithNormalAndSurfacePoint(const FVec3& X, FVec3& Normal, FEmbeddingCoordinate& SurfaceCoord, bool bIncludeEmptyCells) const
 {
 	TArray<FEmbeddingCoordinate> Coordinates;
-	if (GetEmbeddingCoordinates(X, Coordinates))
+	if (GetEmbeddingCoordinates(X, Coordinates, bIncludeEmptyCells))
 	{
 		// Choose "deepest" point here.
 		FReal MinPhiUndeformed = std::numeric_limits<FReal>::max();
@@ -699,6 +705,80 @@ FReal TWeightedLatticeImplicitObject<TConcrete>::PhiWithNormalAndSurfacePoint(co
 	}
 }
 
+template<typename TConcrete>
+void TWeightedLatticeImplicitObject<TConcrete>::BatchPhiWithNormalAndGreatestInfluenceBone(const TConstArrayView<Softs::FPAndInvM> PAndInvMArray, 
+	const Softs::FSolverRigidTransform3& SolverToThis, const Softs::FSolverReal Thickness,
+	TArray<Softs::FSolverReal>& OutBatchPhis, TArray<Softs::FSolverVec3>& OutBatchNormals, TArray<int32>& OutBatchGreatestInfluenceSolverBones) const
+{
+	TRACE_CPUPROFILER_EVENT_SCOPE(TWeightedLatticeImplicitObject_BatchPhiWithNormalAndGreatestInfluenceBone);
+	const int32 NumParticles = PAndInvMArray.Num();
+	OutBatchPhis.SetNumUninitialized(NumParticles);
+	OutBatchNormals.SetNumUninitialized(NumParticles);
+	OutBatchGreatestInfluenceSolverBones.SetNumUninitialized(NumParticles);
+	
+	const Softs::FPAndInvM* const PAndInvM = PAndInvMArray.GetData();
+	Softs::FSolverReal* const Phis = OutBatchPhis.GetData();
+	Softs::FSolverVec3* const Normals = OutBatchNormals.GetData();
+	int32* const GreatestInfluenceSolverBones = OutBatchGreatestInfluenceSolverBones.GetData();
+
+	const int32 NumBatches = FMath::CeilToInt((FRealSingle)(NumParticles) / (FRealSingle)Chaos_WeightedLattice_PhiAndNormalBatchSize);
+	PhysicsParallelFor(NumBatches, [this, NumParticles, &PAndInvM, &SolverToThis, Thickness, &Phis, &Normals, &GreatestInfluenceSolverBones](int32 BatchNumber)
+	{
+		const int32 BatchBegin = (Chaos_WeightedLattice_PhiAndNormalBatchSize * BatchNumber);
+		const int32 BatchEnd = FMath::Min(NumParticles, BatchBegin + Chaos_WeightedLattice_PhiAndNormalBatchSize);
+		for (int32 Index = BatchBegin; Index < BatchEnd; ++Index)
+		{
+			if (PAndInvM[Index].InvM == (Softs::FSolverReal)0.)
+			{
+				continue;
+			}
+
+			const FVec3 RigidSpacePosition(SolverToThis.TransformPositionNoScale(PAndInvM[Index].P));
+			FWeightedLatticeImplicitObject::FEmbeddingCoordinate SurfaceCoord;
+			FVec3 Normal;
+			constexpr bool bIncludeEmptyCellsFalse = false;
+			Phis[Index] = (Softs::FSolverReal)PhiWithNormalAndSurfacePoint(RigidSpacePosition, Normal, SurfaceCoord, bIncludeEmptyCellsFalse);
+			Normals[Index] = Softs::FSolverVec3(Normal);
+			const Softs::FSolverReal Penetration = Thickness - Phis[Index]; // This is related to the Normal impulse
+			const int32 StrongestBone = Penetration > (Softs::FSolverReal)0. ? SurfaceCoord.GreatestInfluenceBone(GetBoneData()) : INDEX_NONE;
+			GreatestInfluenceSolverBones[Index] = SolverBoneIndices.IsValidIndex(StrongestBone) ? SolverBoneIndices[StrongestBone] : INDEX_NONE;
+		}
+	});
+}
+
+template<typename TConcrete>
+void TWeightedLatticeImplicitObject<TConcrete>::BatchPhiWithNormal(const TConstArrayView<Softs::FPAndInvM> PAndInvMArray,
+	const Softs::FSolverRigidTransform3& SolverToThis, TArray<Softs::FSolverReal>& OutBatchPhis, TArray<Softs::FSolverVec3>& OutBatchNormals) const
+{
+	TRACE_CPUPROFILER_EVENT_SCOPE(TWeightedLatticeImplicitObject_BatchPhiWithNormal);
+	const int32 NumParticles = PAndInvMArray.Num();
+	OutBatchPhis.SetNumUninitialized(NumParticles);
+	OutBatchNormals.SetNumUninitialized(NumParticles);
+	const Softs::FPAndInvM* const PAndInvM = PAndInvMArray.GetData();
+	Softs::FSolverReal* const Phis = OutBatchPhis.GetData();
+	Softs::FSolverVec3* const Normals = OutBatchNormals.GetData();
+
+	const int32 NumBatches = FMath::CeilToInt((FRealSingle)(NumParticles) / (FRealSingle)Chaos_WeightedLattice_PhiAndNormalBatchSize);
+	PhysicsParallelFor(NumBatches, [this, NumParticles, &PAndInvM, &SolverToThis, &Phis, &Normals](int32 BatchNumber)
+	{
+		const int32 BatchBegin = (Chaos_WeightedLattice_PhiAndNormalBatchSize * BatchNumber);
+		const int32 BatchEnd = FMath::Min(NumParticles, BatchBegin + Chaos_WeightedLattice_PhiAndNormalBatchSize);
+		for (int32 Index = BatchBegin; Index < BatchEnd; ++Index)
+		{
+			if (PAndInvM[Index].InvM == (Softs::FSolverReal)0.)
+			{
+				continue;
+			}
+			const FVec3 RigidSpacePosition(SolverToThis.TransformPositionNoScale(PAndInvM[Index].P));
+			FWeightedLatticeImplicitObject::FEmbeddingCoordinate SurfaceCoord;
+			FVec3 Normal;
+			constexpr bool bIncludeEmptyCellsFalse = false;
+			Phis[Index] = (Softs::FSolverReal)PhiWithNormalAndSurfacePoint(RigidSpacePosition, Normal, SurfaceCoord, bIncludeEmptyCellsFalse);
+			Normals[Index] = Softs::FSolverVec3(Normal);
+		}
+	});
+}
+
 void FWeightedLatticeImplicitObjectBuilder::GenerateGrid(const int32 GridResolution, const TAABB<FReal, 3>& ObjectBBox)
 {
 	constexpr int32 NumGhostCells = 3;
@@ -797,4 +877,4 @@ void FWeightedLatticeImplicitObjectBuilder::NormalizeBoneWeights()
 
 // Explicit instantiations
 template class CHAOS_API TWeightedLatticeImplicitObject<FLevelSet>;
-} // namespace Chaos
\ No newline at end of file
+} // namespace Chaos
diff --git a/Engine/Source/Runtime/Experimental/Chaos/Public/Chaos/ImplicitFwd.h b/Engine/Source/Runtime/Experimental/Chaos/Public/Chaos/ImplicitFwd.h
--- a/Engine/Source/Runtime/Experimental/Chaos/Public/Chaos/ImplicitFwd.h
+++ b/Engine/Source/Runtime/Experimental/Chaos/Public/Chaos/ImplicitFwd.h
@@ -14,6 +14,7 @@ namespace Chaos
 	class FHeightField;
 	class FImplicitObject;
 	class FLevelSet;
+	class FMLLevelSet;
 	class FTriangleMeshImplicitObject;
 	class FImplicitObjectUnionClustered;
 	class FImplicitObjectUnion;
@@ -33,6 +34,7 @@ namespace Chaos
 	using FConvexPtr = TRefCountPtr<FConvex>;
 	using FTriangleMeshImplicitObjectPtr = TRefCountPtr<FTriangleMeshImplicitObject>;
 	using FLevelSetPtr = TRefCountPtr<FLevelSet>;
+	using FMLLevelSetPtr = TRefCountPtr<FMLLevelSet>; 
 	using FHeightFieldPtr = TRefCountPtr<FHeightField>;
 	using FCapsulePtr = TRefCountPtr<FCapsule>;
 	using FPlanePtr = TRefCountPtr<FImplicitPlane3>;
@@ -46,6 +48,7 @@ namespace Chaos
 	using FConvexRef = FConvex*;
 	using FTriangleMeshImplicitObjectRef = FTriangleMeshImplicitObject*;
 	using FLevelSetRef = FLevelSet*;
+	using FMLLevelSetRef = FMLLevelSet*;
 	using FHeightFieldRef = FHeightField*;
 	using FCapsuleRef = FCapsule*;
 	using FPlaneRef = FImplicitPlane3*;
diff --git a/Engine/Source/Runtime/Experimental/Chaos/Public/Chaos/ImplicitObjectType.h b/Engine/Source/Runtime/Experimental/Chaos/Public/Chaos/ImplicitObjectType.h
--- a/Engine/Source/Runtime/Experimental/Chaos/Public/Chaos/ImplicitObjectType.h
+++ b/Engine/Source/Runtime/Experimental/Chaos/Public/Chaos/ImplicitObjectType.h
@@ -28,7 +28,8 @@ namespace Chaos
 			UnionClustered,
 			TaperedCapsule,
 			WeightedLatticeBone,
-
+			MLLevelSet,
+			MLLevelSetBone,
 			//Add entries above this line for serialization
 			ConcreteObjectCount, // Used to ensure bitflags do not overlap concrete type
 			IsWeightedLattice = 1 << 5,
@@ -95,6 +96,8 @@ namespace Chaos
 				TEXT("EImplicitObjectType:UnionClustered"),
 				TEXT("EImplicitObjectType:TaperedCapsule"),
 				TEXT("EImplicitObjectType:WeightedLatticeBone"),
+				TEXT("EImplicitObjectType:MLLevelSet"),
+				TEXT("EImplicitObjectType:MLLevelSetBone"),
 			};
 
 		static_assert(Chaos::ImplicitObjectType::ConcreteObjectCount == UE_ARRAY_COUNT(Names), "");
diff --git a/Engine/Source/Runtime/Experimental/Chaos/Public/Chaos/MLLevelSetNeuralInference.h b/Engine/Source/Runtime/Experimental/Chaos/Public/Chaos/MLLevelSetNeuralInference.h
new file mode 100644
--- /dev/null
+++ b/Engine/Source/Runtime/Experimental/Chaos/Public/Chaos/MLLevelSetNeuralInference.h
@@ -0,0 +1,37 @@
+// Copyright Epic Games, Inc. All Rights Reserved.
+#pragma once
+
+#include "NNEModelData.h"  
+#include "NNERuntimeCPU.h" 
+
+class FMLLevelSet;
+
+namespace Chaos
+{ 
+
+// Class for MLLevelSet Weight Depended Models
+class FMLLevelSetNeuralInference
+{
+  public:
+	  CHAOS_API FMLLevelSetNeuralInference() {}
+	  CHAOS_API FMLLevelSetNeuralInference(TSharedPtr<UE::NNE::IModelCPU> NNEModel, TArray<TArray<int32>>& InModelWeightsShapes);
+	  CHAOS_API FMLLevelSetNeuralInference(FMLLevelSetNeuralInference&& Other);
+	  CHAOS_API FMLLevelSetNeuralInference(const FMLLevelSetNeuralInference& Other);
+	  FMLLevelSetNeuralInference& operator=(const FMLLevelSetNeuralInference& Other);
+	  FMLLevelSetNeuralInference* Copy() const;
+	  void SetModelWeightsShapes(TArray<TArray<uint32>>& InModelWeightsShapes) { ModelWeightsShapes = InModelWeightsShapes; }
+	  const bool IsValid() const { return ModelInstance.IsValid(); } 
+	  void RunInference(TArray<float, TAlignedHeapAllocator<64>>& InputData, TArray<float, TAlignedHeapAllocator<64>>& OutputData, uint32 SingleInputSize, uint32 SingleOutputSize, TArray<TArray<float, TAlignedHeapAllocator<64>>>& ModelWeightsIn) const; // Batched Version of the RunMLModel. 
+	  bool Serialize(FArchive& Ar);
+ 
+  private:
+	// Different instances of FMLLevelSetNeuralInference can share the same NNEModel (even with deep copy). 
+	// Since each model has its own ModelInstance, this would not cause a problem during inference.
+	TSharedPtr<UE::NNE::IModelCPU> NNEModel; 
+	TSharedPtr<UE::NNE::IModelInstanceCPU> ModelInstance; //can this be TUniquePtr?
+	TArray<TArray<uint32>> ModelWeightsShapes; 
+	CHAOS_API FMLLevelSetNeuralInference(TSharedPtr<UE::NNE::IModelCPU> NNEModel); //needed for serialization
+	
+	friend FMLLevelSet;
+};
+}
diff --git a/Engine/Source/Runtime/Experimental/Chaos/Public/Chaos/MLLevelset.h b/Engine/Source/Runtime/Experimental/Chaos/Public/Chaos/MLLevelset.h
new file mode 100644
--- /dev/null
+++ b/Engine/Source/Runtime/Experimental/Chaos/Public/Chaos/MLLevelset.h
@@ -0,0 +1,243 @@
+// Copyright Epic Games, Inc. All Rights Reserved.
+#pragma once
+
+#include "Chaos/ArrayND.h"
+#include "Chaos/Box.h"
+#include "Chaos/ImplicitObject.h"
+#include "Chaos/ImplicitObjectScaled.h"
+#include "Chaos/MLLevelSetNeuralInference.h"
+#include "Chaos/PBDSoftsEvolutionFwd.h" 
+
+struct FKMLLevelSetElem;
+
+namespace Chaos
+{
+// Not actually a collision object. CollisionParticles that are used to update bones used by the FMLLevelSet have this as their geometry
+// It is used to get relative transform between consecutive bones (active joint vs parent of active joint)
+// Note that not necessarily parent  = parent of active joint
+class FMLLevelSetBoneProxy : public FImplicitObject
+{
+public:
+	using FImplicitObject::GetTypeName;
+
+	FMLLevelSetBoneProxy()
+		: FImplicitObject(EImplicitObject::DisableCollisions, ImplicitObjectType::MLLevelSetBone)
+	{}
+
+	static constexpr EImplicitObjectType StaticType() { return ImplicitObjectType::MLLevelSetBone; }
+
+	virtual FReal PhiWithNormal(const FVec3& x, FVec3& Normal) const override
+	{
+		return UE_BIG_NUMBER;
+	}
+
+	virtual uint32 GetTypeHash() const override
+	{
+		return 0;
+	}
+};
+
+struct FMLLevelSetNNEModelData
+{
+	TArray<int32> ModelArchitectureActivationNodeSizes;
+	const FString MLModelWeightsString;
+	const FString NNEModelPath; 
+	TObjectPtr<UNNEModelData> NNEModelData;
+};
+
+struct FMLLevelSetImportData
+{
+	const FName ParentBoneName;
+	TArray<FName> ActiveBoneNames;
+	TArray<FMLLevelSetNNEModelData> NNEModelDataArr; //Either {SignedDistanceModel} or {SignedDistanceModel, IncorrectZoneModel} 
+	TArray<FVector3f> ActiveBonesReferenceRotations;
+	TArray<FVector3f> ActiveBonesReferenceTranslations;
+	const float SignedDistanceScaling;
+	TArray<TArray<int32>> ActiveBonesRotationComponents;
+	const FVector3f TrainingGridMin;
+	TArray<FVector3f> TrainingGridAxesXYZ;
+	FIntVector DebugGridResolution;
+};
+
+class FMLLevelSet final : public FImplicitObject
+{
+  public:
+	CHAOS_API FMLLevelSet(FMLLevelSetImportData&& MLLevelSetImportData);
+
+	FMLLevelSet(const FMLLevelSet& Other) = delete; 
+	CHAOS_API FMLLevelSet(FMLLevelSet&& Other);
+	CHAOS_API virtual ~FMLLevelSet();
+
+	CHAOS_API virtual Chaos::FImplicitObjectPtr CopyGeometry() const override;
+	CHAOS_API virtual Chaos::FImplicitObjectPtr CopyGeometryWithScale(const FVec3& Scale) const override;
+
+	CHAOS_API virtual FReal PhiWithNormal(const FVec3& x, FVec3& Normal) const override;
+	CHAOS_API void BatchPhiWithNormal(const TConstArrayView<Softs::FPAndInvM> PAndInvMArray,
+		const Softs::FSolverRigidTransform3& SolverToThis, TArray<Softs::FSolverReal>& OutBatchPhis, TArray<Softs::FSolverVec3>& OutBatchNormals,
+		const Chaos::Softs::FSolverReal CollisionThickness, const int32 MLLevelsetThread, const int32 BatchBegin, const int32 BatchEnd) const;	
+	CHAOS_API void BatchPhiWithNormalWithGreatestInfluenceBone(const TConstArrayView<Softs::FPAndInvM> PAndInvMArray,
+		const Softs::FSolverRigidTransform3& SolverToThis, TArray<Softs::FSolverReal>& OutBatchPhis, TArray<Softs::FSolverVec3>& OutBatchNormals,
+		TArray<int32>& OutBatchGreatestInfluenceSolverBones, const Chaos::Softs::FSolverReal CollisionThickness, int32 MLLevelsetThread, int32 BatchBegin, int32 BatchEnd) const;
+
+	CHAOS_API FReal SignedDistance(const FVec3& x) const;
+	CHAOS_API void UpdateActiveBonesRelativeTransforms(TArray<FTransform>& InActiveBonesRelativeTransforms);
+	CHAOS_API void UpdateActiveBonesRelativeTransformsAndUpdateDebugPhi(TArray<FTransform>& InActiveBonesRelativeTransforms);
+	CHAOS_API void UpdateNeuralInferencesNumber(const int32 InNeuralInferencesNumber);
+
+	const TArray<int32>& GetSolverActiveBoneIndices() const { return SolverActiveBoneIndices; }
+	void SetSolverActiveBoneIndices(TArray<int32>&& InSolverActiveBoneIndices) { SolverActiveBoneIndices = MoveTemp(InSolverActiveBoneIndices); }
+	const int32 GetNumberOfActiveBones() const { return ActiveBoneNames.Num(); }
+	void SetSolverParentBoneIndex(int32 InSolverParentBoneIndex) { SolverParentBoneIndex = InSolverParentBoneIndex; }
+
+	const FName GetParentBoneName() const { return ParentBoneName; }
+	const TArray<FName>& GetActiveBoneNames() const { return ActiveBoneNames; }
+	virtual const FAABB3 BoundingBox() const override { return LocalBoundingBox; } 
+	const FVector3f GetTrainingGridMin() const { return TrainingGridMin; }
+	const FVector3f GetTrainingGridVector(int32 Index) const { return TrainingGridUnitAxesXYZ[Index] * TrainingGridAxesLengthsXYZ[Index]; }
+	
+	// Returns a const ref to the underlying grid structure
+	const TUniformGrid<FReal, 3>& GetGrid() const { return DebugGrid; }
+
+	// Returns a const ref to the underlying phi grid
+	const TArrayND<FReal, 3>& GetPhiArray() const { return DebugPhi; }
+
+	FORCEINLINE static constexpr EImplicitObjectType StaticType()
+	{
+		return ImplicitObjectType::MLLevelSet;
+	}
+
+	FORCEINLINE void SerializeImp(FArchive& Ar)
+	{
+		FImplicitObject::SerializeImp(Ar);
+		TBox<FReal, 3>::SerializeAsAABB(Ar, LocalBoundingBox);
+		Ar << ParentBoneName;
+		Ar << ActiveBoneNames;
+		Ar << ActiveBonesRelativeTransforms;
+		Ar << SignedDistanceModelWeights;
+		Ar << SignedDistanceModelWeightsShapes;
+		Ar << IncorrectZoneModelWeights;
+		Ar << IncorrectZoneModelWeightsShapes;
+		Ar << bUseIncorrectZoneModel;
+		Ar << SignedDistanceScaling;
+		Ar << TrainingGridMin;
+		Ar << TrainingGridUnitAxesXYZ;
+		Ar << TrainingGridAxesLengthsXYZ;
+		Ar << TotalNumberOfRotationComponents;
+		Ar << ActiveBonesRotationComponents;
+		Ar << ActiveBonesReferenceRotations;
+		Ar << ActiveBonesReferenceTranslations;
+		Ar << DebugGrid; 
+		Ar << DebugPhi;
+		
+		if (Ar.IsLoading())
+		{
+			SignedDistanceNeuralInferences.Empty();
+			SignedDistanceNeuralInferences.Add(FMLLevelSetNeuralInference(NNESignedDistanceModel, SignedDistanceModelWeightsShapes));
+			if (bUseIncorrectZoneModel)
+			{
+				IncorrectZoneNeuralInferences.Empty();
+				IncorrectZoneNeuralInferences.Add(FMLLevelSetNeuralInference(NNEIncorrectZoneModel, IncorrectZoneModelWeightsShapes));
+			}
+		}
+	}
+
+	virtual void Serialize(FChaosArchive& Ar) override
+	{
+		SerializeImp(Ar);
+	}
+
+	virtual void Serialize(FArchive& Ar) override
+	{
+		SerializeImp(Ar);
+	}	
+
+	// Used to generate a simple debug surface
+	CHAOS_API void GetZeroIsosurfaceGridCellFaces(TArray<FVector3f>& Vertices, TArray<FIntVector>& Tris) const;
+	CHAOS_API void GetInteriorCells(TArray<TVec3<int32>>& InteriorCells, const FReal InteriorThreshold) const;
+	CHAOS_API void CreatePhiFromMLModel();
+
+	virtual uint32 GetTypeHash() const override
+	{
+		uint32 Result = 0;
+
+		for(int32 LayerIdx = 0; LayerIdx < SignedDistanceModelWeights.Num(); ++LayerIdx)
+		{
+			for (int32 i = 0; i < SignedDistanceModelWeights[LayerIdx].Num(); ++i)
+			{
+				Result = HashCombine(Result, ::GetTypeHash(SignedDistanceModelWeights[LayerIdx][i]));
+			}			
+		}
+
+		for (int32 LayerIdx = 0; LayerIdx < IncorrectZoneModelWeights.Num(); ++LayerIdx)
+		{
+			for (int32 i = 0; i < IncorrectZoneModelWeights[LayerIdx].Num(); ++i)
+			{
+				Result = HashCombine(Result, ::GetTypeHash(IncorrectZoneModelWeights[LayerIdx][i]));
+			}
+		}
+
+		return Result;
+	}
+
+  private:
+	void ProcessTrainingGridAxesVectors();
+	void ComputeSignedDistanceNetworkWeightsInput(TArray<TArray<float, TAlignedHeapAllocator<64>>>& NetworkWeightsInput) const;
+	void ComputeIncorrectZoneNetworkWeightsInput(TArray<TArray<float, TAlignedHeapAllocator<64>>>& NetworkWeightsInput) const;
+	void ComputeWeightsInput(const TArray<FTransform>& RelativeBoneTransformationsInput, const TArray<TArray<float>>& NetworkWeights, const TArray<TArray<int32>>& NetworkWeightsShapes, TArray<TArray<float, TAlignedHeapAllocator<64>>>& NetworkWeightsInput) const;
+	void LoadMLModelWeightsFromString(const FString& MLModelWeightsString, TArray<TArray<float>>& ModelWeightArray);
+	void BuildNNEModel(const TArray<int32>& ModelArchitectureActivationNodeSizes, TObjectPtr<UNNEModelData> NNEModelData, const FString& ModelWeightsString,
+		TSharedPtr<UE::NNE::IModelCPU>& NNEModel, TArray<FMLLevelSetNeuralInference>& NeuralInferences, TArray<TArray<int32>>& ModelWeightsShapes, TArray<TArray<float>>& ModelWeightArray);
+	void BuildNNEModel(TObjectPtr<UNNEModelData> InNNEModelData, TSharedPtr<UE::NNE::IModelCPU>& NNEModel);
+	const bool IsFTransformArraysDifferent(TArray<FTransform>& FTransformArr1, TArray<FTransform>& FTransformArr2, FReal Tol = UE_KINDA_SMALL_NUMBER) const;
+
+	// We assume the Greatest Influence Bone is always the parent bone. 
+	const int32 GetGreatestInfluenceBoneSolverIndex() const { return SolverParentBoneIndex; }
+	// We assume the closest active bone joint is always the first active bone. 
+	const int32 GetClosestActiveBoneIndex(const FVec3 ParticlePositionMS) const { return 0; }  
+
+  private:
+	//Relative Bone Transform Related Data
+	FName ParentBoneName;
+	TArray<FName> ActiveBoneNames; 
+	TArray<int32> SolverActiveBoneIndices; //indices set by the solver for bone transform query
+	int32 SolverParentBoneIndex;
+	TArray<FTransform> ActiveBonesRelativeTransforms;
+
+	// NNE ML Model Elements
+	// ToDo: Move these elements inside a new struct or FMLLevelSetNeuralInference
+	TSharedPtr<UE::NNE::IModelCPU> NNESignedDistanceModel; 
+	TArray<FMLLevelSetNeuralInference> SignedDistanceNeuralInferences;
+	TArray<TArray<float>> SignedDistanceModelWeights; 
+	TArray<TArray<int32>> SignedDistanceModelWeightsShapes;
+
+	// NNE Incorrect Zone Model Elements
+	// ToDo: Move these elements inside a new struct or FMLLevelSetNeuralInference
+	TSharedPtr<UE::NNE::IModelCPU> NNEIncorrectZoneModel;
+	TArray<FMLLevelSetNeuralInference> IncorrectZoneNeuralInferences;
+	TArray<TArray<float>> IncorrectZoneModelWeights;
+	TArray<TArray<int32>> IncorrectZoneModelWeightsShapes;
+	bool bUseIncorrectZoneModel;
+	
+	// Inference Input Transformation. 
+	FAABB3 LocalBoundingBox;
+	float SignedDistanceScaling;
+	FVector3f TrainingGridMin;
+	TArray<FVector3f> TrainingGridUnitAxesXYZ;
+	TArray<float> TrainingGridAxesLengthsXYZ;
+	int32 TotalNumberOfRotationComponents;
+	TArray<TArray<int32>> ActiveBonesRotationComponents;
+	TArray<FVector3f> ActiveBonesReferenceRotations; 
+	TArray<FVector3f>ActiveBonesReferenceTranslations; 
+
+	//Debug Drawing
+	TUniformGrid<FReal, 3> DebugGrid;
+	TArrayND<FReal, 3> DebugPhi;
+
+private:
+	FMLLevelSet() : FImplicitObject(EImplicitObject::HasBoundingBox, ImplicitObjectType::MLLevelSet) {}	//needed for copy()
+	CHAOS_API FMLLevelSet(TObjectPtr<UNNEModelData> InNNESignedDistanceModelData, TObjectPtr<UNNEModelData> InNNEIncorrectZoneModel); //needed for serialization
+
+	friend FImplicitObject;	//needed for serialization 
+	friend ::FKMLLevelSetElem; //needed for serialization 
+};
+}
diff --git a/Engine/Source/Runtime/Experimental/Chaos/Public/Chaos/PBDSoftBodyCollisionConstraint.h b/Engine/Source/Runtime/Experimental/Chaos/Public/Chaos/PBDSoftBodyCollisionConstraint.h
--- a/Engine/Source/Runtime/Experimental/Chaos/Public/Chaos/PBDSoftBodyCollisionConstraint.h
+++ b/Engine/Source/Runtime/Experimental/Chaos/Public/Chaos/PBDSoftBodyCollisionConstraint.h
@@ -50,10 +50,17 @@ public:
 
 	void SetWriteDebugContacts(bool bWrite) { bWriteDebugContacts = bWrite; }
 
+	// This method uses mutable data that is not thread-safe.
 	CHAOS_API void Apply(FSolverParticlesRange& Particles, const FSolverReal Dt, const TArray<FSolverCollisionParticlesRange>& CollisionParticles) const;
-
+	
 	CHAOS_API void UpdateLinearSystem(const FSolverParticlesRange& Particles, const FSolverReal Dt, const TArray<FSolverCollisionParticlesRange>& CollisionParticles, FEvolutionLinearSystem& LinearSystem) const;
 
+	CHAOS_API const bool DoCollisionPerTimestep() const { return bDoCollisionPerTimestep; }
+	CHAOS_API void SetDoCollisionPerTimestep(bool bInDoCollisionPerTimestep) { bDoCollisionPerTimestep = bInDoCollisionPerTimestep; }
+	CHAOS_API TArray<int32>& GetCollisionIndices() { return CollisionIndices; }
+	CHAOS_API TArray<FSolverVec3>& GetCollisionTargets() { return CollisionTargets; }
+	CHAOS_API TArray<FSolverVec3>& GetCollisionNormalsOutput() { return CollisionNormalsOutput; }
+
 private:
 	template<bool bLockAndWriteContacts, bool bWithFriction>
 	void ApplyInternal(FSolverParticlesRange& Particles, const FSolverReal Dt, const TArray<FSolverCollisionParticlesRange>& CollisionParticles) const;
@@ -61,6 +68,8 @@ private:
 	void ApplyInternalCCD(FSolverParticlesRange& Particles, const FSolverReal Dt, const TArray<FSolverCollisionParticlesRange>& CollisionParticles) const;
 	template<bool bWithFriction>
 	void ApplyInternalISPC(FSolverParticlesRange& Particles, const FSolverReal Dt, const TArray<FSolverCollisionParticlesRange>& CollisionParticles) const;
+	template<bool bWithFriction>
+	void ApplyInternalISPCPerTimestep(FSolverParticlesRange& Particles, const FSolverReal Dt, const TArray<FSolverCollisionParticlesRange>& CollisionParticles) const;
 
 protected:
 	const TArray<FSolverRigidTransform3>& LastCollisionTransforms; // Used by CCD
@@ -77,6 +86,17 @@ protected:
 	TArray<FSolverVec3>* const Normals;
 	TArray<FSolverReal>* const Phis;
 	mutable FCriticalSection DebugMutex;
+
+private:
+	// These arrays are used to do batch PhiAndNormal queries. They live here so they don't need to be reallocated each Apply call. 
+	mutable TArray<FSolverReal> BatchPhis;
+	mutable TArray<FSolverVec3> BatchNormals;
+	mutable TArray<int32> BatchVelocityBones;
+
+	mutable TArray<int32> CollisionIndices;
+	mutable TArray<FSolverVec3> CollisionTargets;
+	mutable TArray<FSolverVec3> CollisionNormalsOutput;
+	bool bDoCollisionPerTimestep = false;
 };
 
 
diff --git a/Engine/Source/Runtime/Experimental/Chaos/Public/Chaos/ParticlesRange.h b/Engine/Source/Runtime/Experimental/Chaos/Public/Chaos/ParticlesRange.h
--- a/Engine/Source/Runtime/Experimental/Chaos/Public/Chaos/ParticlesRange.h
+++ b/Engine/Source/Runtime/Experimental/Chaos/Public/Chaos/ParticlesRange.h
@@ -53,6 +53,10 @@ public:
 	int32 GetOffset() const { return Offset; }
 	int32 GetRangeSize() const { return RangeSize; }
 	int32 Size() const { return RangeSize; } // So this has same interface as Particles
+	bool IsValidIndex(int32 Index) const 
+	{
+		return (Index >= 0) && (Index < RangeSize);
+	}
 
 protected:
 
diff --git a/Engine/Source/Runtime/Experimental/Chaos/Public/Chaos/SoftsEvolution.h b/Engine/Source/Runtime/Experimental/Chaos/Public/Chaos/SoftsEvolution.h
--- a/Engine/Source/Runtime/Experimental/Chaos/Public/Chaos/SoftsEvolution.h
+++ b/Engine/Source/Runtime/Experimental/Chaos/Public/Chaos/SoftsEvolution.h
@@ -212,6 +212,8 @@ public:
 	bool GetDoQuasistatics() const { return bDoQuasistatics; }
 	void SetDisableTimeDependentNumIterations(bool bDisable) { bDisableTimeDependentNumIterations = bDisable; }
 	CHAOS_API void SetSolverProperties(const FCollectionPropertyConstFacade& PropertyCollection);
+	bool DoCollisionPerTimestep() const { return bDoCollisionPerTimestep; }
+	void SetDoCollisionPerTimestep(bool bInDoCollisionPerTimestep) { bDoCollisionPerTimestep = bInDoCollisionPerTimestep; }
 
 private:
 
@@ -436,6 +438,8 @@ private:
 	bool bDoQuasistatics;
 	FSolverReal SolverFrequency; 
 	FLinearSystemParameters LinearSystemParameters; // Per-solver parameters that need to be passed to the linear system solver
+	bool bDoCollisionPerTimestep = false; // Use PlanarConstraints
+
 
 	// Per-Particle data
 	FSolverParticles Particles;
diff --git a/Engine/Source/Runtime/Experimental/Chaos/Public/Chaos/WeightedLatticeImplicitObject.h b/Engine/Source/Runtime/Experimental/Chaos/Public/Chaos/WeightedLatticeImplicitObject.h
--- a/Engine/Source/Runtime/Experimental/Chaos/Public/Chaos/WeightedLatticeImplicitObject.h
+++ b/Engine/Source/Runtime/Experimental/Chaos/Public/Chaos/WeightedLatticeImplicitObject.h
@@ -5,6 +5,7 @@
 #include "Chaos/ArrayND.h"
 #include "Chaos/AABB.h"
 #include "Chaos/HierarchicalSpatialHash.h"
+#include "Chaos/PBDSoftsEvolutionFwd.h"
 
 struct FKSkinnedLevelSetElem;
 
@@ -186,6 +187,10 @@ public:
 	TWeightedLatticeImplicitObject(ObjectType&& InObject, TUniformGrid<FReal, 3>&& InGrid,
 		TArrayND<FWeightedLatticeInfluenceData, 3>&& InBoneData, TArray<FName>&& InUsedBones, TArray<FTransform>&& InReferenceRelativeTransforms);
 
+	virtual ~TWeightedLatticeImplicitObject() override = default;
+
+	TWeightedLatticeImplicitObject(const TWeightedLatticeImplicitObject& Other) = delete;
+
 	TWeightedLatticeImplicitObject(TWeightedLatticeImplicitObject&& Other);
 
 	static constexpr EImplicitObjectType StaticType()
@@ -211,6 +216,20 @@ public:
 	// When bIncludeEmptyCells is true, the closest non-empty lattice cell will be used to find the closest surface point.
 	FReal PhiWithNormalAndSurfacePoint(const FVec3& X, FVec3& Normal, FEmbeddingCoordinate& SurfaceCoord, bool bIncludeEmptyCells = false) const;
 
+	// Batch calls used by cloth.
+	// Note that if x is in an empty cell / outside the deformed lattice, Phi will be UE_BIG_NUMBER and Normal will be unchanged (BatchNormals will be resized using SetNumUninitialized, so it may be left uninitialized)
+	// GreatestInfluenceBone will not be calculated (and set to INDEX_NONE) if Phi is >= Thickness.
+	// All outputs will be unchanged (and possibly uninitialized) for points with InvM == 0.
+	void BatchPhiWithNormalAndGreatestInfluenceBone(const TConstArrayView<Softs::FPAndInvM> PAndInvMArray,
+		const Softs::FSolverRigidTransform3& SolverToThis, const Softs::FSolverReal Thickness,
+		TArray<Softs::FSolverReal>& OutBatchPhis, TArray<Softs::FSolverVec3>& OutBatchNormals, TArray<int32>& OutBatchGreatestInfluenceSolverBones) const;
+
+	// Batch calls used by cloth.
+	// Note that if x is in an empty cell / outside the deformed lattice, Phi will be UE_BIG_NUMBER and Normal will be unchanged (BatchNormals will be resized using SetNumUninitialized, so it may be left uninitialized)
+	// All outputs will be unchanged (and possibly uninitialized) for points with InvM == 0.
+	void BatchPhiWithNormal(const TConstArrayView<Softs::FPAndInvM> PAndInvMArray,
+		const Softs::FSolverRigidTransform3& SolverToThis, TArray<Softs::FSolverReal>& OutBatchPhis, TArray<Softs::FSolverVec3>& OutBatchNormals) const;
+
 	virtual FReal PhiWithNormal(const FVec3& x, FVec3& Normal) const override
 	{
 		FEmbeddingCoordinate SurfaceCoordUnused;
@@ -229,12 +248,6 @@ private:
 	// Serialized data
 	ObjectType Object;
 
-	static TWeightedLatticeImplicitObject<TConcrete>* CopyHelper(const TWeightedLatticeImplicitObject<TConcrete>* Obj)
-	{
-		FImplicitObjectPtr CopiedShape = Obj->Object->CopyGeometry();
-		return new TWeightedLatticeImplicitObject<TConcrete>(reinterpret_cast<ObjectType&&>(CopiedShape), *Obj);
-	}
-
 	TWeightedLatticeImplicitObject(ObjectType&& InObject, const FWeightedLatticeImplicitObject& Other)
 		: FWeightedLatticeImplicitObject(Other)
 		, Object(MoveTemp(InObject))
diff --git a/Engine/Source/Runtime/Experimental/Chaos/Public/Chaos/XPBDPlanarConstraints.h b/Engine/Source/Runtime/Experimental/Chaos/Public/Chaos/XPBDPlanarConstraints.h
new file mode 100644
--- /dev/null
+++ b/Engine/Source/Runtime/Experimental/Chaos/Public/Chaos/XPBDPlanarConstraints.h
@@ -0,0 +1,57 @@
+// Copyright Epic Games, Inc. All Rights Reserved.
+#pragma once
+
+#include "Chaos/PBDSpringConstraintsBase.h"
+
+namespace Chaos::Softs
+{
+	template <typename T>
+	class TXPBDPlanarConstraints
+	{
+	public:
+		TXPBDPlanarConstraints() { LambdaArray.Init((T)0., 0); }
+
+		virtual ~TXPBDPlanarConstraints() {}
+
+		void Apply(FSolverParticlesRange& Particles, const T Dt, const TArray<int32>& CollisionIndices, const TArray<TVec3<T>>& CollisionTargets, const TArray<TVec3<T>>& CollisionNormals) 
+		{
+			TRACE_CPUPROFILER_EVENT_SCOPE(STAT_ChaosXPBDPlanarConstraintApply);
+			PhysicsParallelFor(CollisionIndices.Num(), [this, &Particles, Dt, &CollisionIndices, &CollisionTargets, &CollisionNormals](const int32 ConstraintIndex)
+				{
+					ApplySingleConstraint(Particles, Dt, ConstraintIndex, CollisionIndices[ConstraintIndex], CollisionTargets[ConstraintIndex], CollisionNormals[ConstraintIndex]);
+				});
+		}
+
+		void Init(const FSolverParticlesRange& InParticles, const T Dt, const int32 CollisionSize) { LambdaArray.Init((T)0., CollisionSize); }
+
+		void SetTolerance(const T TolIn) { Tol = TolIn; }
+
+		void SetStiffness(const T StiffnessIn) { Stiffness = StiffnessIn; }
+
+	private:
+
+		void ApplySingleConstraint(FSolverParticlesRange& Particles, const T Dt, const int32 ConstraintIndex, const int32 ParticleIndex, const TVec3<T>& CollisionTarget, const TVec3<T>& CollisionNormal) 
+		{
+			const TVec3<T> Diff = Particles.P(ParticleIndex) - CollisionTarget;
+			const T DiffDotNormal = Chaos::TVec3<T>::DotProduct(Diff, CollisionNormal);
+			if (DiffDotNormal < Tol && Particles.InvM(ParticleIndex) != T(0.) && Stiffness>.0)
+			{
+				T Constraint = Tol - DiffDotNormal;
+				T AlphaTilde = T(1) / (Dt * Dt * Stiffness);
+				if (Stiffness > StiffnessThreshold)
+				{
+					AlphaTilde = T(0);
+				}
+				T DeltaLambda = (-Constraint - AlphaTilde * LambdaArray[ConstraintIndex]) / (Particles.InvM(ParticleIndex) + AlphaTilde);
+				LambdaArray[ConstraintIndex] += DeltaLambda;
+				Particles.P(ParticleIndex) += (-DeltaLambda) * Particles.InvM(ParticleIndex) * CollisionNormal;
+			}
+		}
+
+	protected:
+		T Tol = T(.1);
+		T Stiffness = 1e10;
+		T StiffnessThreshold = 1e9;
+		TArray<T> LambdaArray;
+	};
+}  // End namespace Chaos::Softs
\ No newline at end of file
