Geometry collection : Fix potential crash in InitializeBodiesPT because of an access to the wrong thread data
- children were iterated using the game thread dynamic collection instead of the physics thread one

diff --git a/Engine/Source/Runtime/Experimental/Chaos/Private/PhysicsProxy/GeometryCollectionPhysicsProxy.cpp b/Engine/Source/Runtime/Experimental/Chaos/Private/PhysicsProxy/GeometryCollectionPhysicsProxy.cpp
--- a/Engine/Source/Runtime/Experimental/Chaos/Private/PhysicsProxy/GeometryCollectionPhysicsProxy.cpp
+++ b/Engine/Source/Runtime/Experimental/Chaos/Private/PhysicsProxy/GeometryCollectionPhysicsProxy.cpp
@@ -1360,7 +1360,7 @@ void FGeometryCollectionPhysicsProxy::InitializeBodiesPT(Chaos::FPBDRigidsSolver
 				else
 				{
 					// Cluster parent
-					GameThreadCollection.IterateThroughChildren(TransformGroupIndex, [&](int32 ChildIndex)
+					PhysicsThreadCollection.IterateThroughChildren(TransformGroupIndex, [&](int32 ChildIndex)
 					{
 						if (SubTreeContainsSimulatableParticle.IsValidIndex(ChildIndex) && SubTreeContainsSimulatableParticle[ChildIndex])
 						{
@@ -1404,7 +1404,7 @@ void FGeometryCollectionPhysicsProxy::InitializeBodiesPT(Chaos::FPBDRigidsSolver
 				}
 
 				Chaos::FMatrix33 FullInertia = Chaos::FMatrix33(0);
-				GameThreadCollection.IterateThroughChildren(TransformGroupIndex, [&](int32 ChildIndex)
+				PhysicsThreadCollection.IterateThroughChildren(TransformGroupIndex, [&](int32 ChildIndex)
 					{
 						Masses[TransformGroupIndex] += Masses[ChildIndex];
 						const Chaos::FMatrix33 ChildWorldSpaceI = Chaos::Utilities::ComputeWorldSpaceInertia(Transforms[ChildIndex].GetRotation(), Inertias[ChildIndex]);
