PBDSoftBodyCollisionConstraint.cpp warning issue fixed.

diff --git a/Engine/Source/Runtime/Experimental/Chaos/Private/Chaos/PBDSoftBodyCollisionConstraint.cpp b/Engine/Source/Runtime/Experimental/Chaos/Private/Chaos/PBDSoftBodyCollisionConstraint.cpp
--- a/Engine/Source/Runtime/Experimental/Chaos/Private/Chaos/PBDSoftBodyCollisionConstraint.cpp
+++ b/Engine/Source/Runtime/Experimental/Chaos/Private/Chaos/PBDSoftBodyCollisionConstraint.cpp
@@ -743,15 +743,7 @@ void FPBDSoftBodyCollisionConstraintBase::ApplyInternalISPC(FSolverParticlesRang
 			const int32 NumberOfMLLevelSets = MLLevelSetIndexes.Num();
 			if (NumberOfMLLevelSets > 0)
 			{
-				int32 NumberOfParallelThreadsPerMLLevelset = 1;
-
-				for (int32 MLCollisionIndex : MLLevelSetIndexes)
-				{
-					FMLLevelSet* MLLevelSet = CollisionParticlesRange.GetGeometry(MLCollisionIndex)->GetObject<FMLLevelSet>();
-					MLLevelSet->UpdateNeuralInferencesNumber(NumberOfParallelThreadsPerMLLevelset);
-				}
-
-				const int32 NumberOfParallelCalls = NumberOfParallelThreadsPerMLLevelset * NumberOfMLLevelSets;
+				const int32 NumberOfParallelCalls = NumberOfMLLevelSets;
 				TArray<TArray<FSolverReal>> BatchPhisArr;
 				TArray<TArray<FSolverVec3>> BatchNormalsArr;
 				TArray<TArray<int32>> BatchVelocityBonesArr;
@@ -769,26 +761,13 @@ void FPBDSoftBodyCollisionConstraintBase::ApplyInternalISPC(FSolverParticlesRang
 				}
 
 				PhysicsParallelFor(MLLevelSetIndexes.Num(), [this, &Particles, Dt, &MLLevelSetIndexes, &CollisionParticlesRange,
-					&PAndInvM, &BatchVelocityBonesArr, &BatchNormalsArr, &BatchPhisArr, NumberOfParallelThreadsPerMLLevelset](int32 IndexML)
+					&PAndInvM, &BatchVelocityBonesArr, &BatchNormalsArr, &BatchPhisArr](int32 IndexML)
 					{
-						const int32 MLLevelSetIndex = IndexML / NumberOfParallelThreadsPerMLLevelset;
+						const int32 MLLevelSetIndex = IndexML;
 						const int32 CollisionIndex = MLLevelSetIndexes[MLLevelSetIndex];
-
-						const int32 MLLevelSetThread = IndexML % NumberOfParallelThreadsPerMLLevelset;
-						const int32 MLLevelSetThreadBlock = Particles.GetRangeSize() / NumberOfParallelThreadsPerMLLevelset;
-						int32 OuterParticleBatchBegin;
-						int32 OuterParticleBatchEnd;
-
-						if (MLLevelSetThread < NumberOfParallelThreadsPerMLLevelset - 1)
-						{
-							OuterParticleBatchBegin = MLLevelSetThread * MLLevelSetThreadBlock;
-							OuterParticleBatchEnd = (MLLevelSetThread + 1) * MLLevelSetThreadBlock;
-						}
-						else
-						{
-							OuterParticleBatchBegin = (NumberOfParallelThreadsPerMLLevelset - 1) * MLLevelSetThreadBlock;;
-							OuterParticleBatchEnd = Particles.GetRangeSize() - 1;
-						}
+						const int32 MLLevelSetThread = 0;
+						const int32 OuterParticleBatchBegin = 0;
+						const int32 OuterParticleBatchEnd = Particles.GetRangeSize() - 1;						
 
 						FMLLevelSet* MLLevelSet = CollisionParticlesRange.GetGeometry(CollisionIndex)->GetObject<FMLLevelSet>();
 						FSolverRigidTransform3 Frame = FSolverRigidTransform3(CollisionParticlesRange.X(CollisionIndex), CollisionParticlesRange.R(CollisionIndex)).Inverse();
@@ -1169,24 +1148,14 @@ void FPBDSoftBodyCollisionConstraintBase::ApplyInternalISPCPerTimestep(FSolverPa
 					&PAndInvM, &BatchVelocityBonesArr, &BatchNormalsArr, &BatchPhisArr, NumberOfParallelThreadsPerMLLevelset,
 					&BatchCollisionIndices, &BatchCollisionTargets, &BatchCollisionNormalsOutput](int32 IndexML)
 					{
-						const int32 MLLevelSetIndex = IndexML / NumberOfParallelThreadsPerMLLevelset;
+						const int32 MLLevelSetIndex = IndexML;
 						const int32 CollisionIndex = MLLevelSetIndexes[MLLevelSetIndex];
 
-						const int32 MLLevelSetThread = IndexML % NumberOfParallelThreadsPerMLLevelset;
+						const int32 MLLevelSetThread = 0;
 						const int32 MLLevelSetThreadBlock = Particles.GetRangeSize() / NumberOfParallelThreadsPerMLLevelset;
-						int32 OuterParticleBatchBegin;
-						int32 OuterParticleBatchEnd;
 
-						if (MLLevelSetThread < NumberOfParallelThreadsPerMLLevelset - 1)
-						{
-							OuterParticleBatchBegin = MLLevelSetThread * MLLevelSetThreadBlock;
-							OuterParticleBatchEnd = (MLLevelSetThread + 1) * MLLevelSetThreadBlock;
-						}
-						else
-						{
-							OuterParticleBatchBegin = (NumberOfParallelThreadsPerMLLevelset - 1) * MLLevelSetThreadBlock;;
-							OuterParticleBatchEnd = Particles.GetRangeSize() - 1;
-						}
+						int32 OuterParticleBatchBegin = (NumberOfParallelThreadsPerMLLevelset - 1) * MLLevelSetThreadBlock;
+						int32 OuterParticleBatchEnd = Particles.GetRangeSize() - 1;
 
 						FMLLevelSet* MLLevelSet = CollisionParticlesRange.GetGeometry(CollisionIndex)->GetObject<FMLLevelSet>();
 						FSolverRigidTransform3 Frame = FSolverRigidTransform3(CollisionParticlesRange.X(CollisionIndex), CollisionParticlesRange.R(CollisionIndex)).Inverse();
