Iris
* Fix bitarrayview issue when the NetRefHandleManager max internal index grows during the PreUpdate phase.
* Moved PreUpdate call from poller struct to the ObjectBridge directly.

#rnx

diff --git a/Engine/Plugins/Runtime/ReplicationSystemTestPlugin/Source/Private/Tests/ReplicationSystem/TestReplicationConfig.cpp b/Engine/Plugins/Runtime/ReplicationSystemTestPlugin/Source/Private/Tests/ReplicationSystem/TestReplicationConfig.cpp
--- a/Engine/Plugins/Runtime/ReplicationSystemTestPlugin/Source/Private/Tests/ReplicationSystem/TestReplicationConfig.cpp
+++ b/Engine/Plugins/Runtime/ReplicationSystemTestPlugin/Source/Private/Tests/ReplicationSystem/TestReplicationConfig.cpp
@@ -190,5 +190,77 @@ UE_NET_TEST_FIXTURE(FReplicationConfigTestFixture, TestNetObjectListGrowEvent)
 	//CreateReplicatedObjects(1);
 }
 
+
+UE_NET_TEST_FIXTURE(FReplicationConfigTestFixture, TestNetObjectListGrow_DuringPollPhase)
+{
+	const uint32 MaxNumObjects = 96;
+	const uint32 InitNumObjects = 32;
+	const uint32 GrowCount = 32;
+
+	OverrideServerConfig.MaxReplicatedObjectCount = MaxNumObjects;
+	OverrideServerConfig.InitialNetObjectListCount = InitNumObjects;
+	OverrideServerConfig.NetObjectListGrowCount = GrowCount;
+
+	StartReplicationSystem();
+
+	FDelegateHandle ServerDelegate;
+
+	FNetRefHandleManager* ServerNetRefHandleManager = &Server->ReplicationSystem->GetReplicationSystemInternal()->GetNetRefHandleManager();
+	ON_SCOPE_EXIT
+	{
+		ServerNetRefHandleManager->GetOnMaxInternalNetRefIndexIncreasedDelegate().Remove(ServerDelegate);
+	};
+
+
+	bool bHasServerRealloc = false;
+	ServerDelegate = ServerNetRefHandleManager->GetOnMaxInternalNetRefIndexIncreasedDelegate().AddLambda([&bHasServerRealloc](uint32 NewMaxIndex)
+	{
+		bHasServerRealloc = true;
+	});
+
+	// Create enough objects to fill the init list (entry[0] is already reserved for invalid object)
+	for (uint32 i = 0; i < (InitNumObjects - 1); ++i)
+	{
+		UObjectReplicationBridge::FRootObjectReplicationParams Params;
+		Params.bNeedsPreUpdate = true;
+		ServerObjects.Emplace(Server->CreateObject(Params));
+	}
+
+	// Send and deliver packet
+	Server->UpdateAndSend({ Client });
+	
+	UTestReplicatedIrisObject* ServerObject = ServerObjects[0];
+	UE_NET_ASSERT_NE(ServerObject, nullptr);
+
+	UTestReplicatedIrisObject* ServerSubObject(nullptr);
+
+	// Add a PreUpdate callback where we create a new subobject
+	auto PreUpdateObject = [&](TArrayView<UObject*> Instances, const UReplicationBridge* Bridge)
+	{
+		for (UObject* ReplicatedObject : Instances)
+		{
+			if (ServerObject == ReplicatedObject)
+			{
+				if (ServerSubObject == nullptr)
+				{
+					// Create a subobject
+					ServerSubObject = Server->CreateSubObject<UTestReplicatedIrisObject>(ServerObject->NetRefHandle);
+				}
+			}
+		}
+	};
+	Server->GetReplicationBridge()->SetExternalPreUpdateFunctor(PreUpdateObject);
+
+	Server->UpdateAndSend( { Client });
+
+	// Check that the server realloc happened.
+	UE_NET_ASSERT_EQ(bHasServerRealloc, true);
+
+	// Make sure the subobject was spawned on the client
+	UTestReplicatedIrisObject* ClientSubObject = Cast<UTestReplicatedIrisObject>(Client->GetReplicationBridge()->GetReplicatedObject(ServerSubObject->NetRefHandle));
+	UE_NET_ASSERT_NE(ClientSubObject, nullptr);
+}
+
+
 } // end namespace UE::Net::Private
 
diff --git a/Engine/Source/Runtime/Experimental/Iris/Core/Private/Iris/ReplicationSystem/ObjectReplicationBridge.cpp b/Engine/Source/Runtime/Experimental/Iris/Core/Private/Iris/ReplicationSystem/ObjectReplicationBridge.cpp
--- a/Engine/Source/Runtime/Experimental/Iris/Core/Private/Iris/ReplicationSystem/ObjectReplicationBridge.cpp
+++ b/Engine/Source/Runtime/Experimental/Iris/Core/Private/Iris/ReplicationSystem/ObjectReplicationBridge.cpp
@@ -873,19 +873,35 @@ void UObjectReplicationBridge::OnStartPreSendUpdate()
 void UObjectReplicationBridge::PreSendUpdate()
 {
 	using namespace UE::Net;
+	using namespace UE::Net::Private;
 
-	FNetBitArrayView ObjectsConsideredForPolling = GetReplicationSystem()->GetReplicationSystemInternal()->GetNetRefHandleManager().GetPolledObjectsInternalIndices();
-	ObjectsConsideredForPolling.ClearAllBits();
+	FNetRefHandleManager& LocalNetRefHandleManager = GetReplicationSystem()->GetReplicationSystemInternal()->GetNetRefHandleManager();
 
-	BuildPollList(ObjectsConsideredForPolling);
+	{
+		const uint32 InitialNumBits = LocalNetRefHandleManager.GetPolledObjectsInternalIndices().GetNumBits();
+		
+		FNetBitArray InitialPollList(InitialNumBits);
+		FNetBitArrayView InitialPollListView(MakeNetBitArrayView(InitialPollList));
+		BuildPollList(InitialPollListView);
 
-	PreUpdate(ObjectsConsideredForPolling);
+		PreUpdate(InitialPollListView);
+		
+		// PreUpdate is allowed to generate new objects, so the netrefhandlemanager bitarrays may have grown.
+		// Limit the size of the NetHandleManager list to the pre-grown array.
+		FNetBitArrayView CachedPollList(LocalNetRefHandleManager.GetPolledObjectsInternalIndices().GetData(), InitialNumBits);
+
+		CachedPollList.Copy(InitialPollList);
+	}
 
 	FinalizeDirtyObjects();
 
-	ReconcileNewSubObjects(ObjectsConsideredForPolling);
+	{
+		FNetBitArrayView ObjectsConsideredForPolling = LocalNetRefHandleManager.GetPolledObjectsInternalIndices();
+		
+		ReconcileNewSubObjects(ObjectsConsideredForPolling);
 
-	PollAndCopy(ObjectsConsideredForPolling);
+		PollAndCopy(ObjectsConsideredForPolling);
+	}
 }
 
 void UObjectReplicationBridge::OnPostSendUpdate()
@@ -966,17 +982,36 @@ void UObjectReplicationBridge::ForcePollObject(FNetRefHandle Handle)
 {
 	using namespace UE::Net::Private;
 
-	if (Handle.IsValid())
+	if (!Handle.IsValid())
 	{
-		IRIS_PROFILER_SCOPE(UObjectReplicationBridge_ForcePollAndCopyObject);
+		return;
+	}
+
+	IRIS_PROFILER_SCOPE(UObjectReplicationBridge_ForcePollAndCopyObject);
 
-		FObjectPoller::FInitParams PollerInitParams;
-		PollerInitParams.ObjectReplicationBridge = this;
-		PollerInitParams.ReplicationSystemInternal = GetReplicationSystem()->GetReplicationSystemInternal();
-		FObjectPoller Poller(PollerInitParams);
+	const FNetRefHandleManager& LocalNetRefHandleManager = GetReplicationSystem()->GetReplicationSystemInternal()->GetNetRefHandleManager();
 
-		Poller.PollAndCopySingleObject(Handle);
+	const FInternalNetRefIndex ObjectIndex = LocalNetRefHandleManager.GetInternalIndex(Handle);
+	if (ObjectIndex == FNetRefHandleManager::InvalidInternalIndex)
+	{
+		return;
+	}
+
+	if (PreUpdateInstanceFunction)
+	{
+		UObject* Instance = LocalNetRefHandleManager.GetReplicatedInstances()[ObjectIndex];
+		if (Instance && LocalNetRefHandleManager.GetObjectsWithPreUpdate().GetBit(ObjectIndex))
+		{
+			PreUpdateInstanceFunction(MakeArrayView<UObject*>(&Instance, 1U), this);
+		}
 	}
+
+	FObjectPoller::FInitParams PollerInitParams;
+	PollerInitParams.ObjectReplicationBridge = this;
+	PollerInitParams.ReplicationSystemInternal = GetReplicationSystem()->GetReplicationSystemInternal();
+	FObjectPoller Poller(PollerInitParams);
+
+	Poller.PollAndCopySingleObject(ObjectIndex);
 }
 
 void UObjectReplicationBridge::BuildPollList(UE::Net::FNetBitArrayView ObjectsConsideredForPolling)
@@ -1122,20 +1157,56 @@ void UObjectReplicationBridge::BuildPollList(UE::Net::FNetBitArrayView ObjectsCo
 	
 void UObjectReplicationBridge::PreUpdate(const UE::Net::FNetBitArrayView ObjectsConsideredForPolling)
 {
+	using namespace UE::Net;
 	using namespace UE::Net::Private;
 
 	IRIS_CSV_PROFILER_SCOPE(Iris, ReplicationBridge_PreUpdate);
+	IRIS_PROFILER_SCOPE_VERBOSE(PreUpdatePass);
 
-	// TODO: Get rid of Poller
-	FObjectPoller::FInitParams PollerInitParams;
-	PollerInitParams.ObjectReplicationBridge = this;
-	PollerInitParams.ReplicationSystemInternal = GetReplicationSystem()->GetReplicationSystemInternal();
-	
-	FObjectPoller Poller(PollerInitParams);
-	Poller.PreUpdatePass(ObjectsConsideredForPolling);
+	if (!PreUpdateInstanceFunction)
+	{
+		return;
+	}
 
-	FObjectPoller::FPreUpdateAndPollStats Stats = Poller.GetPollStats();
-	UE_NET_TRACE_FRAME_STATSCOUNTER(ReplicationSystem->GetId(), ReplicationSystem.PreUpdatedObjectCount, Stats.PreUpdatedObjectCount, ENetTraceVerbosity::Trace);
+	const FNetRefHandleManager& LocalNetRefHandleManager = GetReplicationSystem()->GetReplicationSystemInternal()->GetNetRefHandleManager();
+
+	const TNetChunkedArray<TObjectPtr<UObject>>& ReplicatedInstances = LocalNetRefHandleManager.GetReplicatedInstances();
+
+	uint32 PreUpdatedObjectCount = 0;
+
+	constexpr uint32 PreUpdateBatchCount = 128U;
+	UObject* BatchedObjects[PreUpdateBatchCount];
+	uint32 BatchedObjectCount = 0U;
+
+	auto BatchedPreUpdate = [&](FInternalNetRefIndex Objectindex)
+	{
+		// Flush if needed
+		if (BatchedObjectCount == PreUpdateBatchCount)
+		{
+			PreUpdateInstanceFunction(MakeArrayView<UObject*>(BatchedObjects, BatchedObjectCount), this);
+			PreUpdatedObjectCount += BatchedObjectCount;
+			BatchedObjectCount = 0U;
+		}
+
+		UObject* Instance = ReplicatedInstances[Objectindex];
+		BatchedObjects[BatchedObjectCount] = Instance;
+		BatchedObjectCount += Instance ? 1U : 0U;
+	};
+
+	// Make a copy of the list we'll iterate on since the PreUpdate callbacks can create new objects and realloc the NetRefHandleManager bitarrays
+	FNetBitArray ObjectsWithPreUpdate;
+	ObjectsWithPreUpdate.InitAndCopy(LocalNetRefHandleManager.GetObjectsWithPreUpdate());
+
+	FNetBitArrayView::ForAllSetBits(ObjectsConsideredForPolling, MakeNetBitArrayView(ObjectsWithPreUpdate), FNetBitArrayView::AndOp, BatchedPreUpdate);
+
+	// Flush last batch
+	if (BatchedObjectCount > 0)
+	{
+		PreUpdateInstanceFunction(MakeArrayView(BatchedObjects, BatchedObjectCount), this);
+		PreUpdatedObjectCount += BatchedObjectCount;
+	}
+
+	UE_NET_TRACE_FRAME_STATSCOUNTER(ReplicationSystem->GetId(), ReplicationSystem.PreUpdatedObjectCount, PreUpdatedObjectCount, ENetTraceVerbosity::Trace);
 }
 
 void UObjectReplicationBridge::PollAndCopy(const UE::Net::FNetBitArrayView ObjectsConsideredForPolling)
diff --git a/Engine/Source/Runtime/Experimental/Iris/Core/Private/Iris/ReplicationSystem/Polling/ObjectPoller.cpp b/Engine/Source/Runtime/Experimental/Iris/Core/Private/Iris/ReplicationSystem/Polling/ObjectPoller.cpp
--- a/Engine/Source/Runtime/Experimental/Iris/Core/Private/Iris/ReplicationSystem/Polling/ObjectPoller.cpp
+++ b/Engine/Source/Runtime/Experimental/Iris/Core/Private/Iris/ReplicationSystem/Polling/ObjectPoller.cpp
@@ -33,54 +33,6 @@ FObjectPoller::FObjectPoller(const FInitParams& InitParams)
 	// DirtyObjectsThisFrame is acquired only during polling 
 }
 
-void FObjectPoller::PreUpdatePass(const FNetBitArrayView& ObjectsConsideredForPolling)
-{
-	IRIS_PROFILER_SCOPE_VERBOSE(PreUpdatePass);
-
-	constexpr uint32 PreUpdateBatchCount = 128U;
-
-	UObject* BatchedObjects[PreUpdateBatchCount];
-	uint32 BatchedObjectCount = 0U;
-
-	if (!ObjectReplicationBridge->PreUpdateInstanceFunction)
-	{
-		return;
-	}
-	
-	auto BatchedPreUpdate = [this, &BatchedObjects, &BatchedObjectCount](FInternalNetRefIndex Objectindex)
-	{
-		// Flush if needed
-		if (BatchedObjectCount == PreUpdateBatchCount)
-		{
-			ObjectReplicationBridge->PreUpdateInstanceFunction(MakeArrayView<UObject*>(BatchedObjects, BatchedObjectCount), ObjectReplicationBridge);
-			PollStats.PreUpdatedObjectCount += BatchedObjectCount;
-			BatchedObjectCount = 0U;
-		}
-		UObject* Instance = ReplicatedInstances[Objectindex];
-		BatchedObjects[BatchedObjectCount] = Instance;
-		BatchedObjectCount += Instance ? 1U : 0U;
-	};
-
-	FNetBitArrayView::ForAllSetBits(ObjectsConsideredForPolling, LocalNetRefHandleManager.GetObjectsWithPreUpdate(), FNetBitArrayView::AndOp, BatchedPreUpdate);
-
-	// Flush last batch
-	if (BatchedObjectCount > 0)
-	{
-		ObjectReplicationBridge->PreUpdateInstanceFunction(MakeArrayView(BatchedObjects, BatchedObjectCount), ObjectReplicationBridge);
-		PollStats.PreUpdatedObjectCount += BatchedObjectCount;
-	}
-}
-
-void FObjectPoller::CallPreUpdate(FInternalNetRefIndex ObjectIndex)
-{
-	UObject* Instance = ReplicatedInstances[ObjectIndex];
-	if (Instance && LocalNetRefHandleManager.GetObjectsWithPreUpdate().GetBit(ObjectIndex))
-	{
-		ObjectReplicationBridge->PreUpdateInstanceFunction(MakeArrayView<UObject*>(&Instance, 1U), ObjectReplicationBridge);	
-		++PollStats.PreUpdatedObjectCount;
-	}
-}
-
 void FObjectPoller::PollAndCopyObjects(const FNetBitArrayView& ObjectsConsideredForPolling)
 {
 	FDirtyObjectsAccessor DirtyObjectsAccessor(ReplicationSystemInternal->GetDirtyNetObjectTracker());
@@ -108,23 +60,16 @@ void FObjectPoller::PollAndCopyObjects(const FNetBitArrayView& ObjectsConsidered
 	NetStatsContext = nullptr;
 }
 
-void FObjectPoller::PollAndCopySingleObject(FNetRefHandle Handle)
+void FObjectPoller::PollAndCopySingleObject(FInternalNetRefIndex ObjectIndex)
 {
-	if (uint32 InternalObjectIndex = LocalNetRefHandleManager.GetInternalIndex(Handle))
-	{
-		NetStatsContext = ReplicationSystemInternal->GetNetTypeStats().GetNetStatsContext();
-
-		CallPreUpdate(InternalObjectIndex);
-
-		FDirtyObjectsAccessor DirtyObjectsAccessor(ReplicationSystemInternal->GetDirtyNetObjectTracker());
-		DirtyObjectsThisFrame = DirtyObjectsAccessor.GetDirtyNetObjects();
+	FDirtyObjectsAccessor DirtyObjectsAccessor(ReplicationSystemInternal->GetDirtyNetObjectTracker());
+	DirtyObjectsThisFrame = DirtyObjectsAccessor.GetDirtyNetObjects();
 
-		ForcePollObject(InternalObjectIndex);
+	ForcePollObject(ObjectIndex);
 
-		// Clear ref to locked dirty bit array
-		DirtyObjectsThisFrame = FNetBitArrayView();
-		NetStatsContext = nullptr;
-	}
+	// Clear ref to locked dirty bit array
+	DirtyObjectsThisFrame = FNetBitArrayView();
+	NetStatsContext = nullptr;
 }
 
 void FObjectPoller::ForcePollObject(FInternalNetRefIndex ObjectIndex)
diff --git a/Engine/Source/Runtime/Experimental/Iris/Core/Private/Iris/ReplicationSystem/Polling/ObjectPoller.h b/Engine/Source/Runtime/Experimental/Iris/Core/Private/Iris/ReplicationSystem/Polling/ObjectPoller.h
--- a/Engine/Source/Runtime/Experimental/Iris/Core/Private/Iris/ReplicationSystem/Polling/ObjectPoller.h
+++ b/Engine/Source/Runtime/Experimental/Iris/Core/Private/Iris/ReplicationSystem/Polling/ObjectPoller.h
@@ -50,23 +50,14 @@ public:
 
 	const FPreUpdateAndPollStats& GetPollStats() const { return PollStats; }
 
-	/**
-	 * Call the PreUpdateFunction, generally AActor::PreReplication, on polled objects that need it.
-	 * @param ObjectsConsideredForPolling The list of actors set to be polled this frame
-	 */
-	void PreUpdatePass(const FNetBitArrayView& ObjectsConsideredForPolling);
-
 	/** Poll all the objects whose bit index is set in the array and copy any dirty data into ReplicationState buffers*/
 	void PollAndCopyObjects(const FNetBitArrayView& ObjectsConsideredForPolling);
 
 	/** Poll a single replicated object */
-	void PollAndCopySingleObject(FNetRefHandle Handle);
+	void PollAndCopySingleObject(FInternalNetRefIndex ObjectIndex);
 
 private:
 
-	/** Calls the PreUpdate function if the object needs it */
-	void CallPreUpdate(FInternalNetRefIndex ObjectIndex);
-
 	/** Polls an object in every circumstance */
 	void ForcePollObject(FInternalNetRefIndex ObjectIndex);
 
