Fixed over-allocation in the active arrays used to buffer results from physics. Previously we would allocate a buffer large enough for every object in the scene - even if we only have one potential dirty particle

diff --git a/Engine/Source/Runtime/Experimental/Chaos/Private/PBDRigidsSolver.cpp b/Engine/Source/Runtime/Experimental/Chaos/Private/PBDRigidsSolver.cpp
--- a/Engine/Source/Runtime/Experimental/Chaos/Private/PBDRigidsSolver.cpp
+++ b/Engine/Source/Runtime/Experimental/Chaos/Private/PBDRigidsSolver.cpp
@@ -2068,12 +2068,17 @@ namespace Chaos
 		FPullPhysicsData* PullData = MarshallingManager.GetCurrentPullData_Internal();
 
 		TParticleView<FPBDRigidParticles>& DirtyParticles = GetParticles().GetDirtyParticlesView();
+		const int32 NumDirty = DirtyParticles.Num();
 
+		// We can have at most NumDirty overall active particles across all dirty proxy lists
+		// and in each list at most NumProxies. So we take the min, otherwise for example in a world
+		// with a few million static particles and one dynamic we'll allocate enough storage for the
+		// millions of particles despite only ever syncing one
 		TArray<FSingleParticlePhysicsProxy*> ActiveRigid;
-		ActiveRigid.Reserve(SingleParticlePhysicsProxies_PT.Num());
+		ActiveRigid.Reserve(FMath::Min(SingleParticlePhysicsProxies_PT.Num(), NumDirty));
 
 		TArray<FClusterUnionPhysicsProxy*> ActiveClusterUnions;
-		ActiveClusterUnions.Reserve(ClusterUnionPhysicsProxies_Internal.Num());
+		ActiveClusterUnions.Reserve(FMath::Min(ClusterUnionPhysicsProxies_Internal.Num(), NumDirty));
 
 		const bool bIsResim = GetEvolution()->IsResimming();
 
