Log OnDemandIOStore cache usage at startup

diff --git a/Engine/Source/Runtime/Experimental/IoStore/OnDemand/Private/IoStoreOnDemand.cpp b/Engine/Source/Runtime/Experimental/IoStore/OnDemand/Private/IoStoreOnDemand.cpp
--- a/Engine/Source/Runtime/Experimental/IoStore/OnDemand/Private/IoStoreOnDemand.cpp
+++ b/Engine/Source/Runtime/Experimental/IoStore/OnDemand/Private/IoStoreOnDemand.cpp
@@ -1186,11 +1186,11 @@ FIoStatus FIoStoreOnDemandModule::GetInstallSizesByMountId(const FOnDemandGetIns
 	return FIoStatus(EIoErrorCode::InvalidCode, NotInitializedError);
 }
 
-TIoStatusOr<FOnDemandCacheUsage> FIoStoreOnDemandModule::GetCacheUsage() const
+TIoStatusOr<FOnDemandCacheUsage> FIoStoreOnDemandModule::GetCacheUsage(ELogOnDemandCacheUsage LogMode) const
 {
 	if (IoStore)
 	{
-		return IoStore->GetCacheUsage();
+		return IoStore->GetCacheUsage(LogMode);
 	}
 
 	return FIoStatus(EIoErrorCode::InvalidCode, NotInitializedError);
diff --git a/Engine/Source/Runtime/Experimental/IoStore/OnDemand/Private/OnDemandIoStore.cpp b/Engine/Source/Runtime/Experimental/IoStore/OnDemand/Private/OnDemandIoStore.cpp
--- a/Engine/Source/Runtime/Experimental/IoStore/OnDemand/Private/OnDemandIoStore.cpp
+++ b/Engine/Source/Runtime/Experimental/IoStore/OnDemand/Private/OnDemandIoStore.cpp
@@ -23,6 +23,7 @@
 #include "HAL/PlatformProcess.h"
 #include "IO/IoChunkEncoding.h"
 #include "IO/IoContainerHeader.h"
+#include "IO/IoStoreOnDemand.h"
 #include "IO/PackageStore.h"
 #include "Misc/CommandLine.h"
 #include "Misc/ConfigCacheIni.h"
@@ -877,6 +878,9 @@ FIoStatus FOnDemandIoStore::InitializeOnDemandInstallCache()
 	}
 #endif
 
+	// Log initial cache state
+	GetCacheUsage(ELogOnDemandCacheUsage::Default);
+
 	return EIoErrorCode::Ok;
 }
 
@@ -1929,9 +1933,20 @@ void FOnDemandIoStore::GetReferencedContent(TArray<FSharedOnDemandContainer>& Ou
 	}
 }
 
-FOnDemandCacheUsage FOnDemandIoStore::GetCacheUsage() const
+FOnDemandCacheUsage FOnDemandIoStore::GetCacheUsage(ELogOnDemandCacheUsage LogMode) const
 {
 	FOnDemandInstallCacheStorageUsage Usage = InstallCache->GetStorageUsage();
+
+	if (LogMode == ELogOnDemandCacheUsage::Default)
+	{
+		UE_LOG(LogIoStoreOnDemand, Display, TEXT(""));
+		UE_LOG(LogIoStoreOnDemand, Display, TEXT("* IoStore On Demand Cache Usage"));
+		UE_LOG(LogIoStoreOnDemand, Display, TEXT("* \tMax Size:              %" UINT64_FMT), Usage.MaxSize);
+		UE_LOG(LogIoStoreOnDemand, Display, TEXT("* \tTotal Size:            %" UINT64_FMT), Usage.TotalSize);
+		UE_LOG(LogIoStoreOnDemand, Display, TEXT("* \tReferenced Block Size: %" UINT64_FMT), Usage.ReferencedBlockSize);
+		UE_LOG(LogIoStoreOnDemand, Display, TEXT(""));
+	}
+
 	return FOnDemandCacheUsage
 	{
 		.MaxSize = Usage.MaxSize,
diff --git a/Engine/Source/Runtime/Experimental/IoStore/OnDemand/Private/OnDemandIoStore.h b/Engine/Source/Runtime/Experimental/IoStore/OnDemand/Private/OnDemandIoStore.h
--- a/Engine/Source/Runtime/Experimental/IoStore/OnDemand/Private/OnDemandIoStore.h
+++ b/Engine/Source/Runtime/Experimental/IoStore/OnDemand/Private/OnDemandIoStore.h
@@ -27,6 +27,7 @@ namespace UE::IoStore
 
 class IOnDemandPackageStoreBackend;
 class IOnDemandInstallCache;
+enum class ELogOnDemandCacheUsage : uint8;
 using FSharedPackageStoreBackend	= TSharedPtr<IOnDemandPackageStoreBackend>;
 using FSharedInstallCache			= TSharedPtr<IOnDemandInstallCache>;
 using FWeakOnDemandIoStore			= TWeakPtr<FOnDemandIoStore, ESPMode::ThreadSafe>;
@@ -300,7 +301,7 @@ public:
 	FOnDemandChunkInfo		GetInstalledChunkInfo(const FIoChunkId& ChunkId);
 	void					ReleaseContent(FOnDemandInternalContentHandle& ContentHandle);
 	void					GetReferencedContent(TArray<FSharedOnDemandContainer>& OutContainers, TArray<TBitArray<>>& OutChunkEntryIndices);
-	FOnDemandCacheUsage		GetCacheUsage() const;
+	FOnDemandCacheUsage		GetCacheUsage(ELogOnDemandCacheUsage LogMode = ELogOnDemandCacheUsage::None) const;
 
 private:
 	FIoStatus				GetContainersForInstall(
diff --git a/Engine/Source/Runtime/Experimental/IoStore/OnDemand/Public/IO/IoStoreOnDemand.h b/Engine/Source/Runtime/Experimental/IoStore/OnDemand/Public/IO/IoStoreOnDemand.h
--- a/Engine/Source/Runtime/Experimental/IoStore/OnDemand/Public/IO/IoStoreOnDemand.h
+++ b/Engine/Source/Runtime/Experimental/IoStore/OnDemand/Public/IO/IoStoreOnDemand.h
@@ -430,6 +430,12 @@ struct FOnDemandCacheUsage
 	uint64 ReferencedBlockSize = 0;
 };
 
+enum class ELogOnDemandCacheUsage : uint8
+{
+	None = 0,
+	Default,
+};
+
 /** Token used for signalling an operation to be cancelled. */
 class FOnDemandCancellationToken
 {
@@ -496,7 +502,7 @@ public:
 	  * if there are requests currently being processed. 
 	  * This should only be used for purposes such as debugging telemetry. 
 	  */
-	UE_API TIoStatusOr<FOnDemandCacheUsage> GetCacheUsage() const;
+	UE_API TIoStatusOr<FOnDemandCacheUsage> GetCacheUsage(ELogOnDemandCacheUsage LogMode = ELogOnDemandCacheUsage::None) const;
 
 	virtual void StartupModule() override;
 	virtual void ShutdownModule() override;
