Add content handle log statement

diff --git a/Engine/Source/Runtime/Experimental/IoStore/OnDemand/Private/OnDemandIoStore.cpp b/Engine/Source/Runtime/Experimental/IoStore/OnDemand/Private/OnDemandIoStore.cpp
--- a/Engine/Source/Runtime/Experimental/IoStore/OnDemand/Private/OnDemandIoStore.cpp
+++ b/Engine/Source/Runtime/Experimental/IoStore/OnDemand/Private/OnDemandIoStore.cpp
@@ -1507,7 +1507,7 @@ FOnDemandInstallResult FOnDemandIoStore::TickInstallRequest(const FInstallReques
 		return OutResult;
 	}
 
-	const TSharedPtr<FOnDemandInternalContentHandle, ESPMode::ThreadSafe>& ContentHandle = InstallRequest.Args.ContentHandle.Handle;
+	const TSharedPtr<FOnDemandInternalContentHandle>& ContentHandle = InstallRequest.Args.ContentHandle.Handle;
 	if (ContentHandle->IoStore.IsValid() == false)
 	{
 		// First time this content handle is used
@@ -1537,8 +1537,8 @@ FOnDemandInstallResult FOnDemandIoStore::TickInstallRequest(const FInstallReques
 		}
 	}
 
-	UE_LOG(LogIoStoreOnDemand, Log, TEXT("Downloaded %llu (%.2lf MiB) of total %llu (%.2lf MiB) requested chunk(s)'"),
-		ToDownloadChunkCount, double(ToDownloadBytes) / 1024.0 / 1024.0, TotalChunkCount, double(TotalBytes) / 1024.0 / 1024.0);
+	UE_LOG(LogIoStoreOnDemand, Log, TEXT("Downloaded %llu (%.2lf MiB) of total %llu (%.2lf MiB) requested chunk(s), ContentHandle='%s'"),
+		ToDownloadChunkCount, double(ToDownloadBytes) / 1024.0 / 1024.0, TotalChunkCount, double(TotalBytes) / 1024.0 / 1024.0, *LexToString(*ContentHandle));
 
 	return OutResult;
 }
