Tentative fix for avoiding assert when removing geometry from a cluster union
(happens rarely when exiting PIE with maps that contains cluster union )

diff --git a/Engine/Source/Runtime/Experimental/Chaos/Public/Chaos/ClusterUnionManager.h b/Engine/Source/Runtime/Experimental/Chaos/Public/Chaos/ClusterUnionManager.h
--- a/Engine/Source/Runtime/Experimental/Chaos/Public/Chaos/ClusterUnionManager.h
+++ b/Engine/Source/Runtime/Experimental/Chaos/Public/Chaos/ClusterUnionManager.h
@@ -401,28 +401,31 @@ namespace Chaos
 			return;
 		}
 
-		check(AllChildParticles.Num() == ClusterParticle->ShapesArray().Num());
-		
-		TArray<int32> ShapeIndicesToRemove;
-		ShapeIndicesToRemove.Reserve(ShapeParticles.Num());
-		for(TParticle* ShapeParticle : ShapeParticles)
+		// Don't try to remove anything if the shape array is already empty 
+		if (!ClusterParticle->ShapesArray().IsEmpty())
 		{
-			check(ShapeParticle != nullptr);
-			const int32 Index = AllChildParticles.Find(ShapeParticle);
-			if (Index != INDEX_NONE)
+			check(AllChildParticles.Num() == ClusterParticle->ShapesArray().Num());
+
+			TArray<int32> ShapeIndicesToRemove;
+			ShapeIndicesToRemove.Reserve(ShapeParticles.Num());
+			for (TParticle* ShapeParticle : ShapeParticles)
 			{
-				ShapeIndicesToRemove.Add(Index);
+				check(ShapeParticle != nullptr);
+				const int32 Index = AllChildParticles.Find(ShapeParticle);
+				if (Index != INDEX_NONE)
+				{
+					ShapeIndicesToRemove.Add(Index);
+				}
 			}
-		}
-
-		ShapeIndicesToRemove.Sort();
 
-		ClusterParticle->RemoveShapesAtSortedIndices(ShapeIndicesToRemove);
+			ShapeIndicesToRemove.Sort();
 
-		RemoveArrayItemsAtSortedIndices(AllChildParticles, ShapeIndicesToRemove);
+			ClusterParticle->RemoveShapesAtSortedIndices(ShapeIndicesToRemove);
 
-		check(AllChildParticles.Num() == ClusterParticle->ShapesArray().Num());
+			RemoveArrayItemsAtSortedIndices(AllChildParticles, ShapeIndicesToRemove);
 
+			check(AllChildParticles.Num() == ClusterParticle->ShapesArray().Num());
+		}
 		// If we remove particles from the cluster union geometry then we need to switch the geometry back to a FImplicitObjectUnionClustered to avoid errors with empty unions.
 		if (ClusterParticle->ShapesArray().IsEmpty())
 		{
