Allow setting IAD cache size with
-Iad.FileCache.DiskQuota

diff --git a/Engine/Source/Runtime/Experimental/IoStore/OnDemand/Private/OnDemandInstallCache.cpp b/Engine/Source/Runtime/Experimental/IoStore/OnDemand/Private/OnDemandInstallCache.cpp
--- a/Engine/Source/Runtime/Experimental/IoStore/OnDemand/Private/OnDemandInstallCache.cpp
+++ b/Engine/Source/Runtime/Experimental/IoStore/OnDemand/Private/OnDemandInstallCache.cpp
@@ -1705,7 +1705,7 @@ TSharedPtr<IOnDemandInstallCache> MakeOnDemandInstallCache(
 		return TSharedPtr<IOnDemandInstallCache>();
 	}
 
-	return MakeShareable<IOnDemandInstallCache>(new FOnDemandInstallCache(Config, IoStore));
+	return MakeShared<FOnDemandInstallCache>(Config, IoStore);
 }
 
 } // namespace UE::IoStore
diff --git a/Engine/Source/Runtime/Experimental/IoStore/OnDemand/Private/OnDemandIoStore.cpp b/Engine/Source/Runtime/Experimental/IoStore/OnDemand/Private/OnDemandIoStore.cpp
--- a/Engine/Source/Runtime/Experimental/IoStore/OnDemand/Private/OnDemandIoStore.cpp
+++ b/Engine/Source/Runtime/Experimental/IoStore/OnDemand/Private/OnDemandIoStore.cpp
@@ -789,13 +789,22 @@ FIoStatus FOnDemandIoStore::InitializeOnDemandInstallCache()
 	{
 		FOnDemandInstallCacheConfig CacheConfig;
 
-		if (FString ValueStr; GConfig->GetString(TEXT("OnDemandInstall"), TEXT("FileCache.DiskQuota"), ValueStr, GEngineIni))
+		const int64 DiskQuota = []
 		{
-			int64 DiskQuota = ParseSizeParam(ValueStr);
-			if (DiskQuota > 0)
+			if (FString ParamValue; FParse::Value(FCommandLine::Get(), TEXT("-Iad.FileCache.DiskQuota="), ParamValue))
 			{
-				CacheConfig.DiskQuota = DiskQuota;
+				return ParseSizeParam(ParamValue);
 			}
+			else if (FString ValueStr; GConfig->GetString(TEXT("OnDemandInstall"), TEXT("FileCache.DiskQuota"), ValueStr, GEngineIni))
+			{
+				return ParseSizeParam(ValueStr);
+			}
+			return int64(-1);
+		}();
+
+		if (DiskQuota > 0)
+		{
+			CacheConfig.DiskQuota = DiskQuota;
 		}
 
 		CacheConfig.RootDirectory = Private::GetInstallCacheDirectory();
@@ -830,7 +839,7 @@ FIoStatus FOnDemandIoStore::InitializeOnDemandInstallCache()
 
 #if !(UE_BUILD_SHIPPING|UE_BUILD_TEST)
 	FString ParamValue;
-	if (FParse::Value(FCommandLine::Get(), TEXT("Iad.Fill="), ParamValue))
+	if (FParse::Value(FCommandLine::Get(), TEXT("-Iad.Fill="), ParamValue))
 	{
 		ParamValue.TrimStartAndEndInline();
 		int64 FillSize = -1;
