Transform Constraints: remove useless child ticking

#rnx

diff --git a/Engine/Source/Runtime/Experimental/Animation/Constraints/Private/TransformConstraint.cpp b/Engine/Source/Runtime/Experimental/Animation/Constraints/Private/TransformConstraint.cpp
--- a/Engine/Source/Runtime/Experimental/Animation/Constraints/Private/TransformConstraint.cpp
+++ b/Engine/Source/Runtime/Experimental/Animation/Constraints/Private/TransformConstraint.cpp
@@ -611,22 +611,12 @@ void UTickableTransformConstraint::ResolveBoundObjects(FMovieSceneSequenceID Loc
 
 void UTickableTransformConstraint::Evaluate(bool bTickHandlesAlso) const
 {
-	if (bTickHandlesAlso)
+	if (IsFullyActive())
 	{
-		if (ParentTRSHandle)
+		if (bTickHandlesAlso)
 		{
 			ParentTRSHandle->TickTarget();
 		}
-
-		Super::Evaluate();
-		//todo test this more may be able to remove it
-		if (ChildTRSHandle)
-		{
-			ChildTRSHandle->TickTarget();
-		}
-	}
-	else
-	{
 		Super::Evaluate();
 	}
 }
