Fixed a few Chaos ISPC per-kernel variables not having or not using the preprocessor define for the default value.

diff --git a/Engine/Plugins/ChaosCloth/Source/ChaosCloth/Private/ChaosCloth/ChaosClothingSimulation.cpp b/Engine/Plugins/ChaosCloth/Source/ChaosCloth/Private/ChaosCloth/ChaosClothingSimulation.cpp
--- a/Engine/Plugins/ChaosCloth/Source/ChaosCloth/Private/ChaosCloth/ChaosClothingSimulation.cpp
+++ b/Engine/Plugins/ChaosCloth/Source/ChaosCloth/Private/ChaosCloth/ChaosClothingSimulation.cpp
@@ -49,7 +49,7 @@ static_assert(sizeof(ispc::FVector3f) == sizeof(FVector3f), "sizeof(ispc::FVecto
 static_assert(sizeof(ispc::FVector) == sizeof(Chaos::FVec3), "sizeof(ispc::FVector) != sizeof(Chaos::FVec3)");
 static_assert(sizeof(ispc::FTransform) == sizeof(Chaos::FRigidTransform3), "sizeof(ispc::FTransform) != sizeof(Chaos::FRigidTransform3)");
 
-bool bChaos_GetSimData_ISPC_Enabled = true;
+bool bChaos_GetSimData_ISPC_Enabled = CHAOS_GET_SIM_DATA_ISPC_ENABLED_DEFAULT;
 FAutoConsoleVariableRef CVarChaosGetSimDataISPCEnabled(TEXT("p.Chaos.GetSimData.ISPC"), bChaos_GetSimData_ISPC_Enabled, TEXT("Whether to use ISPC optimizations when getting simulation data"));
 #endif
 
diff --git a/Engine/Plugins/ChaosClothAsset/Source/ChaosClothAssetEngine/Private/ChaosClothAsset/ClothSimulationProxy.cpp b/Engine/Plugins/ChaosClothAsset/Source/ChaosClothAssetEngine/Private/ChaosClothAsset/ClothSimulationProxy.cpp
--- a/Engine/Plugins/ChaosClothAsset/Source/ChaosClothAssetEngine/Private/ChaosClothAsset/ClothSimulationProxy.cpp
+++ b/Engine/Plugins/ChaosClothAsset/Source/ChaosClothAssetEngine/Private/ChaosClothAsset/ClothSimulationProxy.cpp
@@ -37,7 +37,7 @@ namespace UE::Chaos::ClothAsset
 	static_assert(sizeof(ispc::FVector3f) == sizeof(FVector3f), "sizeof(ispc::FVector3f) != sizeof(FVector3f)");
 	static_assert(sizeof(ispc::FTransform) == sizeof(::Chaos::FRigidTransform3), "sizeof(ispc::FTransform) != sizeof(::Chaos::FRigidTransform3)");
 
-	bool bTransformClothSimulData_ISPC_Enabled = true;
+	bool bTransformClothSimulData_ISPC_Enabled = CHAOS_TRANSFORM_CLOTH_SIMUL_DATA_ISPC_ENABLED_DEFAULT;
 	FAutoConsoleVariableRef CVarTransformClothSimukDataISPCEnabled(TEXT("p.ChaosClothAsset.TransformClothSimulData.ISPC"), bTransformClothSimulData_ISPC_Enabled, TEXT("Whether to use ISPC optimizations when transforming simulation data back to reference bone space."));
 #endif
 
diff --git a/Engine/Source/Runtime/Experimental/Chaos/Private/Chaos/PBDAxialSpringConstraints.cpp b/Engine/Source/Runtime/Experimental/Chaos/Private/Chaos/PBDAxialSpringConstraints.cpp
--- a/Engine/Source/Runtime/Experimental/Chaos/Private/Chaos/PBDAxialSpringConstraints.cpp
+++ b/Engine/Source/Runtime/Experimental/Chaos/Private/Chaos/PBDAxialSpringConstraints.cpp
@@ -10,7 +10,7 @@
 DECLARE_CYCLE_STAT(TEXT("Chaos PBD Axial Spring Constraint"), STAT_PBD_AxialSpring, STATGROUP_Chaos);
 
 #if INTEL_ISPC && !UE_BUILD_SHIPPING
-bool bChaos_AxialSpring_ISPC_Enabled = true;
+bool bChaos_AxialSpring_ISPC_Enabled = CHAOS_AXIAL_SPRING_ISPC_ENABLED_DEFAULT;
 FAutoConsoleVariableRef CVarChaosAxialSpringISPCEnabled(TEXT("p.Chaos.AxialSpring.ISPC"), bChaos_AxialSpring_ISPC_Enabled, TEXT("Whether to use ISPC optimizations in AxialSpring constraints"));
 
 static_assert(sizeof(ispc::FVector4f) == sizeof(Chaos::Softs::FPAndInvM), "sizeof(ispc::FVector4f) != sizeof(Chaos::Softs::FPAndInvM");
diff --git a/Engine/Source/Runtime/Experimental/Chaos/Private/Chaos/PBDBendingConstraintsBase.cpp b/Engine/Source/Runtime/Experimental/Chaos/Private/Chaos/PBDBendingConstraintsBase.cpp
--- a/Engine/Source/Runtime/Experimental/Chaos/Private/Chaos/PBDBendingConstraintsBase.cpp
+++ b/Engine/Source/Runtime/Experimental/Chaos/Private/Chaos/PBDBendingConstraintsBase.cpp
@@ -5,7 +5,7 @@
 #endif
 
 #if INTEL_ISPC && !UE_BUILD_SHIPPING
-bool bChaos_Bending_ISPC_Enabled = true;
+bool bChaos_Bending_ISPC_Enabled = CHAOS_BENDING_ISPC_ENABLED_DEFAULT;
 FAutoConsoleVariableRef CVarChaosBendingISPCEnabled(TEXT("p.Chaos.Bending.ISPC"), bChaos_Bending_ISPC_Enabled, TEXT("Whether to use ISPC optimizations in Bending constraints"));
 
 static_assert(sizeof(ispc::FVector4f) == sizeof(Chaos::Softs::FPAndInvM), "sizeof(ispc::FVector4f) != sizeof(Chaos::Softs::FPAndInvM");
diff --git a/Engine/Source/Runtime/Experimental/Chaos/Private/Chaos/PBDCollisionSpringConstraintsBase.cpp b/Engine/Source/Runtime/Experimental/Chaos/Private/Chaos/PBDCollisionSpringConstraintsBase.cpp
--- a/Engine/Source/Runtime/Experimental/Chaos/Private/Chaos/PBDCollisionSpringConstraintsBase.cpp
+++ b/Engine/Source/Runtime/Experimental/Chaos/Private/Chaos/PBDCollisionSpringConstraintsBase.cpp
@@ -20,7 +20,7 @@ static_assert(sizeof(ispc::FVector3f) == sizeof(Chaos::Softs::FSolverVec3), "siz
 static_assert(sizeof(ispc::FVector4f) == sizeof(Chaos::Softs::FPAndInvM), "sizeof(ispc::FVector4f) != sizeof(Chaos::Softs::FPAndInvM");
 static_assert(sizeof(ispc::FIntVector) == sizeof(Chaos::TVec3<int32>), "sizeof(ispc::FIntVector) != sizeof(Chaos::TVec3<int32>");
 
-bool bChaos_CollisionSpring_ISPC_Enabled = true;
+bool bChaos_CollisionSpring_ISPC_Enabled = CHAOS_COLLISION_SPRING_ISPC_ENABLED_DEFAULT;
 FAutoConsoleVariableRef CVarChaosCollisionSpringsISPCEnabled(TEXT("p.Chaos.CollisionSpring.ISPC"), bChaos_CollisionSpring_ISPC_Enabled, TEXT("Whether to use ISPC optimizations in collision spring constraints"));
 #endif
 
diff --git a/Engine/Source/Runtime/Experimental/Chaos/Private/Chaos/PBDEvolution.cpp b/Engine/Source/Runtime/Experimental/Chaos/Private/Chaos/PBDEvolution.cpp
--- a/Engine/Source/Runtime/Experimental/Chaos/Private/Chaos/PBDEvolution.cpp
+++ b/Engine/Source/Runtime/Experimental/Chaos/Private/Chaos/PBDEvolution.cpp
@@ -40,7 +40,7 @@ TAutoConsoleVariable<bool> CVarChaosPBDEvolutionWriteCCDContacts(TEXT("p.Chaos.P
 TAutoConsoleVariable<bool> CVarChaosPBDEvolutionParallelIntegrate(TEXT("p.Chaos.PBDEvolution.ParalleIntegrate"), false, TEXT("Run the integration step in parallel for."), ECVF_Cheat);
 
 #if INTEL_ISPC && !UE_BUILD_SHIPPING
-bool bChaos_PostIterationUpdates_ISPC_Enabled = true;
+bool bChaos_PostIterationUpdates_ISPC_Enabled = CHAOS_POST_ITERATION_UPDATES_ISPC_ENABLED;
 FAutoConsoleVariableRef CVarChaosPostIterationUpdatesISPCEnabled(TEXT("p.Chaos.PostIterationUpdates.ISPC"), bChaos_PostIterationUpdates_ISPC_Enabled, TEXT("Whether to use ISPC optimizations in PBD Post iteration updates"));
 
 static_assert(sizeof(ispc::FVector3f) == sizeof(Chaos::Softs::FSolverVec3), "sizeof(ispc::FVector3f) != sizeof(Chaos::Softs::FSolverVec3");
diff --git a/Engine/Source/Runtime/Experimental/Chaos/Private/Chaos/PBDLongRangeConstraints.cpp b/Engine/Source/Runtime/Experimental/Chaos/Private/Chaos/PBDLongRangeConstraints.cpp
--- a/Engine/Source/Runtime/Experimental/Chaos/Private/Chaos/PBDLongRangeConstraints.cpp
+++ b/Engine/Source/Runtime/Experimental/Chaos/Private/Chaos/PBDLongRangeConstraints.cpp
@@ -9,7 +9,7 @@
 static_assert(sizeof(ispc::FVector4f) == sizeof(Chaos::Softs::FPAndInvM), "sizeof(ispc::FVector4f) != sizeof(Chaos::Softs::FPAndInvM)");
 static_assert(sizeof(ispc::FTether) == sizeof(Chaos::Softs::FPBDLongRangeConstraintsBase::FTether), "sizeof(ispc::FTether) != sizeof(Chaos::Softs::FPBDLongRangeConstraintsBase::FTether)");
 
-bool bChaos_LongRange_ISPC_Enabled = true;
+bool bChaos_LongRange_ISPC_Enabled = CHAOS_LONG_RANGE_ISPC_ENABLED_DEFAULT;
 FAutoConsoleVariableRef CVarChaosLongRangeISPCEnabled(TEXT("p.Chaos.LongRange.ISPC"), bChaos_LongRange_ISPC_Enabled, TEXT("Whether to use ISPC optimizations in long range constraints"));
 #endif
 
diff --git a/Engine/Source/Runtime/Experimental/Chaos/Private/Chaos/PBDSphericalConstraint.cpp b/Engine/Source/Runtime/Experimental/Chaos/Private/Chaos/PBDSphericalConstraint.cpp
--- a/Engine/Source/Runtime/Experimental/Chaos/Private/Chaos/PBDSphericalConstraint.cpp
+++ b/Engine/Source/Runtime/Experimental/Chaos/Private/Chaos/PBDSphericalConstraint.cpp
@@ -10,7 +10,7 @@
 #endif
 
 #if INTEL_ISPC && !UE_BUILD_SHIPPING
-bool bChaos_Spherical_ISPC_Enabled = true;
+bool bChaos_Spherical_ISPC_Enabled = CHAOS_SPHERICAL_ISPC_ENABLED_DEFAULT;
 FAutoConsoleVariableRef CVarChaosSphericalISPCEnabled(TEXT("p.Chaos.Spherical.ISPC"), bChaos_Spherical_ISPC_Enabled, TEXT("Whether to use ISPC optimizations in spherical constraints"));
 
 static_assert(sizeof(ispc::FVector4f) == sizeof(Chaos::Softs::FPAndInvM), "sizeof(ispc::FVector4f) != sizeof(Chaos::Softs::FPAndInvM)");
diff --git a/Engine/Source/Runtime/Experimental/Chaos/Private/Chaos/PBDSpringConstraints.cpp b/Engine/Source/Runtime/Experimental/Chaos/Private/Chaos/PBDSpringConstraints.cpp
--- a/Engine/Source/Runtime/Experimental/Chaos/Private/Chaos/PBDSpringConstraints.cpp
+++ b/Engine/Source/Runtime/Experimental/Chaos/Private/Chaos/PBDSpringConstraints.cpp
@@ -18,7 +18,7 @@ DECLARE_CYCLE_STAT(TEXT("Chaos PBD Spring Constraint"), STAT_PBD_Spring, STATGRO
 static_assert(sizeof(ispc::FVector4f) == sizeof(Chaos::Softs::FPAndInvM), "sizeof(ispc::FVector4f) != sizeof(Chaos::Softs::FPAndInvM)");
 static_assert(sizeof(ispc::FIntVector2) == sizeof(Chaos::TVec2<int32>), "sizeof(ispc::FIntVector2) != sizeof(Chaos::TVec2<int32>)");
 
-bool bChaos_Spring_ISPC_Enabled = true;
+bool bChaos_Spring_ISPC_Enabled = CHAOS_SPRING_ISPC_ENABLED_DEFAULT;
 FAutoConsoleVariableRef CVarChaosSpringISPCEnabled(TEXT("p.Chaos.Spring.ISPC"), bChaos_Spring_ISPC_Enabled, TEXT("Whether to use ISPC optimizations in Spring constraints"));
 #endif
 
diff --git a/Engine/Source/Runtime/Experimental/Chaos/Private/Chaos/SoftsMultiResConstraints.cpp b/Engine/Source/Runtime/Experimental/Chaos/Private/Chaos/SoftsMultiResConstraints.cpp
--- a/Engine/Source/Runtime/Experimental/Chaos/Private/Chaos/SoftsMultiResConstraints.cpp
+++ b/Engine/Source/Runtime/Experimental/Chaos/Private/Chaos/SoftsMultiResConstraints.cpp
@@ -19,7 +19,7 @@ static_assert(sizeof(ispc::FIntVector) == sizeof(Chaos::TVec3<int32>), "sizeof(i
 static_assert(sizeof(ispc::FVector4f) == sizeof(Chaos::TVec4<Chaos::Softs::FSolverReal>), "sizeof(ispc::FVector4f) == sizeof(Chaos::TVec4<Chaos::Softs::FSolverReal>)");
 static_assert(sizeof(ispc::FIntVector) == sizeof(Chaos::TVec3<int32>), "sizeof(ispc::FIntVector) != sizeof(Chaos::TVec3<int32>");
 
-bool bChaos_MultiRes_ISPC_Enabled = true;
+bool bChaos_MultiRes_ISPC_Enabled = CHAOS_MULTIRES_ISPC_ENABLED_DEFAULT;
 FAutoConsoleVariableRef CVarChaosMultiResISPCEnabled(TEXT("p.Chaos.MultiRes.ISPC"), bChaos_MultiRes_ISPC_Enabled, TEXT("Whether to use ISPC optimizations in MultiRes constraints"));
 bool bChaos_MultiRes_SparseWeightMap_Enabled = false;
 FAutoConsoleVariableRef CVarChaosMultiResSparseWeightMapEnabled(TEXT("p.Chaos.MultiRes.SparseWeightMap"), bChaos_MultiRes_SparseWeightMap_Enabled, TEXT("Exploit the sparse weight map structure and skip the particles with 0 stiffness at the beginning and at the end"));
diff --git a/Engine/Source/Runtime/Experimental/Chaos/Private/Chaos/VelocityField.cpp b/Engine/Source/Runtime/Experimental/Chaos/Private/Chaos/VelocityField.cpp
--- a/Engine/Source/Runtime/Experimental/Chaos/Private/Chaos/VelocityField.cpp
+++ b/Engine/Source/Runtime/Experimental/Chaos/Private/Chaos/VelocityField.cpp
@@ -12,7 +12,7 @@ static_assert(sizeof(ispc::FIntVector) == sizeof(Chaos::TVec3<int32>), "sizeof(i
 #endif
 
 #if INTEL_ISPC && !UE_BUILD_SHIPPING
-bool bChaos_VelocityField_ISPC_Enabled = true;
+bool bChaos_VelocityField_ISPC_Enabled = CHAOS_VELOCITY_FIELD_ISPC_ENABLED_DEFAULT;
 FAutoConsoleVariableRef CVarChaosVelocityFieldISPCEnabled(TEXT("p.Chaos.VelocityField.ISPC"), bChaos_VelocityField_ISPC_Enabled, TEXT("Whether to use ISPC optimizations in velocity field calculations"));
 #endif
 
diff --git a/Engine/Source/Runtime/Experimental/Chaos/Private/Chaos/XPBDBendingConstraints.cpp b/Engine/Source/Runtime/Experimental/Chaos/Private/Chaos/XPBDBendingConstraints.cpp
--- a/Engine/Source/Runtime/Experimental/Chaos/Private/Chaos/XPBDBendingConstraints.cpp
+++ b/Engine/Source/Runtime/Experimental/Chaos/Private/Chaos/XPBDBendingConstraints.cpp
@@ -10,7 +10,7 @@
 DECLARE_CYCLE_STAT(TEXT("Chaos XPBD Bending Constraint"), STAT_XPBD_Bending, STATGROUP_Chaos);
 
 #if INTEL_ISPC && !UE_BUILD_SHIPPING
-bool bChaos_XPBDBending_ISPC_Enabled = true;
+bool bChaos_XPBDBending_ISPC_Enabled = CHAOS_XPBD_BENDING_ISPC_ENABLED;
 FAutoConsoleVariableRef CVarChaosXPBDBendingISPCEnabled(TEXT("p.Chaos.XPBDBending.ISPC"), bChaos_XPBDBending_ISPC_Enabled, TEXT("Whether to use ISPC optimizations in XPBD Bending constraints"));
 
 static_assert(sizeof(ispc::FVector4f) == sizeof(Chaos::Softs::FPAndInvM), "sizeof(ispc::FVector4f) != sizeof(Chaos::Softs::FPAndInvM");
diff --git a/Engine/Source/Runtime/Experimental/Chaos/Private/Chaos/XPBDSpringConstraints.cpp b/Engine/Source/Runtime/Experimental/Chaos/Private/Chaos/XPBDSpringConstraints.cpp
--- a/Engine/Source/Runtime/Experimental/Chaos/Private/Chaos/XPBDSpringConstraints.cpp
+++ b/Engine/Source/Runtime/Experimental/Chaos/Private/Chaos/XPBDSpringConstraints.cpp
@@ -21,7 +21,7 @@ static_assert(sizeof(ispc::FVector4f) == sizeof(Chaos::Softs::FPAndInvM), "sizeo
 static_assert(sizeof(ispc::FVector3f) == sizeof(Chaos::Softs::FSolverVec3), "sizeof(ispc::FVector3f) == sizeof(Chaos::Softs::FSolverVec3)");
 static_assert(sizeof(ispc::FIntVector2) == sizeof(Chaos::TVec2<int32>), "sizeof(ispc::FIntVector2) != sizeof(Chaos::TVec2<int32>)");
 
-bool bChaos_XPBDSpring_ISPC_Enabled = true;
+bool bChaos_XPBDSpring_ISPC_Enabled = CHAOS_XPBD_SPRING_ISPC_ENABLED_DEFAULT;
 FAutoConsoleVariableRef CVarChaosXPBDSpringISPCEnabled(TEXT("p.Chaos.XPBDSpring.ISPC"), bChaos_XPBDSpring_ISPC_Enabled, TEXT("Whether to use ISPC optimizations in XPBD Spring constraints"));
 #endif
 
diff --git a/Engine/Source/Runtime/Experimental/Chaos/Private/Chaos/XPBDStretchBiasElementConstraints.cpp b/Engine/Source/Runtime/Experimental/Chaos/Private/Chaos/XPBDStretchBiasElementConstraints.cpp
--- a/Engine/Source/Runtime/Experimental/Chaos/Private/Chaos/XPBDStretchBiasElementConstraints.cpp
+++ b/Engine/Source/Runtime/Experimental/Chaos/Private/Chaos/XPBDStretchBiasElementConstraints.cpp
@@ -14,7 +14,7 @@
 DECLARE_CYCLE_STAT(TEXT("Chaos XPBD Stretch Bias Constraint"), STAT_XPBD_StretchBias, STATGROUP_Chaos);
 
 #if INTEL_ISPC && !UE_BUILD_SHIPPING
-bool bChaos_XPBDStretchBiasElement_ISPC_Enabled = true;
+bool bChaos_XPBDStretchBiasElement_ISPC_Enabled = CHAOS_XPBD_STRETCH_BIAS_ELEMENT_ISPC_ENABLED_DEFAULT;
 FAutoConsoleVariableRef CVarChaosXPBDStretchBiasISPCEnabled(TEXT("p.Chaos.XPBDStretchBias.ISPC"), bChaos_XPBDStretchBiasElement_ISPC_Enabled, TEXT("Whether to use ISPC optimizations in XPBD Stretch Bias constraints"));
 static_assert(sizeof(ispc::FVector4f) == sizeof(Chaos::Softs::FPAndInvM), "sizeof(ispc::FVector4f) != sizeof(Chaos::Softs::FPAndInvM");
 static_assert(sizeof(ispc::FIntVector) == sizeof(Chaos::TVec3<int32>), "sizeof(ispc::FIntVector) != sizeof(Chaos::TVec3<int32>");
diff --git a/Engine/Source/Runtime/Experimental/Chaos/Public/Chaos/PBDBendingConstraintsBase.h b/Engine/Source/Runtime/Experimental/Chaos/Public/Chaos/PBDBendingConstraintsBase.h
--- a/Engine/Source/Runtime/Experimental/Chaos/Public/Chaos/PBDBendingConstraintsBase.h
+++ b/Engine/Source/Runtime/Experimental/Chaos/Public/Chaos/PBDBendingConstraintsBase.h
@@ -400,11 +400,13 @@ protected:
 
 }  // End namespace Chaos::Softs
 
+#if !defined(CHAOS_BENDING_ISPC_ENABLED_DEFAULT)
+#define CHAOS_BENDING_ISPC_ENABLED_DEFAULT 1
+#endif
+
 // Support ISPC enable/disable in non-shipping builds
-#if !INTEL_ISPC
-const bool bChaos_Bending_ISPC_Enabled = false;
-#elif UE_BUILD_SHIPPING
-const bool bChaos_Bending_ISPC_Enabled = true;
+#if !INTEL_ISPC || UE_BUILD_SHIPPING
+static constexpr bool bChaos_Bending_ISPC_Enabled = INTEL_ISPC && CHAOS_BENDING_ISPC_ENABLED_DEFAULT;
 #else
 extern CHAOS_API bool bChaos_Bending_ISPC_Enabled;
 #endif
diff --git a/Engine/Source/Runtime/Experimental/Chaos/Public/Chaos/PBDCollisionSpringConstraintsBase.h b/Engine/Source/Runtime/Experimental/Chaos/Public/Chaos/PBDCollisionSpringConstraintsBase.h
--- a/Engine/Source/Runtime/Experimental/Chaos/Public/Chaos/PBDCollisionSpringConstraintsBase.h
+++ b/Engine/Source/Runtime/Experimental/Chaos/Public/Chaos/PBDCollisionSpringConstraintsBase.h
@@ -250,14 +250,15 @@ private:
 };
 }  // End namespace Chaos::Softs
 
+#if !defined(CHAOS_COLLISION_SPRING_ISPC_ENABLED_DEFAULT)
+#define CHAOS_COLLISION_SPRING_ISPC_ENABLED_DEFAULT 1
+#endif
+
 // Support ISPC enable/disable in non-shipping builds
-#if !INTEL_ISPC
-const bool bChaos_CollisionSpring_ISPC_Enabled = false;
-#elif UE_BUILD_SHIPPING
-const bool bChaos_CollisionSpring_ISPC_Enabled = true;
+#if !INTEL_ISPC || UE_BUILD_SHIPPING
+static constexpr bool bChaos_CollisionSpring_ISPC_Enabled = INTEL_ISPC && CHAOS_COLLISION_SPRING_ISPC_ENABLED_DEFAULT;
 #else
 extern CHAOS_API bool bChaos_CollisionSpring_ISPC_Enabled;
 #endif
 
-
 #endif
diff --git a/Engine/Source/Runtime/Experimental/Chaos/Public/Chaos/SoftsMultiResConstraints.h b/Engine/Source/Runtime/Experimental/Chaos/Public/Chaos/SoftsMultiResConstraints.h
--- a/Engine/Source/Runtime/Experimental/Chaos/Public/Chaos/SoftsMultiResConstraints.h
+++ b/Engine/Source/Runtime/Experimental/Chaos/Public/Chaos/SoftsMultiResConstraints.h
@@ -129,11 +129,13 @@ private:
 
 }  // End namespace Chaos::Softs
 
+#if !defined(CHAOS_MULTIRES_ISPC_ENABLED_DEFAULT)
+#define CHAOS_MULTIRES_ISPC_ENABLED_DEFAULT 1
+#endif
+
 // Support ISPC enable/disable in non-shipping builds
-#if !INTEL_ISPC
-const bool bChaos_MultiRes_ISPC_Enabled = false;
-#elif UE_BUILD_SHIPPING
-const bool bChaos_MultiRes_ISPC_Enabled = true;
+#if !INTEL_ISPC || UE_BUILD_SHIPPING
+static constexpr bool bChaos_MultiRes_ISPC_Enabled = INTEL_ISPC && CHAOS_MULTIRES_ISPC_ENABLED_DEFAULT;
 #else
 extern CHAOS_API bool bChaos_MultiRes_ISPC_Enabled;
 #endif
diff --git a/Engine/Source/Runtime/Experimental/Chaos/Public/Chaos/XPBDBendingConstraints.h b/Engine/Source/Runtime/Experimental/Chaos/Public/Chaos/XPBDBendingConstraints.h
--- a/Engine/Source/Runtime/Experimental/Chaos/Public/Chaos/XPBDBendingConstraints.h
+++ b/Engine/Source/Runtime/Experimental/Chaos/Public/Chaos/XPBDBendingConstraints.h
@@ -409,11 +409,13 @@ extern bool bChaos_XPBDBending_SplitLambdaDamping;
 
 }  // End namespace Chaos::Softs
 
+#if !defined(CHAOS_XPBD_BENDING_ISPC_ENABLED)
+#define CHAOS_XPBD_BENDING_ISPC_ENABLED 1
+#endif
+
 // Support ISPC enable/disable in non-shipping builds
-#if !INTEL_ISPC
-const bool bChaos_XPBDBending_ISPC_Enabled = false;
-#elif UE_BUILD_SHIPPING
-const bool bChaos_XPBDBending_ISPC_Enabled = true;
+#if !INTEL_ISPC || UE_BUILD_SHIPPING
+static constexpr bool bChaos_XPBDBending_ISPC_Enabled = INTEL_ISPC && CHAOS_XPBD_BENDING_ISPC_ENABLED;
 #else
 extern CHAOS_API bool bChaos_XPBDBending_ISPC_Enabled;
 #endif
diff --git a/Engine/Source/Runtime/Experimental/Chaos/Public/Chaos/XPBDSpringConstraints.h b/Engine/Source/Runtime/Experimental/Chaos/Public/Chaos/XPBDSpringConstraints.h
--- a/Engine/Source/Runtime/Experimental/Chaos/Public/Chaos/XPBDSpringConstraints.h
+++ b/Engine/Source/Runtime/Experimental/Chaos/Public/Chaos/XPBDSpringConstraints.h
@@ -398,11 +398,13 @@ private:
 
 }  // End namespace Chaos::Softs
 
+#if !defined(CHAOS_XPBD_SPRING_ISPC_ENABLED_DEFAULT)
+#define CHAOS_XPBD_SPRING_ISPC_ENABLED_DEFAULT 1
+#endif
+
 // Support ISPC enable/disable in non-shipping builds
-#if !INTEL_ISPC
-const bool bChaos_XPBDSpring_ISPC_Enabled = false;
-#elif UE_BUILD_SHIPPING
-const bool bChaos_XPBDSpring_ISPC_Enabled = true;
+#if !INTEL_ISPC || UE_BUILD_SHIPPING
+static constexpr bool bChaos_XPBDSpring_ISPC_Enabled = INTEL_ISPC && CHAOS_XPBD_SPRING_ISPC_ENABLED_DEFAULT;
 #else
 extern CHAOS_API bool bChaos_XPBDSpring_ISPC_Enabled;
 #endif
diff --git a/Engine/Source/Runtime/Experimental/Chaos/Public/Chaos/XPBDStretchBiasElementConstraints.h b/Engine/Source/Runtime/Experimental/Chaos/Public/Chaos/XPBDStretchBiasElementConstraints.h
--- a/Engine/Source/Runtime/Experimental/Chaos/Public/Chaos/XPBDStretchBiasElementConstraints.h
+++ b/Engine/Source/Runtime/Experimental/Chaos/Public/Chaos/XPBDStretchBiasElementConstraints.h
@@ -180,11 +180,13 @@ private:
 
 }  // End namespace Chaos::Softs
 
+#if !defined(CHAOS_XPBD_STRETCH_BIAS_ELEMENT_ISPC_ENABLED_DEFAULT)
+#define CHAOS_XPBD_STRETCH_BIAS_ELEMENT_ISPC_ENABLED_DEFAULT 1
+#endif
+
 // Support ISPC enable/disable in non-shipping builds
-#if !INTEL_ISPC
-const bool bChaos_XPBDStretchBiasElement_ISPC_Enabled = false;
-#elif UE_BUILD_SHIPPING
-const bool bChaos_XPBDStretchBiasElement_ISPC_Enabled = true;
+#if !INTEL_ISPC || UE_BUILD_SHIPPING
+static constexpr bool bChaos_XPBDStretchBiasElement_ISPC_Enabled = INTEL_ISPC && CHAOS_XPBD_STRETCH_BIAS_ELEMENT_ISPC_ENABLED_DEFAULT;
 #else
 extern CHAOS_API bool bChaos_XPBDStretchBiasElement_ISPC_Enabled;
 #endif
