Fix crash in editor where the cluster union manager tries to remove shapes from a cluster union that has no left particles
Also make sure the editor only trigger an ensure instead of a check

diff --git a/Engine/Source/Runtime/Experimental/Chaos/Public/Chaos/ClusterUnionManager.h b/Engine/Source/Runtime/Experimental/Chaos/Public/Chaos/ClusterUnionManager.h
--- a/Engine/Source/Runtime/Experimental/Chaos/Public/Chaos/ClusterUnionManager.h
+++ b/Engine/Source/Runtime/Experimental/Chaos/Public/Chaos/ClusterUnionManager.h
@@ -402,9 +402,26 @@ namespace Chaos
 		}
 
 		// Don't try to remove anything if the shape array is already empty 
-		if (!ClusterParticle->ShapesArray().IsEmpty())
+		if (!ClusterParticle->ShapesArray().IsEmpty() && !AllChildParticles.IsEmpty())
 		{
-			check(AllChildParticles.Num() == ClusterParticle->ShapesArray().Num());
+#if WITH_EDITOR
+			ensureMsgf(
+				AllChildParticles.Num() == ClusterParticle->ShapesArray().Num(),
+				TEXT("RemoveParticlesFromClusterUnionGeometry : More than one shape per child particle : ShapeArray=[%d] AllChildParticles=[%d] ShapeParticles=[%d]"),
+				ClusterParticle->ShapesArray().Num(),
+				AllChildParticles.Num(),
+				ShapeParticles.Num()
+			);
+#else
+			checkf(
+				AllChildParticles.Num() == ClusterParticle->ShapesArray().Num(),
+				TEXT("RemoveParticlesFromClusterUnionGeometry : More than one shape per child particle : ShapeArray=[%d] AllChildParticles=[%d] ShapeParticles=[%d]"),
+				ClusterParticle->ShapesArray().Num(),
+				AllChildParticles.Num(),
+				ShapeParticles.Num()
+			);
+#endif
+			
 
 			TArray<int32> ShapeIndicesToRemove;
 			ShapeIndicesToRemove.Reserve(ShapeParticles.Num());
@@ -424,7 +441,23 @@ namespace Chaos
 
 			RemoveArrayItemsAtSortedIndices(AllChildParticles, ShapeIndicesToRemove);
 
-			check(AllChildParticles.Num() == ClusterParticle->ShapesArray().Num());
+#if WITH_EDITOR
+			ensureMsgf(
+				AllChildParticles.Num() == ClusterParticle->ShapesArray().Num(),
+				TEXT("RemoveParticlesFromClusterUnionGeometry : More than one shape per child particle : ShapeArray=[%d] AllChildParticles=[%d] ShapeParticles=[%d]"),
+				ClusterParticle->ShapesArray().Num(),
+				AllChildParticles.Num(),
+				ShapeParticles.Num()
+			);
+#else
+			checkf(
+				AllChildParticles.Num() == ClusterParticle->ShapesArray().Num(),
+				TEXT("RemoveParticlesFromClusterUnionGeometry : More than one shape per child particle : ShapeArray=[%d] AllChildParticles=[%d] ShapeParticles=[%d]"),
+				ClusterParticle->ShapesArray().Num(),
+				AllChildParticles.Num(),
+				ShapeParticles.Num()
+			);
+#endif
 		}
 		// If we remove particles from the cluster union geometry then we need to switch the geometry back to a FImplicitObjectUnionClustered to avoid errors with empty unions.
 		if (ClusterParticle->ShapesArray().IsEmpty())
