Dataflow: handle case where a reroute with no input connected will fail to post a cache entry
- added support for null cache entry , that resolve to default values when queried for a type

diff --git a/Engine/Source/Runtime/Experimental/Dataflow/Core/Private/Dataflow/DataflowInputOutput.cpp b/Engine/Source/Runtime/Experimental/Dataflow/Core/Private/Dataflow/DataflowInputOutput.cpp
--- a/Engine/Source/Runtime/Experimental/Dataflow/Core/Private/Dataflow/DataflowInputOutput.cpp
+++ b/Engine/Source/Runtime/Experimental/Dataflow/Core/Private/Dataflow/DataflowInputOutput.cpp
@@ -307,6 +307,12 @@ void FDataflowOutput::ForwardInput(const FDataflowInput* Input, Dataflow::FConte
 				Context.SetDataReference(CacheKey(), Property, ConnectionOut->CacheKey());
 			}
 		}
+		else
+		{
+			// if there's no connection we make a invalid reference 
+			// so when the input is going to pull the cached value , it will return a default value instead
+			Context.SetNullData(CacheKey(), Property, GetOwningNodeGuid(), GetOwningNodeValueHash(), GetOwningNodeTimestamp());
+		}
 	}
 }
 
diff --git a/Engine/Source/Runtime/Experimental/Dataflow/Core/Public/Dataflow/DataflowNodeParameters.h b/Engine/Source/Runtime/Experimental/Dataflow/Core/Public/Dataflow/DataflowNodeParameters.h
--- a/Engine/Source/Runtime/Experimental/Dataflow/Core/Public/Dataflow/DataflowNodeParameters.h
+++ b/Engine/Source/Runtime/Experimental/Dataflow/Core/Public/Dataflow/DataflowNodeParameters.h
@@ -50,12 +50,16 @@ namespace Dataflow
 		TArray<FName> Outputs;
 	};
 
+	//--------------------------------------------------------------------
+	// base class for all context cache entries 
+	//--------------------------------------------------------------------
 	struct FContextCacheElementBase 
 	{
 		enum EType
 		{
 			CacheElementTyped,
-			CacheElementReference
+			CacheElementReference,
+			CacheElementNull
 		};
 
 		FContextCacheElementBase(EType CacheElementType, FGuid InNodeGuid = FGuid(), const FProperty* InProperty = nullptr, uint32 InNodeHash = 0, FTimestamp InTimestamp = FTimestamp::Invalid)
@@ -95,6 +99,9 @@ namespace Dataflow
 		FTimestamp Timestamp = FTimestamp::Invalid;
 	};
 
+	//--------------------------------------------------------------------
+	// Value storing context cache entry - strongly typed
+	//--------------------------------------------------------------------
 	template<class T>
 	struct FContextCacheElement : public FContextCacheElementBase 
 	{
@@ -114,6 +121,9 @@ namespace Dataflow
 		const FDataType Data;                        // Decaying T removes any reference and gets the correct underlying storage data type
 	};
 
+	//--------------------------------------------------------------------
+	// Reference to another context cache entry 
+	//--------------------------------------------------------------------
 	template<class T>
 	struct FContextCacheElementReference : public FContextCacheElementBase
 	{
@@ -130,6 +140,18 @@ namespace Dataflow
 		const FContextCacheKey DataKey; // this is a key to another cache element
 	};
 
+	//--------------------------------------------------------------------
+	// Null entry, this will always return a default value 
+	//--------------------------------------------------------------------
+	struct FContextCacheElementNull : public FContextCacheElementBase
+	{
+		FContextCacheElementNull(FGuid InNodeGuid, const FProperty* InProperty, FContextCacheKey InDataKey, uint32 InNodeHash, FTimestamp Timestamp)
+			: FContextCacheElementBase(EType::CacheElementNull, InNodeGuid, InProperty, InNodeHash, Timestamp)
+		{}
+
+		virtual TUniquePtr<FContextCacheElementBase> CreateReference(FContextCacheKey InReferenceDataKey) const override;
+	};
+
 	// cache element method implementation 
 	template<class T>
 	const T& FContextCacheElementBase::GetTypedData(FContext& Context, const FProperty* PropertyIn, const T& Default) const
@@ -143,6 +165,10 @@ namespace Dataflow
 		{
 			return static_cast<const FContextCacheElementReference<T>&>(*this).GetData(Context, PropertyIn, Default);
 		}
+		if (Type == EType::CacheElementNull)
+		{
+			return Default; 
+		}
 		check(false); // should never happen
 		return Default;
 	}
@@ -227,6 +253,15 @@ namespace Dataflow
 			}
 		}
 
+		// this is useful when there's a need to have to have cache entry but  the type is not known and there no connected output
+		// ( like reroute nodes with unconnected input for example ) 
+		// in that case posting an invalid reference, will allow the evaluatino to go through and the node reading it will get a default value instead
+		void SetNullData(FContextCacheKey InKey, const FProperty* InProperty, const FGuid& InNodeGuid, uint32 InNodeHash, const FTimestamp& InTimestamp)
+		{
+			TUniquePtr<FContextCacheElementNull> CacheNullElement = MakeUnique<FContextCacheElementNull>(InNodeGuid, InProperty, InKey, InNodeHash, InTimestamp);
+			SetDataImpl(InKey, MoveTemp(CacheNullElement));
+		}
+
 		virtual TUniquePtr<FContextCacheElementBase>* GetDataImpl(FContextCacheKey Key) = 0;
 
 		template<class T>
@@ -409,4 +444,10 @@ namespace Dataflow
 	{
 		return MakeUnique<FContextCacheElementReference<T>>(GetNodeGuid(), GetProperty(), InReferenceDataKey, GetNodeHash(), GetTimestamp());
 	}
+
+	inline TUniquePtr<FContextCacheElementBase> FContextCacheElementNull::CreateReference(FContextCacheKey InReferenceDataKey) const
+	{
+		// a null entry will always return a null entry as a reference
+		return MakeUnique<FContextCacheElementNull>(GetNodeGuid(), GetProperty(), InReferenceDataKey, GetNodeHash(), GetTimestamp());
+	}
 }
