Dataflow / Geometry collection : fix crash when geometry collection are empty
- this also allow terminal node to actually properly process empty collection

diff --git a/Engine/Plugins/Experimental/GeometryCollectionPlugin/Source/GeometryCollectionNodes/Private/Dataflow/GeometryCollectionAssetNodes.cpp b/Engine/Plugins/Experimental/GeometryCollectionPlugin/Source/GeometryCollectionNodes/Private/Dataflow/GeometryCollectionAssetNodes.cpp
--- a/Engine/Plugins/Experimental/GeometryCollectionPlugin/Source/GeometryCollectionNodes/Private/Dataflow/GeometryCollectionAssetNodes.cpp
+++ b/Engine/Plugins/Experimental/GeometryCollectionPlugin/Source/GeometryCollectionNodes/Private/Dataflow/GeometryCollectionAssetNodes.cpp
@@ -65,12 +65,6 @@ void FGeometryCollectionTerminalDataflowNode::SetAssetValue(TObjectPtr<UObject>
 			const FMaterialInstanceArray& InMaterialInstances = GetValue(Context, &MaterialInstances);
 			const FInstancedMeshesArray& InInstancedMeshes = GetValue(Context, &InstancedMeshes);
 
-			if (InCollection.NumElements(FGeometryCollection::TransformGroup) == 0)
-			{
-				UE_LOG(LogChaos, Warning, TEXT("Cannot set Geometry Collection asset: Geometry Collection must have at least one transform."));
-				return;
-			}
-
 			const bool bHasInternalMaterial = false; // with data flow there's no assumption of internal materials
 			if (InMaterialInstances.Num() > 0)
 			{
diff --git a/Engine/Source/Runtime/Experimental/GeometryCollectionEngine/Private/GeometryCollection/GeometryCollectionComponent.cpp b/Engine/Source/Runtime/Experimental/GeometryCollectionEngine/Private/GeometryCollection/GeometryCollectionComponent.cpp
--- a/Engine/Source/Runtime/Experimental/GeometryCollectionEngine/Private/GeometryCollection/GeometryCollectionComponent.cpp
+++ b/Engine/Source/Runtime/Experimental/GeometryCollectionEngine/Private/GeometryCollection/GeometryCollectionComponent.cpp
@@ -1282,7 +1282,7 @@ void UGeometryCollectionComponent::SetSimulatePhysics(bool bEnabled)
 	if (bEnabled && !PhysicsProxy && RestCollection)
 	{
 		EnsurePhysicsStateCreated();
-		if (!PhysicsProxy)
+		if (!PhysicsProxy && !RestCollection->IsEmpty())
 		{
 			RegisterAndInitializePhysicsProxy();
 		}
diff --git a/Engine/Source/Runtime/Experimental/GeometryCollectionEngine/Private/GeometryCollection/GeometryCollectionObject.cpp b/Engine/Source/Runtime/Experimental/GeometryCollectionEngine/Private/GeometryCollection/GeometryCollectionObject.cpp
--- a/Engine/Source/Runtime/Experimental/GeometryCollectionEngine/Private/GeometryCollection/GeometryCollectionObject.cpp
+++ b/Engine/Source/Runtime/Experimental/GeometryCollectionEngine/Private/GeometryCollection/GeometryCollectionObject.cpp
@@ -542,6 +542,11 @@ void UGeometryCollection::GetSharedSimulationParams(FSharedSimulationParameters&
 	OutParams.SizeSpecificData.Sort();	//can we do this at editor time on post edit change?
 }
 
+bool UGeometryCollection::IsEmpty() const
+{
+	return (NumElements(FGeometryCollection::TransformGroup) == 0);
+}
+
 void UGeometryCollection::Reset()
 {
 	if (GeometryCollection.IsValid())
@@ -557,10 +562,9 @@ void UGeometryCollection::Reset()
 
 void UGeometryCollection::ResetFrom(const FManagedArrayCollection& InCollection, const TArray<UMaterial*>& InMaterials, bool bHasInternalMaterials)
 {
+	Reset();
 	if (GeometryCollection.IsValid())
 	{
-		Reset();
-
 		InCollection.CopyTo(GeometryCollection.Get());
 
 		// todo(Chaos) : we could certainly run a "dependent attribute update method here instead of having to known about convex specifically 
@@ -573,9 +577,9 @@ void UGeometryCollection::ResetFrom(const FManagedArrayCollection& InCollection,
 
 void UGeometryCollection::ResetFrom(const FManagedArrayCollection& InCollection, const TArray<UMaterialInterface*>& InMaterialInstances, bool bHasInternalMaterials)
 {
+	Reset();
 	if (GeometryCollection.IsValid())
 	{
-		Reset();
 		InCollection.CopyTo(GeometryCollection.Get());
 		// todo(Chaos) : we could certainly run a "dependent attribute update method here instead of having to known about convex specifically 
 		UpdateConvexGeometryIfMissing();
@@ -604,9 +608,9 @@ int32 UGeometryCollection::AppendGeometry(const UGeometryCollection & Element, b
 }
 
 /** NumElements */
-int32 UGeometryCollection::NumElements(const FName & Group) const
+int32 UGeometryCollection::NumElements(const FName& Group) const
 {
-	return GeometryCollection->NumElements(Group);
+	return GeometryCollection? GeometryCollection->NumElements(Group): 0;
 }
 
 /** RemoveElements */
diff --git a/Engine/Source/Runtime/Experimental/GeometryCollectionEngine/Public/GeometryCollection/GeometryCollectionObject.h b/Engine/Source/Runtime/Experimental/GeometryCollectionEngine/Public/GeometryCollection/GeometryCollectionObject.h
--- a/Engine/Source/Runtime/Experimental/GeometryCollectionEngine/Public/GeometryCollection/GeometryCollectionObject.h
+++ b/Engine/Source/Runtime/Experimental/GeometryCollectionEngine/Public/GeometryCollection/GeometryCollectionObject.h
@@ -388,6 +388,9 @@ public:
 	TSharedPtr<FGeometryCollection, ESPMode::ThreadSafe>       GetGeometryCollection() { return GeometryCollection; }
 	const TSharedPtr<FGeometryCollection, ESPMode::ThreadSafe> GetGeometryCollection() const { return GeometryCollection; }
 
+	/** Return true if the collection is "Empty" ( has no transforms at all ) */
+	GEOMETRYCOLLECTIONENGINE_API bool IsEmpty() const;
+
 	/** Return collection to initial (ie. empty) state. */
 	GEOMETRYCOLLECTIONENGINE_API void Reset();
 
