Fix to the calculation that converts the network poll frequency into frames.

diff --git a/Engine/Source/Runtime/Experimental/Iris/Core/Private/Iris/ReplicationSystem/ObjectReplicationBridge.cpp b/Engine/Source/Runtime/Experimental/Iris/Core/Private/Iris/ReplicationSystem/ObjectReplicationBridge.cpp
--- a/Engine/Source/Runtime/Experimental/Iris/Core/Private/Iris/ReplicationSystem/ObjectReplicationBridge.cpp
+++ b/Engine/Source/Runtime/Experimental/Iris/Core/Private/Iris/ReplicationSystem/ObjectReplicationBridge.cpp
@@ -2157,7 +2157,12 @@ uint8 UObjectReplicationBridge::ConvertPollFrequencyIntoFrames(float PollFrequen
 		return 0U;
 	}
 
-	const uint32 FramesBetweenUpdatesForObject = static_cast<uint32>(MaxTickRate / FPlatformMath::Max(0.001f, PollFrequency));
+	uint32 FramesBetweenUpdatesForObject = static_cast<uint32>(MaxTickRate / FPlatformMath::Max(0.001f, PollFrequency));
+	if (FramesBetweenUpdatesForObject > 0)
+	{
+		FramesBetweenUpdatesForObject--;
+	}
+
 	return static_cast<uint8>(FMath::Clamp<uint32>(FramesBetweenUpdatesForObject, 0U, UE::Net::Private::FObjectPollFrequencyLimiter::GetMaxPollingFrames()));
 }
 
