Fix the transactionally safe physics lock to allow creation/deletion within transactions too.

diff --git a/Engine/Source/Runtime/Experimental/Chaos/Private/Tests/AutoRTFMPhysSceneLockTests.cpp b/Engine/Source/Runtime/Experimental/Chaos/Private/Tests/AutoRTFMPhysSceneLockTests.cpp
new file mode 100644
--- /dev/null
+++ b/Engine/Source/Runtime/Experimental/Chaos/Private/Tests/AutoRTFMPhysSceneLockTests.cpp
@@ -0,0 +1,137 @@
+// Copyright Epic Games, Inc. All Rights Reserved.
+
+#include "Misc/AutomationTest.h"
+#include "Misc/App.h"
+#include "AutoRTFM/AutoRTFM.h"
+#include "Framework/Threading.h"
+#include "Templates/UniquePtr.h"
+
+#if WITH_DEV_AUTOMATION_TESTS
+
+#define CHECK_EQ(A, B) \
+	do { TestEqual(TEXT(__FILE__ ":" UE_STRINGIZE(__LINE__) ": TestEqual(" #A ", " #B ")"), (A), (B)); } while (0)
+
+IMPLEMENT_SIMPLE_AUTOMATION_TEST(FAutoRTFMChaosPhysSceneLock, "AutoRTFM + FPhysSceneLock", \
+	                             EAutomationTestFlags::EngineFilter | EAutomationTestFlags::ClientContext | \
+								 EAutomationTestFlags::ServerContext | EAutomationTestFlags::CommandletContext)
+
+bool FAutoRTFMChaosPhysSceneLock::RunTest(const FString& Parameters)
+{
+	if (!AutoRTFM::ForTheRuntime::IsAutoRTFMRuntimeEnabled())
+	{
+		ExecutionInfo.AddEvent(FAutomationEvent(EAutomationEventType::Info, 
+								                TEXT("SKIPPED 'FAutoRTFMChaosPhysSceneLock' test. AutoRTFM disabled.")));
+		return true;
+	}
+
+	// It should be safe to declare when not used.
+	AutoRTFM::Transact([&]
+		{
+			Chaos::FPhysSceneLock CreateInTransaction;
+		});
+
+	// Read lock then abort is fine.
+	AutoRTFM::Transact([&]
+		{
+			Chaos::FPhysSceneLock CreateInTransaction;
+			CreateInTransaction.ReadLock();
+			AutoRTFM::AbortTransaction();
+		});
+
+	// Read lock/unlock pair then abort is fine.
+	AutoRTFM::Transact([&]
+		{
+			Chaos::FPhysSceneLock CreateInTransaction;
+
+			CreateInTransaction.ReadLock();
+			CreateInTransaction.ReadUnlock();
+			AutoRTFM::AbortTransaction();
+		});
+
+	// Read lock/unlock pair fine.
+	AutoRTFM::Transact([&]
+		{
+			Chaos::FPhysSceneLock CreateInTransaction;
+
+			CreateInTransaction.ReadLock();
+			CreateInTransaction.ReadUnlock();
+		});
+
+	// Write lock then abort is fine.
+	AutoRTFM::Transact([&]
+		{
+			Chaos::FPhysSceneLock CreateInTransaction;
+			CreateInTransaction.WriteLock();
+			AutoRTFM::AbortTransaction();
+		});
+
+	// Write lock/unlock pair then abort is fine.
+	AutoRTFM::Transact([&]
+		{
+			Chaos::FPhysSceneLock CreateInTransaction;
+
+			CreateInTransaction.WriteLock();
+			CreateInTransaction.WriteUnlock();
+			AutoRTFM::AbortTransaction();
+		});
+
+	// Write lock/unlock pair fine.
+	AutoRTFM::Transact([&]
+		{
+			Chaos::FPhysSceneLock CreateInTransaction;
+
+			CreateInTransaction.WriteLock();
+			CreateInTransaction.WriteUnlock();
+		});
+
+	// Now do the same but outwith a transaction.
+	Chaos::FPhysSceneLock CreateOutwithTransaction;
+
+	// Read lock then abort is fine.
+
+	AutoRTFM::Transact([&]
+		{
+			CreateOutwithTransaction.ReadLock();
+			AutoRTFM::AbortTransaction();
+		});
+
+	// Read lock/unlock pair then abort is fine.
+	AutoRTFM::Transact([&]
+		{
+			CreateOutwithTransaction.ReadLock();
+			CreateOutwithTransaction.ReadUnlock();
+			AutoRTFM::AbortTransaction();
+		});
+
+	// Read lock/unlock pair fine.
+	AutoRTFM::Transact([&]
+		{
+			CreateOutwithTransaction.ReadLock();
+			CreateOutwithTransaction.ReadUnlock();
+		});
+
+	// Write lock then abort is fine.
+	AutoRTFM::Transact([&]
+		{
+			CreateOutwithTransaction.WriteLock();
+			AutoRTFM::AbortTransaction();
+		});
+
+	// Write lock/unlock pair then abort is fine.
+	AutoRTFM::Transact([&]
+		{
+			CreateOutwithTransaction.WriteLock();
+			CreateOutwithTransaction.WriteUnlock();
+			AutoRTFM::AbortTransaction();
+		});
+
+	AutoRTFM::Transact([&]
+		{
+			CreateOutwithTransaction.WriteLock();
+			CreateOutwithTransaction.WriteUnlock();
+		});
+
+	return true;
+}
+
+#endif  // WITH_DEV_AUTOMATION_TESTS
diff --git a/Engine/Source/Runtime/Experimental/Chaos/Public/Framework/Threading.h b/Engine/Source/Runtime/Experimental/Chaos/Public/Framework/Threading.h
--- a/Engine/Source/Runtime/Experimental/Chaos/Public/Framework/Threading.h
+++ b/Engine/Source/Runtime/Experimental/Chaos/Public/Framework/Threading.h
@@ -12,6 +12,7 @@
 #include <atomic>
 #include "HAL/CriticalSection.h"
 #include "AutoRTFM/AutoRTFM.h"
+#include "Templates/SharedPointer.h"
 
 #ifndef PHYSICS_THREAD_CONTEXT
 	#if (!UE_BUILD_SHIPPING && !UE_BUILD_TEST)
@@ -973,6 +974,47 @@ FORCEINLINE void EnsureIsInGameThreadContext()
 	//                                    effects of the open-threads writes!
 	struct FPhysSceneLockTransactionallySafe final
 	{
+		// Always open because the constructor arguments will create the underlying lock.
+		UE_AUTORTFM_ALWAYS_OPEN
+		FPhysSceneLockTransactionallySafe() : State(MakeShared<FState>())
+		{
+			if (AutoRTFM::IsTransactional())
+			{
+				const AutoRTFM::EContextStatus Status = AutoRTFM::Close([this]
+					{
+						AutoRTFM::PushOnAbortHandler(this, [this]
+							{
+								this->~FPhysSceneLockTransactionallySafe();
+							});
+					});
+
+				ensure(AutoRTFM::EContextStatus::OnTrack == Status);
+			}
+		}
+
+		~FPhysSceneLockTransactionallySafe()
+		{
+			if (AutoRTFM::IsTransactional())
+			{
+				const AutoRTFM::EContextStatus Status = AutoRTFM::Close([this]
+					{
+						AutoRTFM::PopOnAbortHandler(this);
+
+						// We explicitly copy the state here for the case that `this` was stack
+						// allocated and has already died before the on-commit is hit.
+						AutoRTFM::OnCommit([State = this->State]
+							{
+								ensure(0 == State->TransactionalLockCount);
+							});
+					});
+
+				ensure(AutoRTFM::EContextStatus::OnTrack == Status);
+			}
+
+			// As the State was constructed in the open, it must be released in the open.
+			AutoRTFM::Open([&] { State = nullptr; });
+		}
+
 		void ReadLock()
 		{
 			if (AutoRTFM::IsTransactional() || AutoRTFM::IsCommittingOrAborting())
@@ -982,8 +1024,8 @@ FORCEINLINE void EnsureIsInGameThreadContext()
 			}
 			else
 			{
-				Lock.ReadLock();
-				ensure(0 == TransactionalLockCount);
+				State->Lock.ReadLock();
+				ensure(0 == State->TransactionalLockCount);
 			}
 		}
 
@@ -996,8 +1038,8 @@ FORCEINLINE void EnsureIsInGameThreadContext()
 			}
 			else
 			{
-				ensure(0 == TransactionalLockCount);
-				Lock.ReadUnlock();
+				ensure(0 == State->TransactionalLockCount);
+				State->Lock.ReadUnlock();
 			}
 		}
 
@@ -1005,33 +1047,36 @@ FORCEINLINE void EnsureIsInGameThreadContext()
 		{
 			if (AutoRTFM::IsTransactional() || AutoRTFM::IsCommittingOrAborting())
 			{
-				AutoRTFM::Open([&]
+				AutoRTFM::Open([this]
 					{
 						// The transactional system which can increment TransactionalLockCount
 						// is always single-threaded, thus this is safe to check without atomicity.
-						if (0 == TransactionalLockCount)
+						if (0 == State->TransactionalLockCount)
 						{
-							Lock.WriteLock();
+							State->Lock.WriteLock();
 						}
 
-						TransactionalLockCount += 1;
+						State->TransactionalLockCount += 1;
 					});
 
-				AutoRTFM::OnAbort([this]
+				// We explicitly copy the state here for the case that `this` was stack
+				// allocated and has already died before the on-abort is hit.
+				AutoRTFM::OnAbort([State = this->State]
 					{
-						ensure(0 != TransactionalLockCount);
-						TransactionalLockCount -= 1;
+						ensure(0 != State->TransactionalLockCount);
+
+						State->TransactionalLockCount -= 1;
 
-						if (0 == TransactionalLockCount)
+						if (0 == State->TransactionalLockCount)
 						{
-							Lock.WriteUnlock();
+							State->Lock.WriteUnlock();
 						}
 					});
 			}
 			else
 			{
-				Lock.WriteLock();
-				ensure(0 == TransactionalLockCount);
+				State->Lock.WriteLock();
+				ensure(0 == State->TransactionalLockCount);
 			}
 		}
 
@@ -1039,27 +1084,37 @@ FORCEINLINE void EnsureIsInGameThreadContext()
 		{
 			if (AutoRTFM::IsTransactional() || AutoRTFM::IsCommittingOrAborting())
 			{
-				AutoRTFM::OnCommit([this]
+				// We explicitly copy the state here for the case that `this` was stack
+				// allocated and has already died before the on-commit is hit.
+				AutoRTFM::OnCommit([State = this->State]
 					{
-						ensure(0 != TransactionalLockCount);
-						TransactionalLockCount -= 1;
+						ensure(0 != State->TransactionalLockCount);
 
-						if (0 == TransactionalLockCount)
+						State->TransactionalLockCount -= 1;
+
+						if (0 == State->TransactionalLockCount)
 						{
-							Lock.WriteUnlock();
+							State->Lock.WriteUnlock();
 						}
 					});
 			}
 			else
 			{
-				ensure(0 == TransactionalLockCount);
-				Lock.WriteUnlock();
+				ensure(0 == State->TransactionalLockCount);
+				State->Lock.WriteUnlock();
 			}
 		}
 
 	private:
-		FPhysSceneLockNonTransactional Lock;
-		uint32 TransactionalLockCount;
+		UE_NONCOPYABLE(FPhysSceneLockTransactionallySafe)
+
+		struct FState final
+		{
+			FPhysSceneLockNonTransactional Lock;
+			uint32 TransactionalLockCount = 0;
+		};
+
+		TSharedPtr<FState> State;
 	};
 
 	using FPhysSceneLock = FPhysSceneLockTransactionallySafe;
