[AutoRTFM] Undefine CHECK_EQ at the end of test cpps to allow unity builds to succeed.

We use CHECK_EQ in multiple tests; the compiler will complain if it is #defined more than once, which may happen if the builder fuses test files.

diff --git a/Engine/Source/Runtime/Core/Tests/Containers/AutoRTFMLockFreeFixedSizeAllocatorTest.cpp b/Engine/Source/Runtime/Core/Tests/Containers/AutoRTFMLockFreeFixedSizeAllocatorTest.cpp
--- a/Engine/Source/Runtime/Core/Tests/Containers/AutoRTFMLockFreeFixedSizeAllocatorTest.cpp
+++ b/Engine/Source/Runtime/Core/Tests/Containers/AutoRTFMLockFreeFixedSizeAllocatorTest.cpp
@@ -132,4 +132,6 @@ bool FAutoRTFMLockFreeFixedSizeAllocator::RunTest(const FString& Parameters)
 	return true;
 }
 
+#undef CHECK_EQ
+
 #endif  // WITH_DEV_AUTOMATION_TESTS
diff --git a/Engine/Source/Runtime/Core/Tests/HAL/AutoRTFMTransactionallySafeCounterTest.cpp b/Engine/Source/Runtime/Core/Tests/HAL/AutoRTFMTransactionallySafeCounterTest.cpp
--- a/Engine/Source/Runtime/Core/Tests/HAL/AutoRTFMTransactionallySafeCounterTest.cpp
+++ b/Engine/Source/Runtime/Core/Tests/HAL/AutoRTFMTransactionallySafeCounterTest.cpp
@@ -148,4 +148,6 @@ bool FAutoRTFMTransactionallySafeCounter::RunTest(const FString& Parameters)
 	return true;
 }
 
+#undef CHECK_EQ
+
 #endif  // WITH_DEV_AUTOMATION_TESTS
diff --git a/Engine/Source/Runtime/Experimental/Chaos/Private/Tests/AutoRTFMRefCountedObjectTests.cpp b/Engine/Source/Runtime/Experimental/Chaos/Private/Tests/AutoRTFMRefCountedObjectTests.cpp
--- a/Engine/Source/Runtime/Experimental/Chaos/Private/Tests/AutoRTFMRefCountedObjectTests.cpp
+++ b/Engine/Source/Runtime/Experimental/Chaos/Private/Tests/AutoRTFMRefCountedObjectTests.cpp
@@ -107,4 +107,6 @@ bool FAutoRTFMChaosRefCountedObject::RunTest(const FString& Parameters)
 	return true;
 }
 
+#undef CHECK_EQ
+
 #endif  // WITH_DEV_AUTOMATION_TESTS
