On Demand IOStore Cache
- Move FCas CurrentBlock to FOnDemandInstallCache
- CurrentBlock is now atomic since it may be read from other threads during the OnDemandIOStore tick.

diff --git a/Engine/Source/Runtime/Experimental/IoStore/OnDemand/Private/OnDemandInstallCache.cpp b/Engine/Source/Runtime/Experimental/IoStore/OnDemand/Private/OnDemandInstallCache.cpp
--- a/Engine/Source/Runtime/Experimental/IoStore/OnDemand/Private/OnDemandInstallCache.cpp
+++ b/Engine/Source/Runtime/Experimental/IoStore/OnDemand/Private/OnDemandInstallCache.cpp
@@ -117,7 +117,6 @@ struct FCas
 	using FLastAccess		= TMap<FCasBlockId, int64>;
 
 	FIoStatus				Initialize(FStringView Directory);
-	FCasLocation			FindChunk(const FIoHash& Hash, bool& bIsLocationInCurrentBlock) const;
 	FCasLocation			FindChunk(const FIoHash& Hash) const;
 	FCasBlockId				CreateBlock();
 	FIoStatus				DeleteBlock(FCasBlockId BlockId, TArray<FCasAddr>& OutAddrs);
@@ -136,8 +135,7 @@ struct FCas
 	TSet<FCasBlockId>	BlockIds;
 	FLastAccess			LastAccess;
 	FReadHandles		ReadHandles;	
-	uint32				MaxBlockSize = 32 << 20; //TODO: Make configurable
-	FCasBlockId			CurrentBlock;
+	uint32				MaxBlockSize = 32 << 20; //TODO: Make configurable	
 	mutable UE::FMutex	Mutex;
 };
 
@@ -149,7 +147,6 @@ FIoStatus FCas::Initialize(FStringView Directory)
 	Lookup.Empty();
 	BlockIds.Empty();
 	LastAccess.Empty();
-	CurrentBlock = FCasBlockId::Invalid;
 
 	TStringBuilder<256> Path;
 	FPathViews::Append(Path, RootDirectory, TEXT("blocks"));
@@ -171,7 +168,7 @@ FIoStatus FCas::Initialize(FStringView Directory)
 	return EIoErrorCode::Ok;
 };
 
-FCasLocation FCas::FindChunk(const FIoHash& Hash, bool& bIsLocationInCurrentBlock) const
+FCasLocation FCas::FindChunk(const FIoHash& Hash) const
 {
 	const FCasAddr* Addr	= reinterpret_cast<const FCasAddr*>(&Hash);
 	const uint32 TypeHash	= GetTypeHash(*Addr);
@@ -179,7 +176,6 @@ FCasLocation FCas::FindChunk(const FIoHash& Hash, bool& bIsLocationInCurrentBloc
 		UE::TUniqueLock Lock(Mutex);
 		if (const FCasLocation* Loc = Lookup.FindByHash(TypeHash, *Addr))
 		{
-			bIsLocationInCurrentBlock = Loc->BlockId == CurrentBlock;
 			return *Loc;
 		}
 	}
@@ -187,12 +183,6 @@ FCasLocation FCas::FindChunk(const FIoHash& Hash, bool& bIsLocationInCurrentBloc
 	return FCasLocation{};
 }
 
-FCasLocation FCas::FindChunk(const FIoHash& Hash) const
-{
-	bool bTmp = false;
-	return FindChunk(Hash, bTmp);
-}
-
 FCasBlockId FCas::CreateBlock()
 {
 	IPlatformFile&	Ipf = FPlatformFileManager::Get().GetPlatformFile();
@@ -974,14 +964,15 @@ private:
 	FIoStatus									FlushPendingChunks(FPendingChunks& Block);
 	FString										GetJournalFilename() const { return CacheDirectory / TEXT("cas.jrn"); }
 
-	FOnDemandIoStore&		IoStore;
-	FString					CacheDirectory;
-	FCas					Cas;
-	FUniquePendingChunks	PendingChunks;
-	FSharedBackendContext	BackendContext;
-	FIoRequestList			CompletedRequests;
-	UE::FMutex				Mutex;
-	uint64					MaxCacheSize;
+	FOnDemandIoStore&			IoStore;
+	FString						CacheDirectory;
+	FCas						Cas;
+	std::atomic<FCasBlockId>	CurrentBlock{ FCasBlockId::Invalid };
+	FUniquePendingChunks		PendingChunks;
+	FSharedBackendContext		BackendContext;
+	FIoRequestList				CompletedRequests;
+	UE::FMutex					Mutex;
+	uint64						MaxCacheSize{ 0 };
 };
 
 ///////////////////////////////////////////////////////////////////////////////
@@ -1024,7 +1015,7 @@ FOnDemandInstallCache::FOnDemandInstallCache(const FOnDemandInstallCacheConfig&
 		case FCasJournal::FEntry::EType::BlockCreated:
 		{
 			const FCasJournal::FEntry::FBlockOperation& Op = JournalEntry.BlockOperation;
-			Cas.CurrentBlock = Op.BlockId;
+			CurrentBlock = Op.BlockId;
 			Cas.BlockIds.Add(Op.BlockId);
 			break;
 		}
@@ -1032,10 +1023,8 @@ FOnDemandInstallCache::FOnDemandInstallCache(const FOnDemandInstallCacheConfig&
 		{
 			const FCasJournal::FEntry::FBlockOperation& Op = JournalEntry.BlockOperation;
 			Cas.BlockIds.Remove(Op.BlockId);
-			if (Cas.CurrentBlock == Op.BlockId)
-			{
-				Cas.CurrentBlock = FCasBlockId::Invalid;
-			}
+			FCasBlockId MaybeCurrentBlock = Op.BlockId;
+			CurrentBlock.compare_exchange_strong(MaybeCurrentBlock, FCasBlockId::Invalid);
 			break;
 		}
 		case FCasJournal::FEntry::EType::BlockAccess:
@@ -1195,8 +1184,7 @@ bool FOnDemandInstallCache::Resolve(FIoRequestImpl* Request)
 		return false;
 	}
 
-	bool bIsLocationInCurrentBlock = false;
-	const FCasLocation CasLoc = Cas.FindChunk(ChunkInfo.Hash(), bIsLocationInCurrentBlock);
+	const FCasLocation CasLoc = Cas.FindChunk(ChunkInfo.Hash());
 	if (CasLoc.IsValid() == false)
 	{
 		return false;
@@ -1222,6 +1210,7 @@ bool FOnDemandInstallCache::Resolve(FIoRequestImpl* Request)
 	Cas.TrackAccess(CasLoc.BlockId);
 
 	// Use synchronous file read API when reading from and writing to the same cache block concurrently.
+	const bool bIsLocationInCurrentBlock = CasLoc.BlockId == CurrentBlock;
 	const bool bSyncRead = bIsLocationInCurrentBlock || CVars::bForceSyncIO;
 	if (bSyncRead)
 	{
@@ -1503,10 +1492,9 @@ FIoStatus FOnDemandInstallCache::Purge(const FCasBlockInfoMap& BlockInfo, const
 			return Status;
 		}
 
-		if (Cas.CurrentBlock == BlockId)
-		{
-			Cas.CurrentBlock = FCasBlockId::Invalid;
-		}
+		// This should be the only thread writing to CurrentBlock
+		FCasBlockId MaybeCurrentBlock = BlockId;
+		CurrentBlock.compare_exchange_strong(MaybeCurrentBlock, FCasBlockId::Invalid);
 
 		OutTotalPurgedBytes += Info.FileSize;
 
@@ -1573,23 +1561,27 @@ FIoStatus FOnDemandInstallCache::FlushPendingChunks(FPendingChunks& Chunks)
 {
 	ON_SCOPE_EXIT { Chunks.Reset(); };
 
+	// This should be the only thread writing to CurrentBlock
+	FCasBlockId CurrentBlockId = CurrentBlock;
+
 	while (Chunks.IsEmpty() == false)
 	{
 		FCasJournal::FTransaction Transaction = FCasJournal::Begin(GetJournalFilename());
-
-		if (Cas.CurrentBlock.IsValid() == false)
+		
+		if (CurrentBlockId.IsValid() == false)
 		{
-			Cas.CurrentBlock = Cas.CreateBlock();
-			ensure(Cas.CurrentBlock.IsValid());
-			Transaction.BlockCreated(Cas.CurrentBlock);
+			CurrentBlockId = Cas.CreateBlock();
+			ensure(CurrentBlockId.IsValid());
+			CurrentBlock = CurrentBlockId;
+			Transaction.BlockCreated(CurrentBlockId);
 		}
 
-		TUniquePtr<IFileHandle>	CasFileHandle = Cas.OpenWrite(Cas.CurrentBlock);
+		TUniquePtr<IFileHandle>	CasFileHandle = Cas.OpenWrite(CurrentBlockId);
 		if (CasFileHandle.IsValid() == false)
 		{
 			return FIoStatusBuilder(EIoErrorCode::FileOpenFailed)
 				<< TEXT("Failed to open cache block file '")
-				<< Cas.GetBlockFilename(Cas.CurrentBlock)
+				<< Cas.GetBlockFilename(CurrentBlockId)
 				<< TEXT("'");
 		}
 
@@ -1613,14 +1605,14 @@ FIoStatus FOnDemandInstallCache::FlushPendingChunks(FPendingChunks& Chunks)
 		if (Ar.Tell() > 0)
 		{
 			UE_LOG(LogIoStoreOnDemand, Log, TEXT("Writing %.2lf MiB to CAS block %u"),
-				ToMiB(Ar.Tell()), Cas.CurrentBlock.Id);
+				ToMiB(Ar.Tell()), CurrentBlockId.Id);
 
 			if (CasFileHandle->Write(Ar.GetData(), Ar.Tell()) == false)
 			{
 				return FIoStatusBuilder(EIoErrorCode::WriteError)
 					<< TEXT("Failed to serialize chunks to cache block");
 			}
-			Cas.TrackAccess(Cas.CurrentBlock);
+			Cas.TrackAccess(CurrentBlockId);
 
 			if (CasFileHandle->Flush() == false)
 			{
@@ -1629,14 +1621,14 @@ FIoStatus FOnDemandInstallCache::FlushPendingChunks(FPendingChunks& Chunks)
 			}
 
 			check(ChunkHashes.Num() == Offsets.Num());
-			check(Cas.CurrentBlock.IsValid());
+			check(CurrentBlockId.IsValid());
 			for (int32 Idx = 0, Count = Offsets.Num(); Idx < Count; ++Idx)
 			{
 				const FCasAddr	CasAddr = FCasAddr::From(ChunkHashes[Idx]);
 				const uint32	ChunkOffset = IntCastChecked<uint32>(Offsets[Idx]);
 
 				FCasLocation& Loc = Cas.Lookup.FindOrAdd(CasAddr);
-				Loc.BlockId	= Cas.CurrentBlock;
+				Loc.BlockId	= CurrentBlockId;
 				Loc.BlockOffset	= ChunkOffset;
 				Transaction.ChunkLocation(Loc, CasAddr);
 			}
@@ -1649,7 +1641,7 @@ FIoStatus FOnDemandInstallCache::FlushPendingChunks(FPendingChunks& Chunks)
 
 		if (Chunks.IsEmpty() == false)
 		{
-			Cas.CurrentBlock = FCasBlockId::Invalid;
+			CurrentBlockId = FCasBlockId::Invalid;
 		}
 	}
 
