Dataflow : add mesh transform  to StaticMeshToCollection node

diff --git a/Engine/Plugins/Experimental/GeometryCollectionPlugin/Source/GeometryCollectionNodes/Private/Dataflow/GeometryCollectionStaticMeshToCollectionNode.cpp b/Engine/Plugins/Experimental/GeometryCollectionPlugin/Source/GeometryCollectionNodes/Private/Dataflow/GeometryCollectionStaticMeshToCollectionNode.cpp
--- a/Engine/Plugins/Experimental/GeometryCollectionPlugin/Source/GeometryCollectionNodes/Private/Dataflow/GeometryCollectionStaticMeshToCollectionNode.cpp
+++ b/Engine/Plugins/Experimental/GeometryCollectionPlugin/Source/GeometryCollectionNodes/Private/Dataflow/GeometryCollectionStaticMeshToCollectionNode.cpp
@@ -17,6 +17,7 @@ FStaticMeshToCollectionDataflowNode::FStaticMeshToCollectionDataflowNode(const D
 	: FDataflowNode(InParam, InGuid)
 {
 	RegisterInputConnection(&StaticMesh);
+	RegisterInputConnection(&MeshTransform);
 	RegisterOutputConnection(&Collection);
 	RegisterOutputConnection(&Materials);
 	RegisterOutputConnection(&MaterialInstances);
@@ -32,9 +33,10 @@ void FStaticMeshToCollectionDataflowNode::Evaluate(Dataflow::FContext& Context,
 	TArray<FGeometryCollectionAutoInstanceMesh> OutInstancedMeshes;
 
 	TObjectPtr<UStaticMesh> StaticMeshVal = GetValue(Context, &StaticMesh, StaticMesh);
+	const FTransform& InMeshTransform = GetValue(Context, &MeshTransform, MeshTransform);
 	if (StaticMeshVal)
 	{
-		FGeometryCollectionEngineConversion::ConvertStaticMeshToGeometryCollection(StaticMeshVal, OutCollection, OutMaterialInstances, OutInstancedMeshes, bSetInternalFromMaterialIndex, bSplitComponents);
+		FGeometryCollectionEngineConversion::ConvertStaticMeshToGeometryCollection(StaticMeshVal, InMeshTransform, OutCollection, OutMaterialInstances, OutInstancedMeshes, bSetInternalFromMaterialIndex, bSplitComponents);
 	}
 
 	TArray<TObjectPtr<UMaterial>> OutMaterials;
diff --git a/Engine/Plugins/Experimental/GeometryCollectionPlugin/Source/GeometryCollectionNodes/Public/Dataflow/GeometryCollectionStaticMeshToCollectionNode.h b/Engine/Plugins/Experimental/GeometryCollectionPlugin/Source/GeometryCollectionNodes/Public/Dataflow/GeometryCollectionStaticMeshToCollectionNode.h
--- a/Engine/Plugins/Experimental/GeometryCollectionPlugin/Source/GeometryCollectionNodes/Public/Dataflow/GeometryCollectionStaticMeshToCollectionNode.h
+++ b/Engine/Plugins/Experimental/GeometryCollectionPlugin/Source/GeometryCollectionNodes/Public/Dataflow/GeometryCollectionStaticMeshToCollectionNode.h
@@ -27,6 +27,10 @@ public:
 	UPROPERTY(EditAnywhere, Category = "Asset", meta = (DataflowInput, DisplayName = "StaticMesh"));
 	TObjectPtr<UStaticMesh> StaticMesh;
 
+	/** Transform to apply to the mesh before converting it to a collection */
+	UPROPERTY(EditAnywhere, Category = "Asset", meta = (DataflowInput));
+	FTransform MeshTransform = FTransform::Identity;
+
 	/** Set the internal faces from material index */
 	UPROPERTY(EditAnywhere, Category = "Asset");
 	bool bSetInternalFromMaterialIndex = true;
diff --git a/Engine/Source/Runtime/Experimental/GeometryCollectionEngine/Private/GeometryCollection/GeometryCollectionEngineConversion.cpp b/Engine/Source/Runtime/Experimental/GeometryCollectionEngine/Private/GeometryCollection/GeometryCollectionEngineConversion.cpp
--- a/Engine/Source/Runtime/Experimental/GeometryCollectionEngine/Private/GeometryCollection/GeometryCollectionEngineConversion.cpp
+++ b/Engine/Source/Runtime/Experimental/GeometryCollectionEngine/Private/GeometryCollection/GeometryCollectionEngineConversion.cpp
@@ -1555,6 +1555,11 @@ void FGeometryCollectionEngineConversion::ConvertStaticMeshToGeometryCollection(
 }
 
 void FGeometryCollectionEngineConversion::ConvertStaticMeshToGeometryCollection(const TObjectPtr<UStaticMesh> StaticMesh, FManagedArrayCollection& OutCollection, TArray<TObjectPtr<UMaterialInterface>>& OutMaterialInstances, TArray<FGeometryCollectionAutoInstanceMesh>& OutInstancedMeshes, bool bSetInternalFromMaterialIndex, bool bSplitComponents)
+{
+	ConvertStaticMeshToGeometryCollection(StaticMesh, FTransform::Identity, OutCollection, OutMaterialInstances, OutInstancedMeshes, bSetInternalFromMaterialIndex, bSplitComponents);
+}
+
+void FGeometryCollectionEngineConversion::ConvertStaticMeshToGeometryCollection(const TObjectPtr<UStaticMesh> StaticMesh, const FTransform& MeshTransform, FManagedArrayCollection& OutCollection, TArray<TObjectPtr<UMaterialInterface>>& OutMaterialInstances, TArray<FGeometryCollectionAutoInstanceMesh>& OutInstancedMeshes, bool bSetInternalFromMaterialIndex, bool bSplitComponents)
 {
 #if WITH_EDITORONLY_DATA
 	if (UGeometryCollection* NewGeometryCollection = NewObject<UGeometryCollection>())
@@ -1562,8 +1567,6 @@ void FGeometryCollectionEngineConversion::ConvertStaticMeshToGeometryCollection(
 		// If any of the static meshes have Nanite enabled, also enable on the new geometry collection asset for convenience.
 		NewGeometryCollection->EnableNanite |= StaticMesh->IsNaniteEnabled();
 
-		FTransform ComponentTransform = FTransform::Identity;
-
 		// Record the contributing source on the asset.
 		FSoftObjectPath SourceSoftObjectPath(StaticMesh);
 
@@ -1583,8 +1586,8 @@ void FGeometryCollectionEngineConversion::ConvertStaticMeshToGeometryCollection(
 
 		bool bAddInternalMaterials = false;
 
-		NewGeometryCollection->GeometrySource.Emplace(SourceSoftObjectPath, ComponentTransform, SourceMaterials, bSplitComponents, bSetInternalFromMaterialIndex);
-		FGeometryCollectionEngineConversion::AppendStaticMesh(StaticMesh, SourceMaterials, ComponentTransform, NewGeometryCollection, false, bAddInternalMaterials, bSplitComponents, bSetInternalFromMaterialIndex);
+		NewGeometryCollection->GeometrySource.Emplace(SourceSoftObjectPath, MeshTransform, SourceMaterials, bSplitComponents, bSetInternalFromMaterialIndex);
+		FGeometryCollectionEngineConversion::AppendStaticMesh(StaticMesh, SourceMaterials, MeshTransform, NewGeometryCollection, false, bAddInternalMaterials, bSplitComponents, bSetInternalFromMaterialIndex);
 
 		NewGeometryCollection->InitializeMaterials();
 
diff --git a/Engine/Source/Runtime/Experimental/GeometryCollectionEngine/Public/GeometryCollection/GeometryCollectionEngineConversion.h b/Engine/Source/Runtime/Experimental/GeometryCollectionEngine/Public/GeometryCollection/GeometryCollectionEngineConversion.h
--- a/Engine/Source/Runtime/Experimental/GeometryCollectionEngine/Public/GeometryCollection/GeometryCollectionEngineConversion.h
+++ b/Engine/Source/Runtime/Experimental/GeometryCollectionEngine/Public/GeometryCollection/GeometryCollectionEngineConversion.h
@@ -192,6 +192,7 @@ public:
 	*/
 	static GEOMETRYCOLLECTIONENGINE_API void ConvertStaticMeshToGeometryCollection(const TObjectPtr<UStaticMesh> StaticMesh, FManagedArrayCollection& OutCollection, TArray<TObjectPtr<UMaterial>>& OutMaterials, TArray<FGeometryCollectionAutoInstanceMesh>& OutInstancedMeshes, bool bSetInternalFromMaterialIndex = true, bool bSplitComponents = false);
 	static GEOMETRYCOLLECTIONENGINE_API void ConvertStaticMeshToGeometryCollection(const TObjectPtr<UStaticMesh> StaticMesh, FManagedArrayCollection& OutCollection, TArray<TObjectPtr<UMaterialInterface>>& OutMaterialInstances, TArray<FGeometryCollectionAutoInstanceMesh>& OutInstancedMeshes, bool bSetInternalFromMaterialIndex = true, bool bSplitComponents = false);
+	static GEOMETRYCOLLECTIONENGINE_API void ConvertStaticMeshToGeometryCollection(const TObjectPtr<UStaticMesh> StaticMesh, const FTransform& MeshTransform, FManagedArrayCollection& OutCollection, TArray<TObjectPtr<UMaterialInterface>>& OutMaterialInstances, TArray<FGeometryCollectionAutoInstanceMesh>& OutInstancedMeshes, bool bSetInternalFromMaterialIndex = true, bool bSplitComponents = false);
 
 	/**
 	*  Converts a UGeometryCollection asset to an FManagedArrayCollection, and arrays of materials and instanced meshes
