Chaos::Fix crash when there are several island groups with joint constraints.

diff --git a/Engine/Source/Runtime/Experimental/Chaos/Private/Chaos/Joint/PBDJointContainerSolver.cpp b/Engine/Source/Runtime/Experimental/Chaos/Private/Chaos/Joint/PBDJointContainerSolver.cpp
--- a/Engine/Source/Runtime/Experimental/Chaos/Private/Chaos/Joint/PBDJointContainerSolver.cpp
+++ b/Engine/Source/Runtime/Experimental/Chaos/Private/Chaos/Joint/PBDJointContainerSolver.cpp
@@ -127,26 +127,10 @@ namespace Chaos
 
 		void FPBDJointContainerSolver::Reset(const int32 InMaxConstraints)
 		{
-			int32 NumNonLinearConstraints = 0;
-			int32 NumLinearConstraints = 0;
-			for (int32 i = 0; i < InMaxConstraints; i++)
-			{
-				if (i < ConstraintContainer.GetNumConstraints())
-				{
-					if (ConstraintContainer.GetConstraintSettings(i).bUseLinearSolver)
-					{
-						NumLinearConstraints += 1;
-					}
-					else
-					{
-						NumNonLinearConstraints += 1;
-					}
-				}
-			}
-			LinearConstraintSolvers.SetNum(NumLinearConstraints);
-			NonLinearConstraintSolvers.SetNum(NumNonLinearConstraints);
-			ContainerLinearConstraintGlobalIndices.Reset(NumLinearConstraints);
-			ContainerNonLinearConstraintGlobalIndices.Reset(NumNonLinearConstraints);
+			LinearConstraintSolvers.SetNum(InMaxConstraints);
+			NonLinearConstraintSolvers.SetNum(InMaxConstraints);
+			ContainerLinearConstraintGlobalIndices.Reset(InMaxConstraints);
+			ContainerNonLinearConstraintGlobalIndices.Reset(InMaxConstraints);
 		}
 
 		void FPBDJointContainerSolver::AddConstraints()
