Dataflow:
- Added random BoneColor generation to FGeometryCollectionTerminalDataflowNode

diff --git a/Engine/Source/Runtime/Experimental/GeometryCollectionEngine/Private/GeometryCollection/GeometryCollectionObject.cpp b/Engine/Source/Runtime/Experimental/GeometryCollectionEngine/Private/GeometryCollection/GeometryCollectionObject.cpp
--- a/Engine/Source/Runtime/Experimental/GeometryCollectionEngine/Private/GeometryCollection/GeometryCollectionObject.cpp
+++ b/Engine/Source/Runtime/Experimental/GeometryCollectionEngine/Private/GeometryCollection/GeometryCollectionObject.cpp
@@ -560,6 +560,26 @@ void UGeometryCollection::Reset()
 	}
 }
 
+namespace Dataflow::Private
+{
+	static void SetRandomBoneColor(TSharedPtr<FGeometryCollection, ESPMode::ThreadSafe>& InGeometryCollection)
+	{
+		TManagedArray<FLinearColor>& BoneColors = InGeometryCollection->BoneColor;
+
+		const int32 NumBones = BoneColors.Num();
+		FRandomStream RandomStream(NumBones);
+
+		for (int32 Idx = 0; Idx < NumBones; ++Idx)
+		{
+			const uint8 R = static_cast<uint8>(RandomStream.FRandRange(5, 105));
+			const uint8 G = static_cast<uint8>(RandomStream.FRandRange(5, 105));
+			const uint8 B = static_cast<uint8>(RandomStream.FRandRange(5, 105));
+
+			BoneColors[Idx] = FLinearColor(FColor(R, G, B, 255));
+		}
+	}
+}
+
 void UGeometryCollection::ResetFrom(const FManagedArrayCollection& InCollection, const TArray<UMaterial*>& InMaterials, bool bHasInternalMaterials)
 {
 	Reset();
@@ -572,6 +592,9 @@ void UGeometryCollection::ResetFrom(const FManagedArrayCollection& InCollection,
 				
 		Materials.Append(InMaterials);
 		InitializeMaterials(bHasInternalMaterials);
+
+		// Randomize BoneColor
+		Dataflow::Private::SetRandomBoneColor(GeometryCollection);
 	}
 }
 
@@ -585,6 +608,9 @@ void UGeometryCollection::ResetFrom(const FManagedArrayCollection& InCollection,
 		UpdateConvexGeometryIfMissing();
 		Materials.Append(InMaterialInstances);
 		InitializeMaterials(bHasInternalMaterials);
+
+		// Randomize BoneColor
+		Dataflow::Private::SetRandomBoneColor(GeometryCollection);
 	}
 }
 
