[Iris] Allow Out Of Bounds attachments to be sent ASAP and also allow reliable RPCs to be marked as "send immediately."

diff --git a/Engine/Source/Runtime/Experimental/Iris/Core/Private/Iris/ReplicationSystem/ReplicationSystem.cpp b/Engine/Source/Runtime/Experimental/Iris/Core/Private/Iris/ReplicationSystem/ReplicationSystem.cpp
--- a/Engine/Source/Runtime/Experimental/Iris/Core/Private/Iris/ReplicationSystem/ReplicationSystem.cpp
+++ b/Engine/Source/Runtime/Experimental/Iris/Core/Private/Iris/ReplicationSystem/ReplicationSystem.cpp
@@ -1086,14 +1086,8 @@ bool UReplicationSystem::SetRPCSendPolicyFlags(const UFunction* Function, UE::Ne
 		return false;
 	}
 
-	if (EnumHasAnyFlags(SendFlags, UE::Net::ENetObjectAttachmentSendPolicyFlags::SendImmediate) && (Function->FunctionFlags & FUNC_NetReliable))
-	{
-		ensureMsgf(false, TEXT("ENetObjectAttachmentSendPolicyFlags::SendImmediate is not allowed to use on Reliable RPC: %s"), *GetNameSafe(Function));
-		return false;
-	}
-
 	UE_LOG(LogIris, Verbose, TEXT("SetRPCSendPolicyFlags %s::%s => %s "), *GetNameSafe(Function->GetOuterUClass()), *GetNameSafe(Function), LexToString(SendFlags));
-		
+
 	Impl->AttachmentSendPolicyFlags.Add(FObjectKey(Function), SendFlags);
 	return true;
 }
diff --git a/Engine/Source/Runtime/Experimental/Iris/Core/Private/Iris/ReplicationSystem/ReplicationWriter.cpp b/Engine/Source/Runtime/Experimental/Iris/Core/Private/Iris/ReplicationSystem/ReplicationWriter.cpp
--- a/Engine/Source/Runtime/Experimental/Iris/Core/Private/Iris/ReplicationSystem/ReplicationWriter.cpp
+++ b/Engine/Source/Runtime/Experimental/Iris/Core/Private/Iris/ReplicationSystem/ReplicationWriter.cpp
@@ -3463,7 +3463,7 @@ UDataStream::EWriteResult FReplicationWriter::Write(FNetSerializationContext& Co
 	// https://jira.it.epicgames.com/browse/UE-127371
 	if (HasDataToSend(WriteContext))
 	{
-		if (WriteContext.bHasHugeObjectToSend)
+		if (WriteContext.bHasHugeObjectToSend || (WriteContext.bHasOOBAttachmentsToSend && CanSendObject(ObjectIndexForOOBAttachment)))
 		{
 			WriteResult = UDataStream::EWriteResult::HasMoreData;
 		}
