Move the I/O store trace macro to capture synchronous reads.

#rnx

diff --git a/Engine/Source/Runtime/Experimental/IoStore/OnDemand/Private/OnDemandInstallCache.cpp b/Engine/Source/Runtime/Experimental/IoStore/OnDemand/Private/OnDemandInstallCache.cpp
--- a/Engine/Source/Runtime/Experimental/IoStore/OnDemand/Private/OnDemandInstallCache.cpp
+++ b/Engine/Source/Runtime/Experimental/IoStore/OnDemand/Private/OnDemandInstallCache.cpp
@@ -1207,6 +1207,7 @@ bool FOnDemandInstallCache::Resolve(FIoRequestImpl* Request)
 		return false;
 	}
 
+	TRACE_IOSTORE_BACKEND_REQUEST_STARTED(Request, this);
 	Cas.TrackAccess(CasLoc.BlockId);
 
 	// Use synchronous file read API when reading from and writing to the same cache block concurrently.
@@ -1290,7 +1291,6 @@ bool FOnDemandInstallCache::Resolve(FIoRequestImpl* Request)
 		});
 	};
 
-	TRACE_IOSTORE_BACKEND_REQUEST_STARTED(Request, this);
 	ChunkRequest.FileReadRequest.Reset(FileHandle->ReadRequest(
 		CasLoc.BlockOffset + ChunkRequest.ChunkRange.GetOffset(),
 		ChunkRequest.ChunkRange.GetLength(),
@@ -1300,6 +1300,7 @@ bool FOnDemandInstallCache::Resolve(FIoRequestImpl* Request)
 
 	if (ChunkRequest.FileReadRequest.IsValid() == false)
 	{
+		TRACE_IOSTORE_BACKEND_REQUEST_FAILED(Request);
 		TUniquePtr<FChunkRequest> Detached = FChunkRequest::Detach(*Request);
 		return false;
 	}
