IAD - change TryLock to Lock when updating on-demand package store

#rnx

diff --git a/Engine/Source/Runtime/Experimental/IoStore/OnDemand/Private/OnDemandPackageStoreBackend.cpp b/Engine/Source/Runtime/Experimental/IoStore/OnDemand/Private/OnDemandPackageStoreBackend.cpp
--- a/Engine/Source/Runtime/Experimental/IoStore/OnDemand/Private/OnDemandPackageStoreBackend.cpp
+++ b/Engine/Source/Runtime/Experimental/IoStore/OnDemand/Private/OnDemandPackageStoreBackend.cpp
@@ -134,7 +134,7 @@ void FOnDemandPackageStoreBackend::OnMounted(TSharedRef<const FPackageStoreBacke
 
 void FOnDemandPackageStoreBackend::BeginRead()
 {
-	ensure(Mutex.TryLock());
+	Mutex.Lock();
 	UpdateLookupTables();
 }
 
