[NetPhysics] Refactor physics resimulation flow to further allow optimizations by disabling individual caches of data that might not be needed.
- Rely on FRewindData for both rewind and stepping through non-resim objects histories so we don't only rely on the ResimCache which holds duplicate data.
- Deactivate FEvolutionResimCache by default, API in FPhysicsSolverBase to toggle on and off and the normal EnableRewindData call will enable it if parameter for it is true.
- Put in CVar, enabled by default, to be able to disable FPushPhysicsData dirty proxies from being cached in history: p.Resim.CachePushDataDirtyProxies
- Put in CVar, enabled by default, to be able to disable FPushPhysicsData async inputs from being cached in history: p.Resim.CachePushDataAsyncInputs
- Put in CVar, enabled by default, to optimize RewindData: p.Resim.RewindDataOptimization
- Mark islands for resim if they have particles in them marked for FullResim
- Take actual physics tick-rate into account for physics tick buffering and rewind data history size, instead of value in project settings.

diff --git a/Engine/Source/Runtime/Engine/Classes/GameFramework/PlayerController.h b/Engine/Source/Runtime/Engine/Classes/GameFramework/PlayerController.h
--- a/Engine/Source/Runtime/Engine/Classes/GameFramework/PlayerController.h
+++ b/Engine/Source/Runtime/Engine/Classes/GameFramework/PlayerController.h
@@ -2415,8 +2415,9 @@ public:
 private:
 
 	/** The static offset between the local async physics tick frame number and the server's, kept in sync via time-dilation
-	*	This is used to synchronize events that happen in the async physics tick */
-	int32 NetworkPhysicsTickOffset = INDEX_NONE;
+	* This is used to synchronize events that happen in the async physics tick 
+	* This is 0 on the server */
+	int32 NetworkPhysicsTickOffset = 0;
 	bool bNetworkPhysicsTickOffsetAssigned = false;
 
 	/** Tick Offset sync interval timer */
diff --git a/Engine/Source/Runtime/Engine/Private/PhysicsEngine/NetworkPhysicsComponent.cpp b/Engine/Source/Runtime/Engine/Private/PhysicsEngine/NetworkPhysicsComponent.cpp
--- a/Engine/Source/Runtime/Engine/Private/PhysicsEngine/NetworkPhysicsComponent.cpp
+++ b/Engine/Source/Runtime/Engine/Private/PhysicsEngine/NetworkPhysicsComponent.cpp
@@ -397,8 +397,7 @@ void UNetworkPhysicsSystem::OnWorldPostInit(UWorld* World, const UWorld::Initial
 				{
 					if (Solver->GetRewindData() == nullptr)
 					{
-						const int32 NumFrames = UPhysicsSettings::Get()->GetPhysicsHistoryCount();
-						Solver->EnableRewindCapture(NumFrames, true);
+						Solver->EnableRewindCapture();
 					}
 				}
 			}
@@ -1193,7 +1192,7 @@ void FAsyncNetworkPhysicsComponent::ConsumeAsyncInput(const int32 PhysicsStep)
 						PhysicsData = nullptr;
 
 #if DEBUG_NETWORK_PHYSICS
-						AsyncInput->InputData->DebugData(FString::Printf(TEXT("SERVER | PT | Failed to extract | LatestFrame: %d | Name: %s"), AsyncInput->InputData->GetLatestFrame(), *GetActorName()));
+						AsyncInput->InputData->DebugData(FString::Printf(TEXT("SERVER | PT | Failed to extract LatestInputReceiveData | LatestFrame: %d | Name: %s"), AsyncInput->InputData->GetLatestFrame(), *GetActorName()));
 #endif
 					}
 				}
@@ -1307,14 +1306,14 @@ void FAsyncNetworkPhysicsComponent::OnPreProcessInputs_Internal(const int32 Phys
 			{
 				PhysicsData->ApplyData(ActorComponent.Get());
 #if DEBUG_NETWORK_PHYSICS
-				UE_LOG(LogChaos, Log, TEXT("		Applying extracted state from history | bExactFrame = %d | LocalFrame = %d | ServerFrame = %d | InputFrame = %d | Data: %s")
+				UE_LOG(LogChaos, Log, TEXT("			Applying extracted state from history | bExactFrame = %d | LocalFrame = %d | ServerFrame = %d | InputFrame = %d | Data: %s")
 					, bExactFrame, PhysicsData->LocalFrame, PhysicsData->ServerFrame, PhysicsData->InputFrame, *PhysicsData->DebugData());
 #endif
 			}
 #if DEBUG_NETWORK_PHYSICS
 			else if (PhysicsStep <= StateHistory->GetLatestFrame())
 			{
-				UE_LOG(LogChaos, Log, TEXT("		FAILED to extract and apply state from history | bExactFrame = %d | -- Printing history --"), bExactFrame);
+				UE_LOG(LogChaos, Log, TEXT("		Non-Determinism: FAILED to extract and apply state from history | bExactFrame = %d | -- Printing history --"), bExactFrame);
 				StateHistory->DebugData(FString::Printf(TEXT("StateHistory | Component = %s"), *GetActorName()));
 			}
 #endif
@@ -1333,7 +1332,7 @@ void FAsyncNetworkPhysicsComponent::OnPreProcessInputs_Internal(const int32 Phys
 				if (ComponentSettings.GetApplyDataInsteadOfMergeData())
 				{
 #if DEBUG_NETWORK_PHYSICS
-					UE_LOG(LogChaos, Log, TEXT("	Reapplying multiple data due to receiving an important data that was previously missed. FromFrame: %d | ToFrame: %d | IsLocallyControlled = %d"), NewImportantInputFrame, (NextExpectedLocalFrame - 1), IsLocallyControlled());
+					UE_LOG(LogChaos, Log, TEXT("		Non-Determinism: Reapplying multiple data due to receiving an important data that was previously missed. FromFrame: %d | ToFrame: %d | IsLocallyControlled = %d"), NewImportantInputFrame, (NextExpectedLocalFrame - 1), IsLocallyControlled());
 #endif
 					// Apply all inputs in range
 					InputHistory->ApplyDataRange(NewImportantInputFrame, NextExpectedLocalFrame - 1, ActorComponent.Get(), /*bOnlyImportant*/false);
@@ -1343,7 +1342,7 @@ void FAsyncNetworkPhysicsComponent::OnPreProcessInputs_Internal(const int32 Phys
 					// Merge all inputs from earliest new important
 					NextExpectedLocalFrame = NewImportantInputFrame;
 #if DEBUG_NETWORK_PHYSICS
-					UE_LOG(LogChaos, Log, TEXT("	Prepare to reapply multiple data through MergeData due to receiving an important data that was previously missed. FromFrame: %d | ToFrame: %d | IsLocallyControlled = %d"), NewImportantInputFrame, (NextExpectedLocalFrame - 1), IsLocallyControlled());
+					UE_LOG(LogChaos, Log, TEXT("		Non-Determinism: Prepare to reapply multiple data through MergeData due to receiving an important data that was previously missed. FromFrame: %d | ToFrame: %d | IsLocallyControlled = %d"), NewImportantInputFrame, (NextExpectedLocalFrame - 1), IsLocallyControlled());
 #endif
 				}
 			}
@@ -1365,7 +1364,7 @@ void FAsyncNetworkPhysicsComponent::OnPreProcessInputs_Internal(const int32 Phys
 					if (ComponentSettings.GetApplyDataInsteadOfMergeData())
 					{
 #if DEBUG_NETWORK_PHYSICS
-						UE_LOG(LogChaos, Log, TEXT("	Applying multiple data instead of merging, from LocalFrame %d into LocalFrame %d | IsLocallyControlled = %d"), NextExpectedLocalFrame, PhysicsData->LocalFrame, IsLocallyControlled());
+						UE_LOG(LogChaos, Log, TEXT("		Non-Determinism: Applying multiple data instead of merging, from LocalFrame %d into LocalFrame %d | IsLocallyControlled = %d"), NextExpectedLocalFrame, PhysicsData->LocalFrame, IsLocallyControlled());
 #endif
 						// Iterate over each input and call ApplyData, except on the last, it will get handled by the normal ApplyData call further down
 						const int32 LastFrame = PhysicsData->LocalFrame;
@@ -1380,7 +1379,7 @@ void FAsyncNetworkPhysicsComponent::OnPreProcessInputs_Internal(const int32 Phys
 					else
 					{
 #if DEBUG_NETWORK_PHYSICS
-						UE_LOG(LogChaos, Log, TEXT("	Merging inputs from LocalFrame %d into LocalFrame %d | IsLocallyControlled = %d"), NextExpectedLocalFrame, PhysicsData->LocalFrame, IsLocallyControlled());
+						UE_LOG(LogChaos, Log, TEXT("		Non-Determinism: Merging inputs from LocalFrame %d into LocalFrame %d | IsLocallyControlled = %d"), NextExpectedLocalFrame, PhysicsData->LocalFrame, IsLocallyControlled());
 #endif
 						// Merge all inputs since last used input
 						InputHistory->MergeData(NextExpectedLocalFrame, PhysicsData);
@@ -1393,7 +1392,7 @@ void FAsyncNetworkPhysicsComponent::OnPreProcessInputs_Internal(const int32 Phys
 #if DEBUG_NETWORK_PHYSICS
 					if (PhysicsStep > InputHistory->GetLatestFrame())
 					{
-						UE_LOG(LogChaos, Log, TEXT("	Input buffer is empty, input for frame %d was extrapolated from frame: %d"), PhysicsStep, PhysicsData->LocalFrame);
+						UE_LOG(LogChaos, Log, TEXT("		Non-Determinism: Input buffer Empty, input for frame %d was extrapolated from frame: %d"), PhysicsStep, PhysicsData->LocalFrame);
 					}
 #endif
 					PhysicsData->bReceivedData = true; // Mark the input data as received so that it doesn't get overwritten by incoming client inputs
@@ -1406,7 +1405,7 @@ void FAsyncNetworkPhysicsComponent::OnPreProcessInputs_Internal(const int32 Phys
 
 #if DEBUG_NETWORK_PHYSICS
 				{
-					UE_LOG(LogChaos, Log, TEXT("		Applying extracted input from history | LocalFrame = %d | ServerFrame = %d | InputFrame = %d | IsResim = %d | IsLocallyControlled = %d | InputDecay = %f | Data: %s")
+					UE_LOG(LogChaos, Log, TEXT("			Applying extracted input from history | LocalFrame = %d | ServerFrame = %d | InputFrame = %d | IsResim = %d | IsLocallyControlled = %d | InputDecay = %f | Data: %s")
 						, PhysicsData->LocalFrame, PhysicsData->ServerFrame, PhysicsData->InputFrame, bIsSolverResim, IsLocallyControlled(), GetCurrentInputDecay(PhysicsData), *PhysicsData->DebugData());
 				}
 #endif
@@ -1414,7 +1413,7 @@ void FAsyncNetworkPhysicsComponent::OnPreProcessInputs_Internal(const int32 Phys
 #if DEBUG_NETWORK_PHYSICS
 			else if (PhysicsStep <= InputHistory->GetLatestFrame())
 			{
-				UE_LOG(LogChaos, Log, TEXT("		FAILED to extract and apply input from history | IsResim = %d | IsLocallyControlled = %d | -- Printing history --"), bIsSolverResim, IsLocallyControlled());
+				UE_LOG(LogChaos, Log, TEXT("		Non-Determinism: FAILED to extract and apply input from history | IsResim = %d | IsLocallyControlled = %d | -- Printing history --"), bIsSolverResim, IsLocallyControlled());
 				InputHistory->DebugData(FString::Printf(TEXT("InputHistory | Name = %s"), *GetActorName()));
 			}
 #endif
@@ -1742,14 +1741,19 @@ const int32 FAsyncNetworkPhysicsComponent::SetupRewindData()
 
 	if (Chaos::FPBDRigidsSolver* RigidSolver = GetRigidSolver())
 	{
-		NumFrames = RigidSolver->GetPhysicsHistoryCount();
+		NumFrames = FMath::Max<int32>(1, FMath::CeilToInt32((0.001f * Chaos::FPBDRigidsSolver::GetPhysicsHistoryTimeLength()) / RigidSolver->GetAsyncDeltaTime()));
+
+		if (IsServer())
+		{
+			return NumFrames;
+		}
 
 		// Don't let this actor initialize RewindData if not using resimulation
 		if (GetPhysicsReplicationMode() == EPhysicsReplicationMode::Resimulation)
 		{
 			if (RigidSolver->IsNetworkPhysicsPredictionEnabled() && RigidSolver->GetRewindData() == nullptr)
 			{
-				RigidSolver->EnableRewindCapture(NumFrames, true);
+				RigidSolver->EnableRewindCapture();
 			}
 		}
 
diff --git a/Engine/Source/Runtime/Engine/Private/PhysicsEngine/PhysLevel.cpp b/Engine/Source/Runtime/Engine/Private/PhysicsEngine/PhysLevel.cpp
--- a/Engine/Source/Runtime/Engine/Private/PhysicsEngine/PhysLevel.cpp
+++ b/Engine/Source/Runtime/Engine/Private/PhysicsEngine/PhysLevel.cpp
@@ -95,6 +95,11 @@ public:
 		return GetSettings()->PhysicsPrediction.ResimulationSettings.ResimulationErrorAngularVelocityThreshold;
 	}
 
+	virtual float GetPhysicsHistoryTimeLength() const
+	{
+		return GetSettings()->PhysicsPrediction.MaxSupportedLatencyPrediction;
+	}
+
 	virtual int32 GetPhysicsHistoryCount() const override
 	{
 		return GetSettings()->GetPhysicsHistoryCount();
diff --git a/Engine/Source/Runtime/Engine/Private/PhysicsEngine/PhysicsReplication.cpp b/Engine/Source/Runtime/Engine/Private/PhysicsEngine/PhysicsReplication.cpp
--- a/Engine/Source/Runtime/Engine/Private/PhysicsEngine/PhysicsReplication.cpp
+++ b/Engine/Source/Runtime/Engine/Private/PhysicsEngine/PhysicsReplication.cpp
@@ -929,8 +929,7 @@ void FPhysicsReplicationAsync::OnPreSimulate_Internal()
 				// We don't have rewind data but an actor is set to replicate using resimulation; we need to enable rewind capture.
 				if (ensure(Chaos::FPBDRigidsSolver::IsNetworkPhysicsPredictionEnabled() && RigidsSolver->IsUsingFixedDt()))
 				{
-					const int32 NumFrames = FMath::Max<int32>(1, Chaos::FPBDRigidsSolver::GetPhysicsHistoryCount());
-					RigidsSolver->EnableRewindCapture(NumFrames, true);
+					RigidsSolver->EnableRewindCapture();
 				}
 			}
 
diff --git a/Engine/Source/Runtime/Engine/Private/PlayerController.cpp b/Engine/Source/Runtime/Engine/Private/PlayerController.cpp
--- a/Engine/Source/Runtime/Engine/Private/PlayerController.cpp
+++ b/Engine/Source/Runtime/Engine/Private/PlayerController.cpp
@@ -6274,9 +6274,24 @@ void APlayerController::ServerSendLatestAsyncPhysicsTimestamp_Implementation(FAs
 	}
 	PRAGMA_ENABLE_DEPRECATION_WARNINGS
 
+	// Get the fixed timestep from project settings
+	float AsyncFixedTimeStepSize = UPhysicsSettings::Get()->AsyncFixedTimeStepSize;
+
+	if (UWorld* World = GetWorld())
+	{
+		if (FPhysScene* PhysScene = World->GetPhysicsScene())
+		{
+			if (Chaos::FPhysicsSolver* Solver = PhysScene->GetSolver())
+			{
+				// Get fixed timestep from solver since it can have been altered
+				AsyncFixedTimeStepSize = Solver->GetAsyncDeltaTime();
+			}
+		}
+	}
+
 	// Get current server timestamp and add the frame buffer to the ServerFrame
 	FAsyncPhysicsTimestamp ActualTimestamp = GetPhysicsTimestamp();
-	const int32 BufferTickSize = FMath::CeilToInt((NetworkPhysicsCvars::TickOffsetBufferTime / 1000.f) / UPhysicsSettings::Get()->AsyncFixedTimeStepSize);
+	const int32 BufferTickSize = FMath::CeilToInt((NetworkPhysicsCvars::TickOffsetBufferTime / 1000.f) / AsyncFixedTimeStepSize);
 	ActualTimestamp.ServerFrame += BufferTickSize;
 
 	// Mark offset as assigned when we get a valid predicted server frame.
@@ -6285,7 +6300,7 @@ void APlayerController::ServerSendLatestAsyncPhysicsTimestamp_Implementation(FAs
 
 	const float CurrentTime = GetWorld()->GetTimeSeconds();
 	const int32 TimestampDiff = FMath::Abs(PredictedServerFrame - ActualTimestamp.ServerFrame);
-	const float TimestampTimeDiff = TimestampDiff * UPhysicsSettings::Get()->AsyncFixedTimeStepSize;
+	const float TimestampTimeDiff = TimestampDiff * AsyncFixedTimeStepSize;
 	NetworkPhysicsTickOffsetDesyncAccumulatedTime = (TimestampDiff == 0) ? CurrentTime : NetworkPhysicsTickOffsetDesyncAccumulatedTime;
 
 	// Send update to client if offset is not assigned or over correction limits
@@ -6295,8 +6310,13 @@ void APlayerController::ServerSendLatestAsyncPhysicsTimestamp_Implementation(FAs
 		|| CurrentTime - NetworkPhysicsTickOffsetDesyncAccumulatedTime > (NetworkPhysicsCvars::TickOffsetCorrectionTimeLimit / 1000.0f))
 	{
 		Timestamp.ServerFrame = ActualTimestamp.ServerFrame;
-		NetworkPhysicsTickOffset = Timestamp.ServerFrame - Timestamp.LocalFrame;
 		NetworkPhysicsTickOffsetDesyncAccumulatedTime = CurrentTime;
+
+#if DEBUG_NETWORK_PHYSICS
+		UE_LOG(LogPlayerController, Log, TEXT("APlayerController::ServerSendLatestAsyncPhysicsTimestamp_Implementation. Sync physics tick with client. ClientFrame: %d, ServerFrame: %d, BufferSize: %d, PredictedFrame: %d)")
+			, Timestamp.LocalFrame, Timestamp.ServerFrame, BufferTickSize, PredictedServerFrame);
+#endif
+
 		ClientSetupNetworkPhysicsTimestamp(Timestamp); /* Reliable RPC */
 	}
 
@@ -6370,6 +6390,11 @@ void APlayerController::ClientSetupNetworkPhysicsTimestamp_Implementation(FAsync
 	// Assign async physics tick offset
 	bNetworkPhysicsTickOffsetAssigned = true;
 	NetworkPhysicsTickOffset = Timestamp.ServerFrame - Timestamp.LocalFrame;
+
+#if DEBUG_NETWORK_PHYSICS
+	UE_LOG(LogPlayerController, Log, TEXT("APlayerController::ClientSetupNetworkPhysicsTimestamp_Implementation. ClientFrame: %d, ServerFrame: %d, NetworkPhysicsTickOffset: %d)")
+		, Timestamp.LocalFrame, Timestamp.ServerFrame, NetworkPhysicsTickOffset);
+#endif
 }
 
 void APlayerController::ClientAckTimeDilation_Implementation(float TimeDilation, int32 ServerStep)
diff --git a/Engine/Source/Runtime/Experimental/Chaos/Private/Chaos/ChaosMarshallingManager.cpp b/Engine/Source/Runtime/Experimental/Chaos/Private/Chaos/ChaosMarshallingManager.cpp
--- a/Engine/Source/Runtime/Experimental/Chaos/Private/Chaos/ChaosMarshallingManager.cpp
+++ b/Engine/Source/Runtime/Experimental/Chaos/Private/Chaos/ChaosMarshallingManager.cpp
@@ -10,6 +10,12 @@ namespace Chaos
 int32 SimDelay = 0;
 FAutoConsoleVariableRef CVarSimDelay(TEXT("p.simDelay"),SimDelay,TEXT(""));
 
+bool bCachePushDataDirtyProxies = true;
+FAutoConsoleVariableRef CVarCachePushDataDirtyProxies(TEXT("p.Resim.CachePushDataDirtyProxies"), bCachePushDataDirtyProxies, TEXT("Default = false. Set true to enable resim caching dirty proxies in the push data from game thread to physics thread. This will make physics proxy changes from GT play out during a resimulation."));
+
+bool bCachePushDataAsyncInputs = true;
+FAutoConsoleVariableRef CVarCachePushDataAsyncInputs(TEXT("p.Resim.CachePushDataAsyncInputs"), bCachePushDataAsyncInputs, TEXT("Default = false. Set true to enable resim caching of async inputs in the push data from game thread to physics thread. This will make async inputs available again during a resimulation."));
+
 FChaosMarshallingManager::FChaosMarshallingManager()
 : ExternalTime_External(0)
 , ExternalTimestamp_External(0)
@@ -128,16 +134,7 @@ FPushPhysicsData* FChaosMarshallingManager::StepInternalTime_External()
 void FChaosMarshallingManager::FreeData_Internal(FPushPhysicsData* PushData)
 {
 	UE_CHAOS_ASYNC_INITBODY_WRITESCOPELOCK(MarshallingManagerLock);
-	//TODO: we know entire manager is cleared, so we can probably just iterate over its pools and reset
-	//instead of going through dirty proxies. If perf matters fix this
-	FDirtyPropertiesManager* Manager = &PushData->DirtyPropertiesManager;
-	FShapeDirtyData* ShapeDirtyData = PushData->DirtyProxiesDataBuffer.GetShapesDirtyData();
-
-	PushData->DirtyProxiesDataBuffer.ForEachProxy([Manager, ShapeDirtyData](int32 DataIdx, FDirtyProxy& Dirty)
-	{
-		Dirty.Clear(*Manager, DataIdx, ShapeDirtyData);
-	});
-
+	PushData->ResetDirtyProxiesBuffer();
 	PushData->Reset();
 	PushDataPool.Enqueue(PushData);
 }
@@ -163,7 +160,6 @@ void FPushPhysicsData::Reset()
 
 	DirtyProxiesDataBuffer.Reset();
 	SimCallbackInputs.Reset();
-	SimCallbackObjectsToRemove.Reset();
 	ResetForHistory();
 }
 
@@ -171,6 +167,20 @@ void FPushPhysicsData::ResetForHistory()
 {
 	SimCommands.Reset();
 	SimCallbackObjectsToAdd.Reset();
+	SimCallbackObjectsToRemove.Reset();
+}
+
+void FPushPhysicsData::ResetDirtyProxiesBuffer()
+{
+	//TODO: we know entire manager is cleared, so we can probably just iterate over its pools and reset
+	//instead of going through dirty proxies. If perf matters fix this
+	FDirtyPropertiesManager* Manager = &DirtyPropertiesManager;
+	FShapeDirtyData* ShapeDirtyData = DirtyProxiesDataBuffer.GetShapesDirtyData();
+
+	DirtyProxiesDataBuffer.ForEachProxy([Manager, ShapeDirtyData](int32 DataIdx, FDirtyProxy& Dirty)
+		{
+			Dirty.Clear(*Manager, DataIdx, ShapeDirtyData);
+		});
 }
 
 void FChaosMarshallingManager::FreeDataToHistory_Internal(FPushPhysicsData* PushData)
@@ -181,6 +191,17 @@ void FChaosMarshallingManager::FreeDataToHistory_Internal(FPushPhysicsData* Push
 	}
 	else
 	{
+		if (bCachePushDataDirtyProxies == false)
+		{
+			PushData->ResetDirtyProxiesBuffer();
+			PushData->DirtyProxiesDataBuffer.Reset();
+		}
+
+		if (bCachePushDataAsyncInputs == false)
+		{
+			PushData->SimCallbackInputs.Reset();
+		}
+
 		PushData->ResetForHistory();
 		HistoryQueue_Internal.Insert(PushData, 0);
 		SetHistoryLength_Internal(HistoryLength);
diff --git a/Engine/Source/Runtime/Experimental/Chaos/Private/Chaos/Island/IslandManager.cpp b/Engine/Source/Runtime/Experimental/Chaos/Private/Chaos/Island/IslandManager.cpp
--- a/Engine/Source/Runtime/Experimental/Chaos/Private/Chaos/Island/IslandManager.cpp
+++ b/Engine/Source/Runtime/Experimental/Chaos/Private/Chaos/Island/IslandManager.cpp
@@ -146,9 +146,9 @@ namespace Chaos::Private
 		return !bIsStationary;
 	}
 
-	bool IsParticleNeedsResim(const FGeometryParticleHandle* Particle)
+	bool ParticlesIslandNeedsResim(const FGeometryParticleHandle* Particle)
 	{
-		return (Particle != nullptr) && (Particle->SyncState() != ESyncState::InSync);
+		return (Particle != nullptr) && (Particle->SyncState() != ESyncState::InSync || Particle->ResimType() == EResimType::FullResim);
 	}
 
 	bool GetIslandParticleSleepThresholds(
@@ -1061,7 +1061,7 @@ namespace Chaos::Private
 			Node->Flags.bIsDynamic = IsParticleDynamic(Particle);
 			Node->Flags.bIsSleeping = IsParticleSleeping(Particle);
 			Node->Flags.bIsMoving = IsParticleMoving(Particle);
-			Node->Flags.bNeedsResim = IsParticleNeedsResim(Particle);
+			Node->Flags.bNeedsResim = ParticlesIslandNeedsResim(Particle);
 			UpdateGraphNodeSleepSettings(Node);
 
 			return Node;
@@ -1222,7 +1222,7 @@ namespace Chaos::Private
 		const bool bIsDynamic = IsParticleDynamic(Node->Particle);
 		const bool bIsMoving = IsParticleMoving(Node->Particle);
 		const bool bIsSleeping = IsParticleSleeping(Node->Particle);
-		const bool bNeedsResim = IsParticleNeedsResim(Node->Particle);
+		const bool bNeedsResim = ParticlesIslandNeedsResim(Node->Particle);
 
 		Node->Flags.bIsDynamic = bIsDynamic;
 		Node->Flags.bIsMoving = bIsMoving;
diff --git a/Engine/Source/Runtime/Experimental/Chaos/Private/Chaos/PBDRigidsEvolutionGBF.cpp b/Engine/Source/Runtime/Experimental/Chaos/Private/Chaos/PBDRigidsEvolutionGBF.cpp
--- a/Engine/Source/Runtime/Experimental/Chaos/Private/Chaos/PBDRigidsEvolutionGBF.cpp
+++ b/Engine/Source/Runtime/Experimental/Chaos/Private/Chaos/PBDRigidsEvolutionGBF.cpp
@@ -672,12 +672,8 @@ void FPBDRigidsEvolutionGBF::AdvanceOneTimeStepImpl(const FReal Dt, const FSubSt
 		SCOPE_CYCLE_COUNTER(STAT_Evolution_BuildGroups);
 		CSV_SCOPED_TIMING_STAT(PhysicsVerbose, StepSolver_BuildGroups);
 
-		// If we are resimming and if we have a particle cache we only build the island the groups 
-		// for islands that required to be simulated based of desynced particles
-		FEvolutionResimCache* ResimCache = GetCurrentStepResimCache();
-		const bool bIsResimming = (ResimCache != nullptr) && ResimCache->IsResimming();
-
-		NumGroups = IslandGroupManager.BuildGroups(bIsResimming);
+		// If we are resimulating, only build island groups for islands that require to be simulated based of particles sync state and resim type
+		NumGroups = IslandGroupManager.BuildGroups(bIsResim);
 	}
 
 
diff --git a/Engine/Source/Runtime/Experimental/Chaos/Private/PBDRigidsSolver.cpp b/Engine/Source/Runtime/Experimental/Chaos/Private/PBDRigidsSolver.cpp
--- a/Engine/Source/Runtime/Experimental/Chaos/Private/PBDRigidsSolver.cpp
+++ b/Engine/Source/Runtime/Experimental/Chaos/Private/PBDRigidsSolver.cpp
@@ -1100,29 +1100,41 @@ namespace Chaos
 	}
 
 	void FPBDRigidsSolver::EnableRewindCapture(int32 NumFrames, bool InUseCollisionResimCache)
+	{
+		SetUseCollisionResimCache(InUseCollisionResimCache);
+		EnableRewindCapture(NumFrames);
+	}
+
+	void FPBDRigidsSolver::EnableRewindCapture(int32 NumFrames)
 	{
 		//TODO: this function calls both internal and external - sort of assumed during initialization. Should decide what thread it's called on and mark it as either external or internal
 		if (MRewindData.IsValid())
 		{
-			MRewindData->Init(((FPBDRigidsSolver*)this), NumFrames, InUseCollisionResimCache, ((FPBDRigidsSolver*)this)->GetCurrentFrame());
+			MRewindData->Init(((FPBDRigidsSolver*)this), NumFrames, GetCurrentFrame());
 		}
 		else
 		{
-			MRewindData = MakeUnique<FRewindData>(((FPBDRigidsSolver*)this), NumFrames, InUseCollisionResimCache, ((FPBDRigidsSolver*)this)->GetCurrentFrame()); // FIXME
+			MRewindData = MakeUnique<FRewindData>(((FPBDRigidsSolver*)this), NumFrames, GetCurrentFrame());
 		}
-		bUseCollisionResimCache = InUseCollisionResimCache;
+
 		const int32 NumFramesSet = GetRewindData() != nullptr ? GetRewindData()->Capacity() : NumFrames;
 		MarshallingManager.SetHistoryLength_Internal(NumFramesSet);
 		MEvolution->SetRewindData(GetRewindData());
-		
-		if (MRewindCallback) 
+
+		if (MRewindCallback)
 		{
 			MRewindCallback->RewindData = GetRewindData();
 		}
-		
+
 		UpdateIsDeterministic();
 
-		UE_LOG(LogChaos, Log, TEXT("PBDRigidsSolver::EnableRewindCapture - Starting physics data history caching for rewind / resimulation. History Size: %d"), NumFramesSet);
+		UE_LOG(LogChaos, Log, TEXT("PBDRigidsSolver::EnableRewindCapture - Starting physics data history caching for rewind / resimulation. History Size: %d. Supported Latency: %f. TickRate: %d. "), NumFramesSet, FPBDRigidsSolver::GetPhysicsHistoryTimeLength(), FMath::RoundToInt32(1.0f / GetAsyncDeltaTime()));
+	}
+
+	void FPBDRigidsSolver::EnableRewindCapture()
+	{
+		int32 NumFrames = FMath::Max<int32>(1, FMath::CeilToInt32((0.001f * FPBDRigidsSolver::GetPhysicsHistoryTimeLength()) / GetAsyncDeltaTime()));
+		EnableRewindCapture(NumFrames);
 	}
 
 	void FPBDRigidsSolver::Reset()
@@ -1142,11 +1154,9 @@ namespace Chaos
 		//todo: do we need this?
 		//MarshallingManager.Reset();
 
-		const int32 PhysicsHistoryLength = FChaosSolversModule::GetModule()->GetSettingsProvider().GetPhysicsHistoryCount();
-
-		if (bUseCollisionResimCache && PhysicsHistoryLength >= 0)
+		if (bUseCollisionResimCache)
 		{
-			EnableRewindCapture(PhysicsHistoryLength, true);
+			EnableRewindCapture(true);
 		}
 
 		MEvolution->SetCaptureRewindDataFunction([this](const TParticleView<TPBDRigidParticles<FReal,3>>& ActiveParticles)
@@ -1952,6 +1962,11 @@ namespace Chaos
 					{
 						if ((LastStep - Step) < RecordedPushData.Num())
 						{
+							if (!bFirst && !bUseCollisionResimCache)
+							{
+								MRewindData->StepNonResimParticles(Step);
+							}
+
 							if (PhysicsReplicationCVars::ResimulationCVars::bApplyTargetsWhileResimulating || bFirst)
 							{
 								// Update all the particles having received a target from the server
diff --git a/Engine/Source/Runtime/Experimental/Chaos/Private/RewindData.cpp b/Engine/Source/Runtime/Experimental/Chaos/Private/RewindData.cpp
--- a/Engine/Source/Runtime/Experimental/Chaos/Private/RewindData.cpp
+++ b/Engine/Source/Runtime/Experimental/Chaos/Private/RewindData.cpp
@@ -155,6 +155,37 @@ bool FJointStateBase::IsInSync(const FPBDJointConstraintHandle& Handle, const FF
 	return true;
 }
 
+
+bool bCVarRewindDataOptimization = true;
+FAutoConsoleVariableRef CVarRewindDataOptimization(TEXT("p.Resim.RewindDataOptimization"), bCVarRewindDataOptimization, TEXT("Default value for RewinData optimization, note that this can be overridden at runtime by API calls. Effect: Only alter the minimum required properties during a resim for particles not marked for FullResim"));
+
+FRewindData::FRewindData(FPBDRigidsSolver* InSolver, int32 NumFrames, bool InRewindDataOptimization, int32 InCurrentFrame)
+	: Managers(NumFrames + 1)	//give 1 extra for saving at head
+	, Solver(InSolver)
+	, CurFrame(InCurrentFrame)
+	, LatestFrame(InCurrentFrame)
+	, FramesSaved(0)
+	, DataIdxOffset(0)
+	, bNeedsSave(false)
+	, bRewindDataOptimization(InRewindDataOptimization)
+	, LatestTargetFrame(0)
+{
+}
+
+FRewindData::FRewindData(FPBDRigidsSolver* InSolver, int32 NumFrames, int32 InCurrentFrame)
+	: Managers(NumFrames + 1)	//give 1 extra for saving at head
+	, Solver(InSolver)
+	, CurFrame(InCurrentFrame)
+	, LatestFrame(InCurrentFrame)
+	, FramesSaved(0)
+	, DataIdxOffset(0)
+	, bNeedsSave(false)
+	, bRewindDataOptimization(bCVarRewindDataOptimization)
+	, LatestTargetFrame(0)
+{
+}
+
+
 void FRewindData::ApplyInputs(const int32 ApplyFrame, const bool bResetSolver)
 {
 	for (TWeakPtr<FBaseRewindHistory>& InputHistory : InputHistories)
@@ -239,13 +270,16 @@ void FRewindData::ApplyTargets(const int32 Frame, const bool bResetSimulation)
 			Particle->SetDisabled(Data.Disabled());
 			Solver->GetEvolution()->SetParticleObjectState(Particle, Data.ObjectState());
 
-			if(Data.ObjectState() == EObjectStateType::Dynamic)
+			// Todo: EResimType should be set by a resimulation system and ApplyTargets() should only process particles marked for resim
+			switch (Data.ObjectState())
 			{
+			case EObjectStateType::Dynamic:
+			case EObjectStateType::Sleeping:
 				Particle->SetResimType(EResimType::FullResim);
-			}
-			else if((Data.ObjectState() == EObjectStateType::Static) || (Data.ObjectState() == EObjectStateType::Kinematic))
-			{
+				break;
+			default:
 				Particle->SetResimType(EResimType::ResimAsFollower);
+				break;
 			}
 		});
 #if !(UE_BUILD_SHIPPING || UE_BUILD_TEST)
@@ -426,6 +460,19 @@ bool FRewindData::RewindToFrame(int32 Frame)
 
 		History.CachePreCorrectionState(*PTParticle);
 
+		// Todo: This should be set by bubble resimulation so that Dynamic and Sleeping particles outside of relevancy of resim doesn't actually resimulate as dynamic particles.
+		switch (PTParticle->ObjectState())
+		{
+		case EObjectStateType::Dynamic:
+		case EObjectStateType::Sleeping:
+			PTParticle->SetResimType(EResimType::FullResim);
+			break;
+		default:
+
+			PTParticle->SetResimType(EResimType::ResimAsFollower);
+			break;
+		}
+
 		const bool bResimAsFollower = DirtyParticleInfo.bResimAsFollower;
 
 		bool bAnyChange = RewindHelper(PTParticle, bResimAsFollower, History.ParticlePositionRotation, [](auto Particle, const auto& Data) {Particle->SetXR(Data); });
@@ -482,6 +529,131 @@ bool FRewindData::RewindToFrame(int32 Frame)
 	return true;
 }
 
+void FRewindData::StepNonResimParticles(const int32 Frame)
+{
+	const FFrameAndPhase FrameAndPhase{ Frame, FFrameAndPhase::PrePushData };
+	auto RewindHelper = [FrameAndPhase, this](auto Obj, auto& Property, const auto& RewindFunc) -> bool
+	{
+		if (auto Val = Property.Read(FrameAndPhase, PropertiesPool))
+		{
+			return RewindFunc(Obj, *Val);
+		}
+		return false;
+	};
+
+	for (FDirtyParticleInfo& DirtyParticleInfo : DirtyParticles)
+	{
+		FGeometryParticleHandle* PTParticle = DirtyParticleInfo.GetObjectPtr();
+		if (PTParticle->ResimType() != EResimType::ResimAsFollower)
+		{
+			continue;
+		}
+
+		bool bHasChanged = false;
+		const FGeometryParticleStateBase& History = DirtyParticleInfo.GetHistory();
+
+		//  Set Postion and Rotation
+		bHasChanged = RewindHelper(PTParticle, History.ParticlePositionRotation, [](auto Particle, const auto& Data) -> bool
+			{
+				if (Particle->GetX() != Data.GetX() || Particle->GetR() != Data.GetR())
+				{
+					Particle->SetXR(Data);
+					return true; 
+				}
+				return false;
+			});
+
+		// Set Velocity and Angular Velocity
+		bHasChanged |= RewindHelper(PTParticle->CastToKinematicParticle(), History.Velocities, [](auto Particle, const auto& Data) -> bool
+			{
+				if (Particle->GetV() != Data.GetV() || Particle->GetW() != Data.GetW())
+				{
+					Particle->SetV(Data.GetV()); Particle->SetW(Data.GetW());
+					return true;
+				}
+				return false;
+			});
+
+		// If XRVW has not changed for the non-resim particle, continue to the next particle
+		if (!bHasChanged)
+		{
+			continue;
+		}
+
+		if (bRewindDataOptimization)
+		{
+			// Set kinematic target
+			RewindHelper(PTParticle->CastToKinematicParticle(), History.KinematicTarget, [](auto Particle, const auto& Data) -> bool { Particle->SetKinematicTarget(Data); return true; });
+
+			// Set disabled true/false and object state
+			bool bHasUpdatedSOAs = RewindHelper(PTParticle->CastToRigidParticle(), History.DynamicsMisc, [this](auto Particle, const auto& Data) -> bool
+				{
+					if (Particle == nullptr)
+					{
+						return false; // SOAs views have not been updated
+					}
+
+					if (Particle->Disabled() != Data.Disabled())
+					{
+						if (Data.Disabled())
+						{
+							Solver->GetEvolution()->DisableParticle(Particle);
+						}
+						else
+						{
+							Solver->GetEvolution()->EnableParticle(Particle);
+						}
+					}
+
+					if (Particle->ObjectState() != Data.ObjectState())
+					{
+						Solver->GetEvolution()->SetParticleObjectState(Particle, Data.ObjectState());
+						return true; // SOA views are updated when calling this function
+					}
+
+					return false; // SOAs views have not been updated
+				});
+
+			// If not already done, update SOA views else particles might not get updated
+			if (!bHasUpdatedSOAs)
+			{
+				if (FPBDRigidParticleHandle* Rigid = PTParticle->CastToRigidParticle())
+				{
+					Solver->GetEvolution()->GetParticles().SetDynamicParticleSOA(Rigid->Handle());
+				}
+				else if (FPBDRigidClusteredParticleHandle* Clustered = PTParticle->CastToClustered())
+				{
+					Solver->GetEvolution()->GetParticles().SetClusteredParticleSOA(Clustered->Handle());
+				}	
+			}
+		}
+		else
+		{
+			RewindHelper(PTParticle, History.NonFrequentData, [this](auto Particle, const auto& Data) -> bool
+				{
+					Solver->GetEvolution()->InvalidateParticle(Particle); // Clear collision/constraints before updating NonFrequentData
+					Particle->SetNonFrequentData(Data);
+					return true;
+				});
+			RewindHelper(PTParticle->CastToKinematicParticle(), History.KinematicTarget, [](auto Particle, const auto& Data) -> bool { Particle->SetKinematicTarget(Data); return true; });
+			RewindHelper(PTParticle->CastToRigidParticle(), History.Dynamics, [](auto Particle, const auto& Data) -> bool { Particle->SetDynamics(Data); return true; });
+			RewindHelper(PTParticle->CastToRigidParticle(), History.DynamicsMisc, [this](auto Particle, const auto& Data) -> bool { Solver->SetParticleDynamicMisc(Particle, Data); return true; });
+			RewindHelper(PTParticle->CastToRigidParticle(), History.MassProps, [](auto Particle, const auto& Data) -> bool { Particle->SetMassProps(Data); return true; });
+		}
+
+		// If the particle is dynamic we must fix the collision anchors so that friction doesn't undo the movement
+		if (PTParticle->ObjectState() == EObjectStateType::Dynamic)
+		{
+			PTParticle->ParticleCollisions().VisitCollisions(
+				[this, PTParticle](FPBDCollisionConstraint& Collision)
+				{
+					Collision.UpdateParticleTransform(PTParticle);
+					return ECollisionVisitorResult::Continue;
+				});
+		}
+	}
+}
+
 template <bool bSkipDynamics, typename TDirtyInfo>
 void FRewindData::DesyncIfNecessary(TDirtyInfo& Info, const FFrameAndPhase FrameAndPhase)
 {
@@ -616,6 +788,11 @@ void FRewindData::DumpHistory_Internal(const int32 FramePrintOffset, const FStri
 	UE_LOG(LogChaos, Warning, TEXT("Saved:%s"), *Path);
 }
 
+bool FRewindData::GetUseCollisionResimCache() const
+{
+	return Solver ? Solver->GetUseCollisionResimCache() : false;
+}
+
 CHAOS_API int32 SkipDesyncTest = 0;
 FAutoConsoleVariableRef CVarSkipDesyncTest(TEXT("p.SkipDesyncTest"), SkipDesyncTest, TEXT("Skips hard desync test, this means all particles will assume to be clean except spawning at different times. This is useful for a perf lower bound, not actually correct"));
 
diff --git a/Engine/Source/Runtime/Experimental/Chaos/Public/Chaos/ChaosMarshallingManager.h b/Engine/Source/Runtime/Experimental/Chaos/Public/Chaos/ChaosMarshallingManager.h
--- a/Engine/Source/Runtime/Experimental/Chaos/Public/Chaos/ChaosMarshallingManager.h
+++ b/Engine/Source/Runtime/Experimental/Chaos/Public/Chaos/ChaosMarshallingManager.h
@@ -259,6 +259,7 @@ struct FPushPhysicsData
 
 	void Reset();
 	void ResetForHistory();
+	void ResetDirtyProxiesBuffer();
 	void CopySubstepData(const FPushPhysicsData& FirstStepData);
 };
 
diff --git a/Engine/Source/Runtime/Experimental/Chaos/Public/Chaos/Framework/PhysicsSolverBase.h b/Engine/Source/Runtime/Experimental/Chaos/Public/Chaos/Framework/PhysicsSolverBase.h
--- a/Engine/Source/Runtime/Experimental/Chaos/Public/Chaos/Framework/PhysicsSolverBase.h
+++ b/Engine/Source/Runtime/Experimental/Chaos/Public/Chaos/Framework/PhysicsSolverBase.h
@@ -445,6 +445,20 @@ namespace Chaos
 			return MRewindCallback.IsValid() && MRewindData.IsValid();
 		}
 
+		/** Enable or disable an additional resim cache based on IResimCacheBase (FEvolutionResimCache by default) that also caches particle collision constraints
+		* this is disabled by default since cached collisions will be wrong during a resimulation for particles that are desynced
+		* and particles that are in sync doesn't need collision data since we just step through their states in history." */
+		void SetUseCollisionResimCache(bool InUseCollisionResimCache)
+		{
+			bUseCollisionResimCache = InUseCollisionResimCache;
+		}
+
+		/** Check if an additional resim cache based on IResimCacheBase (FEvolutionResimCache by default) is being used. Read FPhysicsSolverBase.SetUseCollisionResimCache() for more info. */
+		bool GetUseCollisionResimCache() const
+		{
+			return bUseCollisionResimCache;
+		}
+
 		void SetPhysicsReplication(FPhysicsReplicationAsync* InPhysicsReplication)
 		{
 			PhysicsReplication = InPhysicsReplication;
@@ -738,6 +752,14 @@ namespace Chaos
 			return NetworkPhysicsEnabled;
 		}
 
+		/** Get the time length to cache physics history for, based on Project Settings -> Physics -> Physics Prediction -> MaxSupportedLatencyPrediction */
+		static float GetPhysicsHistoryTimeLength()
+		{
+			const float PhysicsHistoryTimeLength = FChaosSolversModule::GetModule()->GetSettingsProvider().GetPhysicsHistoryTimeLength();
+			return PhysicsHistoryTimeLength;
+		}
+
+
 		/** Get the number of physics history frames to cache */
 		static int32 GetPhysicsHistoryCount()
 		{
diff --git a/Engine/Source/Runtime/Experimental/Chaos/Public/ChaosSolversModule.h b/Engine/Source/Runtime/Experimental/Chaos/Public/ChaosSolversModule.h
--- a/Engine/Source/Runtime/Experimental/Chaos/Public/ChaosSolversModule.h
+++ b/Engine/Source/Runtime/Experimental/Chaos/Public/ChaosSolversModule.h
@@ -87,6 +87,11 @@ public:
 		return 0;
 	}
 
+	virtual float GetPhysicsHistoryTimeLength() const
+	{
+		return 0;
+	}
+
 	virtual int32 GetPhysicsHistoryCount() const
 	{
 		return 0;
diff --git a/Engine/Source/Runtime/Experimental/Chaos/Public/PBDRigidsSolver.h b/Engine/Source/Runtime/Experimental/Chaos/Public/PBDRigidsSolver.h
--- a/Engine/Source/Runtime/Experimental/Chaos/Public/PBDRigidsSolver.h
+++ b/Engine/Source/Runtime/Experimental/Chaos/Public/PBDRigidsSolver.h
@@ -229,6 +229,8 @@ namespace Chaos
 
 		CHAOS_API void EnableRewindCapture(int32 NumFrames, bool InUseCollisionResimCache, TUniquePtr<IRewindCallback>&& RewindCallback);
 		CHAOS_API void EnableRewindCapture(int32 NumFrames, bool InUseCollisionResimCache);
+		CHAOS_API void EnableRewindCapture(int32 NumFrames);
+		CHAOS_API void EnableRewindCapture();
 
 		/**/
 		FPBDRigidsEvolution* GetEvolution() { return MEvolution.Get(); }
@@ -289,8 +291,10 @@ namespace Chaos
 		/** Copy the simulation material list to the query material list, to be done when the SQ commits an update */
 		CHAOS_API void SyncQueryMaterials_External();
 
+		UE_DEPRECATED(5.5, "Deprecated, use GetUseCollisionResimCache() instead")
+		bool RewindUsesCollisionResimCache() const { return GetUseCollisionResimCache(); }
+
 		CHAOS_API void FinalizeRewindData(const TParticleView<FPBDRigidParticles>& DirtyParticles);
-		bool RewindUsesCollisionResimCache() const { return bUseCollisionResimCache; }
 
 		FPerSolverFieldSystem& GetPerSolverField() { return *PerSolverField; }
 		const FPerSolverFieldSystem& GetPerSolverField() const { return *PerSolverField; }
diff --git a/Engine/Source/Runtime/Experimental/Chaos/Public/RewindData.h b/Engine/Source/Runtime/Experimental/Chaos/Public/RewindData.h
--- a/Engine/Source/Runtime/Experimental/Chaos/Public/RewindData.h
+++ b/Engine/Source/Runtime/Experimental/Chaos/Public/RewindData.h
@@ -210,7 +210,7 @@ public :
 		}
 		else if(!bExactFrame)
 		{
-#if DEBUG_REWIND_DATA
+#if DEBUG_NETWORK_PHYSICS
 			if (bResetSolver)
 			{
 				UE_LOG(LogChaos, Warning, TEXT("		Unable to extract data at frame %d while rewinding the simulation"), ExtractFrame);
@@ -239,7 +239,7 @@ public :
 				ExtractedData.InputFrame = INDEX_NONE; // Clear InputFrame since this history entry is now altered and doesn't correspond to the source entry with the same InputFrame
 
 #if DEBUG_NETWORK_PHYSICS
-				UE_LOG(LogChaos, Log, TEXT("		Smoothing data between frame %d and %d - > [%d %d]"), DataHistory[MinFrameIndex].LocalFrame, DataHistory[MaxFrameIndex].LocalFrame, static_cast<DataType*>(HistoryData)->InputFrame, static_cast<DataType*>(HistoryData)->ServerFrame);
+				UE_LOG(LogChaos, Log, TEXT("		Interpolating data between frame %d and %d - > [%d %d]"), DataHistory[MinFrameIndex].LocalFrame, DataHistory[MaxFrameIndex].LocalFrame, static_cast<DataType*>(HistoryData)->InputFrame, static_cast<DataType*>(HistoryData)->ServerFrame);
 #endif
 				return true;
 			}
@@ -1504,25 +1504,24 @@ class FPBDRigidsSolver;
 class FRewindData
 {
 public:
-	FRewindData(FPBDRigidsSolver* InSolver, int32 NumFrames, bool InResimOptimization, int32 InCurrentFrame)
-	: Managers(NumFrames+1)	//give 1 extra for saving at head
-	, Solver(InSolver)
-	, CurFrame(InCurrentFrame)
-	, LatestFrame(InCurrentFrame)
-	, FramesSaved(0)
-	, DataIdxOffset(0)
-	, bNeedsSave(false)
-	, bResimOptimization(InResimOptimization)
-	, LatestTargetFrame(0)
+	FRewindData(FPBDRigidsSolver* InSolver, int32 NumFrames, bool InRewindDataOptimization, int32 InCurrentFrame);
+	FRewindData(FPBDRigidsSolver* InSolver, int32 NumFrames, int32 InCurrentFrame);
+
+	void Init(FPBDRigidsSolver* InSolver, int32 NumFrames, bool InRewindDataOptimization, int32 InCurrentFrame)
 	{
+		Solver = InSolver;
+		CurFrame = InCurrentFrame;
+		LatestFrame = InCurrentFrame;
+		bRewindDataOptimization = InRewindDataOptimization;
+		LatestTargetFrame = 0;
+		Managers = TCircularBuffer<FFrameManagerInfo>(NumFrames + 1);
 	}
 
-	void Init(FPBDRigidsSolver* InSolver, int32 NumFrames, bool InResimOptimization, int32 InCurrentFrame)
+	void Init(FPBDRigidsSolver* InSolver, int32 NumFrames, int32 InCurrentFrame)
 	{
 		Solver = InSolver;
 		CurFrame = InCurrentFrame;
 		LatestFrame = InCurrentFrame;
-		bResimOptimization = InResimOptimization;
 		LatestTargetFrame = 0;
 		Managers = TCircularBuffer<FFrameManagerInfo>(NumFrames + 1);
 	}
@@ -1595,12 +1594,14 @@ public:
 
 	IResimCacheBase* GetCurrentStepResimCache() const
 	{
-		const bool PhysicsPredictionEnabled = FPhysicsSolverBase::IsNetworkPhysicsPredictionEnabled();
-		return PhysicsPredictionEnabled && bResimOptimization ? Managers[CurFrame].ExternalResimCache.Get() : nullptr;
+		return Managers[CurFrame].ExternalResimCache.Get();
 	}
 
 	void CHAOS_API DumpHistory_Internal(const int32 FramePrintOffset, const FString& Filename = FString(TEXT("Dump")));
 
+	/** Check if a resim cache based on IResimCacheBase (FEvolutionResimCache by default) is being used. Read FPhysicsSolverBase.SetUseCollisionResimCache() for more info. */
+	bool GetUseCollisionResimCache() const;
+
 	template <typename CreateCache>
 	void AdvanceFrame(FReal DeltaTime, const CreateCache& CreateCacheFunc)
 	{
@@ -1609,7 +1610,7 @@ public:
 		Managers[CurFrame].FrameCreatedFor = CurFrame;
 		TUniquePtr<IResimCacheBase>& ResimCache = Managers[CurFrame].ExternalResimCache;
 
-		if (bResimOptimization)
+		if (GetUseCollisionResimCache())
 		{
 			if (IsResim())
 			{
@@ -1719,6 +1720,10 @@ public:
 	/** Get the latest frame resim has been blocked from rewinding past */
 	int32 GetBlockedResimFrame() { return BlockResimFrame; }
 
+	/** Set if RewindData optimizations should be enabled or not. 
+	* Effect: Only alter the minimum required properties during a resim for particles not marked for FullResim */
+	void SetRewindDataOptimization(bool InRewindDataOptimization) { bRewindDataOptimization = InRewindDataOptimization; }
+
 	/** Check if we have received targets already for the last frame simulated,
 	* if so compare those with the result of the simulation and if they desync return the frame value to request a rewind for to correct the desync
 	* NOTE: This only happens when the client has desynced and is behind the server, so we receive server states for frames not yet simulated */
@@ -1888,6 +1893,9 @@ private:
 	/** Apply targets positions and velocities while resimulating */
 	void ApplyTargets(const int32 Frame, const bool bResetSimulation);
 
+	/** Apply resim data for objects not simulating during resimlation */
+	void StepNonResimParticles(const int32 Frame);
+
 	template <typename TDirtyInfo>
 	static void DesyncObject(TDirtyInfo& Info, const FFrameAndPhase FrameAndPhase)
 	{
@@ -1911,7 +1919,7 @@ private:
 	int32 FramesSaved;
 	int32 DataIdxOffset;
 	bool bNeedsSave;	//Indicates that some data is pointing at head and requires saving before a rewind
-	bool bResimOptimization;
+	bool bRewindDataOptimization;
 	int32 ResimFrame = INDEX_NONE;
 	int32 LatestTargetFrame;
 
