Fix nanite geometry collection not visually rotating on gizmo driven rotations.
Needed to reapply modified instance transforms after primitive transform update.
This is because, for nanite, Geom Collection GPU instance transforms are stored relative to the relative (non-rotated) world transform.
This bug wasn't visible for non-gizmo use movement cases which go through full scene proxy recreate.
This bug wasn't visible for gizmo translation use case because this is in the relative transform.

diff --git a/Engine/Source/Runtime/Experimental/GeometryCollectionEngine/Private/GeometryCollection/GeometryCollectionSceneProxy.cpp b/Engine/Source/Runtime/Experimental/GeometryCollectionEngine/Private/GeometryCollection/GeometryCollectionSceneProxy.cpp
--- a/Engine/Source/Runtime/Experimental/GeometryCollectionEngine/Private/GeometryCollection/GeometryCollectionSceneProxy.cpp
+++ b/Engine/Source/Runtime/Experimental/GeometryCollectionEngine/Private/GeometryCollection/GeometryCollectionSceneProxy.cpp
@@ -1191,8 +1191,17 @@ FNaniteGeometryCollectionSceneProxy::FNaniteGeometryCollectionSceneProxy(UGeomet
 
 	CollisionResponse = Component->GetCollisionResponseToChannels();
 
-	FGeometryCollectionDynamicData* DynamicData = Component->InitDynamicData(true);
-	SetDynamicData_RenderThread(DynamicData, Component->GetRenderMatrix());
+	FGeometryCollectionDynamicData* InitDynamicData = Component->InitDynamicData(true);
+	SetDynamicData_RenderThread(InitDynamicData, Component->GetRenderMatrix());
+}
+
+FNaniteGeometryCollectionSceneProxy::~FNaniteGeometryCollectionSceneProxy()
+{
+	if (DynamicData != nullptr)
+	{
+		GDynamicDataPool.Release(DynamicData);
+		DynamicData = nullptr;
+	}
 }
 
 void FNaniteGeometryCollectionSceneProxy::CreateRenderThreadResources(FRHICommandListBase& RHICmdList)
@@ -1266,6 +1275,9 @@ uint32 FNaniteGeometryCollectionSceneProxy::GetMemoryFootprint() const
 
 void FNaniteGeometryCollectionSceneProxy::OnTransformChanged(FRHICommandListBase& RHICmdList)
 {
+	Super::OnTransformChanged(RHICmdList);
+
+	SetDynamicData_RenderThread(DynamicData, GetLocalToWorld());
 }
 
 void FNaniteGeometryCollectionSceneProxy::GetNaniteResourceInfo(uint32& ResourceID, uint32& HierarchyOffset, uint32& ImposterIndex) const
@@ -1298,6 +1310,16 @@ Nanite::FResourceMeshInfo FNaniteGeometryCollectionSceneProxy::GetResourceMeshIn
 
 void FNaniteGeometryCollectionSceneProxy::SetDynamicData_RenderThread(FGeometryCollectionDynamicData* NewDynamicData, const FMatrix &PrimitiveLocalToWorld)
 {
+	if (NewDynamicData != DynamicData)
+	{
+		if (DynamicData)
+		{
+			GDynamicDataPool.Release(DynamicData);
+			DynamicData = nullptr;
+		}
+		DynamicData = NewDynamicData;
+	}
+
 	// Are we currently simulating?
 	if (NewDynamicData->IsDynamic)
 	{
@@ -1369,8 +1391,6 @@ void FNaniteGeometryCollectionSceneProxy::SetDynamicData_RenderThread(FGeometryC
 		// Rendering base geometry, use rest transforms rather than simulated transforms.
 		// ...
 	}
-
-	GDynamicDataPool.Release(NewDynamicData);
 }
 
 void FNaniteGeometryCollectionSceneProxy::ResetPreviousTransforms_RenderThread()
diff --git a/Engine/Source/Runtime/Experimental/GeometryCollectionEngine/Private/GeometryCollection/GeometryCollectionSceneProxy.h b/Engine/Source/Runtime/Experimental/GeometryCollectionEngine/Private/GeometryCollection/GeometryCollectionSceneProxy.h
--- a/Engine/Source/Runtime/Experimental/GeometryCollectionEngine/Private/GeometryCollection/GeometryCollectionSceneProxy.h
+++ b/Engine/Source/Runtime/Experimental/GeometryCollectionEngine/Private/GeometryCollection/GeometryCollectionSceneProxy.h
@@ -336,7 +336,7 @@ public:
 	using Super = Nanite::FSceneProxyBase;
 	
 	FNaniteGeometryCollectionSceneProxy(UGeometryCollectionComponent* Component);
-	virtual ~FNaniteGeometryCollectionSceneProxy() = default;
+	virtual ~FNaniteGeometryCollectionSceneProxy();
 
 public:
 	// FPrimitiveSceneProxy interface.
@@ -416,6 +416,8 @@ protected:
 
 	FInstanceSceneDataBuffers InstanceSceneDataBuffersImpl;
 
+	FGeometryCollectionDynamicData* DynamicData = nullptr;
+
 	// Geometry collection doesn't currently support baked light maps, so we use this simple empty light cache info for all nanite geometry collection proxies
 	class FEmptyLightCacheInfo : public FLightCacheInterface
 	{
