Add AllowWrite flag to FCas::OpenAsyncRead and remove special case for reading from current block.

diff --git a/Engine/Source/Runtime/Experimental/IoStore/OnDemand/Private/OnDemandInstallCache.cpp b/Engine/Source/Runtime/Experimental/IoStore/OnDemand/Private/OnDemandInstallCache.cpp
--- a/Engine/Source/Runtime/Experimental/IoStore/OnDemand/Private/OnDemandInstallCache.cpp
+++ b/Engine/Source/Runtime/Experimental/IoStore/OnDemand/Private/OnDemandInstallCache.cpp
@@ -31,15 +31,6 @@ namespace UE::IoStore
 ///////////////////////////////////////////////////////////////////////////////
 namespace CVars
 {
-
-bool bForceSyncIO = false;
-static FAutoConsoleVariableRef CVarForceSyncIO(
-	TEXT("IoStore.OnDemand.ForceSyncIO"),
-	bForceSyncIO,
-	TEXT("Whether to force using synchronous file reads even if cache block is immutable"),
-	ECVF_ReadOnly
-);
-
 }
 
 ///////////////////////////////////////////////////////////////////////////////
@@ -58,6 +49,7 @@ double ToMiB(uint64 Value)
 using FSharedAsyncFileHandle	= TSharedPtr<IAsyncReadFileHandle, ESPMode::ThreadSafe>;
 using FWeakAsyncFileHandle		= TWeakPtr<IAsyncReadFileHandle, ESPMode::ThreadSafe>;
 using FUniqueFileHandle			= TUniquePtr<IFileHandle>;
+using FSharedFileOpenResult		= TValueOrError<FSharedAsyncFileHandle, FFileSystemError>;
 using FCasAddr					= FHash96;
 
 ///////////////////////////////////////////////////////////////////////////////
@@ -122,7 +114,7 @@ struct FCas
 	FIoStatus				DeleteBlock(FCasBlockId BlockId, TArray<FCasAddr>& OutAddrs);
 	FString					GetBlockFilename(FCasBlockId BlockId) const;
 	FFileOpenResult 		OpenRead(FCasBlockId BlockId);
-	FSharedAsyncFileHandle	OpenAsyncRead(FCasBlockId  BlockId);
+	FSharedFileOpenResult	OpenAsyncRead(FCasBlockId  BlockId);
 	FUniqueFileHandle		OpenWrite(FCasBlockId BlockId);
 	void					TrackAccess(FCasBlockId BlockId, int64 UtcTicks);
 	void					TrackAccess(FCasBlockId BlockId) { TrackAccess(BlockId, FDateTime::UtcNow().GetTicks()); }
@@ -263,7 +255,7 @@ FFileOpenResult FCas::OpenRead(FCasBlockId BlockId)
 	return Ipf.OpenRead(*Filename, IPlatformFile::EOpenReadFlags::AllowWrite);
 }
 
-FSharedAsyncFileHandle FCas::OpenAsyncRead(FCasBlockId BlockId)
+FSharedFileOpenResult FCas::OpenAsyncRead(FCasBlockId BlockId)
 {
 	UE::TUniqueLock Lock(Mutex);
 
@@ -271,20 +263,25 @@ FSharedAsyncFileHandle FCas::OpenAsyncRead(FCasBlockId BlockId)
 	{
 		if (FSharedAsyncFileHandle Handle = MaybeHandle->Pin(); Handle.IsValid())
 		{
-			return Handle;
+			return MakeValue(MoveTemp(Handle));
 		}
 	}
 
 	IPlatformFile&			Ipf = FPlatformFileManager::Get().GetPlatformFile();
 	const FString			Filename = GetBlockFilename(BlockId);
-	FSharedAsyncFileHandle	NewHandle(Ipf.OpenAsyncRead(*Filename));
+	FFileOpenAsyncResult	HandleResult(Ipf.OpenAsyncRead(*Filename, IPlatformFile::EOpenReadFlags::AllowWrite));
 
-	if (NewHandle.IsValid())
+	if (HandleResult.HasValue())
 	{
-		ReadHandles.FindOrAdd(BlockId, NewHandle);
+		FSharedAsyncFileHandle NewHandle(HandleResult.GetValue().Release());
+		if (NewHandle.IsValid())
+		{
+			ReadHandles.FindOrAdd(BlockId, NewHandle);
+		}
+		return MakeValue(MoveTemp(NewHandle));
 	}
 
-	return NewHandle;
+	return MakeError(HandleResult.StealError());
 }
 
 FUniqueFileHandle FCas::OpenWrite(FCasBlockId BlockId)
@@ -1210,69 +1207,16 @@ bool FOnDemandInstallCache::Resolve(FIoRequestImpl* Request)
 	TRACE_IOSTORE_BACKEND_REQUEST_STARTED(Request, this);
 	Cas.TrackAccess(CasLoc.BlockId);
 
-	// Use synchronous file read API when reading from and writing to the same cache block concurrently.
-	const bool bIsLocationInCurrentBlock = CasLoc.BlockId == CurrentBlock;
-	const bool bSyncRead = bIsLocationInCurrentBlock || CVars::bForceSyncIO;
-	if (bSyncRead)
-	{
-		// The internal request parameters are attached/owned by the I/O request via
-		// the backend data parameter. The chunk request is deleted in GetCompletedRequests
-		FChunkRequest::Attach(*Request, new FChunkRequest(
-			FSharedAsyncFileHandle(),
-			Request,
-			MoveTemp(ChunkInfo),
-			ChunkRange.ConsumeValueOrDie(),
-			RequestSize));
-
-		UE::Tasks::Launch(UE_SOURCE_LOCATION, [this, Request, CasLoc]
-		{
-			FChunkRequest& ChunkRequest = FChunkRequest::GetRef(*Request);
-			bool bOk = false;
-
-			FFileOpenResult FileOpenResult = Cas.OpenRead(CasLoc.BlockId);
-			if (FileOpenResult.IsValid())
-			{
-				TUniquePtr<IFileHandle> FileHandle = FileOpenResult.StealValue();
-				const int64 CasBlockOffset = CasLoc.BlockOffset + ChunkRequest.ChunkRange.GetOffset();
-				if (Request->IsCancelled() == false && FileHandle->Seek(CasBlockOffset))
-				{
-					bOk = FileHandle->Read(ChunkRequest.EncodedChunk.GetData(), ChunkRequest.EncodedChunk.GetSize());
-					if (bOk == false)
-					{
-						const FString Filename = Cas.GetBlockFilename(CasLoc.BlockId);
-						UE_LOG(LogIoStoreOnDemand, Error, TEXT("Failed to read %llu bytes at offset %lld in CAS block '%s'"),
-							ChunkRequest.EncodedChunk.GetSize(),
-							CasBlockOffset,
-							*Filename);
-					}
-				}
-				else
-				{
-					const FString Filename = Cas.GetBlockFilename(CasLoc.BlockId);
-					UE_LOG(LogIoStoreOnDemand, Error, TEXT("Failed to seek to offset %lld in CAS block '%s'"), CasBlockOffset, *Filename);
-				}
-			}
-			else
-			{
-				const FString Filename = Cas.GetBlockFilename(CasLoc.BlockId);
-				UE_LOG(LogIoStoreOnDemand, Error, TEXT("Failed to open CAS block '%s' for reading, reason '%s'"),
-					*Filename, *FileOpenResult.GetError().GetMessage());
-			}
-
-			const bool bWasCancelled = bOk == false;
-			CompleteRequest(Request, bWasCancelled);
-		});
-
-		return true;
-	}
-
-	FSharedAsyncFileHandle FileHandle = Cas.OpenAsyncRead(CasLoc.BlockId);
-	if (FileHandle.IsValid() == false)
+	FSharedFileOpenResult FileOpenResult = Cas.OpenAsyncRead(CasLoc.BlockId);
+	if (FileOpenResult.HasError())
 	{
 		const FString Filename = Cas.GetBlockFilename(CasLoc.BlockId);
-		UE_LOG(LogIoStoreOnDemand, Error, TEXT("Failed to open CAS block '%s' for async reading"), *Filename);
+		UE_LOG(LogIoStoreOnDemand, Error, TEXT("Failed to open CAS block '%s' for async reading, reason '%s'"), *Filename, *FileOpenResult.GetError().GetMessage());
 		TUniquePtr<FChunkRequest> Detached = FChunkRequest::Detach(*Request);
+		return false;
 	}
+	
+	FSharedAsyncFileHandle FileHandle(FileOpenResult.StealValue());
 
 	// The internal request parameters are attached/owned by the I/O request via
 	// the backend data parameter. The chunk request is deleted in GetCompletedRequests
