IAS' cache governor correctly regulates when op-bound

#rnx

diff --git a/Engine/Source/Runtime/Experimental/IoStore/OnDemand/Private/JournaledCache.cpp b/Engine/Source/Runtime/Experimental/IoStore/OnDemand/Private/JournaledCache.cpp
--- a/Engine/Source/Runtime/Experimental/IoStore/OnDemand/Private/JournaledCache.cpp
+++ b/Engine/Source/Runtime/Experimental/IoStore/OnDemand/Private/JournaledCache.cpp
@@ -1511,7 +1511,17 @@ void FGovernor::Set(uint32 Allowance, uint32 Ops, uint32 Seconds, int64 CycleFre
 		return;
 	}
 
-	int32 BlockCount = Allowance / CommitBufferSize;
+	// A small commit buffer will make us op bound
+	if (uint32 CommitWidth = CommitBufferSize * Ops; CommitWidth < Allowance)
+	{
+		OpAllowance = CommitBufferSize;
+		OpInterval = (Seconds * CycleFreq * 3) / Ops;
+		MaxOpCount = 1;
+		return;
+	}
+
+	// Allowance bound.
+	int32 BlockCount = FMath::Max<int32>(1, Allowance / CommitBufferSize);
 	int32 CommitOpCost = BlockCount * 3;
 	MaxOpCount = (Ops - CommitOpCost) / BlockCount;
 
