Iris
* Added PostInit callback to UNetObjectFactory
* Moved PostInit code from bridges to their respective factory

#rnx

diff --git a/Engine/Plugins/Runtime/ReplicationSystemTestPlugin/Source/Private/Tests/ReplicationSystem/ReplicatedTestObject.cpp b/Engine/Plugins/Runtime/ReplicationSystemTestPlugin/Source/Private/Tests/ReplicationSystem/ReplicatedTestObject.cpp
--- a/Engine/Plugins/Runtime/ReplicationSystemTestPlugin/Source/Private/Tests/ReplicationSystem/ReplicatedTestObject.cpp
+++ b/Engine/Plugins/Runtime/ReplicationSystemTestPlugin/Source/Private/Tests/ReplicationSystem/ReplicatedTestObject.cpp
@@ -109,13 +109,6 @@ void UReplicatedTestObjectBridge::EndReplication(UReplicatedTestObject* Instance
 	StopReplicatingNetObject(Instance, Flags);
 }
 
-void UReplicatedTestObjectBridge::EndInstantiateFromRemote(FNetRefHandle Handle)
-{
-	UReplicatedTestObject* Instance = CastChecked<UReplicatedTestObject>(GetReplicatedObject(Handle));
-
-	Instance->NetRefHandle = Handle;
-}
-
 void UReplicatedTestObjectBridge::DestroyInstanceFromRemote(const FDestroyInstanceParams& Params)
 {
 	if (!Params.Instance)
diff --git a/Engine/Plugins/Runtime/ReplicationSystemTestPlugin/Source/Private/Tests/ReplicationSystem/ReplicatedTestObject.h b/Engine/Plugins/Runtime/ReplicationSystemTestPlugin/Source/Private/Tests/ReplicationSystem/ReplicatedTestObject.h
--- a/Engine/Plugins/Runtime/ReplicationSystemTestPlugin/Source/Private/Tests/ReplicationSystem/ReplicatedTestObject.h
+++ b/Engine/Plugins/Runtime/ReplicationSystemTestPlugin/Source/Private/Tests/ReplicationSystem/ReplicatedTestObject.h
@@ -583,7 +583,6 @@ protected:
 
 	friend FSupressCreateInstanceFailedEnsureScope;
 
-	virtual void EndInstantiateFromRemote(FNetRefHandle Handle) override;
 	virtual void DestroyInstanceFromRemote(const FDestroyInstanceParams& Params) override;
 	virtual bool IsAllowedToDestroyInstance(const UObject* Instance) const override;
 
diff --git a/Engine/Plugins/Runtime/ReplicationSystemTestPlugin/Source/Private/Tests/ReplicationSystem/ReplicatedTestObjectFactory.cpp b/Engine/Plugins/Runtime/ReplicationSystemTestPlugin/Source/Private/Tests/ReplicationSystem/ReplicatedTestObjectFactory.cpp
--- a/Engine/Plugins/Runtime/ReplicationSystemTestPlugin/Source/Private/Tests/ReplicationSystem/ReplicatedTestObjectFactory.cpp
+++ b/Engine/Plugins/Runtime/ReplicationSystemTestPlugin/Source/Private/Tests/ReplicationSystem/ReplicatedTestObjectFactory.cpp
@@ -178,3 +178,10 @@ UNetObjectFactory::FInstantiateResult UReplicatedTestObjectFactory::InstantiateR
 	InstantiateResult.Flags |= EReplicationBridgeCreateNetRefHandleResultFlags::AllowDestroyInstanceFromRemote;
 	return InstantiateResult;
 }
+
+void UReplicatedTestObjectFactory::PostInit(const FPostInitContext& Context)
+{
+	UReplicatedTestObject* Instance = CastChecked<UReplicatedTestObject>(Context.Instance);
+
+	Instance->NetRefHandle = Context.Handle;
+}
\ No newline at end of file
diff --git a/Engine/Plugins/Runtime/ReplicationSystemTestPlugin/Source/Private/Tests/ReplicationSystem/ReplicatedTestObjectFactory.h b/Engine/Plugins/Runtime/ReplicationSystemTestPlugin/Source/Private/Tests/ReplicationSystem/ReplicatedTestObjectFactory.h
--- a/Engine/Plugins/Runtime/ReplicationSystemTestPlugin/Source/Private/Tests/ReplicationSystem/ReplicatedTestObjectFactory.h
+++ b/Engine/Plugins/Runtime/ReplicationSystemTestPlugin/Source/Private/Tests/ReplicationSystem/ReplicatedTestObjectFactory.h
@@ -39,6 +39,8 @@ public:
 
 	virtual FInstantiateResult InstantiateReplicatedObjectFromHeader(const FInstantiateContext& Context, const UE::Net::FNetObjectCreationHeader* InHeader) override;
 
+	virtual void PostInit(const FPostInitContext& Context) override;
+
 protected:
 
 	virtual TUniquePtr<UE::Net::FNetObjectCreationHeader> CreateAndFillHeader(UE::Net::FNetRefHandle Handle) override;
diff --git a/Engine/Source/Runtime/Engine/Private/Net/Iris/ReplicationSystem/EngineReplicationBridge.cpp b/Engine/Source/Runtime/Engine/Private/Net/Iris/ReplicationSystem/EngineReplicationBridge.cpp
--- a/Engine/Source/Runtime/Engine/Private/Net/Iris/ReplicationSystem/EngineReplicationBridge.cpp
+++ b/Engine/Source/Runtime/Engine/Private/Net/Iris/ReplicationSystem/EngineReplicationBridge.cpp
@@ -523,17 +523,6 @@ void UEngineReplicationBridge::StopReplicatingComponent(UActorComponent* ActorCo
 	}
 }
 
-void UEngineReplicationBridge::EndInstantiateFromRemote(FNetRefHandle Handle)
-{
-	if (AActor* Actor = Cast<AActor>(GetReplicatedObject(Handle)))
-	{
-		// Optional
-		//Actor->NetHandle = Handle;
-
-		Actor->PostNetInit();
-	}
-}
-
 //$IRIS todo: move this to be a factory callback
 void UEngineReplicationBridge::OnSubObjectCreatedFromReplication(FNetRefHandle SubObjectHandle)
 {
diff --git a/Engine/Source/Runtime/Engine/Private/Net/Iris/ReplicationSystem/NetActorFactory.cpp b/Engine/Source/Runtime/Engine/Private/Net/Iris/ReplicationSystem/NetActorFactory.cpp
--- a/Engine/Source/Runtime/Engine/Private/Net/Iris/ReplicationSystem/NetActorFactory.cpp
+++ b/Engine/Source/Runtime/Engine/Private/Net/Iris/ReplicationSystem/NetActorFactory.cpp
@@ -470,6 +470,11 @@ void UNetActorFactory::PostInstantiation(const FPostInstantiationContext& Contex
 	ActorBridge->WakeUpObjectInstantiatedFromRemote(Actor);
 }
 
+void UNetActorFactory::PostInit(const FPostInitContext& Context)
+{
+	AActor* Actor = CastChecked<AActor>(Context.Instance);
+	Actor->PostNetInit();
+}
 
  namespace UE::Net
  {
diff --git a/Engine/Source/Runtime/Engine/Public/Net/Iris/ReplicationSystem/EngineReplicationBridge.h b/Engine/Source/Runtime/Engine/Public/Net/Iris/ReplicationSystem/EngineReplicationBridge.h
--- a/Engine/Source/Runtime/Engine/Public/Net/Iris/ReplicationSystem/EngineReplicationBridge.h
+++ b/Engine/Source/Runtime/Engine/Public/Net/Iris/ReplicationSystem/EngineReplicationBridge.h
@@ -104,7 +104,6 @@ protected:
 	virtual void Initialize(UReplicationSystem* ReplicationSystem) override;
 	virtual void Deinitialize() override;
 	virtual void OnSubObjectCreatedFromReplication(FNetRefHandle SubObjectHandle) override;
-	virtual void EndInstantiateFromRemote(FNetRefHandle Handle) override;
 	virtual void DestroyInstanceFromRemote(const FDestroyInstanceParams& Params) override;
 	virtual void GetInitialDependencies(FNetRefHandle Handle, FNetDependencyInfoArray& OutDependencies) const override;
 	virtual bool RemapPathForPIE(uint32 ConnectionId, FString& Path, bool bReading) const override;
diff --git a/Engine/Source/Runtime/Engine/Public/Net/Iris/ReplicationSystem/NetActorFactory.h b/Engine/Source/Runtime/Engine/Public/Net/Iris/ReplicationSystem/NetActorFactory.h
--- a/Engine/Source/Runtime/Engine/Public/Net/Iris/ReplicationSystem/NetActorFactory.h
+++ b/Engine/Source/Runtime/Engine/Public/Net/Iris/ReplicationSystem/NetActorFactory.h
@@ -121,6 +121,8 @@ public:
 
 	virtual void PostInstantiation(const FPostInstantiationContext& Context) override;
 
+	virtual void PostInit(const FPostInitContext& Context) override;
+
 protected:
 
 	virtual TUniquePtr<UE::Net::FNetObjectCreationHeader> CreateAndFillHeader(UE::Net::FNetRefHandle Handle) override;
diff --git a/Engine/Source/Runtime/Experimental/Iris/Core/Private/Iris/ReplicationSystem/NetRefHandleManager.cpp b/Engine/Source/Runtime/Experimental/Iris/Core/Private/Iris/ReplicationSystem/NetRefHandleManager.cpp
--- a/Engine/Source/Runtime/Experimental/Iris/Core/Private/Iris/ReplicationSystem/NetRefHandleManager.cpp
+++ b/Engine/Source/Runtime/Experimental/Iris/Core/Private/Iris/ReplicationSystem/NetRefHandleManager.cpp
@@ -360,13 +360,15 @@ FNetRefHandle FNetRefHandleManager::CreateNetObject(FNetRefHandle WantedHandle,
 }
 
 // Create NetRefHandle not owned by us
-FNetRefHandle FNetRefHandleManager::CreateNetObjectFromRemote(FNetRefHandle WantedHandle, const FReplicationProtocol* ReplicationProtocol)
+FNetRefHandle FNetRefHandleManager::CreateNetObjectFromRemote(FNetRefHandle WantedHandle, const FReplicationProtocol* ReplicationProtocol,  UE::Net::FNetObjectFactoryId FactoryId)
 {
 	if (!ensureMsgf(WantedHandle.IsValid() && !WantedHandle.IsCompleteHandle(), TEXT("FNetRefHandleManager::CreateNetObjectFromRemote Expected WantedHandle %s to be valid and incomplete"), *WantedHandle.ToString()))
 	{
 		return FNetRefHandle();
 	}
 
+	check(FactoryId != UE::Net::InvalidNetObjectFactoryId);
+
 	FNetRefHandle NetRefHandle = MakeNetRefHandle(WantedHandle.GetId(), ReplicationSystemId);
 
 	const FInternalNetRefIndex InternalIndex = InternalCreateNetObject(NetRefHandle, FNetHandle(), ReplicationProtocol);
@@ -374,6 +376,8 @@ FNetRefHandle FNetRefHandleManager::CreateNetObjectFromRemote(FNetRefHandle Want
 	{
 		FReplicatedObjectData& Data = ReplicatedObjectData[InternalIndex];
 
+		Data.NetFactoryId = FactoryId;
+
 		// Allocate storage for incoming data
 		Data.ReceiveStateBuffer = (uint8*)FMemory::Malloc(ReplicationProtocol->InternalTotalSize, ReplicationProtocol->InternalTotalAlignment);
 
diff --git a/Engine/Source/Runtime/Experimental/Iris/Core/Private/Iris/ReplicationSystem/NetRefHandleManager.h b/Engine/Source/Runtime/Experimental/Iris/Core/Private/Iris/ReplicationSystem/NetRefHandleManager.h
--- a/Engine/Source/Runtime/Experimental/Iris/Core/Private/Iris/ReplicationSystem/NetRefHandleManager.h
+++ b/Engine/Source/Runtime/Experimental/Iris/Core/Private/Iris/ReplicationSystem/NetRefHandleManager.h
@@ -153,7 +153,7 @@ public:
 	FNetRefHandle CreateNetObject(FNetRefHandle WantedHandle, FNetHandle GlobalHandle, const FReplicationProtocol* ReplicationProtocol);
 
 	// Create NetObject on request from remote
-	FNetRefHandle CreateNetObjectFromRemote(FNetRefHandle WantedHandle, const FReplicationProtocol* ReplicationProtocol);
+	FNetRefHandle CreateNetObjectFromRemote(FNetRefHandle WantedHandle, const FReplicationProtocol* ReplicationProtocol, FNetObjectFactoryId FactoryId);
 
 	// Attach Instance protocol to handle
 	// Instance can be null, we only track the Instance for legacy support
diff --git a/Engine/Source/Runtime/Experimental/Iris/Core/Private/Iris/ReplicationSystem/ObjectReplicationBridge.cpp b/Engine/Source/Runtime/Experimental/Iris/Core/Private/Iris/ReplicationSystem/ObjectReplicationBridge.cpp
--- a/Engine/Source/Runtime/Experimental/Iris/Core/Private/Iris/ReplicationSystem/ObjectReplicationBridge.cpp
+++ b/Engine/Source/Runtime/Experimental/Iris/Core/Private/Iris/ReplicationSystem/ObjectReplicationBridge.cpp
@@ -918,7 +918,7 @@ FReplicationBridgeCreateNetRefHandleResult UObjectReplicationBridge::CreateNetRe
 	else
 	{
 		// Create NetHandle
-		FNetRefHandle Handle = InternalCreateNetObjectFromRemote(WantedNetHandle, ReplicationProtocol);
+		FNetRefHandle Handle = InternalCreateNetObjectFromRemote(WantedNetHandle, ReplicationProtocol, FactoryId);
 		CreateResult.NetRefHandle = Handle;
 		if (Handle.IsValid())
 		{
@@ -939,9 +939,20 @@ void UObjectReplicationBridge::SubObjectCreatedFromReplication(FNetRefHandle Han
 	OnSubObjectCreatedFromReplication(Handle);
 }
 
-void UObjectReplicationBridge::PostApplyInitialState(FNetRefHandle Handle)
+void UObjectReplicationBridge::PostApplyInitialState(UE::Net::Private::FInternalNetRefIndex InternalObjectIndex)
 {
-	EndInstantiateFromRemote(Handle);
+	using namespace UE::Net;
+	using namespace UE::Net::Private;
+
+	const FNetObjectFactoryId FactoryId = NetRefHandleManager->GetReplicatedObjectDataNoCheck(InternalObjectIndex).NetFactoryId;
+	UNetObjectFactory* Factory = GetNetFactory(FactoryId);
+
+	const UNetObjectFactory::FPostInitContext Context 
+	{ 
+		.Instance = NetRefHandleManager->GetReplicatedObjectInstance(InternalObjectIndex),
+		.Handle = NetRefHandleManager->GetNetRefHandleFromInternalIndex(InternalObjectIndex)
+	};
+	Factory->PostInit(Context);
 }
 
 void UObjectReplicationBridge::PreSendUpdateSingleHandle(FNetRefHandle RefHandle)
diff --git a/Engine/Source/Runtime/Experimental/Iris/Core/Private/Iris/ReplicationSystem/ReplicationBridge.cpp b/Engine/Source/Runtime/Experimental/Iris/Core/Private/Iris/ReplicationSystem/ReplicationBridge.cpp
--- a/Engine/Source/Runtime/Experimental/Iris/Core/Private/Iris/ReplicationSystem/ReplicationBridge.cpp
+++ b/Engine/Source/Runtime/Experimental/Iris/Core/Private/Iris/ReplicationSystem/ReplicationBridge.cpp
@@ -97,10 +97,6 @@ void UReplicationBridge::SubObjectCreatedFromReplication(FNetRefHandle SubObject
 {
 }
 
-void UReplicationBridge::PostApplyInitialState(FNetRefHandle Handle)
-{
-}
-
 void UReplicationBridge::DetachInstanceFromRemote(FNetRefHandle Handle, EReplicationBridgeDestroyInstanceReason DestroyReason, EReplicationBridgeDestroyInstanceFlags DestroyFlags)
 {
 }
@@ -316,9 +312,9 @@ void UReplicationBridge::CallSubObjectCreatedFromReplication(FNetRefHandle SubOb
 	SubObjectCreatedFromReplication(SubObjectHandle);
 }
 
-void UReplicationBridge::CallPostApplyInitialState(FNetRefHandle Handle)
+void UReplicationBridge::CallPostApplyInitialState(UE::Net::Private::FInternalNetRefIndex InternalObjectIndex)
 {
-	PostApplyInitialState(Handle);
+	PostApplyInitialState(InternalObjectIndex);
 }
 
 UReplicationBridge::~UReplicationBridge()
@@ -413,9 +409,9 @@ UE::Net::FNetRefHandle UReplicationBridge::InternalCreateNetObject(FNetRefHandle
 	return InternalCreateNetObject(AllocatedHandle, FNetHandle(), ReplicationProtocol);
 }
 
-UE::Net::FNetRefHandle UReplicationBridge::InternalCreateNetObjectFromRemote(FNetRefHandle WantedNetHandle, const UE::Net::FReplicationProtocol* ReplicationProtocol)
+UE::Net::FNetRefHandle UReplicationBridge::InternalCreateNetObjectFromRemote(FNetRefHandle WantedNetHandle, const UE::Net::FReplicationProtocol* ReplicationProtocol, UE::Net::FNetObjectFactoryId FactoryId)
 {
-	FNetRefHandle Handle = NetRefHandleManager->CreateNetObjectFromRemote(WantedNetHandle, ReplicationProtocol);
+	FNetRefHandle Handle = NetRefHandleManager->CreateNetObjectFromRemote(WantedNetHandle, ReplicationProtocol, FactoryId);
 
 	if (Handle.IsValid())
 	{
diff --git a/Engine/Source/Runtime/Experimental/Iris/Core/Private/Iris/ReplicationSystem/ReplicationReader.cpp b/Engine/Source/Runtime/Experimental/Iris/Core/Private/Iris/ReplicationSystem/ReplicationReader.cpp
--- a/Engine/Source/Runtime/Experimental/Iris/Core/Private/Iris/ReplicationSystem/ReplicationReader.cpp
+++ b/Engine/Source/Runtime/Experimental/Iris/Core/Private/Iris/ReplicationSystem/ReplicationReader.cpp
@@ -1800,16 +1800,15 @@ void FReplicationReader::DispatchStateData(FNetSerializationContext& Context)
 			}
 		}
 
-		// In the last pass, RPC`s and cleanup cached data
+		// In the last pass, RPC's and cleanup cached data
 		for (FPostDispatchObjectInfo& PostDispatchObjectInfo : MakeArrayView(PostDispatchObjectInfos, NumObjectsPendingPostDistpatch))
 		{
 			FDispatchObjectInfo& Info = *PostDispatchObjectInfo.Info;
 
-			// If object was created this frame it`s initial state is now applied
+			// If the object was created this frame it's initial state is now applied
 			if (Info.bHasState && Info.bIsInitialState)
 			{
-				const FNetRefHandleManager::FReplicatedObjectData& ObjectData = NetRefHandleManager->GetReplicatedObjectDataNoCheck(Info.InternalIndex);
-				ReplicationBridge->CallPostApplyInitialState(ObjectData.RefHandle);
+				ReplicationBridge->CallPostApplyInitialState(Info.InternalIndex);
 			}
 
 			// Dispatch attachment and enqueue for later resolving
diff --git a/Engine/Source/Runtime/Experimental/Iris/Core/Public/Iris/ReplicationSystem/NetObjectFactory.h b/Engine/Source/Runtime/Experimental/Iris/Core/Public/Iris/ReplicationSystem/NetObjectFactory.h
--- a/Engine/Source/Runtime/Experimental/Iris/Core/Public/Iris/ReplicationSystem/NetObjectFactory.h
+++ b/Engine/Source/Runtime/Experimental/Iris/Core/Public/Iris/ReplicationSystem/NetObjectFactory.h
@@ -84,6 +84,7 @@ public:
 	struct FInstantiateResult;
 	struct FInstantiateContext;
 	struct FPostInstantiationContext;
+	struct FPostInitContext;
 
 	/**
 	 * Creates the header containing all information required to instantiate a remote version of the object represented by the handle.
@@ -124,6 +125,13 @@ public:
 	 */
 	virtual void PostInstantiation(const FPostInstantiationContext& Context) {}
 
+	/**
+	 * Optional callback triggered after we applied the initial replicated properties to the instantiated object.
+	 * From here the remote object is ready to be used by the game engine.
+	 * @param Context Gives you access to the new object.
+	 */
+	virtual void PostInit(const FPostInitContext& Context) {}
+
 public:
 
 	/** Result of the instantiate request */
@@ -164,6 +172,15 @@ public:
 		uint32 ConnectionId = 0;
 	};
 
+	/** Contextual information to use in the PostInit callback */
+	struct FPostInitContext
+	{
+		/** The object instantiated */
+		UObject* Instance = nullptr;
+		/** The handle of the object */
+		UE::Net::FNetRefHandle Handle;
+	};
+
 protected:
 
 	virtual void OnInit() {}
diff --git a/Engine/Source/Runtime/Experimental/Iris/Core/Public/Iris/ReplicationSystem/ObjectReplicationBridge.h b/Engine/Source/Runtime/Experimental/Iris/Core/Public/Iris/ReplicationSystem/ObjectReplicationBridge.h
--- a/Engine/Source/Runtime/Experimental/Iris/Core/Public/Iris/ReplicationSystem/ObjectReplicationBridge.h
+++ b/Engine/Source/Runtime/Experimental/Iris/Core/Public/Iris/ReplicationSystem/ObjectReplicationBridge.h
@@ -243,7 +243,7 @@ protected:
 	IRISCORE_API virtual bool CacheNetRefHandleCreationInfo(FNetRefHandle Handle) override;
 	IRISCORE_API virtual FReplicationBridgeCreateNetRefHandleResult CreateNetRefHandleFromRemote(FNetRefHandle RootObjectNetHandle, FNetRefHandle WantedNetHandle, FReplicationBridgeSerializationContext& Context) override;
 	IRISCORE_API virtual void SubObjectCreatedFromReplication(FNetRefHandle SubObjectHandle) override;
-	IRISCORE_API virtual void PostApplyInitialState(FNetRefHandle Handle) override;
+	IRISCORE_API virtual void PostApplyInitialState(UE::Net::Private::FInternalNetRefIndex InternalObjectIndex) final;
 	IRISCORE_API virtual void DetachInstanceFromRemote(FNetRefHandle Handle, EReplicationBridgeDestroyInstanceReason DestroyReason, EReplicationBridgeDestroyInstanceFlags DestroyFlags) override;
 	IRISCORE_API virtual void DetachInstance(FNetRefHandle Handle) override;
 	IRISCORE_API virtual void OnProtocolMismatchReported(FNetRefHandle RefHandle, uint32 ConnectionId) override;
@@ -254,9 +254,6 @@ protected:
 	/** Invoked for new replicated SubObjects after state has been applied to owner */
 	virtual void OnSubObjectCreatedFromReplication(FNetRefHandle SubObjectHandle) {};
 
-	/** Invoked after remote NetHandle has been created and initial state is applied. */
-	virtual void EndInstantiateFromRemote(FNetRefHandle Handle) {};
-
 	struct FDestroyInstanceParams
 	{
 		UObject* Instance = nullptr;
diff --git a/Engine/Source/Runtime/Experimental/Iris/Core/Public/Iris/ReplicationSystem/ReplicationBridge.h b/Engine/Source/Runtime/Experimental/Iris/Core/Public/Iris/ReplicationSystem/ReplicationBridge.h
--- a/Engine/Source/Runtime/Experimental/Iris/Core/Public/Iris/ReplicationSystem/ReplicationBridge.h
+++ b/Engine/Source/Runtime/Experimental/Iris/Core/Public/Iris/ReplicationSystem/ReplicationBridge.h
@@ -23,15 +23,18 @@ class UNetDriver;
 namespace UE::Net
 {
 	enum class ENetRefHandleError : uint32;
+	typedef uint8 FNetObjectFactoryId;
 
 	struct FNetDependencyInfo;
-	class FNetTokenStoreState;
-	class FReplicationFragment;
 	struct FReplicationInstanceProtocol;
 	struct FReplicationProtocol;
-	class FNetSerializationContext;
+
 	class FNetBitStreamReader;
 	class FNetBitStreamWriter;
+	class FNetSerializationContext;
+	class FNetTokenStoreState;	
+	class FReplicationFragment;
+
 	namespace Private
 	{
 		typedef uint32 FInternalNetRefIndex;
@@ -165,8 +168,8 @@ protected:
 	/** Invoked right before we apply the state for a new received subobject but after we have applied state for owning/root object in order to behave like old replication system */
 	IRISCORE_API virtual void SubObjectCreatedFromReplication(FNetRefHandle SubObjectRefHandle);
 
-	/** Invoke after we have applied the initial state for an object.*/
-	IRISCORE_API virtual void PostApplyInitialState(FNetRefHandle Handle);
+	/** Invoked after we have applied the initial state for an object.*/
+	virtual void PostApplyInitialState(UE::Net::Private::FInternalNetRefIndex InternalObjectIndex) {}
 
 	/**
 	 * Called when the instance is detached from the protocol on request by the remote. 
@@ -195,7 +198,7 @@ protected:
 	IRISCORE_API FNetRefHandle InternalCreateNetObject(FNetRefHandle AllocatedHandle, const UE::Net::FReplicationProtocol* ReplicationProtocol);
 
 	/** Create a NetRefHandle / NetObject on request from the authoritative end. */
-	IRISCORE_API FNetRefHandle InternalCreateNetObjectFromRemote(FNetRefHandle WantedNetHandle, const UE::Net::FReplicationProtocol* ReplicationProtocol);
+	IRISCORE_API FNetRefHandle InternalCreateNetObjectFromRemote(FNetRefHandle WantedNetHandle, const UE::Net::FReplicationProtocol* ReplicationProtocol,  UE::Net::FNetObjectFactoryId FactoryId);
 
 	/** Attach instance to NetRefHandle. */
 	IRISCORE_API void InternalAttachInstanceToNetRefHandle(FNetRefHandle RefHandle, bool bBindInstanceProtocol, UE::Net::FReplicationInstanceProtocol* InstanceProtocol, UObject* Instance, FNetHandle NetHandle);
@@ -257,7 +260,7 @@ private:
 	bool CallWriteNetRefHandleCreationInfo(FReplicationBridgeSerializationContext& Context, FNetRefHandle Handle);
 	bool CallWriteNetRefHandleDestructionInfo(FReplicationBridgeSerializationContext& Context, FNetRefHandle Handle);
 	void CallSubObjectCreatedFromReplication(FNetRefHandle SubObjectHandle);
-	void CallPostApplyInitialState(FNetRefHandle Handle);
+	void CallPostApplyInitialState(UE::Net::Private::FInternalNetRefIndex InternalObjectIndex);
 	void CallPruneStaleObjects();
 	void CallGetInitialDependencies(FNetRefHandle Handle, FNetDependencyInfoArray& OutDependencies) const;
 	void CallDetachInstance(FNetRefHandle Handle);
