Constraints evaluation: improvements testing
    - added Constraints.SkeletalMesh.TickingMode ([0, 2] - default: 2) to test constraints' post evaluation skeletal mesh ticking
    0 - do not tick any related skeletal mesh
    1 - only tick the constrained skeletal mesh
    2 - tick all the skeletal meshes attached to the constrained actor

when a constrained actor has several skeletal meshes attached to it, but only one skeletal mesh is constrained, it may not be necessary to tick all of them to update bone transforms

#rnx

diff --git a/Engine/Plugins/Animation/ControlRig/Source/ControlRig/Private/Constraints/ControlRigTransformableHandle.cpp b/Engine/Plugins/Animation/ControlRig/Source/ControlRig/Private/Constraints/ControlRigTransformableHandle.cpp
--- a/Engine/Plugins/Animation/ControlRig/Source/ControlRig/Private/Constraints/ControlRigTransformableHandle.cpp
+++ b/Engine/Plugins/Animation/ControlRig/Source/ControlRig/Private/Constraints/ControlRigTransformableHandle.cpp
@@ -107,7 +107,7 @@ void UTransformableControlHandle::TickTarget() const
 		return;
 	}
 	
-	if (const USkeletalMeshComponent* SkeletalMeshComponent = GetSkeletalMesh())
+	if (USkeletalMeshComponent* SkeletalMeshComponent = GetSkeletalMesh())
 	{
 		return TransformableHandleUtils::TickDependantComponents(SkeletalMeshComponent);
 	}
diff --git a/Engine/Source/Runtime/Experimental/Animation/Constraints/Private/TransformableHandleUtils.cpp b/Engine/Source/Runtime/Experimental/Animation/Constraints/Private/TransformableHandleUtils.cpp
--- a/Engine/Source/Runtime/Experimental/Animation/Constraints/Private/TransformableHandleUtils.cpp
+++ b/Engine/Source/Runtime/Experimental/Animation/Constraints/Private/TransformableHandleUtils.cpp
@@ -8,14 +8,38 @@
 
 namespace TransformableHandleUtils
 {
+
+static int32 SkeletalMeshTickingMode = 2;
+static FAutoConsoleVariableRef CVarSkeletalMeshTickingMode(
+	TEXT("Constraints.SkeletalMesh.TickingMode"),
+	SkeletalMeshTickingMode,
+	TEXT( "Constraint skeletal mesh ticking mode ([0, 2] - default: 2):\n" )
+		TEXT( "0 - do not tick any related skeletal mesh\n" )
+		TEXT( "1 - only tick the constrained skeletal mesh\n")
+		TEXT( "2 - tick all the skeletal meshes attached to the constrained actor\n"),
+		FConsoleVariableDelegate::CreateLambda([](IConsoleVariable*)
+	{
+		SkeletalMeshTickingMode = FMath::Clamp(SkeletalMeshTickingMode, 0, 2);
+	})
+);
 	
-void TickDependantComponents(const USceneComponent* InComponent)
+void TickDependantComponents(USceneComponent* InComponent)
 {
+	if (SkeletalMeshTickingMode <= 0)
+	{
+		return;
+	}
+	
 	if (!InComponent)
 	{
 		return;
 	}
 
+	if (SkeletalMeshTickingMode == 1)
+	{
+		return TickSkeletalMeshComponent(Cast<USkeletalMeshComponent>(InComponent));
+	}
+
 	static constexpr bool bIncludeFromChildActors = true;
 
 	const AActor* Parent = InComponent->GetOwner();
diff --git a/Engine/Source/Runtime/Experimental/Animation/Constraints/Public/TransformableHandleUtils.h b/Engine/Source/Runtime/Experimental/Animation/Constraints/Public/TransformableHandleUtils.h
--- a/Engine/Source/Runtime/Experimental/Animation/Constraints/Public/TransformableHandleUtils.h
+++ b/Engine/Source/Runtime/Experimental/Animation/Constraints/Public/TransformableHandleUtils.h
@@ -9,6 +9,6 @@ class USkeletalMeshComponent;
 
 namespace TransformableHandleUtils
 {
-	CONSTRAINTS_API void TickDependantComponents(const USceneComponent* InComponent);
+	CONSTRAINTS_API void TickDependantComponents(USceneComponent* InComponent);
 	CONSTRAINTS_API void TickSkeletalMeshComponent(USkeletalMeshComponent* InSkeletalMeshComponent);
 }
