Add support for non-nanite fallback mesh for nanite geometry collections.
Strip geometry collection nanite render data when cooking non-nanite platforms.

diff --git a/Engine/Source/Runtime/Experimental/GeometryCollectionEngine/Private/GeometryCollection/DerivedDataGeometryCollectionCooker.cpp b/Engine/Source/Runtime/Experimental/GeometryCollectionEngine/Private/GeometryCollection/DerivedDataGeometryCollectionCooker.cpp
--- a/Engine/Source/Runtime/Experimental/GeometryCollectionEngine/Private/GeometryCollection/DerivedDataGeometryCollectionCooker.cpp
+++ b/Engine/Source/Runtime/Experimental/GeometryCollectionEngine/Private/GeometryCollection/DerivedDataGeometryCollectionCooker.cpp
@@ -102,7 +102,7 @@ bool FDerivedDataGeometryCollectionRenderDataCooker::Build(TArray<uint8>& OutDat
 	{
 		Chaos::FErrorReporter ErrorReporter(GeometryCollection.GetName());
 
-		TUniquePtr<FGeometryCollectionRenderData> RenderData = FGeometryCollectionRenderData::Create(*Collection, GeometryCollection.EnableNanite, GeometryCollection.bUseFullPrecisionUVs, GeometryCollection.bConvertVertexColorsToSRGB);
+		TUniquePtr<FGeometryCollectionRenderData> RenderData = FGeometryCollectionRenderData::Create(*Collection, GeometryCollection.EnableNanite, GeometryCollection.bEnableNaniteFallback, GeometryCollection.bUseFullPrecisionUVs, GeometryCollection.bConvertVertexColorsToSRGB);
 		RenderData->Serialize(ChaosAr, GeometryCollection);
 
 		if (false && ErrorReporter.EncounteredAnyErrors())
diff --git a/Engine/Source/Runtime/Experimental/GeometryCollectionEngine/Private/GeometryCollection/GeometryCollectionObject.cpp b/Engine/Source/Runtime/Experimental/GeometryCollectionEngine/Private/GeometryCollection/GeometryCollectionObject.cpp
--- a/Engine/Source/Runtime/Experimental/GeometryCollectionEngine/Private/GeometryCollection/GeometryCollectionObject.cpp
+++ b/Engine/Source/Runtime/Experimental/GeometryCollectionEngine/Private/GeometryCollection/GeometryCollectionObject.cpp
@@ -113,6 +113,7 @@ UGeometryCollection::UGeometryCollection(const FObjectInitializer& ObjectInitial
 	, bStripOnCook(false)
 	, bStripRenderDataOnCook(false)
 	, EnableNanite(false)
+	, bEnableNaniteFallback(false)
 #if WITH_EDITORONLY_DATA
 	, CollisionType_DEPRECATED(ECollisionTypeEnum::Chaos_Volumetric)
 	, ImplicitType_DEPRECATED(EImplicitTypeEnum::Chaos_Implicit_Convex)
@@ -1458,7 +1459,7 @@ void UGeometryCollection::RebuildRenderData()
 	if (RenderDataGuid != StateGuid)
 	{
 		ReleaseResources();
-		RenderData = FGeometryCollectionRenderData::Create(*GetGeometryCollection(), EnableNanite, bUseFullPrecisionUVs, bConvertVertexColorsToSRGB);
+		RenderData = FGeometryCollectionRenderData::Create(*GetGeometryCollection(), EnableNanite, bEnableNaniteFallback, bUseFullPrecisionUVs, bConvertVertexColorsToSRGB);
 		InitResources();
 		PropagateMarkDirtyToComponents();
 		RenderDataGuid = StateGuid;
@@ -1816,6 +1817,11 @@ void UGeometryCollection::PostEditChangeProperty(struct FPropertyChangedEvent& P
 			bDoInvalidateCollection = true;
 			bRebuildRenderData = true;
 		}
+		else if (PropertyChangedEvent.Property->GetFName() == GET_MEMBER_NAME_CHECKED(UGeometryCollection, bEnableNaniteFallback))
+		{
+			bDoInvalidateCollection = true;
+			bRebuildRenderData = true;
+		}
 		else if (PropertyChangedEvent.Property->GetFName() == GET_MEMBER_NAME_CHECKED(UGeometryCollection, bUseFullPrecisionUVs))
 		{
 			bDoInvalidateCollection = true;
diff --git a/Engine/Source/Runtime/Experimental/GeometryCollectionEngine/Private/GeometryCollection/GeometryCollectionRenderData.cpp b/Engine/Source/Runtime/Experimental/GeometryCollectionEngine/Private/GeometryCollection/GeometryCollectionRenderData.cpp
--- a/Engine/Source/Runtime/Experimental/GeometryCollectionEngine/Private/GeometryCollection/GeometryCollectionRenderData.cpp
+++ b/Engine/Source/Runtime/Experimental/GeometryCollectionEngine/Private/GeometryCollection/GeometryCollectionRenderData.cpp
@@ -684,7 +684,7 @@ void CreateNaniteData(FGeometryCollectionBuiltMeshData&& InMeshData, FGeometryCo
 	}
 }
 
-TUniquePtr<FGeometryCollectionRenderData> FGeometryCollectionRenderData::Create(FGeometryCollection& InCollection, bool bInEnableNanite, bool bInUseFullPrecisionUVs, bool bConvertVertexColorsToSRGB)
+TUniquePtr<FGeometryCollectionRenderData> FGeometryCollectionRenderData::Create(FGeometryCollection& InCollection, bool bInEnableNanite, bool bInEnableNaniteFallback, bool bInUseFullPrecisionUVs, bool bConvertVertexColorsToSRGB)
 {
 	TRACE_CPUPROFILER_EVENT_SCOPE(FGeometryCollectionRenderData::Create);
 
@@ -693,18 +693,16 @@ TUniquePtr<FGeometryCollectionRenderData> FGeometryCollectionRenderData::Create(
 	FGeometryCollectionBuiltMeshData MeshBuildData;
 	if (BuildMeshDataFromGeometryCollection(InCollection, MeshBuildData, bConvertVertexColorsToSRGB))
 	{
-	{
-		if (bInEnableNanite)
+		if (!bInEnableNanite || bInEnableNaniteFallback)
 		{
-			CreateNaniteData(MoveTemp(MeshBuildData), *RenderData.Get());
+			CreateMeshData(InCollection, MeshBuildData, bInUseFullPrecisionUVs, *RenderData.Get());
 		}
-		else
+
+		if (bInEnableNanite)
 		{
-			// Could always create mesh data if we want to be able to enable/disable nanite at runtime in cooked build.
-			CreateMeshData(InCollection, MeshBuildData, bInUseFullPrecisionUVs, *RenderData.Get());
+			CreateNaniteData(MoveTemp(MeshBuildData), *RenderData.Get());
 		}
 	}
-	}
 
 	return RenderData;
 }
@@ -735,7 +733,22 @@ void FGeometryCollectionRenderData::Serialize(FArchive& Ar, UGeometryCollection&
 	}
 
 	Ar << bHasMeshData;
-	Ar << bHasNaniteData;
+
+#if WITH_EDITOR
+	bool bStripNaniteDataFromCook = Ar.IsCooking() && !DoesTargetPlatformSupportNanite(Ar.CookingTarget());
+#else
+	bool bStripNaniteDataFromCook = false;
+#endif
+
+	if (bStripNaniteDataFromCook)
+	{
+		bool bDummy = false;
+		Ar << bDummy;
+	}
+	else
+	{
+		Ar << bHasNaniteData;
+	}
 
 	if (bHasMeshData)
 	{
@@ -748,7 +761,7 @@ void FGeometryCollectionRenderData::Serialize(FArchive& Ar, UGeometryCollection&
 		MeshDescription = {};
 	}
 
-	if (bHasNaniteData)
+	if (bHasNaniteData && !bStripNaniteDataFromCook)
 	{
 		InitNaniteResources(NaniteResourcesPtr);
 
diff --git a/Engine/Source/Runtime/Experimental/GeometryCollectionEngine/Private/GeometryCollection/GeometryCollectionRenderData.h b/Engine/Source/Runtime/Experimental/GeometryCollectionEngine/Private/GeometryCollection/GeometryCollectionRenderData.h
--- a/Engine/Source/Runtime/Experimental/GeometryCollectionEngine/Private/GeometryCollection/GeometryCollectionRenderData.h
+++ b/Engine/Source/Runtime/Experimental/GeometryCollectionEngine/Private/GeometryCollection/GeometryCollectionRenderData.h
@@ -116,7 +116,7 @@ public:
 	~FGeometryCollectionRenderData();
 
 #if WITH_EDITOR
-	static TUniquePtr<FGeometryCollectionRenderData> Create(FGeometryCollection& InCollection, bool bInEnableNanite, bool bInUseFullPrecisionUVs, bool bConvertVertexColorsToSRGB);
+	static TUniquePtr<FGeometryCollectionRenderData> Create(FGeometryCollection& InCollection, bool bInEnableNanite, bool bInEnableNaniteFallback, bool bInUseFullPrecisionUVs, bool bConvertVertexColorsToSRGB);
 #endif
 
 	bool IsInitialized()
diff --git a/Engine/Source/Runtime/Experimental/GeometryCollectionEngine/Public/GeometryCollection/GeometryCollectionObject.h b/Engine/Source/Runtime/Experimental/GeometryCollectionEngine/Public/GeometryCollection/GeometryCollectionObject.h
--- a/Engine/Source/Runtime/Experimental/GeometryCollectionEngine/Public/GeometryCollection/GeometryCollectionObject.h
+++ b/Engine/Source/Runtime/Experimental/GeometryCollectionEngine/Public/GeometryCollection/GeometryCollectionObject.h
@@ -610,6 +610,12 @@ public:
 	UPROPERTY(EditAnywhere, BlueprintReadWrite, BlueprintSetter=SetEnableNanite, Category = "Nanite")
 	bool EnableNanite;
 
+	/**
+	 * Enable Non-Nanite fallback mesh when Nanite support is enabled.
+	 */
+	UPROPERTY(EditAnywhere, Category = "Nanite", Meta = (EditCondition = "EnableNanite"))
+	bool bEnableNaniteFallback;
+
 	UFUNCTION(BlueprintCallable, Category = "Rendering")
 	GEOMETRYCOLLECTIONENGINE_API void SetConvertVertexColorsToSRGB(bool bValue);
 
