Fix debug draw on RBAN
Scaled asset are now displayed correctly, with the correct size

diff --git a/Engine/Source/Runtime/Experimental/Chaos/Private/Chaos/ChaosDebugDraw.cpp b/Engine/Source/Runtime/Experimental/Chaos/Private/Chaos/ChaosDebugDraw.cpp
--- a/Engine/Source/Runtime/Experimental/Chaos/Private/Chaos/ChaosDebugDraw.cpp
+++ b/Engine/Source/Runtime/Experimental/Chaos/Private/Chaos/ChaosDebugDraw.cpp
@@ -900,14 +900,16 @@ namespace Chaos
 			{
 				const TSphere<FReal, 3>* Sphere = Implicit->template GetObject<TSphere<FReal, 3>>();
 				const FVec3 P = ShapeTransform.TransformPosition(Sphere->GetCenter());
-				FDebugDrawQueue::GetInstance().DrawDebugSphere(P, Sphere->GetRadius(), 8, ShapeColor, false, Duration, uint8(Settings.DrawPriority), Settings.ShapeThicknesScale * Settings.LineThickness);
+				const FVec3 Scale = ShapeTransform.GetScale3D();
+				FDebugDrawQueue::GetInstance().DrawDebugSphere(P, Sphere->GetRadius() * Scale.Z, 8, ShapeColor, false, Duration, uint8(Settings.DrawPriority), Settings.ShapeThicknesScale * Settings.LineThickness);
 				break;
 			}
 			case ImplicitObjectType::Box:
 			{
 				const TBox<FReal, 3>* Box = Implicit->template GetObject<TBox<FReal, 3>>();
 				const FVec3 P = ShapeTransform.TransformPosition(Box->GetCenter());
-				FDebugDrawQueue::GetInstance().DrawDebugBox(P, (FReal)0.5 * Box->Extents(), ShapeTransform.GetRotation(), ShapeColor, false, Duration, uint8(Settings.DrawPriority), Settings.ShapeThicknesScale * Settings.LineThickness);
+				const FVec3 Scale = ShapeTransform.GetScale3D();
+				FDebugDrawQueue::GetInstance().DrawDebugBox(P, (FReal)0.5 * Box->Extents() * Scale, ShapeTransform.GetRotation(), ShapeColor, false, Duration, uint8(Settings.DrawPriority), Settings.ShapeThicknesScale * Settings.LineThickness);
 				break;
 			}
 			case ImplicitObjectType::Plane:
@@ -917,7 +919,8 @@ namespace Chaos
 				const FCapsule* Capsule = Implicit->template GetObject<FCapsule>();
 				const FVec3 P = ShapeTransform.TransformPosition(Capsule->GetCenter());
 				const FRotation3 Q = ShapeTransform.GetRotation() * FRotationMatrix::MakeFromZ(Capsule->GetAxis());
-				FDebugDrawQueue::GetInstance().DrawDebugCapsule(P, (FReal)0.5 * Capsule->GetHeight() + Capsule->GetRadius(), Capsule->GetRadius(), Q, ShapeColor, false, Duration, uint8(Settings.DrawPriority), Settings.ShapeThicknesScale * Settings.LineThickness);
+				const FVec3 Scale = ShapeTransform.GetScale3D();
+				FDebugDrawQueue::GetInstance().DrawDebugCapsule(P, ((FReal)0.5 * Capsule->GetHeight() + Capsule->GetRadius()) * Scale.Z, Capsule->GetRadius() * Scale.X, Q, ShapeColor, false, Duration, uint8(Settings.DrawPriority), Settings.ShapeThicknesScale * Settings.LineThickness);
 				break;
 			}
 			case ImplicitObjectType::LevelSet:
