ChaosFlesh: CIS fix
change warning messages to display/verbose

diff --git a/Engine/Source/Runtime/Experimental/Chaos/Private/Chaos/Deformable/GaussSeidelMainConstraint.cpp b/Engine/Source/Runtime/Experimental/Chaos/Private/Chaos/Deformable/GaussSeidelMainConstraint.cpp
--- a/Engine/Source/Runtime/Experimental/Chaos/Private/Chaos/Deformable/GaussSeidelMainConstraint.cpp
+++ b/Engine/Source/Runtime/Experimental/Chaos/Private/Chaos/Deformable/GaussSeidelMainConstraint.cpp
@@ -1,6 +1,7 @@
 // Copyright Epic Games, Inc. All Rights Reserved.
 #include "Chaos/Deformable/GaussSeidelMainConstraint.h"
 #include "Chaos/Math/Krylov.h"
+#include "ChaosLog.h"
 
 namespace Chaos::Softs
 {
@@ -176,9 +177,9 @@ namespace Chaos::Softs
 
 		NewtonNorm = FMath::Sqrt(NewtonNorm);
 
-		UE_LOG(LogTemp, Warning, TEXT("Current Iteration: %d"), CurrentIt);
+		UE_LOG(LogChaos, Display, TEXT("Current Iteration: %d"), CurrentIt);
 
-		UE_LOG(LogTemp, Warning, TEXT("Newton Residual is %f"), NewtonNorm);
+		UE_LOG(LogChaos, Display, TEXT("Newton Residual is %f"), NewtonNorm);
 
 		PassedIters++;
 
diff --git a/Engine/Source/Runtime/Experimental/Chaos/Public/Chaos/Math/Krylov.h b/Engine/Source/Runtime/Experimental/Chaos/Public/Chaos/Math/Krylov.h
--- a/Engine/Source/Runtime/Experimental/Chaos/Public/Chaos/Math/Krylov.h
+++ b/Engine/Source/Runtime/Experimental/Chaos/Public/Chaos/Math/Krylov.h
@@ -278,11 +278,11 @@ void LanczosCG(
 
 	if (check_residual)
 	{
-		UE_LOG(LogChaos, Warning, TEXT("Lanczos used max iterations (%d). Residual = %g."), max_it, residual);
+		UE_LOG(LogChaos, Verbose, TEXT("Lanczos used max iterations (%d). Residual = %g."), max_it, residual);
 	}
 	else
 	{
-		UE_LOG(LogChaos, Warning, TEXT("Lanczos used max iterations (%d)."), max_it);
+		UE_LOG(LogChaos, Verbose, TEXT("Lanczos used max iterations (%d)."), max_it);
 	}
 }
 
