Fix geometry collection component removal logic preventing internal clusters from crumbling into smaller pieces
- HasDynamicInternalClusterParent and HasClusterUnionParent are mutually exclusive and testing both was resulting in getting always false for computing HasDynamicInternalClusterParent

diff --git a/Engine/Source/Runtime/Experimental/GeometryCollectionEngine/Private/GeometryCollection/GeometryCollectionComponent.cpp b/Engine/Source/Runtime/Experimental/GeometryCollectionEngine/Private/GeometryCollection/GeometryCollectionComponent.cpp
--- a/Engine/Source/Runtime/Experimental/GeometryCollectionEngine/Private/GeometryCollection/GeometryCollectionComponent.cpp
+++ b/Engine/Source/Runtime/Experimental/GeometryCollectionEngine/Private/GeometryCollection/GeometryCollectionComponent.cpp
@@ -6533,7 +6533,7 @@ void UGeometryCollectionComponent::IncrementSleepTimer(float DeltaTime)
 				{
 					// this children has an dynamic internal cluster parent so it can't be removed but we need tyo process the internal cluster by looking at the original parent properties
 					const int32 OriginalParentIdx = OriginalParents[TransformIdx];
-					const bool HasDynamicInternalClusterParent = DynamicStateFacade.HasDynamicInternalClusterParent(TransformIdx) && DynamicStateFacade.HasClusterUnionParent(TransformIdx);
+					const bool HasDynamicInternalClusterParent = DynamicStateFacade.HasDynamicInternalClusterParent(TransformIdx);
 					if (OriginalParentIdx > INDEX_NONE && HasDynamicInternalClusterParent && RemoveOnSleepFacade.IsRemovalActive(OriginalParentIdx))
 					{
 						const bool UseClusterCrumbling = true; // with sleep removal : internal clusters always crumble - this will change when we merge the removal feature together
@@ -6603,7 +6603,7 @@ void UGeometryCollectionComponent::IncrementBreakTimer(float DeltaTime)
 				{
 					// this children has an internal cluster parent so it can't be removed but we need tyo process the internal cluster by looking at the original parent properties
 					const int32 OriginalParentIdx = OriginalParents[TransformIdx];
-					const bool HasDynamicInternalClusterParent = DynamicStateFacade.HasDynamicInternalClusterParent(TransformIdx) && !DynamicStateFacade.HasClusterUnionParent(TransformIdx);
+					const bool HasDynamicInternalClusterParent = DynamicStateFacade.HasDynamicInternalClusterParent(TransformIdx);
 
 					if (OriginalParentIdx > INDEX_NONE && HasDynamicInternalClusterParent && RemoveOnBreakFacade.IsRemovalActive(OriginalParentIdx))
 					{
