Dataflow : Fix various issues isues and crashes with Dataflow nodes
- Fix out of bound related crashes when using the select by box node ( when there was more transforms than geometry in the collection )
    - now properly access the centroid array by geometry index
- StaticMeshToCollection node now outputs collection with a single root when using the Split Element option

diff --git a/Engine/Source/Runtime/Experimental/Chaos/Private/GeometryCollection/Facades/CollectionTransformSelectionFacade.cpp b/Engine/Source/Runtime/Experimental/Chaos/Private/GeometryCollection/Facades/CollectionTransformSelectionFacade.cpp
--- a/Engine/Source/Runtime/Experimental/Chaos/Private/GeometryCollection/Facades/CollectionTransformSelectionFacade.cpp
+++ b/Engine/Source/Runtime/Experimental/Chaos/Private/GeometryCollection/Facades/CollectionTransformSelectionFacade.cpp
@@ -841,18 +841,23 @@ namespace GeometryCollection::Facades
 
 			for (int32 TransformIdx = 0; TransformIdx < NumTransforms; ++TransformIdx)
 			{
-				const FVector CentroidInBoneSpace(Centroids[TransformIdx]);
+				const TManagedArray<int32>& TransformToGeometryIndices = TransformToGeometryIndexAttribute.Get();
+				const int32 GeometryIndex = TransformToGeometryIndices[TransformIdx];
+				if (Centroids.IsValidIndex(GeometryIndex))
+				{
+					const FVector CentroidInBoneSpace(Centroids[GeometryIndex]);
 
-				// Transform from BoneSpace to CollectionSpace
-				const FTransform CollectionSpaceTransform = TransformFacade.ComputeCollectionSpaceTransform(TransformIdx);
-				const FVector CentroidInCollectionSpace = CollectionSpaceTransform.TransformPosition(CentroidInBoneSpace);
+					// Transform from BoneSpace to CollectionSpace
+					const FTransform CollectionSpaceTransform = TransformFacade.ComputeCollectionSpaceTransform(TransformIdx);
+					const FVector CentroidInCollectionSpace = CollectionSpaceTransform.TransformPosition(CentroidInBoneSpace);
 
-				// Transform with specified transform
-				const FVector CentroidInBoxSpace = InBoxTransform.InverseTransformPosition(CentroidInCollectionSpace);
+					// Transform with specified transform
+					const FVector CentroidInBoxSpace = InBoxTransform.InverseTransformPosition(CentroidInCollectionSpace);
 
-				if (InBox.IsInside(CentroidInBoxSpace))
-				{
-					OutSelection.Add(TransformIdx);
+					if (InBox.IsInside(CentroidInBoxSpace))
+					{
+						OutSelection.Add(TransformIdx);
+					}
 				}
 			}
 		}
diff --git a/Engine/Source/Runtime/Experimental/GeometryCollectionEngine/Private/GeometryCollection/GeometryCollectionEngineConversion.cpp b/Engine/Source/Runtime/Experimental/GeometryCollectionEngine/Private/GeometryCollection/GeometryCollectionEngineConversion.cpp
--- a/Engine/Source/Runtime/Experimental/GeometryCollectionEngine/Private/GeometryCollection/GeometryCollectionEngineConversion.cpp
+++ b/Engine/Source/Runtime/Experimental/GeometryCollectionEngine/Private/GeometryCollection/GeometryCollectionEngineConversion.cpp
@@ -1589,12 +1589,18 @@ void FGeometryCollectionEngineConversion::ConvertStaticMeshToGeometryCollection(
 		NewGeometryCollection->GeometrySource.Emplace(SourceSoftObjectPath, MeshTransform, SourceMaterials, bSplitComponents, bSetInternalFromMaterialIndex);
 		FGeometryCollectionEngineConversion::AppendStaticMesh(StaticMesh, SourceMaterials, MeshTransform, NewGeometryCollection, false, bAddInternalMaterials, bSplitComponents, bSetInternalFromMaterialIndex);
 
+		// make sure we have only one root if we split components
+		TSharedPtr<FGeometryCollection> OutCollectionPtr = NewGeometryCollection->GetGeometryCollection();
+		if (bSplitComponents && FGeometryCollectionClusteringUtility::ContainsMultipleRootBones(OutCollectionPtr.Get()))
+		{
+			FGeometryCollectionClusteringUtility::ClusterAllBonesUnderNewRoot(OutCollectionPtr.Get());
+		}
+
 		NewGeometryCollection->InitializeMaterials();
 
 		// Materials
 		OutMaterialInstances.Append(NewGeometryCollection->Materials);
 
-		TSharedPtr<FGeometryCollection> OutCollectionPtr = NewGeometryCollection->GetGeometryCollection();
 		OutCollectionPtr->CopyTo(&OutCollection);
 	}
 #endif //WITH_EDITORONLY_DATA
