IoStoreOnDemand journal file tests

diff --git a/Engine/Source/Runtime/Experimental/IoStore/OnDemand/Private/OnDemandInstallCache.cpp b/Engine/Source/Runtime/Experimental/IoStore/OnDemand/Private/OnDemandInstallCache.cpp
--- a/Engine/Source/Runtime/Experimental/IoStore/OnDemand/Private/OnDemandInstallCache.cpp
+++ b/Engine/Source/Runtime/Experimental/IoStore/OnDemand/Private/OnDemandInstallCache.cpp
@@ -25,6 +25,12 @@
 #include "Tasks/Task.h"
 #include "ProfilingDebugging/IoStoreTrace.h"
 
+#if defined(WITH_IOSTORE_ONDEMAND_TESTS)
+#include "TestHarness.h"
+#include "TestMacros/Assertions.h"
+#include <catch2/generators/catch_generators.hpp>
+#endif
+
 namespace UE::IoStore
 {
 
@@ -539,6 +545,7 @@ struct FCasJournal
 	static FIoStatus		Replay(const FString& JournalFile, FEntryHandler&& Handler);
 	static FIoStatus		Create(const FString& JournalFile);
 	static FTransaction		Begin(FString&& JournalFile);
+	static FTransaction		Begin(const FString& JournalFile) { return Begin(FString(JournalFile)); }
 	static FIoStatus		Commit(FTransaction&& Transaction);
 };
 
@@ -1660,4 +1667,125 @@ TSharedPtr<IOnDemandInstallCache> MakeOnDemandInstallCache(
 	return MakeShared<FOnDemandInstallCache>(Config, IoStore);
 }
 
+///////////////////////////////////////////////////////////////////////////////
+#if defined(WITH_IOSTORE_ONDEMAND_TESTS)
+
+class FTmpDirectoryScope
+{
+public:
+	explicit FTmpDirectoryScope(const FString& InDir)
+		: Ifm(IFileManager::Get())
+		, Dir(InDir)
+	{
+		const bool bTree			= true;
+		const bool bRequireExists	= false;
+		Ifm.DeleteDirectory(*Dir, bRequireExists, bTree);
+		Ifm.MakeDirectory(*Dir, bTree);
+	}
+
+	~FTmpDirectoryScope()
+	{
+		const bool bTree			= true;
+		const bool bRequireExists	= false;
+		Ifm.DeleteDirectory(*Dir, bRequireExists, bTree);
+	}
+private:
+	IFileManager& Ifm;
+	FString Dir;
+};
+
+TEST_CASE("IoStore::OnDemand::InstallCache::Journal", "[IoStoreOnDemand][InstallCache]")
+{
+	const FString TestBaseDir = TEXT("TestTmpDir");
+
+	SECTION("CreateJournalFile")
+	{
+		FTmpDirectoryScope _(TestBaseDir);
+		const FString JournalFile = TestBaseDir / TEXT("test.jrn");
+		FIoStatus Status = FCasJournal::Create(JournalFile);
+		CHECK(Status.IsOk());
+	}
+
+	SECTION("SimpleTransaction")
+	{
+		FTmpDirectoryScope _(TestBaseDir);
+		const FString JournalFile = TestBaseDir / TEXT("test.jrn");
+		FIoStatus Status = FCasJournal::Create(JournalFile);
+		CHECK(Status.IsOk());
+
+		FCasJournal::FTransaction Transaction = FCasJournal::Begin(JournalFile);
+		Transaction.BlockCreated(FCasBlockId(1));
+		Status = FCasJournal::Commit(MoveTemp(Transaction));
+		CHECK(Status.IsOk());
+	}
+
+	SECTION("ReplayChunkLocations")
+	{
+		//Arrange
+		TArray<FCasAddr>	ExpectedAddresses;
+		TArray<uint32>		ExpectedBlockOffsets;
+		const FCasBlockId	ExpectedBlockId(42);
+		
+		for (int32 Idx = 1; Idx < 33; ++Idx)
+		{
+			ExpectedAddresses.Add(FCasAddr::From(reinterpret_cast<const uint8*>(&Idx), sizeof(uint32)));
+			ExpectedBlockOffsets.Add(Idx);
+		}
+
+		// Act
+		FTmpDirectoryScope _(TestBaseDir);
+		const FString JournalFile = TestBaseDir / TEXT("test.jrn");
+		FIoStatus Status = FCasJournal::Create(JournalFile);
+		CHECK(Status.IsOk());
+
+		FCasJournal::FTransaction Transaction = FCasJournal::Begin(JournalFile);
+		for (int32 Idx = 0; const FCasAddr& Addr : ExpectedAddresses)
+		{
+			Transaction.ChunkLocation(
+				FCasLocation
+				{
+					.BlockId = ExpectedBlockId,
+					.BlockOffset = ExpectedBlockOffsets[Idx]
+				},
+				Addr);
+		}
+
+		Status = FCasJournal::Commit(MoveTemp(Transaction));
+		CHECK(Status.IsOk());
+
+		// Assert
+		TArray<FCasJournal::FEntry::FChunkLocation> Locs;
+		Status = FCasJournal::Replay(
+			JournalFile,
+			[&Locs](const FCasJournal::FEntry& JournalEntry)
+			{
+				switch(JournalEntry.Type())
+				{
+				case FCasJournal::FEntry::EType::ChunkLocation:
+				{
+					Locs.Add(JournalEntry.ChunkLocation);
+					break;
+				}
+				default:
+					CHECK(false);
+					break;
+				};
+			});
+		CHECK(Status.IsOk());
+		CHECK(Locs.Num() == ExpectedAddresses.Num());
+		for (int32 Idx = 0; const FCasJournal::FEntry::FChunkLocation& Loc : Locs)
+		{
+			const FCasLocation ExpectedLoc = FCasLocation
+			{
+				.BlockId = ExpectedBlockId,
+				.BlockOffset = uint32(Idx + 1)
+			};
+			CHECK(Loc.CasLocation.BlockId == ExpectedLoc.BlockId);
+			CHECK(Loc.CasLocation.BlockOffset == ExpectedLoc.BlockOffset);
+		}
+	}
+}
+
+#endif // WITH_IOSTORE_ONDEMAND_TESTS
+
 } // namespace UE::IoStore
diff --git a/Engine/Source/Runtime/Experimental/IoStore/Tests/IoStoreOnDemandTests.Build.cs b/Engine/Source/Runtime/Experimental/IoStore/Tests/IoStoreOnDemandTests.Build.cs
new file mode 100644
--- /dev/null
+++ b/Engine/Source/Runtime/Experimental/IoStore/Tests/IoStoreOnDemandTests.Build.cs
@@ -0,0 +1,23 @@
+// Copyright Epic Games, Inc. All Rights Reserved.
+using UnrealBuildTool;
+
+public class IoStoreOnDemandTests : TestModuleRules
+{
+    static IoStoreOnDemandTests()
+    {
+		TestMetadata = new Metadata();
+        TestMetadata.TestName = "IoStoreOnDemand";
+        TestMetadata.TestShortName = "IoStoreOnDemand";
+    }
+
+    public IoStoreOnDemandTests(ReadOnlyTargetRules Target) : base(Target)
+	{
+		OptimizeCode = CodeOptimization.Never;
+
+		PrivateDependencyModuleNames.AddRange(
+			new string[] {
+				"Core",
+				"IoStoreOnDemand",
+			});
+	}
+}
diff --git a/Engine/Source/Runtime/Experimental/IoStore/Tests/IoStoreOnDemandTests.Target.cs b/Engine/Source/Runtime/Experimental/IoStore/Tests/IoStoreOnDemandTests.Target.cs
new file mode 100644
--- /dev/null
+++ b/Engine/Source/Runtime/Experimental/IoStore/Tests/IoStoreOnDemandTests.Target.cs
@@ -0,0 +1,17 @@
+// Copyright Epic Games, Inc. All Rights Reserved.
+
+using UnrealBuildTool;
+
+[SupportedPlatforms("Win64")]
+public class IoStoreOnDemandTestsTarget : TestTargetRules
+{
+	public IoStoreOnDemandTestsTarget(TargetInfo Target) : base(Target)
+	{
+		bCompileAgainstCoreUObject			= true;
+		bCompileAgainstApplicationCore		= true;
+		bAdaptiveUnityDisablesOptimizations = true;
+
+		GlobalDefinitions.Add("CATCH_CONFIG_ENABLE_BENCHMARKING=1");
+		GlobalDefinitions.Add("WITH_IOSTORE_ONDEMAND_TESTS=1");
+	}
+}
