Geometry collection : fix issues where geometry collection do not properly update the velocity buffer
- this is because the previous transforms are not properly set

diff --git a/Engine/Source/Runtime/Experimental/GeometryCollectionEngine/Private/GeometryCollection/GeometryCollectionComponent.cpp b/Engine/Source/Runtime/Experimental/GeometryCollectionEngine/Private/GeometryCollection/GeometryCollectionComponent.cpp
--- a/Engine/Source/Runtime/Experimental/GeometryCollectionEngine/Private/GeometryCollection/GeometryCollectionComponent.cpp
+++ b/Engine/Source/Runtime/Experimental/GeometryCollectionEngine/Private/GeometryCollection/GeometryCollectionComponent.cpp
@@ -651,7 +651,6 @@ UGeometryCollectionComponent::UGeometryCollectionComponent(const FObjectInitiali
 	, bEnableBoneSelection(false)
 	, IsObjectDynamic(false)
 	, IsObjectLoading(true)
-	, bIsMoving(false)
 	, ViewLevel(-1)
 	, NavmeshInvalidationTimeSliceIndex(0)
 	, ComponentSpaceTransforms(this)
@@ -1126,20 +1125,7 @@ FPrimitiveSceneProxy* UGeometryCollectionComponent::CreateSceneProxy()
 			RestCollection->HasNaniteData() &&
 			GGeometryCollectionNanite != 0)
 		{
-			FNaniteGeometryCollectionSceneProxy* NaniteProxy = new FNaniteGeometryCollectionSceneProxy(this);
-			LocalSceneProxy = NaniteProxy;
-
-			// ForceMotionBlur means we maintain bIsMoving, regardless of actual state.
-			if (bForceMotionBlur)
-			{
-				bIsMoving = true;
-				ENQUEUE_RENDER_COMMAND(NaniteProxyOnMotionEnd)(
-					[NaniteProxy] (FRHICommandListBase&)
-					{
-						NaniteProxy->OnMotionBegin();
-					}
-				);
-			}
+			LocalSceneProxy = new FNaniteGeometryCollectionSceneProxy(this);
 		}
 		else if (RestCollection->HasMeshData())
 		{
@@ -1781,39 +1767,15 @@ void UGeometryCollectionComponent::RestTransformsChanged()
 {
 	if (SceneProxy)
 	{
-		FGeometryCollectionDynamicData* DynamicData = GDynamicDataPool.Allocate();
-		DynamicData->SetPrevTransforms(ComponentSpaceTransforms.RequestAllTransforms());
-
 		OnTransformsDirty();
 
-		DynamicData->SetTransforms(ComponentSpaceTransforms.RequestAllTransforms());
-		DynamicData->IsDynamic = true;
-
 #if WITH_EDITOR
-			// We need to do this in case we're controlled by Sequencer in editor, which doesn't invoke PostEditChangeProperty
-			UpdateCachedBounds();
-			SendRenderTransform_Concurrent();
+		// We need to do this in case we're controlled by Sequencer in editor, which doesn't invoke PostEditChangeProperty
+		UpdateCachedBounds();
+		SendRenderTransform_Concurrent();
 #endif
-		if (SceneProxy->IsNaniteMesh())
-		{
-			FNaniteGeometryCollectionSceneProxy* GeometryCollectionSceneProxy = static_cast<FNaniteGeometryCollectionSceneProxy*>(SceneProxy);
-			ENQUEUE_RENDER_COMMAND(SendRenderDynamicData)(
-				[GeometryCollectionSceneProxy, DynamicData] (FRHICommandListBase&)
-				{
-					GeometryCollectionSceneProxy->SetDynamicData_RenderThread(DynamicData, GeometryCollectionSceneProxy->GetLocalToWorld());
-				}
-			);
-		}
-		else
-		{
-			FGeometryCollectionSceneProxy* GeometryCollectionSceneProxy = static_cast<FGeometryCollectionSceneProxy*>(SceneProxy);
-			ENQUEUE_RENDER_COMMAND(SendRenderDynamicData)(
-				[GeometryCollectionSceneProxy, DynamicData] (FRHICommandListBase& RHICmdList)
-				{
-					GeometryCollectionSceneProxy->SetDynamicData_RenderThread(RHICmdList, DynamicData);
-				}
-			);
-		}
+
+		SendDynamicDataToSceneProxy();
 	}
 	else
 	{
@@ -3221,88 +3183,31 @@ void UGeometryCollectionComponent::GetRestTransforms(TArray<FMatrix44f>& OutRest
 FGeometryCollectionDynamicData* UGeometryCollectionComponent::InitDynamicData(bool bInitialization)
 {
 	SCOPE_CYCLE_COUNTER(STAT_GCInitDynamicData);
+	FGeometryCollectionDynamicData* DynamicData = GDynamicDataPool.Allocate();
+	DynamicData->SetTransforms(ComponentSpaceTransforms.RequestAllTransforms());
 
-	FGeometryCollectionDynamicData* DynamicData = nullptr;
-
-	const bool bEditorMode = bShowBoneColors || bEnableBoneSelection;
-	const bool bIsDynamic  = GetIsObjectDynamic() || bEditorMode || bInitialization;
-
-	if (bIsDynamic)
-	{
-		DynamicData = GDynamicDataPool.Allocate();
-		DynamicData->IsDynamic = true;
-		DynamicData->IsLoading = GetIsObjectLoading();
-
-		const TArray<FTransform3f>& CompSpaceTransforms = ComponentSpaceTransforms.RequestAllTransforms();
-
-		// If we have no transforms stored in the dynamic data, then assign both prev and current to the same global matrices
-		// Copy existing global matrices into prev transforms
-		DynamicData->PrevTransforms = DynamicData->Transforms;
-
-		// Copy global matrices over to DynamicData
-		bool bComputeChanges = true;
-
-		// if the number of matrices has changed between frames, then sync previous to current
-		if (CompSpaceTransforms.Num() != DynamicData->PrevTransforms.Num())
-		{
-			DynamicData->SetPrevTransforms(CompSpaceTransforms);
-			DynamicData->ChangedCount = CompSpaceTransforms.Num();
-			bComputeChanges = false; // Optimization to just force all transforms as changed and skip comparison
-		}
-
-		DynamicData->SetTransforms(CompSpaceTransforms);
-
-		// The number of transforms for current and previous should match now
-		check(DynamicData->PrevTransforms.Num() == DynamicData->Transforms.Num());
-
-		if (bComputeChanges)
-		{
-			DynamicData->DetermineChanges();
-		}
-	}
-
-	if (!bEditorMode && !bInitialization)
+#if WITH_EDITOR
+	// zero out transfrom matrices if they are marked to be hidden 
+	if (RestCollection && RestCollection->GetGeometryCollection())
 	{
-		if (DynamicData && DynamicData->ChangedCount == 0)
-		{
-			GDynamicDataPool.Release(DynamicData);
-			DynamicData = nullptr;
+		static const FName HideAttribute{ "Hide" };
 
-			// Change of state?
-			if (bIsMoving && !bForceMotionBlur)
-			{
-				bIsMoving = false;
-				if (SceneProxy && SceneProxy->IsNaniteMesh())
-				{
-					FNaniteGeometryCollectionSceneProxy* NaniteProxy = static_cast<FNaniteGeometryCollectionSceneProxy*>(SceneProxy);
-					ENQUEUE_RENDER_COMMAND(NaniteProxyOnMotionEnd)(
-						[NaniteProxy] (FRHICommandListBase&)
-						{
-							NaniteProxy->OnMotionEnd();
-						}
-					);
-				}
-			}
-		}
-		else
+		const FGeometryCollection& Collection = *RestCollection->GetGeometryCollection();
+		if (const TManagedArray<bool>* HideTransforms = Collection.FindAttribute<bool>(HideAttribute, FGeometryCollection::TransformGroup))
 		{
-			// Change of state?
-			if (!bIsMoving && !bForceMotionBlur)
+			if (DynamicData->Transforms.Num() == HideTransforms->Num())
 			{
-				bIsMoving = true;
-				if (SceneProxy && SceneProxy->IsNaniteMesh())
+				for (int32 TransformIndex = 0; TransformIndex < HideTransforms->Num(); TransformIndex++)
 				{
-					FNaniteGeometryCollectionSceneProxy* NaniteProxy = static_cast<FNaniteGeometryCollectionSceneProxy*>(SceneProxy);
-					ENQUEUE_RENDER_COMMAND(NaniteProxyOnMotionBegin)(
-						[NaniteProxy] (FRHICommandListBase&)
-						{
-							NaniteProxy->OnMotionBegin();
-						}
-					);
+					if ((*HideTransforms)[TransformIndex])
+					{
+						DynamicData->Transforms[TransformIndex] = FMatrix44f(EForceInit::ForceInitToZero);
+					}
 				}
 			}
 		}
 	}
+#endif
 
 	return DynamicData;
 }
@@ -3336,6 +3241,13 @@ void UGeometryCollectionComponent::OnUpdateTransform(EUpdateTransformFlags Updat
 	{
 		PhysicsProxy->SetWorldTransform_External(GetComponentTransform());
 	}
+
+	if (SceneProxy && SceneProxy->IsNaniteMesh())
+	{
+		// Nanite scene proxy requires an render update because it may fail to detect rotation and scale changes in the transform 
+		MarkRenderTransformDirty();
+		MarkRenderDynamicDataDirty();
+	}
 }
 
 bool UGeometryCollectionComponent::HasAnySockets() const
@@ -4884,61 +4796,45 @@ void UGeometryCollectionComponent::OnDestroyPhysicsState()
 	OnComponentPhysicsStateChanged.Broadcast(this, EComponentPhysicsStateChange::Destroyed);
 }
 
-void UGeometryCollectionComponent::SendRenderDynamicData_Concurrent()
+void UGeometryCollectionComponent::SendDynamicDataToSceneProxy()
 {
-	//UE_LOG(UGCC_LOG, Log, TEXT("GeometryCollectionComponent[%p]::SendRenderDynamicData_Concurrent()"), this);
-	Super::SendRenderDynamicData_Concurrent();
-
-	// Only update the dynamic data if the dynamic collection is dirty
-	if (SceneProxy && ((DynamicCollection && DynamicCollection->IsDirty()) || CachePlayback))
+	FGeometryCollectionDynamicData* DynamicData = InitDynamicData();
+	if (SceneProxy && DynamicData)
 	{
-		FGeometryCollectionDynamicData* DynamicData = InitDynamicData(false /* initialization */);
-
-		if (DynamicData || SceneProxy->IsNaniteMesh())
+		if (SceneProxy->IsNaniteMesh())
 		{
 			INC_DWORD_STAT_BY(STAT_GCTotalTransforms, DynamicData ? DynamicData->Transforms.Num() : 0);
-			INC_DWORD_STAT_BY(STAT_GCChangedTransforms, DynamicData ? DynamicData->ChangedCount : 0);
 
-			// #todo (bmiller) Once ISMC changes have been complete, this is the best place to call this method
-			// but we can't currently because it's an inappropriate place to call MarkRenderStateDirty on the ISMC.
-			// RefreshEmbeddedGeometry();
-
-			// Enqueue command to send to render thread
-			if (SceneProxy->IsNaniteMesh())
-			{
-				FNaniteGeometryCollectionSceneProxy* GeometryCollectionSceneProxy = static_cast<FNaniteGeometryCollectionSceneProxy*>(SceneProxy);
-				ENQUEUE_RENDER_COMMAND(SendRenderDynamicData)(
-					[GeometryCollectionSceneProxy, DynamicData] (FRHICommandListBase&)
-					{
-						if (DynamicData)
-						{
-							GeometryCollectionSceneProxy->SetDynamicData_RenderThread(DynamicData, GeometryCollectionSceneProxy->GetLocalToWorld());
-						}
-						else
-						{
-							// No longer dynamic, make sure previous transforms are reset
-							GeometryCollectionSceneProxy->ResetPreviousTransforms_RenderThread();
-						}
-					}
-				);
-			}
-			else
-			{
-				FGeometryCollectionSceneProxy* GeometryCollectionSceneProxy = static_cast<FGeometryCollectionSceneProxy*>(SceneProxy);
-				ENQUEUE_RENDER_COMMAND(SendRenderDynamicData)(
-					[GeometryCollectionSceneProxy, DynamicData](FRHICommandListBase& RHICmdList)
-					{
-						if (GeometryCollectionSceneProxy)
-						{
-							GeometryCollectionSceneProxy->SetDynamicData_RenderThread(RHICmdList, DynamicData);
-						}
-					}
-				);
-			}
-		}		
+			const FMatrix RenderMatrix{ GetRenderMatrix() };
+			FNaniteGeometryCollectionSceneProxy* GeometryCollectionSceneProxy = static_cast<FNaniteGeometryCollectionSceneProxy*>(SceneProxy);
+			ENQUEUE_RENDER_COMMAND(SendRenderDynamicData)(
+				[GeometryCollectionSceneProxy, DynamicData, RenderMatrix](FRHICommandListBase&)
+				{
+					GeometryCollectionSceneProxy->SetDynamicData_RenderThread(DynamicData, RenderMatrix);
+				}
+			);
+		}
+		else
+		{
+			FGeometryCollectionSceneProxy* GeometryCollectionSceneProxy = static_cast<FGeometryCollectionSceneProxy*>(SceneProxy);
+			ENQUEUE_RENDER_COMMAND(SendRenderDynamicData)(
+				[GeometryCollectionSceneProxy, DynamicData](FRHICommandListBase& RHICmdList)
+				{
+					GeometryCollectionSceneProxy->SetDynamicData_RenderThread(RHICmdList, DynamicData);
+				}
+			);
+		}
 	}
 }
 
+void UGeometryCollectionComponent::SendRenderDynamicData_Concurrent()
+{
+	//UE_LOG(UGCC_LOG, Log, TEXT("GeometryCollectionComponent[%p]::SendRenderDynamicData_Concurrent()"), this);
+	Super::SendRenderDynamicData_Concurrent();
+
+	SendDynamicDataToSceneProxy();
+}
+
 void UGeometryCollectionComponent::SetCollisionObjectType(ECollisionChannel Channel)
 {
 	Super::SetCollisionObjectType(Channel);
diff --git a/Engine/Source/Runtime/Experimental/GeometryCollectionEngine/Private/GeometryCollection/GeometryCollectionSceneProxy.cpp b/Engine/Source/Runtime/Experimental/GeometryCollectionEngine/Private/GeometryCollection/GeometryCollectionSceneProxy.cpp
--- a/Engine/Source/Runtime/Experimental/GeometryCollectionEngine/Private/GeometryCollection/GeometryCollectionSceneProxy.cpp
+++ b/Engine/Source/Runtime/Experimental/GeometryCollectionEngine/Private/GeometryCollection/GeometryCollectionSceneProxy.cpp
@@ -148,8 +148,7 @@ FGeometryCollectionSceneProxy::FGeometryCollectionSceneProxy(UGeometryCollection
 		}
 	}
 
-	Component->GetRestTransforms(RestTransforms);
-	NumTransforms = RestTransforms.Num();
+	NumTransforms = GeometryCollection? GeometryCollection->NumElements(FTransformCollection::TransformGroup) : 0;
 
 #if GEOMETRYCOLLECTION_EDITOR_SELECTION
 	// Render by SubSection if we are in the rigid body picker.
@@ -190,19 +189,6 @@ FGeometryCollectionSceneProxy::FGeometryCollectionSceneProxy(UGeometryCollection
 		}
 	}
 
-	// Get hidden geometry and zero the associated transforms.
-	Component->GetHiddenTransforms(HiddenTransforms);
-	if (HiddenTransforms.Num())
-	{
-		check(HiddenTransforms.Num() == RestTransforms.Num());
-		for (int32 TransformIndex = 0; TransformIndex < RestTransforms.Num(); ++TransformIndex)
-		{
-			if (HiddenTransforms[TransformIndex])
-			{
-				RestTransforms[TransformIndex] = FMatrix44f(EForceInit::ForceInitToZero);
-			}
-		}
-	}
 #endif
 
 	// #todo(dmp): This flag means that when motion blur is turned on, it will always render geometry collections into the
@@ -267,7 +253,7 @@ void FGeometryCollectionSceneProxy::SetupVertexFactory(FRHICommandListBase& RHIC
 	{
 		Data.BoneMapSRV = MeshResource.BoneMapVertexBuffer.GetSRV();
 		Data.BoneTransformSRV = TransformBuffers[CurrentTransformBufferIndex].VertexBufferSRV;
-		Data.BonePrevTransformSRV = PrevTransformBuffers[CurrentTransformBufferIndex].VertexBufferSRV;
+		Data.BonePrevTransformSRV = Data.BoneTransformSRV; // setup : both prev and crrent are the same 
 	}
 	else
 	{
@@ -295,20 +281,15 @@ void FGeometryCollectionSceneProxy::CreateRenderThreadResources(FRHICommandListB
 	{
 		// Initialize transform buffers and upload rest transforms.
 		TransformBuffers.AddDefaulted(1);
-		PrevTransformBuffers.AddDefaulted(1);
 
 		TransformBuffers[0].NumTransforms = NumTransforms;
-		PrevTransformBuffers[0].NumTransforms = NumTransforms;
 		TransformBuffers[0].InitResource(RHICmdList);
-		PrevTransformBuffers[0].InitResource(RHICmdList);
 
 		const bool bLocalGeometryCollectionTripleBufferUploads = (GGeometryCollectionTripleBufferUploads != 0) && bSupportsTripleBufferVertexUpload;
 		const EResourceLockMode LockMode = bLocalGeometryCollectionTripleBufferUploads ? RLM_WriteOnly_NoOverwrite : RLM_WriteOnly;
 
 		FGeometryCollectionTransformBuffer& TransformBuffer = GetCurrentTransformBuffer();
 		TransformBuffer.UpdateDynamicData(RHICmdList, DynamicData->Transforms, LockMode);
-		FGeometryCollectionTransformBuffer& PrevTransformBuffer = GetCurrentPrevTransformBuffer();
-		PrevTransformBuffer.UpdateDynamicData(RHICmdList, DynamicData->PrevTransforms, LockMode);
 	}
 	else
 	{
@@ -381,7 +362,6 @@ void FGeometryCollectionSceneProxy::DestroyRenderThreadResources()
 		for (int32 i = 0; i < TransformBuffers.Num(); i++)
 		{
 			TransformBuffers[i].ReleaseResource();
-			PrevTransformBuffers[i].ReleaseResource();
 		}
 		TransformBuffers.Reset();
 	}
@@ -426,12 +406,6 @@ void FGeometryCollectionSceneProxy::SetDynamicData_RenderThread(FRHICommandListB
 	{
 		return;
 	}
-	// Early out if if we are applying (non-dynamic) rest transforms over multiple frames.
-	if (!DynamicData->IsDynamic && TransformVertexBuffersContainsRestTransforms)
-	{
-		return;
-	}
-	TransformVertexBuffersContainsRestTransforms = !DynamicData->IsDynamic;
 		
 	if (bSupportsManualVertexFetch)
 	{
@@ -440,14 +414,12 @@ void FGeometryCollectionSceneProxy::SetDynamicData_RenderThread(FRHICommandListB
 		if (bLocalGeometryCollectionTripleBufferUploads && TransformBuffers.Num() == 1)
 		{
 			TransformBuffers.AddDefaulted(2);
-			PrevTransformBuffers.AddDefaulted(2);
+			check(TransformBuffers.Num() == 3);
 
-			for (int32 i = 1; i < 3; i++)
+			for (int32 i = 1; i < TransformBuffers.Num(); i++)
 			{
 				TransformBuffers[i].NumTransforms = NumTransforms;
-				PrevTransformBuffers[i].NumTransforms = NumTransforms;
 				TransformBuffers[i].InitResource(RHICmdList);
-				PrevTransformBuffers[i].InitResource(RHICmdList);
 			}
 		}
 
@@ -462,34 +434,8 @@ void FGeometryCollectionSceneProxy::SetDynamicData_RenderThread(FRHICommandListB
 
 			VertexFactory.SetBoneTransformSRV(TransformBuffer.VertexBufferSRV);
 			VertexFactory.SetBonePrevTransformSRV(PrevTransformBuffer.VertexBufferSRV);
-
-#if WITH_EDITOR
-			// Implement hiding geometry in editor by zeroing the transform.
-			// Could move this to InitDynamicData?
-			if (HiddenTransforms.Num())
-			{
-				for (int32 TransformIndex = 0; TransformIndex < DynamicData->Transforms.Num(); ++TransformIndex)
-				{
-					if (HiddenTransforms[TransformIndex])
-					{
-						DynamicData->Transforms[TransformIndex] = FMatrix44f(EForceInit::ForceInitToZero);
-						DynamicData->PrevTransforms[TransformIndex] = FMatrix44f(EForceInit::ForceInitToZero);
-					}
-				}
-			}
-#endif
-				
-			if (DynamicData->IsDynamic)
-			{
-				TransformBuffer.UpdateDynamicData(RHICmdList, DynamicData->Transforms, LockMode);
-				PrevTransformBuffer.UpdateDynamicData(RHICmdList, DynamicData->PrevTransforms, LockMode);
-			}
-			else
-			{
-				// If we are rendering the base mesh geometry then use RestTransforms for both current and previous transforms.
-				TransformBuffer.UpdateDynamicData(RHICmdList, RestTransforms, LockMode);
-				PrevTransformBuffer.UpdateDynamicData(RHICmdList, RestTransforms, LockMode);
-			}
+	
+			TransformBuffer.UpdateDynamicData(RHICmdList, DynamicData->Transforms, LockMode);
 
 			UpdateLooseParameter(VertexFactory, TransformBuffer.VertexBufferSRV, PrevTransformBuffer.VertexBufferSRV, MeshResource.BoneMapVertexBuffer.GetSRV());
 
@@ -505,7 +451,7 @@ void FGeometryCollectionSceneProxy::SetDynamicData_RenderThread(FRHICommandListB
 	}
 	else
 	{
-		UpdateSkinnedPositions(RHICmdList, DynamicData->IsDynamic ? DynamicData->Transforms : RestTransforms);
+		UpdateSkinnedPositions(RHICmdList, DynamicData->Transforms);
 	}
 
 #if RHI_RAYTRACING
@@ -696,13 +642,8 @@ void FGeometryCollectionSceneProxy::GetDynamicMeshElements(const TArray<const FS
 				continue;
 			}
 
-			// If not dynamic then use the section array with interior fracture surfaces removed.
-			bool bRemoveInternalFaces = DynamicData != nullptr && !DynamicData->IsDynamic && MeshDescription.SectionsNoInternal.Num();
-
-#if WITH_EDITOR
 			// If hiding geometry in editor then we don't remove hidden faces.
-			bRemoveInternalFaces &= HiddenTransforms.Num() == 0;
-#endif
+			const bool bRemoveInternalFaces = false;
 
 #if GEOMETRYCOLLECTION_EDITOR_SELECTION
 			// If using subsections then use the subsection array. 
@@ -833,9 +774,7 @@ void FGeometryCollectionSceneProxy::GetDynamicRayTracingInstances(FRayTracingIns
 
 	SetupVertexFactory(Collector.GetRHICommandList(), GeometryCollectionVertexFactory);
 
-	// If not dynamic then use the section array with interior fracture surfaces removed.
-	const bool bRemoveInternalFaces = DynamicData != nullptr && !DynamicData->IsDynamic && MeshDescription.SectionsNoInternal.Num();
-	TArray<FGeometryCollectionMeshElement> const& SectionArray = bRemoveInternalFaces ? MeshDescription.SectionsNoInternal : MeshDescription.Sections;
+	TArray<FGeometryCollectionMeshElement> const& SectionArray = MeshDescription.Sections;
 
 	UpdatingRayTracingGeometry_RenderingThread(SectionArray);
 
@@ -985,12 +924,10 @@ uint32 FGeometryCollectionSceneProxy::GetAllocatedSize() const
 		+ Materials.GetAllocatedSize()
 		+ MeshDescription.Sections.GetAllocatedSize()
 		+ MeshDescription.SubSections.GetAllocatedSize()
-		+ RestTransforms.GetAllocatedSize()
 		+ (SkinnedPositionVertexBuffer.GetAllowCPUAccess() ? SkinnedPositionVertexBuffer.GetStride() * SkinnedPositionVertexBuffer.GetNumVertices() : 0)
 #if WITH_EDITOR
 		+ BoneColors.GetAllocatedSize()
 		+ (ColorVertexBuffer.GetAllowCPUAccess() ? ColorVertexBuffer.GetStride() * ColorVertexBuffer.GetNumVertices() : 0)
-		+ HiddenTransforms.GetAllocatedSize()
 #endif
 #if GEOMETRYCOLLECTION_EDITOR_SELECTION
 		+ HitProxies.GetAllocatedSize()
@@ -1006,7 +943,6 @@ uint32 FGeometryCollectionSceneProxy::GetAllocatedSize() const
 FNaniteGeometryCollectionSceneProxy::FNaniteGeometryCollectionSceneProxy(UGeometryCollectionComponent* Component)
 : Nanite::FSceneProxyBase(Component)
 , GeometryCollection(Component->GetRestCollection())
-, bCurrentlyInMotion(false)
 , bRequiresGPUSceneUpdate(false)
 , bEnableBoneSelection(false)
 {
@@ -1273,13 +1209,6 @@ uint32 FNaniteGeometryCollectionSceneProxy::GetMemoryFootprint() const
 	return sizeof(*this) + GetAllocatedSize();
 }
 
-void FNaniteGeometryCollectionSceneProxy::OnTransformChanged(FRHICommandListBase& RHICmdList)
-{
-	Super::OnTransformChanged(RHICmdList);
-
-	SetDynamicData_RenderThread(DynamicData, GetLocalToWorld());
-}
-
 void FNaniteGeometryCollectionSceneProxy::GetNaniteResourceInfo(uint32& ResourceID, uint32& HierarchyOffset, uint32& ImposterIndex) const
 {
 	ResourceID = NaniteResourceID;
@@ -1320,77 +1249,66 @@ void FNaniteGeometryCollectionSceneProxy::SetDynamicData_RenderThread(FGeometryC
 		DynamicData = NewDynamicData;
 	}
 
-	// Are we currently simulating?
-	if (NewDynamicData->IsDynamic)
-	{
-		FInstanceSceneDataBuffers::FAccessTag AccessTag(PointerHash(this));
-		FInstanceSceneDataBuffers::FWriteView ProxyData = InstanceSceneDataBuffersImpl.BeginWriteAccess(AccessTag);
-		InstanceSceneDataBuffersImpl.SetPrimitiveLocalToWorld(PrimitiveLocalToWorld, AccessTag);
+	FInstanceSceneDataBuffers::FAccessTag AccessTag(PointerHash(this));
+	FInstanceSceneDataBuffers::FWriteView ProxyData = InstanceSceneDataBuffersImpl.BeginWriteAccess(AccessTag);
+	InstanceSceneDataBuffersImpl.SetPrimitiveLocalToWorld(PrimitiveLocalToWorld, AccessTag);
 
-		const TSharedPtr<FGeometryCollection, ESPMode::ThreadSafe> Collection = GeometryCollection->GetGeometryCollection();
-		const TManagedArray<int32>& TransformToGeometryIndices	 = Collection->TransformToGeometryIndex;
-		const TManagedArray<TSet<int32>>& TransformChildren		 = Collection->Children;
-		const TManagedArray<int32>& SimulationType				 = Collection->SimulationType;
+	const TSharedPtr<FGeometryCollection, ESPMode::ThreadSafe> Collection = GeometryCollection->GetGeometryCollection();
+	const TManagedArray<int32>& TransformToGeometryIndices	 = Collection->TransformToGeometryIndex;
+	const TManagedArray<TSet<int32>>& TransformChildren		 = Collection->Children;
+	const TManagedArray<int32>& SimulationType				 = Collection->SimulationType;
 
-		const int32 TransformCount = NewDynamicData->Transforms.Num();
-		check(TransformCount == TransformToGeometryIndices.Num());
-		check(TransformCount == TransformChildren.Num());
-		check(TransformCount == NewDynamicData->PrevTransforms.Num());
+	const int32 TransformCount = NewDynamicData->Transforms.Num();
+	check(TransformCount == TransformToGeometryIndices.Num());
+	check(TransformCount == TransformChildren.Num());
 
+	// set the prev by copying the last current
+	ProxyData.PrevInstanceToPrimitiveRelative = ProxyData.InstanceToPrimitiveRelative;
+	bCanSkipRedundantTransformUpdates = false; // shoudl we compare the transform to better decide about this ? 
 
-		ProxyData.InstanceToPrimitiveRelative.Reset(TransformCount);
-		ProxyData.PrevInstanceToPrimitiveRelative.Reset(TransformCount);
-		ProxyData.InstanceLocalBounds.Reset(TransformCount);
-		ProxyData.InstanceHierarchyOffset.Reset(TransformCount);
+	ProxyData.InstanceToPrimitiveRelative.Reset(TransformCount);
+	ProxyData.InstanceLocalBounds.Reset(TransformCount);
+	ProxyData.InstanceHierarchyOffset.Reset(TransformCount);
 		
 #if GEOMETRYCOLLECTION_EDITOR_SELECTION
-		ProxyData.InstanceEditorData.Reset(bEnableBoneSelection ? TransformCount : 0);
+	ProxyData.InstanceEditorData.Reset(bEnableBoneSelection ? TransformCount : 0);
 #endif
 
-		ProxyData.Flags.bHasPerInstanceDynamicData = true;
-		ProxyData.Flags.bHasPerInstanceLocalBounds = true;
-		ProxyData.Flags.bHasPerInstanceHierarchyOffset = true;
+	ProxyData.Flags.bHasPerInstanceDynamicData = true;
+	ProxyData.Flags.bHasPerInstanceLocalBounds = true;
+	ProxyData.Flags.bHasPerInstanceHierarchyOffset = true;
 
-		for (int32 TransformIndex = 0; TransformIndex < TransformCount; ++TransformIndex)
+	for (int32 TransformIndex = 0; TransformIndex < TransformCount; ++TransformIndex)
+	{
+		const int32 TransformToGeometryIndex = TransformToGeometryIndices[TransformIndex];
+		if (SimulationType[TransformIndex] != FGeometryCollection::ESimulationTypes::FST_Rigid)
 		{
-			const int32 TransformToGeometryIndex = TransformToGeometryIndices[TransformIndex];
-			if (SimulationType[TransformIndex] != FGeometryCollection::ESimulationTypes::FST_Rigid)
-			{
-				continue;
-			}
-
-			const FGeometryNaniteData& NaniteData = GeometryNaniteData[TransformToGeometryIndex];
-
-			const FRenderTransform& InstanceToPrimitiveRelative = ProxyData.InstanceToPrimitiveRelative.Emplace_GetRef(InstanceSceneDataBuffersImpl.ComputeInstanceToPrimitiveRelative(NewDynamicData->Transforms[TransformIndex], AccessTag));
-
-			FRenderTransform& PrevInstanceToPrimitiveRelative = ProxyData.PrevInstanceToPrimitiveRelative.Emplace_GetRef();
+			continue;
+		}
 
-			if (bCurrentlyInMotion)
-			{
-				PrevInstanceToPrimitiveRelative = InstanceSceneDataBuffersImpl.ComputeInstanceToPrimitiveRelative(NewDynamicData->PrevTransforms[TransformIndex], AccessTag);
-			}
-			else
-			{
-				PrevInstanceToPrimitiveRelative = InstanceToPrimitiveRelative;
-			}
+		const FGeometryNaniteData& NaniteData = GeometryNaniteData[TransformToGeometryIndex];
+		const FRenderTransform& InstanceToPrimitiveRelative = ProxyData.InstanceToPrimitiveRelative.Emplace_GetRef(InstanceSceneDataBuffersImpl.ComputeInstanceToPrimitiveRelative(NewDynamicData->Transforms[TransformIndex], AccessTag));
 
-			ProxyData.InstanceLocalBounds.Emplace(PadInstanceLocalBounds(NaniteData.LocalBounds));
-			ProxyData.InstanceHierarchyOffset.Emplace(NaniteData.HierarchyOffset);
+		ProxyData.InstanceLocalBounds.Emplace(PadInstanceLocalBounds(NaniteData.LocalBounds));
+		ProxyData.InstanceHierarchyOffset.Emplace(NaniteData.HierarchyOffset);
 			
 #if GEOMETRYCOLLECTION_EDITOR_SELECTION
-			if (bEnableBoneSelection)
-			{
-				ProxyData.InstanceEditorData.Emplace(FInstanceEditorData::Pack(HitProxies[TransformToGeometryIndex]->Id.GetColor(), false));
-			}
-#endif
+		if (bEnableBoneSelection)
+		{
+			ProxyData.InstanceEditorData.Emplace(FInstanceEditorData::Pack(HitProxies[TransformToGeometryIndex]->Id.GetColor(), false));
 		}
-		InstanceSceneDataBuffersImpl.EndWriteAccess(AccessTag);
+#endif
 	}
-	else
+
+	// make sure the previous transform count do match the current one
+	// if not simply use the current as previous
+	if (ProxyData.PrevInstanceToPrimitiveRelative.Num() != ProxyData.InstanceToPrimitiveRelative.Num())
 	{
-		// Rendering base geometry, use rest transforms rather than simulated transforms.
-		// ...
+		ProxyData.PrevInstanceToPrimitiveRelative = ProxyData.InstanceToPrimitiveRelative;
+		bCanSkipRedundantTransformUpdates = true;
 	}
+
+	InstanceSceneDataBuffersImpl.EndWriteAccess(AccessTag);
 }
 
 void FNaniteGeometryCollectionSceneProxy::ResetPreviousTransforms_RenderThread()
@@ -1423,20 +1341,6 @@ void FNaniteGeometryCollectionSceneProxy::FlushGPUSceneUpdate_GameThread()
 	);
 }
 
-void FNaniteGeometryCollectionSceneProxy::OnMotionBegin()
-{
-	bCurrentlyInMotion = true;
-	bCanSkipRedundantTransformUpdates = false;
-}
-
-void FNaniteGeometryCollectionSceneProxy::OnMotionEnd()
-{
-	bCurrentlyInMotion = false;
-	bCanSkipRedundantTransformUpdates = true;
-	ResetPreviousTransforms_RenderThread();
-}
-
-
 bool FNaniteGeometryCollectionSceneProxy::ShowCollisionMeshes(const FEngineShowFlags& EngineShowFlags) const
 {
 	if (IsCollisionEnabled())
diff --git a/Engine/Source/Runtime/Experimental/GeometryCollectionEngine/Private/GeometryCollection/GeometryCollectionSceneProxy.h b/Engine/Source/Runtime/Experimental/GeometryCollectionEngine/Private/GeometryCollection/GeometryCollectionSceneProxy.h
--- a/Engine/Source/Runtime/Experimental/GeometryCollectionEngine/Private/GeometryCollection/GeometryCollectionSceneProxy.h
+++ b/Engine/Source/Runtime/Experimental/GeometryCollectionEngine/Private/GeometryCollection/GeometryCollectionSceneProxy.h
@@ -82,10 +82,7 @@ inline void CopyTransformsWithConversionWhenNeeded(TArray<FMatrix44f>& DstTransf
 struct FGeometryCollectionDynamicData
 {
 	TArray<FMatrix44f> Transforms;
-	TArray<FMatrix44f> PrevTransforms;
-	uint32 ChangedCount;
-	uint8 IsDynamic : 1;
-	uint8 IsLoading : 1;
+	uint64 FrameIndex = 0;
 
 	FGeometryCollectionDynamicData()
 	{
@@ -95,16 +92,7 @@ struct FGeometryCollectionDynamicData
 	void Reset()
 	{
 		Transforms.Reset();
-		PrevTransforms.Reset();
-		IsDynamic = false;
-		IsLoading = false;
-	}
-
-	UE_DEPRECATED(5.3, "Use FTransform version of SetTransforms instead")
-	void SetTransforms(const TArray<FMatrix>& InTransforms)
-	{
-		// use for LWC as FMatrix and FMatrix44f are different when LWC is on 
-		CopyTransformsWithConversionWhenNeeded(Transforms, InTransforms);
+		FrameIndex = GFrameCounter;
 	}
 
 	void SetTransforms(const TArray<FTransform>& InTransforms)
@@ -117,64 +105,6 @@ struct FGeometryCollectionDynamicData
 	{
 		CopyTransformsWithConversionWhenNeeded(Transforms, InTransforms);
 	}
-
-	UE_DEPRECATED(5.3, "Use FTransform version of SetPrevTransforms instead")
-	void SetPrevTransforms(const TArray<FMatrix>& InTransforms)
-	{
-		// use for LWC as FMatrix and FMatrix44f are different when LWC is on 
-		CopyTransformsWithConversionWhenNeeded(PrevTransforms, InTransforms);
-	}
-
-	void SetPrevTransforms(const TArray<FTransform>& InTransforms)
-	{
-		// use for LWC as FMatrix and FMatrix44f are different when LWC is on 
-		CopyTransformsWithConversionWhenNeeded(PrevTransforms, InTransforms);
-	}
-
-	void SetPrevTransforms(const TArray<FTransform3f>& InTransforms)
-	{
-		CopyTransformsWithConversionWhenNeeded(PrevTransforms, InTransforms);
-	}
-
-	UE_DEPRECATED(5.3, "Use FTransform version of SetAllTransforms instead")
-	void SetAllTransforms(const TArray<FMatrix>& InTransforms)
-	{
-		PRAGMA_DISABLE_DEPRECATION_WARNINGS
-		SetTransforms(InTransforms);
-		PRAGMA_ENABLE_DEPRECATION_WARNINGS
-		PrevTransforms = Transforms;
-		ChangedCount = Transforms.Num();
-	}
-
-	void SetAllTransforms(const TArray<FTransform>& InTransforms)
-	{
-		SetTransforms(InTransforms);
-		PrevTransforms = Transforms;
-		ChangedCount = Transforms.Num();
-	}
-
-	void DetermineChanges()
-	{
-		// Check if previous transforms are the same as current
-		const float EqualTolerance = 1e-6;
-
-		check(Transforms.Num() == PrevTransforms.Num());
-		if (Transforms.Num() != PrevTransforms.Num())
-		{
-			ChangedCount = Transforms.Num();
-		}
-		else
-		{
-			ChangedCount = 0;
-			for (int32 TransformIndex = 0; TransformIndex < Transforms.Num(); ++TransformIndex)
-			{
-				if (!PrevTransforms[TransformIndex].Equals(Transforms[TransformIndex], EqualTolerance))
-				{
-					++ChangedCount;
-				}
-			}
-		}
-	}
 };
 
 
@@ -204,7 +134,6 @@ private:
 *   NOTE : This class is still in flux, and has a few pending todos. Your comments and 
 *   thoughts are appreciated though. The remaining items to address involve:
 *   - @todo double buffer - The double buffering of the FGeometryCollectionDynamicData.
-*   - @todo previous state - Saving the previous FGeometryCollectionDynamicData for rendering motion blur.
 *   - @todo GPU skin : Make the skinning use the GpuVertexShader
 */
 class FGeometryCollectionSceneProxy final : public FPrimitiveSceneProxy
@@ -217,7 +146,6 @@ class FGeometryCollectionSceneProxy final : public FPrimitiveSceneProxy
 	FGeometryCollectionMeshDescription MeshDescription;
 
 	int32 NumTransforms = 0;
-	TArray<FMatrix44f> RestTransforms;
 	TSharedPtr<FGeometryCollection, ESPMode::ThreadSafe> GeometryCollection;
 
 	FCollisionResponseContainer CollisionResponse;
@@ -230,11 +158,9 @@ class FGeometryCollectionSceneProxy final : public FPrimitiveSceneProxy
 	FPositionVertexBuffer SkinnedPositionVertexBuffer;
 
 	int32 CurrentTransformBufferIndex = 0;
-	bool TransformVertexBuffersContainsRestTransforms = true;
 	bool bSupportsTripleBufferVertexUpload = false;
 	bool bRenderResourcesCreated = false;
 	TArray<FGeometryCollectionTransformBuffer, TInlineAllocator<3>> TransformBuffers;
-	TArray<FGeometryCollectionTransformBuffer, TInlineAllocator<3>> PrevTransformBuffers;
 
 	FGeometryCollectionDynamicData* DynamicData = nullptr;
 
@@ -245,8 +171,7 @@ class FGeometryCollectionSceneProxy final : public FPrimitiveSceneProxy
 	FColorVertexBuffer ColorVertexBuffer;
 	FGeometryCollectionVertexFactory VertexFactoryDebugColor;
 	UMaterialInterface* BoneSelectedMaterial = nullptr;
-	TArray<bool> HiddenTransforms;
-#endif
+	#endif
 
 #if GEOMETRYCOLLECTION_EDITOR_SELECTION
 	bool bUsesSubSections = false;
@@ -310,7 +235,9 @@ protected:
 	}
 	FGeometryCollectionTransformBuffer& GetCurrentPrevTransformBuffer()
 	{
-		return PrevTransformBuffers[CurrentTransformBufferIndex];
+		const int32 NumBuffers = TransformBuffers.Num();
+		const int32 PreviousIndex = (CurrentTransformBufferIndex + NumBuffers - 1) % NumBuffers;
+		return TransformBuffers[PreviousIndex];
 	}
 
 	void CycleTransformBuffers(bool bCycle)
@@ -351,8 +278,6 @@ public:
 
 	virtual uint32 GetMemoryFootprint() const override;
 
-	virtual void OnTransformChanged(FRHICommandListBase& RHICmdList) override;
-
 	// FSceneProxyBase interface.
 	virtual void GetNaniteResourceInfo(uint32& ResourceID, uint32& HierarchyOffset, uint32& ImposterIndex) const override;
 
@@ -375,9 +300,6 @@ public:
 		return bRequiresGPUSceneUpdate;
 	}
 
-	void OnMotionBegin();
-	void OnMotionEnd();
-
 	inline virtual void GetLCIs(FLCIArray& LCIs) override
 	{
 		FLightCacheInterface* LCI = &EmptyLightCacheInfo;
@@ -406,7 +328,6 @@ protected:
 	uint32 bCastShadow : 1;
 	uint32 bReverseCulling : 1;
 	uint32 bHasMaterialErrors : 1;
-	uint32 bCurrentlyInMotion : 1;
 	uint32 bRequiresGPUSceneUpdate : 1;
 	uint32 bEnableBoneSelection : 1;
 	
diff --git a/Engine/Source/Runtime/Experimental/GeometryCollectionEngine/Public/GeometryCollection/GeometryCollectionComponent.h b/Engine/Source/Runtime/Experimental/GeometryCollectionEngine/Public/GeometryCollection/GeometryCollectionComponent.h
--- a/Engine/Source/Runtime/Experimental/GeometryCollectionEngine/Public/GeometryCollection/GeometryCollectionComponent.h
+++ b/Engine/Source/Runtime/Experimental/GeometryCollectionEngine/Public/GeometryCollection/GeometryCollectionComponent.h
@@ -1542,7 +1542,7 @@ protected:
 	GEOMETRYCOLLECTIONENGINE_API virtual void CollectPSOPrecacheData(const FPSOPrecacheParams& BasePrecachePSOParams, FMaterialInterfacePSOPrecacheParamsList& OutParams) override;
 
 	/** Populate the dynamic particle data for the render thread. */
-	GEOMETRYCOLLECTIONENGINE_API FGeometryCollectionDynamicData* InitDynamicData(bool bInitialization);
+	GEOMETRYCOLLECTIONENGINE_API FGeometryCollectionDynamicData* InitDynamicData(bool bInitialization = false);
 
 	/** Reset the dynamic collection from the current rest state. */
 	GEOMETRYCOLLECTIONENGINE_API void ResetDynamicCollection();
@@ -1649,6 +1649,7 @@ private:
 	void ProcessRepStateDataOnPT();
 	void ProcessRepDynamicDataOnPT();
 	void InitializeRemovalDynamicAttributesIfNeeded();
+	void SendDynamicDataToSceneProxy();
 
 	// called when the rest transform are updated from SetRestState / ResetRestTransforms
 	// this updates only the renderer, the dynamic collection should be initialized when calling this function
@@ -1664,9 +1665,6 @@ private:
 	uint8 IsObjectDynamic : 1;
 	uint8 IsObjectLoading : 1;
 
-	/** True if GeometryCollection transforms have changed from previous tick. */
-	uint8 bIsMoving : 1;
-
 	int ViewLevel;
 
 	uint32 NavmeshInvalidationTimeSliceIndex;
