IAD cache journal tests for cache block created/deleted.

#rnx

diff --git a/Engine/Source/Runtime/Experimental/IoStore/OnDemand/Private/OnDemandInstallCache.cpp b/Engine/Source/Runtime/Experimental/IoStore/OnDemand/Private/OnDemandInstallCache.cpp
--- a/Engine/Source/Runtime/Experimental/IoStore/OnDemand/Private/OnDemandInstallCache.cpp
+++ b/Engine/Source/Runtime/Experimental/IoStore/OnDemand/Private/OnDemandInstallCache.cpp
@@ -1784,6 +1784,56 @@ TEST_CASE("IoStore::OnDemand::InstallCache::Journal", "[IoStoreOnDemand][Install
 			CHECK(Loc.CasLocation.BlockOffset == ExpectedLoc.BlockOffset);
 		}
 	}
+
+	SECTION("ReplayBlockCreatedAndDeleted")
+	{
+		// Arrange
+		const FCasBlockId ExpectedBlockId(42);
+
+		// Act
+		FTmpDirectoryScope _(TestBaseDir);
+		const FString JournalFile = TestBaseDir / TEXT("test.jrn");
+
+		FIoStatus Status = FCasJournal::Create(JournalFile);
+		CHECK(Status.IsOk());
+
+		FCasJournal::FTransaction Tx = FCasJournal::Begin(JournalFile);
+		Tx.BlockCreated(ExpectedBlockId);
+		Tx.BlockDeleted(ExpectedBlockId);
+
+		Status = FCasJournal::Commit(MoveTemp(Tx));
+		CHECK(Status.IsOk());
+
+		// Assert
+		FCasBlockId CreatedBlockId;
+		FCasBlockId DeletedBlockId;
+
+		Status = FCasJournal::Replay(
+			JournalFile,
+			[&CreatedBlockId, &DeletedBlockId](const FCasJournal::FEntry& JournalEntry)
+			{
+				switch(JournalEntry.Type())
+				{
+				case FCasJournal::FEntry::EType::BlockCreated:
+				{
+					CreatedBlockId = JournalEntry.BlockOperation.BlockId;
+					break;
+				}
+				case FCasJournal::FEntry::EType::BlockDeleted:
+				{
+					DeletedBlockId = JournalEntry.BlockOperation.BlockId;
+					break;
+				}
+				default:
+					CHECK(false);
+					break;
+				};
+			});
+
+		CHECK(Status.IsOk());
+		CHECK(CreatedBlockId == ExpectedBlockId);
+		CHECK(DeletedBlockId == ExpectedBlockId);
+	}
 }
 
 #endif // WITH_IOSTORE_ONDEMAND_TESTS
