IAD cache journal test for block access time.

diff --git a/Engine/Source/Runtime/Experimental/IoStore/OnDemand/Private/OnDemandInstallCache.cpp b/Engine/Source/Runtime/Experimental/IoStore/OnDemand/Private/OnDemandInstallCache.cpp
--- a/Engine/Source/Runtime/Experimental/IoStore/OnDemand/Private/OnDemandInstallCache.cpp
+++ b/Engine/Source/Runtime/Experimental/IoStore/OnDemand/Private/OnDemandInstallCache.cpp
@@ -1834,6 +1834,53 @@ TEST_CASE("IoStore::OnDemand::InstallCache::Journal", "[IoStoreOnDemand][Install
 		CHECK(CreatedBlockId == ExpectedBlockId);
 		CHECK(DeletedBlockId == ExpectedBlockId);
 	}
+
+	SECTION("ReplayBlockAccess")
+	{
+		// Arrange
+		const FCasBlockId ExpectedBlockId(462);
+		const uint64 ExpectedTicks = FDateTime::UtcNow().GetTicks();
+
+		// Act
+		FTmpDirectoryScope _(TestBaseDir);
+		const FString JournalFile = TestBaseDir / TEXT("test.jrn");
+
+		FIoStatus Status = FCasJournal::Create(JournalFile);
+		CHECK(Status.IsOk());
+
+		FCasJournal::FTransaction Tx = FCasJournal::Begin(JournalFile);
+		Tx.BlockAccess(ExpectedBlockId, ExpectedTicks);
+
+		Status = FCasJournal::Commit(MoveTemp(Tx));
+		CHECK(Status.IsOk());
+
+		// Assert
+		FCasBlockId BlockId;
+		uint64 Ticks = 0;
+
+		Status = FCasJournal::Replay(
+			JournalFile,
+			[&BlockId, &Ticks](const FCasJournal::FEntry& JournalEntry)
+			{
+				switch(JournalEntry.Type())
+				{
+				case FCasJournal::FEntry::EType::BlockAccess:
+				{
+					const FCasJournal::FEntry::FBlockOperation& Op = JournalEntry.BlockOperation;
+					BlockId	= Op.BlockId;
+					Ticks	= Op.UtcTicks;
+					break;
+				}
+				default:
+					CHECK(false);
+					break;
+				};
+			});
+
+		CHECK(Status.IsOk());
+		CHECK(BlockId == ExpectedBlockId);
+		CHECK(Ticks == ExpectedTicks);
+	}
 }
 
 #endif // WITH_IOSTORE_ONDEMAND_TESTS
