NetCore - Implemented FConnectionHandle which carry both parent connection and child connection information. This can be usefule for APIs where it's important to keep track of child connections as well.

diff --git a/Engine/Plugins/Runtime/GameplayAbilities/Source/GameplayAbilities/Private/Serialization/MinimalGameplayCueReplicationProxyReplicationFragment.cpp b/Engine/Plugins/Runtime/GameplayAbilities/Source/GameplayAbilities/Private/Serialization/MinimalGameplayCueReplicationProxyReplicationFragment.cpp
--- a/Engine/Plugins/Runtime/GameplayAbilities/Source/GameplayAbilities/Private/Serialization/MinimalGameplayCueReplicationProxyReplicationFragment.cpp
+++ b/Engine/Plugins/Runtime/GameplayAbilities/Source/GameplayAbilities/Private/Serialization/MinimalGameplayCueReplicationProxyReplicationFragment.cpp
@@ -127,7 +127,7 @@ void FMinimalGameplayCueReplicationProxyReplicationFragment::MimicMinimalGamepla
 		{			
 			if (const UNetConnection* OwnerNetConnection = OwningActor->GetNetConnection())
 			{
-				if (OwnerNetConnection->GetConnectionId() == Context.NetSerializationContext->GetLocalConnectionId())
+				if (OwnerNetConnection->GetConnectionHandle().GetParentConnectionId() == Context.NetSerializationContext->GetLocalConnectionId())
 				{
 					UpdateOwnerTagMap = false;
 				}
diff --git a/Engine/Plugins/Runtime/GameplayAbilities/Source/GameplayAbilities/Private/Serialization/PredictionKeyNetSerializer.cpp b/Engine/Plugins/Runtime/GameplayAbilities/Source/GameplayAbilities/Private/Serialization/PredictionKeyNetSerializer.cpp
--- a/Engine/Plugins/Runtime/GameplayAbilities/Source/GameplayAbilities/Private/Serialization/PredictionKeyNetSerializer.cpp
+++ b/Engine/Plugins/Runtime/GameplayAbilities/Source/GameplayAbilities/Private/Serialization/PredictionKeyNetSerializer.cpp
@@ -116,7 +116,7 @@ void FPredictionKeyNetSerializer::Quantize(FNetSerializationContext&, const FNet
 		{
 			if (const UNetConnection* Connection = PackageMap->GetConnection())
 			{
-				Target.OwningConnectionId = Connection->GetParentConnectionId();
+				Target.OwningConnectionId = Connection->GetConnectionHandle().GetParentConnectionId();
 			}
 		}
 	}
diff --git a/Engine/Plugins/Runtime/ReplicationSystemTestPlugin/Source/Private/Tests/ReplicationSystem/Filtering/TestConnectionHandle.cpp b/Engine/Plugins/Runtime/ReplicationSystemTestPlugin/Source/Private/Tests/ReplicationSystem/Filtering/TestConnectionHandle.cpp
new file mode 100644
--- /dev/null
+++ b/Engine/Plugins/Runtime/ReplicationSystemTestPlugin/Source/Private/Tests/ReplicationSystem/Filtering/TestConnectionHandle.cpp
@@ -0,0 +1,51 @@
+// Copyright Epic Games, Inc. All Rights Reserved.
+
+#include "NetworkAutomationTest.h"
+#include "NetworkAutomationTestMacros.h"
+#include "IrisTestMessageStreamOperators.h"
+#include "Iris/ReplicationSystem/Filtering/SharedConnectionFilterStatus.h"
+
+namespace UE::Net::Private
+{
+
+class FConnectionHandleTestFixture : public FNetworkAutomationTestSuiteFixture
+{
+};
+
+UE_NET_TEST_FIXTURE(FConnectionHandleTestFixture, ConnectionHandleIsInvalidByDefault)
+{
+	FConnectionHandle ConnHandle;
+	UE_NET_ASSERT_FALSE(ConnHandle.IsValid());
+}
+
+UE_NET_TEST_FIXTURE(FConnectionHandleTestFixture, ConnectionHandleWithParentConnectioIdIsValid)
+{
+	constexpr uint32 ParentConnId = 1U;
+	FConnectionHandle ConnHandle(ParentConnId);
+	UE_NET_ASSERT_TRUE(ConnHandle.IsValid());
+}
+
+UE_NET_TEST_FIXTURE(FConnectionHandleTestFixture, ConnectionHandleReturnsExpectedParentConnectionId)
+{
+	constexpr uint32 ParentConnId = 4U;
+	FConnectionHandle ConnHandle(ParentConnId);
+	UE_NET_ASSERT_EQ(ConnHandle.GetParentConnectionId(), ParentConnId);
+}
+
+UE_NET_TEST_FIXTURE(FConnectionHandleTestFixture, ConnectionHandleReturnsExpectedChildConnectionId)
+{
+	constexpr uint32 ParentConnId = 4U;
+	constexpr uint32 ChildConnId = 7U;
+	FConnectionHandle ConnHandle(ParentConnId, ChildConnId);
+	UE_NET_ASSERT_EQ(ConnHandle.GetChildConnectionId(), ChildConnId);
+}
+
+UE_NET_TEST_FIXTURE(FConnectionHandleTestFixture, ConnectionHandleWithOnlyValidChildConnectioIdIsInvalid)
+{
+	constexpr uint32 ParentConnId = UE::Net::InvalidConnectionId;
+	constexpr uint32 ChildConnId = 15U;
+	FConnectionHandle ConnHandle(ParentConnId, ChildConnId);
+	UE_NET_ASSERT_FALSE(ConnHandle.IsValid());
+}
+
+}
diff --git a/Engine/Plugins/Runtime/ReplicationSystemTestPlugin/Source/Private/Tests/ReplicationSystem/Filtering/TestSharedConnectionFilterStatus.cpp b/Engine/Plugins/Runtime/ReplicationSystemTestPlugin/Source/Private/Tests/ReplicationSystem/Filtering/TestSharedConnectionFilterStatus.cpp
new file mode 100644
--- /dev/null
+++ b/Engine/Plugins/Runtime/ReplicationSystemTestPlugin/Source/Private/Tests/ReplicationSystem/Filtering/TestSharedConnectionFilterStatus.cpp
@@ -0,0 +1,265 @@
+// Copyright Epic Games, Inc. All Rights Reserved.
+
+#include "NetworkAutomationTest.h"
+#include "NetworkAutomationTestMacros.h"
+#include "IrisTestMessageStreamOperators.h"
+#include "Iris/IrisConstants.h"
+#include "Iris/Core/IrisLog.h"
+#include "Iris/ReplicationSystem/Filtering/SharedConnectionFilterStatus.h"
+#include "Logging/LogScopedVerbosityOverride.h"
+
+namespace UE::Net::Private
+{
+
+class FSharedConnectionFilterStatusTestFixture : public FNetworkAutomationTestSuiteFixture
+{
+};
+
+class FSharedConnectionFilterStatusCollectionTestFixture : public FNetworkAutomationTestSuiteFixture
+{
+};
+
+// FSharedConnectionFilterStatus tests
+UE_NET_TEST_FIXTURE(FSharedConnectionFilterStatusTestFixture, ReplicationIsDisallowedByDefault)
+{
+	FSharedConnectionFilterStatus FilterStatus;
+	UE_NET_ASSERT_EQ(FilterStatus.GetFilterStatus(), ENetFilterStatus::Disallow);
+}
+
+UE_NET_TEST_FIXTURE(FSharedConnectionFilterStatusTestFixture, ReplicationIsNeverAllowedOnInvalidConnection)
+{
+	LOG_SCOPE_VERBOSITY_OVERRIDE(LogIrisFiltering, ELogVerbosity::Fatal);
+	FSharedConnectionFilterStatus FilterStatus;
+	FilterStatus.SetFilterStatus(FConnectionHandle(InvalidConnectionId), ENetFilterStatus::Allow);
+	UE_NET_ASSERT_EQ(FilterStatus.GetFilterStatus(), ENetFilterStatus::Disallow);
+}
+
+UE_NET_TEST_FIXTURE(FSharedConnectionFilterStatusTestFixture, ConnectionWithFilterStatusEstablishesParentConnection)
+{
+	constexpr uint32 ParentConnId = 1U;
+	constexpr uint32 OtherParentConnId = 2U;
+	FConnectionHandle ConnHandle0(ParentConnId);
+	FConnectionHandle ConnHandle1(OtherParentConnId);
+
+	FSharedConnectionFilterStatus FilterStatus;
+	FilterStatus.SetFilterStatus(ConnHandle0, ENetFilterStatus::Disallow);
+	UE_NET_ASSERT_EQ(FilterStatus.GetParentConnectionId(), ConnHandle0.GetParentConnectionId());
+
+	// Trying to set filter status for a different parent connection ID should not modify the group as parent connection ID has already been established
+	{
+		LOG_SCOPE_VERBOSITY_OVERRIDE(LogIrisFiltering, ELogVerbosity::Fatal);
+		FilterStatus.SetFilterStatus(ConnHandle1, ENetFilterStatus::Allow);
+	}
+	UE_NET_ASSERT_EQ(FilterStatus.GetParentConnectionId(), ConnHandle0.GetParentConnectionId());
+	UE_NET_ASSERT_EQ(FilterStatus.GetFilterStatus(), ENetFilterStatus::Disallow);
+}
+
+UE_NET_TEST_FIXTURE(FSharedConnectionFilterStatusTestFixture, EstablishedParentConnectionIsKeptWhenRemovingChildConnection)
+{
+	constexpr uint32 ParentConnId = 1U;
+	constexpr uint32 ChildConnId = 1U;
+	FConnectionHandle ConnHandle(ParentConnId, ChildConnId);
+
+	FSharedConnectionFilterStatus FilterStatus;
+	FilterStatus.SetFilterStatus(ConnHandle, ENetFilterStatus::Allow);
+	FilterStatus.RemoveConnection(ConnHandle);
+	UE_NET_ASSERT_EQ(FilterStatus.GetParentConnectionId(), ConnHandle.GetParentConnectionId());
+}
+
+UE_NET_TEST_FIXTURE(FSharedConnectionFilterStatusTestFixture, EstablishedParentConnectionIsRemovedWhenRemovingParentConnection)
+{
+	constexpr uint32 ParentConnId = 1U;
+	constexpr uint32 ChildConnId = 1U;
+	FConnectionHandle ConnHandle(ParentConnId, ChildConnId);
+
+	FSharedConnectionFilterStatus FilterStatus;
+	FilterStatus.SetFilterStatus(ConnHandle, ENetFilterStatus::Allow);
+	FilterStatus.RemoveConnection(FConnectionHandle(ConnHandle.GetParentConnectionId()));
+	UE_NET_ASSERT_EQ(FilterStatus.GetParentConnectionId(), InvalidConnectionId);
+}
+
+UE_NET_TEST_FIXTURE(FSharedConnectionFilterStatusTestFixture, ReplicationIsDisallowedWhenNoConnectionAllows)
+{
+	constexpr uint32 ParentConnectionId = 4711U;
+	FSharedConnectionFilterStatus FilterStatus;
+	for (const uint32 ChildConnectionId : {0, 1, 3, 2, 7})
+	{
+		FilterStatus.SetFilterStatus(FConnectionHandle(4711U, ChildConnectionId), ENetFilterStatus::Disallow);
+		UE_NET_ASSERT_EQ(FilterStatus.GetFilterStatus(), ENetFilterStatus::Disallow);
+	}
+}
+
+UE_NET_TEST_FIXTURE(FSharedConnectionFilterStatusTestFixture, ReplicationIsAllowedWhenOneOrMoreConnectionAllows)
+{
+	constexpr uint32 ParentConnectionId = 7512U;
+	FSharedConnectionFilterStatus FilterStatus;
+	// Note that the child connection ID currently needs a mix of odd and even numbers as we're using that information to choose between Allow and Disallow.
+	for (const uint32 ChildConnectionId : {0, 1, 3, 2, 7})
+	{
+		FilterStatus.SetFilterStatus(FConnectionHandle(4711U, ChildConnectionId), (ChildConnectionId & 1U ? ENetFilterStatus::Disallow :  ENetFilterStatus::Allow));
+	}
+
+	UE_NET_ASSERT_EQ(FilterStatus.GetFilterStatus(), ENetFilterStatus::Allow);
+}
+
+UE_NET_TEST_FIXTURE(FSharedConnectionFilterStatusTestFixture, FilterStatusIsAdjustedWhenConnectionIsRemoved)
+{
+	constexpr uint32 ParentConnId = 1U;
+	FConnectionHandle AllowReplicationConnHandle(ParentConnId, 1U);
+	FConnectionHandle DisallowReplicationConnHandle(ParentConnId, 2U);
+
+	// Test allowing first and disallowing second before removing the allow connection
+	{
+		FSharedConnectionFilterStatus FilterStatus;
+		FilterStatus.SetFilterStatus(AllowReplicationConnHandle, ENetFilterStatus::Allow);
+		FilterStatus.SetFilterStatus(DisallowReplicationConnHandle, ENetFilterStatus::Disallow);
+		FilterStatus.RemoveConnection(AllowReplicationConnHandle);
+		UE_NET_ASSERT_EQ(FilterStatus.GetFilterStatus(), ENetFilterStatus::Disallow);
+	}
+
+	// Test disallowing first and allowing second before removing the allow connection
+	{
+		FSharedConnectionFilterStatus FilterStatus;
+		FilterStatus.SetFilterStatus(DisallowReplicationConnHandle, ENetFilterStatus::Disallow);
+		FilterStatus.SetFilterStatus(AllowReplicationConnHandle, ENetFilterStatus::Allow);
+		FilterStatus.RemoveConnection(AllowReplicationConnHandle);
+		UE_NET_ASSERT_EQ(FilterStatus.GetFilterStatus(), ENetFilterStatus::Disallow);
+	}
+}
+
+UE_NET_TEST_FIXTURE(FSharedConnectionFilterStatusTestFixture, FilterStatusIsAdjustedWhenParentConnectionIsRemoved)
+{
+	constexpr uint32 ParentConnId = 1U;
+	FConnectionHandle AllowReplicationConnHandle(ParentConnId, 1U);
+
+	FSharedConnectionFilterStatus FilterStatus;
+	FilterStatus.SetFilterStatus(AllowReplicationConnHandle, ENetFilterStatus::Allow);
+	// Removing the parent connection should act as removing all child connections too.
+	FilterStatus.RemoveConnection(FConnectionHandle(ParentConnId));
+	UE_NET_ASSERT_EQ(FilterStatus.GetFilterStatus(), ENetFilterStatus::Disallow);
+}
+
+// FSharedConnectionFilterStatusCollection tests
+UE_NET_TEST_FIXTURE(FSharedConnectionFilterStatusCollectionTestFixture, ReplicationIsDisallowedByDefault)
+{
+	FSharedConnectionFilterStatusCollection Collection;
+	for (const uint32 ParentConnectionId : {InvalidConnectionId, 1U, 4711U, 99U})
+	{
+		UE_NET_ASSERT_EQ(Collection.GetFilterStatus(ParentConnectionId), ENetFilterStatus::Disallow);
+	}
+}
+
+UE_NET_TEST_FIXTURE(FSharedConnectionFilterStatusCollectionTestFixture, ReplicationIsAllowedIfAnyChildConnectionAllows)
+{
+	FSharedConnectionFilterStatusCollection Collection;
+
+	constexpr uint32 ParentConnectionIds[] = {1, 4711, 99};
+	struct FChildFilterStatus
+	{
+		uint32 ChildConnectionId;
+		ENetFilterStatus FilterStatus;
+	}
+	ChildFilterStatuses[] =
+	{
+		{0, ENetFilterStatus::Allow},
+		{3, ENetFilterStatus::Disallow},
+		{2, ENetFilterStatus::Disallow},
+	};
+
+	for (const uint32 ParentConnectionId : ParentConnectionIds)
+	{
+		for (const FChildFilterStatus& ChildFilterStatus : ChildFilterStatuses)
+		{
+			Collection.SetFilterStatus(FConnectionHandle(ParentConnectionId, ChildFilterStatus.ChildConnectionId), ChildFilterStatus.FilterStatus);
+		}
+	}
+
+	// Now that all filter statuses have been set we can verify the result
+	for (const uint32 ParentConnectionId : ParentConnectionIds)
+	{
+		UE_NET_ASSERT_EQ(Collection.GetFilterStatus(ParentConnectionId), ENetFilterStatus::Allow);
+	}
+}
+
+UE_NET_TEST_FIXTURE(FSharedConnectionFilterStatusCollectionTestFixture, FilterStatusIsAdjustedWhenConnectionIsRemoved)
+{
+	FSharedConnectionFilterStatusCollection Collection;
+
+	constexpr uint32 ParentConnectionIds[] = {1, 4711, 99};
+	struct FChildFilterStatus
+	{
+		uint32 ChildConnectionId;
+		ENetFilterStatus FilterStatus;
+	}
+	ChildFilterStatuses[] =
+	{
+		{1, ENetFilterStatus::Allow},
+		{3, ENetFilterStatus::Disallow},
+		{2, ENetFilterStatus::Disallow},
+	};
+
+	// Set filter status
+	for (const uint32 ParentConnectionId : ParentConnectionIds)
+	{
+		for (const FChildFilterStatus& ChildFilterStatus : ChildFilterStatuses)
+		{
+			Collection.SetFilterStatus(FConnectionHandle(ParentConnectionId, ChildFilterStatus.ChildConnectionId), ChildFilterStatus.FilterStatus);
+		}
+	}
+
+	// Remove all connections
+	for (const uint32 ParentConnectionId : ParentConnectionIds)
+	{
+		for (const FChildFilterStatus& ChildFilterStatus : ChildFilterStatuses)
+		{
+			Collection.RemoveConnection(FConnectionHandle(ParentConnectionId, ChildFilterStatus.ChildConnectionId));
+		}
+	}
+
+	// Now that all filter operations have been performed we can verify the result
+	for (const uint32 ParentConnectionId : ParentConnectionIds)
+	{
+		UE_NET_ASSERT_EQ(Collection.GetFilterStatus(ParentConnectionId), ENetFilterStatus::Disallow);
+	}
+}
+
+UE_NET_TEST_FIXTURE(FSharedConnectionFilterStatusCollectionTestFixture, FilterStatusIsAdjustedWhenParentConnectionIsRemoved)
+{
+	FSharedConnectionFilterStatusCollection Collection;
+
+	constexpr uint32 ParentConnectionIds[] = {1, 4711, 99};
+	struct FChildFilterStatus
+	{
+		uint32 ChildConnectionId;
+		ENetFilterStatus FilterStatus;
+	}
+	ChildFilterStatuses[] =
+	{
+		{1, ENetFilterStatus::Allow},
+		{3, ENetFilterStatus::Disallow},
+		{2, ENetFilterStatus::Disallow},
+	};
+
+	// Set filter status
+	for (const uint32 ParentConnectionId : ParentConnectionIds)
+	{
+		for (const FChildFilterStatus& ChildFilterStatus : ChildFilterStatuses)
+		{
+			Collection.SetFilterStatus(FConnectionHandle(ParentConnectionId, ChildFilterStatus.ChildConnectionId), ChildFilterStatus.FilterStatus);
+		}
+	}
+
+	// Remove all parent connections
+	for (const uint32 ParentConnectionId : ParentConnectionIds)
+	{
+		Collection.RemoveConnection(FConnectionHandle(ParentConnectionId));
+	}
+
+	// Now that all filter operations have been performed we can verify the result
+	for (const uint32 ParentConnectionId : ParentConnectionIds)
+	{
+		UE_NET_ASSERT_EQ(Collection.GetFilterStatus(ParentConnectionId), ENetFilterStatus::Disallow);
+	}
+}
+
+}
diff --git a/Engine/Source/Runtime/Engine/Classes/Engine/ChildConnection.h b/Engine/Source/Runtime/Engine/Classes/Engine/ChildConnection.h
--- a/Engine/Source/Runtime/Engine/Classes/Engine/ChildConnection.h
+++ b/Engine/Source/Runtime/Engine/Classes/Engine/ChildConnection.h
@@ -82,4 +82,7 @@ public:
 	virtual void InitLocalConnection(UNetDriver* InDriver, class FSocket* InSocket, const FURL& InURL, EConnectionState InState, int32 InMaxPacket = 0, int32 InPacketOverhead = 0) override {}
 
 	virtual void InitChildConnection(UNetDriver* InDriver, UNetConnection* InParent);
+
+private:
+	void AssignConnectionHandle();
 };
diff --git a/Engine/Source/Runtime/Engine/Classes/Engine/NetConnection.h b/Engine/Source/Runtime/Engine/Classes/Engine/NetConnection.h
--- a/Engine/Source/Runtime/Engine/Classes/Engine/NetConnection.h
+++ b/Engine/Source/Runtime/Engine/Classes/Engine/NetConnection.h
@@ -29,6 +29,7 @@
 #include "Net/Common/Packets/PacketTraits.h"
 #include "Net/Core/Misc/ResizableCircularQueue.h"
 #include "Net/NetAnalyticsTypes.h"
+#include "Net/Core/Connection/ConnectionHandle.h"
 #include "Net/Core/Connection/NetCloseResult.h"
 #include "Net/TrafficControl.h"
 #include "Net/NetDormantHolder.h"
@@ -630,10 +631,27 @@ public:
 	int32			LogCallCount;
 	int32			LogSustainedCount;
 
-	uint32 GetConnectionId() const { return ConnectionId; }
-	void SetConnectionId(uint32 InConnectionId) { ConnectionId = InConnectionId; }
+	/** Returns a connection identifier which contains both parent connection and child connection information. */
+	UE::Net::FConnectionHandle GetConnectionHandle() const
+	{
+		return ConnectionHandle;
+	}
+
+	UE_DEPRECATED(5.6, "Use GetConnectionHandle() instead")
+	uint32 GetConnectionId() const
+	{ 
+		return ConnectionHandle.IsParentConnection() ? ConnectionHandle.GetParentConnectionId() : UE::Net::FConnectionHandle().GetParentConnectionId();
+	}
 
+	UE_DEPRECATED(5.6, "External code should not be setting connection IDs. For valid use cases SetConnectionHandle() should be used instead.")
+	void SetConnectionId(uint32 InConnectionId)
+	{
+		// Using this deprecated method will result in previous behavior which didn't support child connections in a meaningful way. We assume this is a parent connection.
+		ConnectionHandle = UE::Net::FConnectionHandle(InConnectionId);
+	}
+	
 	/** If this is a child connection it will return the topmost parent coonnection ID, otherwise it will return its own ID. */
+	UE_DEPRECATED(5.6, "Use GetConnectionHandle instead")
 	ENGINE_API uint32 GetParentConnectionId() const;
 
 	FNetTraceCollector* GetInTraceCollector() const;
@@ -1736,8 +1754,8 @@ private:
 	/** Whether or not PacketOrderCache is presently being flushed */
 	bool bFlushingPacketOrderCache;
 
-	/** Unique ID that can be used instead of passing around a pointer to the connection */
-	uint32 ConnectionId;
+	/** Per NetDriver unique ID that can be used instead of passing around a pointer to the connection. Be careful if caching as the handle can be re-purposed as connections are destroyed and created. */
+	UE::Net::FConnectionHandle ConnectionHandle;
 
 #if UE_NET_TRACE_ENABLED
 	FNetTraceCollector* InTraceCollector = nullptr;
@@ -1857,6 +1875,11 @@ protected:
 
 	void InitChannelData();
 
+	void SetConnectionHandle(UE::Net::FConnectionHandle Handle)
+	{
+		ConnectionHandle = Handle;
+	}
+
 public:
 	/**
 	 * Retrieve stored set of replicated sub-objects of the given actor at the time of the last dormancy flush
diff --git a/Engine/Source/Runtime/Engine/Classes/Engine/NetDriver.h b/Engine/Source/Runtime/Engine/Classes/Engine/NetDriver.h
--- a/Engine/Source/Runtime/Engine/Classes/Engine/NetDriver.h
+++ b/Engine/Source/Runtime/Engine/Classes/Engine/NetDriver.h
@@ -16,6 +16,7 @@
 #include "Misc/NetworkGuid.h"
 #include "UObject/CoreNet.h"
 #include "Channel.h"
+#include "Net/Core/Connection/ConnectionHandle.h"
 #include "Net/Core/Misc/DDoSDetection.h"
 #include "IPAddress.h"
 #include "Net/NetAnalyticsTypes.h"
@@ -1988,8 +1989,12 @@ public:
 	inline uint32 AllocateConnectionId() { return ConnectionIdHandler.Allocate(); }
 	inline void FreeConnectionId(uint32 Id) { return ConnectionIdHandler.Free(Id); };
 
-	/** Returns the NetConnection associated with the ConnectionId. Slow. */
-	ENGINE_API UNetConnection* GetConnectionById(uint32 ConnectionId) const;
+	/** Returns the NetConnection associated with the ParentConnectionId. Slow. */
+	UE_DEPRECATED(5.6, "Use GetConnectionByHandle()")
+	ENGINE_API UNetConnection* GetConnectionById(uint32 ParentConnectionId) const;
+
+	/** Returns the NetConnection associated with the ConnectionHandle. Slow as it iterates over all connections to find one with matching handle. */
+	ENGINE_API UNetConnection* GetConnectionByHandle(UE::Net::FConnectionHandle ConnectionHandle) const;
 
 	/** Returns identifier used for NetTrace */
 	inline uint32 GetNetTraceId() const { return NetTraceId; }
diff --git a/Engine/Source/Runtime/Engine/Private/ActorReplication.cpp b/Engine/Source/Runtime/Engine/Private/ActorReplication.cpp
--- a/Engine/Source/Runtime/Engine/Private/ActorReplication.cpp
+++ b/Engine/Source/Runtime/Engine/Private/ActorReplication.cpp
@@ -1108,7 +1108,7 @@ void AActor::UpdateOwningNetConnection()
 	uint32 NewOwningNetConnectionId = 0U;
 	if (const UNetConnection* NetConnection = GetNetConnection())
 	{
-		NewOwningNetConnectionId = NetConnection->GetParentConnectionId();
+		NewOwningNetConnectionId = NetConnection->GetConnectionHandle().GetParentConnectionId();
 	}
 
 	// If this actor isn't replicated there's no way for us to tell whether we need to update our children.
diff --git a/Engine/Source/Runtime/Engine/Private/Components/CharacterMovementComponent.cpp b/Engine/Source/Runtime/Engine/Private/Components/CharacterMovementComponent.cpp
--- a/Engine/Source/Runtime/Engine/Private/Components/CharacterMovementComponent.cpp
+++ b/Engine/Source/Runtime/Engine/Private/Components/CharacterMovementComponent.cpp
@@ -8882,7 +8882,7 @@ static UIrisObjectReferencePackageMap* GetIrisPackageMapToReadReferences(const U
 		if (UIrisObjectReferencePackageMap* ObjectReferencePackageMap = Bridge->GetObjectReferencePackageMap())
 		{
 			const UReplicationSystem* ReplicationSystem = UE::Net::FReplicationSystemUtil::GetReplicationSystem(NetConnection->GetDriver());
-			ObjectReferencePackageMap->InitForRead(&PackedBits.PackageMapExports, ReplicationSystem->GetNetTokenResolveContext(NetConnection->GetConnectionId()));
+			ObjectReferencePackageMap->InitForRead(&PackedBits.PackageMapExports, ReplicationSystem->GetNetTokenResolveContext(NetConnection->GetConnectionHandle().GetParentConnectionId()));
 
 			return ObjectReferencePackageMap;
 		}
diff --git a/Engine/Source/Runtime/Engine/Private/DataChannel.cpp b/Engine/Source/Runtime/Engine/Private/DataChannel.cpp
--- a/Engine/Source/Runtime/Engine/Private/DataChannel.cpp
+++ b/Engine/Source/Runtime/Engine/Private/DataChannel.cpp
@@ -1958,7 +1958,7 @@ void UControlChannel::ReceivedBunch( FInBunch& Bunch )
 #if UE_WITH_IRIS
 				if (UReplicationSystem* IrisRepSystem = Connection->Driver->GetReplicationSystem())
 				{
-					IrisRepSystem->ReportProtocolMismatch(NetRefHandleId, Connection->GetConnectionId());
+					IrisRepSystem->ReportProtocolMismatch(NetRefHandleId, Connection->GetConnectionHandle().GetParentConnectionId());
 				}
 #endif
 			}
@@ -1972,7 +1972,7 @@ void UControlChannel::ReceivedBunch( FInBunch& Bunch )
 #if UE_WITH_IRIS
 				if (UReplicationSystem* IrisRepSystem = Connection->Driver->GetReplicationSystem())
 				{
-					IrisRepSystem->ReportErrorWithNetRefHandle(NetRefHandleError, NetRefHandleId, Connection->GetConnectionId());
+					IrisRepSystem->ReportErrorWithNetRefHandle(NetRefHandleError, NetRefHandleId, Connection->GetConnectionHandle().GetParentConnectionId());
 				}
 #endif
 			}
diff --git a/Engine/Source/Runtime/Engine/Private/Net/Experimental/Iris/DataStreamChannel.cpp b/Engine/Source/Runtime/Engine/Private/Net/Experimental/Iris/DataStreamChannel.cpp
--- a/Engine/Source/Runtime/Engine/Private/Net/Experimental/Iris/DataStreamChannel.cpp
+++ b/Engine/Source/Runtime/Engine/Private/Net/Experimental/Iris/DataStreamChannel.cpp
@@ -74,7 +74,7 @@ void UDataStreamChannel::Init(UNetConnection* InConnection, int32 InChIndex, ECh
 	if (UReplicationSystem* ReplicationSystem = InConnection->Driver->GetReplicationSystem())
 	{
 		bIsReadyToHandshake = 1U;
-		ReplicationSystem->InitDataStreams(InConnection->GetConnectionId(), DataStreamManager);
+		ReplicationSystem->InitDataStreams(InConnection->GetConnectionHandle().GetParentConnectionId(), DataStreamManager);
 	}
 #endif // UE_WITH_IRIS
 }
@@ -87,7 +87,7 @@ void UDataStreamChannel::ReInit()
 	if (UReplicationSystem* ReplicationSystem = Connection->Driver->GetReplicationSystem())
 	{
 		bIsReadyToHandshake = 1U;
-		ReplicationSystem->InitDataStreams(Connection->GetConnectionId(), DataStreamManager);
+		ReplicationSystem->InitDataStreams(Connection->GetConnectionHandle().GetParentConnectionId(), DataStreamManager);
 	}
 #endif
 }
@@ -155,7 +155,7 @@ void UDataStreamChannel::ReceivedBunch(FInBunch& Bunch)
 			UEngineReplicationBridge* ActorBridge = FReplicationSystemUtil::GetActorReplicationBridge(Connection);
 			if (ActorBridge)
 			{
-				ActorBridge->SendErrorWithNetRefHandle(UE::Net::ENetRefHandleError::BitstreamCorrupted, SerializationContext.GetErrorHandleContext(), Connection->GetConnectionId());
+				ActorBridge->SendErrorWithNetRefHandle(UE::Net::ENetRefHandleError::BitstreamCorrupted, SerializationContext.GetErrorHandleContext(), Connection->GetConnectionHandle().GetParentConnectionId());
 			}
 		}
 		Bunch.SetError();
diff --git a/Engine/Source/Runtime/Engine/Private/Net/Iris/ReplicationSystem/EngineReplicationBridge.cpp b/Engine/Source/Runtime/Engine/Private/Net/Iris/ReplicationSystem/EngineReplicationBridge.cpp
--- a/Engine/Source/Runtime/Engine/Private/Net/Iris/ReplicationSystem/EngineReplicationBridge.cpp
+++ b/Engine/Source/Runtime/Engine/Private/Net/Iris/ReplicationSystem/EngineReplicationBridge.cpp
@@ -47,6 +47,7 @@
 
 #include "Net/DataBunch.h"
 #include "Net/DataChannel.h"
+#include "Net/Core/Connection/ConnectionHandle.h"
 #include "Net/Core/Connection/NetEnums.h"
 #include "Net/Core/Connection/NetCloseResult.h"
 #include "Net/Core/Misc/NetSubObjectRegistry.h"
@@ -793,7 +794,7 @@ void UEngineReplicationBridge::OnProtocolMismatchReported(FNetRefHandle RefHandl
 
 		if (IsClassCritical(ObjectClass))
 		{
-			if (UNetConnection* ClientConnection = NetDriver->GetConnectionById(ConnectionId))
+			if (UNetConnection* ClientConnection = NetDriver->GetConnectionByHandle(UE::Net::FConnectionHandle(ConnectionId)))
 			{
 				FString ErrorMsg = FString::Printf(TEXT("Protocol mismatch: %s:%s. Class: %s"), *RefHandle.ToString(), *GetNameSafe(ReplicatedObject), *GetNameSafe(ObjectClass));
 				UE_LOG(LogIrisBridge, Error, TEXT("%s: Closing connection due to: %s"), ToCStr(ClientConnection->Describe()), ToCStr(ErrorMsg));
@@ -818,7 +819,7 @@ void UEngineReplicationBridge::SendErrorWithNetRefHandle(UE::Net::ENetRefHandleE
 {
 	if (NetDriver)
 	{
-		if (UNetConnection* ClientConnection = NetDriver->GetConnectionById(ConnectionId))
+		if (UNetConnection* ClientConnection = NetDriver->GetConnectionByHandle(UE::Net::FConnectionHandle(ConnectionId)))
 		{
 			uint64 RawHandleId = RefHandle.GetId();
 			FNetControlMessage<NMT_IrisNetRefHandleError>::Send(ClientConnection, ErrorType, RawHandleId);
@@ -919,7 +920,7 @@ FString UEngineReplicationBridge::PrintConnectionInfo(uint32 ConnectionId) const
 {
 	if (NetDriver)
 	{
-		if (UNetConnection* ClientConnection = NetDriver->GetConnectionById(ConnectionId))
+		if (UNetConnection* ClientConnection = NetDriver->GetConnectionByHandle(UE::Net::FConnectionHandle(ConnectionId)))
 		{
 			return FString::Printf(TEXT("ConnectionId:%u ViewTarget: %s Named: %s"), ConnectionId, *GetNameSafe(ClientConnection->ViewTarget), *ClientConnection->Describe());
 		}
diff --git a/Engine/Source/Runtime/Engine/Private/Net/Iris/ReplicationSystem/NetActorFactory.cpp b/Engine/Source/Runtime/Engine/Private/Net/Iris/ReplicationSystem/NetActorFactory.cpp
--- a/Engine/Source/Runtime/Engine/Private/Net/Iris/ReplicationSystem/NetActorFactory.cpp
+++ b/Engine/Source/Runtime/Engine/Private/Net/Iris/ReplicationSystem/NetActorFactory.cpp
@@ -24,7 +24,7 @@
 #include "Iris/Serialization/VectorNetSerializers.h"
 
 #include "Net/DataBunch.h"
-
+#include "Net/Core/Connection/ConnectionHandle.h"
 #include "Net/Iris/ReplicationSystem/EngineReplicationBridge.h"
 
 #include "ProfilingDebugging/AssetMetadataTrace.h"
@@ -454,7 +454,7 @@ void UNetActorFactory::PostInstantiation(const FPostInstantiationContext& Contex
 
 	// OnActorChannelOpen
 	{
-		UNetConnection* Connection = NetDriver->GetConnectionById(Context.ConnectionId);
+		UNetConnection* Connection = NetDriver->GetConnectionByHandle(FConnectionHandle(Context.ConnectionId));
 		FInBunch Bunch(Connection, const_cast<uint8*>(BaseHeader->CustomCreationData.GetData()), BaseHeader->CustomCreationDataBitCount);
 		Actor->OnActorChannelOpen(Bunch, Connection);
 
diff --git a/Engine/Source/Runtime/Engine/Private/Net/Iris/ReplicationSystem/ReplicationSystemUtil.cpp b/Engine/Source/Runtime/Engine/Private/Net/Iris/ReplicationSystem/ReplicationSystemUtil.cpp
--- a/Engine/Source/Runtime/Engine/Private/Net/Iris/ReplicationSystem/ReplicationSystemUtil.cpp
+++ b/Engine/Source/Runtime/Engine/Private/Net/Iris/ReplicationSystem/ReplicationSystemUtil.cpp
@@ -18,6 +18,7 @@
 #include "HAL/IConsoleManager.h"
 #include "Templates/Casts.h"
 #include "Net/Subsystems/NetworkSubsystem.h"
+#include "Net/Core/Connection/ConnectionHandle.h"
 #include "Net/Core/Misc/NetConditionGroupManager.h"
 #include "Net/Core/NetHandle/NetHandleManager.h"
 #include "GameFramework/PlayerController.h"
@@ -560,7 +561,7 @@ void FReplicationSystemUtil::RemoveSubObjectGroupMembership(const APlayerControl
 	{
 		if (UReplicationSystem* ReplicationSystem = Conn->GetDriver() ? Conn->GetDriver()->GetReplicationSystem() : nullptr)
 		{
-			ReplicationSystem->SetSubObjectFilterStatus(NetGroup, Conn->GetConnectionId(), ENetFilterStatus::Disallow);
+			ReplicationSystem->SetSubObjectFilterStatus(NetGroup, Conn->GetConnectionHandle(), ENetFilterStatus::Disallow);
 		}
 	}
 }
@@ -572,13 +573,13 @@ void FReplicationSystemUtil::UpdateSubObjectGroupMemberships(const APlayerContro
 	{
 		if (UReplicationSystem* ReplicationSystem = Conn->GetDriver() ? Conn->GetDriver()->GetReplicationSystem() : nullptr)
 		{
-			const uint32 ConnId = Conn->GetParentConnectionId();
+			const FConnectionHandle ConnectionHandle = Conn->GetConnectionHandle();
 			for (const FName NetGroup : PC->GetNetConditionGroups())
 			{
 				if (!IsSpecialNetConditionGroup(NetGroup))
 				{
 					FNetObjectGroupHandle SubObjectGroupHandle = ReplicationSystem->GetOrCreateSubObjectFilter(NetGroup);
-					ReplicationSystem->SetSubObjectFilterStatus(NetGroup, ConnId, ENetFilterStatus::Allow);
+					ReplicationSystem->SetSubObjectFilterStatus(NetGroup, ConnectionHandle, ENetFilterStatus::Allow);
 				}
 			}
 		}
diff --git a/Engine/Source/Runtime/Engine/Private/NetConnection.cpp b/Engine/Source/Runtime/Engine/Private/NetConnection.cpp
--- a/Engine/Source/Runtime/Engine/Private/NetConnection.cpp
+++ b/Engine/Source/Runtime/Engine/Private/NetConnection.cpp
@@ -382,7 +382,6 @@ PRAGMA_ENABLE_DEPRECATION_WARNINGS
 ,	PacketOrderCacheStartIdx(0)
 ,	PacketOrderCacheCount(0)
 ,	bFlushingPacketOrderCache(false)
-,	ConnectionId(0)
 {
 	NetworkCustomVersions = FNetworkVersion::GetNetworkCustomVersions();
 }
@@ -463,7 +462,7 @@ void UNetConnection::InitBase(UNetDriver* InDriver,class FSocket* InSocket, cons
 	NetTraceId = Driver->GetNetTraceId();
 #endif
 
-	SetConnectionId(InDriver->AllocateConnectionId());
+	SetConnectionHandle(UE::Net::FConnectionHandle(InDriver->AllocateConnectionId()));
 
 	const double DriverElapsedTime = Driver->GetElapsedTime();
 
@@ -562,7 +561,7 @@ void UNetConnection::InitBase(UNetDriver* InDriver,class FSocket* InSocket, cons
 	// Init RemoteNetTokenState.
 	if (Driver->GetNetTokenStore())
 	{
-		Driver->GetNetTokenStore()->InitRemoteNetTokenStoreState(GetConnectionId());
+		Driver->GetNetTokenStore()->InitRemoteNetTokenStoreState(GetConnectionHandle().GetParentConnectionId());
 	}
 #endif
 
@@ -602,14 +601,14 @@ void UNetConnection::InitBase(UNetDriver* InDriver,class FSocket* InSocket, cons
 
 	NetPing = FNetPing::CreateNetPing(this);
 
-	UE_NET_TRACE_CONNECTION_CREATED(NetTraceId, GetConnectionId());
-	UE_NET_TRACE_CONNECTION_STATE_UPDATED(NetTraceId, GetConnectionId(), static_cast<uint8>(GetConnectionState()));
+	UE_NET_TRACE_CONNECTION_CREATED(NetTraceId, GetConnectionHandle().GetParentConnectionId());
+	UE_NET_TRACE_CONNECTION_STATE_UPDATED(NetTraceId, GetConnectionHandle().GetParentConnectionId(), static_cast<uint8>(GetConnectionState()));
 
 #if UE_WITH_IRIS
 	if (UReplicationSystem* ReplicationSystem = Driver->GetReplicationSystem())
 	{
-		ReplicationSystem->AddConnection(GetConnectionId());
-		ReplicationSystem->SetConnectionUserData(GetConnectionId(), this);
+		ReplicationSystem->AddConnection(GetConnectionHandle().GetParentConnectionId());
+		ReplicationSystem->SetConnectionUserData(GetConnectionHandle().GetParentConnectionId(), this);
 	}
 #endif // UE_WITH_IRIS
 }
@@ -635,7 +634,7 @@ void UNetConnection::InitConnection(UNetDriver* InDriver, EConnectionState InSta
 	NetTraceId = Driver->GetNetTraceId();
 #endif
 
-	SetConnectionId(InDriver->AllocateConnectionId());
+	SetConnectionHandle(UE::Net::FConnectionHandle(InDriver->AllocateConnectionId()));
 
 	// We won't be sending any packets, so use a default size
 	MaxPacket = (InMaxPacket == 0 || InMaxPacket > MAX_PACKET_SIZE) ? MAX_PACKET_SIZE : InMaxPacket;
@@ -674,8 +673,8 @@ void UNetConnection::InitConnection(UNetDriver* InDriver, EConnectionState InSta
 	PackageMapClient->Initialize(this, Driver->GuidCache);
 	PackageMap = PackageMapClient;
 
-	UE_NET_TRACE_CONNECTION_CREATED(NetTraceId, GetConnectionId());
-	UE_NET_TRACE_CONNECTION_STATE_UPDATED(NetTraceId, GetConnectionId(), static_cast<uint8>(GetConnectionState()));
+	UE_NET_TRACE_CONNECTION_CREATED(NetTraceId, GetConnectionHandle().GetParentConnectionId());
+	UE_NET_TRACE_CONNECTION_STATE_UPDATED(NetTraceId, GetConnectionHandle().GetParentConnectionId(), static_cast<uint8>(GetConnectionState()));
 }
 
 void UNetConnection::InitHandler()
@@ -801,7 +800,7 @@ void UNetConnection::NotifyAnalyticsProvider()
 
 void UNetConnection::NotifyConnectionUpdated()
 {
-	if (const uint32 MyConnectionId = GetConnectionId() && OwningActor && RemoteAddr)
+	if (const uint32 MyConnectionId = GetConnectionHandle().GetParentConnectionId() && OwningActor && RemoteAddr)
 	{
 		UE_NET_TRACE_CONNECTION_UPDATED(NetTraceId, MyConnectionId, *(RemoteAddr->ToString(true)), *(OwningActor->GetName()));
 	}
@@ -1022,7 +1021,7 @@ EConnectionState const UNetConnection::GetConnectionState() const
 void UNetConnection::SetConnectionState(EConnectionState ConnectionState)
 {
 	State = ConnectionState;
-	UE_NET_TRACE_CONNECTION_STATE_UPDATED(NetTraceId, GetConnectionId(), static_cast<uint8>(State));
+	UE_NET_TRACE_CONNECTION_STATE_UPDATED(NetTraceId, GetConnectionHandle().GetParentConnectionId(), static_cast<uint8>(State));
 }
 PRAGMA_ENABLE_DEPRECATION_WARNINGS
 
@@ -1137,7 +1136,7 @@ void UNetConnection::GracefulClose(FNetResult&& CloseReason)
 #if UE_WITH_IRIS
 			if (UReplicationSystem* ReplicationSystem = Driver->GetReplicationSystem())
 			{
-				ReplicationSystem->SetConnectionGracefullyClosing(ConnectionId);
+				ReplicationSystem->SetConnectionGracefullyClosing(ConnectionHandle.GetParentConnectionId());
 			}
 #endif
 
@@ -1400,25 +1399,24 @@ void UNetConnection::CleanUp()
 		}
 	}
 
-	const uint32 MyConnectionId = GetConnectionId();
-	if (MyConnectionId)
+	if (ConnectionHandle.IsParentConnection())
 	{
-		UE_NET_TRACE_CONNECTION_CLOSED(NetTraceId, MyConnectionId);
+		UE_NET_TRACE_CONNECTION_CLOSED(NetTraceId, ConnectionHandle.GetParentConnectionId());
 	}
 
 	if (Driver != nullptr)
 	{
 		// It would be nicer to have the Driver handle this internally, but unfortunately the ServerConnection member is public.
 		// Otherwise we'd be able to do the appropriate logic in Add/Remove Client/ServerConnection
-		if (MyConnectionId)
+		if (ConnectionHandle.IsParentConnection())
 		{
 #if UE_WITH_IRIS
 			if (UReplicationSystem* ReplicationSystem = Driver->GetReplicationSystem())
 			{
-				ReplicationSystem->RemoveConnection(MyConnectionId);
+				ReplicationSystem->RemoveConnection(ConnectionHandle.GetParentConnectionId());
 			}
 #endif // UE_WITH_IRIS
-			Driver->FreeConnectionId(MyConnectionId);
+			Driver->FreeConnectionId(ConnectionHandle.GetParentConnectionId());
 		}
 	}
 
@@ -1507,6 +1505,40 @@ void UChildConnection::InitChildConnection(UNetDriver* InDriver, UNetConnection*
 	CurrentNetSpeed = InParent->CurrentNetSpeed;
 
 	InitChannelData();
+
+	AssignConnectionHandle();
+}
+
+void UChildConnection::AssignConnectionHandle()
+{
+	if (const UNetConnection* ParentConn = Parent.Get())
+	{
+		// Don't assume the children remain on a particular index or that the array isn't ever shrinking. This code will assign as low an ID as possible.
+		const uint32 MaxNeededChildId = ParentConn->Children.Num() + 1;
+		TBitArray<> UsedChildIds;
+		UsedChildIds.SetNum(static_cast<int32>(MaxNeededChildId) + 1, false);
+		constexpr uint32 InvalidChildId = 0;
+		UsedChildIds[InvalidChildId] = true;
+		for (const UChildConnection* ChildConn : ParentConn->Children)
+		{
+			const uint32 ChildId = Cast<UNetConnection>(ChildConn)->GetConnectionHandle().GetChildConnectionId();
+			// If the ChildId is outside our small range we don't need to know it's used and the bitarray isn't large enough to store the information.
+			if (ChildId <= MaxNeededChildId)
+			{
+				UsedChildIds[ChildId] = true;
+			}
+		}
+
+		// There will be at least one unused child ID.
+		int32 ChildId = UsedChildIds.Find(false);
+		if (ChildId == INDEX_NONE)
+		{
+			ensureMsgf(ChildId != INDEX_NONE, TEXT("Logic failed assigning child connection ID."));
+			ChildId = InvalidChildId;
+		}
+		
+		SetConnectionHandle(UE::Net::FConnectionHandle(ParentConn->GetConnectionHandle().GetParentConnectionId(), static_cast<uint32>(ChildId)));
+	}
 }
 
 void UChildConnection::GracefulClose(FNetCloseResult&& CloseReason)
@@ -1608,7 +1640,7 @@ FNetLevelVisibilityTransactionId UNetConnection::UpdateLevelStreamStatusChangedT
 				const UReplicationBridge* ReplicationBridge = ReplicationSystem->GetReplicationBridge();
 				if (UE::Net::FNetObjectGroupHandle GroupHandle = ReplicationBridge->GetLevelGroup(Level); GroupHandle.IsValid())
 				{
-					ReplicationSystem->SetGroupFilterStatus(GroupHandle, GetConnectionId(), UE::Net::ENetFilterStatus::Disallow);
+					ReplicationSystem->SetGroupFilterStatus(GroupHandle, GetConnectionHandle().GetParentConnectionId(), UE::Net::ENetFilterStatus::Disallow);
 				}
 			}
 		}
@@ -1764,7 +1796,7 @@ void UNetConnection::UpdateLevelVisibilityInternal(const FUpdateLevelVisibilityL
 					const UReplicationBridge* ReplicationBridge = ReplicationSystem->GetReplicationBridge();
 					if (UE::Net::FNetObjectGroupHandle GroupHandle = ReplicationBridge->GetLevelGroup(Level); GroupHandle.IsValid())
 					{
-						ReplicationSystem->SetGroupFilterStatus(GroupHandle, GetConnectionId(), UE::Net::ENetFilterStatus::Allow);
+						ReplicationSystem->SetGroupFilterStatus(GroupHandle, GetConnectionHandle().GetParentConnectionId(), UE::Net::ENetFilterStatus::Allow);
 					}
 				}
 			}
@@ -1862,7 +1894,7 @@ void UNetConnection::UpdateLevelVisibilityInternal(const FUpdateLevelVisibilityL
 					const UReplicationBridge* ReplicationBridge = ReplicationSystem->GetReplicationBridge();
 					if (UE::Net::FNetObjectGroupHandle GroupHandle = ReplicationBridge->GetLevelGroup(Level); GroupHandle.IsValid())
 					{
-						ReplicationSystem->SetGroupFilterStatus(GroupHandle, GetConnectionId(), UE::Net::ENetFilterStatus::Disallow);
+						ReplicationSystem->SetGroupFilterStatus(GroupHandle, GetConnectionHandle().GetParentConnectionId(), UE::Net::ENetFilterStatus::Disallow);
 					}
 				}
 			}
@@ -2303,12 +2335,12 @@ void UNetConnection::FlushNet(bool bIgnoreSimulation)
 #if UE_NET_TRACE_ENABLED	
 		if (OutTraceCollector)
 		{
-			UE_NET_TRACE_FLUSH_COLLECTOR(OutTraceCollector, NetTraceId, GetConnectionId(), ENetTracePacketType::Outgoing);
+			UE_NET_TRACE_FLUSH_COLLECTOR(OutTraceCollector, NetTraceId, GetConnectionHandle().GetParentConnectionId(), ENetTracePacketType::Outgoing);
 			UE_NET_TRACE_DESTROY_COLLECTOR(OutTraceCollector);
 			OutTraceCollector = nullptr;
 		}
 		// Report end of packet
-		UE_NET_TRACE_PACKET_SEND(NetTraceId, GetConnectionId(), OutPacketId, SendBuffer.GetNumBits());
+		UE_NET_TRACE_PACKET_SEND(NetTraceId, GetConnectionHandle().GetParentConnectionId(), OutPacketId, SendBuffer.GetNumBits());
 #endif
 
 		// Send now.
@@ -2606,7 +2638,7 @@ void UNetConnection::ReceivedNak( int32 NakPacketId )
 {
 	UE_LOG(LogNetTraffic, Verbose, TEXT("   Received nak %i"), NakPacketId);
 
-	UE_NET_TRACE_PACKET_DROPPED(NetTraceId, GetConnectionId(), NakPacketId, ENetTracePacketType::Outgoing);
+	UE_NET_TRACE_PACKET_DROPPED(NetTraceId, GetConnectionHandle().GetParentConnectionId(), NakPacketId, ENetTracePacketType::Outgoing);
 
 	SCOPE_CYCLE_COUNTER(Stat_NetConnectionReceivedNak);
 
@@ -3275,19 +3307,19 @@ void UNetConnection::ReceivedPacket( FBitReader& Reader, bool bIsReinjectedPacke
 #if UE_NET_TRACE_ENABLED
 	if (InTraceCollector)
 	{
-		UE_NET_TRACE_FLUSH_COLLECTOR(InTraceCollector, NetTraceId, GetConnectionId(), ENetTracePacketType::Incoming);
+		UE_NET_TRACE_FLUSH_COLLECTOR(InTraceCollector, NetTraceId, GetConnectionHandle().GetParentConnectionId(), ENetTracePacketType::Incoming);
 		UE_NET_TRACE_DESTROY_COLLECTOR(InTraceCollector);
 		InTraceCollector = nullptr;
 	}
 #endif
 
 	// Trace end marker of this incoming data packet.
-	UE_NET_TRACE_PACKET_RECV(NetTraceId, GetConnectionId(), InPacketId, Reader.GetNumBits());
+	UE_NET_TRACE_PACKET_RECV(NetTraceId, GetConnectionHandle().GetParentConnectionId(), InPacketId, Reader.GetNumBits());
 
 	if ((bHasBunchErrors || bSkipAck) && !IsInternalAck())
 	{
 		// Trace packet as lost on receiving end to indicate bSkipAck or errors in traced data
-		UE_NET_TRACE_PACKET_DROPPED(NetTraceId, GetConnectionId(), InPacketId, ENetTracePacketType::Incoming);
+		UE_NET_TRACE_PACKET_DROPPED(NetTraceId, GetConnectionHandle().GetParentConnectionId(), InPacketId, ENetTracePacketType::Incoming);
 	}
 }
 
@@ -4923,7 +4955,7 @@ void UNetConnection::HandleClientPlayer( APlayerController *PC, UNetConnection*
 		// Enable replication
 		if (UReplicationSystem* ReplicationSystem = Driver->GetReplicationSystem())
 		{
-			ReplicationSystem->SetReplicationEnabledForConnection(GetConnectionId(), true);
+			ReplicationSystem->SetReplicationEnabledForConnection(GetConnectionHandle().GetParentConnectionId(), true);
 		}
 	}
 #endif // UE_WITH_IRIS
@@ -5342,7 +5374,7 @@ void UNetConnection::SetClientLoginState( const EClientLoginState::Type NewState
 		if (NewState == EClientLoginState::ReceivedJoin)
 		{
 			// Enable replication
-			ReplicationSystem->SetReplicationEnabledForConnection(GetConnectionId(), true);
+			ReplicationSystem->SetReplicationEnabledForConnection(GetConnectionHandle().GetParentConnectionId(), true);
 		}
 	}
 #endif // UE_WITH_IRIS
@@ -5758,12 +5790,7 @@ void UNetConnection::SendChallengeControlMessage(const FEncryptionKeyResponse& R
 
 uint32 UNetConnection::GetParentConnectionId() const
 {
-	if (const UChildConnection* ChildConnection = const_cast<UNetConnection*>(this)->GetUChildConnection())
-	{
-		return ChildConnection->Parent->GetParentConnectionId();
-	}
-
-	return GetConnectionId();
+	return GetConnectionHandle().GetParentConnectionId();
 }
 
 void UNetConnection::NotifyActorDestroyed(AActor* Actor, bool IsSeamlessTravel /* = false */)
diff --git a/Engine/Source/Runtime/Engine/Private/NetDriver.cpp b/Engine/Source/Runtime/Engine/Private/NetDriver.cpp
--- a/Engine/Source/Runtime/Engine/Private/NetDriver.cpp
+++ b/Engine/Source/Runtime/Engine/Private/NetDriver.cpp
@@ -2288,8 +2288,8 @@ void UNetDriver::TryUpgradeNetworkFeatures(EEngineNetworkRuntimeFeatures RemoteF
 			checkf(ServerConnection, TEXT("NMT_Upgrade control message received with no server connection on %s"), *GetDescription());
 			
 			// Tell Iris and the Datastream about the current netconnection
-			ReplicationSystem->AddConnection(ServerConnection->GetConnectionId());
-			ReplicationSystem->SetConnectionUserData(ServerConnection->GetConnectionId(), ServerConnection);
+			ReplicationSystem->AddConnection(ServerConnection->GetConnectionHandle().GetParentConnectionId());
+			ReplicationSystem->SetConnectionUserData(ServerConnection->GetConnectionHandle().GetParentConnectionId(), ServerConnection);
 			for (UChannel* Channel : ServerConnection->Channels)
 			{
 				if (Channel)
@@ -6993,16 +6993,34 @@ void UNetDriver::SetReplicationDriver(UReplicationDriver* NewReplicationDriver)
 	NotifyGameInstanceUpdated();
 }
 
-UNetConnection* UNetDriver::GetConnectionById(uint32 ConnectionId) const
+UNetConnection* UNetDriver::GetConnectionById(uint32 ParentConnectionId) const
 {
-	if (ServerConnection != nullptr && ServerConnection->GetConnectionId() == ConnectionId)
+	if (ServerConnection != nullptr && ServerConnection->GetConnectionHandle().GetParentConnectionId() == ParentConnectionId)
 	{
 		return ServerConnection;
 	}
 
 	for (UNetConnection* Connection : ClientConnections)
 	{
-		if (Connection && Connection->GetConnectionId() == ConnectionId)
+		if (Connection && Connection->GetConnectionHandle().GetParentConnectionId() == ParentConnectionId)
+		{
+			return Connection;
+		}
+	}
+
+	return nullptr;
+}
+
+UNetConnection* UNetDriver::GetConnectionByHandle(UE::Net::FConnectionHandle ConnectionHandle) const
+{
+	if (ServerConnection != nullptr && ServerConnection->GetConnectionHandle() == ConnectionHandle)
+	{
+		return ServerConnection;
+	}
+
+	for (UNetConnection* Connection : ClientConnections)
+	{
+		if (Connection && Connection->GetConnectionHandle() == ConnectionHandle)
 		{
 			return Connection;
 		}
@@ -7056,7 +7074,7 @@ void UNetDriver::UpdateGroupFilterStatusForLevel(const ULevel* Level, UE::Net::F
 	for (UNetConnection* Connection : ClientConnections)
 	{
 		const bool bIsVisible = (bIsPersistentLevel && WorldPackageName == Connection->GetClientWorldPackageName()) || Connection->ClientVisibleLevelNames.Contains(LevelPackageName);
-		ConnectionMask.SetBitValue(Connection->GetConnectionId(), bIsVisible);
+		ConnectionMask.SetBitValue(Connection->GetConnectionHandle().GetParentConnectionId(), bIsVisible);
 	}
 
 	// Set group filter status
@@ -7240,7 +7258,7 @@ void UNetDriver::UpdateIrisReplicationViews() const
 
 			FillIrisReplicationViews(AllConnections, ReplicationView);
 
-			ReplicationSystem->SetReplicationView(ClientConnection->GetConnectionId(), ReplicationView);
+			ReplicationSystem->SetReplicationView(ClientConnection->GetConnectionHandle().GetParentConnectionId(), ReplicationView);
 
 			ReplicationView.Views.Reset();
 			AllConnections.Reset();
@@ -7259,7 +7277,7 @@ void UNetDriver::UpdateIrisReplicationViews() const
 
 		FillIrisReplicationViews(AllConnections, ReplicationView);
 
-		ReplicationSystem->SetReplicationView(ServerConnection->GetConnectionId(), ReplicationView);
+		ReplicationSystem->SetReplicationView(ServerConnection->GetConnectionHandle().GetParentConnectionId(), ReplicationView);
 	}
 }
 
@@ -7311,7 +7329,7 @@ void UNetDriver::PostDispatchSendUpdate()
 
 			for (uint32 ConnId : ConnectionsToSend)
 			{
-				UNetConnection* NetConnection = GetConnectionById(ConnId);
+				UNetConnection* NetConnection = GetConnectionByHandle(UE::Net::FConnectionHandle(ConnId));
 				if (NetConnection && NetConnection->Channels.IsValidIndex(DataStreamChannelIndex))
 				{
 					if (UDataStreamChannel* DataStreamChannel = Cast<UDataStreamChannel>(NetConnection->Channels[DataStreamChannelIndex]))
@@ -7659,12 +7677,7 @@ void UNetDriver::ProcessRemoteFunction(
 			{
 				if (UNetConnection* Connection = Actor->GetNetConnection())
 				{
-					if (UChildConnection* ChildConnection = Connection->GetUChildConnection())
-					{
-						Connection = ChildConnection->Parent;
-					}
-
-					if (ReplicationSystem->SendRPC(Connection->GetConnectionId(), Actor, SubObject, Function, Parameters))
+					if (ReplicationSystem->SendRPC(Connection->GetConnectionHandle().GetParentConnectionId(), Actor, SubObject, Function, Parameters))
 					{
 						return;
 					}
@@ -8600,7 +8613,7 @@ FConsoleCommandWithArgsDelegate::CreateLambda([](const TArray< FString >& Args)
 			UE_LOG(LogNet, Display, TEXT("Printing Server Connection for: %s (Definition=%s)"), *NetDriver->NetDriverName.ToString(), *NetDriver->GetNetDriverDefinition().ToString());
 
 			UE_LOG(LogNet, Display, TEXT("\tServerConnection: ConnectionId=%u ViewTarget=%s FullDescription=%s"), 
-				NetDriver->ServerConnection->GetConnectionId(), 
+				NetDriver->ServerConnection->GetConnectionHandle().GetParentConnectionId(), 
 				*GetNameSafe(NetDriver->ServerConnection->ViewTarget),
 				*NetDriver->ServerConnection->Describe()
 			);
@@ -8612,7 +8625,7 @@ FConsoleCommandWithArgsDelegate::CreateLambda([](const TArray< FString >& Args)
 			for (UNetConnection* NetConnection : NetDriver->ClientConnections)
 			{
 				UE_LOG(LogNet, Display, TEXT("\tClientConnection: ConnectionId=%u ViewTarget=%s NetId=%s FullDescription=%s"), 
-					NetConnection->GetConnectionId(), 
+					NetConnection->GetConnectionHandle().GetParentConnectionId(), 
 					*GetNameSafe(NetConnection->ViewTarget),
 					*NetConnection->PlayerId.ToDebugString(), 
 					*NetConnection->Describe()
diff --git a/Engine/Source/Runtime/Engine/Private/PackageMapClient.cpp b/Engine/Source/Runtime/Engine/Private/PackageMapClient.cpp
--- a/Engine/Source/Runtime/Engine/Private/PackageMapClient.cpp
+++ b/Engine/Source/Runtime/Engine/Private/PackageMapClient.cpp
@@ -297,7 +297,7 @@ void UPackageMapClient::Initialize(UNetConnection * InConnection, TSharedPtr<FNe
 	NetTokenResolveContext.NetTokenStore = InConnection->GetDriver()->GetNetTokenStore();
 	if (NetTokenResolveContext.NetTokenStore)
 	{
-		NetTokenResolveContext.RemoteNetTokenStoreState = NetTokenResolveContext.NetTokenStore->GetRemoteNetTokenStoreState(InConnection->GetConnectionId());
+		NetTokenResolveContext.RemoteNetTokenStoreState = NetTokenResolveContext.NetTokenStore->GetRemoteNetTokenStoreState(InConnection->GetConnectionHandle().GetParentConnectionId());
 	}
 #endif
 }
@@ -2275,7 +2275,7 @@ void UPackageMapClient::ReceiveCustomExportsBunch(FInBunch &InBunch)
 	InBunch.SerializeBits(&bHasExportsToRead, 1);
 
 	UE::Net::FNetTokenStore* NetTokenStore = NetTokenResolveContext.NetTokenStore;
-	UE::Net::FNetTokenStoreState* RemoteNetTokenStoreState = NetTokenStore->GetRemoteNetTokenStoreState(Connection->GetConnectionId());
+	UE::Net::FNetTokenStoreState* RemoteNetTokenStoreState = NetTokenStore->GetRemoteNetTokenStoreState(Connection->GetConnectionHandle().GetParentConnectionId());
 
 	while (bHasExportsToRead && !InBunch.IsError())
 	{
diff --git a/Engine/Source/Runtime/Experimental/Iris/Core/Private/Iris/Core/IrisLog.cpp b/Engine/Source/Runtime/Experimental/Iris/Core/Private/Iris/Core/IrisLog.cpp
--- a/Engine/Source/Runtime/Experimental/Iris/Core/Private/Iris/Core/IrisLog.cpp
+++ b/Engine/Source/Runtime/Experimental/Iris/Core/Private/Iris/Core/IrisLog.cpp
@@ -3,3 +3,4 @@
 #include "Iris/Core/IrisLog.h"
 
 DEFINE_LOG_CATEGORY(LogIris);
+DEFINE_LOG_CATEGORY(LogIrisFiltering);
diff --git a/Engine/Source/Runtime/Experimental/Iris/Core/Private/Iris/ReplicationSystem/Filtering/ReplicationFiltering.cpp b/Engine/Source/Runtime/Experimental/Iris/Core/Private/Iris/ReplicationSystem/Filtering/ReplicationFiltering.cpp
--- a/Engine/Source/Runtime/Experimental/Iris/Core/Private/Iris/ReplicationSystem/Filtering/ReplicationFiltering.cpp
+++ b/Engine/Source/Runtime/Experimental/Iris/Core/Private/Iris/ReplicationSystem/Filtering/ReplicationFiltering.cpp
@@ -23,9 +23,6 @@
 namespace UE::Net::Private
 {
 
-DEFINE_LOG_CATEGORY_STATIC(LogIrisFiltering, Log, All);
-
-
 bool bCVarRepFilterCullNonRelevant = true;
 static FAutoConsoleVariableRef CVarRepFilterCullNonRelevant(TEXT("Net.Iris.CullNonRelevant"), bCVarRepFilterCullNonRelevant, TEXT("When enabled will cull replicated actors that are not relevant to any client."), ECVF_Default);
 
@@ -695,12 +692,18 @@ void FReplicationFiltering::RemoveConnection(uint32 ConnectionId)
 	}
 
 	// Reset SubObject filter for removed connection
-	SubObjectFilterGroups.ForAllSetBits([this, ConnectionId](uint32 GroupIndex)
+	FConnectionHandle ConnectionHandle(ConnectionId);
+	SubObjectFilterGroups.ForAllSetBits([this, ConnectionId, ConnectionHandle](uint32 GroupIndex)
 	{
 		if (!FNetObjectGroupHandle::IsReservedNetObjectGroupIndex(static_cast<FNetObjectGroupHandle::FGroupIndexType>(GroupIndex)))
 		{
-			SetConnectionFilterStatus(*GetPerObjectInfo(GroupInfos[GroupIndex].ConnectionStateIndex), ConnectionId, ENetFilterStatus::Disallow);
-			DirtySubObjectFilterGroups.SetBit(GroupIndex);
+			// Clear the connection specific info in the group filter info and disallow replication.
+			if (FPerSubObjectFilterGroupInfo* GroupInfo = GetPerSubObjectFilterGroupInfo(static_cast<FNetObjectGroupHandle::FGroupIndexType>(GroupIndex)))
+			{
+				GroupInfo->ConnectionFilterStatus.RemoveConnection(ConnectionHandle);
+				SetConnectionFilterStatus(*GetPerObjectInfo(GroupInfo->ConnectionStateIndex), ConnectionId, ENetFilterStatus::Disallow);
+				DirtySubObjectFilterGroups.SetBit(GroupIndex);
+			}
 		}
 	});
 }
@@ -1876,10 +1879,8 @@ void FReplicationFiltering::AddSubObjectFilter(FNetObjectGroupHandle GroupHandle
 	}
 
 	SubObjectFilterGroups.SetBit(GroupIndex);
-	GroupInfos[GroupIndex].ConnectionStateIndex = AllocPerObjectInfo();
 
-	// By default we filter out all connections
-	SetPerObjectInfoFilterStatus(*GetPerObjectInfo(GroupInfos[GroupIndex].ConnectionStateIndex), ENetFilterStatus::Disallow);
+	CreatePerSubObjectGroupFilterInfo(GroupIndex);
 
 	UE_LOG(LogIrisFiltering, Verbose, TEXT("ReplicationFiltering::AddSubObjectFilter Group: %s FilterStatus: DisallowReplication"), *Groups->GetGroupNameString(GroupHandle), GroupIndex);
 }
@@ -1891,9 +1892,8 @@ void FReplicationFiltering::RemoveSubObjectFilter(FNetObjectGroupHandle GroupHan
 	{
 		// Mark group as no longer a SubObjectFilter group
 		SubObjectFilterGroups.ClearBit(GroupIndex);
-		const PerObjectInfoIndexType ConnectionStateIndex = GroupInfos[GroupIndex].ConnectionStateIndex;
-		GroupInfos[GroupIndex].ConnectionStateIndex = 0U;
-		FreePerObjectInfo(ConnectionStateIndex);
+
+		DestroyPerSubObjectGroupFilterInfo(GroupIndex);
 
 		UE_LOG(LogIrisFiltering, Verbose, TEXT("ReplicationFiltering::RemoveSubObjectFilter Group: %s"), *Groups->GetGroupNameString(GroupHandle));
 	}
@@ -1904,13 +1904,17 @@ void FReplicationFiltering::UpdateSubObjectFilters()
 	IRIS_PROFILER_SCOPE(FReplicationFiltering_UpdateSubObjectFilters);
 
 	// We want to remove all groups that have no members and no enabled connections
-	auto UpdateSubObjectFilterGroup = [this](uint32 GroupIndex)
+	auto UpdateSubObjectFilterGroup = [this](uint32 BitIndex)
 	{
-		if (const FNetObjectGroup* Group = Groups->GetGroupFromIndex((FNetObjectGroupHandle::FGroupIndexType)GroupIndex))
+		FNetObjectGroupHandle::FGroupIndexType GroupIndex = static_cast<FNetObjectGroupHandle::FGroupIndexType>(BitIndex);
+		if (const FNetObjectGroup* Group = Groups->GetGroupFromIndex(GroupIndex);
+			Group && Group->Members.IsEmpty())
 		{
-			if (Group->Members.IsEmpty() && !IsAnyConnectionFilterStatusAllowed(*GetPerObjectInfo(GroupInfos[GroupIndex].ConnectionStateIndex)))
+			const FPerSubObjectFilterGroupInfo* GroupInfo = GetPerSubObjectFilterGroupInfo(GroupIndex);
+			if (!IsAnyConnectionFilterStatusAllowed(*GetPerObjectInfo(GroupInfo->ConnectionStateIndex)))
 			{
 				UE_LOG(LogIrisFiltering, Verbose, TEXT("UpdateSubObjectFilters is destroying group %s since its empty"), *Group->GroupName.ToString());
+				// Note that the below call will cause functions to be called in this class, like RemoveSubObjectFilter.
 				ReplicationSystem->DestroyGroup(Groups->GetHandleFromGroup(Group));
 			}
 		}
@@ -1920,102 +1924,48 @@ void FReplicationFiltering::UpdateSubObjectFilters()
 	DirtySubObjectFilterGroups.ClearAllBits();
 }
 
-void FReplicationFiltering::SetSubObjectFilterStatus(FNetObjectGroupHandle GroupHandle, ENetFilterStatus ReplicationStatus)
-{
-	IRIS_PROFILER_SCOPE(SetSubObjectFilterStatus)
-
-	const FNetObjectGroupHandle::FGroupIndexType GroupIndex = GroupHandle.GetGroupIndex();
-	if (GroupHandle.IsReservedNetObjectGroup())
-	{
-		UE_LOG(LogIrisFiltering, Warning, TEXT("FReplicationFiltering::SetSubObjectFilterStatus - Trying to set filter for reserved Group: %s which is not allowed."), *Groups->GetGroupNameString(GroupHandle));
-		return;
-	}
-
-	if (!SubObjectFilterGroups.GetBit(GroupIndex))
-	{
-		UE_LOG(LogIrisFiltering, Warning, TEXT("FReplicationFiltering::SetSubObjectFilterStatus - Trying to ReplicationStatus for Group: %s that is not a SubObjectFilterGroup"), *Groups->GetGroupNameString(GroupHandle));
-		return;
-	}
-
-	FPerObjectInfo* FilterInfo = GetPerObjectInfo(GroupInfos[GroupHandle.GetGroupIndex()].ConnectionStateIndex);
-	SetPerObjectInfoFilterStatus(*FilterInfo, ReplicationStatus);
-	if (!IsAnyConnectionFilterStatusAllowed(*FilterInfo))
-	{
-		DirtySubObjectFilterGroups.SetBit(GroupIndex);
-	}
-}
-
-void FReplicationFiltering::SetSubObjectFilterStatus(FNetObjectGroupHandle GroupHandle, const FNetBitArrayView& ConnectionsBitArray, ENetFilterStatus ReplicationStatus)
+void FReplicationFiltering::SetSubObjectFilterStatus(FNetObjectGroupHandle GroupHandle, FConnectionHandle ConnectionHandle, ENetFilterStatus ReplicationStatus)
 {
 	const FNetObjectGroupHandle::FGroupIndexType GroupIndex = GroupHandle.GetGroupIndex();
-	if (ConnectionsBitArray.GetNumBits() > ValidConnections.GetNumBits())
-	{
-		UE_LOG(LogIrisFiltering, Warning, TEXT("FReplicationFiltering::SetSubObjectFilterStatus - Trying to set filter for %s, with invalid Connections parameters."), *Groups->GetGroupNameString(GroupHandle));
-		return;
-	}
-
 	if (GroupHandle.IsReservedNetObjectGroup())
 	{
 		UE_LOG(LogIrisFiltering, Warning, TEXT("FReplicationFiltering::SetSubObjectFilterStatus - Trying to set filter for reserved Group: %s which is not allowed."), *Groups->GetGroupNameString(GroupHandle));
 		return;
 	}
 
-	if (!SubObjectFilterGroups.GetBit(GroupIndex))
-	{
-		UE_LOG(LogIrisFiltering, Warning, TEXT("FReplicationFiltering::SetSubObjectFilterStatus - Trying to ReplicationStatus for Group: %s that is not a SubObjectFilterGroup"), *Groups->GetGroupNameString(GroupHandle));
-		return;
-	}
-
-	const ENetFilterStatus InvertedReplicationStatus = ReplicationStatus == ENetFilterStatus::Allow ? ENetFilterStatus::Disallow : ENetFilterStatus::Allow;
-	const ENetFilterStatus DefaultReplicationStatus = ENetFilterStatus::Disallow;
-
-	// It is intentional that we iterate over all connections
-	for (uint32 ConnectionId = 1U; ConnectionId < ValidConnections.GetNumBits(); ++ConnectionId)
+	const uint32 ParentConnectionId = ConnectionHandle.GetParentConnectionId();
+	if (ensure(ValidConnections.GetBit(ParentConnectionId) && SubObjectFilterGroups.GetBit(GroupIndex)))
 	{
-		if (ConnectionId < ConnectionsBitArray.GetNumBits())
-		{			
-			SetSubObjectFilterStatus(GroupHandle, ConnectionId, ConnectionsBitArray.GetBit(ConnectionId) ? ReplicationStatus : InvertedReplicationStatus);
-		}
-		else
+		UE_LOG(LogIrisFiltering, Verbose, TEXT("ReplicationFiltering::SetSubObjectFilterStatus Group: %s, ConnectionHandle: %u:%u, FilterStatus: %u"), *Groups->GetGroupNameString(GroupHandle), ConnectionHandle.GetParentConnectionId(), ConnectionHandle.GetChildConnectionId(), ReplicationStatus == ENetFilterStatus::Allow ? 1U : 0U);
+		FPerSubObjectFilterGroupInfo* GroupInfo = GetPerSubObjectFilterGroupInfo(GroupIndex);
+		GroupInfo->ConnectionFilterStatus.SetFilterStatus(ConnectionHandle, ReplicationStatus);
+		ENetFilterStatus ConnectionReplicationStatus = GroupInfo->ConnectionFilterStatus.GetFilterStatus(ParentConnectionId);
+		FPerObjectInfo* ConnectionState = GetPerObjectInfo(GroupInfo->ConnectionStateIndex);
+		SetConnectionFilterStatus(*ConnectionState, ParentConnectionId, ConnectionReplicationStatus);
+		if (!IsAnyConnectionFilterStatusAllowed(*ConnectionState))
 		{
-			// To avoid issues with unspecified connections we set the status to the default
-			SetSubObjectFilterStatus(GroupHandle, ConnectionId, DefaultReplicationStatus);
+			DirtySubObjectFilterGroups.SetBit(GroupIndex);
 		}
 	}
 }
 
-void FReplicationFiltering::SetSubObjectFilterStatus(FNetObjectGroupHandle GroupHandle, uint32 ConnectionId, ENetFilterStatus ReplicationStatus)
+bool FReplicationFiltering::GetSubObjectFilterStatus(FNetObjectGroupHandle GroupHandle, uint32 ParentConnectionId, ENetFilterStatus& OutReplicationStatus) const
 {
 	const FNetObjectGroupHandle::FGroupIndexType GroupIndex = GroupHandle.GetGroupIndex();
-	if (GroupHandle.IsReservedNetObjectGroup())
-	{
-		UE_LOG(LogIrisFiltering, Warning, TEXT("FReplicationFiltering::SetSubObjectFilterStatus - Trying to set filter for reserved Group: %s which is not allowed."), *Groups->GetGroupNameString(GroupHandle));
-		return;
-	}
-
-	if (ensure(ValidConnections.GetBit(ConnectionId) && SubObjectFilterGroups.GetBit(GroupIndex)))
+	if (!(ValidConnections.GetBit(ParentConnectionId) && SubObjectFilterGroups.GetBit(GroupIndex)))
 	{
-		UE_LOG(LogIrisFiltering, Verbose, TEXT("ReplicationFiltering::SetSubObjectFilterStatus Group: %s, ConnectionId: %u, FilterStatus: %u"), *Groups->GetGroupNameString(GroupHandle), ConnectionId, ReplicationStatus == ENetFilterStatus::Allow ? 1U : 0U);
-		FPerObjectInfo* FilterInfo = GetPerObjectInfo(GroupInfos[GroupIndex].ConnectionStateIndex);
-		SetConnectionFilterStatus(*FilterInfo, ConnectionId, ReplicationStatus);
-		if (!IsAnyConnectionFilterStatusAllowed(*FilterInfo))
-		{
-			DirtySubObjectFilterGroups.SetBit(GroupIndex);
-		}
+		return false;
 	}
-}
 
-bool FReplicationFiltering::GetSubObjectFilterStatus(FNetObjectGroupHandle GroupHandle, uint32 ConnectionId, ENetFilterStatus& OutReplicationStatus) const
-{
-	if (!(ValidConnections.GetBit(ConnectionId) && SubObjectFilterGroups.GetBit(GroupHandle.GetGroupIndex())))
+	if (const FPerSubObjectFilterGroupInfo* GroupInfo = GetPerSubObjectFilterGroupInfo(GroupIndex);
+		ensure(GroupInfo))
 	{
-		return false;
+		const FPerObjectInfo* ConnectionState = GetPerObjectInfo(GroupInfo->ConnectionStateIndex);
+		OutReplicationStatus = GetConnectionFilterStatus(*ConnectionState, ParentConnectionId);
+		return true;
 	}
 
-	const FPerObjectInfo* ConnectionState = GetPerObjectInfo(GroupInfos[GroupHandle.GetGroupIndex()].ConnectionStateIndex);
-	OutReplicationStatus = GetConnectionFilterStatus(*ConnectionState, ConnectionId);
-
-	return true;
+	return false;
 }
 
 bool FReplicationFiltering::AddExclusionFilterGroup(FNetObjectGroupHandle GroupHandle)
@@ -2632,6 +2582,36 @@ void FReplicationFiltering::SetPerObjectInfoFilterStatus(FPerObjectInfo& ObjectI
 	}
 }
 
+FReplicationFiltering::FPerSubObjectFilterGroupInfo& FReplicationFiltering::CreatePerSubObjectGroupFilterInfo(FNetObjectGroupHandle::FGroupIndexType GroupIndex)
+{
+	FPerSubObjectFilterGroupInfo& GroupInfo = SubObjectFilterGroupInfos.Add(GroupIndex);
+	// Ensure group info didn't already exist
+	ensure(GroupInfo.ConnectionStateIndex == 0);
+	GroupInfo.ConnectionStateIndex = AllocPerObjectInfo();
+	SetPerObjectInfoFilterStatus(*GetPerObjectInfo(GroupInfo.ConnectionStateIndex), ENetFilterStatus::Disallow);
+
+	return GroupInfo;
+}
+
+void FReplicationFiltering::DestroyPerSubObjectGroupFilterInfo(FNetObjectGroupHandle::FGroupIndexType GroupIndex)
+{
+	if (FPerSubObjectFilterGroupInfo* GroupInfo = SubObjectFilterGroupInfos.Find(GroupIndex))
+	{
+		FreePerObjectInfo(GroupInfo->ConnectionStateIndex);
+		SubObjectFilterGroupInfos.Remove(GroupIndex);
+	}
+}
+
+FReplicationFiltering::FPerSubObjectFilterGroupInfo* FReplicationFiltering::GetPerSubObjectFilterGroupInfo(FNetObjectGroupHandle::FGroupIndexType GroupIndex)
+{
+	return SubObjectFilterGroupInfos.Find(GroupIndex);
+}
+
+const FReplicationFiltering::FPerSubObjectFilterGroupInfo* FReplicationFiltering::GetPerSubObjectFilterGroupInfo(FNetObjectGroupHandle::FGroupIndexType GroupIndex) const
+{
+	return SubObjectFilterGroupInfos.Find(GroupIndex);
+}
+
 void FReplicationFiltering::InitFilters()
 {
 	const UNetObjectFilterDefinitions* FilterDefinitions = GetDefault<UNetObjectFilterDefinitions>();
diff --git a/Engine/Source/Runtime/Experimental/Iris/Core/Private/Iris/ReplicationSystem/Filtering/ReplicationFiltering.h b/Engine/Source/Runtime/Experimental/Iris/Core/Private/Iris/ReplicationSystem/Filtering/ReplicationFiltering.h
--- a/Engine/Source/Runtime/Experimental/Iris/Core/Private/Iris/ReplicationSystem/Filtering/ReplicationFiltering.h
+++ b/Engine/Source/Runtime/Experimental/Iris/Core/Private/Iris/ReplicationSystem/Filtering/ReplicationFiltering.h
@@ -4,10 +4,12 @@
 
 #include "CoreTypes.h"
 #include "Net/Core/NetBitArray.h"
+#include "Iris/ReplicationSystem/Filtering/SharedConnectionFilterStatus.h"
 #include "Iris/ReplicationSystem/Filtering/NetObjectFilter.h"
 #include "Iris/ReplicationSystem/Filtering/ObjectScopeHysteresisUpdater.h"
 #include "Iris/ReplicationSystem/Filtering/ReplicationFilteringConfig.h"
 #include "Iris/ReplicationSystem/NetObjectGroupHandle.h"
+#include "Net/Core/Connection/ConnectionHandle.h"
 #include "Containers/Array.h"
 #include "UObject/StrongObjectPtr.h"
 
@@ -140,10 +142,8 @@ public:
 	void RemoveSubObjectFilter(FNetObjectGroupHandle GroupHandle);
 	bool IsSubObjectFilterGroup(FNetObjectGroupHandle GroupHandle) const { return GroupHandle.IsValid() && SubObjectFilterGroups.GetBit(GroupHandle.GetGroupIndex()); }
 
-	void SetSubObjectFilterStatus(FNetObjectGroupHandle GroupHandle, ENetFilterStatus ReplicationStatus);
-	void SetSubObjectFilterStatus(FNetObjectGroupHandle GroupHandle, const FNetBitArrayView& ConnectionsBitArray, ENetFilterStatus);
-	void SetSubObjectFilterStatus(FNetObjectGroupHandle GroupHandle, uint32 ConnectionId, ENetFilterStatus ReplicationStatus);
-	bool GetSubObjectFilterStatus(FNetObjectGroupHandle GroupHandle, uint32 ConnectionId, ENetFilterStatus& OutReplicationStatus) const;
+	void SetSubObjectFilterStatus(FNetObjectGroupHandle GroupHandle, FConnectionHandle ConnectionHandle, ENetFilterStatus ReplicationStatus);
+	bool GetSubObjectFilterStatus(FNetObjectGroupHandle GroupHandle, uint32 ParentConnectionId, ENetFilterStatus& OutReplicationStatus) const;
 
 	/** Print the filter information we have regarding the passed ObjectIndex and his relation to the passed Connection */
 	FString PrintFilterObjectInfo(FInternalNetRefIndex ObjectIndex, uint32 ConnectionId) const;
@@ -187,6 +187,15 @@ private:
 		PerObjectInfoIndexType ConnectionStateIndex;
 	};
 
+	class FPerSubObjectFilterGroupInfo
+	{
+	public:
+		// Tracks parent and child connection filter status. Necessary for splitscreen support.
+		FSharedConnectionFilterStatusCollection ConnectionFilterStatus;
+		// Index to storage for parent connection bit array. 
+		PerObjectInfoIndexType ConnectionStateIndex = 0;
+	};
+
 	struct FFilterInfo
 	{
 		TStrongObjectPtr<UNetObjectFilter> Filter;
@@ -197,6 +206,7 @@ private:
 private:
 	class FUpdateDirtyObjectsBatchHelper;
 	friend FNetObjectFilteringInfoAccessor;
+	friend FPerSubObjectFilterGroupInfo;
 	
 	static void StaticChecks();
 
@@ -246,6 +256,12 @@ private:
 
 	void SetPerObjectInfoFilterStatus(FPerObjectInfo& ObjectInfo, ENetFilterStatus ReplicationStatus);
 
+	// SubObjectGroup filtering support
+	FPerSubObjectFilterGroupInfo& CreatePerSubObjectGroupFilterInfo(FNetObjectGroupHandle::FGroupIndexType GroupIndex);
+	void DestroyPerSubObjectGroupFilterInfo(FNetObjectGroupHandle::FGroupIndexType GroupIndex);
+	FPerSubObjectFilterGroupInfo* GetPerSubObjectFilterGroupInfo(FNetObjectGroupHandle::FGroupIndexType GroupIndex);
+	const FPerSubObjectFilterGroupInfo* GetPerSubObjectFilterGroupInfo(FNetObjectGroupHandle::FGroupIndexType GroupIndex) const;
+
 	ENetFilterStatus GetConnectionFilterStatus(const FPerObjectInfo& ObjectInfo, uint32 ConnectionId) const;
 	bool IsAnyConnectionFilterStatusAllowed(const FPerObjectInfo& ObjectInfo) const;
 	bool IsAnyConnectionFilterStatusDisallowed(const FPerObjectInfo& ObjectInfo) const;
@@ -346,6 +362,9 @@ private:
 	TArray<FPerGroupInfo> GroupInfos;
 	uint32 MaxGroupCount = 0;
 
+	// SubObject filter groups
+	TMap<uint32, FPerSubObjectFilterGroupInfo> SubObjectFilterGroupInfos;
+
 	// Hysteresis frame counts for dynamically filtered objects
 	TArray<uint8> ObjectScopeHysteresisFrameCounts;
 
@@ -361,9 +380,11 @@ private:
 	/** Inclusion filtering groups with newly added members and that need to include objects for at least one connection. */
 	FNetBitArray DirtyInclusionFilterGroups;
 
-	//$IRIS TODO: These need better documentation
+	// Group indices which are subobject filter groups
 	FNetBitArray SubObjectFilterGroups;
+	// Group indices which are subobject filter groups and in need of updating
 	FNetBitArray DirtySubObjectFilterGroups;
+	// Object indices with a connection filter
 	FNetBitArray AllConnectionFilteredObjects;
 
 	TArray<PerObjectInfoIndexType> ObjectIndexToPerObjectInfoIndex;
diff --git a/Engine/Source/Runtime/Experimental/Iris/Core/Private/Iris/ReplicationSystem/Filtering/SharedConnectionFilterStatus.cpp b/Engine/Source/Runtime/Experimental/Iris/Core/Private/Iris/ReplicationSystem/Filtering/SharedConnectionFilterStatus.cpp
new file mode 100644
--- /dev/null
+++ b/Engine/Source/Runtime/Experimental/Iris/Core/Private/Iris/ReplicationSystem/Filtering/SharedConnectionFilterStatus.cpp
@@ -0,0 +1,126 @@
+// Copyright Epic Games, Inc. All Rights Reserved.
+
+#include "Iris/ReplicationSystem/Filtering/SharedConnectionFilterStatus.h"
+#include "Iris/Core/IrisLog.h"
+
+namespace UE::Net
+{
+
+// FSharedConnectionFilterStatus implementation
+bool FSharedConnectionFilterStatus::SetFilterStatus(FConnectionHandle ConnectionHandle, ENetFilterStatus FilterStatus)
+{
+	if (!ConnectionHandle.IsValid())
+	{
+		UE_LOG(LogIrisFiltering, Error, TEXT("%s"), TEXT("Trying to set filter status for an invalid FConnectionHandle on an FConnectionHandleFilterGroup"));
+		return false;
+	}
+
+	if (ParentConnectionId != InvalidConnectionId && ConnectionHandle.GetParentConnectionId() != ParentConnectionId)
+	{
+		UE_LOG(LogIrisFiltering, Error, TEXT("FConnectionHandleFilterGroup ignoring SetFilterStatus call due to connection ID mismatch. Expected %u got %u."), ParentConnectionId, ConnectionHandle.GetParentConnectionId());
+		return false;
+	}
+
+	ParentConnectionId = ConnectionHandle.GetParentConnectionId();
+	if (FilterStatus == ENetFilterStatus::Allow)
+	{
+		AllowConnections.Add(ConnectionHandle.GetChildConnectionId());
+	}
+	else
+	{
+		AllowConnections.Remove(ConnectionHandle.GetChildConnectionId());
+	}
+
+	return true;
+}
+
+void FSharedConnectionFilterStatus::RemoveConnection(FConnectionHandle ConnectionHandle)
+{
+	if (ConnectionHandle.GetParentConnectionId() == ParentConnectionId)
+	{
+		if (ConnectionHandle.IsParentConnection())
+		{
+			AllowConnections.Reset();
+			// Allow this instance to be repurposed for a different parent connection.
+			ParentConnectionId = InvalidConnectionId;
+		}
+		else
+		{
+			AllowConnections.Remove(ConnectionHandle.GetChildConnectionId());
+		}
+	}
+}
+
+// FSharedConnectionFilterStatusCollection implementation
+void FSharedConnectionFilterStatusCollection::SetFilterStatus(FConnectionHandle ConnectionHandle, ENetFilterStatus FilterStatus)
+{
+	if (!ConnectionHandle.IsValid())
+	{
+		UE_LOG(LogIrisFiltering, Error, TEXT("%s"), TEXT("Trying to set filter status for an invalid FConnectionHandle on an FConnectionHandleFilterGroupCollection"));
+		return;
+	}
+
+	const uint32 ParentConnId = ConnectionHandle.GetParentConnectionId();
+	if (FilterStatus == ENetFilterStatus::Allow)
+	{
+		FSharedConnectionFilterStatus& SharedFilterStatus = FindOrAddSharedConnectionFilterStatus(ParentConnId);
+		SharedFilterStatus.SetFilterStatus(ConnectionHandle, FilterStatus);
+	}
+	else
+	{
+		if (FSharedConnectionFilterStatus* SharedFilterStatus = FindSharedConnectionFilterStatus(ParentConnId))
+		{
+			SharedFilterStatus->SetFilterStatus(ConnectionHandle, FilterStatus);
+			// If after setting the filter status the group disallows replication then we can remove it
+			if (SharedFilterStatus->GetFilterStatus() == ENetFilterStatus::Disallow)
+			{
+				ParentToFilterStatus.Remove(ParentConnId);
+			}
+		}
+	}
+}
+
+void FSharedConnectionFilterStatusCollection::RemoveConnection(FConnectionHandle ConnectionHandle)
+{
+	if (ConnectionHandle.IsParentConnection())
+	{
+		ParentToFilterStatus.Remove(ConnectionHandle.GetParentConnectionId());
+	}
+	else if (ConnectionHandle.IsChildConnection())
+	{
+		if (FSharedConnectionFilterStatus* SharedFilterStatus = FindSharedConnectionFilterStatus(ConnectionHandle.GetParentConnectionId()))
+		{
+			SharedFilterStatus->RemoveConnection(ConnectionHandle);
+		}
+	}
+}
+
+ENetFilterStatus FSharedConnectionFilterStatusCollection::GetFilterStatus(uint32 ParentConnectionId) const
+{
+	if (const FSharedConnectionFilterStatus* SharedFilterStatus = FindSharedConnectionFilterStatus(ParentConnectionId))
+	{
+		return SharedFilterStatus->GetFilterStatus();
+	}
+
+	return ENetFilterStatus::Disallow;
+}
+
+FSharedConnectionFilterStatus* FSharedConnectionFilterStatusCollection::FindSharedConnectionFilterStatus(uint32 ParentConnectionId)
+{
+	return ParentToFilterStatus.Find(ParentConnectionId);
+}
+
+const FSharedConnectionFilterStatus* FSharedConnectionFilterStatusCollection::FindSharedConnectionFilterStatus(uint32 ParentConnectionId) const
+{
+	return ParentToFilterStatus.Find(ParentConnectionId);
+}
+
+FSharedConnectionFilterStatus& FSharedConnectionFilterStatusCollection::FindOrAddSharedConnectionFilterStatus(uint32 ParentConnectionId)
+{
+	FSharedConnectionFilterStatus& SharedFilterStatus = ParentToFilterStatus.FindOrAdd(ParentConnectionId);
+	// Set disallow status to establish the ParentConnectionId of the group. It's not strictly necessary but it could help detect issues with the collection implementation.
+	SharedFilterStatus.SetFilterStatus(FConnectionHandle(ParentConnectionId), ENetFilterStatus::Disallow);
+	return SharedFilterStatus;
+}
+
+}
diff --git a/Engine/Source/Runtime/Experimental/Iris/Core/Private/Iris/ReplicationSystem/ReplicationSystem.cpp b/Engine/Source/Runtime/Experimental/Iris/Core/Private/Iris/ReplicationSystem/ReplicationSystem.cpp
--- a/Engine/Source/Runtime/Experimental/Iris/Core/Private/Iris/ReplicationSystem/ReplicationSystem.cpp
+++ b/Engine/Source/Runtime/Experimental/Iris/Core/Private/Iris/ReplicationSystem/ReplicationSystem.cpp
@@ -1315,7 +1315,7 @@ UE::Net::FNetObjectGroupHandle UReplicationSystem::GetSubObjectFilterGroupHandle
 	return FNetObjectGroupHandle();
 }
 
-void UReplicationSystem::SetSubObjectFilterStatus(FName GroupName, uint32 ConnectionId, UE::Net::ENetFilterStatus ReplicationStatus)
+void UReplicationSystem::SetSubObjectFilterStatus(FName GroupName, UE::Net::FConnectionHandle ConnectionHandle, UE::Net::ENetFilterStatus ReplicationStatus)
 {
 	using namespace UE::Net;
 	using namespace UE::Net::Private;
@@ -1336,7 +1336,7 @@ void UReplicationSystem::SetSubObjectFilterStatus(FName GroupName, uint32 Connec
 	if (GroupHandle.IsValid())
 	{
 		FReplicationFiltering& Filtering = Impl->ReplicationSystemInternal.GetFiltering();
-		Filtering.SetSubObjectFilterStatus(GroupHandle, ConnectionId, ReplicationStatus);
+		Filtering.SetSubObjectFilterStatus(GroupHandle, ConnectionHandle, ReplicationStatus);
 
 		FReplicationConditionals& Conditionals = Impl->ReplicationSystemInternal.GetConditionals();
 		Conditionals.MarkLifeTimeConditionalsDirtyForObjectsInGroup(GroupHandle);
diff --git a/Engine/Source/Runtime/Experimental/Iris/Core/Public/Iris/Core/IrisLog.h b/Engine/Source/Runtime/Experimental/Iris/Core/Public/Iris/Core/IrisLog.h
--- a/Engine/Source/Runtime/Experimental/Iris/Core/Public/Iris/Core/IrisLog.h
+++ b/Engine/Source/Runtime/Experimental/Iris/Core/Public/Iris/Core/IrisLog.h
@@ -5,3 +5,4 @@
 #include "Logging/LogMacros.h"
 
 IRISCORE_API DECLARE_LOG_CATEGORY_EXTERN(LogIris, Log, All);
+IRISCORE_API DECLARE_LOG_CATEGORY_EXTERN(LogIrisFiltering, Log, All);
diff --git a/Engine/Source/Runtime/Experimental/Iris/Core/Public/Iris/ReplicationSystem/Filtering/SharedConnectionFilterStatus.h b/Engine/Source/Runtime/Experimental/Iris/Core/Public/Iris/ReplicationSystem/Filtering/SharedConnectionFilterStatus.h
new file mode 100644
--- /dev/null
+++ b/Engine/Source/Runtime/Experimental/Iris/Core/Public/Iris/ReplicationSystem/Filtering/SharedConnectionFilterStatus.h
@@ -0,0 +1,71 @@
+// Copyright Epic Games, Inc. All Rights Reserved.
+
+#pragma once
+
+#include "HAL/Platform.h"
+#include "Iris/IrisConstants.h"
+#include "Iris/ReplicationSystem/ReplicationView.h" // for UE_IRIS_INLINE_VIEWS_PER_CONNECTION
+#include "Iris/ReplicationSystem/Filtering/NetObjectFilter.h"
+#include "Net/Core/Connection/ConnectionHandle.h"
+#include "Containers/Map.h"
+#include "Containers/Set.h"
+
+namespace UE::Net
+{
+
+/**
+ * Keeps track of the filter status for a connection and its child connections. All FConnectionHandles must have the same ParentConnection. As soon as the filter status is set for a FConnectionHandle all subsequent operations must have the same FConnectionHandle.
+ * The default state is that the filter status is Disallowed. If there's at least one connection that has set the filter status to Allowed then will the filter status for the group be Allowed.
+ */
+class FSharedConnectionFilterStatus
+{
+public:
+	/** Returns true if the filter status could be successfully set for the ConnectionHandle. A return value of false indicates either the connection handle is invalid or its ParentConnectionId doesn't match prior calls to SetFilterStatus. */
+	IRISCORE_API bool SetFilterStatus(FConnectionHandle ConnectionHandle, ENetFilterStatus FilterStatus);
+	
+	/** Returns Disallowed if no connections have set the filter status to Allowed and returns Allowed otherwise. */
+	ENetFilterStatus GetFilterStatus() const;
+	
+	/** Removes the filter status for ConnectionHandle. If it's a parent connection it will act as if all child connections were removed too and allow the instance to be used with a different ParentConnectionId. */
+	IRISCORE_API void RemoveConnection(FConnectionHandle ConnectionHandle);
+
+	/** Returns the parent connection ID the group operates on or InvalidConnectionId if no valid connection handles ever set filter status. */
+	uint32 GetParentConnectionId() const;
+
+private:
+	// Only keep track of connections which allow replication as the replication status is Disallow by default.
+	TSet<uint32, DefaultKeyFuncs<uint32>, TInlineSetAllocator<UE_IRIS_INLINE_VIEWS_PER_CONNECTION>> AllowConnections;
+	uint32 ParentConnectionId = InvalidConnectionId;
+};
+
+/** Keeps track of the filter status for multiple connections and their child connections. A FConnectionHandleFilterGroup is stored per unique ParentConnectionId calling SetFilterStatus. */
+class FSharedConnectionFilterStatusCollection
+{
+public:
+	IRISCORE_API void SetFilterStatus(FConnectionHandle ConnectionHandle, ENetFilterStatus FilterStatus);
+
+	/** Returns Disallowed if no connections with the supplied ParentConnectionId have set the filter status to Allowed. If at least one connection with the ParentConnectionId allows replication then Allowed is returned. */
+	IRISCORE_API ENetFilterStatus GetFilterStatus(uint32 ParentConnectionId) const;
+
+	/** Remove a connection from filter status records. If it's a parent connection the corresponding FConnectionHandleFilterGroup will be removed altogether. */
+	IRISCORE_API void RemoveConnection(FConnectionHandle ConnectionHandle);
+
+private:
+	FSharedConnectionFilterStatus* FindSharedConnectionFilterStatus(uint32 ParentConnectionId);
+	const FSharedConnectionFilterStatus* FindSharedConnectionFilterStatus(uint32 ParentConnectionId) const;
+	FSharedConnectionFilterStatus& FindOrAddSharedConnectionFilterStatus(uint32 ParentConnectionId);
+
+	TMap<uint32, FSharedConnectionFilterStatus> ParentToFilterStatus;
+};
+
+inline ENetFilterStatus FSharedConnectionFilterStatus::GetFilterStatus() const
+{
+	return AllowConnections.IsEmpty() ? ENetFilterStatus::Disallow : ENetFilterStatus::Allow;
+}
+
+inline uint32 FSharedConnectionFilterStatus::GetParentConnectionId() const
+{
+	return ParentConnectionId;
+}
+
+}
diff --git a/Engine/Source/Runtime/Experimental/Iris/Core/Public/Iris/ReplicationSystem/ReplicationSystem.h b/Engine/Source/Runtime/Experimental/Iris/Core/Public/Iris/ReplicationSystem/ReplicationSystem.h
--- a/Engine/Source/Runtime/Experimental/Iris/Core/Public/Iris/ReplicationSystem/ReplicationSystem.h
+++ b/Engine/Source/Runtime/Experimental/Iris/Core/Public/Iris/ReplicationSystem/ReplicationSystem.h
@@ -10,6 +10,7 @@
 #include "Iris/ReplicationSystem/NetRefHandle.h"
 #include "Iris/ReplicationSystem/NetObjectGroupHandle.h"
 #include "Iris/ReplicationSystem/ReplicationSystemTypes.h"
+#include "Net/Core/Connection/ConnectionHandle.h"
 #include "Net/Core/NetHandle/NetHandle.h"
 
 #include "ReplicationSystem.generated.h"
@@ -555,7 +556,7 @@ public:
 	IRISCORE_API FNetObjectGroupHandle GetSubObjectFilterGroupHandle(FName GroupName) const;
 
 	/** Set status of GroupFilter for specific connection. */
-	IRISCORE_API void SetSubObjectFilterStatus(FName GroupName, uint32 ConnectionId, UE::Net::ENetFilterStatus ReplicationStatus);
+	IRISCORE_API void SetSubObjectFilterStatus(FName GroupName, UE::Net::FConnectionHandle ConnectionHandle, UE::Net::ENetFilterStatus ReplicationStatus);
 
 	/** Remove group from filtering system, will cancel effects of the group. */
 	IRISCORE_API void RemoveSubObjectFilter(FName GroupName);
diff --git a/Engine/Source/Runtime/Net/Core/Public/Net/Core/Connection/ConnectionHandle.h b/Engine/Source/Runtime/Net/Core/Public/Net/Core/Connection/ConnectionHandle.h
new file mode 100644
--- /dev/null
+++ b/Engine/Source/Runtime/Net/Core/Public/Net/Core/Connection/ConnectionHandle.h
@@ -0,0 +1,90 @@
+// Copyright Epic Games, Inc. All Rights Reserved.
+
+#pragma once
+
+#include "HAL/Platform.h"
+#include "Templates/TypeHash.h"
+
+namespace UE::Net
+{
+
+/**
+ * The handle is unique per netdriver. The child connection ID is not unique with respect to other handles with a different parent connection ID.
+ * Child connection are not assigned proper connection IDs as we will never replicate to them. Only parent connections will be replicated to.
+ * When there is need to deal with child connections in replication scenarios the ConnectionHandle can aid with that.
+ */
+class FConnectionHandle
+{
+public:
+	/** Initializes an invalid handle. */
+	inline FConnectionHandle() = default;
+	explicit FConnectionHandle(uint32 InParentConnectionId);
+	FConnectionHandle(uint32 InParentConnectionId, uint32 InChildConnectionId);
+
+	bool operator==(const FConnectionHandle& Other) const;
+
+	/** Returns true if this is a valid handle. */
+	bool IsValid() const;
+	/** Returns true if the handle is valid and represents a parent connection. */
+	bool IsParentConnection() const;
+	/** Returns true if the handle is valid and represents a child connection. */
+	bool IsChildConnection() const;
+
+	/** Returns the parent connection id for all types of valid handles. */
+	uint32 GetParentConnectionId() const;
+	/** Returns a non-zero ID for a valid child connection handle, zero for parent connections and invalid handles. */
+	uint32 GetChildConnectionId() const;
+
+private:
+	friend uint32 GetTypeHash(const FConnectionHandle& Handle);
+
+	uint32 ParentConnectionId = 0;
+	// For a valid handle a ChildConnectionId of zero indicates it's the parent connection itself.
+	uint32 ChildConnectionId = 0;
+};
+
+inline FConnectionHandle::FConnectionHandle(uint32 InParentConnectionId)
+: ParentConnectionId(InParentConnectionId)
+, ChildConnectionId(0)
+{
+}
+
+inline FConnectionHandle::FConnectionHandle(uint32 InParentConnectionId, uint32 InChildConnectionId)
+: ParentConnectionId(InParentConnectionId)
+, ChildConnectionId(InChildConnectionId)
+{
+}
+
+inline bool FConnectionHandle::operator==(const FConnectionHandle& Other) const = default;
+
+inline bool FConnectionHandle::IsValid() const
+{
+	return ParentConnectionId > 0U;
+}
+
+inline bool FConnectionHandle::IsParentConnection() const
+{
+	return IsValid() && ChildConnectionId == 0U;
+}
+
+inline bool FConnectionHandle::IsChildConnection() const
+{
+	return IsValid() && ChildConnectionId != 0U;
+}
+
+inline uint32 FConnectionHandle::GetParentConnectionId() const
+{
+	return ParentConnectionId;
+}
+
+inline uint32 FConnectionHandle::GetChildConnectionId() const
+{
+	return IsValid() ? ChildConnectionId : 0U;
+}
+
+inline uint32 GetTypeHash(const FConnectionHandle& Handle)
+{
+	return ::GetTypeHash((uint64(Handle.ParentConnectionId) << 32U) | uint64(Handle.ChildConnectionId));
+}
+
+}
