Constraints: ensure attachment dependency is enforced

#rnx

diff --git a/Engine/Source/Runtime/Experimental/Animation/Constraints/Private/TransformConstraint.cpp b/Engine/Source/Runtime/Experimental/Animation/Constraints/Private/TransformConstraint.cpp
--- a/Engine/Source/Runtime/Experimental/Animation/Constraints/Private/TransformConstraint.cpp
+++ b/Engine/Source/Runtime/Experimental/Animation/Constraints/Private/TransformConstraint.cpp
@@ -280,6 +280,25 @@ void LogDependency( const FString& InDescription, const UTransformableHandle* In
 		*GetHandleLabel(InParentHandle), *GetHandleLabel(InChildHandle),
 		*GetConstraintLabel(InParentConstraint), *GetConstraintLabel(InChildConstraint));
 }
+
+void LogAttachmentDependency( const FString& InDescription, const UTransformableHandle* InAttachHandle, const UTransformableHandle* InChildHandle,
+								const UTickableConstraint* InAttachConstraint, const UTickableConstraint* InChildConstraint)
+{
+	if (!bDebugDependencies)
+	{
+		return;
+	}
+	
+	if (!InAttachHandle || !InChildHandle || !InAttachConstraint || !InChildConstraint)
+	{
+		return;
+	}
+
+	UE_LOG(LogTemp, Warning, TEXT("%s: '%s' is an attach parent of '%s' so '%s' must tick before '%s'"),
+		*InDescription,
+		*GetHandleLabel(InAttachHandle), *GetHandleLabel(InChildHandle),
+		*GetConstraintLabel(InAttachConstraint), *GetConstraintLabel(InChildConstraint));
+}
 	
 }
 
@@ -2037,6 +2056,9 @@ bool FTransformConstraintUtils::BuildDependencies(UWorld* InWorld, UTickableTran
 		}
 	}
 
+	// build dependencies regarding attachments
+	BuildAttachmentsDependencies(InWorld, InConstraint);
+
 	if (bIncludeTarget && !bSelf)
 	{
 		// get all constraints acting on the same target 
@@ -2142,6 +2164,63 @@ bool FTransformConstraintUtils::BuildDependencies(UWorld* InWorld, UTickableTran
 	return true;
 }
 
+void FTransformConstraintUtils::BuildAttachmentsDependencies(UWorld* InWorld, const UTickableTransformConstraint* InConstraint)
+{
+	using namespace ConstraintLocals;
+	using ConstraintWeakPtr = TWeakObjectPtr<UTickableConstraint>;
+	
+	if (!ensure(InWorld))
+	{
+		return;
+	}
+
+	if (!InConstraint || !InConstraint->IsValid())
+	{
+		return;
+	}
+
+	const FConstraintsManagerController& Controller = FConstraintsManagerController::Get(InWorld);
+	
+	static const TCHAR* AttachmentDependencyDesc = TEXT("Attachment Dependency");
+	
+	const UTransformableHandle* ChildHandle = InConstraint->ChildTRSHandle.Get();
+	if (const USceneComponent* ChildComponent = Cast<USceneComponent>(ChildHandle->GetTarget().Get()))
+	{
+		static constexpr bool bIncludeAllDescendants = true;
+		
+		TArray<USceneComponent*> ChildComponentChildren;
+		ChildComponent->GetChildrenComponents(bIncludeAllDescendants, ChildComponentChildren);
+		
+		for (USceneComponent* ChildChildComponent: ChildComponentChildren)
+		{
+			const uint32 ChildHash = GetConstrainableHash(ChildChildComponent);
+			if (ChildHash != 0)
+			{
+				auto IsHashChildOfConstraint = [ChildHash](const ConstraintWeakPtr& Constraint)
+				{
+					const UTickableTransformConstraint* TransformConstraint = Cast<UTickableTransformConstraint>(Constraint.Get());
+					const UTransformableHandle* ChildHandle = TransformConstraint ? TransformConstraint->ChildTRSHandle.Get() : nullptr;
+					if (!ChildHandle || !ChildHandle->IsValid())
+					{
+						return false;
+					}
+					return ChildHandle->GetHash() == ChildHash;
+				};
+				
+				const TArray< ConstraintWeakPtr > AttachChildConstraints = Controller.GetConstraintsByPredicate(IsHashChildOfConstraint);
+				for (const ConstraintWeakPtr& AttachChildConstraint: AttachChildConstraints)
+				{
+					Controller.SetConstraintsDependencies(InConstraint->ConstraintID, AttachChildConstraint->ConstraintID);
+					if (const UTickableTransformConstraint* AttachTransformConstraint = Cast<UTickableTransformConstraint>(AttachChildConstraint))
+					{
+						LogAttachmentDependency(AttachmentDependencyDesc, ChildHandle, AttachTransformConstraint->ChildTRSHandle,  InConstraint, AttachTransformConstraint);
+					}
+				}
+			}
+		}
+	}
+}
+
 void FTransformConstraintUtils::UpdateTransformBasedOnConstraint(FTransform& CurrentTransform, USceneComponent* SceneComponent)
 {
 	TArray< TWeakObjectPtr<UTickableConstraint> > Constraints;
diff --git a/Engine/Source/Runtime/Experimental/Animation/Constraints/Public/TransformConstraint.h b/Engine/Source/Runtime/Experimental/Animation/Constraints/Public/TransformConstraint.h
--- a/Engine/Source/Runtime/Experimental/Animation/Constraints/Public/TransformConstraint.h
+++ b/Engine/Source/Runtime/Experimental/Animation/Constraints/Public/TransformConstraint.h
@@ -440,6 +440,7 @@ public:
 
 struct FTransformConstraintUtils
 {
+public: 
 	/** Fills a sorted constraint array that InChild actor is the child of. */
 	static CONSTRAINTS_API void GetParentConstraints(
 		UWorld* World,
@@ -506,6 +507,10 @@ struct FTransformConstraintUtils
 
 	/** Ensure default dependencies between constraints. */
 	static CONSTRAINTS_API bool BuildDependencies(UWorld* InWorld, UTickableTransformConstraint* Constraint);
+	
+private:
+	/** Ensures that attachment dependencies are reflected at the constraints level. */
+	static void BuildAttachmentsDependencies(UWorld* InWorld, const UTickableTransformConstraint* Constraint);
 };
 
 /**
