Iris - Adjusted LocalIndex calculations and added ensure to track down crash.

#rnx

diff --git a/Engine/Source/Runtime/Experimental/Iris/Core/Private/Iris/ReplicationSystem/Filtering/ObjectScopeHysteresisUpdater.cpp b/Engine/Source/Runtime/Experimental/Iris/Core/Private/Iris/ReplicationSystem/Filtering/ObjectScopeHysteresisUpdater.cpp
--- a/Engine/Source/Runtime/Experimental/Iris/Core/Private/Iris/ReplicationSystem/Filtering/ObjectScopeHysteresisUpdater.cpp
+++ b/Engine/Source/Runtime/Experimental/Iris/Core/Private/Iris/ReplicationSystem/Filtering/ObjectScopeHysteresisUpdater.cpp
@@ -157,10 +157,10 @@ FObjectScopeHysteresisUpdater::FLocalIndex FObjectScopeHysteresisUpdater::GetOrC
 	uint32 LocalIndex = UsedLocalIndices.FindFirstZero();
 	if (LocalIndex == FNetBitArray::InvalidIndex)
 	{
+		LocalIndex = UsedLocalIndices.GetNumBits();
 		UsedLocalIndices.AddBits(LocalIndexGrowCount);
 		LocalIndexToNetRefIndex.AddZeroed(LocalIndexGrowCount);
 		FrameCounters.AddZeroed(LocalIndexGrowCount);
-		LocalIndex = UsedLocalIndices.FindFirstZero();
 	}
 
 	UsedLocalIndices.SetBit(LocalIndex);
@@ -169,6 +169,8 @@ FObjectScopeHysteresisUpdater::FLocalIndex FObjectScopeHysteresisUpdater::GetOrC
 	ObjectsToUpdate.SetBit(NetRefIndex);
 	MaxLocalIndex = FPlatformMath::Max(MaxLocalIndex, LocalIndex + 1U);
 
+	ensureMsgf(MaxLocalIndex <= UsedLocalIndices.GetNumBits(), TEXT("MaxLocalIndex has grown beyond supported size. MaxLocalIndex: %u LocalIndex: %u Array sizes: %u"), MaxLocalIndex, LocalIndex, UsedLocalIndices.GetNumBits());
+
 	return LocalIndex;
 }
 
