IAD - always update on-demand package store lookup tables when unmounting containers.

diff --git a/Engine/Source/Runtime/Experimental/IoStore/OnDemand/Private/OnDemandPackageStoreBackend.cpp b/Engine/Source/Runtime/Experimental/IoStore/OnDemand/Private/OnDemandPackageStoreBackend.cpp
--- a/Engine/Source/Runtime/Experimental/IoStore/OnDemand/Private/OnDemandPackageStoreBackend.cpp
+++ b/Engine/Source/Runtime/Experimental/IoStore/OnDemand/Private/OnDemandPackageStoreBackend.cpp
@@ -115,7 +115,7 @@ FIoStatus FOnDemandPackageStoreBackend::Unmount(const FString& ContainerName)
 	});
 	check(NumRemoved <= 1);
 
-	bNeedsUpdate = NumRemoved > 0;
+	bNeedsUpdate = true;
 	return NumRemoved == 0 ? EIoErrorCode::NotFound : EIoErrorCode::Ok;
 }
 
