Fix Collision Event returning twice the same component when proxy are swapped

diff --git a/Engine/Source/Runtime/Experimental/ChaosSolverEngine/Private/Chaos/ChaosGameplayEventDispatcher.cpp b/Engine/Source/Runtime/Experimental/ChaosSolverEngine/Private/Chaos/ChaosGameplayEventDispatcher.cpp
--- a/Engine/Source/Runtime/Experimental/ChaosSolverEngine/Private/Chaos/ChaosGameplayEventDispatcher.cpp
+++ b/Engine/Source/Runtime/Experimental/ChaosSolverEngine/Private/Chaos/ChaosGameplayEventDispatcher.cpp
@@ -399,7 +399,7 @@ void UChaosGameplayEventDispatcher::HandleCollisionEvents(const Chaos::FCollisio
 			{
 				const FChaosHandlerSet& HandlerSet = It.Value();
 
-				UPrimitiveComponent* const Comp0 = Cast<UPrimitiveComponent>(It.Key());
+				UPrimitiveComponent* Comp0 = Cast<UPrimitiveComponent>(It.Key());
 				const TArray<IPhysicsProxyBase*>* PhysicsProxyArray = Scene.GetOwnedPhysicsProxies(Comp0);
 
 				if (PhysicsProxyArray)
@@ -425,7 +425,13 @@ void UChaosGameplayEventDispatcher::HandleCollisionEvents(const Chaos::FCollisio
 
 								Chaos::FCollidingData const& CollisionDataItem = CollisionData[CollisionIdx];
 
-								IPhysicsProxyBase* const PhysicsProxy1 = CollisionDataItem.Proxy2 ? CollisionDataItem.Proxy2: PhysicsProxy0;
+								IPhysicsProxyBase* PhysicsProxy1 = CollisionDataItem.Proxy2 ? CollisionDataItem.Proxy2: PhysicsProxy0;
+
+								if (bSwapOrder)
+								{
+									PhysicsProxy1 = CollisionDataItem.Proxy1;
+									PhysicsProxy0 = CollisionDataItem.Proxy2;
+								}
 
 								// Are the proxies pending destruction? If they are no longer tracked by the PhysScene, the proxy is deleted or pending deletion.
 								if (Scene.GetOwningComponent<UPrimitiveComponent>(PhysicsProxy0) == nullptr || Scene.GetOwningComponent<UPrimitiveComponent>(PhysicsProxy1) == nullptr)
@@ -446,6 +452,10 @@ void UChaosGameplayEventDispatcher::HandleCollisionEvents(const Chaos::FCollisio
 									if (bNewEntry)
 									{
 										UPrimitiveComponent* const Comp1 = Scene.GetOwningComponent<UPrimitiveComponent>(PhysicsProxy1);
+										if (bSwapOrder)
+										{
+											Comp0 = Scene.GetOwningComponent<UPrimitiveComponent>(PhysicsProxy0);
+										}
 
 										// fill in legacy contact data
 										NotifyInfo.bCallEvent0 = true;
@@ -479,6 +489,10 @@ void UChaosGameplayEventDispatcher::HandleCollisionEvents(const Chaos::FCollisio
 									if (bNewEntry)
 									{
 										UPrimitiveComponent* const Comp1 = Scene.GetOwningComponent<UPrimitiveComponent>(PhysicsProxy1);
+										if (bSwapOrder)
+										{
+											Comp0 = Scene.GetOwningComponent<UPrimitiveComponent>(PhysicsProxy0);
+										}
 
 										// fill in Chaos contact data
 										ChaosNotifyInfo.CollisionInfo.Component = Comp0;
