Change Dataflow namespace to UE::Dataflow to be ok with coding standard

diff --git a/Engine/Source/Runtime/Experimental/Dataflow/Core/Private/Dataflow/DataflowNode.cpp b/Engine/Source/Runtime/Experimental/Dataflow/Core/Private/Dataflow/DataflowNode.cpp
--- a/Engine/Source/Runtime/Experimental/Dataflow/Core/Private/Dataflow/DataflowNode.cpp
+++ b/Engine/Source/Runtime/Experimental/Dataflow/Core/Private/Dataflow/DataflowNode.cpp
@@ -22,9 +22,9 @@ const FLinearColor FDataflowNode::DefaultNodeBodyTintColor = FLinearColor(0.f, 0
 
 const FName FDataflowAnyType::TypeName = TEXT("FDataflowAnyType");
 
-namespace Dataflow::Private
+namespace UE::Dataflow::Private
 {
-static uint32 GetArrayElementOffsetFromReference(const FArrayProperty* const ArrayProperty, const Dataflow::FConnectionReference& Reference)
+static uint32 GetArrayElementOffsetFromReference(const FArrayProperty* const ArrayProperty, const UE::Dataflow::FConnectionReference& Reference)
 {
 	check(ArrayProperty);
 	if (const void* const AddressAtIndex = ArrayProperty->GetValueAddressAtIndex_Direct(ArrayProperty->Inner, const_cast<void*>(Reference.ContainerReference), Reference.Index))
@@ -131,7 +131,7 @@ void FDataflowNode::AddInput(FDataflowInput* InPtr)
 {
 	if (InPtr)
 	{
-		for (const TPair<Dataflow::FConnectionKey, FDataflowInput*>& Elem : ExpandedInputs)
+		for (const TPair<UE::Dataflow::FConnectionKey, FDataflowInput*>& Elem : ExpandedInputs)
 		{
 			const FDataflowInput* const In = Elem.Value;
 			ensureMsgf(!In->GetName().IsEqual(InPtr->GetName()), TEXT("Add Input Failed: Existing Node input already defined with name (%s)"), *InPtr->GetName().ToString());
@@ -139,7 +139,7 @@ void FDataflowNode::AddInput(FDataflowInput* InPtr)
 
 		check(InPtr->GetOwningNode() == this);
 
-		const Dataflow::FConnectionKey Key(InPtr->GetOffset(), InPtr->GetContainerIndex(), InPtr->GetContainerElementOffset());
+		const UE::Dataflow::FConnectionKey Key(InPtr->GetOffset(), InPtr->GetContainerIndex(), InPtr->GetContainerElementOffset());
 		if (ensure(!ExpandedInputs.Contains(Key)))
 		{
 			ExpandedInputs.Add(Key, InPtr);
@@ -154,7 +154,7 @@ int32 FDataflowNode::GetNumInputs() const
 
 FDataflowInput* FDataflowNode::FindInput(FName InName)
 {
-	for (TPair<Dataflow::FConnectionKey, FDataflowInput*>& Elem : ExpandedInputs)
+	for (TPair<UE::Dataflow::FConnectionKey, FDataflowInput*>& Elem : ExpandedInputs)
 	{
 		FDataflowInput* const Con = Elem.Value;
 		if (Con->GetName().IsEqual(InName))
@@ -167,7 +167,7 @@ FDataflowInput* FDataflowNode::FindInput(FName InName)
 
 const FDataflowInput* FDataflowNode::FindInput(FName InName) const
 {
-	for (const TPair<Dataflow::FConnectionKey, FDataflowInput*>& Elem : ExpandedInputs)
+	for (const TPair<UE::Dataflow::FConnectionKey, FDataflowInput*>& Elem : ExpandedInputs)
 	{
 		const FDataflowInput* const Con = Elem.Value;
 		if (Con->GetName().IsEqual(InName))
@@ -178,7 +178,7 @@ const FDataflowInput* FDataflowNode::FindInput(FName InName) const
 	return nullptr;
 }
 
-const FDataflowInput* FDataflowNode::FindInput(const Dataflow::FConnectionKey& Key) const
+const FDataflowInput* FDataflowNode::FindInput(const UE::Dataflow::FConnectionKey& Key) const
 {
 	if (const FDataflowInput* const* Con = ExpandedInputs.Find(Key))
 	{
@@ -188,9 +188,9 @@ const FDataflowInput* FDataflowNode::FindInput(const Dataflow::FConnectionKey& K
 	return nullptr;
 }
 
-const FDataflowInput* FDataflowNode::FindInput(const Dataflow::FConnectionReference& Reference) const
+const FDataflowInput* FDataflowNode::FindInput(const UE::Dataflow::FConnectionReference& Reference) const
 {
-	const Dataflow::FConnectionKey Key = GetKeyFromReference(Reference);
+	const UE::Dataflow::FConnectionKey Key = GetKeyFromReference(Reference);
 	if (const FDataflowInput* const Con = FindInput(Key))
 	{
 		check(Con->RealAddress() == Reference.Reference);
@@ -199,7 +199,7 @@ const FDataflowInput* FDataflowNode::FindInput(const Dataflow::FConnectionRefere
 	if (Reference.ContainerReference == nullptr && !InputArrayProperties.IsEmpty())
 	{
 		// Search through all connections to see if Reference is the RealAddress of an array property.
-		for (const TPair<Dataflow::FConnectionKey, FDataflowInput*>& Elem : ExpandedInputs)
+		for (const TPair<UE::Dataflow::FConnectionKey, FDataflowInput*>& Elem : ExpandedInputs)
 		{
 			const FDataflowInput* const Con = Elem.Value;
 			if (Con->RealAddress() == Reference.Reference)
@@ -211,7 +211,7 @@ const FDataflowInput* FDataflowNode::FindInput(const Dataflow::FConnectionRefere
 	return nullptr;
 }
 
-FDataflowInput* FDataflowNode::FindInput(const Dataflow::FConnectionKey& Key)
+FDataflowInput* FDataflowNode::FindInput(const UE::Dataflow::FConnectionKey& Key)
 {
 	if (FDataflowInput* const* Con = ExpandedInputs.Find(Key))
 	{
@@ -221,9 +221,9 @@ FDataflowInput* FDataflowNode::FindInput(const Dataflow::FConnectionKey& Key)
 	return nullptr;
 }
 
-FDataflowInput* FDataflowNode::FindInput(const Dataflow::FConnectionReference& Reference)
+FDataflowInput* FDataflowNode::FindInput(const UE::Dataflow::FConnectionReference& Reference)
 {
-	const Dataflow::FConnectionKey Key = GetKeyFromReference(Reference);
+	const UE::Dataflow::FConnectionKey Key = GetKeyFromReference(Reference);
 	if (FDataflowInput* const Con = FindInput(Key))
 	{
 		check(Con->RealAddress() == Reference.Reference);
@@ -232,7 +232,7 @@ FDataflowInput* FDataflowNode::FindInput(const Dataflow::FConnectionReference& R
 	if (Reference.ContainerReference == nullptr && !InputArrayProperties.IsEmpty())
 	{
 		// Search through all connections to see if Reference is the RealAddress of an array property.
-		for (TPair<Dataflow::FConnectionKey, FDataflowInput*>& Elem : ExpandedInputs)
+		for (TPair<UE::Dataflow::FConnectionKey, FDataflowInput*>& Elem : ExpandedInputs)
 		{
 			FDataflowInput* const Con = Elem.Value;
 			if (Con->RealAddress() == Reference.Reference)
@@ -246,7 +246,7 @@ FDataflowInput* FDataflowNode::FindInput(const Dataflow::FConnectionReference& R
 
 const FDataflowInput* FDataflowNode::FindInput(const FGuid& InGuid) const
 {
-	for (const TPair<Dataflow::FConnectionKey, FDataflowInput*>& Elem : ExpandedInputs)
+	for (const TPair<UE::Dataflow::FConnectionKey, FDataflowInput*>& Elem : ExpandedInputs)
 	{
 		const FDataflowInput* const Con = Elem.Value;
 		if (Con->GetGuid() == InGuid)
@@ -266,7 +266,7 @@ TArray< FDataflowInput* > FDataflowNode::GetInputs() const
 
 void FDataflowNode::ClearInputs()
 {
-	for (TPair<Dataflow::FConnectionKey, FDataflowInput*>& Elem : ExpandedInputs)
+	for (TPair<UE::Dataflow::FConnectionKey, FDataflowInput*>& Elem : ExpandedInputs)
 	{
 		FDataflowInput* const Con = Elem.Value;
 		delete Con;
@@ -276,7 +276,7 @@ void FDataflowNode::ClearInputs()
 
 bool FDataflowNode::HasHideableInputs() const
 {
-	for (const TPair<Dataflow::FConnectionKey, FDataflowInput*>& Elem : ExpandedInputs)
+	for (const TPair<UE::Dataflow::FConnectionKey, FDataflowInput*>& Elem : ExpandedInputs)
 	{
 		const FDataflowInput* const Con = Elem.Value;
 		if (Con->GetCanHidePin())
@@ -289,7 +289,7 @@ bool FDataflowNode::HasHideableInputs() const
 
 bool FDataflowNode::HasHiddenInputs() const
 {
-	for (const TPair<Dataflow::FConnectionKey, FDataflowInput*>& Elem : ExpandedInputs)
+	for (const TPair<UE::Dataflow::FConnectionKey, FDataflowInput*>& Elem : ExpandedInputs)
 	{
 		const FDataflowInput* const Con = Elem.Value;
 		if (Con->GetPinIsHidden())
@@ -390,7 +390,7 @@ PRAGMA_ENABLE_DEPRECATION_WARNINGS
 	return nullptr;
 }
 
-const FDataflowOutput* FDataflowNode::FindOutput(const Dataflow::FConnectionKey& Key) const
+const FDataflowOutput* FDataflowNode::FindOutput(const UE::Dataflow::FConnectionKey& Key) const
 {
 PRAGMA_DISABLE_DEPRECATION_WARNINGS  // Until Outputs becomes private
 	if (const FDataflowOutput* const* Con = Outputs.Find(Key.Offset))
@@ -402,9 +402,9 @@ PRAGMA_ENABLE_DEPRECATION_WARNINGS
 	return nullptr;
 }
 
-const FDataflowOutput* FDataflowNode::FindOutput(const Dataflow::FConnectionReference& Reference) const
+const FDataflowOutput* FDataflowNode::FindOutput(const UE::Dataflow::FConnectionReference& Reference) const
 {
-	const Dataflow::FConnectionKey Key = GetKeyFromReference(Reference);
+	const UE::Dataflow::FConnectionKey Key = GetKeyFromReference(Reference);
 	if (const FDataflowOutput* const Con = FindOutput(Key))
 	{
 		check(Con->RealAddress() == Reference.Reference);
@@ -413,7 +413,7 @@ const FDataflowOutput* FDataflowNode::FindOutput(const Dataflow::FConnectionRefe
 	return nullptr;
 }
 
-FDataflowOutput* FDataflowNode::FindOutput(const Dataflow::FConnectionKey& Key)
+FDataflowOutput* FDataflowNode::FindOutput(const UE::Dataflow::FConnectionKey& Key)
 {
 PRAGMA_DISABLE_DEPRECATION_WARNINGS  // Until Outputs becomes private
 	if (FDataflowOutput* const* Con = Outputs.Find(Key.Offset))
@@ -425,9 +425,9 @@ PRAGMA_ENABLE_DEPRECATION_WARNINGS
 	return nullptr;
 }
 
-FDataflowOutput* FDataflowNode::FindOutput(const Dataflow::FConnectionReference& Reference)
+FDataflowOutput* FDataflowNode::FindOutput(const UE::Dataflow::FConnectionReference& Reference)
 {
-	const Dataflow::FConnectionKey Key = GetKeyFromReference(Reference);
+	const UE::Dataflow::FConnectionKey Key = GetKeyFromReference(Reference);
 	if (FDataflowOutput* const Con = FindOutput(Key))
 	{
 		check(Con->RealAddress() == Reference.Reference);
@@ -513,29 +513,29 @@ PRAGMA_ENABLE_DEPRECATION_WARNINGS
 	return false;
 }
 
-TArray<Dataflow::FPin> FDataflowNode::GetPins() const
+TArray<UE::Dataflow::FPin> FDataflowNode::GetPins() const
 {
-	TArray<Dataflow::FPin> RetVal;
-	for (const TPair<Dataflow::FConnectionKey, FDataflowInput*>& Elem : ExpandedInputs)
+	TArray<UE::Dataflow::FPin> RetVal;
+	for (const TPair<UE::Dataflow::FConnectionKey, FDataflowInput*>& Elem : ExpandedInputs)
 	{
 		const FDataflowInput* const Con = Elem.Value;
-		RetVal.Add({ Dataflow::FPin::EDirection::INPUT,Con->GetType(), Con->GetName(), Con->GetPinIsHidden()});
+		RetVal.Add({ UE::Dataflow::FPin::EDirection::INPUT,Con->GetType(), Con->GetName(), Con->GetPinIsHidden()});
 	}
 PRAGMA_DISABLE_DEPRECATION_WARNINGS  // Until Outputs becomes private
 	for (const TPair<int32, FDataflowOutput*>& Elem : Outputs)
 PRAGMA_ENABLE_DEPRECATION_WARNINGS
 	{
 		const FDataflowOutput* const Con = Elem.Value;
-		RetVal.Add({ Dataflow::FPin::EDirection::OUTPUT,Con->GetType(), Con->GetName(), Con->GetPinIsHidden() });
+		RetVal.Add({ UE::Dataflow::FPin::EDirection::OUTPUT,Con->GetType(), Con->GetName(), Con->GetPinIsHidden() });
 	}
 	return RetVal;
 }
 
-void FDataflowNode::UnregisterPinConnection(const Dataflow::FPin& Pin)
+void FDataflowNode::UnregisterPinConnection(const UE::Dataflow::FPin& Pin)
 {
-	if (Pin.Direction == Dataflow::FPin::EDirection::INPUT)
+	if (Pin.Direction == UE::Dataflow::FPin::EDirection::INPUT)
 	{
-		for (TMap<Dataflow::FConnectionKey, FDataflowInput*>::TIterator Iter = ExpandedInputs.CreateIterator(); Iter; ++Iter)
+		for (TMap<UE::Dataflow::FConnectionKey, FDataflowInput*>::TIterator Iter = ExpandedInputs.CreateIterator(); Iter; ++Iter)
 		{
 			FDataflowInput* Con = Iter.Value();
 			if (Con->GetName().IsEqual(Pin.Name) && Con->GetType().IsEqual(Pin.Type))
@@ -549,7 +549,7 @@ void FDataflowNode::UnregisterPinConnection(const Dataflow::FPin& Pin)
 			}
 		}
 	}
-	else if (Pin.Direction == Dataflow::FPin::EDirection::OUTPUT)
+	else if (Pin.Direction == UE::Dataflow::FPin::EDirection::OUTPUT)
 	{
 PRAGMA_DISABLE_DEPRECATION_WARNINGS  // Until Outputs becomes private
 		for (TMap<int32, FDataflowOutput*>::TIterator Iter = Outputs.CreateIterator(); Iter; ++Iter)
@@ -569,7 +569,7 @@ PRAGMA_ENABLE_DEPRECATION_WARNINGS
 	}
 }
 
-void FDataflowNode::Invalidate(const Dataflow::FTimestamp& InModifiedTimestamp)
+void FDataflowNode::Invalidate(const UE::Dataflow::FTimestamp& InModifiedTimestamp)
 {
 	if (bPauseInvalidations)
 	{
@@ -598,12 +598,12 @@ PRAGMA_ENABLE_DEPRECATION_WARNINGS
 
 const FProperty* FDataflowNode::FindProperty(const UStruct* Struct, const void* InProperty, const FName& PropertyName, TArray<const FProperty*>* OutPropertyChain) const
 {
-	return Dataflow::Private::FindProperty(Struct, this, InProperty, PropertyName, OutPropertyChain);
+	return UE::Dataflow::Private::FindProperty(Struct, this, InProperty, PropertyName, OutPropertyChain);
 }
 
 const FProperty& FDataflowNode::FindPropertyChecked(const UStruct* Struct, const void* InProperty, const FName& PropertyName, TArray<const FProperty*>* OutPropertyChain) const
 {
-	return Dataflow::Private::FindPropertyChecked(Struct, this, InProperty, PropertyName, OutPropertyChain);
+	return UE::Dataflow::Private::FindPropertyChecked(Struct, this, InProperty, PropertyName, OutPropertyChain);
 }
 
 const FProperty* FDataflowNode::FindProperty(const UStruct* Struct, const FName& PropertyFullName, TArray<const FProperty*>* OutPropertyChain) const
@@ -660,15 +660,15 @@ uint32 FDataflowNode::GetConnectionOffsetFromReference(const void* Reference) co
 	return (uint32)((size_t)Reference - (size_t)this);
 }
 
-Dataflow::FConnectionKey FDataflowNode::GetKeyFromReference(const Dataflow::FConnectionReference& Reference) const
+UE::Dataflow::FConnectionKey FDataflowNode::GetKeyFromReference(const UE::Dataflow::FConnectionReference& Reference) const
 {
-	Dataflow::FConnectionKey Key;
+	UE::Dataflow::FConnectionKey Key;
 	Key.Offset = Reference.ContainerReference ? GetConnectionOffsetFromReference(Reference.ContainerReference) : GetConnectionOffsetFromReference(Reference.Reference);
 	Key.ContainerIndex = Reference.Index;
 	Key.ContainerElementOffset = INDEX_NONE;
 	if (const FArrayProperty* const* ArrayProperty = InputArrayProperties.Find(Key.Offset))
 	{
-		Key.ContainerElementOffset = Dataflow::Private::GetArrayElementOffsetFromReference(*ArrayProperty, Reference);
+		Key.ContainerElementOffset = UE::Dataflow::Private::GetArrayElementOffsetFromReference(*ArrayProperty, Reference);
 	}
 	return Key;
 }
@@ -769,7 +769,7 @@ FText FDataflowNode::GetPropertyDisplayNameText(const TArray<const FProperty*>&
 #endif
 }
 
-void FDataflowNode::InitConnectionParametersFromPropertyReference(const FStructOnScope& StructOnScope, const void* PropertyRef, const FName& PropertyName, Dataflow::FConnectionParameters& OutParams)
+void FDataflowNode::InitConnectionParametersFromPropertyReference(const FStructOnScope& StructOnScope, const void* PropertyRef, const FName& PropertyName, UE::Dataflow::FConnectionParameters& OutParams)
 {
 	const UStruct* Struct = StructOnScope.GetStruct();
 	check(Struct);
@@ -786,11 +786,11 @@ void FDataflowNode::InitConnectionParametersFromPropertyReference(const FStructO
 	check(OutParams.Offset == GetPropertyOffset(PropertyChain));
 }
 
-FDataflowInput& FDataflowNode::RegisterInputConnectionInternal(const Dataflow::FConnectionReference& Reference, const FName& PropertyName)
+FDataflowInput& FDataflowNode::RegisterInputConnectionInternal(const UE::Dataflow::FConnectionReference& Reference, const FName& PropertyName)
 {
 	TUniquePtr<FStructOnScope> ScriptOnStruct = TUniquePtr<FStructOnScope>(NewStructOnScope());
 	check(ScriptOnStruct);
-	Dataflow::FInputParameters InputParams;
+	UE::Dataflow::FInputParameters InputParams;
 	InitConnectionParametersFromPropertyReference(*ScriptOnStruct, Reference.Reference, PropertyName, InputParams);
 	FDataflowInput* const Input = new FDataflowInput(InputParams);
 	check(Input->RealAddress() == Reference.Reference);
@@ -800,14 +800,14 @@ FDataflowInput& FDataflowNode::RegisterInputConnectionInternal(const Dataflow::F
 	return *Input;
 }
 
-FDataflowInput& FDataflowNode::RegisterInputArrayConnectionInternal(const Dataflow::FConnectionReference& Reference, const FName& ElementPropertyName,
+FDataflowInput& FDataflowNode::RegisterInputArrayConnectionInternal(const UE::Dataflow::FConnectionReference& Reference, const FName& ElementPropertyName,
 	const FName& ArrayPropertyName)
 {
 	TUniquePtr<FStructOnScope> ScriptOnStruct = TUniquePtr<FStructOnScope>(NewStructOnScope());
 	check(ScriptOnStruct);
 	const UStruct* Struct = ScriptOnStruct->GetStruct();
 	check(Struct);
-	Dataflow::FArrayInputParameters InputParams;
+	UE::Dataflow::FArrayInputParameters InputParams;
 	InputParams.Owner = this;
 
 	// Find the Array property.
@@ -835,7 +835,7 @@ FDataflowInput& FDataflowNode::RegisterInputArrayConnectionInternal(const Datafl
 	}
 	else if (const FStructProperty* const InnerStruct = CastField<FStructProperty>(InputParams.ArrayProperty->Inner))
 	{
-		InputParams.Property = &Dataflow::Private::FindPropertyChecked(InnerStruct->Struct, AddressAtIndex, Reference.Reference, ElementPropertyName, &PropertyChain);
+		InputParams.Property = &UE::Dataflow::Private::FindPropertyChecked(InnerStruct->Struct, AddressAtIndex, Reference.Reference, ElementPropertyName, &PropertyChain);
 		PropertyChain.Add(InnerStruct);
 	}
 
@@ -846,7 +846,7 @@ FDataflowInput& FDataflowNode::RegisterInputArrayConnectionInternal(const Datafl
 	const FString CPPType = InputParams.Property->GetCPPType(&ExtendedType);
 	InputParams.Type = FName(CPPType + ExtendedType);
 	InputParams.Name = GetPropertyFullName(PropertyChain, Reference.Index);
-	InputParams.InnerOffset = Dataflow::Private::GetArrayElementOffsetFromReference(InputParams.ArrayProperty, Reference);
+	InputParams.InnerOffset = UE::Dataflow::Private::GetArrayElementOffsetFromReference(InputParams.ArrayProperty, Reference);
 
 	InputArrayProperties.Emplace(InputParams.Offset, InputParams.ArrayProperty);
 
@@ -856,7 +856,7 @@ FDataflowInput& FDataflowNode::RegisterInputArrayConnectionInternal(const Datafl
 	return *Input;
 }
 
-void FDataflowNode::UnregisterInputConnection(const Dataflow::FConnectionKey& Key)
+void FDataflowNode::UnregisterInputConnection(const UE::Dataflow::FConnectionKey& Key)
 {
 	if (ExpandedInputs.Remove(Key))
 	{
@@ -865,11 +865,11 @@ void FDataflowNode::UnregisterInputConnection(const Dataflow::FConnectionKey& Ke
 	}
 }
 
-FDataflowOutput& FDataflowNode::RegisterOutputConnectionInternal(const Dataflow::FConnectionReference& Reference, const FName& PropertyName)
+FDataflowOutput& FDataflowNode::RegisterOutputConnectionInternal(const UE::Dataflow::FConnectionReference& Reference, const FName& PropertyName)
 {
 	TUniquePtr<FStructOnScope> ScriptOnStruct = TUniquePtr<FStructOnScope>(NewStructOnScope());
 	check(ScriptOnStruct);
-	Dataflow::FOutputParameters OutputParams;
+	UE::Dataflow::FOutputParameters OutputParams;
 	InitConnectionParametersFromPropertyReference(*ScriptOnStruct, Reference.Reference, PropertyName, OutputParams);
 	FDataflowOutput* OutputConnection = new FDataflowOutput(OutputParams);
 	check(OutputConnection->RealAddress() == Reference.Reference);
@@ -1049,7 +1049,7 @@ bool FDataflowNode::ValidateConnections()
 #if 0
 		 // disabling this out this for now as this fail all over the place for some dataflow graphs 
 		 // we may get rid of the metadata constraints we may not need it anymore ( to be decided later )
-			for (const TPair<Dataflow::FConnectionKey, FDataflowInput*>& ExpandedInput : ExpandedInputs)
+			for (const TPair<UE::Dataflow::FConnectionKey, FDataflowInput*>& ExpandedInput : ExpandedInputs)
 			{
 				const FDataflowInput* Input = ExpandedInput.Value;
 
@@ -1095,23 +1095,23 @@ TUniquePtr<const FStructOnScope> FDataflowNode::NewStructOnScopeConst() const
 
 FString FDataflowNode::GetToolTip() const
 {
-	Dataflow::FFactoryParameters FactoryParameters = ::Dataflow::FNodeFactory::GetInstance()->GetParameters(GetType());
+	UE::Dataflow::FFactoryParameters FactoryParameters = UE::Dataflow::FNodeFactory::GetInstance()->GetParameters(GetType());
 
 	return FactoryParameters.ToolTip;
 }
 
-FText FDataflowNode::GetPinDisplayName(const FName& PropertyFullName, const Dataflow::FPin::EDirection Direction) const
+FText FDataflowNode::GetPinDisplayName(const FName& PropertyFullName, const UE::Dataflow::FPin::EDirection Direction) const
 {
 	int32 ContainerIndex = INDEX_NONE;
 
-	if (Direction == Dataflow::FPin::EDirection::INPUT)
+	if (Direction == UE::Dataflow::FPin::EDirection::INPUT)
 	{
 		if (const FDataflowInput* const Input = FindInput(PropertyFullName))
 		{
 			ContainerIndex = Input->GetContainerIndex();
 		}
 	}
-	else if (Direction == Dataflow::FPin::EDirection::OUTPUT)
+	else if (Direction == UE::Dataflow::FPin::EDirection::OUTPUT)
 	{
 		if (const FDataflowOutput* const Output = FindOutput(PropertyFullName))
 		{
@@ -1134,26 +1134,26 @@ FText FDataflowNode::GetPinDisplayName(const FName& PropertyFullName, const Data
 	return FText();
 }
 
-FString FDataflowNode::GetPinToolTip(const FName& PropertyFullName, const Dataflow::FPin::EDirection Direction) const
+FString FDataflowNode::GetPinToolTip(const FName& PropertyFullName, const UE::Dataflow::FPin::EDirection Direction) const
 {
 #if WITH_EDITORONLY_DATA
-	if (Direction == Dataflow::FPin::EDirection::INPUT)
+	if (Direction == UE::Dataflow::FPin::EDirection::INPUT)
 	{
 		if (const FDataflowInput* const Input = FindInput(PropertyFullName))
 		{
 			if (const FProperty* const Property = Input->GetProperty())
 			{
-				return Dataflow::Private::GetPinToolTipFromProperty(Property);
+				return UE::Dataflow::Private::GetPinToolTipFromProperty(Property);
 			}
 		}
 	}
-	else if (Direction == Dataflow::FPin::EDirection::OUTPUT)
+	else if (Direction == UE::Dataflow::FPin::EDirection::OUTPUT)
 	{
 		if (const FDataflowOutput* const Output = FindOutput(PropertyFullName))
 		{
 			if (const FProperty* const Property = Output->GetProperty())
 			{
-				return Dataflow::Private::GetPinToolTipFromProperty(Property);
+				return UE::Dataflow::Private::GetPinToolTipFromProperty(Property);
 			}
 		}
 	}
@@ -1163,7 +1163,7 @@ FString FDataflowNode::GetPinToolTip(const FName& PropertyFullName, const Datafl
 		{
 			if (const FProperty* const Property = FindProperty(Struct, PropertyFullName))
 			{
-				return Dataflow::Private::GetPinToolTipFromProperty(Property);
+				return UE::Dataflow::Private::GetPinToolTipFromProperty(Property);
 			}
 		}
 	}
@@ -1172,26 +1172,26 @@ FString FDataflowNode::GetPinToolTip(const FName& PropertyFullName, const Datafl
 	return {};
 }
 
-TArray<FString> FDataflowNode::GetPinMetaData(const FName& PropertyFullName, const Dataflow::FPin::EDirection Direction) const
+TArray<FString> FDataflowNode::GetPinMetaData(const FName& PropertyFullName, const UE::Dataflow::FPin::EDirection Direction) const
 {
 #if WITH_EDITORONLY_DATA
-	if (Direction == Dataflow::FPin::EDirection::INPUT)
+	if (Direction == UE::Dataflow::FPin::EDirection::INPUT)
 	{
 		if (const FDataflowInput* const Input = FindInput(PropertyFullName))
 		{
 			if (const FProperty* const Property = Input->GetProperty())
 			{
-				return Dataflow::Private::GetPinMetaDataFromProperty(Property);
+				return UE::Dataflow::Private::GetPinMetaDataFromProperty(Property);
 			}
 		}
 	}
-	else if (Direction == Dataflow::FPin::EDirection::OUTPUT)
+	else if (Direction == UE::Dataflow::FPin::EDirection::OUTPUT)
 	{
 		if (const FDataflowOutput* const Output = FindOutput(PropertyFullName))
 		{
 			if (const FProperty* const Property = Output->GetProperty())
 			{
-				return Dataflow::Private::GetPinMetaDataFromProperty(Property);
+				return UE::Dataflow::Private::GetPinMetaDataFromProperty(Property);
 			}
 		}
 	}
@@ -1201,7 +1201,7 @@ TArray<FString> FDataflowNode::GetPinMetaData(const FName& PropertyFullName, con
 		{
 			if (const FProperty* const Property = FindProperty(Struct, PropertyFullName))
 			{
-				return Dataflow::Private::GetPinMetaDataFromProperty(Property);
+				return UE::Dataflow::Private::GetPinMetaDataFromProperty(Property);
 			}
 		}
 	}
@@ -1221,7 +1221,7 @@ void FDataflowNode::CopyNodeProperties(const TSharedPtr<FDataflowNode> CopyFromD
 	this->SerializeInternal(ArReader);
 }
 
-void FDataflowNode::ForwardInput(Dataflow::FContext& Context, const Dataflow::FConnectionReference& InputReference, const Dataflow::FConnectionReference& Reference) const
+void FDataflowNode::ForwardInput(UE::Dataflow::FContext& Context, const UE::Dataflow::FConnectionReference& InputReference, const UE::Dataflow::FConnectionReference& Reference) const
 {
 	if (const FDataflowOutput* Output = FindOutput(Reference))
 	{
@@ -1260,18 +1260,18 @@ void FDataflowNode::NotifyConnectionTypeChanged(FDataflowConnection* Connection)
 {
 	if (Connection->IsAnyType())
 	{
-		if (Connection->GetDirection() == Dataflow::FPin::EDirection::INPUT)
+		if (Connection->GetDirection() == UE::Dataflow::FPin::EDirection::INPUT)
 		{
 			OnInputTypeChanged((FDataflowInput*)Connection);
 		}
-		if (Connection->GetDirection() == Dataflow::FPin::EDirection::OUTPUT)
+		if (Connection->GetDirection() == UE::Dataflow::FPin::EDirection::OUTPUT)
 		{
 			OnOutputTypeChanged((FDataflowOutput*)Connection);
 		}
 	}
 }
 
-bool FDataflowNode::SetInputConcreteType(const Dataflow::FConnectionReference& InputReference, FName NewType)
+bool FDataflowNode::SetInputConcreteType(const UE::Dataflow::FConnectionReference& InputReference, FName NewType)
 {
 	if (FDataflowInput* Input = FindInput(InputReference))
 	{
@@ -1283,7 +1283,7 @@ bool FDataflowNode::SetInputConcreteType(const Dataflow::FConnectionReference& I
 	return false;
 }
 
-bool FDataflowNode::SetOutputConcreteType(const Dataflow::FConnectionReference& OutputReference, FName NewType)
+bool FDataflowNode::SetOutputConcreteType(const UE::Dataflow::FConnectionReference& OutputReference, FName NewType)
 {
 	if (FDataflowOutput* Output = FindOutput(OutputReference))
 	{
@@ -1298,7 +1298,7 @@ bool FDataflowNode::SetOutputConcreteType(const Dataflow::FConnectionReference&
 bool FDataflowNode::SetAllConnectionConcreteType(FName NewType)
 {
 	bool bChanged = false;
-	for (const TPair<Dataflow::FConnectionKey, FDataflowInput*>& InputEntry : ExpandedInputs)
+	for (const TPair<UE::Dataflow::FConnectionKey, FDataflowInput*>& InputEntry : ExpandedInputs)
 	{
 		FDataflowInput* const Input = InputEntry.Value;
 		if (Input && Input->GetType() != NewType)
diff --git a/Engine/Source/Runtime/Experimental/Dataflow/Core/Private/Dataflow/DataflowNodeFactory.cpp b/Engine/Source/Runtime/Experimental/Dataflow/Core/Private/Dataflow/DataflowNodeFactory.cpp
--- a/Engine/Source/Runtime/Experimental/Dataflow/Core/Private/Dataflow/DataflowNodeFactory.cpp
+++ b/Engine/Source/Runtime/Experimental/Dataflow/Core/Private/Dataflow/DataflowNodeFactory.cpp
@@ -11,15 +11,15 @@
 
 DEFINE_LOG_CATEGORY_STATIC(LogDataflowFactory, Warning, All);
 
-namespace Dataflow
+namespace UE::Dataflow
 {
 	FNodeFactory* FNodeFactory::Instance = nullptr;
 
 	void FNodeFactory::RegisterDefaultNodes()
 	{
-		Dataflow::RegisterCoreNodes();
-		Dataflow::RegisterDataflowMathNodes();
-		Dataflow::RegisterDataflowVectorNodes();
+		RegisterCoreNodes();
+		RegisterDataflowMathNodes();
+		RegisterDataflowVectorNodes();
 	}
 
 	TSharedPtr<FDataflowNode> FNodeFactory::NewNodeFromRegisteredType(FGraph& Graph, const FNewNodeParameters& Param)
@@ -120,9 +120,9 @@ namespace Dataflow
 	bool FNodeFactory::IsNodeDeprecated(const FName NodeType)
 	{
 		// Display node deprecated if InNode is deprecated
-		if (Dataflow::FNodeFactory* Factory = GetInstance())
+		if (FNodeFactory* Factory = GetInstance())
 		{
-			const Dataflow::FFactoryParameters& Param = Factory->GetParameters(NodeType);
+			const FFactoryParameters& Param = Factory->GetParameters(NodeType);
 			if (Param.IsValid() && Param.IsDeprecated())
 			{
 				return true;
@@ -135,9 +135,9 @@ namespace Dataflow
 	bool FNodeFactory::IsNodeExperimental(const FName NodeType)
 	{
 		// Display node experimental if InNode is experimental
-		if (Dataflow::FNodeFactory* Factory = GetInstance())
+		if (FNodeFactory* Factory = GetInstance())
 		{
-			const Dataflow::FFactoryParameters& Param = Factory->GetParameters(NodeType);
+			const FFactoryParameters& Param = Factory->GetParameters(NodeType);
 			if (Param.IsValid() && Param.IsExperimental())
 			{
 				return true;
diff --git a/Engine/Source/Runtime/Experimental/Dataflow/Core/Private/Dataflow/DataflowOverrideNode.cpp b/Engine/Source/Runtime/Experimental/Dataflow/Core/Private/Dataflow/DataflowOverrideNode.cpp
--- a/Engine/Source/Runtime/Experimental/Dataflow/Core/Private/Dataflow/DataflowOverrideNode.cpp
+++ b/Engine/Source/Runtime/Experimental/Dataflow/Core/Private/Dataflow/DataflowOverrideNode.cpp
@@ -39,7 +39,7 @@ bool FDataflowOverrideNode::ShouldInvalidate(FName InKey) const
 }
 
 template <>
-int32 FDataflowOverrideNode::GetDefaultValue(Dataflow::FContext& Context) const
+int32 FDataflowOverrideNode::GetDefaultValue(UE::Dataflow::FContext& Context) const
 {
 	const FString DefaultValue = GetValue<FString>(Context, &Default, Default);
 	int32 DefaultIntValue = 0;
@@ -53,7 +53,7 @@ int32 FDataflowOverrideNode::GetDefaultValue(Dataflow::FContext& Context) const
 }
 
 template <>
-float FDataflowOverrideNode::GetDefaultValue(Dataflow::FContext& Context) const
+float FDataflowOverrideNode::GetDefaultValue(UE::Dataflow::FContext& Context) const
 {
 	const FString DefaultValue = GetValue<FString>(Context, &Default, Default);
 	float DefaultFloatValue = 0.f;
@@ -66,7 +66,7 @@ float FDataflowOverrideNode::GetDefaultValue(Dataflow::FContext& Context) const
 	return DefaultFloatValue;
 }
 
-FString FDataflowOverrideNode::GetValueFromAsset(Dataflow::FContext& Context, const UObject* InOwner) const
+FString FDataflowOverrideNode::GetValueFromAsset(UE::Dataflow::FContext& Context, const UObject* InOwner) const
 {
 	FName InKey = GetValue<FName>(Context, &Key, Key);
 	FString EmptyString;
diff --git a/Engine/Source/Runtime/Experimental/Dataflow/Core/Private/Dataflow/DataflowSettings.cpp b/Engine/Source/Runtime/Experimental/Dataflow/Core/Private/Dataflow/DataflowSettings.cpp
--- a/Engine/Source/Runtime/Experimental/Dataflow/Core/Private/Dataflow/DataflowSettings.cpp
+++ b/Engine/Source/Runtime/Experimental/Dataflow/Core/Private/Dataflow/DataflowSettings.cpp
@@ -4,23 +4,23 @@
 
 #include UE_INLINE_GENERATED_CPP_BY_NAME(DataflowSettings)
 
-namespace Dataflow::Private
+namespace UE::Dataflow::Private
 {
 	static const FLinearColor CManagedArrayCollectionPinTypeColor = FLinearColor(0.353393f, 0.454175f, 1.0f, 1.0f);
 	static const FLinearColor CArrayPinTypeColor = FLinearColor(1.0f, 0.172585f, 0.0f, 1.0f);
 	static const FLinearColor CBoxPinTypeColor = FLinearColor(0.013575f, 0.770000f, 0.429609f, 1.0f);
 	static const FLinearColor CSpherePinTypeColor = FLinearColor(0.2f, 0.6f, 1.f, 1.0f);
 	static const FLinearColor CDataflowAnyTypePinTypeColor = FLinearColor(0.3f, 0.3f, 0.3f, 1.0f);
-} // namespace Dataflow::Private
+} // namespace UE::Dataflow::Private
 
 UDataflowSettings::UDataflowSettings(const FObjectInitializer& ObjectInitlaizer)
 	: Super(ObjectInitlaizer)
 {
-	ManagedArrayCollectionPinTypeColor = Dataflow::Private::CManagedArrayCollectionPinTypeColor;
-	ArrayPinTypeColor = Dataflow::Private::CArrayPinTypeColor;
-	BoxPinTypeColor = Dataflow::Private::CBoxPinTypeColor;
-	SpherePinTypeColor = Dataflow::Private::CSpherePinTypeColor;
-	DataflowAnyTypePinTypeColor = Dataflow::Private::CDataflowAnyTypePinTypeColor;
+	ManagedArrayCollectionPinTypeColor = UE::Dataflow::Private::CManagedArrayCollectionPinTypeColor;
+	ArrayPinTypeColor = UE::Dataflow::Private::CArrayPinTypeColor;
+	BoxPinTypeColor = UE::Dataflow::Private::CBoxPinTypeColor;
+	SpherePinTypeColor = UE::Dataflow::Private::CSpherePinTypeColor;
+	DataflowAnyTypePinTypeColor = UE::Dataflow::Private::CDataflowAnyTypePinTypeColor;
 }
 
 FName UDataflowSettings::GetCategoryName() const
diff --git a/Engine/Source/Runtime/Experimental/Dataflow/Core/Private/Dataflow/DataflowVectorNodes.cpp b/Engine/Source/Runtime/Experimental/Dataflow/Core/Private/Dataflow/DataflowVectorNodes.cpp
--- a/Engine/Source/Runtime/Experimental/Dataflow/Core/Private/Dataflow/DataflowVectorNodes.cpp
+++ b/Engine/Source/Runtime/Experimental/Dataflow/Core/Private/Dataflow/DataflowVectorNodes.cpp
@@ -5,7 +5,7 @@
 #include "Dataflow/DataflowNodeFactory.h"
 #include "Dataflow/DataflowNodeColorsRegistry.h"
 
-namespace Dataflow
+namespace UE::Dataflow
 {
 	void RegisterDataflowVectorNodes()
 	{
@@ -28,7 +28,7 @@ namespace Dataflow
 
 //-----------------------------------------------------------
 
-FDataflowVectorMakeVec2Node::FDataflowVectorMakeVec2Node(const Dataflow::FNodeParameters& InParam, FGuid InGuid)
+FDataflowVectorMakeVec2Node::FDataflowVectorMakeVec2Node(const UE::Dataflow::FNodeParameters& InParam, FGuid InGuid)
 	: FDataflowNode(InParam, InGuid)
 {
 	RegisterInputConnection(&X);
@@ -38,7 +38,7 @@ FDataflowVectorMakeVec2Node::FDataflowVectorMakeVec2Node(const Dataflow::FNodePa
 	SetOutputConcreteType<FVector2D>(&Vector2D);
 }
 
-void FDataflowVectorMakeVec2Node::Evaluate(Dataflow::FContext& Context, const FDataflowOutput* Out) const
+void FDataflowVectorMakeVec2Node::Evaluate(UE::Dataflow::FContext& Context, const FDataflowOutput* Out) const
 {
 	if (Out->IsA(&Vector2D))
 	{
@@ -50,7 +50,7 @@ void FDataflowVectorMakeVec2Node::Evaluate(Dataflow::FContext& Context, const FD
 
 //-----------------------------------------------------------
 
-FDataflowVectorMakeVec3Node::FDataflowVectorMakeVec3Node(const Dataflow::FNodeParameters& InParam, FGuid InGuid)
+FDataflowVectorMakeVec3Node::FDataflowVectorMakeVec3Node(const UE::Dataflow::FNodeParameters& InParam, FGuid InGuid)
 	: FDataflowNode(InParam, InGuid)
 {
 	RegisterInputConnection(&X);
@@ -61,7 +61,7 @@ FDataflowVectorMakeVec3Node::FDataflowVectorMakeVec3Node(const Dataflow::FNodePa
 	SetOutputConcreteType<FVector>(&Vector3D);
 }
 
-void FDataflowVectorMakeVec3Node::Evaluate(Dataflow::FContext& Context, const FDataflowOutput* Out) const
+void FDataflowVectorMakeVec3Node::Evaluate(UE::Dataflow::FContext& Context, const FDataflowOutput* Out) const
 {
 	if (Out->IsA(&Vector3D))
 	{
@@ -74,7 +74,7 @@ void FDataflowVectorMakeVec3Node::Evaluate(Dataflow::FContext& Context, const FD
 
 //-----------------------------------------------------------
 
-FDataflowVectorMakeVec4Node::FDataflowVectorMakeVec4Node(const Dataflow::FNodeParameters& InParam, FGuid InGuid)
+FDataflowVectorMakeVec4Node::FDataflowVectorMakeVec4Node(const UE::Dataflow::FNodeParameters& InParam, FGuid InGuid)
 	: FDataflowNode(InParam, InGuid)
 {
 	RegisterInputConnection(&X);
@@ -86,7 +86,7 @@ FDataflowVectorMakeVec4Node::FDataflowVectorMakeVec4Node(const Dataflow::FNodePa
 	SetOutputConcreteType<FVector4>(&Vector4D);
 }
 
-void FDataflowVectorMakeVec4Node::Evaluate(Dataflow::FContext& Context, const FDataflowOutput* Out) const
+void FDataflowVectorMakeVec4Node::Evaluate(UE::Dataflow::FContext& Context, const FDataflowOutput* Out) const
 {
 	if (Out->IsA(&Vector4D))
 	{
@@ -100,7 +100,7 @@ void FDataflowVectorMakeVec4Node::Evaluate(Dataflow::FContext& Context, const FD
 
 //-----------------------------------------------------------
 
-FDataflowVectorBreakNode::FDataflowVectorBreakNode(const Dataflow::FNodeParameters& InParam, FGuid InGuid)
+FDataflowVectorBreakNode::FDataflowVectorBreakNode(const UE::Dataflow::FNodeParameters& InParam, FGuid InGuid)
 	: FDataflowNode(InParam, InGuid)
 {
 	RegisterInputConnection(&V);
@@ -116,7 +116,7 @@ FDataflowVectorBreakNode::FDataflowVectorBreakNode(const Dataflow::FNodeParamete
 	SetOutputConcreteType<double>(&W);
 }
 
-void FDataflowVectorBreakNode::Evaluate(Dataflow::FContext& Context, const FDataflowOutput* Out) const
+void FDataflowVectorBreakNode::Evaluate(UE::Dataflow::FContext& Context, const FDataflowOutput* Out) const
 {
 	if (Out->IsA(&X))
 	{
@@ -138,7 +138,7 @@ void FDataflowVectorBreakNode::Evaluate(Dataflow::FContext& Context, const FData
 
 //-----------------------------------------------------------
 
-FDataflowVectorAddNode::FDataflowVectorAddNode(const Dataflow::FNodeParameters& InParam, FGuid InGuid)
+FDataflowVectorAddNode::FDataflowVectorAddNode(const UE::Dataflow::FNodeParameters& InParam, FGuid InGuid)
 	: FDataflowNode(InParam, InGuid)
 {
 	RegisterInputConnection(&A);
@@ -149,7 +149,7 @@ FDataflowVectorAddNode::FDataflowVectorAddNode(const Dataflow::FNodeParameters&
 
 }
 
-void FDataflowVectorAddNode::Evaluate(Dataflow::FContext& Context, const FDataflowOutput* Out) const
+void FDataflowVectorAddNode::Evaluate(UE::Dataflow::FContext& Context, const FDataflowOutput* Out) const
 {
 	if (Out->IsA(&V))
 	{
@@ -161,7 +161,7 @@ void FDataflowVectorAddNode::Evaluate(Dataflow::FContext& Context, const FDatafl
 
 //-----------------------------------------------------------
 
-FDataflowVectorSubtractNode::FDataflowVectorSubtractNode(const Dataflow::FNodeParameters& InParam, FGuid InGuid)
+FDataflowVectorSubtractNode::FDataflowVectorSubtractNode(const UE::Dataflow::FNodeParameters& InParam, FGuid InGuid)
 	: FDataflowNode(InParam, InGuid)
 {
 	RegisterInputConnection(&A);
@@ -172,7 +172,7 @@ FDataflowVectorSubtractNode::FDataflowVectorSubtractNode(const Dataflow::FNodePa
 
 }
 
-void FDataflowVectorSubtractNode::Evaluate(Dataflow::FContext& Context, const FDataflowOutput* Out) const
+void FDataflowVectorSubtractNode::Evaluate(UE::Dataflow::FContext& Context, const FDataflowOutput* Out) const
 {
 	if (Out->IsA(&V))
 	{
@@ -184,7 +184,7 @@ void FDataflowVectorSubtractNode::Evaluate(Dataflow::FContext& Context, const FD
 
 //-----------------------------------------------------------
 
-FDataflowVectorDotProductNode::FDataflowVectorDotProductNode(const Dataflow::FNodeParameters& InParam, FGuid InGuid)
+FDataflowVectorDotProductNode::FDataflowVectorDotProductNode(const UE::Dataflow::FNodeParameters& InParam, FGuid InGuid)
 	: FDataflowNode(InParam, InGuid)
 {
 	RegisterInputConnection(&A);
@@ -194,7 +194,7 @@ FDataflowVectorDotProductNode::FDataflowVectorDotProductNode(const Dataflow::FNo
 	SetOutputConcreteType<double>(&DotProduct);
 }
 
-void FDataflowVectorDotProductNode::Evaluate(Dataflow::FContext& Context, const FDataflowOutput* Out) const
+void FDataflowVectorDotProductNode::Evaluate(UE::Dataflow::FContext& Context, const FDataflowOutput* Out) const
 {
 	if (Out->IsA(&DotProduct))
 	{
@@ -207,7 +207,7 @@ void FDataflowVectorDotProductNode::Evaluate(Dataflow::FContext& Context, const
 
 //-----------------------------------------------------------
 
-FDataflowVectorLengthNode::FDataflowVectorLengthNode(const Dataflow::FNodeParameters& InParam, FGuid InGuid)
+FDataflowVectorLengthNode::FDataflowVectorLengthNode(const UE::Dataflow::FNodeParameters& InParam, FGuid InGuid)
 	: FDataflowNode(InParam, InGuid)
 {
 	RegisterInputConnection(&V);
@@ -216,7 +216,7 @@ FDataflowVectorLengthNode::FDataflowVectorLengthNode(const Dataflow::FNodeParame
 	SetOutputConcreteType<double>(&Length);
 }
 
-void FDataflowVectorLengthNode::Evaluate(Dataflow::FContext& Context, const FDataflowOutput* Out) const
+void FDataflowVectorLengthNode::Evaluate(UE::Dataflow::FContext& Context, const FDataflowOutput* Out) const
 {
 	if (Out->IsA(&Length))
 	{
@@ -228,7 +228,7 @@ void FDataflowVectorLengthNode::Evaluate(Dataflow::FContext& Context, const FDat
 
 //-----------------------------------------------------------
 
-FDataflowVectorSquaredLengthNode::FDataflowVectorSquaredLengthNode(const Dataflow::FNodeParameters& InParam, FGuid InGuid)
+FDataflowVectorSquaredLengthNode::FDataflowVectorSquaredLengthNode(const UE::Dataflow::FNodeParameters& InParam, FGuid InGuid)
 	: FDataflowNode(InParam, InGuid)
 {
 	RegisterInputConnection(&V);
@@ -237,7 +237,7 @@ FDataflowVectorSquaredLengthNode::FDataflowVectorSquaredLengthNode(const Dataflo
 	SetOutputConcreteType<double>(&SquaredLength);
 }
 
-void FDataflowVectorSquaredLengthNode::Evaluate(Dataflow::FContext& Context, const FDataflowOutput* Out) const
+void FDataflowVectorSquaredLengthNode::Evaluate(UE::Dataflow::FContext& Context, const FDataflowOutput* Out) const
 {
 	if (Out->IsA(&SquaredLength))
 	{
@@ -249,7 +249,7 @@ void FDataflowVectorSquaredLengthNode::Evaluate(Dataflow::FContext& Context, con
 
 //-----------------------------------------------------------
 
-FDataflowVectorDistanceNode::FDataflowVectorDistanceNode(const Dataflow::FNodeParameters& InParam, FGuid InGuid)
+FDataflowVectorDistanceNode::FDataflowVectorDistanceNode(const UE::Dataflow::FNodeParameters& InParam, FGuid InGuid)
 	: FDataflowNode(InParam, InGuid)
 {
 	RegisterInputConnection(&A);
@@ -259,7 +259,7 @@ FDataflowVectorDistanceNode::FDataflowVectorDistanceNode(const Dataflow::FNodePa
 	SetOutputConcreteType<double>(&Distance);
 }
 
-void FDataflowVectorDistanceNode::Evaluate(Dataflow::FContext& Context, const FDataflowOutput* Out) const
+void FDataflowVectorDistanceNode::Evaluate(UE::Dataflow::FContext& Context, const FDataflowOutput* Out) const
 {
 	if (Out->IsA(&Distance))
 	{
@@ -273,7 +273,7 @@ void FDataflowVectorDistanceNode::Evaluate(Dataflow::FContext& Context, const FD
 
 //-----------------------------------------------------------
 
-FDataflowVectorCrossProductNode::FDataflowVectorCrossProductNode(const Dataflow::FNodeParameters& InParam, FGuid InGuid)
+FDataflowVectorCrossProductNode::FDataflowVectorCrossProductNode(const UE::Dataflow::FNodeParameters& InParam, FGuid InGuid)
 	: FDataflowNode(InParam, InGuid)
 {
 	RegisterInputConnection(&A);
@@ -283,7 +283,7 @@ FDataflowVectorCrossProductNode::FDataflowVectorCrossProductNode(const Dataflow:
 	SetOutputConcreteType<FVector>(&CrossProduct);
 }
 
-void FDataflowVectorCrossProductNode::Evaluate(Dataflow::FContext& Context, const FDataflowOutput* Out) const
+void FDataflowVectorCrossProductNode::Evaluate(UE::Dataflow::FContext& Context, const FDataflowOutput* Out) const
 {
 	if (Out->IsA(&CrossProduct))
 	{
@@ -297,7 +297,7 @@ void FDataflowVectorCrossProductNode::Evaluate(Dataflow::FContext& Context, cons
 
 //-----------------------------------------------------------
 
-FDataflowVectorScaleNode::FDataflowVectorScaleNode(const Dataflow::FNodeParameters& InParam, FGuid InGuid)
+FDataflowVectorScaleNode::FDataflowVectorScaleNode(const UE::Dataflow::FNodeParameters& InParam, FGuid InGuid)
 	: FDataflowNode(InParam, InGuid)
 {
 	Scale.Value = 1.0; // set default
@@ -306,7 +306,7 @@ FDataflowVectorScaleNode::FDataflowVectorScaleNode(const Dataflow::FNodeParamete
 	RegisterOutputConnection(&Scaled);
 }
 
-void FDataflowVectorScaleNode::Evaluate(Dataflow::FContext& Context, const FDataflowOutput* Out) const
+void FDataflowVectorScaleNode::Evaluate(UE::Dataflow::FContext& Context, const FDataflowOutput* Out) const
 {
 	if (Out->IsA(&Scaled))
 	{
@@ -337,14 +337,14 @@ bool FDataflowVectorScaleNode::OnOutputTypeChanged(const FDataflowOutput* Output
 
 //-----------------------------------------------------------
 
-FDataflowVectorNormalize::FDataflowVectorNormalize(const Dataflow::FNodeParameters& InParam, FGuid InGuid)
+FDataflowVectorNormalize::FDataflowVectorNormalize(const UE::Dataflow::FNodeParameters& InParam, FGuid InGuid)
 	: FDataflowNode(InParam, InGuid)
 {
 	RegisterInputConnection(&V);
 	RegisterOutputConnection(&Normalized);
 }
 
-void FDataflowVectorNormalize::Evaluate(Dataflow::FContext& Context, const FDataflowOutput* Out) const
+void FDataflowVectorNormalize::Evaluate(UE::Dataflow::FContext& Context, const FDataflowOutput* Out) const
 {
 	if (Out->IsA(&Normalized))
 	{
diff --git a/Engine/Source/Runtime/Experimental/Dataflow/Core/Public/Dataflow/DataflowConnection.h b/Engine/Source/Runtime/Experimental/Dataflow/Core/Public/Dataflow/DataflowConnection.h
--- a/Engine/Source/Runtime/Experimental/Dataflow/Core/Public/Dataflow/DataflowConnection.h
+++ b/Engine/Source/Runtime/Experimental/Dataflow/Core/Public/Dataflow/DataflowConnection.h
@@ -18,7 +18,7 @@ struct FDataflowNode;
 struct FDataflowInput;
 struct FDataflowOutput;
 
-namespace Dataflow
+namespace UE::Dataflow
 {
 
 	template<class T> inline DATAFLOWCORE_API FName GraphConnectionTypeName();
@@ -139,14 +139,14 @@ protected:
 	FGuid  Guid;
 	IDataflowTypePolicy* TypePolicy = nullptr;
 	uint32 Offset;
-	Dataflow::FPin::EDirection Direction;
+	UE::Dataflow::FPin::EDirection Direction;
 	bool bIsAnyType:1 = false;
 	bool bHasConcreteType : 1 = false;
 	bool bCanHidePin:1 = false;
 	bool bPinIsHidden:1 = false;
 
 	friend struct FDataflowNode;
-	friend class Dataflow::FGraph;
+	friend class UE::Dataflow::FGraph;
 
 protected:
 	DATAFLOWCORE_API bool IsOwningNodeEnabled() const;
@@ -174,8 +174,8 @@ protected:
 public:
 	FDataflowConnection() {};
 	UE_DEPRECATED(5.5, "Deprecated constructor : use FConnectionParameters to pass parameters")
-	DATAFLOWCORE_API FDataflowConnection(Dataflow::FPin::EDirection Direction, FName InType, FName InName, FDataflowNode* OwningNode = nullptr, const FProperty* InProperty = nullptr, FGuid InGuid = FGuid::NewGuid());
-	DATAFLOWCORE_API FDataflowConnection(Dataflow::FPin::EDirection Direction, const Dataflow::FConnectionParameters& Params);
+	DATAFLOWCORE_API FDataflowConnection(UE::Dataflow::FPin::EDirection Direction, FName InType, FName InName, FDataflowNode* OwningNode = nullptr, const FProperty* InProperty = nullptr, FGuid InGuid = FGuid::NewGuid());
+	DATAFLOWCORE_API FDataflowConnection(UE::Dataflow::FPin::EDirection Direction, const UE::Dataflow::FConnectionParameters& Params);
 	virtual ~FDataflowConnection() {};
 
 	FDataflowNode* GetOwningNode() { return OwningNode; }
@@ -183,20 +183,20 @@ public:
 
 	DATAFLOWCORE_API FGuid GetOwningNodeGuid() const;
 	DATAFLOWCORE_API uint32 GetOwningNodeValueHash() const;
-	DATAFLOWCORE_API Dataflow::FTimestamp GetOwningNodeTimestamp() const;
+	DATAFLOWCORE_API UE::Dataflow::FTimestamp GetOwningNodeTimestamp() const;
 
 	const FProperty* GetProperty() const { return Property; }
 
 	DATAFLOWCORE_API FString GetPropertyTooltip() const;
 	DATAFLOWCORE_API FString GetPropertyTypeNameTooltip() const;
 
-	Dataflow::FPin::EDirection GetDirection() const { return Direction; }
+	UE::Dataflow::FPin::EDirection GetDirection() const { return Direction; }
 	uint32 GetOffset() const { return Offset; }
 	virtual int32 GetContainerIndex() const { return INDEX_NONE; }
 	virtual uint32 GetContainerElementOffset() const { return INDEX_NONE; }
-	Dataflow::FConnectionKey GetConnectionKey() const 
+	UE::Dataflow::FConnectionKey GetConnectionKey() const 
 	{		
-		return Dataflow::FConnectionKey(GetOffset(), GetContainerIndex(), GetContainerElementOffset());
+		return UE::Dataflow::FConnectionKey(GetOffset(), GetContainerIndex(), GetContainerElementOffset());
 	}
 
 	FName GetType() const { return Type; }
@@ -208,7 +208,7 @@ public:
 	void SetName(FName InName) { Name = InName; }
 
 	virtual void* RealAddress() const { ensure(OwningNode);  return (void*)((size_t)OwningNode + (size_t)GetOffset()); };
-	Dataflow::FContextCacheKey CacheKey() const { return GetTypeHash(Guid); };
+	UE::Dataflow::FContextCacheKey CacheKey() const { return GetTypeHash(Guid); };
 
 	virtual bool AddConnection(FDataflowConnection* In) { return false; };
 	virtual bool RemoveConnection(FDataflowConnection* In) { return false; }
@@ -225,7 +225,7 @@ public:
 		return (size_t)RealAddress() == (size_t)InVar;
 	}
 
-	virtual void Invalidate(const Dataflow::FTimestamp& ModifiedTimestamp = Dataflow::FTimestamp::Current()) {};
+	virtual void Invalidate(const UE::Dataflow::FTimestamp& ModifiedTimestamp = UE::Dataflow::FTimestamp::Current()) {};
 
 	bool GetCanHidePin() const { return bCanHidePin; }
 	bool GetPinIsHidden() const { return bCanHidePin && bPinIsHidden; }
diff --git a/Engine/Source/Runtime/Experimental/Dataflow/Core/Public/Dataflow/DataflowCoreNodes.h b/Engine/Source/Runtime/Experimental/Dataflow/Core/Public/Dataflow/DataflowCoreNodes.h
--- a/Engine/Source/Runtime/Experimental/Dataflow/Core/Public/Dataflow/DataflowCoreNodes.h
+++ b/Engine/Source/Runtime/Experimental/Dataflow/Core/Public/Dataflow/DataflowCoreNodes.h
@@ -20,8 +20,8 @@ struct FDataflowReRouteNode : public FDataflowNode
 	DATAFLOW_NODE_DEFINE_INTERNAL(FDataflowReRouteNode, "ReRouteNode", "Core", "")
 
 public:
-	FDataflowReRouteNode(const Dataflow::FNodeParameters& Param, FGuid InGuid = FGuid::NewGuid());
-	virtual void Evaluate(Dataflow::FContext& Context, const FDataflowOutput* Out) const;
+	FDataflowReRouteNode(const UE::Dataflow::FNodeParameters& Param, FGuid InGuid = FGuid::NewGuid());
+	virtual void Evaluate(UE::Dataflow::FContext& Context, const FDataflowOutput* Out) const;
 
 
 public:
@@ -42,8 +42,8 @@ struct FDataflowBranchNode : public FDataflowNode
 	DATAFLOW_NODE_DEFINE_INTERNAL(FDataflowBranchNode, "Branch", "FlowControl", "")
 
 public:
-	FDataflowBranchNode(const Dataflow::FNodeParameters& Param, FGuid InGuid = FGuid::NewGuid());
-	virtual void Evaluate(Dataflow::FContext& Context, const FDataflowOutput* Out) const;
+	FDataflowBranchNode(const UE::Dataflow::FNodeParameters& Param, FGuid InGuid = FGuid::NewGuid());
+	virtual void Evaluate(UE::Dataflow::FContext& Context, const FDataflowOutput* Out) const;
 
 public:
 	UPROPERTY(meta = (DataflowInput, DisplayName = "TrueValue"))
@@ -72,8 +72,8 @@ struct FDataflowSelectNode : public FDataflowNode
 	DATAFLOW_NODE_DEFINE_INTERNAL(FDataflowSelectNode, "Select", "FlowControl", "")
 
 public:
-	FDataflowSelectNode(const Dataflow::FNodeParameters& Param, FGuid InGuid = FGuid::NewGuid());
-	virtual void Evaluate(Dataflow::FContext& Context, const FDataflowOutput* Out) const;
+	FDataflowSelectNode(const UE::Dataflow::FNodeParameters& Param, FGuid InGuid = FGuid::NewGuid());
+	virtual void Evaluate(UE::Dataflow::FContext& Context, const FDataflowOutput* Out) const;
 
 public:
 	UPROPERTY()
@@ -88,14 +88,14 @@ public:
 private:
 	virtual bool OnInputTypeChanged(const FDataflowInput* Input) override;
 	virtual bool OnOutputTypeChanged(const FDataflowOutput* Input) override;
-	virtual TArray<Dataflow::FPin> AddPins() override;
+	virtual TArray<UE::Dataflow::FPin> AddPins() override;
 	virtual bool CanAddPin() const override { return true; }
 	virtual bool CanRemovePin() const override { return Inputs.Num() > NumInitialInputs; }
-	virtual TArray<Dataflow::FPin> GetPinsToRemove() const override;
-	virtual void OnPinRemoved(const Dataflow::FPin& Pin) override;
+	virtual TArray<UE::Dataflow::FPin> GetPinsToRemove() const override;
+	virtual void OnPinRemoved(const UE::Dataflow::FPin& Pin) override;
 	virtual void Serialize(FArchive& Ar) override;
 
-	Dataflow::TConnectionReference<FDataflowAnyType> GetConnectionReference(int32 Index) const;
+	UE::Dataflow::TConnectionReference<FDataflowAnyType> GetConnectionReference(int32 Index) const;
 
 	static constexpr int32 NumRequiredDataflowInputs = 1;
 	static constexpr int32 NumInitialInputs = 2;
@@ -113,15 +113,15 @@ struct FDataflowPrintNode : public FDataflowNode
 	DATAFLOW_NODE_DEFINE_INTERNAL(FDataflowPrintNode, "Print", "Core", "")
 
 public:
-	FDataflowPrintNode(const Dataflow::FNodeParameters& Param, FGuid InGuid = FGuid::NewGuid());
-	virtual void Evaluate(Dataflow::FContext& Context, const FDataflowOutput* Out) const;
+	FDataflowPrintNode(const UE::Dataflow::FNodeParameters& Param, FGuid InGuid = FGuid::NewGuid());
+	virtual void Evaluate(UE::Dataflow::FContext& Context, const FDataflowOutput* Out) const;
 
 public:
 	UPROPERTY(meta = (DataflowInput, DisplayName = "Value"))
 	FDataflowStringConvertibleTypes Value;
 };
 
-namespace Dataflow
+namespace UE::Dataflow
 {
 	void RegisterCoreNodes();
 }
diff --git a/Engine/Source/Runtime/Experimental/Dataflow/Core/Public/Dataflow/DataflowGraph.h b/Engine/Source/Runtime/Experimental/Dataflow/Core/Public/Dataflow/DataflowGraph.h
--- a/Engine/Source/Runtime/Experimental/Dataflow/Core/Public/Dataflow/DataflowGraph.h
+++ b/Engine/Source/Runtime/Experimental/Dataflow/Core/Public/Dataflow/DataflowGraph.h
@@ -10,7 +10,7 @@
 
 struct FDataflowConnection;
 
-namespace Dataflow
+namespace UE::Dataflow
 {
 
 	struct FLink {
@@ -191,13 +191,13 @@ namespace Dataflow
 }
 
 
-FORCEINLINE FArchive& operator<<(FArchive& Ar, Dataflow::FLink& Value)
+FORCEINLINE FArchive& operator<<(FArchive& Ar, UE::Dataflow::FLink& Value)
 {
 	Ar << Value.InputNode << Value.OutputNode << Value.Input << Value.Output;
 	return Ar;
 }
 
-FORCEINLINE FArchive& operator<<(Chaos::FChaosArchive& Ar, Dataflow::FLink& Value)
+FORCEINLINE FArchive& operator<<(Chaos::FChaosArchive& Ar, UE::Dataflow::FLink& Value)
 {
 	Ar << Value.InputNode << Value.OutputNode << Value.Input << Value.Output;
 	return Ar;
diff --git a/Engine/Source/Runtime/Experimental/Dataflow/Core/Public/Dataflow/DataflowInputOutput.h b/Engine/Source/Runtime/Experimental/Dataflow/Core/Public/Dataflow/DataflowInputOutput.h
--- a/Engine/Source/Runtime/Experimental/Dataflow/Core/Public/Dataflow/DataflowInputOutput.h
+++ b/Engine/Source/Runtime/Experimental/Dataflow/Core/Public/Dataflow/DataflowInputOutput.h
@@ -17,7 +17,7 @@ struct FDataflowOutput;
 //
 //  Input
 //
-namespace Dataflow
+namespace UE::Dataflow
 {
 	struct FInputParameters : public FConnectionParameters
 	{
@@ -53,9 +53,9 @@ protected:
 
 public:
 	UE_DEPRECATED(5.5, "Deprecated constructor : Guid is now passed through FInputParameters")
-	DATAFLOWCORE_API FDataflowInput(const Dataflow::FInputParameters& Param, FGuid InGuid);
+	DATAFLOWCORE_API FDataflowInput(const UE::Dataflow::FInputParameters& Param, FGuid InGuid);
 
-	DATAFLOWCORE_API FDataflowInput(const Dataflow::FInputParameters& Param = {});
+	DATAFLOWCORE_API FDataflowInput(const UE::Dataflow::FInputParameters& Param = {});
 
 	virtual bool AddConnection(FDataflowConnection* InOutput) override;
 	virtual bool RemoveConnection(FDataflowConnection* InOutput) override;
@@ -72,21 +72,21 @@ public:
 	* @return the typed value of the input 
 	*/
 	template<class T>
-	const T& GetValue(Dataflow::FContext& Context, const T& Default) const;
+	const T& GetValue(UE::Dataflow::FContext& Context, const T& Default) const;
 
 	template<typename TAnyType>
-	typename TAnyType::FStorageType GetValueFromAnyType(Dataflow::FContext& Context, const typename TAnyType::FStorageType& Default) const;
+	typename TAnyType::FStorageType GetValueFromAnyType(UE::Dataflow::FContext& Context, const typename TAnyType::FStorageType& Default) const;
 
 	/**
 	* pull the value from the upstream connections
 	* the upstream graph is evaluated if necessary and values are cached along the way 
 	*/
-	DATAFLOWCORE_API void PullValue(Dataflow::FContext& Context) const;
+	DATAFLOWCORE_API void PullValue(UE::Dataflow::FContext& Context) const;
 
 	template<class T>
-	TFuture<const T&> GetValueParallel(Dataflow::FContext& Context, const T& Default) const;
+	TFuture<const T&> GetValueParallel(UE::Dataflow::FContext& Context, const T& Default) const;
 
-	virtual void Invalidate(const Dataflow::FTimestamp& ModifiedTimestamp = Dataflow::FTimestamp::Current()) override;
+	virtual void Invalidate(const UE::Dataflow::FTimestamp& ModifiedTimestamp = UE::Dataflow::FTimestamp::Current()) override;
 };
 
 USTRUCT()
@@ -101,7 +101,7 @@ private:
 	// uint32 Offset; // On base class. This is the Offset to ArrayProperty from OwningNode.
 
 public:
-	explicit FDataflowArrayInput(int32 InIndex = INDEX_NONE, const Dataflow::FArrayInputParameters& Param = {});
+	explicit FDataflowArrayInput(int32 InIndex = INDEX_NONE, const UE::Dataflow::FArrayInputParameters& Param = {});
 
 	DATAFLOWCORE_API virtual void* RealAddress() const override;
 	virtual int32 GetContainerIndex() const override { return Index; }
@@ -111,7 +111,7 @@ public:
 //
 // Output
 //
-namespace Dataflow
+namespace UE::Dataflow
 {
 	struct FOutputParameters: public FConnectionParameters
 	{
@@ -133,7 +133,7 @@ struct FDataflowOutput : public FDataflowConnection
 	UE_DEPRECATED(5.5, "Use PassthroughKey instead")
 	uint32 PassthroughOffset = INDEX_NONE;
 
-	Dataflow::FConnectionKey PassthroughKey;
+	UE::Dataflow::FConnectionKey PassthroughKey;
 
 protected:
 	friend struct FDataflowInput;
@@ -145,9 +145,9 @@ public:
 	mutable TSharedPtr<FCriticalSection> OutputLock;
 	
 	UE_DEPRECATED(5.5, "Deprecated constructor : Guid is now passed through FOutputParameters")
-	DATAFLOWCORE_API FDataflowOutput(const Dataflow::FOutputParameters& Param, FGuid InGuid);
+	DATAFLOWCORE_API FDataflowOutput(const UE::Dataflow::FOutputParameters& Param, FGuid InGuid);
 
-	DATAFLOWCORE_API FDataflowOutput(const Dataflow::FOutputParameters& Param = {});
+	DATAFLOWCORE_API FDataflowOutput(const UE::Dataflow::FOutputParameters& Param = {});
 
 	PRAGMA_DISABLE_DEPRECATION_WARNINGS
 	~FDataflowOutput() = default;
@@ -171,11 +171,11 @@ public:
 	UE_DEPRECATED(5.5, "Use SetPassthroughInput instead")
 	virtual void SetPassthroughOffset(const uint32 InPassthroughOffset)
 	{
-		SetPassthroughInput(Dataflow::FConnectionKey(InPassthroughOffset, INDEX_NONE, INDEX_NONE));
+		SetPassthroughInput(UE::Dataflow::FConnectionKey(InPassthroughOffset, INDEX_NONE, INDEX_NONE));
 	}
 
-	DATAFLOWCORE_API FDataflowOutput& SetPassthroughInput(const Dataflow::FConnectionReference& Reference);
-	DATAFLOWCORE_API FDataflowOutput& SetPassthroughInput(const Dataflow::FConnectionKey& Key);
+	DATAFLOWCORE_API FDataflowOutput& SetPassthroughInput(const UE::Dataflow::FConnectionReference& Reference);
+	DATAFLOWCORE_API FDataflowOutput& SetPassthroughInput(const UE::Dataflow::FConnectionKey& Key);
 
 	DATAFLOWCORE_API const FDataflowInput* GetPassthroughInput() const;
 
@@ -189,7 +189,7 @@ public:
 	}
  
 	template<class T>
-	void SetValue(T&& InVal, Dataflow::FContext& Context) const
+	void SetValue(T&& InVal, UE::Dataflow::FContext& Context) const
 	{
 		if (Property)
 		{
@@ -198,7 +198,7 @@ public:
 	}
 
 	template<typename TAnyType>
-	void SetValueFromAnyType(const typename TAnyType::FStorageType& InVal, Dataflow::FContext& Context) const
+	void SetValueFromAnyType(const typename TAnyType::FStorageType& InVal, UE::Dataflow::FContext& Context) const
 	{
 		TAnyType::FPolicyType::VisitPolicyByType(GetType(),
 			[this, &Context, &InVal](auto SingleTypePolicy)
@@ -211,7 +211,7 @@ public:
 	}
 
 	template<class T>
-	const T& GetValue(Dataflow::FContext& Context, const T& Default) const
+	const T& GetValue(UE::Dataflow::FContext& Context, const T& Default) const
 	{
 		if (!this->Evaluate(Context))
 		{
@@ -226,24 +226,24 @@ public:
 		return Default;
 	}
 
-	bool HasCachedValue(Dataflow::FContext& Context) const
+	bool HasCachedValue(UE::Dataflow::FContext& Context) const
 	{
 		return Context.HasData(CacheKey(), GetOwningNodeTimestamp());
 	}
 
 	// there's no need for a templatized version as the parameter will not be used
 	// the method do check if the type of the input is the same as the output type though 
-	DATAFLOWCORE_API void ForwardInput(const Dataflow::FConnectionReference& InputReference, Dataflow::FContext& Context) const;
-	DATAFLOWCORE_API void ForwardInput(const FDataflowInput* Input, Dataflow::FContext& Context) const;
+	DATAFLOWCORE_API void ForwardInput(const UE::Dataflow::FConnectionReference& InputReference, UE::Dataflow::FContext& Context) const;
+	DATAFLOWCORE_API void ForwardInput(const FDataflowInput* Input, UE::Dataflow::FContext& Context) const;
 
 
-	DATAFLOWCORE_API bool EvaluateImpl(Dataflow::FContext& Context) const;
+	DATAFLOWCORE_API bool EvaluateImpl(UE::Dataflow::FContext& Context) const;
 	
-	DATAFLOWCORE_API bool Evaluate(Dataflow::FContext& Context) const;
+	DATAFLOWCORE_API bool Evaluate(UE::Dataflow::FContext& Context) const;
 
-	DATAFLOWCORE_API TFuture<bool> EvaluateParallel(Dataflow::FContext& Context) const;
+	DATAFLOWCORE_API TFuture<bool> EvaluateParallel(UE::Dataflow::FContext& Context) const;
 
-	DATAFLOWCORE_API virtual void Invalidate(const Dataflow::FTimestamp& ModifiedTimestamp = Dataflow::FTimestamp::Current()) override;
+	DATAFLOWCORE_API virtual void Invalidate(const UE::Dataflow::FTimestamp& ModifiedTimestamp = UE::Dataflow::FTimestamp::Current()) override;
 
 private:
 };
@@ -258,7 +258,7 @@ struct TStructOpsTypeTraits<FDataflowOutput> : public  TStructOpsTypeTraitsBase2
 };
 
 template<typename T>
-const T& FDataflowInput::GetValue(Dataflow::FContext& Context, const T& Default) const
+const T& FDataflowInput::GetValue(UE::Dataflow::FContext& Context, const T& Default) const
 {
 	if (const FDataflowOutput* ConnectionOut = GetConnection())
 	{
@@ -276,14 +276,14 @@ const T& FDataflowInput::GetValue(Dataflow::FContext& Context, const T& Default)
 }
 
 template<typename TAnyType>
-typename TAnyType::FStorageType FDataflowInput::GetValueFromAnyType(Dataflow::FContext& Context, const typename TAnyType::FStorageType& Default) const
+typename TAnyType::FStorageType FDataflowInput::GetValueFromAnyType(UE::Dataflow::FContext& Context, const typename TAnyType::FStorageType& Default) const
 {
 	typename TAnyType::FStorageType ReturnValue = Default;
 	if (const FDataflowOutput* ConnectionOut = GetConnection())
 	{
 		if (ConnectionOut->Evaluate(Context))
 		{
-			if (const TUniquePtr<Dataflow::FContextCacheElementBase>* CacheEntry = Context.GetDataImpl(ConnectionOut->CacheKey()))
+			if (const TUniquePtr<UE::Dataflow::FContextCacheElementBase>* CacheEntry = Context.GetDataImpl(ConnectionOut->CacheKey()))
 			{
 				if (*CacheEntry)
 				{
@@ -303,7 +303,7 @@ typename TAnyType::FStorageType FDataflowInput::GetValueFromAnyType(Dataflow::FC
 }
 
 template<class T>
-TFuture<const T&> FDataflowInput::GetValueParallel(Dataflow::FContext& Context, const T& Default) const
+TFuture<const T&> FDataflowInput::GetValueParallel(UE::Dataflow::FContext& Context, const T& Default) const
 {
 	return Async(EAsyncExecution::TaskGraph, [&]() -> const T& { return this->GetValue<T>(Context, Default); });
 }
diff --git a/Engine/Source/Runtime/Experimental/Dataflow/Core/Public/Dataflow/DataflowMathNodes.h b/Engine/Source/Runtime/Experimental/Dataflow/Core/Public/Dataflow/DataflowMathNodes.h
--- a/Engine/Source/Runtime/Experimental/Dataflow/Core/Public/Dataflow/DataflowMathNodes.h
+++ b/Engine/Source/Runtime/Experimental/Dataflow/Core/Public/Dataflow/DataflowMathNodes.h
@@ -32,12 +32,12 @@ struct FDataflowMathOneInputOperatorNode : public FDataflowNode
 
 public:
 	FDataflowMathOneInputOperatorNode() {};
-	FDataflowMathOneInputOperatorNode(const Dataflow::FNodeParameters& InParam, FGuid InGuid = FGuid::NewGuid());
-	virtual void Evaluate(Dataflow::FContext& Context, const FDataflowOutput* Out) const override;
+	FDataflowMathOneInputOperatorNode(const UE::Dataflow::FNodeParameters& InParam, FGuid InGuid = FGuid::NewGuid());
+	virtual void Evaluate(UE::Dataflow::FContext& Context, const FDataflowOutput* Out) const override;
 
 protected:
 	void RegisterInputsAndOutputs();
-	virtual double ComputeResult(Dataflow::FContext& Context, double InA) const { ensure(false); return 0.0; };
+	virtual double ComputeResult(UE::Dataflow::FContext& Context, double InA) const { ensure(false); return 0.0; };
 };
 
 /** Two inputs operators base class */
@@ -57,12 +57,12 @@ struct FDataflowMathTwoInputsOperatorNode : public FDataflowNode
 
 public:
 	FDataflowMathTwoInputsOperatorNode() {};
-	FDataflowMathTwoInputsOperatorNode(const Dataflow::FNodeParameters& InParam, FGuid InGuid = FGuid::NewGuid());
-	virtual void Evaluate(Dataflow::FContext& Context, const FDataflowOutput* Out) const override;
+	FDataflowMathTwoInputsOperatorNode(const UE::Dataflow::FNodeParameters& InParam, FGuid InGuid = FGuid::NewGuid());
+	virtual void Evaluate(UE::Dataflow::FContext& Context, const FDataflowOutput* Out) const override;
 
 protected:
 	void RegisterInputsAndOutputs();
-	virtual double ComputeResult(Dataflow::FContext& Context, double InA, double InB) const { ensure(false); return 0.0; };
+	virtual double ComputeResult(UE::Dataflow::FContext& Context, double InA, double InB) const { ensure(false); return 0.0; };
 };
 
 /** Addition (A + B) */
@@ -73,8 +73,8 @@ struct FDataflowMathAddNode : public FDataflowMathTwoInputsOperatorNode
 	DATAFLOW_NODE_DEFINE_INTERNAL(FDataflowMathAddNode, "Add", DATAFLOW_MATH_NODES_CATEGORY, "")
 
 public:
-	FDataflowMathAddNode(const Dataflow::FNodeParameters& InParam, FGuid InGuid = FGuid::NewGuid());
-	virtual double ComputeResult(Dataflow::FContext& Context, double InA, double InB) const override;
+	FDataflowMathAddNode(const UE::Dataflow::FNodeParameters& InParam, FGuid InGuid = FGuid::NewGuid());
+	virtual double ComputeResult(UE::Dataflow::FContext& Context, double InA, double InB) const override;
 };
 
 /** Subtraction (A - B) */
@@ -85,8 +85,8 @@ struct FDataflowMathSubtractNode : public FDataflowMathTwoInputsOperatorNode
 	DATAFLOW_NODE_DEFINE_INTERNAL(FDataflowMathSubtractNode, "Subtract", DATAFLOW_MATH_NODES_CATEGORY, "")
 
 public:
-	FDataflowMathSubtractNode(const Dataflow::FNodeParameters& InParam, FGuid InGuid = FGuid::NewGuid());
-	virtual double ComputeResult(Dataflow::FContext& Context, double InA, double InB) const override;
+	FDataflowMathSubtractNode(const UE::Dataflow::FNodeParameters& InParam, FGuid InGuid = FGuid::NewGuid());
+	virtual double ComputeResult(UE::Dataflow::FContext& Context, double InA, double InB) const override;
 };
 
 /** Multiplication (A * B) */
@@ -97,8 +97,8 @@ struct FDataflowMathMultiplyNode : public FDataflowMathTwoInputsOperatorNode
 	DATAFLOW_NODE_DEFINE_INTERNAL(FDataflowMathMultiplyNode, "Multiply", DATAFLOW_MATH_NODES_CATEGORY, "")
 
 public:
-	FDataflowMathMultiplyNode(const Dataflow::FNodeParameters& InParam, FGuid InGuid = FGuid::NewGuid());
-	virtual double ComputeResult(Dataflow::FContext& Context, double InA, double InB) const override;
+	FDataflowMathMultiplyNode(const UE::Dataflow::FNodeParameters& InParam, FGuid InGuid = FGuid::NewGuid());
+	virtual double ComputeResult(UE::Dataflow::FContext& Context, double InA, double InB) const override;
 };
 
 /**
@@ -115,8 +115,8 @@ struct FDataflowMathDivideNode : public FDataflowMathTwoInputsOperatorNode
 	FDataflowNumericTypes Fallback;
 
 public:
-	FDataflowMathDivideNode(const Dataflow::FNodeParameters& InParam, FGuid InGuid = FGuid::NewGuid());
-	virtual double ComputeResult(Dataflow::FContext& Context, double InA, double InB) const override;
+	FDataflowMathDivideNode(const UE::Dataflow::FNodeParameters& InParam, FGuid InGuid = FGuid::NewGuid());
+	virtual double ComputeResult(UE::Dataflow::FContext& Context, double InA, double InB) const override;
 };
 
 /** Minimum ( Min(A, B) ) */
@@ -127,8 +127,8 @@ struct FDataflowMathMinimumNode : public FDataflowMathTwoInputsOperatorNode
 	DATAFLOW_NODE_DEFINE_INTERNAL(FDataflowMathMinimumNode, "Minimum", DATAFLOW_MATH_NODES_CATEGORY, "")
 
 public:
-	FDataflowMathMinimumNode(const Dataflow::FNodeParameters& InParam, FGuid InGuid = FGuid::NewGuid());
-	virtual double ComputeResult(Dataflow::FContext& Context, double InA, double InB) const override;
+	FDataflowMathMinimumNode(const UE::Dataflow::FNodeParameters& InParam, FGuid InGuid = FGuid::NewGuid());
+	virtual double ComputeResult(UE::Dataflow::FContext& Context, double InA, double InB) const override;
 };
 
 /** Maximum ( Max(A, B) ) */
@@ -139,8 +139,8 @@ struct FDataflowMathMaximumNode : public FDataflowMathTwoInputsOperatorNode
 	DATAFLOW_NODE_DEFINE_INTERNAL(FDataflowMathMaximumNode, "Maximum", DATAFLOW_MATH_NODES_CATEGORY, "")
 
 public:
-	FDataflowMathMaximumNode(const Dataflow::FNodeParameters& InParam, FGuid InGuid = FGuid::NewGuid());
-	virtual double ComputeResult(Dataflow::FContext& Context, double InA, double InB) const override;
+	FDataflowMathMaximumNode(const UE::Dataflow::FNodeParameters& InParam, FGuid InGuid = FGuid::NewGuid());
+	virtual double ComputeResult(UE::Dataflow::FContext& Context, double InA, double InB) const override;
 };
 
 /** 
@@ -157,8 +157,8 @@ struct FDataflowMathReciprocalNode : public FDataflowMathOneInputOperatorNode
 	FDataflowNumericTypes Fallback;
 
 public:
-	FDataflowMathReciprocalNode(const Dataflow::FNodeParameters& InParam, FGuid InGuid = FGuid::NewGuid());
-	virtual double ComputeResult(Dataflow::FContext& Context, double InA) const override;
+	FDataflowMathReciprocalNode(const UE::Dataflow::FNodeParameters& InParam, FGuid InGuid = FGuid::NewGuid());
+	virtual double ComputeResult(UE::Dataflow::FContext& Context, double InA) const override;
 };
 
 /** Square ( A * A ) */
@@ -169,8 +169,8 @@ struct FDataflowMathSquareNode : public FDataflowMathOneInputOperatorNode
 	DATAFLOW_NODE_DEFINE_INTERNAL(FDataflowMathSquareNode, "Square", DATAFLOW_MATH_NODES_CATEGORY, "")
 
 public:
-	FDataflowMathSquareNode(const Dataflow::FNodeParameters& InParam, FGuid InGuid = FGuid::NewGuid());
-	virtual double ComputeResult(Dataflow::FContext& Context, double InA) const override;
+	FDataflowMathSquareNode(const UE::Dataflow::FNodeParameters& InParam, FGuid InGuid = FGuid::NewGuid());
+	virtual double ComputeResult(UE::Dataflow::FContext& Context, double InA) const override;
 };
 
 /** Cube ( A * A * A ) */
@@ -181,8 +181,8 @@ struct FDataflowMathCubeNode : public FDataflowMathOneInputOperatorNode
 	DATAFLOW_NODE_DEFINE_INTERNAL(FDataflowMathCubeNode, "Cube", DATAFLOW_MATH_NODES_CATEGORY, "")
 
 public:
-	FDataflowMathCubeNode(const Dataflow::FNodeParameters& InParam, FGuid InGuid = FGuid::NewGuid());
-	virtual double ComputeResult(Dataflow::FContext& Context, double InA) const override;
+	FDataflowMathCubeNode(const UE::Dataflow::FNodeParameters& InParam, FGuid InGuid = FGuid::NewGuid());
+	virtual double ComputeResult(UE::Dataflow::FContext& Context, double InA) const override;
 };
 
 /** Square Root ( sqrt(A) ) */
@@ -193,8 +193,8 @@ struct FDataflowMathSquareRootNode : public FDataflowMathOneInputOperatorNode
 	DATAFLOW_NODE_DEFINE_INTERNAL(FDataflowMathSquareRootNode, "SquareRoot", DATAFLOW_MATH_NODES_CATEGORY, "")
 
 public:
-	FDataflowMathSquareRootNode(const Dataflow::FNodeParameters& InParam, FGuid InGuid = FGuid::NewGuid());
-	virtual double ComputeResult(Dataflow::FContext& Context, double InA) const override;
+	FDataflowMathSquareRootNode(const UE::Dataflow::FNodeParameters& InParam, FGuid InGuid = FGuid::NewGuid());
+	virtual double ComputeResult(UE::Dataflow::FContext& Context, double InA) const override;
 };
 
 /** 
@@ -211,8 +211,8 @@ struct FDataflowMathInverseSquareRootNode : public FDataflowMathOneInputOperator
 	FDataflowNumericTypes Fallback;
 
 public:
-	FDataflowMathInverseSquareRootNode(const Dataflow::FNodeParameters& InParam, FGuid InGuid = FGuid::NewGuid());
-	virtual double ComputeResult(Dataflow::FContext& Context, double InA) const override;
+	FDataflowMathInverseSquareRootNode(const UE::Dataflow::FNodeParameters& InParam, FGuid InGuid = FGuid::NewGuid());
+	virtual double ComputeResult(UE::Dataflow::FContext& Context, double InA) const override;
 };
 
 /** Negate ( -A ) */
@@ -226,8 +226,8 @@ struct FDataflowMathNegateNode : public FDataflowMathOneInputOperatorNode
 	FDataflowNumericTypes Fallback;
 
 public:
-	FDataflowMathNegateNode(const Dataflow::FNodeParameters& InParam, FGuid InGuid = FGuid::NewGuid());
-	virtual double ComputeResult(Dataflow::FContext& Context, double InA) const override;
+	FDataflowMathNegateNode(const UE::Dataflow::FNodeParameters& InParam, FGuid InGuid = FGuid::NewGuid());
+	virtual double ComputeResult(UE::Dataflow::FContext& Context, double InA) const override;
 };
 
 /** Absolute value ( |A| ) */
@@ -241,8 +241,8 @@ struct FDataflowMathAbsNode : public FDataflowMathOneInputOperatorNode
 	FDataflowNumericTypes Fallback;
 
 public:
-	FDataflowMathAbsNode(const Dataflow::FNodeParameters& InParam, FGuid InGuid = FGuid::NewGuid());
-	virtual double ComputeResult(Dataflow::FContext& Context, double InA) const override;
+	FDataflowMathAbsNode(const UE::Dataflow::FNodeParameters& InParam, FGuid InGuid = FGuid::NewGuid());
+	virtual double ComputeResult(UE::Dataflow::FContext& Context, double InA) const override;
 };
 
 /** Floor ( 1.4 => 1.0 | 1.9 => 1.0 | -5.3 => -6.0 ) */
@@ -253,8 +253,8 @@ struct FDataflowMathFloorNode : public FDataflowMathOneInputOperatorNode
 	DATAFLOW_NODE_DEFINE_INTERNAL(FDataflowMathFloorNode, "Floor", DATAFLOW_MATH_NODES_CATEGORY, "")
 
 public:
-	FDataflowMathFloorNode(const Dataflow::FNodeParameters& InParam, FGuid InGuid = FGuid::NewGuid());
-	virtual double ComputeResult(Dataflow::FContext& Context, double InA) const override;
+	FDataflowMathFloorNode(const UE::Dataflow::FNodeParameters& InParam, FGuid InGuid = FGuid::NewGuid());
+	virtual double ComputeResult(UE::Dataflow::FContext& Context, double InA) const override;
 };
 
 /** Ceil ( 1.4 => 2.0 | 1.9 => 2.0 | -5.3 => -5.0) */
@@ -265,8 +265,8 @@ struct FDataflowMathCeilNode : public FDataflowMathOneInputOperatorNode
 	DATAFLOW_NODE_DEFINE_INTERNAL(FDataflowMathCeilNode, "Ceil", DATAFLOW_MATH_NODES_CATEGORY, "")
 
 public:
-	FDataflowMathCeilNode(const Dataflow::FNodeParameters& InParam, FGuid InGuid = FGuid::NewGuid());
-	virtual double ComputeResult(Dataflow::FContext& Context, double InA) const override;
+	FDataflowMathCeilNode(const UE::Dataflow::FNodeParameters& InParam, FGuid InGuid = FGuid::NewGuid());
+	virtual double ComputeResult(UE::Dataflow::FContext& Context, double InA) const override;
 };
 
 /** Round ( 1.4 => 1.0 | 1.9 => 2.0 | -5.3 => -5.0) */
@@ -277,8 +277,8 @@ struct FDataflowMathRoundNode : public FDataflowMathOneInputOperatorNode
 	DATAFLOW_NODE_DEFINE_INTERNAL(FDataflowMathRoundNode, "Round", DATAFLOW_MATH_NODES_CATEGORY, "")
 
 public:
-	FDataflowMathRoundNode(const Dataflow::FNodeParameters& InParam, FGuid InGuid = FGuid::NewGuid());
-	virtual double ComputeResult(Dataflow::FContext& Context, double InA) const override;
+	FDataflowMathRoundNode(const UE::Dataflow::FNodeParameters& InParam, FGuid InGuid = FGuid::NewGuid());
+	virtual double ComputeResult(UE::Dataflow::FContext& Context, double InA) const override;
 };
 
 /** Trunc ( 1.4 => 1.0 | 1.9 => 1.0 | -5.3 => -5.0) */
@@ -289,8 +289,8 @@ struct FDataflowMathTruncNode : public FDataflowMathOneInputOperatorNode
 	DATAFLOW_NODE_DEFINE_INTERNAL(FDataflowMathTruncNode, "Trunc", DATAFLOW_MATH_NODES_CATEGORY, "")
 
 public:
-	FDataflowMathTruncNode(const Dataflow::FNodeParameters& InParam, FGuid InGuid = FGuid::NewGuid());
-	virtual double ComputeResult(Dataflow::FContext& Context, double InA) const override;
+	FDataflowMathTruncNode(const UE::Dataflow::FNodeParameters& InParam, FGuid InGuid = FGuid::NewGuid());
+	virtual double ComputeResult(UE::Dataflow::FContext& Context, double InA) const override;
 };
 
 /** Frac ( 1.4 => 0.4 | 1.9 => 0.9 | -5.3 => 0.3 ) */
@@ -301,8 +301,8 @@ struct FDataflowMathFracNode : public FDataflowMathOneInputOperatorNode
 	DATAFLOW_NODE_DEFINE_INTERNAL(FDataflowMathFracNode, "Frac", DATAFLOW_MATH_NODES_CATEGORY, "")
 
 public:
-	FDataflowMathFracNode(const Dataflow::FNodeParameters& InParam, FGuid InGuid = FGuid::NewGuid());
-	virtual double ComputeResult(Dataflow::FContext& Context, double InA) const override;
+	FDataflowMathFracNode(const UE::Dataflow::FNodeParameters& InParam, FGuid InGuid = FGuid::NewGuid());
+	virtual double ComputeResult(UE::Dataflow::FContext& Context, double InA) const override;
 };
 
 /** power ( A ^ B) */
@@ -313,8 +313,8 @@ struct FDataflowMathPowNode : public FDataflowMathTwoInputsOperatorNode
 	DATAFLOW_NODE_DEFINE_INTERNAL(FDataflowMathPowNode, "Pow", DATAFLOW_MATH_NODES_CATEGORY, "")
 
 public:
-	FDataflowMathPowNode(const Dataflow::FNodeParameters& InParam, FGuid InGuid = FGuid::NewGuid());
-	virtual double ComputeResult(Dataflow::FContext& Context, double InA, double InB) const override;
+	FDataflowMathPowNode(const UE::Dataflow::FNodeParameters& InParam, FGuid InGuid = FGuid::NewGuid());
+	virtual double ComputeResult(UE::Dataflow::FContext& Context, double InA, double InB) const override;
 };
 
 /**
@@ -331,8 +331,8 @@ struct FDataflowMathLogXNode : public FDataflowMathOneInputOperatorNode
 	FDataflowNumericTypes Base;
 
 public:
-	FDataflowMathLogXNode(const Dataflow::FNodeParameters& InParam, FGuid InGuid = FGuid::NewGuid());
-	virtual double ComputeResult(Dataflow::FContext& Context, double InA) const override;
+	FDataflowMathLogXNode(const UE::Dataflow::FNodeParameters& InParam, FGuid InGuid = FGuid::NewGuid());
+	virtual double ComputeResult(UE::Dataflow::FContext& Context, double InA) const override;
 };
 
 /** Natural log ( Log(A) ) */
@@ -343,8 +343,8 @@ struct FDataflowMathLogNode : public FDataflowMathOneInputOperatorNode
 	DATAFLOW_NODE_DEFINE_INTERNAL(FDataflowMathLogNode, "Log", DATAFLOW_MATH_NODES_CATEGORY, "")
 
 public:
-	FDataflowMathLogNode(const Dataflow::FNodeParameters& InParam, FGuid InGuid = FGuid::NewGuid());
-	virtual double ComputeResult(Dataflow::FContext& Context, double InA) const override;
+	FDataflowMathLogNode(const UE::Dataflow::FNodeParameters& InParam, FGuid InGuid = FGuid::NewGuid());
+	virtual double ComputeResult(UE::Dataflow::FContext& Context, double InA) const override;
 };
 
 /** Exponential ( Exp(A) ) */
@@ -355,8 +355,8 @@ struct FDataflowMathExpNode : public FDataflowMathOneInputOperatorNode
 	DATAFLOW_NODE_DEFINE_INTERNAL(FDataflowMathExpNode, "Exp", DATAFLOW_MATH_NODES_CATEGORY, "")
 
 public:
-	FDataflowMathExpNode(const Dataflow::FNodeParameters& InParam, FGuid InGuid = FGuid::NewGuid());
-	virtual double ComputeResult(Dataflow::FContext& Context, double InA) const override;
+	FDataflowMathExpNode(const UE::Dataflow::FNodeParameters& InParam, FGuid InGuid = FGuid::NewGuid());
+	virtual double ComputeResult(UE::Dataflow::FContext& Context, double InA) const override;
 };
 
 /** return -1, 0, +1 whether the input is respectively negative, zero or positive ( Sign(A) ) */
@@ -367,8 +367,8 @@ struct FDataflowMathSignNode : public FDataflowMathOneInputOperatorNode
 	DATAFLOW_NODE_DEFINE_INTERNAL(FDataflowMathSignNode, "Sign", DATAFLOW_MATH_NODES_CATEGORY, "")
 
 public:
-	FDataflowMathSignNode(const Dataflow::FNodeParameters& InParam, FGuid InGuid = FGuid::NewGuid());
-	virtual double ComputeResult(Dataflow::FContext& Context, double InA) const override;
+	FDataflowMathSignNode(const UE::Dataflow::FNodeParameters& InParam, FGuid InGuid = FGuid::NewGuid());
+	virtual double ComputeResult(UE::Dataflow::FContext& Context, double InA) const override;
 };
 
 /** One minus (1 - A) */
@@ -379,8 +379,8 @@ struct FDataflowMathOneMinusNode : public FDataflowMathOneInputOperatorNode
 	DATAFLOW_NODE_DEFINE_INTERNAL(FDataflowMathOneMinusNode, "OneMinus", DATAFLOW_MATH_NODES_CATEGORY, "")
 
 public:
-	FDataflowMathOneMinusNode(const Dataflow::FNodeParameters& InParam, FGuid InGuid = FGuid::NewGuid());
-	virtual double ComputeResult(Dataflow::FContext& Context, double InA) const override;
+	FDataflowMathOneMinusNode(const UE::Dataflow::FNodeParameters& InParam, FGuid InGuid = FGuid::NewGuid());
+	virtual double ComputeResult(UE::Dataflow::FContext& Context, double InA) const override;
 };
 
 UENUM(BlueprintType)
@@ -419,8 +419,8 @@ struct FDataflowMathConstantNode : public FDataflowNode
 	FDataflowNumericTypes Result;
 
 public:
-	FDataflowMathConstantNode(const Dataflow::FNodeParameters& InParam, FGuid InGuid = FGuid::NewGuid());
-	virtual void Evaluate(Dataflow::FContext& Context, const FDataflowOutput* Out) const override;
+	FDataflowMathConstantNode(const UE::Dataflow::FNodeParameters& InParam, FGuid InGuid = FGuid::NewGuid());
+	virtual void Evaluate(UE::Dataflow::FContext& Context, const FDataflowOutput* Out) const override;
 
 private:
 	double GetConstant() const;
@@ -442,8 +442,8 @@ struct FDataflowMathSinNode : public FDataflowMathOneInputOperatorNode
 	DATAFLOW_NODE_DEFINE_INTERNAL(FDataflowMathSinNode, "Sin", DATAFLOW_MATH_TRIG_NODES_CATEGORY, "")
 
 public:
-	FDataflowMathSinNode(const Dataflow::FNodeParameters& InParam, FGuid InGuid = FGuid::NewGuid());
-	virtual double ComputeResult(Dataflow::FContext& Context, double InA) const override;
+	FDataflowMathSinNode(const UE::Dataflow::FNodeParameters& InParam, FGuid InGuid = FGuid::NewGuid());
+	virtual double ComputeResult(UE::Dataflow::FContext& Context, double InA) const override;
 };
 
 /** Cos(A) with A in radians  */
@@ -454,8 +454,8 @@ struct FDataflowMathCosNode : public FDataflowMathOneInputOperatorNode
 	DATAFLOW_NODE_DEFINE_INTERNAL(FDataflowMathCosNode, "Cos", DATAFLOW_MATH_TRIG_NODES_CATEGORY, "")
 
 public:
-	FDataflowMathCosNode(const Dataflow::FNodeParameters& InParam, FGuid InGuid = FGuid::NewGuid());
-	virtual double ComputeResult(Dataflow::FContext& Context, double InA) const override;
+	FDataflowMathCosNode(const UE::Dataflow::FNodeParameters& InParam, FGuid InGuid = FGuid::NewGuid());
+	virtual double ComputeResult(UE::Dataflow::FContext& Context, double InA) const override;
 };
 
 /** Tan(A) with A in radians  */
@@ -466,8 +466,8 @@ struct FDataflowMathTanNode : public FDataflowMathOneInputOperatorNode
 	DATAFLOW_NODE_DEFINE_INTERNAL(FDataflowMathTanNode, "Tan", DATAFLOW_MATH_TRIG_NODES_CATEGORY, "")
 
 public:
-	FDataflowMathTanNode(const Dataflow::FNodeParameters& InParam, FGuid InGuid = FGuid::NewGuid());
-	virtual double ComputeResult(Dataflow::FContext& Context, double InA) const override;
+	FDataflowMathTanNode(const UE::Dataflow::FNodeParameters& InParam, FGuid InGuid = FGuid::NewGuid());
+	virtual double ComputeResult(UE::Dataflow::FContext& Context, double InA) const override;
 };
 
 /** ArcSin(A) returns a value in radians  */
@@ -478,8 +478,8 @@ struct FDataflowMathArcSinNode : public FDataflowMathOneInputOperatorNode
 	DATAFLOW_NODE_DEFINE_INTERNAL(FDataflowMathArcSinNode, "ArcSin", DATAFLOW_MATH_TRIG_NODES_CATEGORY, "")
 
 public:
-	FDataflowMathArcSinNode(const Dataflow::FNodeParameters& InParam, FGuid InGuid = FGuid::NewGuid());
-	virtual double ComputeResult(Dataflow::FContext& Context, double InA) const override;
+	FDataflowMathArcSinNode(const UE::Dataflow::FNodeParameters& InParam, FGuid InGuid = FGuid::NewGuid());
+	virtual double ComputeResult(UE::Dataflow::FContext& Context, double InA) const override;
 };
 
 /** ArcCos(A) returns a value in radians  */
@@ -490,8 +490,8 @@ struct FDataflowMathArcCosNode : public FDataflowMathOneInputOperatorNode
 	DATAFLOW_NODE_DEFINE_INTERNAL(FDataflowMathArcCosNode, "ArcCos", DATAFLOW_MATH_TRIG_NODES_CATEGORY, "")
 
 public:
-	FDataflowMathArcCosNode(const Dataflow::FNodeParameters& InParam, FGuid InGuid = FGuid::NewGuid());
-	virtual double ComputeResult(Dataflow::FContext& Context, double InA) const override;
+	FDataflowMathArcCosNode(const UE::Dataflow::FNodeParameters& InParam, FGuid InGuid = FGuid::NewGuid());
+	virtual double ComputeResult(UE::Dataflow::FContext& Context, double InA) const override;
 };
 
 /** ArcTan(A) returns a value in radians  */
@@ -502,8 +502,8 @@ struct FDataflowMathArcTanNode : public FDataflowMathOneInputOperatorNode
 	DATAFLOW_NODE_DEFINE_INTERNAL(FDataflowMathArcTanNode, "ArcTan", DATAFLOW_MATH_TRIG_NODES_CATEGORY, "")
 
 public:
-	FDataflowMathArcTanNode(const Dataflow::FNodeParameters& InParam, FGuid InGuid = FGuid::NewGuid());
-	virtual double ComputeResult(Dataflow::FContext& Context, double InA) const override;
+	FDataflowMathArcTanNode(const UE::Dataflow::FNodeParameters& InParam, FGuid InGuid = FGuid::NewGuid());
+	virtual double ComputeResult(UE::Dataflow::FContext& Context, double InA) const override;
 };
 
 /** ArcTan2(A, B) returns a value in radians  */
@@ -514,8 +514,8 @@ struct FDataflowMathArcTan2Node : public FDataflowMathTwoInputsOperatorNode
 	DATAFLOW_NODE_DEFINE_INTERNAL(FDataflowMathArcTan2Node, "ArcTan2", DATAFLOW_MATH_TRIG_NODES_CATEGORY, "")
 
 public:
-	FDataflowMathArcTan2Node(const Dataflow::FNodeParameters& InParam, FGuid InGuid = FGuid::NewGuid());
-	virtual double ComputeResult(Dataflow::FContext& Context, double InA, double InB) const override;
+	FDataflowMathArcTan2Node(const UE::Dataflow::FNodeParameters& InParam, FGuid InGuid = FGuid::NewGuid());
+	virtual double ComputeResult(UE::Dataflow::FContext& Context, double InA, double InB) const override;
 };
 
 
@@ -527,8 +527,8 @@ struct FDataflowMathDegToRadNode : public FDataflowMathOneInputOperatorNode
 	DATAFLOW_NODE_DEFINE_INTERNAL(FDataflowMathDegToRadNode, "DegToRad", DATAFLOW_MATH_TRIG_NODES_CATEGORY, "")
 
 public:
-	FDataflowMathDegToRadNode(const Dataflow::FNodeParameters& InParam, FGuid InGuid = FGuid::NewGuid());
-	virtual double ComputeResult(Dataflow::FContext& Context, double InA) const override;
+	FDataflowMathDegToRadNode(const UE::Dataflow::FNodeParameters& InParam, FGuid InGuid = FGuid::NewGuid());
+	virtual double ComputeResult(UE::Dataflow::FContext& Context, double InA) const override;
 };
 
 /** RadToDeg(A) convert radians to degrees */
@@ -539,8 +539,8 @@ struct FDataflowMathRadToDegNode : public FDataflowMathOneInputOperatorNode
 	DATAFLOW_NODE_DEFINE_INTERNAL(FDataflowMathRadToDegNode, "RadToDeg", DATAFLOW_MATH_TRIG_NODES_CATEGORY, "")
 
 public:
-	FDataflowMathRadToDegNode(const Dataflow::FNodeParameters& InParam, FGuid InGuid = FGuid::NewGuid());
-	virtual double ComputeResult(Dataflow::FContext& Context, double InA) const override;
+	FDataflowMathRadToDegNode(const UE::Dataflow::FNodeParameters& InParam, FGuid InGuid = FGuid::NewGuid());
+	virtual double ComputeResult(UE::Dataflow::FContext& Context, double InA) const override;
 };
 
 //--------------------------------------------------------------------------
@@ -582,7 +582,7 @@ Nodes Left to be converted :
 
 */
 
-namespace Dataflow
+namespace UE::Dataflow
 {
 	void RegisterDataflowMathNodes();
 }
