Change Dataflow namespace to UE::Dataflow to be ok with coding standard

diff --git a/Engine/Source/Runtime/Experimental/Dataflow/Core/Public/Dataflow/DataflowNode.h b/Engine/Source/Runtime/Experimental/Dataflow/Core/Public/Dataflow/DataflowNode.h
--- a/Engine/Source/Runtime/Experimental/Dataflow/Core/Public/Dataflow/DataflowNode.h
+++ b/Engine/Source/Runtime/Experimental/Dataflow/Core/Public/Dataflow/DataflowNode.h
@@ -15,7 +15,7 @@
 
 class UScriptStruct;
 
-namespace Dataflow {
+namespace UE::Dataflow {
 	struct FNodeParameters {
 		FName Name;
 		UObject* OwningObject = nullptr;
@@ -23,8 +23,6 @@ namespace Dataflow {
 	class FGraph;
 }
 
-
-
 /**
 * FNode
 *		Base class for node based evaluation within the Dataflow graph. 
@@ -40,7 +38,7 @@ struct FDataflowNode
 PRAGMA_ENABLE_DEPRECATION_WARNINGS
 	GENERATED_USTRUCT_BODY()
 
-	friend class Dataflow::FGraph;
+	friend class UE::Dataflow::FGraph;
 	friend struct FDataflowConnection;
 
 	UE_DEPRECATED(5.5, "Will be made private in 5.7, use Get/Set methods instead.")
@@ -48,7 +46,7 @@ PRAGMA_ENABLE_DEPRECATION_WARNINGS
 	UE_DEPRECATED(5.5, "Will be made private in 5.7, use Get/Set methods instead.")
 	FName Name;
 	UE_DEPRECATED(5.5, "Will be made private in 5.7, use Get/Set methods instead.")
-	Dataflow::FTimestamp LastModifiedTimestamp;
+	UE::Dataflow::FTimestamp LastModifiedTimestamp;
 
 	UE_DEPRECATED(5.5, "Inputs type has changed and has been made private")
 	TMap< int, FDataflowInput* > Inputs;
@@ -61,14 +59,14 @@ PRAGMA_ENABLE_DEPRECATION_WARNINGS
 	FDataflowNode()
 		: Guid(FGuid())
 		, Name("Invalid")
-		, LastModifiedTimestamp(Dataflow::FTimestamp::Invalid)
+		, LastModifiedTimestamp(UE::Dataflow::FTimestamp::Invalid)
 	{
 	}
 
-	FDataflowNode(const Dataflow::FNodeParameters& Param, FGuid InGuid = FGuid::NewGuid())
+	FDataflowNode(const UE::Dataflow::FNodeParameters& Param, FGuid InGuid = FGuid::NewGuid())
 		: Guid(InGuid)
 		, Name(Param.Name)
-		, LastModifiedTimestamp(Dataflow::FTimestamp::Invalid)
+		, LastModifiedTimestamp(UE::Dataflow::FTimestamp::Invalid)
 	{
 	}
 
@@ -81,7 +79,7 @@ PRAGMA_ENABLE_DEPRECATION_WARNINGS
 	FGuid GetGuid() const { return Guid; }
 	FName GetName() const { return Name; }
 	void SetName(FName InName) { Name = InName; }
-	Dataflow::FTimestamp GetTimestamp() const { return LastModifiedTimestamp;  }
+	UE::Dataflow::FTimestamp GetTimestamp() const { return LastModifiedTimestamp;  }
 	PRAGMA_ENABLE_DEPRECATION_WARNINGS
 	DATAFLOWCORE_API uint32 GetValueHash();
 
@@ -91,10 +89,10 @@ PRAGMA_ENABLE_DEPRECATION_WARNINGS
 	virtual FName GetCategory() const { return ""; }
 	virtual FString GetTags() const { return ""; }
 	DATAFLOWCORE_API virtual FString GetToolTip() const;
-	DATAFLOWCORE_API FString GetPinToolTip(const FName& PropertyName, const Dataflow::FPin::EDirection Direction = Dataflow::FPin::EDirection::NONE) const;
-	DATAFLOWCORE_API FText GetPinDisplayName(const FName& PropertyName, const Dataflow::FPin::EDirection Direction = Dataflow::FPin::EDirection::NONE) const;
-	DATAFLOWCORE_API TArray<FString> GetPinMetaData(const FName& PropertyName, const Dataflow::FPin::EDirection Direction = Dataflow::FPin::EDirection::NONE) const;
-	virtual TArray<Dataflow::FRenderingParameter> GetRenderParameters() const { return GetRenderParametersImpl(); }
+	DATAFLOWCORE_API FString GetPinToolTip(const FName& PropertyName, const UE::Dataflow::FPin::EDirection Direction = UE::Dataflow::FPin::EDirection::NONE) const;
+	DATAFLOWCORE_API FText GetPinDisplayName(const FName& PropertyName, const UE::Dataflow::FPin::EDirection Direction = UE::Dataflow::FPin::EDirection::NONE) const;
+	DATAFLOWCORE_API TArray<FString> GetPinMetaData(const FName& PropertyName, const UE::Dataflow::FPin::EDirection Direction = UE::Dataflow::FPin::EDirection::NONE) const;
+	virtual TArray<UE::Dataflow::FRenderingParameter> GetRenderParameters() const { return GetRenderParametersImpl(); }
 	// Copy node property values from another node
 	UE_DEPRECATED(5.4, "FDataflowNode::CopyNodeProperties is deprecated.")
 	DATAFLOWCORE_API void CopyNodeProperties(const TSharedPtr<FDataflowNode> CopyFromDataflowNode);
@@ -108,42 +106,42 @@ PRAGMA_ENABLE_DEPRECATION_WARNINGS
 	// Connections
 	//
 
-	DATAFLOWCORE_API TArray<Dataflow::FPin> GetPins() const;
+	DATAFLOWCORE_API TArray<UE::Dataflow::FPin> GetPins() const;
 
 	UE_DEPRECATED(5.5, "Use AddPins method instead")
-	virtual Dataflow::FPin AddPin() { return Dataflow::FPin::InvalidPin; }
+	virtual UE::Dataflow::FPin AddPin() { return UE::Dataflow::FPin::InvalidPin; }
 	/** Override this function to add the AddOptionPin functionality to the node's context menu. */
-	virtual TArray<Dataflow::FPin> AddPins()
+	virtual TArray<UE::Dataflow::FPin> AddPins()
 	{
 PRAGMA_DISABLE_DEPRECATION_WARNINGS
-		const Dataflow::FPin DeprecatedAddPin = AddPin();
+		const UE::Dataflow::FPin DeprecatedAddPin = AddPin();
 PRAGMA_ENABLE_DEPRECATION_WARNINGS
-		if(DeprecatedAddPin == Dataflow::FPin::InvalidPin)
+		if(DeprecatedAddPin == UE::Dataflow::FPin::InvalidPin)
 		{
-			return TArray<Dataflow::FPin>();
+			return TArray<UE::Dataflow::FPin>();
 		}
 		return { DeprecatedAddPin };
 	}
 	/** Override this function to add the AddOptionPin functionality to the node's context menu. */
 	virtual bool CanAddPin() const { return false; }
 	UE_DEPRECATED(5.5, "Use GetPinsToRemove method instead")
-	virtual Dataflow::FPin GetPinToRemove() const { return Dataflow::FPin::InvalidPin; }
+	virtual UE::Dataflow::FPin GetPinToRemove() const { return UE::Dataflow::FPin::InvalidPin; }
 	UE_DEPRECATED(5.4, "Use GetPinsToRemove and OnPinRemoved instead.")
-	virtual Dataflow::FPin RemovePin()
+	virtual UE::Dataflow::FPin RemovePin()
 	{
 		PRAGMA_DISABLE_DEPRECATION_WARNINGS
 		return GetPinToRemove();
 		PRAGMA_ENABLE_DEPRECATION_WARNINGS
 	}
 	/** Override this function to add the RemoveOptionPin functionality to the node's context menu. OnPinRemoved will be called in this order.*/
-	virtual TArray<Dataflow::FPin> GetPinsToRemove() const
+	virtual TArray<UE::Dataflow::FPin> GetPinsToRemove() const
 	{
 PRAGMA_DISABLE_DEPRECATION_WARNINGS
-		const Dataflow::FPin DeprecatedRemovePin = GetPinToRemove();
+		const UE::Dataflow::FPin DeprecatedRemovePin = GetPinToRemove();
 PRAGMA_ENABLE_DEPRECATION_WARNINGS
-		if (DeprecatedRemovePin == Dataflow::FPin::InvalidPin)
+		if (DeprecatedRemovePin == UE::Dataflow::FPin::InvalidPin)
 		{
-			return TArray<Dataflow::FPin>();
+			return TArray<UE::Dataflow::FPin>();
 		}
 		return { DeprecatedRemovePin };
 	}
@@ -151,7 +149,7 @@ PRAGMA_ENABLE_DEPRECATION_WARNINGS
 	 * Override this to update any bookkeeping when a pin is being removed.
 	 * This will be called before the pin is unregistered as an input.
 	 */
-	virtual void OnPinRemoved(const Dataflow::FPin& Pin) {}
+	virtual void OnPinRemoved(const UE::Dataflow::FPin& Pin) {}
 	/** Override this function to add the RemoveOptionPin functionality to the node's context menu. */
 	virtual bool CanRemovePin() const { return false; }
 
@@ -167,13 +165,13 @@ PRAGMA_ENABLE_DEPRECATION_WARNINGS
 	DATAFLOWCORE_API bool HasHiddenInputs() const;
 
 	DATAFLOWCORE_API FDataflowInput* FindInput(FName Name);
-	DATAFLOWCORE_API FDataflowInput* FindInput(const Dataflow::FConnectionKey& Key);
+	DATAFLOWCORE_API FDataflowInput* FindInput(const UE::Dataflow::FConnectionKey& Key);
 	/**This version can find array inputs if only the Reference is supplied by searching through all inputs */
-	DATAFLOWCORE_API FDataflowInput* FindInput(const Dataflow::FConnectionReference& Reference);
+	DATAFLOWCORE_API FDataflowInput* FindInput(const UE::Dataflow::FConnectionReference& Reference);
 	DATAFLOWCORE_API const FDataflowInput* FindInput(FName Name) const;
-	DATAFLOWCORE_API const FDataflowInput* FindInput(const Dataflow::FConnectionKey& Key) const;
+	DATAFLOWCORE_API const FDataflowInput* FindInput(const UE::Dataflow::FConnectionKey& Key) const;
 	/**This version can find array inputs if only the Reference is supplied by searching through all inputs */
-	DATAFLOWCORE_API const FDataflowInput* FindInput(const Dataflow::FConnectionReference& Reference) const;
+	DATAFLOWCORE_API const FDataflowInput* FindInput(const UE::Dataflow::FConnectionReference& Reference) const;
 	DATAFLOWCORE_API const FDataflowInput* FindInput(const FGuid& InGuid) const;
 
 	DATAFLOWCORE_API virtual void AddOutput(FDataflowOutput* InPtr);
@@ -185,12 +183,12 @@ PRAGMA_ENABLE_DEPRECATION_WARNINGS
 
 	DATAFLOWCORE_API FDataflowOutput* FindOutput(FName Name);
 	DATAFLOWCORE_API FDataflowOutput* FindOutput(uint32 GuidHash);
-	DATAFLOWCORE_API FDataflowOutput* FindOutput(const Dataflow::FConnectionKey& Key);
-	DATAFLOWCORE_API FDataflowOutput* FindOutput(const Dataflow::FConnectionReference& Reference);
+	DATAFLOWCORE_API FDataflowOutput* FindOutput(const UE::Dataflow::FConnectionKey& Key);
+	DATAFLOWCORE_API FDataflowOutput* FindOutput(const UE::Dataflow::FConnectionReference& Reference);
 	DATAFLOWCORE_API const FDataflowOutput* FindOutput(FName Name) const;
 	DATAFLOWCORE_API const FDataflowOutput* FindOutput(uint32 GuidHash) const;
-	DATAFLOWCORE_API const FDataflowOutput* FindOutput(const Dataflow::FConnectionKey& Key) const;
-	DATAFLOWCORE_API const FDataflowOutput* FindOutput(const Dataflow::FConnectionReference& Reference) const;
+	DATAFLOWCORE_API const FDataflowOutput* FindOutput(const UE::Dataflow::FConnectionKey& Key) const;
+	DATAFLOWCORE_API const FDataflowOutput* FindOutput(const UE::Dataflow::FConnectionReference& Reference) const;
 	DATAFLOWCORE_API const FDataflowOutput* FindOutput(const FGuid& InGuid) const;
 
 	/** Return a property's byte offset from the dataflow base node address using the full property name (must includes its parent struct property names). Does not work with array properties.*/
@@ -213,7 +211,7 @@ PRAGMA_ENABLE_DEPRECATION_WARNINGS
 	virtual FDataflowOutput* RedirectSerializedOutput(const FName& MissingOutputName) { return nullptr; }
 
 	/** Called by editor toolkits when the node is selected, or already selected and invalidated. */
-	virtual void OnSelected(Dataflow::FContext& Context) {}
+	virtual void OnSelected(UE::Dataflow::FContext& Context) {}
 	/** Called by editor toolkits when the node is deselected. */
 	virtual void OnDeselected() {}
 
@@ -229,7 +227,7 @@ PRAGMA_ENABLE_DEPRECATION_WARNINGS
 
 	/** Register the Input and Outputs after the creation in the factory. Use PropertyName to disambiguate a struct name from its first property. */
 	template<typename T>
-	FDataflowInput& RegisterInputConnection(const Dataflow::TConnectionReference<T>& Reference, const FName& PropertyName = NAME_None)
+	FDataflowInput& RegisterInputConnection(const UE::Dataflow::TConnectionReference<T>& Reference, const FName& PropertyName = NAME_None)
 	{
 		FDataflowInput& Input = RegisterInputConnectionInternal(Reference, PropertyName);
 		if constexpr (std::is_base_of_v<FDataflowAnyType, T>)
@@ -242,11 +240,11 @@ PRAGMA_ENABLE_DEPRECATION_WARNINGS
 	template <typename T>
 	FDataflowInput& RegisterInputConnection(const T* Reference, const FName& PropertyName = NAME_None)
 	{
-		return RegisterInputConnection(Dataflow::TConnectionReference<T>(Reference), PropertyName);
+		return RegisterInputConnection(UE::Dataflow::TConnectionReference<T>(Reference), PropertyName);
 	}
 
 	template <typename T>
-	FDataflowOutput& RegisterOutputConnection(const Dataflow::TConnectionReference<T>& Reference, const Dataflow::TConnectionReference<T>& Passthrough = Dataflow::TConnectionReference<T>(nullptr), const FName& PropertyName = NAME_None)
+	FDataflowOutput& RegisterOutputConnection(const UE::Dataflow::TConnectionReference<T>& Reference, const UE::Dataflow::TConnectionReference<T>& Passthrough = UE::Dataflow::TConnectionReference<T>(nullptr), const FName& PropertyName = NAME_None)
 	{
 		FDataflowOutput& Output = RegisterOutputConnectionInternal(Reference, PropertyName);
 		if constexpr (std::is_base_of_v<FDataflowAnyType, T>)
@@ -264,18 +262,18 @@ PRAGMA_ENABLE_DEPRECATION_WARNINGS
 	template <typename T>
 	FDataflowOutput& RegisterOutputConnection(const T* Reference, const T* Passthrough = nullptr, const FName& PropertyName = NAME_None)
 	{
-		return RegisterOutputConnection(Dataflow::TConnectionReference<T>(Reference), Dataflow::TConnectionReference<T>(Passthrough), PropertyName);
+		return RegisterOutputConnection(UE::Dataflow::TConnectionReference<T>(Reference), UE::Dataflow::TConnectionReference<T>(Passthrough), PropertyName);
 	}
 
 	template <typename T>
 	UE_DEPRECATED(5.5, "PassthroughName is no longer needed to register output connections")
 	FDataflowOutput& RegisterOutputConnection(const T* Reference, const T* Passthrough, const FName& PropertyName, const FName& PassthroughName)
 	{
-		return RegisterOutputConnection(Dataflow::TConnectionReference<T>(Reference), Dataflow::TConnectionReference<T>(Passthrough), PropertyName);
+		return RegisterOutputConnection(UE::Dataflow::TConnectionReference<T>(Reference), UE::Dataflow::TConnectionReference<T>(Passthrough), PropertyName);
 	}
 
 	template<typename T>
-	FDataflowInput& RegisterInputArrayConnection(const Dataflow::TConnectionReference<T>& Reference, const FName& ElementPropertyName = NAME_None,
+	FDataflowInput& RegisterInputArrayConnection(const UE::Dataflow::TConnectionReference<T>& Reference, const FName& ElementPropertyName = NAME_None,
 		const FName& ArrayPropertyName = NAME_None)
 	{
 		FDataflowInput& Input = RegisterInputArrayConnectionInternal(Reference, ElementPropertyName, ArrayPropertyName);
@@ -287,7 +285,7 @@ PRAGMA_ENABLE_DEPRECATION_WARNINGS
 	}
 
 	template<typename T>
-	FDataflowInput& FindOrRegisterInputArrayConnection(const Dataflow::TConnectionReference<T>& Reference, const FName& ElementPropertyName = NAME_None,
+	FDataflowInput& FindOrRegisterInputArrayConnection(const UE::Dataflow::TConnectionReference<T>& Reference, const FName& ElementPropertyName = NAME_None,
 		const FName& ArrayPropertyName = NAME_None)
 	{
 		if (FDataflowInput* const FoundInput = FindInput(Reference))
@@ -304,7 +302,7 @@ PRAGMA_ENABLE_DEPRECATION_WARNINGS
 	}
 
 	/** Unregister the input connection if one exists matching this property, and then invalidate the graph. */
-	void UnregisterInputConnection(const Dataflow::FConnectionReference& Reference)
+	void UnregisterInputConnection(const UE::Dataflow::FConnectionReference& Reference)
 	{
 		UnregisterInputConnection(GetKeyFromReference(Reference));
 	}
@@ -313,14 +311,14 @@ PRAGMA_ENABLE_DEPRECATION_WARNINGS
 	{
 		return UnregisterInputConnection(GetKeyFromReference(Reference));
 	}
-	DATAFLOWCORE_API void UnregisterInputConnection(const Dataflow::FConnectionKey& Key);
+	DATAFLOWCORE_API void UnregisterInputConnection(const UE::Dataflow::FConnectionKey& Key);
 	/** Unregister the connection if one exists matching this pin, then invalidate the graph. */
-	DATAFLOWCORE_API void UnregisterPinConnection(const Dataflow::FPin& Pin);
+	DATAFLOWCORE_API void UnregisterPinConnection(const UE::Dataflow::FPin& Pin);
 
 	//
 	// Evaluation
 	//
-	virtual void Evaluate(Dataflow::FContext& Context, const FDataflowOutput*) const { ensure(false); }
+	virtual void Evaluate(UE::Dataflow::FContext& Context, const FDataflowOutput*) const { ensure(false); }
 
 	/**
 	*   GetValue(...)
@@ -333,42 +331,42 @@ PRAGMA_ENABLE_DEPRECATION_WARNINGS
 	*						*Reference will be used as the default if the input is not connected. 
 	*/
 	template<class T, typename = std::enable_if_t<!std::is_base_of_v<FDataflowAnyType, T>>>
-	const T& GetValue(Dataflow::FContext& Context, const Dataflow::TConnectionReference<T>& Reference) const
+	const T& GetValue(UE::Dataflow::FContext& Context, const UE::Dataflow::TConnectionReference<T>& Reference) const
 	{
 		checkSlow(FindInput(Reference));
 		return FindInput(Reference)->template GetValue<T>(Context, *static_cast<const T*>(Reference.Reference));
 	}
 
 	template<typename TAnyType, typename = std::enable_if_t<std::is_base_of_v<FDataflowAnyType, TAnyType>>>
-	typename TAnyType::FStorageType GetValue(Dataflow::FContext& Context, const Dataflow::TConnectionReference<TAnyType>& Reference) const
+	typename TAnyType::FStorageType GetValue(UE::Dataflow::FContext& Context, const UE::Dataflow::TConnectionReference<TAnyType>& Reference) const
 	{
 		checkSlow(Reference.Reference && FindInput(Reference));
 		return FindInput(Reference)->template GetValueFromAnyType<TAnyType>(Context, static_cast<const TAnyType*>(Reference.Reference)->Value);
 	}
 
 	template<class T, typename = std::enable_if_t<!std::is_base_of_v<FDataflowAnyType, T>>>
-	const T& GetValue(Dataflow::FContext& Context, const T* Reference) const
+	const T& GetValue(UE::Dataflow::FContext& Context, const T* Reference) const
 	{
-		return GetValue(Context, Dataflow::TConnectionReference<T>(Reference));
+		return GetValue(Context, UE::Dataflow::TConnectionReference<T>(Reference));
 	}
 
 	template<typename TAnyType, typename = std::enable_if_t<std::is_base_of_v<FDataflowAnyType, TAnyType>>>
-	typename TAnyType::FStorageType GetValue(Dataflow::FContext& Context, const TAnyType* Reference) const
+	typename TAnyType::FStorageType GetValue(UE::Dataflow::FContext& Context, const TAnyType* Reference) const
 	{
-		return GetValue(Context, Dataflow::TConnectionReference<TAnyType>(Reference));
+		return GetValue(Context, UE::Dataflow::TConnectionReference<TAnyType>(Reference));
 	}
 	
 	template<class T>
-	TFuture<const T&> GetValueParallel(Dataflow::FContext& Context, const Dataflow::TConnectionReference<T>& Reference) const
+	TFuture<const T&> GetValueParallel(UE::Dataflow::FContext& Context, const UE::Dataflow::TConnectionReference<T>& Reference) const
 	{
 		checkSlow(FindInput(Reference));
 		return FindInput(Reference)->template GetValueParallel<T>(Context, *static_cast<const T*>(Reference.Reference));
 	}
 
 	template<class T>
-	TFuture<const T&> GetValueParallel(Dataflow::FContext& Context, const T* Reference) const
+	TFuture<const T&> GetValueParallel(UE::Dataflow::FContext& Context, const T* Reference) const
 	{
-		return GetValueParallel(Context, Dataflow::TConnectionReference<T>(Reference));
+		return GetValueParallel(Context, UE::Dataflow::TConnectionReference<T>(Reference));
 	}
 
 	/**
@@ -381,15 +379,15 @@ PRAGMA_ENABLE_DEPRECATION_WARNINGS
 	*   @param Reference : Pointer to a member of this node that corresponds with the output to set.
 	*   @param Default : Default value if the input is not connected.
 	*/
-	template<class T> const T& GetValue(Dataflow::FContext& Context, const Dataflow::TConnectionReference<T>& Reference, const T& Default) const
+	template<class T> const T& GetValue(UE::Dataflow::FContext& Context, const UE::Dataflow::TConnectionReference<T>& Reference, const T& Default) const
 	{
 		checkSlow(FindInput(Reference));
 		return FindInput(Reference)->template GetValue<T>(Context, Default);
 	}
 
-	template<class T> const T& GetValue(Dataflow::FContext& Context, const T* Reference, const T& Default) const
+	template<class T> const T& GetValue(UE::Dataflow::FContext& Context, const T* Reference, const T& Default) const
 	{
-		return GetValue(Context, Dataflow::TConnectionReference<T>(Reference), Default);
+		return GetValue(Context, UE::Dataflow::TConnectionReference<T>(Reference), Default);
 	}
 
 	/**
@@ -408,7 +406,7 @@ PRAGMA_ENABLE_DEPRECATION_WARNINGS
 	*   @param Reference : Pointer to a member of this node that corresponds with the output to set. 
 	*/
 	template<class T, typename = std::enable_if_t<!std::is_base_of_v<FDataflowAnyType, T>>>
-	void SetValue(Dataflow::FContext& Context, T&& Value, const typename TDecay<T>::Type* Reference) const
+	void SetValue(UE::Dataflow::FContext& Context, T&& Value, const typename TDecay<T>::Type* Reference) const
 	{
 		if (const FDataflowOutput* Output = FindOutput(Reference))
 		{
@@ -421,7 +419,7 @@ PRAGMA_ENABLE_DEPRECATION_WARNINGS
 	}
 
 	template<typename TAnyType, typename = std::enable_if_t<std::is_base_of_v<FDataflowAnyType, TAnyType>>>
-	void SetValue(Dataflow::FContext& Context, const typename TAnyType::FStorageType& Value, const TAnyType* Reference) const
+	void SetValue(UE::Dataflow::FContext& Context, const typename TAnyType::FStorageType& Value, const TAnyType* Reference) const
 	{
 		if (const FDataflowOutput* Output = FindOutput(Reference))
 		{
@@ -446,7 +444,7 @@ PRAGMA_ENABLE_DEPRECATION_WARNINGS
 	*   @param InputReference : Pointer to a input member of this node that needs to be forwarded.
 	*   @param Reference : Pointer to a member of this node that corresponds with the output to set.
 	*/
-	DATAFLOWCORE_API void ForwardInput(Dataflow::FContext& Context, const Dataflow::FConnectionReference& InputReference, const Dataflow::FConnectionReference& Reference) const;
+	DATAFLOWCORE_API void ForwardInput(UE::Dataflow::FContext& Context, const UE::Dataflow::FConnectionReference& InputReference, const UE::Dataflow::FConnectionReference& Reference) const;
 
 	/**
 	*   SafeForwardInput(...)
@@ -460,11 +458,11 @@ PRAGMA_ENABLE_DEPRECATION_WARNINGS
 	*   @param Reference : Pointer to a member of this node that corresponds with the output to set.
 	*/
 	template<class T>
-	void SafeForwardInput(Dataflow::FContext& Context, const Dataflow::FConnectionReference& InputReference, const T* Reference) const
+	void SafeForwardInput(UE::Dataflow::FContext& Context, const UE::Dataflow::FConnectionReference& InputReference, const T* Reference) const
 	{
 		if (IsConnected(InputReference))
 		{
-			ForwardInput(Context, InputReference, Dataflow::TConnectionReference<T>(Reference));
+			ForwardInput(Context, InputReference, UE::Dataflow::TConnectionReference<T>(Reference));
 		}
 		else if constexpr (std::is_base_of_v<FDataflowAnyType, T>)
 		{
@@ -483,7 +481,7 @@ PRAGMA_ENABLE_DEPRECATION_WARNINGS
 	*
 	*   @param Reference : Pointer to a member of this node that corresponds with the input.
 	*/
-	bool IsConnected(const Dataflow::FConnectionReference& Reference) const
+	bool IsConnected(const UE::Dataflow::FConnectionReference& Reference) const
 	{
 		checkSlow(FindInput(Reference));
 		return (FindInput(Reference)->GetConnection() != nullptr);
@@ -492,7 +490,7 @@ PRAGMA_ENABLE_DEPRECATION_WARNINGS
 	template<typename T>
 	bool IsConnected(const T* Reference) const
 	{
-		return IsConnected(Dataflow::FConnectionReference(Reference));
+		return IsConnected(UE::Dataflow::FConnectionReference(Reference));
 	}
 
 	void PauseInvalidations()
@@ -500,7 +498,7 @@ PRAGMA_ENABLE_DEPRECATION_WARNINGS
 		if (!bPauseInvalidations)
 		{
 			bPauseInvalidations = true;
-			PausedModifiedTimestamp = Dataflow::FTimestamp::Invalid;
+			PausedModifiedTimestamp = UE::Dataflow::FTimestamp::Invalid;
 		}
 	}
 
@@ -513,7 +511,7 @@ PRAGMA_ENABLE_DEPRECATION_WARNINGS
 		}
 	}
 
-	DATAFLOWCORE_API void Invalidate(const Dataflow::FTimestamp& ModifiedTimestamp = Dataflow::FTimestamp::Current());
+	DATAFLOWCORE_API void Invalidate(const UE::Dataflow::FTimestamp& ModifiedTimestamp = UE::Dataflow::FTimestamp::Current());
 
 	virtual void OnInvalidate() {}
 
@@ -565,19 +563,19 @@ protected:
 	virtual bool OnOutputTypeChanged(const FDataflowOutput* Output) { return false; }
 
 	// returns true if the input type was changed successfully
-	DATAFLOWCORE_API bool SetInputConcreteType(const Dataflow::FConnectionReference& InputReference, FName NewType);
+	DATAFLOWCORE_API bool SetInputConcreteType(const UE::Dataflow::FConnectionReference& InputReference, FName NewType);
 
 	template <typename T>
-	bool SetInputConcreteType(const Dataflow::FConnectionReference& InputReference)
+	bool SetInputConcreteType(const UE::Dataflow::FConnectionReference& InputReference)
 	{
 		return SetInputConcreteType(InputReference, TDataflowPolicyTypeName<T>::GetName());
 	}
 
 	// returns true if the output type was changed successfully
-	DATAFLOWCORE_API bool SetOutputConcreteType(const Dataflow::FConnectionReference& OutputReference, FName NewType);
+	DATAFLOWCORE_API bool SetOutputConcreteType(const UE::Dataflow::FConnectionReference& OutputReference, FName NewType);
 
 	template <typename T>
-	bool SetOutputConcreteType(const Dataflow::FConnectionReference& OutputReference)
+	bool SetOutputConcreteType(const UE::Dataflow::FConnectionReference& OutputReference)
 	{
 		return SetOutputConcreteType(OutputReference, TDataflowPolicyTypeName<T>::GetName());
 	}
@@ -585,14 +583,14 @@ protected:
 	// returns true if any of the types was changed successfully
 	DATAFLOWCORE_API bool SetAllConnectionConcreteType(FName NewType);
 
-	DATAFLOWCORE_API FDataflowInput& RegisterInputConnectionInternal(const Dataflow::FConnectionReference& Reference, const FName& PropertyName = NAME_None);
-	DATAFLOWCORE_API FDataflowOutput& RegisterOutputConnectionInternal(const Dataflow::FConnectionReference& Reference, const FName& PropertyName = NAME_None);
-	DATAFLOWCORE_API FDataflowInput& RegisterInputArrayConnectionInternal(const Dataflow::FConnectionReference& Reference, const FName& ElementPropertyName = NAME_None,
+	DATAFLOWCORE_API FDataflowInput& RegisterInputConnectionInternal(const UE::Dataflow::FConnectionReference& Reference, const FName& PropertyName = NAME_None);
+	DATAFLOWCORE_API FDataflowOutput& RegisterOutputConnectionInternal(const UE::Dataflow::FConnectionReference& Reference, const FName& PropertyName = NAME_None);
+	DATAFLOWCORE_API FDataflowInput& RegisterInputArrayConnectionInternal(const UE::Dataflow::FConnectionReference& Reference, const FName& ElementPropertyName = NAME_None,
 		const FName& ArrayPropertyName = NAME_None);
 
 private:
 
-	void InitConnectionParametersFromPropertyReference(const FStructOnScope& StructOnScope, const void* PropertyRef, const FName& PropertyName, Dataflow::FConnectionParameters& OutParams);
+	void InitConnectionParametersFromPropertyReference(const FStructOnScope& StructOnScope, const void* PropertyRef, const FName& PropertyName, UE::Dataflow::FConnectionParameters& OutParams);
 	// This will add [ContainerIndex] to any array it finds unless ContainerIndex == INDEX_NONE.
 	static FString GetPropertyFullNameString(const TConstArrayView<const FProperty*>& PropertyChain, int32 ContainerIndex = INDEX_NONE);
 	static FName GetPropertyFullName(const TArray<const FProperty*>& PropertyChain, int32 ContainerIndex = INDEX_NONE);
@@ -600,7 +598,7 @@ private:
 	static FString StripContainerIndexFromPropertyFullName(const FString& PropertyFullName);
 	static uint32 GetPropertyOffset(const TArray<const FProperty*>& PropertyChain);
 	uint32 GetConnectionOffsetFromReference(const void* Reference) const;
-	DATAFLOWCORE_API Dataflow::FConnectionKey GetKeyFromReference(const Dataflow::FConnectionReference& Reference) const;
+	DATAFLOWCORE_API UE::Dataflow::FConnectionKey GetKeyFromReference(const UE::Dataflow::FConnectionReference& Reference) const;
 
 	/**
 	* Find a property using the property address and name (not including its parent struct property names).
@@ -613,20 +611,20 @@ private:
 	/** Find a property using the property full name (must includes its parent struct property names). */
 	const FProperty* FindProperty(const UStruct* Struct, const FName& PropertyFullName, TArray<const FProperty*>* OutPropertyChain = nullptr) const;
 
-	virtual TArray<Dataflow::FRenderingParameter> GetRenderParametersImpl() const { return TArray<Dataflow::FRenderingParameter>(); }
+	virtual TArray<UE::Dataflow::FRenderingParameter> GetRenderParametersImpl() const { return TArray<UE::Dataflow::FRenderingParameter>(); }
 
 	bool bHasValidConnections = true;
-	TMap<Dataflow::FConnectionKey, FDataflowInput*> ExpandedInputs;
+	TMap<UE::Dataflow::FConnectionKey, FDataflowInput*> ExpandedInputs;
 	TMap<uint32, const FArrayProperty*> InputArrayProperties; // Used to calculate ContainerElementOffsets
 
 protected:
 	bool bPauseInvalidations = false;
-	Dataflow::FTimestamp PausedModifiedTimestamp = Dataflow::FTimestamp::Invalid; // When unpausing invalidations, Invalidate will be called with this timestamp.
+	UE::Dataflow::FTimestamp PausedModifiedTimestamp = UE::Dataflow::FTimestamp::Invalid; // When unpausing invalidations, Invalidate will be called with this timestamp.
 	FOnNodeInvalidated OnNodeInvalidatedDelegate;
 
 };
 
-namespace Dataflow
+namespace UE::Dataflow
 {
 
 class FDataflowNodePauseInvalidationScope
@@ -661,11 +659,11 @@ private:
 	//
 
 #define DATAFLOW_NODE_REGISTER_CREATION_FACTORY(A)									\
-	::Dataflow::FNodeFactory::RegisterNodeFromType<A>();
+	::UE::Dataflow::FNodeFactory::RegisterNodeFromType<A>();
 
 #define DATAFLOW_NODE_RENDER_TYPE(A, B, ...)												\
-	virtual TArray<::Dataflow::FRenderingParameter> GetRenderParametersImpl() const {		\
-		TArray<::Dataflow::FRenderingParameter> Array;								\
+	virtual TArray<::UE::Dataflow::FRenderingParameter> GetRenderParametersImpl() const {		\
+		TArray<::UE::Dataflow::FRenderingParameter> Array;								\
 		Array.Add({ A, B, {__VA_ARGS__,} });														\
 		return Array;}
 
@@ -696,9 +694,7 @@ private:
 
 #define DATAFLOW_NODE_REGISTER_CREATION_FACTORY_NODE_COLORS_BY_CATEGORY(A, C1, C2)	\
 {																					\
-	::Dataflow::FNodeColorsRegistry::Get().RegisterNodeColors(A, {C1, C2});			\
+	::UE::Dataflow::FNodeColorsRegistry::Get().RegisterNodeColors(A, {C1, C2});			\
 }																					\
 
 }
-
-
diff --git a/Engine/Source/Runtime/Experimental/Dataflow/Core/Public/Dataflow/DataflowNodeColorsRegistry.h b/Engine/Source/Runtime/Experimental/Dataflow/Core/Public/Dataflow/DataflowNodeColorsRegistry.h
--- a/Engine/Source/Runtime/Experimental/Dataflow/Core/Public/Dataflow/DataflowNodeColorsRegistry.h
+++ b/Engine/Source/Runtime/Experimental/Dataflow/Core/Public/Dataflow/DataflowNodeColorsRegistry.h
@@ -12,7 +12,7 @@ struct FDataflowConnection;
 	
 class FLazySingleton;
 
-namespace Dataflow
+namespace UE::Dataflow
 {
 	//
 	//
@@ -37,5 +37,5 @@ namespace Dataflow
 
 		friend FLazySingleton;
 	};
-
 }
+
diff --git a/Engine/Source/Runtime/Experimental/Dataflow/Core/Public/Dataflow/DataflowNodeFactory.h b/Engine/Source/Runtime/Experimental/Dataflow/Core/Public/Dataflow/DataflowNodeFactory.h
--- a/Engine/Source/Runtime/Experimental/Dataflow/Core/Public/Dataflow/DataflowNodeFactory.h
+++ b/Engine/Source/Runtime/Experimental/Dataflow/Core/Public/Dataflow/DataflowNodeFactory.h
@@ -8,7 +8,7 @@
 struct FDataflowNode;
 struct FDataflowConnection;
 
-namespace Dataflow
+namespace UE::Dataflow
 {
 	struct FNewNodeParameters {
 		FGuid Guid;
@@ -91,7 +91,7 @@ namespace Dataflow
 				GetVersionFromTypeName(T::StaticType())
 			};
 
-			auto CreationFunction = [](const ::Dataflow::FNewNodeParameters& InParam)
+			auto CreationFunction = [](const UE::Dataflow::FNewNodeParameters& InParam)
 			{
 				const FNodeParameters NodeParameters = { InParam.Name, InParam.OwningObject };
 				TUniquePtr<T> Val = MakeUnique<T>(NodeParameters, InParam.Guid);
@@ -139,3 +139,4 @@ namespace Dataflow
 	};
 }
 
+
diff --git a/Engine/Source/Runtime/Experimental/Dataflow/Core/Public/Dataflow/DataflowNodeParameters.h b/Engine/Source/Runtime/Experimental/Dataflow/Core/Public/Dataflow/DataflowNodeParameters.h
--- a/Engine/Source/Runtime/Experimental/Dataflow/Core/Public/Dataflow/DataflowNodeParameters.h
+++ b/Engine/Source/Runtime/Experimental/Dataflow/Core/Public/Dataflow/DataflowNodeParameters.h
@@ -16,7 +16,7 @@ struct FDataflowConnection;
 
 #define DATAFLOW_EDITOR_EVALUATION WITH_EDITOR
 
-namespace Dataflow
+namespace UE::Dataflow
 {
 	class FContext;
 
@@ -179,21 +179,21 @@ namespace Dataflow
 	};
 };
 
-inline FArchive& operator<<(FArchive& Ar, Dataflow::FTimestamp& ValueIn)
+inline FArchive& operator<<(FArchive& Ar, UE::Dataflow::FTimestamp& ValueIn)
 {
 	Ar << ValueIn.Value;
 	Ar << ValueIn.Invalid;
 	return Ar;
 }
 
-inline FArchive& operator<<(FArchive& Ar, Dataflow::FContextCache& ValueIn)
+inline FArchive& operator<<(FArchive& Ar, UE::Dataflow::FContextCache& ValueIn)
 {
 	ValueIn.Serialize(Ar);
 	return Ar;
 }
 
 
-namespace Dataflow
+namespace UE::Dataflow
 {
 
 	class FContext
@@ -451,3 +451,5 @@ namespace Dataflow
 		return MakeUnique<FContextCacheElementNull>(GetNodeGuid(), GetProperty(), InReferenceDataKey, GetNodeHash(), GetTimestamp());
 	}
 }
+
+
diff --git a/Engine/Source/Runtime/Experimental/Dataflow/Core/Public/Dataflow/DataflowOverrideNode.h b/Engine/Source/Runtime/Experimental/Dataflow/Core/Public/Dataflow/DataflowOverrideNode.h
--- a/Engine/Source/Runtime/Experimental/Dataflow/Core/Public/Dataflow/DataflowOverrideNode.h
+++ b/Engine/Source/Runtime/Experimental/Dataflow/Core/Public/Dataflow/DataflowOverrideNode.h
@@ -28,7 +28,7 @@ struct FDataflowOverrideNode : public FDataflowNode
 	DATAFLOW_NODE_DEFINE_INTERNAL(FDataflowOverrideNode, "DataflowOverrideNode", "BaseClass", "")
 
 public:
-	FDataflowOverrideNode(const Dataflow::FNodeParameters& Param, FGuid InGuid = FGuid::NewGuid())
+	FDataflowOverrideNode(const UE::Dataflow::FNodeParameters& Param, FGuid InGuid = FGuid::NewGuid())
 		: Super(Param,InGuid) 
 	{
 		RegisterInputConnection(&Key);
@@ -41,23 +41,23 @@ public:
 	DATAFLOWCORE_API bool ShouldInvalidate(FName InKey) const;
 
 	template <class T>
-	T GetDefaultValue(Dataflow::FContext& Context) const
+	T GetDefaultValue(UE::Dataflow::FContext& Context) const
 	{
 		return GetValue<T>(Context, &Default, Default);
 	}
 
 	template <>
-	DATAFLOWCORE_API int32 GetDefaultValue(Dataflow::FContext& Context) const;
+	DATAFLOWCORE_API int32 GetDefaultValue(UE::Dataflow::FContext& Context) const;
 
 	template <>
-	DATAFLOWCORE_API float GetDefaultValue(Dataflow::FContext& Context) const;
+	DATAFLOWCORE_API float GetDefaultValue(UE::Dataflow::FContext& Context) const;
 
-	DATAFLOWCORE_API FString GetValueFromAsset(Dataflow::FContext& Context, const UObject* InOwner) const;
+	DATAFLOWCORE_API FString GetValueFromAsset(UE::Dataflow::FContext& Context, const UObject* InOwner) const;
 
 	//
 	// Evaluate
 	//
-	virtual void Evaluate(Dataflow::FContext& Context, const FDataflowOutput* Out) const { ensure(false); };
+	virtual void Evaluate(UE::Dataflow::FContext& Context, const FDataflowOutput* Out) const { ensure(false); };
 
 public:
 	UPROPERTY(EditAnywhere, Category = "Overrides", meta = (DataflowInput))
diff --git a/Engine/Source/Runtime/Experimental/Dataflow/Core/Public/Dataflow/DataflowTerminalNode.h b/Engine/Source/Runtime/Experimental/Dataflow/Core/Public/Dataflow/DataflowTerminalNode.h
--- a/Engine/Source/Runtime/Experimental/Dataflow/Core/Public/Dataflow/DataflowTerminalNode.h
+++ b/Engine/Source/Runtime/Experimental/Dataflow/Core/Public/Dataflow/DataflowTerminalNode.h
@@ -31,7 +31,7 @@ struct FDataflowTerminalNode : public FDataflowNode
 	FDataflowTerminalNode()
 		: Super() { }
 
-	FDataflowTerminalNode(const Dataflow::FNodeParameters& Param, FGuid InGuid = FGuid::NewGuid())
+	FDataflowTerminalNode(const UE::Dataflow::FNodeParameters& Param, FGuid InGuid = FGuid::NewGuid())
 		: Super(Param,InGuid) {
 	}
 
@@ -51,16 +51,16 @@ struct FDataflowTerminalNode : public FDataflowNode
 	//
 	// Error Checking
 	//
-	virtual void SetAssetValue(TObjectPtr<UObject> Asset, Dataflow::FContext& Context) const { ensure(false); }
+	virtual void SetAssetValue(TObjectPtr<UObject> Asset, UE::Dataflow::FContext& Context) const { ensure(false); }
 
 	//
 	// Evaluate
 	//
 
-	virtual void Evaluate(Dataflow::FContext& Context) const { ensure(false); }
+	virtual void Evaluate(UE::Dataflow::FContext& Context) const { ensure(false); }
 
 private:
-	virtual void Evaluate(Dataflow::FContext& Context, const FDataflowOutput* Out) const override
+	virtual void Evaluate(UE::Dataflow::FContext& Context, const FDataflowOutput* Out) const override
 	{
 		Evaluate(Context);
 	};
diff --git a/Engine/Source/Runtime/Experimental/Dataflow/Core/Public/Dataflow/DataflowVectorNodes.h b/Engine/Source/Runtime/Experimental/Dataflow/Core/Public/Dataflow/DataflowVectorNodes.h
--- a/Engine/Source/Runtime/Experimental/Dataflow/Core/Public/Dataflow/DataflowVectorNodes.h
+++ b/Engine/Source/Runtime/Experimental/Dataflow/Core/Public/Dataflow/DataflowVectorNodes.h
@@ -16,7 +16,7 @@ struct FDataflowVectorMakeVec2Node : public FDataflowNode
 	DATAFLOW_NODE_DEFINE_INTERNAL(FDataflowVectorMakeVec2Node, "MakeVector2", DATAFLOW_MATH_VECTOR_NODES_CATEGORY, "")
 
 public:
-	FDataflowVectorMakeVec2Node(const Dataflow::FNodeParameters& InParam, FGuid InGuid = FGuid::NewGuid());
+	FDataflowVectorMakeVec2Node(const UE::Dataflow::FNodeParameters& InParam, FGuid InGuid = FGuid::NewGuid());
 
 private:
 	/** X component */
@@ -31,7 +31,7 @@ private:
 	UPROPERTY(meta = (DataflowOutput))
 	FDataflowVectorTypes Vector2D;
 
-	virtual void Evaluate(Dataflow::FContext& Context, const FDataflowOutput* Out) const override;
+	virtual void Evaluate(UE::Dataflow::FContext& Context, const FDataflowOutput* Out) const override;
 };
 
 /** Make a 3D Vector */
@@ -42,7 +42,7 @@ struct FDataflowVectorMakeVec3Node : public FDataflowNode
 	DATAFLOW_NODE_DEFINE_INTERNAL(FDataflowVectorMakeVec3Node, "MakeVector3", DATAFLOW_MATH_VECTOR_NODES_CATEGORY, "")
 
 public:
-	FDataflowVectorMakeVec3Node(const Dataflow::FNodeParameters& InParam, FGuid InGuid = FGuid::NewGuid());
+	FDataflowVectorMakeVec3Node(const UE::Dataflow::FNodeParameters& InParam, FGuid InGuid = FGuid::NewGuid());
 
 private:
 	/** X component */
@@ -61,7 +61,7 @@ private:
 	UPROPERTY(meta = (DataflowOutput))
 	FDataflowVectorTypes Vector3D;
 
-	virtual void Evaluate(Dataflow::FContext& Context, const FDataflowOutput* Out) const override;
+	virtual void Evaluate(UE::Dataflow::FContext& Context, const FDataflowOutput* Out) const override;
 };
 
 
@@ -73,7 +73,7 @@ struct FDataflowVectorMakeVec4Node : public FDataflowNode
 	DATAFLOW_NODE_DEFINE_INTERNAL(FDataflowVectorMakeVec4Node, "MakeVector4", DATAFLOW_MATH_VECTOR_NODES_CATEGORY, "")
 
 public:
-	FDataflowVectorMakeVec4Node(const Dataflow::FNodeParameters& InParam, FGuid InGuid = FGuid::NewGuid());
+	FDataflowVectorMakeVec4Node(const UE::Dataflow::FNodeParameters& InParam, FGuid InGuid = FGuid::NewGuid());
 
 private:
 	/** X component */
@@ -96,7 +96,7 @@ private:
 	UPROPERTY(meta = (DataflowOutput))
 	FDataflowVectorTypes Vector4D;
 
-	virtual void Evaluate(Dataflow::FContext& Context, const FDataflowOutput* Out) const override;
+	virtual void Evaluate(UE::Dataflow::FContext& Context, const FDataflowOutput* Out) const override;
 };
 
 
@@ -111,7 +111,7 @@ struct FDataflowVectorBreakNode: public FDataflowNode
 	DATAFLOW_NODE_DEFINE_INTERNAL(FDataflowVectorBreakNode, "BreakVector", DATAFLOW_MATH_VECTOR_NODES_CATEGORY, "")
 
 public:
-	FDataflowVectorBreakNode(const Dataflow::FNodeParameters& InParam, FGuid InGuid = FGuid::NewGuid());
+	FDataflowVectorBreakNode(const UE::Dataflow::FNodeParameters& InParam, FGuid InGuid = FGuid::NewGuid());
 
 private:
 	/** Vector to break into components */
@@ -134,7 +134,7 @@ private:
 	UPROPERTY(meta = (DataflowOutput));
 	FDataflowNumericTypes W;
 
-	virtual void Evaluate(Dataflow::FContext& Context, const FDataflowOutput* Out) const override;
+	virtual void Evaluate(UE::Dataflow::FContext& Context, const FDataflowOutput* Out) const override;
 };
 
 /** Add two vectors component wise : V = (A + B)*/
@@ -145,7 +145,7 @@ struct FDataflowVectorAddNode : public FDataflowNode
 	DATAFLOW_NODE_DEFINE_INTERNAL(FDataflowVectorAddNode, "AddVector", DATAFLOW_MATH_VECTOR_NODES_CATEGORY, "")
 
 public:
-	FDataflowVectorAddNode(const Dataflow::FNodeParameters& InParam, FGuid InGuid = FGuid::NewGuid());
+	FDataflowVectorAddNode(const UE::Dataflow::FNodeParameters& InParam, FGuid InGuid = FGuid::NewGuid());
 
 private:
 	/** A Vector operand */
@@ -160,7 +160,7 @@ private:
 	UPROPERTY(meta = (DataflowOutput))
 	FDataflowVectorTypes V;
 
-	virtual void Evaluate(Dataflow::FContext& Context, const FDataflowOutput* Out) const override;
+	virtual void Evaluate(UE::Dataflow::FContext& Context, const FDataflowOutput* Out) const override;
 };
 
 /** Subtract two vectors component wise: V = (A - B) */
@@ -171,7 +171,7 @@ struct FDataflowVectorSubtractNode : public FDataflowNode
 	DATAFLOW_NODE_DEFINE_INTERNAL(FDataflowVectorSubtractNode, "SubtractVector", DATAFLOW_MATH_VECTOR_NODES_CATEGORY, "")
 
 public:
-	FDataflowVectorSubtractNode(const Dataflow::FNodeParameters& InParam, FGuid InGuid = FGuid::NewGuid());
+	FDataflowVectorSubtractNode(const UE::Dataflow::FNodeParameters& InParam, FGuid InGuid = FGuid::NewGuid());
 
 private:
 	/** A Vector operand */
@@ -186,7 +186,7 @@ private:
 	UPROPERTY(meta = (DataflowOutput))
 	FDataflowVectorTypes V;
 
-	virtual void Evaluate(Dataflow::FContext& Context, const FDataflowOutput* Out) const override;
+	virtual void Evaluate(UE::Dataflow::FContext& Context, const FDataflowOutput* Out) const override;
 };
 
 /** Compute the dot product between two vectors : DotProduct = A.B */
@@ -197,7 +197,7 @@ struct FDataflowVectorDotProductNode : public FDataflowNode
 	DATAFLOW_NODE_DEFINE_INTERNAL(FDataflowVectorDotProductNode, "VectorDotProduct", DATAFLOW_MATH_VECTOR_NODES_CATEGORY, "")
 
 public:
-	FDataflowVectorDotProductNode(const Dataflow::FNodeParameters& InParam, FGuid InGuid = FGuid::NewGuid());
+	FDataflowVectorDotProductNode(const UE::Dataflow::FNodeParameters& InParam, FGuid InGuid = FGuid::NewGuid());
 
 private:
 	/** A Vector operand */
@@ -212,7 +212,7 @@ private:
 	UPROPERTY(meta = (DataflowOutput))
 	FDataflowNumericTypes DotProduct;
 
-	virtual void Evaluate(Dataflow::FContext& Context, const FDataflowOutput* Out) const override;
+	virtual void Evaluate(UE::Dataflow::FContext& Context, const FDataflowOutput* Out) const override;
 };
 
 /** Compute the Length of a vector : Length = |V| */
@@ -223,7 +223,7 @@ struct FDataflowVectorLengthNode : public FDataflowNode
 	DATAFLOW_NODE_DEFINE_INTERNAL(FDataflowVectorLengthNode, "VectorLength", DATAFLOW_MATH_VECTOR_NODES_CATEGORY, "")
 
 public:
-	FDataflowVectorLengthNode(const Dataflow::FNodeParameters& InParam, FGuid InGuid = FGuid::NewGuid());
+	FDataflowVectorLengthNode(const UE::Dataflow::FNodeParameters& InParam, FGuid InGuid = FGuid::NewGuid());
 
 private:
 	/** Vector to get length from */
@@ -234,7 +234,7 @@ private:
 	UPROPERTY(meta = (DataflowOutput))
 	FDataflowNumericTypes Length;
 
-	virtual void Evaluate(Dataflow::FContext& Context, const FDataflowOutput* Out) const override;
+	virtual void Evaluate(UE::Dataflow::FContext& Context, const FDataflowOutput* Out) const override;
 };
 
 /** Compute the Squared length of a vector : Length = |V||V| */
@@ -245,7 +245,7 @@ struct FDataflowVectorSquaredLengthNode : public FDataflowNode
 	DATAFLOW_NODE_DEFINE_INTERNAL(FDataflowVectorSquaredLengthNode, "VectorSquaredLength", DATAFLOW_MATH_VECTOR_NODES_CATEGORY, "")
 
 public:
-	FDataflowVectorSquaredLengthNode(const Dataflow::FNodeParameters& InParam, FGuid InGuid = FGuid::NewGuid());
+	FDataflowVectorSquaredLengthNode(const UE::Dataflow::FNodeParameters& InParam, FGuid InGuid = FGuid::NewGuid());
 
 private:
 	/** Vector to get squared length from */
@@ -256,7 +256,7 @@ private:
 	UPROPERTY(meta = (DataflowOutput))
 	FDataflowNumericTypes SquaredLength;
 
-	virtual void Evaluate(Dataflow::FContext& Context, const FDataflowOutput* Out) const override;
+	virtual void Evaluate(UE::Dataflow::FContext& Context, const FDataflowOutput* Out) const override;
 };
 
 /** Compute the distance between two vectors : Distance = |B-A| */
@@ -267,7 +267,7 @@ struct FDataflowVectorDistanceNode : public FDataflowNode
 	DATAFLOW_NODE_DEFINE_INTERNAL(FDataflowVectorDistanceNode, "VectorDistance", DATAFLOW_MATH_VECTOR_NODES_CATEGORY, "")
 
 public:
-	FDataflowVectorDistanceNode(const Dataflow::FNodeParameters& InParam, FGuid InGuid = FGuid::NewGuid());
+	FDataflowVectorDistanceNode(const UE::Dataflow::FNodeParameters& InParam, FGuid InGuid = FGuid::NewGuid());
 
 private:
 	/** A Vector operand */
@@ -282,7 +282,7 @@ private:
 	UPROPERTY(meta = (DataflowOutput))
 	FDataflowNumericTypes Distance;
 
-	virtual void Evaluate(Dataflow::FContext& Context, const FDataflowOutput* Out) const override;
+	virtual void Evaluate(UE::Dataflow::FContext& Context, const FDataflowOutput* Out) const override;
 };
 
 /** 
@@ -296,7 +296,7 @@ struct FDataflowVectorCrossProductNode : public FDataflowNode
 	DATAFLOW_NODE_DEFINE_INTERNAL(FDataflowVectorCrossProductNode, "VectorCrossProduct", DATAFLOW_MATH_VECTOR_NODES_CATEGORY, "")
 
 public:
-	FDataflowVectorCrossProductNode(const Dataflow::FNodeParameters& InParam, FGuid InGuid = FGuid::NewGuid());
+	FDataflowVectorCrossProductNode(const UE::Dataflow::FNodeParameters& InParam, FGuid InGuid = FGuid::NewGuid());
 
 private:
 	/** A Vector operand */
@@ -311,7 +311,7 @@ private:
 	UPROPERTY(meta = (DataflowOutput))
 	FDataflowVectorTypes CrossProduct;
 
-	virtual void Evaluate(Dataflow::FContext& Context, const FDataflowOutput* Out) const override;
+	virtual void Evaluate(UE::Dataflow::FContext& Context, const FDataflowOutput* Out) const override;
 };
 
 /** Scale a vector by a scalar : Scaled = (V * Scale) */
@@ -322,7 +322,7 @@ struct FDataflowVectorScaleNode : public FDataflowNode
 	DATAFLOW_NODE_DEFINE_INTERNAL(FDataflowVectorScaleNode, "ScaleVector", DATAFLOW_MATH_VECTOR_NODES_CATEGORY, "")
 
 public:
-	FDataflowVectorScaleNode(const Dataflow::FNodeParameters& InParam, FGuid InGuid = FGuid::NewGuid());
+	FDataflowVectorScaleNode(const UE::Dataflow::FNodeParameters& InParam, FGuid InGuid = FGuid::NewGuid());
 
 private:
 	/** Vector to scale */
@@ -337,7 +337,7 @@ private:
 	UPROPERTY(meta = (DataflowOutput))
 	FDataflowVectorTypes Scaled;
 
-	virtual void Evaluate(Dataflow::FContext& Context, const FDataflowOutput* Out) const override;
+	virtual void Evaluate(UE::Dataflow::FContext& Context, const FDataflowOutput* Out) const override;
 
 	virtual bool OnInputTypeChanged(const FDataflowInput* Input) override;
 	virtual bool OnOutputTypeChanged(const FDataflowOutput* Output) override;
@@ -351,7 +351,7 @@ struct FDataflowVectorNormalize : public FDataflowNode
 	DATAFLOW_NODE_DEFINE_INTERNAL(FDataflowVectorNormalize, "NormalizeVector", DATAFLOW_MATH_VECTOR_NODES_CATEGORY, "")
 
 public:
-	FDataflowVectorNormalize(const Dataflow::FNodeParameters& InParam, FGuid InGuid = FGuid::NewGuid());
+	FDataflowVectorNormalize(const UE::Dataflow::FNodeParameters& InParam, FGuid InGuid = FGuid::NewGuid());
 
 private:
 	/** Vector to normalize */
@@ -362,13 +362,13 @@ private:
 	UPROPERTY(meta = (DataflowOutput))
 	FDataflowVectorTypes Normalized;
 
-	virtual void Evaluate(Dataflow::FContext& Context, const FDataflowOutput* Out) const override;
+	virtual void Evaluate(UE::Dataflow::FContext& Context, const FDataflowOutput* Out) const override;
 
 	virtual bool OnInputTypeChanged(const FDataflowInput* Input) override;
 	virtual bool OnOutputTypeChanged(const FDataflowOutput* Output) override;
 };
 
-namespace Dataflow
+namespace UE::Dataflow
 {
 	void RegisterDataflowVectorNodes();
 }
diff --git a/Engine/Source/Runtime/Experimental/Dataflow/Engine/Private/Dataflow/DataflowBlueprintLibrary.cpp b/Engine/Source/Runtime/Experimental/Dataflow/Engine/Private/Dataflow/DataflowBlueprintLibrary.cpp
--- a/Engine/Source/Runtime/Experimental/Dataflow/Engine/Private/Dataflow/DataflowBlueprintLibrary.cpp
+++ b/Engine/Source/Runtime/Experimental/Dataflow/Engine/Private/Dataflow/DataflowBlueprintLibrary.cpp
@@ -15,7 +15,7 @@ void UDataflowBlueprintLibrary::EvaluateTerminalNodeByName(UDataflow* Dataflow,
 			if (const FDataflowTerminalNode* TerminalNode = Node->AsType<const FDataflowTerminalNode>())
 			{
 				UE_LOG(LogChaosDataflow, Verbose, TEXT("UDataflowBlueprintLibrary::EvaluateTerminalNodeByName(): Node [%s]"), *TerminalNodeName.ToString());
-				Dataflow::FEngineContext Context(ResultAsset);
+				UE::Dataflow::FEngineContext Context(ResultAsset);
 				TerminalNode->Evaluate(Context);
 				if (ResultAsset)
 				{
diff --git a/Engine/Source/Runtime/Experimental/Dataflow/Engine/Private/Dataflow/DataflowContent.cpp b/Engine/Source/Runtime/Experimental/Dataflow/Engine/Private/Dataflow/DataflowContent.cpp
--- a/Engine/Source/Runtime/Experimental/Dataflow/Engine/Private/Dataflow/DataflowContent.cpp
+++ b/Engine/Source/Runtime/Experimental/Dataflow/Engine/Private/Dataflow/DataflowContent.cpp
@@ -18,7 +18,7 @@ FAutoConsoleVariableRef CVARDataflowEnableContextCaching(TEXT("p.Dataflow.Editor
 	TEXT("Allow the Dataflow editor to crate and use a pre-evaluated graph when the dataflow editor is re-opened.[def:true]"));
 
 
-namespace DataflowContextHelpers
+namespace UE::DataflowContextHelpers
 {
 
 	/*
@@ -32,7 +32,7 @@ namespace DataflowContextHelpers
 	*/
 	bool BindContextToGraph(TObjectPtr<UObject>& Asset, UDataflow* DataflowAsset)
 	{
-		using namespace Dataflow;
+		using namespace UE::Dataflow;
 
 		TSharedPtr<FGraph> Dataflow = DataflowAsset->GetDataflow();
 		if (!Dataflow) return false;
@@ -87,7 +87,7 @@ namespace DataflowContextHelpers
 	*/
 	bool ValidateCachedNodeHash(TObjectPtr<UObject>& Asset, UDataflow* DataflowAsset)
 	{
-		using namespace Dataflow;
+		using namespace UE::Dataflow;
 
 		TSharedPtr<FGraph> Dataflow = DataflowAsset->GetDataflow();
 		if (!Dataflow) return false;
@@ -127,13 +127,13 @@ namespace DataflowContextHelpers
 
 	bool ResetCacheTimestamp(TObjectPtr<UObject>& Asset, UDataflow* DataflowAsset)
 	{
-		using namespace Dataflow;
+		using namespace UE::Dataflow;
 
 		UDataflowBaseContent* BaseContent = Cast< UDataflowBaseContent>(Asset.Get());
 		if (!BaseContent) return false;
 		if (!BaseContent->GetDataflowContext()) return false;
 
-		Dataflow::FTimestamp NewTimestamp = DataflowAsset->GetRenderingTimestamp().Value + 1;
+		UE::Dataflow::FTimestamp NewTimestamp = DataflowAsset->GetRenderingTimestamp().Value + 1;
 		BaseContent->SetLastModifiedTimestamp(NewTimestamp, false /*bMakeDirty*/);
 
 		TSharedPtr<FEngineContext>& Context = BaseContent->GetDataflowContext();
@@ -306,7 +306,7 @@ void UDataflowBaseContent::SetDataflowOwner(const TObjectPtr<UObject>& InOwner)
 {
 	if(!DataflowContext)
 	{
-		DataflowContext = MakeShared<Dataflow::FEngineContext>(nullptr);
+		DataflowContext = MakeShared<UE::Dataflow::FEngineContext>(nullptr);
 	}
 	DataflowContext->Owner = InOwner;  
 	SetConstructionDirty(true);
@@ -322,7 +322,7 @@ void UDataflowBaseContent::SetDataflowAsset(const TObjectPtr<UDataflow>& Dataflo
 {
 	if(!DataflowContext)
 	{
-		DataflowContext = MakeShared<Dataflow::FEngineContext>(nullptr);
+		DataflowContext = MakeShared<UE::Dataflow::FEngineContext>(nullptr);
 	}
 	DataflowGraph = DataflowAsset;  
 	SetConstructionDirty(true);
@@ -343,7 +343,7 @@ void UDataflowBaseContent::PostEditChangeProperty(FPropertyChangedEvent& Propert
 
 #endif //if WITH_EDITOR
 
-void UDataflowBaseContent::SetLastModifiedTimestamp(Dataflow::FTimestamp InTimestamp, bool bMakeDirty) 
+void UDataflowBaseContent::SetLastModifiedTimestamp(UE::Dataflow::FTimestamp InTimestamp, bool bMakeDirty) 
 { 
 	if (InTimestamp.IsInvalid() || LastModifiedTimestamp < InTimestamp)
 	{
@@ -360,7 +360,7 @@ void UDataflowBaseContent::SetLastModifiedTimestamp(Dataflow::FTimestamp InTimes
 	}
 }
 
-void UDataflowBaseContent::SetDataflowContext(const TSharedPtr<Dataflow::FEngineContext>& InContext) 
+void UDataflowBaseContent::SetDataflowContext(const TSharedPtr<UE::Dataflow::FEngineContext>& InContext) 
 { 
 	DataflowContext = InContext;  
 	SetConstructionDirty(true);
@@ -375,7 +375,7 @@ void UDataflowBaseContent::Serialize(FArchive& Ar)
 
 	if (!DataflowContext)
 	{
-		DataflowContext = MakeShared<Dataflow::FEngineContext>(nullptr);
+		DataflowContext = MakeShared<UE::Dataflow::FEngineContext>(nullptr);
 	}
 	DataflowContext->Serialize(Ar);
 }
diff --git a/Engine/Source/Runtime/Experimental/Dataflow/Engine/Private/Dataflow/DataflowEdNode.cpp b/Engine/Source/Runtime/Experimental/Dataflow/Engine/Private/Dataflow/DataflowEdNode.cpp
--- a/Engine/Source/Runtime/Experimental/Dataflow/Engine/Private/Dataflow/DataflowEdNode.cpp
+++ b/Engine/Source/Runtime/Experimental/Dataflow/Engine/Private/Dataflow/DataflowEdNode.cpp
@@ -24,28 +24,28 @@
 
 DEFINE_LOG_CATEGORY_STATIC(DATAFLOWNODE_LOG, Error, All);
 
-namespace Dataflow::Private
+namespace UE::Dataflow::Private
 {
-static Dataflow::FPin::EDirection EdPinDirectionToDataflowDirection(EEdGraphPinDirection EdDirection)
+static UE::Dataflow::FPin::EDirection EdPinDirectionToDataflowDirection(EEdGraphPinDirection EdDirection)
 {
 	if (EdDirection == EEdGraphPinDirection::EGPD_Input)
 	{
-		return Dataflow::FPin::EDirection::INPUT;
+		return FPin::EDirection::INPUT;
 	}
 	if (EdDirection == EEdGraphPinDirection::EGPD_Output)
 	{
-		return Dataflow::FPin::EDirection::OUTPUT;
+		return FPin::EDirection::OUTPUT;
 	}
-	return Dataflow::FPin::EDirection::NONE;
+	return FPin::EDirection::NONE;
 }
 
-static EEdGraphPinDirection DataflowDirectionToEdPinDirection(Dataflow::FPin::EDirection Direction)
+static EEdGraphPinDirection DataflowDirectionToEdPinDirection(UE::Dataflow::FPin::EDirection Direction)
 {
-	if (Direction == Dataflow::FPin::EDirection::INPUT)
+	if (Direction == FPin::EDirection::INPUT)
 	{
 		return EEdGraphPinDirection::EGPD_Input;
 	}
-	if (Direction == Dataflow::FPin::EDirection::OUTPUT)
+	if (Direction == FPin::EDirection::OUTPUT)
 	{
 		return EEdGraphPinDirection::EGPD_Output;
 	}
@@ -115,7 +115,7 @@ bool UDataflowEdNode::CanEnableWireframeRenderNode() const
 
 TSharedPtr<FDataflowNode> UDataflowEdNode::GetDataflowNode()
 {
-	if(TSharedPtr<Dataflow::FGraph> Dataflow = GetDataflowGraph())
+	if(TSharedPtr<UE::Dataflow::FGraph> Dataflow = GetDataflowGraph())
 	{
 		return Dataflow->FindBaseNode(GetDataflowNodeGuid());
 	}
@@ -124,7 +124,7 @@ TSharedPtr<FDataflowNode> UDataflowEdNode::GetDataflowNode()
 
 TSharedPtr<const FDataflowNode> UDataflowEdNode::GetDataflowNode() const
 {
-	if (TSharedPtr<const Dataflow::FGraph> Dataflow = GetDataflowGraph())
+	if (TSharedPtr<const UE::Dataflow::FGraph> Dataflow = GetDataflowGraph())
 	{
 		return Dataflow->FindBaseNode(GetDataflowNodeGuid());
 	}
@@ -144,9 +144,9 @@ void UDataflowEdNode::AllocateDefaultPins()
 		{
 			if (TSharedPtr<FDataflowNode> DataflowNode = DataflowGraph->FindBaseNode(DataflowNodeGuid))
 			{
-				for (const Dataflow::FPin& Pin : DataflowNode->GetPins())
+				for (const UE::Dataflow::FPin& Pin : DataflowNode->GetPins())
 				{
-					UEdGraphPin* const EdPin = CreatePin(Dataflow::Private::DataflowDirectionToEdPinDirection(Pin.Direction), Pin.Type, Pin.Name);
+					UEdGraphPin* const EdPin = CreatePin(UE::Dataflow::Private::DataflowDirectionToEdPinDirection(Pin.Direction), Pin.Type, Pin.Name);
 					EdPin->bHidden = Pin.bHidden;
 				}
 			}
@@ -215,9 +215,9 @@ void UDataflowEdNode::UpdatePinsFromDataflowNode()
 				}
 				PinsToRemove.Reset();
 
-				for (const Dataflow::FPin& Pin : DataflowNode->GetPins())
+				for (const UE::Dataflow::FPin& Pin : DataflowNode->GetPins())
 				{
-					const EEdGraphPinDirection EdDirection = Dataflow::Private::DataflowDirectionToEdPinDirection(Pin.Direction);
+					const EEdGraphPinDirection EdDirection = UE::Dataflow::Private::DataflowDirectionToEdPinDirection(Pin.Direction);
 					UEdGraphPin* EdPin = FindPin(Pin.Name, EdDirection);
 					if (!EdPin)
 					{
@@ -240,14 +240,14 @@ void UDataflowEdNode::AddOptionPin()
 
 		if (const TSharedPtr<FDataflowNode> DataflowNode = DataflowGraph->FindBaseNode(DataflowNodeGuid))
 		{
-			const TArray<Dataflow::FPin> AddedPins = DataflowNode->AddPins();
-			for (const Dataflow::FPin& Pin : AddedPins)
+			const TArray<UE::Dataflow::FPin> AddedPins = DataflowNode->AddPins();
+			for (const UE::Dataflow::FPin& Pin : AddedPins)
 			{
 				switch (Pin.Direction)
 				{
-				case Dataflow::FPin::EDirection::INPUT:
-				case Dataflow::FPin::EDirection::OUTPUT:
-					CreatePin(Dataflow::Private::DataflowDirectionToEdPinDirection(Pin.Direction), Pin.Type, Pin.Name);
+				case UE::Dataflow::FPin::EDirection::INPUT:
+				case UE::Dataflow::FPin::EDirection::OUTPUT:
+					CreatePin(UE::Dataflow::Private::DataflowDirectionToEdPinDirection(Pin.Direction), Pin.Type, Pin.Name);
 					ReconstructNode();
 					break;
 				default:
@@ -274,15 +274,15 @@ void UDataflowEdNode::RemoveOptionPin()
 
 		if (const TSharedPtr<FDataflowNode> DataflowNode = DataflowGraph->FindBaseNode(DataflowNodeGuid))
 		{
-			const TArray<Dataflow::FPin> RemovePins = DataflowNode->GetPinsToRemove();
-			Dataflow::FDataflowNodePauseInvalidationScope PauseInvalidationScope(DataflowNode.Get()); // Don't call invalidations per pin. Nodes may not evaluate correctly until all pins have been removed.
-			for (const Dataflow::FPin& Pin : RemovePins)
+			const TArray<UE::Dataflow::FPin> RemovePins = DataflowNode->GetPinsToRemove();
+			UE::Dataflow::FDataflowNodePauseInvalidationScope PauseInvalidationScope(DataflowNode.Get()); // Don't call invalidations per pin. Nodes may not evaluate correctly until all pins have been removed.
+			for (const UE::Dataflow::FPin& Pin : RemovePins)
 			{
 				switch (Pin.Direction)
 				{
-				case Dataflow::FPin::EDirection::INPUT:
-				case Dataflow::FPin::EDirection::OUTPUT:
-					if (UEdGraphPin* const EdPin = FindPin(Pin.Name, Dataflow::Private::DataflowDirectionToEdPinDirection(Pin.Direction)))
+				case UE::Dataflow::FPin::EDirection::INPUT:
+				case UE::Dataflow::FPin::EDirection::OUTPUT:
+					if (UEdGraphPin* const EdPin = FindPin(Pin.Name, UE::Dataflow::Private::DataflowDirectionToEdPinDirection(Pin.Direction)))
 					{
 						constexpr bool bNotifyNodes = true;
 						EdPin->BreakAllPinLinks(bNotifyNodes);
@@ -453,7 +453,7 @@ FLinearColor UDataflowEdNode::GetNodeTitleColor() const
 		{
 			if (TSharedPtr<FDataflowNode> DataflowNode = DataflowGraph->FindBaseNode(DataflowNodeGuid))
 			{
-				return Dataflow::FNodeColorsRegistry::Get().GetNodeTitleColor(DataflowNode->GetCategory());
+				return UE::Dataflow::FNodeColorsRegistry::Get().GetNodeTitleColor(DataflowNode->GetCategory());
 			}
 		}
 	}
@@ -464,7 +464,7 @@ FLinearColor UDataflowEdNode::GetNodeBodyTintColor() const
 {
 	if (const TSharedPtr<const FDataflowNode> DataflowNode = GetDataflowNode())
 	{
-		Dataflow::FNodeColorsRegistry::Get().GetNodeBodyTintColor(DataflowNode->GetCategory());
+		UE::Dataflow::FNodeColorsRegistry::Get().GetNodeBodyTintColor(DataflowNode->GetCategory());
 	}
 	return FDataflowNode::DefaultNodeBodyTintColor;
 }
@@ -486,7 +486,7 @@ FText UDataflowEdNode::GetPinDisplayName(const UEdGraphPin* Pin) const
 	{
 		if (const TSharedPtr<const FDataflowNode> DataflowNode = GetDataflowNode())
 		{
-			const FText DisplayName = DataflowNode->GetPinDisplayName(Pin->PinName, Dataflow::Private::EdPinDirectionToDataflowDirection(Pin->Direction));
+			const FText DisplayName = DataflowNode->GetPinDisplayName(Pin->PinName, UE::Dataflow::Private::EdPinDirectionToDataflowDirection(Pin->Direction));
 			if (!DisplayName.IsEmpty())
 			{
 				return DisplayName;
@@ -500,7 +500,7 @@ void UDataflowEdNode::GetPinHoverText(const UEdGraphPin& Pin, FString& HoverText
 {
 	if (const TSharedPtr<const FDataflowNode> DataflowNode = GetDataflowNode())
 	{
-		const Dataflow::FPin::EDirection PinDirection = Dataflow::Private::EdPinDirectionToDataflowDirection(Pin.Direction);
+		const UE::Dataflow::FPin::EDirection PinDirection = UE::Dataflow::Private::EdPinDirectionToDataflowDirection(Pin.Direction);
 
 		FString MetaDataStr;
 		TArray<FString> PinMetaData = DataflowNode->GetPinMetaData(Pin.PinName, PinDirection);
@@ -616,7 +616,7 @@ void UDataflowEdNode::OnPinRemoved(UEdGraphPin* InRemovedPin)
 			{
 				if (FDataflowInput* Con = DataflowNode->FindInput(FName(InRemovedPin->GetName())))
 				{
-					const Dataflow::FPin Pin = { Dataflow::FPin::EDirection::INPUT, Con->GetType(), Con->GetName() };
+					const UE::Dataflow::FPin Pin = { UE::Dataflow::FPin::EDirection::INPUT, Con->GetType(), Con->GetName() };
 					DataflowNode->OnPinRemoved(Pin);
 					DataflowNode->UnregisterPinConnection(Pin);
 				}
@@ -625,7 +625,7 @@ void UDataflowEdNode::OnPinRemoved(UEdGraphPin* InRemovedPin)
 			{
 				if (FDataflowOutput* Con = DataflowNode->FindOutput(FName(InRemovedPin->GetName())))
 				{
-					const Dataflow::FPin Pin = { Dataflow::FPin::EDirection::OUTPUT, Con->GetType(), Con->GetName() };
+					const UE::Dataflow::FPin Pin = { UE::Dataflow::FPin::EDirection::OUTPUT, Con->GetType(), Con->GetName() };
 					DataflowNode->OnPinRemoved(Pin);
 					DataflowNode->UnregisterPinConnection(Pin);
 				}
@@ -770,13 +770,13 @@ bool UDataflowEdNode::IsInputPinShown(FName PinName) const
 #endif //WITH_EDITOR
 
 
-TArray<Dataflow::FRenderingParameter> UDataflowEdNode::GetRenderParameters() const
+TArray<UE::Dataflow::FRenderingParameter> UDataflowEdNode::GetRenderParameters() const
 {
 	if (TSharedPtr<const FDataflowNode> DataflowNode = GetDataflowNode())
 	{
 		return DataflowNode->GetRenderParameters();
 	}
-	return 	TArray<Dataflow::FRenderingParameter>();
+	return 	TArray<UE::Dataflow::FRenderingParameter>();
 }
 
 
diff --git a/Engine/Source/Runtime/Experimental/Dataflow/Engine/Private/Dataflow/DataflowObject.cpp b/Engine/Source/Runtime/Experimental/Dataflow/Engine/Private/Dataflow/DataflowObject.cpp
--- a/Engine/Source/Runtime/Experimental/Dataflow/Engine/Private/Dataflow/DataflowObject.cpp
+++ b/Engine/Source/Runtime/Experimental/Dataflow/Engine/Private/Dataflow/DataflowObject.cpp
@@ -33,7 +33,7 @@ FDataflowAssetEdit::~FDataflowAssetEdit()
 	PostEditCallback();
 }
 
-Dataflow::FGraph* FDataflowAssetEdit::GetGraph()
+UE::Dataflow::FGraph* FDataflowAssetEdit::GetGraph()
 {
 	if (Asset)
 	{
@@ -44,7 +44,7 @@ Dataflow::FGraph* FDataflowAssetEdit::GetGraph()
 
 UDataflow::UDataflow(const FObjectInitializer& ObjectInitializer)
 	: UEdGraph(ObjectInitializer)
-	, Dataflow(new Dataflow::FGraph())
+	, Dataflow(new UE::Dataflow::FGraph())
 {}
 
 void UDataflow::EvaluateTerminalNodeByName(FName NodeName, UObject* Asset)
@@ -102,7 +102,7 @@ void UDataflow::PostLoad()
 	}
 
 	// Resync connections (nodes might have redirected connections
-	for (const Dataflow::FLink& Link : Dataflow->GetConnections())
+	for (const UE::Dataflow::FLink& Link : Dataflow->GetConnections())
 	{
 		TSharedPtr<const FDataflowNode> OutputNode = Dataflow->FindBaseNode(Link.OutputNode);
 		TSharedPtr<const FDataflowNode> InputNode = Dataflow->FindBaseNode(Link.InputNode);
@@ -133,34 +133,34 @@ void UDataflow::PostLoad()
 	}
 #endif
 
-	LastModifiedRenderTarget = Dataflow::FTimestamp::Current();
+	LastModifiedRenderTarget = UE::Dataflow::FTimestamp::Current();
 	UObject::PostLoad();
 }
 
 void UDataflow::AddRenderTarget(TObjectPtr<const UDataflowEdNode> InNode)
 {
-	LastModifiedRenderTarget = Dataflow::FTimestamp::Current();
+	LastModifiedRenderTarget = UE::Dataflow::FTimestamp::Current();
 	check(InNode->ShouldRenderNode());
 	RenderTargets.AddUnique(InNode);
 }
 
 void UDataflow::RemoveRenderTarget(TObjectPtr<const UDataflowEdNode> InNode)
 {
-	LastModifiedRenderTarget = Dataflow::FTimestamp::Current();
+	LastModifiedRenderTarget = UE::Dataflow::FTimestamp::Current();
 	check(!InNode->ShouldRenderNode());
 	RenderTargets.Remove(InNode);
 }
 
 void UDataflow::AddWireframeRenderTarget(TObjectPtr<const UDataflowEdNode> InNode)
 {
-	LastModifiedRenderTarget = Dataflow::FTimestamp::Current();
+	LastModifiedRenderTarget = UE::Dataflow::FTimestamp::Current();
 	check(InNode->ShouldWireframeRenderNode());
 	WireframeRenderTargets.AddUnique(InNode);
 }
 
 void UDataflow::RemoveWireframeRenderTarget(TObjectPtr<const UDataflowEdNode> InNode)
 {
-	LastModifiedRenderTarget = Dataflow::FTimestamp::Current();
+	LastModifiedRenderTarget = UE::Dataflow::FTimestamp::Current();
 	check(!InNode->ShouldWireframeRenderNode());
 	WireframeRenderTargets.Remove(InNode);
 }
diff --git a/Engine/Source/Runtime/Experimental/Dataflow/Engine/Public/Dataflow/DataflowContent.h b/Engine/Source/Runtime/Experimental/Dataflow/Engine/Public/Dataflow/DataflowContent.h
--- a/Engine/Source/Runtime/Experimental/Dataflow/Engine/Public/Dataflow/DataflowContent.h
+++ b/Engine/Source/Runtime/Experimental/Dataflow/Engine/Public/Dataflow/DataflowContent.h
@@ -18,7 +18,7 @@ class UDataflowBaseContent;
 class FPreviewScene;
 class UAnimSingleNodeInstance;
 
-namespace DataflowContextHelpers
+namespace UE::DataflowContextHelpers
 {
 	// Return a new(or saved) content that can store the execution state of the graph. 
 	template<class T>
@@ -93,14 +93,14 @@ public:
 	*	LastModifiedTimestamp - State Invalidation 
 	*   Dataflow timestamp accessors can be used to see if the EvaluationContext has been invalidated. 
 	*/
-	void SetLastModifiedTimestamp(Dataflow::FTimestamp InTimestamp, bool bMakeDirty =true);
-	const Dataflow::FTimestamp& GetLastModifiedTimestamp() const { return LastModifiedTimestamp; }
+	void SetLastModifiedTimestamp(UE::Dataflow::FTimestamp InTimestamp, bool bMakeDirty =true);
+	const UE::Dataflow::FTimestamp& GetLastModifiedTimestamp() const { return LastModifiedTimestamp; }
 
 	/**  
 	*	Context - Dataflow Evaluation State
 	*   Dataflow context stores the evaluated state of the graph. 
 	*/
-	virtual void SetDataflowContext(const TSharedPtr<Dataflow::FEngineContext>& InContext) override;
+	virtual void SetDataflowContext(const TSharedPtr<UE::Dataflow::FEngineContext>& InContext) override;
  
 	/** Rebuild the owner dependent datas  */
 	void UpdateContentDatas();
@@ -155,7 +155,7 @@ protected:
 	TObjectPtr<UObject> TerminalAsset = nullptr;
 
     /** Last data flow evaluated node time stamp */
-	Dataflow::FTimestamp LastModifiedTimestamp = Dataflow::FTimestamp::Invalid;
+	UE::Dataflow::FTimestamp LastModifiedTimestamp = UE::Dataflow::FTimestamp::Invalid;
 
     /** Dirty flag to trigger rendering. Do we need that? since when accessing the member by non const ref we will not dirty it */
 	UPROPERTY()
@@ -258,3 +258,4 @@ protected:
 	/** Boolean to control if the animation asset could be edited or not */
 	bool bHideAnimationAsset = false;
 };
+
diff --git a/Engine/Source/Runtime/Experimental/Dataflow/Engine/Public/Dataflow/DataflowContextObject.h b/Engine/Source/Runtime/Experimental/Dataflow/Engine/Public/Dataflow/DataflowContextObject.h
--- a/Engine/Source/Runtime/Experimental/Dataflow/Engine/Public/Dataflow/DataflowContextObject.h
+++ b/Engine/Source/Runtime/Experimental/Dataflow/Engine/Public/Dataflow/DataflowContextObject.h
@@ -9,7 +9,7 @@
 
 struct FManagedArrayCollection;
 class UDataflow;
-namespace Dataflow
+namespace UE::Dataflow
 {
 	class IDataflowConstructionViewMode;
 	template<class Base> class TEngineContext;
@@ -46,16 +46,16 @@ public:
 	TSharedPtr<const FManagedArrayCollection> GetRenderCollection() const { return RenderCollection; }
 
 	/** ViewMode Access */
-	void SetConstructionViewMode(const Dataflow::IDataflowConstructionViewMode* InMode) { ConstructionViewMode = InMode; }
-	const Dataflow::IDataflowConstructionViewMode* GetConstructionViewMode() const { return ConstructionViewMode; }
+	void SetConstructionViewMode(const UE::Dataflow::IDataflowConstructionViewMode* InMode) { ConstructionViewMode = InMode; }
+	const UE::Dataflow::IDataflowConstructionViewMode* GetConstructionViewMode() const { return ConstructionViewMode; }
 	
 	/**
 	*	Context - Dataflow Evaluation State
 	*   Dataflow context stores the evaluated state of the graph.
 	*/
-	virtual void SetDataflowContext(const TSharedPtr<Dataflow::FEngineContext>& InDataflowContext) { DataflowContext = InDataflowContext; }
-	const TSharedPtr<Dataflow::FEngineContext>& GetDataflowContext() const { return DataflowContext; }
-	TSharedPtr<Dataflow::FEngineContext>& GetDataflowContext() { return DataflowContext; }
+	virtual void SetDataflowContext(const TSharedPtr<UE::Dataflow::FEngineContext>& InDataflowContext) { DataflowContext = InDataflowContext; }
+	const TSharedPtr<UE::Dataflow::FEngineContext>& GetDataflowContext() const { return DataflowContext; }
+	TSharedPtr<UE::Dataflow::FEngineContext>& GetDataflowContext() { return DataflowContext; }
 
 	/** Dataflow Asset */
 	virtual void SetDataflowAsset(const TObjectPtr<UDataflow>& InAsset) { DataflowGraph = InAsset; }
@@ -84,10 +84,10 @@ protected:
 	TObjectPtr<UDataflowEdNode> SelectedNode = nullptr;
 
 	/** Construction view mode for the context object @todo(michael) : is it only for construction or for simulation as well*/
-	const Dataflow::IDataflowConstructionViewMode* ConstructionViewMode = nullptr;
+	const UE::Dataflow::IDataflowConstructionViewMode* ConstructionViewMode = nullptr;
 
 	/** Engine context (data flow owner/asset) to be used for dataflow evaluation */
-	TSharedPtr<Dataflow::FEngineContext> DataflowContext = nullptr;
+	TSharedPtr<UE::Dataflow::FEngineContext> DataflowContext = nullptr;
 
 	/** Dataflow graph for evaluation */
 	UPROPERTY(Transient, SkipSerialization)
diff --git a/Engine/Source/Runtime/Experimental/Dataflow/Engine/Public/Dataflow/DataflowEdNode.h b/Engine/Source/Runtime/Experimental/Dataflow/Engine/Public/Dataflow/DataflowEdNode.h
--- a/Engine/Source/Runtime/Experimental/Dataflow/Engine/Public/Dataflow/DataflowEdNode.h
+++ b/Engine/Source/Runtime/Experimental/Dataflow/Engine/Public/Dataflow/DataflowEdNode.h
@@ -19,7 +19,7 @@
 class FArchive;
 class UEdGraphPin;
 class UObject;
-namespace Dataflow { class FGraph; class FRenderingParameters; class IDataflowConstructionViewMode;  }
+namespace UE::Dataflow { class FGraph; class FRenderingParameters; class IDataflowConstructionViewMode;  }
 namespace GeometryCollection::Facades { class FRenderingFacade; }
 
 UCLASS(MinimalAPI)
@@ -28,7 +28,7 @@ class UDataflowEdNode : public UEdGraphNode
 	GENERATED_UCLASS_BODY()
 
 	FGuid DataflowNodeGuid;
-	TSharedPtr<Dataflow::FGraph> DataflowGraph;
+	TSharedPtr<UE::Dataflow::FGraph> DataflowGraph;
 
 public:
 
@@ -59,9 +59,9 @@ public:
 
 	bool IsBound() { return DataflowGraph && DataflowNodeGuid.IsValid(); }
 
-	TSharedPtr<Dataflow::FGraph> GetDataflowGraph() { return DataflowGraph; }
-	TSharedPtr<const Dataflow::FGraph> GetDataflowGraph() const { return DataflowGraph; }
-	void SetDataflowGraph(TSharedPtr<Dataflow::FGraph> InDataflowGraph) { DataflowGraph = InDataflowGraph; }
+	TSharedPtr<UE::Dataflow::FGraph> GetDataflowGraph() { return DataflowGraph; }
+	TSharedPtr<const UE::Dataflow::FGraph> GetDataflowGraph() const { return DataflowGraph; }
+	void SetDataflowGraph(TSharedPtr<UE::Dataflow::FGraph> InDataflowGraph) { DataflowGraph = InDataflowGraph; }
 
 	DATAFLOWENGINE_API void UpdatePinsFromDataflowNode();
 
@@ -99,7 +99,7 @@ public:
 	DATAFLOWENGINE_API void SetCanEnableWireframeRenderNode(bool bInCanEnable);
 	DATAFLOWENGINE_API bool CanEnableWireframeRenderNode() const;
 
-	DATAFLOWENGINE_API TArray<Dataflow::FRenderingParameter> GetRenderParameters() const;
+	DATAFLOWENGINE_API TArray<UE::Dataflow::FRenderingParameter> GetRenderParameters() const;
 
 private:
 
diff --git a/Engine/Source/Runtime/Experimental/Dataflow/Engine/Public/Dataflow/DataflowObject.h b/Engine/Source/Runtime/Experimental/Dataflow/Engine/Public/Dataflow/DataflowObject.h
--- a/Engine/Source/Runtime/Experimental/Dataflow/Engine/Public/Dataflow/DataflowObject.h
+++ b/Engine/Source/Runtime/Experimental/Dataflow/Engine/Public/Dataflow/DataflowObject.h
@@ -20,7 +20,7 @@ class UDataflow;
 class UObject;
 class UDataflowEdNode;
 class UMaterial;
-namespace Dataflow { class FGraph; }
+namespace UE::Dataflow { class FGraph; }
 
 
 /**
@@ -42,7 +42,7 @@ public:
 	DATAFLOWENGINE_API FDataflowAssetEdit(UDataflow *InAsset, FPostEditFunctionCallback InCallable);
 	DATAFLOWENGINE_API ~FDataflowAssetEdit();
 
-	DATAFLOWENGINE_API Dataflow::FGraph* GetGraph();
+	DATAFLOWENGINE_API UE::Dataflow::FGraph* GetGraph();
 
 private:
 	FPostEditFunctionCallback PostEditCallback;
@@ -63,7 +63,7 @@ enum class EDataflowType : uint8
 /**
 * UDataflow (UObject)
 *
-* UObject wrapper for the Dataflow::FGraph
+* UObject wrapper for the UE::Dataflow::FGraph
 *
 */
 UCLASS(BlueprintType, customconstructor, MinimalAPI)
@@ -71,10 +71,10 @@ class UDataflow : public UEdGraph
 {
 	GENERATED_UCLASS_BODY()
 
-	Dataflow::FTimestamp LastModifiedRenderTarget = Dataflow::FTimestamp::Invalid; 
+	UE::Dataflow::FTimestamp LastModifiedRenderTarget = UE::Dataflow::FTimestamp::Invalid; 
 	TArray< TObjectPtr<const UDataflowEdNode> > RenderTargets; // Not Serialized
 	TArray< TObjectPtr<const UDataflowEdNode> > WireframeRenderTargets; // Not Serialized
-	TSharedPtr<Dataflow::FGraph, ESPMode::ThreadSafe> Dataflow;
+	TSharedPtr<UE::Dataflow::FGraph, ESPMode::ThreadSafe> Dataflow;
 	DATAFLOWENGINE_API void PostEditCallback();
 
 public:
@@ -118,8 +118,8 @@ public:
 	DATAFLOWENGINE_API void Serialize(FArchive& Ar);
 
 	/** Accessors for internal geometry collection */
-	TSharedPtr<const Dataflow::FGraph, ESPMode::ThreadSafe> GetDataflow() const { return Dataflow; }
-	TSharedPtr<Dataflow::FGraph, ESPMode::ThreadSafe> GetDataflow() { return Dataflow; }
+	TSharedPtr<const UE::Dataflow::FGraph, ESPMode::ThreadSafe> GetDataflow() const { return Dataflow; }
+	TSharedPtr<UE::Dataflow::FGraph, ESPMode::ThreadSafe> GetDataflow() { return Dataflow; }
 
 	/**Editing the collection should only be through the edit object.*/
 	FDataflowAssetEdit EditDataflow() const {
@@ -142,7 +142,7 @@ public:
 	DATAFLOWENGINE_API void RemoveWireframeRenderTarget(TObjectPtr<const UDataflowEdNode>);
 	const TArray< TObjectPtr<const UDataflowEdNode> >& GetWireframeRenderTargets() const { return WireframeRenderTargets; }
 
-	const Dataflow::FTimestamp& GetRenderingTimestamp() const { return LastModifiedRenderTarget; }
+	const UE::Dataflow::FTimestamp& GetRenderingTimestamp() const { return LastModifiedRenderTarget; }
 
 #if WITH_EDITORONLY_DATA
 
@@ -182,3 +182,4 @@ private:
 #endif
 };
 
+
diff --git a/Engine/Source/Runtime/Experimental/Dataflow/Engine/Public/Dataflow/DataflowObjectInterface.h b/Engine/Source/Runtime/Experimental/Dataflow/Engine/Public/Dataflow/DataflowObjectInterface.h
--- a/Engine/Source/Runtime/Experimental/Dataflow/Engine/Public/Dataflow/DataflowObjectInterface.h
+++ b/Engine/Source/Runtime/Experimental/Dataflow/Engine/Public/Dataflow/DataflowObjectInterface.h
@@ -8,7 +8,7 @@
 class UObject;
 class FArchive;
 
-namespace Dataflow
+namespace UE::Dataflow
 {
 	template<class Base = FContextSingle>
 	class TEngineContext : public Base
@@ -36,3 +36,4 @@ namespace Dataflow
 	typedef TEngineContext<FContextSingle> FEngineContext;
 	typedef TEngineContext<FContextThreaded> FEngineContextThreaded;
 }
+
diff --git a/Engine/Source/Runtime/Experimental/Dataflow/Simulation/Private/Dataflow/DataflowSimulationContext.cpp b/Engine/Source/Runtime/Experimental/Dataflow/Simulation/Private/Dataflow/DataflowSimulationContext.cpp
--- a/Engine/Source/Runtime/Experimental/Dataflow/Simulation/Private/Dataflow/DataflowSimulationContext.cpp
+++ b/Engine/Source/Runtime/Experimental/Dataflow/Simulation/Private/Dataflow/DataflowSimulationContext.cpp
@@ -91,7 +91,7 @@ int32 TSimulationContext<Base>::NumSimulationProxies(const FString& ProxyType) c
 	return TypedProxies ? TypedProxies->Num() : 0;
 }
 
-template class TSimulationContext<::Dataflow::FContextSingle>;
-template class TSimulationContext<::Dataflow::FContextThreaded>;
+template class TSimulationContext<UE::Dataflow::FContextSingle>;
+template class TSimulationContext<UE::Dataflow::FContextThreaded>;
 	
 }
\ No newline at end of file
diff --git a/Engine/Source/Runtime/Experimental/Dataflow/Simulation/Private/Dataflow/DataflowSimulationManager.cpp b/Engine/Source/Runtime/Experimental/Dataflow/Simulation/Private/Dataflow/DataflowSimulationManager.cpp
--- a/Engine/Source/Runtime/Experimental/Dataflow/Simulation/Private/Dataflow/DataflowSimulationManager.cpp
+++ b/Engine/Source/Runtime/Experimental/Dataflow/Simulation/Private/Dataflow/DataflowSimulationManager.cpp
@@ -79,7 +79,7 @@ namespace UE::Dataflow
 #if CHAOS_DEBUG_DRAW
 				ChaosDD::Private::FChaosDDScopeTaskContext DDTaskContext(ParentDDContext);
 #endif
-				Dataflow::EvaluateSimulationGraph(DataflowAsset, SimulationContext, DeltaTime, SimulationTime);
+				UE::Dataflow::EvaluateSimulationGraph(DataflowAsset, SimulationContext, DeltaTime, SimulationTime);
 			}
 
 		private:
@@ -87,7 +87,7 @@ namespace UE::Dataflow
 			TObjectPtr<UDataflow> DataflowAsset;
 
 			/** Simulation context */
-			TSharedPtr<Dataflow::FDataflowSimulationContext> SimulationContext;
+			TSharedPtr<UE::Dataflow::FDataflowSimulationContext> SimulationContext;
 
 			/** Delta time used to advance the simulation */
 			float DeltaTime;
diff --git a/Engine/Source/Runtime/Experimental/Dataflow/Simulation/Private/Dataflow/DataflowSimulationPlugin.cpp b/Engine/Source/Runtime/Experimental/Dataflow/Simulation/Private/Dataflow/DataflowSimulationPlugin.cpp
--- a/Engine/Source/Runtime/Experimental/Dataflow/Simulation/Private/Dataflow/DataflowSimulationPlugin.cpp
+++ b/Engine/Source/Runtime/Experimental/Dataflow/Simulation/Private/Dataflow/DataflowSimulationPlugin.cpp
@@ -10,11 +10,11 @@
 
 void IDataflowSimulationPlugin::StartupModule()
 {
-	Dataflow::RegisterDataflowSimulationNodes();
+	UE::Dataflow::RegisterDataflowSimulationNodes();
 
-	Dataflow::RegisterNodeFilter(FDataflowSimulationNode::StaticType());
-	Dataflow::RegisterNodeFilter(FDataflowInvalidNode::StaticType());
-	Dataflow::RegisterNodeFilter(FDataflowExecutionNode::StaticType());
+	UE::Dataflow::RegisterNodeFilter(FDataflowSimulationNode::StaticType());
+	UE::Dataflow::RegisterNodeFilter(FDataflowInvalidNode::StaticType());
+	UE::Dataflow::RegisterNodeFilter(FDataflowExecutionNode::StaticType());
 
 	UDataflowSimulationManager::OnStartup();
 }
diff --git a/Engine/Source/Runtime/Experimental/Dataflow/Simulation/Private/Dataflow/DataflowSimulationUtils.cpp b/Engine/Source/Runtime/Experimental/Dataflow/Simulation/Private/Dataflow/DataflowSimulationUtils.cpp
--- a/Engine/Source/Runtime/Experimental/Dataflow/Simulation/Private/Dataflow/DataflowSimulationUtils.cpp
+++ b/Engine/Source/Runtime/Experimental/Dataflow/Simulation/Private/Dataflow/DataflowSimulationUtils.cpp
@@ -18,7 +18,7 @@ namespace UE::Dataflow
 			
 			if(SimulationGraph)
 			{
-				if(const TSharedPtr<::Dataflow::FGraph> DataflowGraph = SimulationGraph->GetDataflow())
+				if(const TSharedPtr<UE::Dataflow::FGraph> DataflowGraph = SimulationGraph->GetDataflow())
 				{
 					// Invalidation of all the simulation nodes that are always dirty
 					for(const TSharedPtr<FDataflowNode>& InvalidNode : DataflowGraph->GetFilteredNodes(FDataflowInvalidNode::StaticType()))
diff --git a/Engine/Source/Runtime/Experimental/Dataflow/Simulation/Public/Dataflow/DataflowSimulationContext.h b/Engine/Source/Runtime/Experimental/Dataflow/Simulation/Public/Dataflow/DataflowSimulationContext.h
--- a/Engine/Source/Runtime/Experimental/Dataflow/Simulation/Public/Dataflow/DataflowSimulationContext.h
+++ b/Engine/Source/Runtime/Experimental/Dataflow/Simulation/Public/Dataflow/DataflowSimulationContext.h
@@ -10,11 +10,11 @@
 namespace UE::Dataflow
 {
 	/** Simulation context that will be used by all the simulation nodes*/
-	template<typename Base = ::Dataflow::FContextSingle>
-	class DATAFLOWSIMULATION_API TSimulationContext : public ::Dataflow::TEngineContext<Base>
+	template<typename Base = UE::Dataflow::FContextSingle>
+	class DATAFLOWSIMULATION_API TSimulationContext : public UE::Dataflow::TEngineContext<Base>
 	{
 	public:
-		DATAFLOW_CONTEXT_INTERNAL(::Dataflow::TEngineContext<Base>, TSimulationContext);
+		DATAFLOW_CONTEXT_INTERNAL(UE::Dataflow::TEngineContext<Base>, TSimulationContext);
 
 		explicit TSimulationContext(const TObjectPtr<UObject>& InOwner)
 				: Super(InOwner)
@@ -114,6 +114,6 @@ namespace UE::Dataflow
 		}
 	}
 
-	using FDataflowSimulationContext = TSimulationContext<::Dataflow::FContextSingle>;
-	using FDataflowSimulationContextThreaded =  TSimulationContext<::Dataflow::FContextThreaded>;
+	using FDataflowSimulationContext = TSimulationContext<UE::Dataflow::FContextSingle>;
+	using FDataflowSimulationContextThreaded =  TSimulationContext<UE::Dataflow::FContextThreaded>;
 }
diff --git a/Engine/Source/Runtime/Experimental/Dataflow/Simulation/Public/Dataflow/DataflowSimulationNodes.h b/Engine/Source/Runtime/Experimental/Dataflow/Simulation/Public/Dataflow/DataflowSimulationNodes.h
--- a/Engine/Source/Runtime/Experimental/Dataflow/Simulation/Public/Dataflow/DataflowSimulationNodes.h
+++ b/Engine/Source/Runtime/Experimental/Dataflow/Simulation/Public/Dataflow/DataflowSimulationNodes.h
@@ -34,7 +34,7 @@ struct FDataflowSimulationNode : public FDataflowNode
 	FDataflowSimulationNode()
 		: Super() { }
 
-	FDataflowSimulationNode(const Dataflow::FNodeParameters& Param, FGuid InGuid = FGuid::NewGuid())
+	FDataflowSimulationNode(const UE::Dataflow::FNodeParameters& Param, FGuid InGuid = FGuid::NewGuid())
 		: Super(Param,InGuid) {
 	}
 
@@ -50,7 +50,7 @@ struct FDataflowSimulationNode : public FDataflowNode
 	}
 
 	/** Evaluate simulation dispatch */
-	virtual void Evaluate(Dataflow::FContext& Context, const FDataflowOutput* Output) const override
+	virtual void Evaluate(UE::Dataflow::FContext& Context, const FDataflowOutput* Output) const override
 	{
 		if(Context.IsA(UE::Dataflow::FDataflowSimulationContext::StaticType()))
 		{
@@ -80,7 +80,7 @@ struct FDataflowInvalidNode : public FDataflowSimulationNode
 	FDataflowInvalidNode()
 		: Super() { }
 
-	FDataflowInvalidNode(const Dataflow::FNodeParameters& Param, FGuid InGuid = FGuid::NewGuid())
+	FDataflowInvalidNode(const UE::Dataflow::FNodeParameters& Param, FGuid InGuid = FGuid::NewGuid())
 		: Super(Param,InGuid) {
 	}
 
@@ -110,7 +110,7 @@ struct FDataflowExecutionNode : public FDataflowSimulationNode
 	FDataflowExecutionNode()
 		: Super() { }
 
-	FDataflowExecutionNode(const Dataflow::FNodeParameters& Param, FGuid InGuid = FGuid::NewGuid())
+	FDataflowExecutionNode(const UE::Dataflow::FNodeParameters& Param, FGuid InGuid = FGuid::NewGuid())
 		: Super(Param,InGuid) {
 	}
 
@@ -161,7 +161,7 @@ struct DATAFLOWSIMULATION_API FGetSimulationTimeDataflowNode : public FDataflowI
 
 public:
 	
-	FGetSimulationTimeDataflowNode(const Dataflow::FNodeParameters& InParam, FGuid InGuid = FGuid::NewGuid())
+	FGetSimulationTimeDataflowNode(const UE::Dataflow::FNodeParameters& InParam, FGuid InGuid = FGuid::NewGuid())
 	: FDataflowInvalidNode(InParam, InGuid)
 	{
 		RegisterOutputConnection(&SimulationTime);
@@ -184,7 +184,7 @@ struct FGetPhysicsSolversDataflowNode : public FDataflowInvalidNode
 
 public:
 	
-	FGetPhysicsSolversDataflowNode(const Dataflow::FNodeParameters& InParam, FGuid InGuid = FGuid::NewGuid())
+	FGetPhysicsSolversDataflowNode(const UE::Dataflow::FNodeParameters& InParam, FGuid InGuid = FGuid::NewGuid())
 	: FDataflowInvalidNode(InParam, InGuid)
 	{
 		RegisterOutputConnection(&PhysicsSolvers);
@@ -211,7 +211,7 @@ struct FAdvancePhysicsSolversDataflowNode : public FDataflowSimulationNode
 
 public:
 	
-	FAdvancePhysicsSolversDataflowNode(const Dataflow::FNodeParameters& InParam, FGuid InGuid = FGuid::NewGuid())
+	FAdvancePhysicsSolversDataflowNode(const UE::Dataflow::FNodeParameters& InParam, FGuid InGuid = FGuid::NewGuid())
 	: FDataflowSimulationNode(InParam, InGuid)
 	{
 		RegisterInputConnection(&SimulationTime);
@@ -240,7 +240,7 @@ struct FFilterSimulationProxiesDataflowNode : public FDataflowSimulationNode
 
 public:
 	
-	FFilterSimulationProxiesDataflowNode(const Dataflow::FNodeParameters& InParam, FGuid InGuid = FGuid::NewGuid())
+	FFilterSimulationProxiesDataflowNode(const UE::Dataflow::FNodeParameters& InParam, FGuid InGuid = FGuid::NewGuid())
 	: FDataflowSimulationNode(InParam, InGuid)
 	{
 		RegisterInputConnection(&SimulationProxies);
@@ -272,7 +272,7 @@ struct FSimulationProxiesTerminalDataflowNode : public FDataflowExecutionNode
 
 public:
 	
-	FSimulationProxiesTerminalDataflowNode(const Dataflow::FNodeParameters& InParam, FGuid InGuid = FGuid::NewGuid())
+	FSimulationProxiesTerminalDataflowNode(const UE::Dataflow::FNodeParameters& InParam, FGuid InGuid = FGuid::NewGuid())
 	: FDataflowExecutionNode(InParam, InGuid)
 	{
 		RegisterInputConnection(&SimulationProxies);
@@ -286,7 +286,7 @@ public:
 	virtual void EvaluateSimulation(UE::Dataflow::FDataflowSimulationContext& SimulationContext, const FDataflowOutput* Output) const override;
 };
 
-namespace Dataflow
+namespace UE::Dataflow
 {
 	void RegisterDataflowSimulationNodes();
 }
diff --git a/Engine/Source/Runtime/Experimental/GeometryCollectionEngine/Private/GeometryCollection/GeometryCollectionObject.cpp b/Engine/Source/Runtime/Experimental/GeometryCollectionEngine/Private/GeometryCollection/GeometryCollectionObject.cpp
--- a/Engine/Source/Runtime/Experimental/GeometryCollectionEngine/Private/GeometryCollection/GeometryCollectionObject.cpp
+++ b/Engine/Source/Runtime/Experimental/GeometryCollectionEngine/Private/GeometryCollection/GeometryCollectionObject.cpp
@@ -561,7 +561,7 @@ void UGeometryCollection::Reset()
 	}
 }
 
-namespace Dataflow::Private
+namespace UE::Dataflow::Private
 {
 	static void SetRandomBoneColor(TSharedPtr<FGeometryCollection, ESPMode::ThreadSafe>& InGeometryCollection)
 	{
@@ -595,7 +595,7 @@ void UGeometryCollection::ResetFrom(const FManagedArrayCollection& InCollection,
 		InitializeMaterials(bHasInternalMaterials);
 
 		// Randomize BoneColor
-		Dataflow::Private::SetRandomBoneColor(GeometryCollection);
+		UE::Dataflow::Private::SetRandomBoneColor(GeometryCollection);
 	}
 }
 
@@ -611,7 +611,7 @@ void UGeometryCollection::ResetFrom(const FManagedArrayCollection& InCollection,
 		InitializeMaterials(bHasInternalMaterials);
 
 		// Randomize BoneColor
-		Dataflow::Private::SetRandomBoneColor(GeometryCollection);
+		UE::Dataflow::Private::SetRandomBoneColor(GeometryCollection);
 	}
 }
 
