NetToken
- Added support for separate indices per NetTokenDataStore to make it possible to use fewer bits for the index for types with few exports.

diff --git a/Engine/Source/Runtime/Engine/Private/PackageMapClient.cpp b/Engine/Source/Runtime/Engine/Private/PackageMapClient.cpp
--- a/Engine/Source/Runtime/Engine/Private/PackageMapClient.cpp
+++ b/Engine/Source/Runtime/Engine/Private/PackageMapClient.cpp
@@ -2245,7 +2245,7 @@ FOutBunch* UPackageMapClient::CreateCustomExportsBunch(const FOutBunch& Outgoing
 			UE_NET_TRACE_DYNAMIC_NAME_SCOPE(*NetToken.ToString(), *ExportsBunch, GetTraceCollector(*ExportsBunch), ENetTraceVerbosity::VeryVerbose);
 
 			// Write token 
-			WriteNetToken(*ExportsBunch, NetToken);
+			NetTokenResolveContext.NetTokenStore->WriteNetToken(*ExportsBunch, NetToken);
 			NetTokenResolveContext.NetTokenStore->WriteTokenData(*ExportsBunch, NetToken);
 
 			// Track what exports we exported until we know what packet we actually did commit to.
@@ -2279,7 +2279,7 @@ void UPackageMapClient::ReceiveCustomExportsBunch(FInBunch &InBunch)
 
 	while (bHasExportsToRead && !InBunch.IsError())
 	{
-		FNetToken ImportedNetToken = ReadNetToken(InBunch);
+		FNetToken ImportedNetToken = NetTokenStore->ReadNetToken(InBunch);
 
 		NetTokenStore->ReadTokenData(InBunch, ImportedNetToken, *RemoteNetTokenStoreState);
 
diff --git a/Engine/Source/Runtime/Experimental/Iris/Core/Private/Iris/ReplicationSystem/NameTokenStore.cpp b/Engine/Source/Runtime/Experimental/Iris/Core/Private/Iris/ReplicationSystem/NameTokenStore.cpp
--- a/Engine/Source/Runtime/Experimental/Iris/Core/Private/Iris/ReplicationSystem/NameTokenStore.cpp
+++ b/Engine/Source/Runtime/Experimental/Iris/Core/Private/Iris/ReplicationSystem/NameTokenStore.cpp
@@ -28,15 +28,14 @@ FNetToken FNameTokenStore::GetOrCreateToken(FName Name)
 	FNetTokenStoreKey Key = GetOrCreateTokenStoreKey(Name);
 	if (Key.IsValid())
 	{
-		const uint32 NameIndex = Key.GetKeyValue();
-		if (StoredTokens[NameIndex].IsValid())
+		const FNetToken ExistingToken = GetNetTokenFromKey(Key);
+		if (ExistingToken.IsValid())
 		{
-			return StoredTokens[NameIndex];
+			return ExistingToken;
 		}
 		else
 		{
-			const FNetToken NewToken = CreateToken(TokenStore.IsAuthority() ? FNetToken::ENetTokenAuthority::Authority : FNetToken::ENetTokenAuthority::None, Key, *TokenStore.GetLocalNetTokenStoreState());
-			StoredTokens[NameIndex] = NewToken;
+			const FNetToken NewToken = CreateAndStoreTokenForKey(Key);
 
 			UE_LOG_FNAMETOKEN(TEXT("FNameTokenStore::GetOrCreateToken - Created %s for %s"), *NewToken.ToString(), *Name.ToString());
 
@@ -47,22 +46,18 @@ FNetToken FNameTokenStore::GetOrCreateToken(FName Name)
 	return FNetToken();
 }
 
-FNetTokenStoreKey FNameTokenStore::GetOrCreateTokenStoreKey(FName Name)
+FNetTokenDataStore::FNetTokenStoreKey FNameTokenStore::GetOrCreateTokenStoreKey(FName Name)
 {
-	// Lock if we have to.. make that a policy
-	//FScopeLock Lock(&CriticalSection);
 	if (const FNetTokenStoreKey* ExistingKey = FNameToKey.Find(Name))
 	{
 		return *ExistingKey;
 	}
-	else if (FNameToKey.Num() < FNetToken::MaxNetTokenCount)
-	{
-		// Store bookkeeping data
-		const FNetTokenStoreKey NewKey = MakeNetTokenStoreKey(GetTypeId(), StoredFNames.Num());
 
-		FNameToKey.Add(Name, NewKey);
+	const FNetTokenStoreKey NewKey = GetNextNetTokenStoreKey();
+	if (NewKey.IsValid())
+	{
 		StoredFNames.Add(Name);
-		StoredTokens.Add(FNetToken());
+		FNameToKey.Add(Name, NewKey);
 
 		return NewKey;
 	}
@@ -71,11 +66,11 @@ FNetTokenStoreKey FNameTokenStore::GetOrCreateTokenStoreKey(FName Name)
 }
 
 FNameTokenStore::FNameTokenStore(FNetTokenStore& InTokenStore)
-: TokenStore(InTokenStore)
+: FNetTokenDataStore(InTokenStore)
 {
-	// Reserve 0
-	StoredFNames.Add(FName());
-	StoredTokens.Add(FNetToken());
+	// As we use an array for our storage we must match the size of the StoredTokens array.
+	// We assume that Index 0 is invalid.
+	StoredFNames.SetNum(StoredTokens.Num());
 }
 
 FName FNameTokenStore::ResolveToken(FNetToken Token, const FNetTokenStoreState* NetTokenStoreState) const
@@ -84,9 +79,13 @@ FName FNameTokenStore::ResolveToken(FNetToken Token, const FNetTokenStoreState*
 	if (Token.IsValid() && ensureMsgf(TokenStoreState, TEXT("FNameTokenStore::ResolveToken Needs valid remote NetTokenStoreState to resolve remote %s"), *Token.ToString()))
 	{
 		const FNetTokenStoreKey StoreKey = GetTokenKey(Token, *TokenStoreState);
-		if (StoreKey.IsValid() && StoreKey.GetKeyValue() < (uint32)StoredFNames.Num())
+		if (StoreKey.IsValid() && StoreKey.GetKeyIndex() < (uint32)StoredFNames.Num())
 		{
-			return StoredFNames[StoreKey.GetKeyValue()];
+			return StoredFNames[StoreKey.GetKeyIndex()];
+		}
+		else
+		{
+			UE_LOG(LogNetToken, Error, TEXT("FNameTokenStore::ResolveToken failed to resolve %s in NetTokenDataStore: %s"), *Token.ToString(), *GetTokenStoreName().ToString());
 		}
 	}
 
@@ -95,21 +94,21 @@ FName FNameTokenStore::ResolveToken(FNetToken Token, const FNetTokenStoreState*
 
 void FNameTokenStore::WriteTokenData(FNetSerializationContext& Context, FNetTokenStoreKey TokenStoreKey) const
 {
-	UE_NET_TRACE_DYNAMIC_NAME_SCOPE(StoredFNames[TokenStoreKey.GetKeyValue()], *Context.GetBitStreamWriter(), Context.GetTraceCollector(), ENetTraceVerbosity::VeryVerbose);
-	UE_LOG_FNAMETOKEN(TEXT("FNameTokenStore::WriteTokenData %s %s"), *(StoredTokens[TokenStoreKey.GetKeyValue()].ToString()), *(StoredFNames[TokenStoreKey.GetKeyValue()].ToString()));
+	UE_NET_TRACE_DYNAMIC_NAME_SCOPE(StoredFNames[TokenStoreKey.GetKeyIndex()], *Context.GetBitStreamWriter(), Context.GetTraceCollector(), ENetTraceVerbosity::VeryVerbose);
+	UE_LOG_FNAMETOKEN(TEXT("FNameTokenStore::WriteTokenData %s %s"), *(StoredTokens[TokenStoreKey.GetKeyIndex()].ToString()), *(StoredFNames[TokenStoreKey.GetKeyIndex()].ToString()));
 	// $TODO: $IRIS: We can be a bit smarter here and utilize the string-number split of FNames to export less data.. JIRA: UE-221753
-	WriteString(Context.GetBitStreamWriter(), StoredFNames[TokenStoreKey.GetKeyValue()].ToString());
+	WriteString(Context.GetBitStreamWriter(), StoredFNames[TokenStoreKey.GetKeyIndex()].ToString());
 }
 
 void FNameTokenStore::WriteTokenData(FArchive& Ar, FNetTokenStoreKey TokenStoreKey) const
 {
-	UE_NET_TRACE_DYNAMIC_NAME_SCOPE(StoredFNames[TokenStoreKey.GetKeyValue()], static_cast<FNetBitWriter&>(Ar), GetTraceCollector(static_cast<FNetBitWriter&>(Ar)), ENetTraceVerbosity::VeryVerbose);
-	UE_LOG_FNAMETOKEN(TEXT("FNameTokenStore::WriteTokenData %s %s"), *(StoredTokens[TokenStoreKey.GetKeyValue()].ToString()), *(StoredFNames[TokenStoreKey.GetKeyValue()].ToString()));
-	FName Name = StoredFNames[TokenStoreKey.GetKeyValue()];
+	UE_NET_TRACE_DYNAMIC_NAME_SCOPE(StoredFNames[TokenStoreKey.GetKeyIndex()], static_cast<FNetBitWriter&>(Ar), GetTraceCollector(static_cast<FNetBitWriter&>(Ar)), ENetTraceVerbosity::VeryVerbose);
+	UE_LOG_FNAMETOKEN(TEXT("FNameTokenStore::WriteTokenData %s %s"), *(StoredTokens[TokenStoreKey.GetKeyIndex()].ToString()), *(StoredFNames[TokenStoreKey.GetKeyIndex()].ToString()));
+	FName Name = StoredFNames[TokenStoreKey.GetKeyIndex()];
 	UPackageMap::StaticSerializeName(Ar, Name);
 }
 
-FNetTokenStoreKey FNameTokenStore::ReadTokenData(FNetSerializationContext& Context, const FNetToken& NetToken)
+FNetTokenDataStore::FNetTokenStoreKey FNameTokenStore::ReadTokenData(FNetSerializationContext& Context, const FNetToken& NetToken)
 {
 	UE_NET_TRACE_NAMED_DYNAMIC_NAME_SCOPE(TokenScope, FName(), *Context.GetBitStreamReader(), Context.GetTraceCollector(), ENetTraceVerbosity::VeryVerbose);
 
@@ -124,18 +123,7 @@ FNetTokenStoreKey FNameTokenStore::ReadTokenData(FNetSerializationContext& Conte
 		FName Name(Temp);
 		UE_NET_TRACE_SET_SCOPE_NAME(TokenScope, Name);
 
-		FNetTokenStoreKey Key = GetOrCreateTokenStoreKey(Name);
-
-		// If this is an authtoken and we are not the authority, update the stored key
-		if (Key.IsValid() && NetToken.IsAssignedByAuthority() && !TokenStore.IsAuthority())
-		{
-			UE_LOG_FNAMETOKEN(TEXT("NameTokenStore::ReadTokenData - Replaced local key %s with %s for name %s"), *StoredTokens[Key.GetKeyValue()].ToString(), *NetToken.ToString(), *Temp);	
-
-			// This way the next time we lookup this key during assignment use the imported authoritative key which we do not have to export
-			StoredTokens[Key.GetKeyValue()] = NetToken;
-		}
-
-		return Key;
+		return GetOrCreateTokenStoreKey(Name);
 	}
 	else
 	{
@@ -143,25 +131,14 @@ FNetTokenStoreKey FNameTokenStore::ReadTokenData(FNetSerializationContext& Conte
 	}
 }
 
-FNetTokenStoreKey FNameTokenStore::ReadTokenData(FArchive& Ar, const FNetToken& NetToken)
+FNetTokenDataStore::FNetTokenStoreKey FNameTokenStore::ReadTokenData(FArchive& Ar, const FNetToken& NetToken)
 {
 	FName Name;
 	UPackageMap::StaticSerializeName(Ar, Name);
 
 	if (!Ar.IsError())
 	{
-		FNetTokenStoreKey Key = GetOrCreateTokenStoreKey(Name);
-
-		// If this is an authtoken and we are not the authority, update the stored key
-		if (Key.IsValid() && NetToken.IsAssignedByAuthority() && !TokenStore.IsAuthority())
-		{
-			UE_LOG_FNAMETOKEN(TEXT("NameTokenStore::ReadTokenData - Replaced local key %s with %s for name %s"), *StoredTokens[Key.GetKeyValue()].ToString(), *NetToken.ToString(), *Name.ToString());	
-
-			// This way the next time we lookup this key during assignment use the imported authoritative key which we do not have to export
-			StoredTokens[Key.GetKeyValue()] = NetToken;
-		}
-
-		return Key;
+		return GetOrCreateTokenStoreKey(Name);
 	}
 	else
 	{
diff --git a/Engine/Source/Runtime/Experimental/Iris/Core/Private/Iris/ReplicationSystem/NetTokenDataStream.cpp b/Engine/Source/Runtime/Experimental/Iris/Core/Private/Iris/ReplicationSystem/NetTokenDataStream.cpp
--- a/Engine/Source/Runtime/Experimental/Iris/Core/Private/Iris/ReplicationSystem/NetTokenDataStream.cpp
+++ b/Engine/Source/Runtime/Experimental/Iris/Core/Private/Iris/ReplicationSystem/NetTokenDataStream.cpp
@@ -8,7 +8,6 @@
 #include "Iris/Serialization/NetBitStreamUtil.h"
 #include "Iris/Serialization/NetExportContext.h"
 #include "Iris/ReplicationSystem/NetExports.h"
-#include "Iris/ReplicationSystem/NetTokenStoreState.h"
 #include "Iris/ReplicationSystem/ReplicationSystem.h"
 #include "Iris/ReplicationSystem/ReplicationSystemInternal.h"
 #include "Iris/Core/IrisLog.h"
@@ -69,12 +68,10 @@ void UNetTokenDataStream::Init(const FInitParameters& Params)
 	// $IRIS $TODO: if we want to make this into a real feature we need to expose some sort of api to mark tokens for pre-export
 	if (Private::bIrisPreExportExistingNetTokensOnConnect)
 	{
-		// Grab all existing tokens and add them for pre-export
-		const uint32 LocalTokenCount = NetTokenStore->GetLocalNetTokenStoreState()->TokenInfos.Num();	
-		NetTokensPendingExport.Reserve(LocalTokenCount - 1);
-		for (uint32 Index = 1; Index < LocalTokenCount; ++Index)
+		TArray<FNetToken> Tokens(NetTokenStore->GetAllNetTokens());
+		for (const FNetToken& Token : Tokens)
 		{
-			NetTokensPendingExport.Add(FNetToken::MakeNetToken(Index, NetTokenStore->IsAuthority() ? FNetToken::ENetTokenAuthority::Authority : FNetToken::ENetTokenAuthority::None));
+			NetTokensPendingExport.Add(Token);
 		}
 	}
 }
@@ -130,7 +127,7 @@ UDataStream::EWriteResult UNetTokenDataStream::WriteData(UE::Net::FNetSerializat
 			UE_NET_TRACE_NAMED_SCOPE(ExportScope, NetTokenExport, SubStream, SubContext.GetTraceCollector(), ENetTraceVerbosity::Verbose);
 
 			SubStream.WriteBool(true);
-			WriteNetToken(SubContext, Token);
+			NetTokenStore->WriteNetToken(SubContext, Token);
 			NetTokenStore->WriteTokenData(SubContext, Token);
 
 			if (SubStream.IsOverflown())
@@ -200,11 +197,8 @@ void UNetTokenDataStream::ReadData(UE::Net::FNetSerializationContext& Context)
 			break;
 		}
 
-		FNetToken Token = ReadNetToken(Context);
-		if (Token.IsValid())
-		{
-			NetTokenStore->ReadTokenData(Context, Token, *RemoteNetTokenStoreState);
-		}
+		FNetToken Token = NetTokenStore->ReadNetToken(Context);
+		NetTokenStore->ReadTokenData(Context, Token, *RemoteNetTokenStoreState);
 	}
 }
 
diff --git a/Engine/Source/Runtime/Experimental/Iris/Core/Private/Iris/ReplicationSystem/NetTokenStore.cpp b/Engine/Source/Runtime/Experimental/Iris/Core/Private/Iris/ReplicationSystem/NetTokenStore.cpp
--- a/Engine/Source/Runtime/Experimental/Iris/Core/Private/Iris/ReplicationSystem/NetTokenStore.cpp
+++ b/Engine/Source/Runtime/Experimental/Iris/Core/Private/Iris/ReplicationSystem/NetTokenStore.cpp
@@ -1,7 +1,6 @@
 // Copyright Epic Games, Inc. All Rights Reserved.
 
 #include "Iris/ReplicationSystem/NetTokenStore.h"
-#include "Iris/ReplicationSystem/NetTokenStoreState.h"
 #include "Iris/ReplicationSystem/ObjectReferenceCacheFwd.h"
 #include "Iris/Serialization/NetBitStreamReader.h"
 #include "Iris/Serialization/NetBitStreamWriter.h"
@@ -10,50 +9,112 @@
 #include "Iris/Serialization/InternalNetSerializationContext.h"
 #include "Iris/Serialization/NetExportContext.h"
 #include "Net/Core/Trace/NetTrace.h"
+#include "UObject/CoreNet.h"
 
 namespace UE::Net
 {
 
-FNetTokenDataStore::FNetTokenDataStore()
-: TypeId(FNetToken::InvalidTokenTypeId)
+class FNetTokenStoreState
 {
+public:
+	using FNetTokenStoreKey = FNetTokenDataStore::FNetTokenStoreKey;
+
+	// The size of the array is managed by FNetTokenDataStream
+	FNetTokenStoreState()
+	{
+		Reset();
+	}
+
+	// Reserve space for tokens
+	bool ReserveTokenCount(uint32 TypeIndex, uint32 NewCount)
+	{
+		if ((TypeIndex >= FNetToken::MaxTypeIdCount) || (NewCount >= FNetToken::MaxNetTokenCount))
+		{
+			return false;
+		}
+
+		const uint32 NewSize = FMath::Max<uint32>(TokenInfoArray[TypeIndex].Num(), NewCount);
+		TokenInfoArray[TypeIndex].SetNumZeroed(NewSize);
+
+		return true;
+	}
+
+	void Reset()
+	{
+		// We reserve the first token as an invalid token
+		for (TArray<FNetTokenStoreKey>& TokenInfos : TokenInfoArray)
+		{
+			TokenInfos.Add(FNetTokenStoreKey());
+		}
+	}
+
+	// Map from NetTokenIndex -> NetTokenDataStoreKey (Index)
+	TArray<FNetTokenDataStore::FNetTokenStoreKey> TokenInfoArray[FNetToken::MaxTypeIdCount];
+};
+
+FNetTokenDataStore::FNetTokenDataStore(FNetTokenStore& InTokenStore)
+: TokenStore(InTokenStore)
+, TypeId(FNetToken::InvalidTokenTypeId)
+{
+	// Reserve first index as invalid.
+	StoredTokens.Add(FNetToken());
 }
 
 FNetTokenDataStore::~FNetTokenDataStore()
 {
 }
 
-FNetTokenStoreKey FNetTokenDataStore::GetTokenKey(FNetToken Token, const FNetTokenStoreState& TokenStoreState) const
+FNetTokenDataStore::FNetTokenStoreKey FNetTokenDataStore::GetTokenKey(FNetToken Token, const FNetTokenStoreState& TokenStoreState) const
 {
-	if (Token.GetIndex() < (uint32)TokenStoreState.TokenInfos.Num())
+	if (Token.GetTypeId() == GetTypeId())
 	{
-		const FNetTokenStoreKey& Key = TokenStoreState.TokenInfos[Token.GetIndex()];
-		//if (ensureAlwaysMsgf(Key.IsValid() && GetTypeId() == Key.GetTypeId(), TEXT("Cannot resolve NetToken %s with TypeId: %u in DataStore with TypeId: %u"), *Token.ToString(), Key.GetTypeId(), GetTypeId()))
-		if (ensureMsgf(Key.IsValid() && GetTypeId() == Key.GetTypeId(), TEXT("Cannot resolve NetToken %s with TypeId: %u in DataStore with TypeId: %u"), *Token.ToString(), Key.GetTypeId(), GetTypeId()))
-		{
-			return Key;
-		}
-	}
+		const TArray<FNetTokenStoreKey>& TokenStoreKeysForType = TokenStoreState.TokenInfoArray[GetTypeId()];
+		const int32 TokenIndex(Token.GetIndex());
 
-	return FNetTokenStoreKey();
+		return TokenIndex < TokenStoreKeysForType.Num() ? TokenStoreKeysForType[TokenIndex] : FNetTokenStoreKey();
+	}
+	else
+	{
+		UE_LOG(LogNetToken, Error, TEXT("FNetTokenDataStore::GetTokenKey Invalid tokentype %s StoreTypeId: %u"), *Token.ToString(), GetTypeId());
+		return FNetTokenStoreKey();
+	}
 }
 
-FNetToken FNetTokenDataStore::CreateToken(FNetToken::ENetTokenAuthority Authority, FNetTokenStoreKey Key, FNetTokenStoreState& TokenStoreState)
+FNetToken FNetTokenDataStore::CreateAndStoreTokenForKey(FNetTokenStoreKey Key)
 {
-	const uint32 NextTokenIndex = TokenStoreState.TokenInfos.Num();
+	FNetTokenStoreState& LocalNetTokenStoreState = *TokenStore.GetLocalNetTokenStoreState();
+
+	// The LocalNetTokenStoreState always maps locally assigned tokens -> Key (Index)
+	const uint32 NextTokenIndex = LocalNetTokenStoreState.TokenInfoArray[GetTypeId()].Num();
 
 	if (!ensure(NextTokenIndex < FNetToken::MaxNetTokenCount))
 	{
 		return FNetToken();
 	}
 
+	FNetToken NewToken = FNetTokenStore::MakeNetToken(TypeId, NextTokenIndex, TokenStore.IsAuthority() ? FNetToken::ENetTokenAuthority::Authority : FNetToken::ENetTokenAuthority::None);
+
 	// Store token info
-	TokenStoreState.TokenInfos.Add(Key);
+	LocalNetTokenStoreState.TokenInfoArray[GetTypeId()].Add(Key);
+
+	// The stored tokens array contains the current NetToken associated with the Key (Index), but it can be updated to use an authoritative token instead.
+	StoredTokens[Key.GetKeyIndex()] = NewToken;
+
+	return NewToken;
+}
+
+void FNetTokenDataStore::StoreTokenForKey(FNetTokenStoreKey Key, FNetToken NetToken)
+{
+	// The stored tokens array contains the current NetToken associated with the Key (Index), but it can be updated to use an authoritative token instead.
+	StoredTokens[Key.GetKeyIndex()] = NetToken;
+}
 
-	return FNetToken::MakeNetToken(NextTokenIndex, Authority);
+FNetToken FNetTokenDataStore::GetNetTokenFromKey(FNetTokenStoreKey Key) const
+{
+	return StoredTokens[Key.GetKeyIndex()];
 }
 
-FNetToken ReadNetToken(UE::Net::FNetSerializationContext& Context)
+FNetToken FNetTokenStore::InternalReadNetToken(UE::Net::FNetSerializationContext& Context, FNetToken::FTypeId TokenTypeId)
 {
 	FNetBitStreamReader* Reader = Context.GetBitStreamReader();
 
@@ -64,9 +125,14 @@ FNetToken ReadNetToken(UE::Net::FNetSerializationContext& Context)
 	if (const bool bIsValid = (TokenIndex != FNetToken::InvalidTokenIndex))
 	{
 		const bool bIsAssignedByAuthority = Reader->ReadBool();
+		if (TokenTypeId == FNetToken::InvalidTokenTypeId)
+		{
+			TokenTypeId = Reader->ReadBits(FNetToken::TokenTypeIdBits);
+		}
+
 		if (!Reader->IsOverflown())
 		{
-			ReadToken = FNetToken::MakeNetToken(TokenIndex, bIsAssignedByAuthority ? FNetToken::ENetTokenAuthority::Authority : FNetToken::ENetTokenAuthority::None);
+			ReadToken = FNetTokenStore::MakeNetToken(TokenTypeId, TokenIndex, bIsAssignedByAuthority ? FNetToken::ENetTokenAuthority::Authority : FNetToken::ENetTokenAuthority::None);
 			UE_NET_TRACE_SET_SCOPE_NAME(TokenScope, *ReadToken.ToString());
 		}
 	}
@@ -74,8 +140,7 @@ FNetToken ReadNetToken(UE::Net::FNetSerializationContext& Context)
 	return ReadToken;
 }
 
-// Note: Be cereful when modifying this methods to not affect replay compatibility.
-void WriteNetToken(UE::Net::FNetSerializationContext& Context, FNetToken Token)
+void FNetTokenStore::InternalWriteNetToken(UE::Net::FNetSerializationContext& Context, FNetToken Token, bool bWriteTokenType)
 {
 	UE_NET_TRACE_DYNAMIC_NAME_SCOPE(*Token.ToString(), *Context.GetBitStreamWriter(), Context.GetTraceCollector(), ENetTraceVerbosity::VeryVerbose);
 	FNetBitStreamWriter* Writer = Context.GetBitStreamWriter();
@@ -85,42 +150,81 @@ void WriteNetToken(UE::Net::FNetSerializationContext& Context, FNetToken Token)
 	if (TokenIndex != FNetToken::InvalidTokenIndex)
 	{
 		Writer->WriteBool(Token.IsAssignedByAuthority());
+		if (bWriteTokenType)
+		{
+			Writer->WriteBits(Token.GetTypeId(), FNetToken::TokenTypeIdBits);
+		}
 	}
 }
 
-FNetToken ReadNetToken(FArchive& Ar)
+FNetToken FNetTokenStore::InternalReadNetToken(FArchive& Ar, FNetToken::FTypeId TokenTypeId)
 {
 	FNetToken ReadToken;
 
 	uint32 TokenIndex = 0;
 	Ar.SerializeIntPacked(TokenIndex);
-
 	if (const bool bIsValid = (TokenIndex != FNetToken::InvalidTokenIndex))
 	{
 		bool bIsAssignedByAuthority = false;
 		Ar.SerializeBits(&bIsAssignedByAuthority, 1);
 
+		if (TokenTypeId == FNetToken::InvalidTokenTypeId)
+		{
+			uint32 TempTokenTypeId = 0;
+			Ar.SerializeBits(&TempTokenTypeId, FNetToken::TokenTypeIdBits);
+			TokenTypeId = TempTokenTypeId;
+		}
+
 		if (!Ar.IsError())
 		{
-			ReadToken = FNetToken::MakeNetToken(TokenIndex, bIsAssignedByAuthority ? FNetToken::ENetTokenAuthority::Authority : FNetToken::ENetTokenAuthority::None);
+			ReadToken = FNetTokenStore::MakeNetToken(TokenTypeId, TokenIndex, bIsAssignedByAuthority ? FNetToken::ENetTokenAuthority::Authority : FNetToken::ENetTokenAuthority::None);
 		}
 	}
 
 	return ReadToken;
 }
 
-// Note: Be cereful when modifying this methods to not affect replay compatibility.
-void WriteNetToken(FArchive& Ar, FNetToken Token)
+// Note: Be careful when modifying this methods to not affect replay compatibility.
+void FNetTokenStore::InternalWriteNetToken(FArchive& Ar, FNetToken Token, bool bWriteTokenType)
 {
+	UE_NET_TRACE_DYNAMIC_NAME_SCOPE(*Token.ToString(), static_cast<FNetBitWriter&>(Ar), GetTraceCollector(static_cast<FNetBitWriter&>(Ar)), ENetTraceVerbosity::VeryVerbose);
 	uint32 TokenIndex = Token.GetIndex();
 	Ar.SerializeIntPacked(TokenIndex);
 	if (TokenIndex != FNetToken::InvalidTokenIndex)
 	{
 		bool bIsAssignedByAuthority = Token.IsAssignedByAuthority();
 		Ar.SerializeBits(&bIsAssignedByAuthority, 1);
+
+		if (bWriteTokenType)
+		{
+			uint32 TokenTypeId = Token.GetTypeId();
+			Ar.SerializeBits(&TokenTypeId, FNetToken::TokenTypeIdBits);
+		}
 	}
 }
 
+FNetToken FNetTokenDataStore::ReadNetToken(UE::Net::FNetSerializationContext& Context)
+{
+	return FNetTokenStore::InternalReadNetToken(Context, GetTypeId());
+}
+
+void FNetTokenDataStore::WriteNetToken(UE::Net::FNetSerializationContext& Context, FNetToken Token)
+{
+	const bool bWriteTokenTypeId = false;
+	return FNetTokenStore::InternalWriteNetToken(Context, Token, bWriteTokenTypeId);
+}
+
+FNetToken FNetTokenDataStore::ReadNetToken(FArchive& Ar)
+{
+	return FNetTokenStore::InternalReadNetToken(Ar, GetTypeId());
+}
+
+void FNetTokenDataStore::WriteNetToken(FArchive& Ar, FNetToken Token)
+{
+	const bool bWriteTokenTypeId = false;
+	FNetTokenStore::InternalWriteNetToken(Ar, Token, bWriteTokenTypeId);
+}
+
 FNetTokenStore::FNetTokenStore()
 : LocalNetTokenStoreState(new FNetTokenStoreState)
 {
@@ -142,8 +246,7 @@ void FNetTokenStore::InitRemoteNetTokenStoreState(uint32 ConnectionId)
 	{
 		if (FNetTokenStoreState* ExistingState = RemoteNetTokenStoreStates[ConnectionId].Get())
 		{
-			ExistingState->ReserveTokenCount(0U);
-
+			ExistingState->Reset();
 		}
 		else
 		{
@@ -218,12 +321,10 @@ void FNetTokenStore::WriteTokenData(FNetSerializationContext& Context, const FNe
 	{
 		FNetBitStreamWriter* Writer = Context.GetBitStreamWriter();
 
-		// Write type
-		const FNetTokenStoreKey& TokenKey = LocalNetTokenStoreState->TokenInfos[NetToken.GetIndex()];
-		Writer->WriteBits(TokenKey.TypeId, FNetToken::TokenTypeIdBits);
+		const FNetTokenStoreKey& TokenKey = LocalNetTokenStoreState->TokenInfoArray[NetToken.GetTypeId()][NetToken.GetIndex()];
 
 		// Write token data
-		TokenDataStores[TokenKey.TypeId].Get<1>().Get()->WriteTokenData(Context, TokenKey);
+		TokenDataStores[NetToken.GetTypeId()].Get<1>().Get()->WriteTokenData(Context, TokenKey);
 	}
 }
 
@@ -231,60 +332,79 @@ void FNetTokenStore::WriteTokenData(FArchive& Ar, const FNetToken NetToken) cons
 {
 	if (NetToken.IsValid())
 	{
-		// Write type
-		const FNetTokenStoreKey TokenKey = LocalNetTokenStoreState->TokenInfos[NetToken.GetIndex()];
-		uint32 TypeId = TokenKey.TypeId;
+		const FNetToken::FTypeId TokenTypeId = NetToken.GetTypeId();
+
+		// We cannot write token data for bad tokens
+		if ((TokenTypeId >= FNetToken::MaxTypeIdCount) || (NetToken.GetIndex() >= (uint32)LocalNetTokenStoreState->TokenInfoArray[TokenTypeId].Num()))
+		{
+			UE_LOG(LogNetToken, Error, TEXT("Trying to write data for unknown NetToken %s"), *NetToken.ToString());
+			return;
+		}
 
-		Ar.SerializeBits(&TypeId, FNetToken::TokenTypeIdBits);
+		const FNetTokenStoreKey TokenKey = LocalNetTokenStoreState->TokenInfoArray[NetToken.GetTypeId()][NetToken.GetIndex()];
 
 		// Write token data
-		TokenDataStores[TokenKey.TypeId].Get<1>()->WriteTokenData(Ar, TokenKey);
+		TokenDataStores[NetToken.GetTypeId()].Get<1>()->WriteTokenData(Ar, TokenKey);
 	}
 }
 
+bool FNetTokenStore::ValidateAndStoreNetTokenData(FNetTokenDataStore& DataStore, FNetTokenStoreState& RemoteNetTokenStoreState, const FNetToken NetToken, const FNetTokenStoreKey StoreKey)
+{
+	if (!StoreKey.IsValid() || !RemoteNetTokenStoreState.ReserveTokenCount(NetToken.GetTypeId(), NetToken.GetIndex() + 1))
+	{
+		return false;
+	}
+
+	TArray<FNetTokenStoreKey>& TokenStoreKeysForType = RemoteNetTokenStoreState.TokenInfoArray[NetToken.GetTypeId()];
+
+	// Since the same tokendata might be exported multiple times we better validate that it is the same data
+	const FNetTokenStoreKey& ExistingStoreKey = TokenStoreKeysForType[NetToken.GetIndex()];
+	if (!ensureAlways(!ExistingStoreKey.IsValid() || StoreKey == ExistingStoreKey))
+	{
+		return false;
+	}
+
+	// If this is an authtoken and we are not the authority, update the stored key so that we can use the authoriative key instead of the local key
+	if (NetToken.IsAssignedByAuthority() && !IsAuthority())
+	{
+		UE_LOG(LogNetToken, Verbose, TEXT("FNetTokenStore::ReadTokenData - Replaced local key %s with %s"), *DataStore.StoredTokens[StoreKey.GetKeyIndex()].ToString(), *NetToken.ToString());	
+
+		// This way the next time we lookup this key during assignment use the imported authoritative key which we do not have to export
+		DataStore.StoredTokens[StoreKey.GetKeyIndex()] = NetToken;
+	}
+	
+	// Store
+	TokenStoreKeysForType[NetToken.GetIndex()] = StoreKey;
+
+	return true;
+}
+
 void FNetTokenStore::ReadTokenData(FNetSerializationContext& Context, const FNetToken NetToken, FNetTokenStoreState& RemoteNetTokenStoreState)
 {
 	if (NetToken.IsValid())
 	{
 		FNetBitStreamReader* Reader = Context.GetBitStreamReader();
 
-		// TODO: Guard this better, a map might be better after-all!
-		// Also need to protect this if we process inbound data in parallel.
-		if (!RemoteNetTokenStoreState.ReserveTokenCount(NetToken.GetIndex() + 1))
-		{
-			Reader->DoOverflow();
-			return;
-		}
-
-		// Read type
-		FNetToken::FTypeId TokenTypeId = (FNetToken::FTypeId)Reader->ReadBits(FNetToken::TokenTypeIdBits);
+		FNetToken::FTypeId TokenTypeId = NetToken.GetTypeId();
 
-		// Validate that we managed to read and verify the type
-		if (Reader->IsOverflown() || TokenTypeId >= (uint32)TokenDataStores.Num() )
+		if (TokenTypeId >= (uint32)TokenDataStores.Num())
 		{
 			Reader->DoOverflow();
+			UE_LOG(LogNetToken, Error, TEXT("Failed to read ReadTokenData for %s."), *NetToken.ToString());
+			ensure(false);
 			return;
 		}
 
-		const FNetTokenStoreKey StoreKey = TokenDataStores[TokenTypeId].Get<1>()->ReadTokenData(Context, NetToken);
+		FNetTokenDataStore* DataStore = TokenDataStores[TokenTypeId].Get<1>().Get();
+		const FNetTokenStoreKey StoreKey = DataStore->ReadTokenData(Context, NetToken);
 
-		// Validate that we managed to read the data for the key
-		if (Reader->IsOverflown() || !StoreKey.IsValid())
+		if (!ValidateAndStoreNetTokenData(*DataStore, RemoteNetTokenStoreState, NetToken, StoreKey))
 		{
 			Reader->DoOverflow();
+			UE_LOG(LogNetToken, Error, TEXT("Failed to read ReadTokenData for %s."), *NetToken.ToString());
+			ensure(false);
 			return;
 		}
-
-		// Since the same tokendata might be exported multiple times we better validate that it is the same data
-		const FNetTokenStoreKey& ExistingStoreKey = RemoteNetTokenStoreState.TokenInfos[NetToken.GetIndex()];
-		if (!ensureAlways(!ExistingStoreKey.IsValid() || StoreKey == ExistingStoreKey))
-		{
-			Reader->DoOverflow();
-			return;
-		}
-	
-		// Store
-		RemoteNetTokenStoreState.TokenInfos[NetToken.GetIndex()] = StoreKey;
 	}
 }
 
@@ -292,48 +412,30 @@ void FNetTokenStore::ReadTokenData(FArchive& Ar, const FNetToken NetToken, FNetT
 {
 	if (NetToken.IsValid())
 	{
-		// TODO: Guard this better, a map might be better after-all!
-		// Also need to protect this if we process inbound data in parallel.
-		if (!RemoteNetTokenStoreState.ReserveTokenCount(NetToken.GetIndex() + 1))
-		{
-			Ar.SetError();
-			return;
-		}
-
 		// Read type
-		FNetToken::FTypeId TokenTypeId = 0U;
-		Ar.SerializeBits(&TokenTypeId, FNetToken::TokenTypeIdBits);
+		FNetToken::FTypeId TokenTypeId = NetToken.GetTypeId();
 
 		// Validate that we managed to read and verify the type
 		if (Ar.IsError() || TokenTypeId >= (uint32)TokenDataStores.Num() )
 		{
 			Ar.SetError();
+			UE_LOG(LogNetToken, Error, TEXT("Failed to read ReadTokenData for %s."), *NetToken.ToString());
 			return;
 		}
 
-		const FNetTokenStoreKey StoreKey = TokenDataStores[TokenTypeId].Get<1>()->ReadTokenData(Ar, NetToken);
-
-		// Validate that we managed to read the data for the key
-		if (Ar.IsError() || !StoreKey.IsValid())
-		{
-			Ar.SetError();
-			return;
-		}
+		FNetTokenDataStore* DataStore = TokenDataStores[TokenTypeId].Get<1>().Get();
+		const FNetTokenStoreKey StoreKey = DataStore->ReadTokenData(Ar, NetToken);
 
-		// Since the same tokendata might be exported multiple times we better validate that it is the same data
-		const FNetTokenStoreKey& ExistingStoreKey = RemoteNetTokenStoreState.TokenInfos[NetToken.GetIndex()];
-		if (!ensureAlways(!ExistingStoreKey.IsValid() || StoreKey == ExistingStoreKey))
+		if (!ValidateAndStoreNetTokenData(*DataStore, RemoteNetTokenStoreState, NetToken, StoreKey))
 		{
 			Ar.SetError();
+			UE_LOG(LogNetToken, Error, TEXT("Failed to read ReadTokenData for %s."), *NetToken.ToString());
+			ensure(false);
 			return;
 		}
-	
-		// Store
-		RemoteNetTokenStoreState.TokenInfos[NetToken.GetIndex()] = StoreKey;
 	}
 }
 
-
 void FNetTokenStore::ConditionalWriteNetTokenData(FNetSerializationContext& Context, Private::FNetExportContext* ExportContext, const FNetToken NetToken) const
 {
 	FNetBitStreamWriter* Writer = Context.GetBitStreamWriter();
@@ -378,36 +480,33 @@ void FNetTokenStore::ConditionalReadNetTokenData(FNetSerializationContext& Conte
 	}
 }
 
-void FNetTokenStore::AppendExportOrWriteInlinedExportData(FNetSerializationContext& Context, FNetToken NetToken)
+void FNetTokenStore::AppendExport(FNetSerializationContext& Context, FNetToken NetToken)
 {
 	using namespace UE::Net::Private;
 
 	FNetExportContext* ExportContext = Context.GetExportContext();
-	const FInternalNetSerializationContext* InternalContext = Context.GetInternalContext();
-	
-	if (InternalContext->bInlineObjectReferenceExports == 0U)
+	if (ensure(ExportContext))
 	{
 		ExportContext->AddPendingExport(NetToken);
 	}
-	else
-	{
-		FNetTokenStore* NetTokenStore = Context.GetNetTokenStore();
-		NetTokenStore->ConditionalWriteNetTokenData(Context, ExportContext, NetToken);
-	}
 }
 
-void FNetTokenStore::ReadInlinedExportData(FNetSerializationContext& Context, FNetToken NetToken)
+TArray<FNetToken> FNetTokenStore::GetAllNetTokens() const
 {
-	using namespace UE::Net::Private;
-
-	const FInternalNetSerializationContext* InternalContext = Context.GetInternalContext();
+	TArray<FNetToken> Result;
 
-	// Conditionally import if using inlined exports
-	if (InternalContext->bInlineObjectReferenceExports == 1U)
+	const FNetTokenStoreState* TokenStoreState = GetLocalNetTokenStoreState();
+	for (const TTuple<FName, TUniquePtr<FNetTokenDataStore>>& TokenDataStore : TokenDataStores)
 	{
-		FNetTokenStore* NetTokenStore = Context.GetNetTokenStore();
-		NetTokenStore->ConditionalReadNetTokenData(Context, NetToken);
+		const uint32 TypeId = TokenDataStore.Get<1>()->GetTypeId();
+		const uint32 NumTokensForType = TokenStoreState->TokenInfoArray[TypeId].Num();
+		if (NumTokensForType > 1)
+		{
+			Result.Append(MakeArrayView(&TokenDataStore.Get<1>()->StoredTokens[1], NumTokensForType - 1));
+		}
 	}
+
+	return MoveTemp(Result);
 }
 
 }
diff --git a/Engine/Source/Runtime/Experimental/Iris/Core/Private/Iris/ReplicationSystem/NetTokenStoreState.cpp b/Engine/Source/Runtime/Experimental/Iris/Core/Private/Iris/ReplicationSystem/NetTokenStoreState.cpp
deleted file mode 100644
diff --git a/Engine/Source/Runtime/Experimental/Iris/Core/Private/Iris/ReplicationSystem/NetTokenStoreState.h b/Engine/Source/Runtime/Experimental/Iris/Core/Private/Iris/ReplicationSystem/NetTokenStoreState.h
deleted file mode 100644
diff --git a/Engine/Source/Runtime/Experimental/Iris/Core/Private/Iris/ReplicationSystem/ObjectReferenceCache.cpp b/Engine/Source/Runtime/Experimental/Iris/Core/Private/Iris/ReplicationSystem/ObjectReferenceCache.cpp
--- a/Engine/Source/Runtime/Experimental/Iris/Core/Private/Iris/ReplicationSystem/ObjectReferenceCache.cpp
+++ b/Engine/Source/Runtime/Experimental/Iris/Core/Private/Iris/ReplicationSystem/ObjectReferenceCache.cpp
@@ -1233,7 +1233,8 @@ void FObjectReferenceCache::WriteFullReferenceInternal(FNetSerializationContext&
 
 		if (Writer->WriteBool(bHasPath))
 		{
-			WriteNetToken(Context, CachedObject->RelativePath);
+			// Note: As we use the StringTokenStore to write the token we also need to read it using the stringtokenstore.
+			StringTokenStore->WriteNetToken(Context, CachedObject->RelativePath);
 			NetTokenStore->ConditionalWriteNetTokenData(Context, ExportContext, CachedObject->RelativePath);
 			WriteFullReferenceInternal(Context, FNetObjectReference(CachedObject->OuterNetRefHandle));
 		}
@@ -1300,7 +1301,7 @@ void FObjectReferenceCache::ReadFullReferenceInternal(FNetSerializationContext&
 
 		if (Reader->ReadBool())
 		{
-			RelativePath = ReadNetToken(Context);
+			RelativePath = StringTokenStore->ReadNetToken(Context);
 			NetTokenStore->ConditionalReadNetTokenData(Context, RelativePath);
 			ReadFullReferenceInternal(Context, OuterRef, RecursionCount + 1U);
 		}
@@ -1386,7 +1387,7 @@ void FObjectReferenceCache::ReadFullReference(FNetSerializationContext& Context,
 		UE_NET_TRACE_SCOPE(ClientAssigned, *Reader, Context.GetTraceCollector(), ENetTraceVerbosity::Verbose);
 
 		const FNetRefHandle NetRefHandle = ReadNetRefHandle(Context);
-		FNetToken RelativePath = ReadNetToken(Context);
+		FNetToken RelativePath = StringTokenStore->ReadNetToken(Context);
 		NetTokenStore->ConditionalReadNetTokenData(Context, RelativePath);
 
 		OutRef.RefHandle = NetRefHandle;
@@ -1425,7 +1426,8 @@ void FObjectReferenceCache::WriteFullReference(FNetSerializationContext& Context
 		UE_NET_TRACE_SCOPE(ClientAssigned, *Writer, Context.GetTraceCollector(), ENetTraceVerbosity::Verbose);
 
 		WriteNetRefHandle(Context, Ref.GetRefHandle());
-		WriteNetToken(Context, Ref.PathToken);
+		// Note: As we use the StringTokenStore to write the token we also need to read it using the stringtokenstore.
+		StringTokenStore->WriteNetToken(Context, Ref.PathToken);
 		NetTokenStore->ConditionalWriteNetTokenData(Context, Context.GetExportContext(), Ref.PathToken);
 		return;
 	}
@@ -1457,7 +1459,8 @@ void FObjectReferenceCache::WriteReference(FNetSerializationContext& Context, FN
 	if (Writer->WriteBool(bIsClientAssignedReference))
 	{
 		WriteNetRefHandle(Context, Ref.GetRefHandle());
-		WriteNetToken(Context, Ref.PathToken);
+		// Note: As we use the StringTokenStore to write the token we also need to read it using the stringtokenstore.
+		StringTokenStore->WriteNetToken(Context, Ref.PathToken);
 
 		return;
 	}
@@ -1493,7 +1496,7 @@ void FObjectReferenceCache::ReadReference(FNetSerializationContext& Context, FNe
 	if (const bool bIsClientAssignedReference = Reader->ReadBool())
 	{
 		const FNetRefHandle NetRefHandle = ReadNetRefHandle(Context);
-		FNetToken RelativePath = ReadNetToken(Context);
+		FNetToken RelativePath = StringTokenStore->ReadNetToken(Context);
 
 		OutRef.RefHandle = NetRefHandle;
 		OutRef.PathToken = RelativePath;
@@ -1590,8 +1593,10 @@ FObjectReferenceCache::EWriteExportsResult FObjectReferenceCache::WritePendingEx
 				if (!(NetToken.IsAssignedByAuthority() && !bIsNetTokenAuthority) && !ExportContext->IsExported(NetToken))
 				{
 					Writer.WriteBool(true);
-					WriteNetToken(Context, NetToken);
+
+					NetTokenStore->WriteNetToken(Context, NetToken);
 					NetTokenStore->WriteTokenData(Context, NetToken);
+
 					ExportContext->AddExported(NetToken);			
 				}
 			}
@@ -1656,7 +1661,7 @@ bool FObjectReferenceCache::ReadExports(FNetSerializationContext& Context, TArra
 			FForceInlineExportScope ForceInlineExportScope(Context.GetInternalContext());
 			while (bHasExportsToRead && !Context.HasErrorOrOverflow())
 			{
-				FNetToken ImportedNetToken = ReadNetToken(Context);
+				FNetToken ImportedNetToken = NetTokenStore->ReadNetToken(Context);
 				NetTokenStore->ReadTokenData(Context, ImportedNetToken, *ResolveContext.RemoteNetTokenStoreState);
 				bHasExportsToRead = Reader.ReadBool();
 			}
diff --git a/Engine/Source/Runtime/Experimental/Iris/Core/Private/Iris/ReplicationSystem/StringTokenStore.cpp b/Engine/Source/Runtime/Experimental/Iris/Core/Private/Iris/ReplicationSystem/StringTokenStore.cpp
--- a/Engine/Source/Runtime/Experimental/Iris/Core/Private/Iris/ReplicationSystem/StringTokenStore.cpp
+++ b/Engine/Source/Runtime/Experimental/Iris/Core/Private/Iris/ReplicationSystem/StringTokenStore.cpp
@@ -27,26 +27,19 @@ namespace UE::Net
 
 FNetToken FStringTokenStore::GetOrCreateToken(const TCHAR* Name, uint32 Length)
 {
-	FNetTokenStoreKey Key = GetOrCreatePersistentString(Name, Length);
+	FNetToken Result;
+	const FNetTokenStoreKey Key = GetOrCreatePersistentString(Name, Length);
 	if (Key.IsValid())
 	{
-		const uint32 StringIndex = Key.GetKeyValue();
-		if (StoredTokens[StringIndex].IsValid())
+		Result = GetNetTokenFromKey(Key);
+		if (!Result.IsValid())		
 		{
-			return StoredTokens[StringIndex];
-		}
-		else
-		{
-			const FNetToken NewToken = CreateToken(TokenStore.IsAuthority() ? FNetToken::ENetTokenAuthority::Authority : FNetToken::ENetTokenAuthority::None, Key, *TokenStore.GetLocalNetTokenStoreState());
-			StoredTokens[StringIndex] = NewToken;
-
-			UE_LOG_STRINGTOKEN(TEXT("FStringTokenStore::GetOrCreateToken - Created %s for %s"), *NewToken.ToString(), Name);
-
-			return NewToken;
+			Result = CreateAndStoreTokenForKey(Key);
+			UE_LOG_STRINGTOKEN(TEXT("FStringTokenStore::GetOrCreateToken - Created %s for %s"), *Result.ToString(), Name);
 		}
 	}
 
-	return FNetToken();
+	return Result;
 }
 
 FNetToken FStringTokenStore::GetOrCreateToken(const FString& String)
@@ -54,31 +47,27 @@ FNetToken FStringTokenStore::GetOrCreateToken(const FString& String)
 	return GetOrCreateToken(ToCStr(String), String.Len());
 }
 
-FNetTokenStoreKey FStringTokenStore::GetOrCreatePersistentString(const TCHAR* Name, uint32 Length)
+FNetTokenDataStore::FNetTokenStoreKey FStringTokenStore::GetOrCreatePersistentString(const TCHAR* Name, uint32 Length)
 {
 	// Hash name
 	const uint32 NameSize = Length * sizeof(TCHAR);
 	uint64 HashedName = CityHash64((const char*)Name, NameSize);
 	
-	// Lock if we have to.. make that a policy
-	//FScopeLock Lock(&CriticalSection);
 	if (const FNetTokenStoreKey* ExistingKey = HashToKey.Find(HashedName))
 	{
 		return *ExistingKey;
 	}
-	else if (StoredTokens.Num() < FNetToken::MaxNetTokenCount)
+
+	const FNetTokenStoreKey NewKey = GetNextNetTokenStoreKey();
+	if (NewKey.IsValid())
 	{
 		// Allocate memory and copy persistent string
 		TCHAR* PersistentString = (TCHAR*)Allocator.Alloc(NameSize + sizeof(TCHAR), alignof(TCHAR));
 		FCString::Strncpy((TCHAR*)PersistentString, Name, Length + 1U);
 
-		// Store bookkeeping data
-		const FNetTokenStoreKey NewKey = MakeNetTokenStoreKey(GetTypeId(), StoredStrings.Num());
-
 		HashToKey.Add(HashedName, NewKey);
 		StoredStrings.Add(PersistentString);
-		StoredTokens.Add(FNetToken());
-
+	
 		return NewKey;
 	}
 
@@ -86,23 +75,27 @@ FNetTokenStoreKey FStringTokenStore::GetOrCreatePersistentString(const TCHAR* Na
 }
 
 FStringTokenStore::FStringTokenStore(FNetTokenStore& InTokenStore)
-: TokenStore(InTokenStore)
+: FNetTokenDataStore(InTokenStore)
 , Allocator()
 {
-	// Reserve 0
-	StoredStrings.Add(nullptr);
-	StoredTokens.Add(FNetToken());
+	// As we use an array for our storage we must match the size of the StoredTokens array.
+	// We assume that Index 0 is invalid.
+	StoredStrings.SetNum(StoredTokens.Num());
 }
 
 const TCHAR* FStringTokenStore::ResolveToken(FNetToken Token, const FNetTokenStoreState* NetTokenStoreState) const
 {
 	const FNetTokenStoreState* TokenStoreState = TokenStore.IsLocalToken(Token) ? TokenStore.GetLocalNetTokenStoreState() : NetTokenStoreState;
-	if (Token.IsValid() && ensureMsgf(TokenStoreState, TEXT("FStringTokenStore::ResolveToken Needs valid remote NetTokenStoreState to resolve remote %s"), *Token.ToString()))
+	if (Token.IsValid() && ensureMsgf(TokenStoreState, TEXT("FStringTokenStore::ResolveToken Needs valid TokenStoreState to resolve %s"), *Token.ToString()))
 	{
 		const FNetTokenStoreKey StoreKey = GetTokenKey(Token, *TokenStoreState);
-		if (StoreKey.IsValid() && StoreKey.GetKeyValue() < (uint32)StoredStrings.Num())
+		if (StoreKey.IsValid() && StoreKey.GetKeyIndex() < (uint32)StoredStrings.Num())
+		{
+			return StoredStrings[StoreKey.GetKeyIndex()];
+		}
+		else
 		{
-			return StoredStrings[StoreKey.GetKeyValue()];
+			UE_LOG(LogNetToken, Error, TEXT("FStringTokenStore::ResolveToken failed to resolve %s in NetTokenDataStore: %s"), *Token.ToString(), *GetTokenStoreName().ToString());
 		}
 	}
 
@@ -112,16 +105,16 @@ const TCHAR* FStringTokenStore::ResolveToken(FNetToken Token, const FNetTokenSto
 void FStringTokenStore::WriteTokenData(FNetSerializationContext& Context, FNetTokenStoreKey TokenStoreKey) const
 {
 	// $TODO: $IRIS: Do not calculate the length of the string to write the data.
-	WriteString(Context.GetBitStreamWriter(), FStringView(StoredStrings[TokenStoreKey.GetKeyValue()]));
+	WriteString(Context.GetBitStreamWriter(), FStringView(StoredStrings[TokenStoreKey.GetKeyIndex()]));
 }
 
 void FStringTokenStore::WriteTokenData(FArchive& Ar, FNetTokenStoreKey TokenStoreKey) const
 {
-	FString Temp(StoredStrings[TokenStoreKey.GetKeyValue()]);
+	FString Temp(StoredStrings[TokenStoreKey.GetKeyIndex()]);
 	Ar << Temp;
 }
 
-FNetTokenStoreKey FStringTokenStore::ReadTokenData(FNetSerializationContext& Context, const FNetToken& NetToken)
+FNetTokenDataStore::FNetTokenStoreKey FStringTokenStore::ReadTokenData(FNetSerializationContext& Context, const FNetToken& NetToken)
 {
 	FNetBitStreamReader* Reader = Context.GetBitStreamReader();
 
@@ -131,17 +124,7 @@ FNetTokenStoreKey FStringTokenStore::ReadTokenData(FNetSerializationContext& Con
 
 	if (!Reader->IsOverflown())
 	{
-		FNetTokenStoreKey Key = GetOrCreatePersistentString(*Temp, Temp.Len());
-
-		// If this is an authtoken and we are not the authority, update the stored key
-		if (Key.IsValid() && NetToken.IsAssignedByAuthority() && !TokenStore.IsAuthority())
-		{
-			UE_LOG_STRINGTOKEN(TEXT("FStringTokenStore::ReadTokenData - Replaced local key %s with %s for string %s"), *StoredTokens[Key.GetKeyValue()].ToString(), *NetToken.ToString(), *Temp);	
-
-			// This way the next time we lookup this key during assignment use the imported authoritative key which we do not have to export
-			StoredTokens[Key.GetKeyValue()] = NetToken;
-		}
-		return Key;
+		return GetOrCreatePersistentString(*Temp, Temp.Len());;
 	}
 	else
 	{
@@ -149,7 +132,7 @@ FNetTokenStoreKey FStringTokenStore::ReadTokenData(FNetSerializationContext& Con
 	}
 }
 
-FNetTokenStoreKey FStringTokenStore::ReadTokenData(FArchive& Ar, const FNetToken& NetToken)
+FNetTokenDataStore::FNetTokenStoreKey FStringTokenStore::ReadTokenData(FArchive& Ar, const FNetToken& NetToken)
 {
 	// Read the token data and add it to the string store without assigning LocalToken
 	FString Temp;
@@ -157,17 +140,7 @@ FNetTokenStoreKey FStringTokenStore::ReadTokenData(FArchive& Ar, const FNetToken
 
 	if (!Ar.IsError())
 	{
-		FNetTokenStoreKey Key = GetOrCreatePersistentString(*Temp, Temp.Len());
-
-		// If this is an authtoken and we are not the authority, update the stored key
-		if (Key.IsValid() && NetToken.IsAssignedByAuthority() && !TokenStore.IsAuthority())
-		{
-			UE_LOG_STRINGTOKEN(TEXT("FStringTokenStore::ReadTokenData - Replaced local key %s with %s for string %s"), *StoredTokens[Key.GetKeyValue()].ToString(), *NetToken.ToString(), *Temp);	
-
-			// This way the next time we lookup this key during assignment use the imported authoritative key which we do not have to export
-			StoredTokens[Key.GetKeyValue()] = NetToken;
-		}
-		return Key;
+		return GetOrCreatePersistentString(*Temp, Temp.Len());
 	}
 	else
 	{
diff --git a/Engine/Source/Runtime/Experimental/Iris/Core/Private/Iris/Serialization/StringNetSerializers.cpp b/Engine/Source/Runtime/Experimental/Iris/Core/Private/Iris/Serialization/StringNetSerializers.cpp
--- a/Engine/Source/Runtime/Experimental/Iris/Core/Private/Iris/Serialization/StringNetSerializers.cpp
+++ b/Engine/Source/Runtime/Experimental/Iris/Core/Private/Iris/Serialization/StringNetSerializers.cpp
@@ -502,10 +502,10 @@ void FNameAsNetTokenNetSerializer::Serialize(FNetSerializationContext& Context,
 	if (Writer->WriteBool(Value.bIsString))
 	{
 		// Always write the token
-		WriteNetToken(Context, Value.NetToken);
+		Context.GetNetTokenStore()->WriteNetTokenWithKnownType<FNameTokenStore>(Context, Value.NetToken);
 
 		// Export or add to pending exports for later export
-		FNetTokenStore::AppendExportOrWriteInlinedExportData(Context, Value.NetToken);
+		FNetTokenStore::AppendExport(Context, Value.NetToken);
 	}
 	else
 	{
@@ -531,16 +531,13 @@ void FNameAsNetTokenNetSerializer::Deserialize(FNetSerializationContext& Context
 		Target.bIsString = 1;
 
 		// Always Read the token
-		FNetToken NetToken = ReadNetToken(Context);
+		FNetToken NetToken = Context.GetNetTokenStore()->ReadNetTokenWithKnownType<FNameTokenStore>(Context);
 
 		if (Reader->IsOverflown())
 		{
 			return;
 		}
 
-		// Read inlined exports if there are any
-		FNetTokenStore::ReadInlinedExportData(Context, NetToken);
-
 		Target.NetToken = NetToken;
 	}
 	else
diff --git a/Engine/Source/Runtime/Experimental/Iris/Core/Public/Iris/ReplicationSystem/NameTokenStore.h b/Engine/Source/Runtime/Experimental/Iris/Core/Public/Iris/ReplicationSystem/NameTokenStore.h
--- a/Engine/Source/Runtime/Experimental/Iris/Core/Public/Iris/ReplicationSystem/NameTokenStore.h
+++ b/Engine/Source/Runtime/Experimental/Iris/Core/Public/Iris/ReplicationSystem/NameTokenStore.h
@@ -21,7 +21,10 @@ public:
 	// Resolve NetToken, to resolve remote tokens RemoteTokenStoreState must be valid
 	IRISCORE_API FName ResolveToken(FNetToken Token, const FNetTokenStoreState* RemoteTokenStoreState = nullptr) const;	
 
-	static FName GetTokenStoreName() { return NameTokenStoreName; }
+	static FName GetTokenStoreName()
+	{ 
+		return NameTokenStoreName;
+	}
 
 protected:
 	// Serialize data for a token, note there is not validation in this function
@@ -36,12 +39,11 @@ protected:
 	FNetTokenStoreKey GetOrCreateTokenStoreKey(FName Name);
 
 private:
+
 	inline static FName NameTokenStoreName = TEXT("NameTokenStore");
 
-	FNetTokenStore& TokenStore;
 	TMap<FName, FNetTokenStoreKey> FNameToKey;
 	TArray<FName> StoredFNames;
-	TArray<FNetToken> StoredTokens;
 };
 
 }
diff --git a/Engine/Source/Runtime/Experimental/Iris/Core/Public/Iris/ReplicationSystem/NetTokenStore.h b/Engine/Source/Runtime/Experimental/Iris/Core/Public/Iris/ReplicationSystem/NetTokenStore.h
--- a/Engine/Source/Runtime/Experimental/Iris/Core/Public/Iris/ReplicationSystem/NetTokenStore.h
+++ b/Engine/Source/Runtime/Experimental/Iris/Core/Public/Iris/ReplicationSystem/NetTokenStore.h
@@ -12,24 +12,34 @@
 * The idea with NetTokens is to allow export of "stable"-pieces of data such as string and names by replacing them with a NetToken during quantization/serialization.
 * The data associated with a NetToken is then exported separately from the data. As soon as the exported data has been acknowledged the data will only be serialized using the NetToken.
 * 
-* As both Client and Server can communicate NetTokens, each side can end up assigning different tokens that will differe from each other.
+* As both Client and Server can communicate NetTokens, each side can end up assigning different tokens that will differs from each other.
 * Here is a high-level overview of the algorithm which works slightly differently based on if using Iris replication or the old replication system.
 *
-* The sending side looks-up or creates a NetToken for the Data being serialized. Servers will mark assigned NetTokens as authoritive, while clients generate a temporary NetToken.
+* The sending side looks-up or creates a NetToken for the Data being serialized. Servers will mark assigned NetTokens as authoritative, while clients generate a temporary NetToken.
 * When a network bunch/batch that contains NetToken is being sent, there is a per-connection look-up to see if we need to append and serialize exports or not, if the data associated with the token
 * has been acknowledged the token will not be exported again.
 *
 * On the receiving side, imported exports are always guaranteed to have been processed before we attempt to read received data containing NetTokens which allows
 * the receiving side to resolve the NetToken to get the actual data.
 *
-* The implmentaion details differs a bit depending on if we are using iris replication or old style replication.
+* The implementation details differs a bit depending on if we are using iris replication or old style replication.
 *
 * A current example that us used by both systems is GameplayTags, For Iris: See GameplayTagNetSerializer.cpp, for the old replication system: See: GameplayTagContainer.cpp.
 *
 * Note: this is still an experimental feature and is currently only available if compiling UE_WITH_IRIS. 
-* Implementation datails is subject to change.
+* Implementation details is subject to change.
 */
 
+/* NetBitArray validation support. */
+#ifndef UE_NET_VALIDATE_NETTOKENTYPE
+#if (UE_BUILD_SHIPPING || UE_BUILD_TEST)
+#	define UE_NET_VALIDATE_NETTOKENTYPE 0
+#else
+#	define UE_NET_VALIDATE_NETTOKENTYPE 1
+#endif
+#endif
+
+
 class UNetTokenDataStream;
 namespace UE::Net
 {
@@ -47,36 +57,100 @@ namespace UE::Net
 namespace UE::Net
 {
 
-// virtual interface for NetTokenDataStores
+/**
+ * NetTokenDataStore
+ * Implemented per type to store and serialize data associated with a NetToken.
+*/
 class FNetTokenDataStore
 {
 public:
+	class FNetTokenStoreKey
+	{
+	public:
+		enum { InvalidKeyIndex = 0U };
+
+		FNetTokenStoreKey()
+		: KeyIndex(InvalidKeyIndex)
+		{
+		}
+
+		explicit FNetTokenStoreKey(uint32 InKeyIndex)
+		: KeyIndex(InKeyIndex)
+		{
+		}
+
+		uint32 GetKeyIndex() const
+		{ 
+			return KeyIndex;
+		}
+
+		bool IsValid() const
+		{ 
+			return KeyIndex != InvalidKeyIndex;
+		}
+
+		bool operator==(const FNetTokenStoreKey& Other) const = default;
+
+	private:
+		uint32 KeyIndex = InvalidKeyIndex;
+	};
+
 	IRISCORE_API virtual ~FNetTokenDataStore();
 
+	// Serialization methods for NetTokens that does not include the TypeId as this is known by the NetTokenDataStore.
+
+	// Write NetToken
+	IRISCORE_API void WriteNetToken(UE::Net::FNetSerializationContext& Context, FNetToken Token);
+
+	// Read NetToken
+	IRISCORE_API FNetToken ReadNetToken(UE::Net::FNetSerializationContext& Context);
+
+	// Write NetToken
+	IRISCORE_API void WriteNetToken(FArchive& Ar, FNetToken Token);
+
+	// Read NetToken
+	IRISCORE_API FNetToken ReadNetToken(FArchive& Ar);
+
 protected:
-	IRISCORE_API FNetTokenDataStore();
+
+	explicit IRISCORE_API FNetTokenDataStore(FNetTokenStore& InTokenStore);
 
 	virtual void WriteTokenData(FNetSerializationContext& Context, FNetTokenStoreKey Key) const = 0;
-	virtual FNetTokenStoreKey ReadTokenData(FNetSerializationContext& Context, const FNetToken& NetToken) = 0;
+	virtual FNetTokenDataStore::FNetTokenStoreKey ReadTokenData(FNetSerializationContext& Context, const FNetToken& NetToken) = 0;
 
 	virtual void WriteTokenData(FArchive& Ar, FNetTokenStoreKey Key) const = 0;
-	virtual FNetTokenStoreKey ReadTokenData(FArchive& Ar, const FNetToken& NetToken) = 0;
+	virtual FNetTokenDataStore::FNetTokenStoreKey ReadTokenData(FArchive& Ar, const FNetToken& NetToken) = 0;
 
-	FNetTokenStoreKey GetTokenKey(FNetToken Token, const FNetTokenStoreState& TokenStoreState) const;
-	inline FNetToken::FTypeId GetTypeId() const { return TypeId; }
+	FNetTokenDataStore::FNetTokenStoreKey GetTokenKey(FNetToken Token, const FNetTokenStoreState& TokenStoreState) const;
+
+	inline FNetToken::FTypeId GetTypeId() const
+	{
+		return TypeId;
+	}
 
 	// Create new NetToken
-	IRISCORE_API FNetToken CreateToken(FNetToken::ENetTokenAuthority Authority, FNetTokenStoreKey Key, FNetTokenStoreState& TokenStoreState);
+	IRISCORE_API FNetToken CreateAndStoreTokenForKey(FNetTokenStoreKey Key);
+	IRISCORE_API void StoreTokenForKey(FNetTokenStoreKey Key, FNetToken NetToken);
+	IRISCORE_API FNetToken GetNetTokenFromKey(FNetTokenStoreKey) const;
 
-	// Make NetTokenStoreKey
-	IRISCORE_API static FNetTokenStoreKey MakeNetTokenStoreKey(FNetToken::FTypeId TokenTypeId, uint32 TokenKeyValue);
+	// Allocate next TokenStoreKey
+	IRISCORE_API FNetTokenDataStore::FNetTokenStoreKey GetNextNetTokenStoreKey();
+
+protected:
+	// Maps from FNetTokenStoreKey (index) to NetToken, this typically is the locally assigned NetToken, but can be overridden if we receive a token from the authority.
+	TArray<FNetToken> StoredTokens;
+	FNetTokenStore& TokenStore;
 
 private:
 	friend FNetTokenStore;
 	FNetToken::FTypeId TypeId;
 };
 
-// This is the token store, we currently have one per ReplicationSystem but it is possible we will share this across game instance.
+/**
+ * FNetTokenStore
+ * Main system for using NetTokensm currently owns type specific NetTokenDataStores and per connection state
+ * Currently we have a unique instance per NetDriver/ReplicationSystem but it is possible we will share this across game instance.
+ */
 class FNetTokenStore
 {
 	UE_NONCOPYABLE(FNetTokenStore);
@@ -92,18 +166,28 @@ public:
 	};
 	IRISCORE_API void Init(FInitParams& InitParams);
 
-	bool IsAuthority() const { return Params.Authority == FNetToken::ENetTokenAuthority::Authority; }
+	/** Returns true if this is the authority, typically the server. */
+	bool IsAuthority() const
+	{
+		return Params.Authority == FNetToken::ENetTokenAuthority::Authority;
+	}
 
-	// A token is local if the authority of the NetTokenStore and the token matches, Invalid tokens are always local.
-	bool IsLocalToken(const FNetToken NetToken) const { return !NetToken.IsValid() || IsAuthority() == NetToken.IsAssignedByAuthority(); }
+	/** A token is local if the authority of the NetTokenStore and the token matches, Invalid tokens are always local. */
+	bool IsLocalToken(const FNetToken NetToken) const
+	{
+		return !NetToken.IsValid() || IsAuthority() == NetToken.IsAssignedByAuthority();
+	}
 
-	// Register DataStore and return true if it was registered, 
+	/** Register DataStore and return true if it was registered. */
 	IRISCORE_API bool RegisterDataStore(TUniquePtr<FNetTokenDataStore> DataStore, FName TokenStoreName);
 
+	/** Get const access to data store by name. */
 	IRISCORE_API const FNetTokenDataStore* GetDataStore(FName Name) const;
+
+	/** Get data store by name. */
 	IRISCORE_API FNetTokenDataStore* GetDataStore(FName Name);
 
-	/** Create and return data store of specified type, it will be owned by NetTokenStore */
+	/** Create and return data store of specified type, it will be owned by NetTokenStore. */
 	template<typename T>
 	bool CreateAndRegisterDataStore()
 	{
@@ -114,73 +198,181 @@ public:
 
 	/** Return data store of specified type. */
 	template<typename T>
-	T* GetDataStore() { return static_cast<T*>(GetDataStore(T::GetTokenStoreName())); }
+	T* GetDataStore()
+	{
+		return static_cast<T*>(GetDataStore(T::GetTokenStoreName()));
+	}
 
 	/** Return data store of specified type. */
 	template<typename T>
-	const T* GetDataStore() const { return static_cast<const T*>(GetDataStore(T::GetTokenStoreName())); }
+	const T* GetDataStore() const
+	{
+		return static_cast<const T*>(GetDataStore(T::GetTokenStoreName()));
+	}
+	
+	// FNetTokenStoreState maps from NetTokenIndex -> NetTokenStoreKey (Index)
+	// Remote and local NetTokens use separate NetTokensStore states.
+
+	/** Get const access to the local NetTokenStoreState */
+	const FNetTokenStoreState* GetLocalNetTokenStoreState() const
+	{
+		return LocalNetTokenStoreState.Get();
+	}
 
-	const FNetTokenStoreState* GetLocalNetTokenStoreState() const { return LocalNetTokenStoreState.Get(); }
-	FNetTokenStoreState* GetLocalNetTokenStoreState() { return LocalNetTokenStoreState.Get();; }
+	/** Get access to the local NetTokenStoreState */
+	FNetTokenStoreState* GetLocalNetTokenStoreState()
+	{
+		return LocalNetTokenStoreState.Get();
+	}
 
 	/** Init RemoteNetTokenStoreState for given ConnectionId, if it already exists it will be reset. */
 	IRISCORE_API void InitRemoteNetTokenStoreState(uint32 ConnectionId);
 
-	/** Get RemoteNetTokenStoreState for given ConnectionId. */
+	/** Get const access to RemoteNetTokenStoreState for given ConnectionId. */
 	IRISCORE_API const FNetTokenStoreState* GetRemoteNetTokenStoreState(uint32 ConnectionId) const;
 
 	/** Get RemoteNetTokenStoreState for given ConnectionId. */
 	IRISCORE_API FNetTokenStoreState* GetRemoteNetTokenStoreState(uint32 ConnectionId);
 
-	// Write data associated with the NetToken
+	/** Write data associated with the NetToken. */
 	IRISCORE_API void WriteTokenData(FNetSerializationContext& Context, const FNetToken NetToken) const;
 
-	// Read data associated with the NetToken
+	/** Read data associated with the NetToken. */
 	IRISCORE_API void ReadTokenData(FNetSerializationContext& Context, const FNetToken NetToken, FNetTokenStoreState& RemoteNetTokenStoreState);
 
-	// Write data associated with the NetToken
+	/** Write data associated with the NetToken. */
 	IRISCORE_API void WriteTokenData(FArchive& Ar, const FNetToken NetToken) const;
 
-	// Read data associated with the NetToken
+	/** Read data associated with the NetToken. */
 	IRISCORE_API void ReadTokenData(FArchive& Ar, const FNetToken NetToken, FNetTokenStoreState& RemoteNetTokenStoreState);
 
-	// Conditionally write NetTokenData unless already exported
+	/** Conditionally write NetTokenData unless already exported. */
 	IRISCORE_API void ConditionalWriteNetTokenData(FNetSerializationContext& Context, Private::FNetExportContext* ExportContext, const FNetToken NetToken) const;
 
-	// Conditionally read NetTokenData if exported
+	/** Conditionally read NetTokenData if exported. */
 	IRISCORE_API void ConditionalReadNetTokenData(FNetSerializationContext& Context, const FNetToken NetToken);
 
-	// Utility methods, consolidate with other changes to NetTokenStore as next step.
-	IRISCORE_API static void AppendExportOrWriteInlinedExportData(FNetSerializationContext&, FNetToken NetToken);
-	IRISCORE_API static void ReadInlinedExportData(FNetSerializationContext&, FNetToken NetToken);
+	/** Utility methods, consolidate with other changes to NetTokenStore as next step. */
+	IRISCORE_API static void AppendExport(FNetSerializationContext&, FNetToken NetToken);
+
+	/** 
+	 * Convenience method to Write a NetToken without writing type bits
+	 * In development builds it will verify that the token type matches the StoreType, otherwise it will skip the lookup of DataStore.
+	 */
+	template <typename T>
+	void WriteNetTokenWithKnownType(FNetSerializationContext& Context, FNetToken NetToken)
+	{
+#if UE_NET_VALIDATE_NETTOKENTYPE
+		// We only verify this in dev builds as it should trap most programmer errors.
+		if (NetToken.IsValid())
+		{
+			const T* DataStore = GetDataStore<T>();
+			if (!DataStore || (DataStore->GetTypeId() != NetToken.GetTypeId()))
+			{
+				UE_LOG(LogNetToken, Error, TEXT("Tried to write NetToken %s using invalid NetTokenStore %s"), *NetToken.ToString(), *(T::GetTokenStoreName().ToString()));
+				// Just to get some attention to the log.
+				ensure(DataStore && (DataStore->GetTypeId() == NetToken.GetTypeId()));
+			}
+		}
+#endif
+		const bool bWriteTypeId = false;
+		FNetTokenStore::InternalWriteNetToken(Context, NetToken, bWriteTypeId);
+	}
+
+	/** 
+	 * Convenience method to read a NetToken without reading type bits.
+	*/
+	template <typename T>
+	FNetToken ReadNetTokenWithKnownType(FNetSerializationContext& Context)
+	{
+		FNetToken NetToken;
+		if (T* DataStore = GetDataStore<T>())
+		{
+			NetToken = DataStore->ReadNetToken(Context);
+		}
+		else
+		{
+			UE_LOG(LogNetToken, Error, TEXT("ReadNetTokenWithKnownType Tried to read NetToken using invalid NetTokenStore %s"), *(T::GetTokenStoreName().ToString()));
+			// Just to get some attention to the log.
+			ensure(DataStore);
+		}
+
+		return NetToken;
+	}
+
+	/** Write NetToken including TypeId */
+	void WriteNetToken(FNetSerializationContext& Context, FNetToken NetToken) const
+	{
+		const bool bWriteTypeId = true;
+		FNetTokenStore::InternalWriteNetToken(Context, NetToken, bWriteTypeId);
+	}
+
+	/** Read NetToken including TypeId */
+	FNetToken ReadNetToken(FNetSerializationContext& Context) const
+	{
+		return InternalReadNetToken(Context, FNetToken::InvalidTokenTypeId);
+	}
+
+	/** Write NetToken including TypeId */
+	void WriteNetToken(FArchive& Ar, FNetToken NetToken) const
+	{
+		const bool bWriteTypeId = true;
+		FNetTokenStore::InternalWriteNetToken(Ar, NetToken, bWriteTypeId);
+	}
 
+	/** Read NetToken including TypeId */
+	FNetToken ReadNetToken(FArchive& Ar) const
+	{
+		return InternalReadNetToken(Ar, FNetToken::InvalidTokenTypeId);
+	}
+	
 private:
+
 	friend UNetTokenDataStream;
 	friend FNetTokenDataStore;
 
+	using FNetTokenStoreKey = FNetTokenDataStore::FNetTokenStoreKey;
+
+	// Internal method to write a NetToken, if bWriteTokenType is true it will write the TypeId as well.
+	IRISCORE_API static void InternalWriteNetToken(UE::Net::FNetSerializationContext& Context, FNetToken Token, bool bWriteTokenType);
+	IRISCORE_API static void InternalWriteNetToken(FArchive& Ar, FNetToken Token, bool bWriteTokenType);
+
+	// Internal method to read a NetToken, if the TokenTypeId is valid it will be used instead of reading it from the bitstream
+	IRISCORE_API static FNetToken InternalReadNetToken(UE::Net::FNetSerializationContext& Context, FNetToken::FTypeId TokenTypeId);
+	IRISCORE_API static FNetToken InternalReadNetToken(FArchive& Ar, FNetToken::FTypeId TokenTypeId);
+
+	bool ValidateAndStoreNetTokenData(FNetTokenDataStore& DataStore, FNetTokenStoreState& RemoteNetTokenStoreState, const FNetToken NetToken, const FNetTokenStoreKey StoreKey);
+	
+	TArray<FNetToken> GetAllNetTokens() const;
+
+	static FNetToken MakeNetToken(uint32 TypeId, uint32 Index, FNetToken::ENetTokenAuthority Authority)
+	{ 
+		check(TypeId < FNetToken::MaxTypeIdCount);
+		check(Index < FNetToken::MaxNetTokenCount);
+
+		return FNetToken(TypeId, Index, Authority);
+	}
+
 	TUniquePtr<FNetTokenStoreState> LocalNetTokenStoreState;
 	TArray<TUniquePtr<FNetTokenStoreState>> RemoteNetTokenStoreStates;
 	TArray<TTuple<FName, TUniquePtr<FNetTokenDataStore>>> TokenDataStores;
 	FInitParams Params;
 };
 
-inline FNetTokenStoreKey FNetTokenDataStore::MakeNetTokenStoreKey(FNetToken::FTypeId TokenTypeId, uint32 TokenKeyValue)
+inline FNetTokenDataStore::FNetTokenStoreKey FNetTokenDataStore::GetNextNetTokenStoreKey()
 {
-	check(TokenTypeId < FNetToken::MaxTypeIdCount);
-	check(TokenKeyValue < FNetToken::MaxNetTokenCount)
-
-	FNetTokenStoreKey TokenKey;
-	TokenKey.TypeId = TokenTypeId;
-	TokenKey.Key = TokenKeyValue;
+	uint32 NextTokenStoreKeyIndex = (uint32)StoredTokens.Num();
+	if (ensure(NextTokenStoreKeyIndex < FNetToken::MaxNetTokenCount))
+	{
+		FNetTokenStoreKey TokenKey(NextTokenStoreKeyIndex);
 
-	return TokenKey;
-}
+		// We fill in the actual key later
+		StoredTokens.Add(FNetToken());
 
-// Read and write tokens
-IRISCORE_API FNetToken ReadNetToken(UE::Net::FNetSerializationContext& Context);
-IRISCORE_API void WriteNetToken(UE::Net::FNetSerializationContext& Context, FNetToken Token);
+		return TokenKey;
+	}
 
-IRISCORE_API FNetToken ReadNetToken(FArchive& Ar);
-IRISCORE_API void WriteNetToken(FArchive& Ar, FNetToken Token);
+	return FNetTokenStoreKey();
+}
 
 }
diff --git a/Engine/Source/Runtime/Experimental/Iris/Core/Public/Iris/ReplicationSystem/StringTokenStore.h b/Engine/Source/Runtime/Experimental/Iris/Core/Public/Iris/ReplicationSystem/StringTokenStore.h
--- a/Engine/Source/Runtime/Experimental/Iris/Core/Public/Iris/ReplicationSystem/StringTokenStore.h
+++ b/Engine/Source/Runtime/Experimental/Iris/Core/Public/Iris/ReplicationSystem/StringTokenStore.h
@@ -26,9 +26,15 @@ public:
 	IRISCORE_API const TCHAR* ResolveToken(FNetToken Token, const FNetTokenStoreState* RemoteTokenStoreState = nullptr) const;	
 
 	// Resolve a token received from remote
-	IRISCORE_API const TCHAR* ResolveRemoteToken(FNetToken Token, const FNetTokenStoreState& NetTokenStoreState) const { return ResolveToken(Token, &NetTokenStoreState); }
+	IRISCORE_API const TCHAR* ResolveRemoteToken(FNetToken Token, const FNetTokenStoreState& NetTokenStoreState) const
+	{ 
+		return ResolveToken(Token, &NetTokenStoreState);
+	}
 
-	static FName GetTokenStoreName() { return StringTokenStoreName; }
+	static FName GetTokenStoreName()
+	{
+		return StringTokenStoreName;
+	}
 
 protected:
 	// Serialize data for a token, note there is not validation in this function
@@ -45,10 +51,8 @@ protected:
 private:
 	inline static FName StringTokenStoreName = TEXT("StringTokenStore");
 
-	FNetTokenStore& TokenStore;
 	TMap<uint64, FNetTokenStoreKey> HashToKey;
 	TArray<const TCHAR*> StoredStrings;
-	TArray<FNetToken> StoredTokens;
 	FMemStackBase Allocator;
 };
 
diff --git a/Engine/Source/Runtime/GameplayTags/Private/GameplayTagContainer.cpp b/Engine/Source/Runtime/GameplayTags/Private/GameplayTagContainer.cpp
--- a/Engine/Source/Runtime/GameplayTags/Private/GameplayTagContainer.cpp
+++ b/Engine/Source/Runtime/GameplayTags/Private/GameplayTagContainer.cpp
@@ -1312,7 +1312,8 @@ namespace UE::GameplayTags::GameplayTagDynamicSerialization
 			TagToken = TagTokenDataStore->GetOrCreateToken(GameplayTag);
 
 			// Write token
-			WriteNetToken(Ar, TagToken);
+			// Important: As we write it directly through the TagTokenStore we must read it in the same way as we skip serializing the type.
+			TagTokenDataStore->WriteNetToken(Ar, TagToken);
 
 			// Register replay export if needed
 			if (TagToken.IsValid())
@@ -1350,7 +1351,8 @@ namespace UE::GameplayTags::GameplayTagDynamicSerialization
 		else if (Ar.IsLoading())
 		{
 			// Read TagNetToken
-			TagToken = ReadNetToken(Ar);
+			TagTokenDataStore->ReadNetToken(Ar);
+
 			if (Ar.IsError())
 			{
 				return false;
@@ -1403,60 +1405,56 @@ namespace UE::GameplayTags::GameplayTagDynamicSerialization
 		{
 			UE::Net::FNetTokenExportContext* ExportContext = FNetTokenExportContext::GetNetTokenExportContext(Ar);
 			UE::Net::FNetTokenStore* NetTokenStore = ExportContext ? ExportContext->GetNetTokenStore() : nullptr;
-			if (ensure(NetTokenStore))
+			FGameplayTagTokenStore* TagTokenStore = NetTokenStore ? NetTokenStore->GetDataStore<UE::Net::FGameplayTagTokenStore>() : nullptr;
+			if (ensure(TagTokenStore))
 			{
 				FNetToken TagToken;
-
 				if (GameplayTag.IsValid())
 				{
-					if (FGameplayTagTokenStore* TagTokenStore = NetTokenStore->GetDataStore<UE::Net::FGameplayTagTokenStore>())
-					{
-						TagToken = TagTokenStore->GetOrCreateToken(GameplayTag);
-					}
-					else
-					{
-						UE_LOG(LogGameplayTags, Error, TEXT("FGameplayTag::NetSerialize::Could not find required FGameplayTagTokenStore"));
-						ensure(false);
-					}
+					TagToken = TagTokenStore->GetOrCreateToken(GameplayTag);
 				}
 
 				UE_NET_TRACE_DYNAMIC_NAME_SCOPE(*TagToken.ToString(), static_cast<FNetBitWriter&>(Ar), GetTraceCollector(static_cast<FNetBitWriter&>(Ar)), ENetTraceVerbosity::VeryVerbose);
 
-				// Write token
-				WriteNetToken(Ar, TagToken);
+				// Write NetToken, 
+				// Important: As we write it directly thorugh the TagTokenStore we also need to read it in the samw way as we skip serializing the type.
+				TagTokenStore->WriteNetToken(Ar, TagToken);
+
+				// Add export
+				ExportContext->AddNetTokenPendingExport(TagToken);
 
-				// Should we append export from WriteNetToken method?
-				if (TagToken.IsValid())
-				{
-					ExportContext->AddNetTokenPendingExport(TagToken);
-				}
 				return true;
 			}
+			else
+			{
+				UE_LOG(LogGameplayTags, Error, TEXT("FGameplayTag::NetSerialize::Could not find required FGameplayTagTokenStore"));
+				ensure(false);
+			}
 		}
 		else if (Ar.IsLoading())
 		{
 			// When reading data we always have a PackageMap so we can get the necessary resolve context from here.
 			const FNetTokenResolveContext* NetTokenResolveContext = Map ? Map->GetNetTokenResolveContext() : nullptr;
-			if (ensure(NetTokenResolveContext && NetTokenResolveContext->NetTokenStore))
+			FGameplayTagTokenStore* TagTokenStore = NetTokenResolveContext ? NetTokenResolveContext->NetTokenStore->GetDataStore<UE::Net::FGameplayTagTokenStore>() : nullptr;
+			if (ensure(TagTokenStore))
 			{
-				FNetToken TagToken = ReadNetToken(Ar);
+				// Read the TagToken using the TagTokenStore
+				FNetToken TagToken = TagTokenStore->ReadNetToken(Ar);
 				if (Ar.IsError())
 				{
 					return false;
 				}
-				if (FGameplayTagTokenStore* TagTokenStore = NetTokenResolveContext->NetTokenStore->GetDataStore<UE::Net::FGameplayTagTokenStore>())
-				{
-					GameplayTag = TagTokenStore->ResolveToken(TagToken, NetTokenResolveContext->RemoteNetTokenStoreState);
-				}
-				else
-				{
-					UE_LOG(LogGameplayTags, Error, TEXT("FGameplayTag::NetSerialize::Could not find required FGameplayTagTokenStore"));	
-					ensure(false);
-					Ar.SetError();
-					return false;
-				}
+
+				// Resolve the TagToken
+				GameplayTag = TagTokenStore->ResolveToken(TagToken, NetTokenResolveContext->RemoteNetTokenStoreState);
 				return true;
 			}
+			else
+			{
+				UE_LOG(LogGameplayTags, Error, TEXT("FGameplayTag::NetSerialize::Could not find required FGameplayTagTokenStore"));	
+				ensure(false);
+				Ar.SetError();
+			}
 		}
 		return false;
 	}
diff --git a/Engine/Source/Runtime/GameplayTags/Private/GameplayTagNetSerializer.cpp b/Engine/Source/Runtime/GameplayTags/Private/GameplayTagNetSerializer.cpp
--- a/Engine/Source/Runtime/GameplayTags/Private/GameplayTagNetSerializer.cpp
+++ b/Engine/Source/Runtime/GameplayTags/Private/GameplayTagNetSerializer.cpp
@@ -99,11 +99,11 @@ void FGameplayTagNetSerializer::Serialize(FNetSerializationContext& Context, con
 			return;
 		}
 
-		// Always write the token
-		WriteNetToken(Context, Value.TagNetToken);
+		// Write token without type, 
+		Context.GetNetTokenStore()->WriteNetTokenWithKnownType<FGameplayTagTokenStore>(Context, Value.TagNetToken);
 
 		// Export or add to pending exports for later export
-		FNetTokenStore::AppendExportOrWriteInlinedExportData(Context, Value.TagNetToken);
+		FNetTokenStore::AppendExport(Context, Value.TagNetToken);
 	}
 }
 
@@ -122,15 +122,12 @@ void FGameplayTagNetSerializer::Deserialize(FNetSerializationContext& Context, c
 	}
 	else
 	{
-		FNetToken NetToken = ReadNetToken(Context);
+		FNetToken NetToken = Context.GetNetTokenStore()->ReadNetTokenWithKnownType<FGameplayTagTokenStore>(Context);
 		if (Reader->IsOverflown())
 		{
 			return;
 		}
 
-		// Read inlined exports if we should
-		FNetTokenStore::ReadInlinedExportData(Context, NetToken);
-
 		if (Reader->IsOverflown())
 		{
 			return;
diff --git a/Engine/Source/Runtime/GameplayTags/Public/GameplayTagTokenStore.h b/Engine/Source/Runtime/GameplayTags/Public/GameplayTagTokenStore.h
--- a/Engine/Source/Runtime/GameplayTags/Public/GameplayTagTokenStore.h
+++ b/Engine/Source/Runtime/GameplayTags/Public/GameplayTagTokenStore.h
@@ -25,9 +25,13 @@ public:
 	// Resolve NetToken, to resolve remote tokens RemoteTokenStoreState must be valid
 	GAMEPLAYTAGS_API FGameplayTag ResolveToken(FNetToken Token, const FNetTokenStoreState* RemoteTokenStoreState = nullptr) const;
 
-	static FName GetTokenStoreName() { return NameTokenStoreName; }
+	static FName GetTokenStoreName()
+	{
+		return GameplayTokenStoreName;
+	}
+
 private:
-	inline static FName NameTokenStoreName = TEXT("GameplayTagTokenStore");
+	inline static FName GameplayTokenStoreName = TEXT("GameplayTagTokenStore");
 };
 
 }
diff --git a/Engine/Source/Runtime/Net/Core/Private/Net/Core/NetToken/NetTokenExportContext.cpp b/Engine/Source/Runtime/Net/Core/Private/Net/Core/NetToken/NetTokenExportContext.cpp
--- a/Engine/Source/Runtime/Net/Core/Private/Net/Core/NetToken/NetTokenExportContext.cpp
+++ b/Engine/Source/Runtime/Net/Core/Private/Net/Core/NetToken/NetTokenExportContext.cpp
@@ -61,7 +61,7 @@ void FNetTokenExportContext::AppendNetTokensPendingExport(TArrayView<const UE::N
 
 void FNetTokenExportContext::AddNetTokenPendingExport(UE::Net::FNetToken NetToken)
 {
-	if (ensure(NetTokensPendingExport))
+	if (ensure(NetTokensPendingExport) && NetToken.IsValid())
 	{
 		NetTokensPendingExport->Add(NetToken);
 		UE_LOG(LogNetToken, Verbose, TEXT("FNetTokenExportContext::AddNetTokenPendingExport %s"), *NetToken.ToString());
diff --git a/Engine/Source/Runtime/Net/Core/Public/Net/Core/NetToken/NetToken.h b/Engine/Source/Runtime/Net/Core/Public/Net/Core/NetToken/NetToken.h
--- a/Engine/Source/Runtime/Net/Core/Public/Net/Core/NetToken/NetToken.h
+++ b/Engine/Source/Runtime/Net/Core/Public/Net/Core/NetToken/NetToken.h
@@ -44,66 +44,69 @@ public:
 	};
 
 public:	
-	FNetToken() : Value(Invalid) {}
-	inline bool IsValid() const { return Value != Invalid; }
-	bool IsAssignedByAuthority() const { return bIsAssignedByAuthority != 0U; }
-	uint32 GetIndex() const { return Index; }
-	bool operator==(const FNetToken& Other) const { return IsAssignedByAuthority() == Other.IsAssignedByAuthority() && Index == Other.Index; }
+	FNetToken()
+	: Value(Invalid)
+	{
+	}
+
+	inline bool IsValid() const
+	{
+		return Index != InvalidTokenIndex;
+	}
+
+	bool IsAssignedByAuthority() const
+	{
+		return bIsAssignedByAuthority != 0U;
+	}
+
+	uint32 GetIndex() const
+	{
+		return Index;
+	}
+
+	FTypeId GetTypeId() const
+	{
+		return TypeId;
+	}
+
+	bool operator==(const FNetToken& Other) const
+	{
+		return Value == Other.Value;
+	}
+
 	FString ToString() const;
 
 	friend uint32 GetTypeHash(const FNetToken& Token)
 	{
 		return ::GetTypeHash(Token.Value);
 	}
-	
-	static FNetToken MakeNetToken(uint32 Index, ENetTokenAuthority Authority) { check(Index < MaxNetTokenCount); return FNetToken(Index, Authority); }
-	
-private:
-	explicit FNetToken(uint32 InIndex, ENetTokenAuthority Authority) { Padding = 0U, Index = InIndex, bIsAssignedByAuthority = Authority == ENetTokenAuthority::Authority ? 1U : 0U; }
 
 private:
+	friend class UE::Net::FNetTokenStore;
 
-	union 
+	explicit FNetToken(uint32 InTypeId, uint32 InIndex, ENetTokenAuthority Authority)
 	{
-		struct
-		{
-			uint32 Index : TokenBits;
-			uint32 bIsAssignedByAuthority : 1U;
-			uint32 Padding : 32 - TokenBits - 1U;
-		};
-		uint32 Value;
-	};
-};
-
-class FNetTokenStoreKey
-{
-public:
-	enum { Invalid = 0U };
-
-	FNetTokenStoreKey() : Value(Invalid) {}
-	bool IsValid() const { return Value != Invalid; }
-	uint32 GetKeyValue() const { return Key; }
-	FNetToken::FTypeId GetTypeId() const { return TypeId; }
-	bool operator==(const FNetTokenStoreKey& Other) const { return Value == Other.Value; }
+		Padding = 0U;
+		TypeId = InTypeId;
+		Index = InIndex;
+		bIsAssignedByAuthority = Authority == ENetTokenAuthority::Authority ? 1U : 0U;
+	}
 
 private:
-	friend class FNetTokenStore;
-	friend class FNetTokenDataStore;
 
-private:
-	union
+	union 
 	{
 		struct
 		{
-			uint32 TypeId : FNetToken::TokenTypeIdBits;
-			uint32 Key : FNetToken::TokenBits;
+			uint32 Index : TokenBits;
+			uint32 TypeId : TokenTypeIdBits;
+			uint32 bIsAssignedByAuthority : 1U;
+			uint32 Padding : 32 - TokenTypeIdBits - TokenBits - 1U;
 		};
 		uint32 Value;
 	};
 };
 
-static_assert(sizeof(FNetTokenStoreKey) == sizeof(uint32), "FNetTokenKey should fit in a uint32");
-
 // Contains necessary context to resolve NetTokens
 class FNetTokenResolveContext
 {
@@ -115,7 +118,7 @@ public:
 inline FString FNetToken::ToString() const
 {
 	FString Result;
-	Result = FString::Printf(TEXT("NetToken (Auth:%u Index=%u)"), IsAssignedByAuthority(), Index);
+	Result = FString::Printf(TEXT("NetToken (Auth:%u TypeId=%u Index=%u)"), IsAssignedByAuthority(), TypeId, Index);
 	return Result;
 }
 
