[ChaosVD] Fix ensure and recording start failure that can happen if you attempt to start a CVD recording when another trace recording session is active.

CVD already attempted to stop the active recording in these situations, but there are cases where the active session might take more than one frame to disconnect, causing the error.

This CL also adds a way to restore the trace channels CVD disabled at the start of the recording

diff --git a/Engine/Source/Runtime/Experimental/ChaosVisualDebugger/Private/ChaosVDRuntimeModule.cpp b/Engine/Source/Runtime/Experimental/ChaosVisualDebugger/Private/ChaosVDRuntimeModule.cpp
--- a/Engine/Source/Runtime/Experimental/ChaosVisualDebugger/Private/ChaosVDRuntimeModule.cpp
+++ b/Engine/Source/Runtime/Experimental/ChaosVisualDebugger/Private/ChaosVDRuntimeModule.cpp
@@ -4,11 +4,17 @@
 
 #include "Containers/Array.h"
 #include "HAL/IConsoleManager.h"
+#include "HAL/PlatformProcess.h"
 #include "Internationalization/Internationalization.h"
 #include "Misc/CommandLine.h"
 #include "Misc/FileHelper.h"
 #include "Misc/MessageDialog.h"
 #include "Misc/Paths.h"
+
+#if WITH_EDITOR
+#include "Misc/ScopedSlowTask.h"
+#endif
+
 #include "Modules/ModuleManager.h"
 #include "ProfilingDebugging/TraceAuxiliary.h"
 
@@ -41,6 +47,11 @@ static FAutoConsoleVariable CVarChaosVDGTimeBetweenFullCaptures(
 	10,
 	TEXT("Time interval in seconds after which a full capture (not only delta changes) should be recorded"));
 
+static FAutoConsoleVariable CVarChaosVDMaxTimeToWaitForDisconnect(
+	TEXT("p.Chaos.VD.MaxTimeToWaitForDisconnectSeconds"),
+	5.0f,
+	TEXT("Max time to wait after attempting to stop an active trace session. After that time has passed if we are still connected, CVD will continue and eventually error out."));
+
 FChaosVDRecordingStateChangedDelegate FChaosVDRuntimeModule::RecordingStartedDelegate = FChaosVDRecordingStateChangedDelegate();
 FChaosVDRecordingStateChangedDelegate FChaosVDRuntimeModule::RecordingStopDelegate = FChaosVDRecordingStateChangedDelegate();
 FChaosVDRecordingStartFailedDelegate FChaosVDRuntimeModule::RecordingStartFailedDelegate = FChaosVDRecordingStartFailedDelegate();
@@ -151,27 +162,22 @@ void FChaosVDRuntimeModule::StartRecording(TConstArrayView<FString> Args)
 	// This is aggressive but until Trace supports multi-sessions, just take over.
 	if (FTraceAuxiliary::IsConnected())
 	{
-		StopTrace();
-	}
-
-	// Until we support allowing other channels, indicate in the logs that we are disabling everything else
-	UE_LOG(LogChaosVDRuntime, Log, TEXT("[%s] Disabling additional trace channels..."), ANSI_TO_TCHAR(__FUNCTION__));
+		UE_LOG(LogChaosVDRuntime, Log, TEXT("[%s] There is an active trace session. attempting to disconnect..."), ANSI_TO_TCHAR(__FUNCTION__));
 
-	// Disable any enabled additional channel
-	UE::Trace::EnumerateChannels([](const ANSICHAR* ChannelName, bool bEnabled, void*)
+		//TODO: We should make the wait async like we do whe we attempt to connect to a live session
+		if (FTraceAuxiliary::Stop() && WaitForTraceSessionDisconnect())
 		{
-			if (bEnabled)
-			{
-				FString ChannelNameFString(ChannelName);
-				UE::Trace::ToggleChannel(ChannelNameFString.GetCharArray().GetData(), false);
-			}
+			UE_LOG(LogChaosVDRuntime, Log, TEXT("[%s] Successful disconnect attempt!."), ANSI_TO_TCHAR(__FUNCTION__));
+		}
+		else
+		{
+			FailureReason = LOCTEXT("FailedToStopActiveRecordingErrorMessage", "Failed to Stop active Trace Session.");
 		}
-		, nullptr);
+	}
 
+	SaveAndDisabledCurrentEnabledTraceChannels();
 
-	UE::Trace::ToggleChannel(TEXT("ChaosVDChannel"), true); 
-	UE::Trace::ToggleChannel(TEXT("Frame"), true);
-	UE::Trace::ToggleChannel(TEXT("Log"), true);
+	EnableRequiredTraceChannels();
 
 	FTraceAuxiliary::FOptions TracingOptions;
 	TracingOptions.bExcludeTail = true;
@@ -189,11 +195,6 @@ void FChaosVDRuntimeModule::StartRecording(TConstArrayView<FString> Args)
 	}
 	else if(Args[0] == TEXT("Server"))
 	{
-		if (FTraceAuxiliary::IsConnected())
-		{
-			FTraceAuxiliary::Stop();
-		}
-
 		const FString Target = Args.IsValidIndex(1) ? Args[1] : TEXT("127.0.0.1");
 
 		bIsRecording = FTraceAuxiliary::Start(
@@ -241,7 +242,6 @@ void FChaosVDRuntimeModule::StartRecording(TConstArrayView<FString> Args)
 			RecordingStartFailedDelegate.Broadcast(FailureReason);
 		}	
 	}
-
 }
 
 void FChaosVDRuntimeModule::StopRecording()
@@ -255,9 +255,7 @@ void FChaosVDRuntimeModule::StopRecording()
 	FTraceAuxiliary::OnTraceStopped.RemoveAll(this);
 
 #if UE_TRACE_ENABLED
-
-	UE::Trace::ToggleChannel(TEXT("ChaosVDChannel"), false);
-	UE::Trace::ToggleChannel(TEXT("Frame"), false); 
+	RestoreTraceChannelsToPreRecordingState();
 
 	StopTrace();
 #endif
@@ -303,4 +301,84 @@ void FChaosVDRuntimeModule::HandleTraceStopRequest(FTraceAuxiliary::EConnectionT
 	bRequestedStop = false;
 }
 
+bool FChaosVDRuntimeModule::WaitForTraceSessionDisconnect()
+{
+	float MaxWaitTime = CVarChaosVDMaxTimeToWaitForDisconnect->GetFloat();
+	float CurrentWaitTime = 0.0f;
+
+#if WITH_EDITOR
+	FScopedSlowTask DisconnectAttemptSlowTask(MaxWaitTime, LOCTEXT("DisconnectAttemptMessage", " Active Trace Session detected, attempting to disconnect ..."));
+
+	constexpr bool bShowCancelButton = false;
+	constexpr bool bAllowInPIE = true;
+	DisconnectAttemptSlowTask.MakeDialog(bShowCancelButton, bAllowInPIE);
+#endif
+
+	while (CurrentWaitTime < MaxWaitTime)
+	{
+		constexpr float WaitInterval = 0.1f;
+		FPlatformProcess::Sleep(0.1f);
+
+		if (!FTraceAuxiliary::IsConnected())
+		{
+			return true;
+		}
+
+		// We don't need to be precise for this, we can just accumulate the wait
+		CurrentWaitTime += WaitInterval;
+
+#if WITH_EDITOR
+		DisconnectAttemptSlowTask.EnterProgressFrame(CurrentWaitTime);
+#endif
+	}
+
+	return FTraceAuxiliary::IsConnected();
+}
+
+void FChaosVDRuntimeModule::SaveAndDisabledCurrentEnabledTraceChannels()
+{
+	// Until we support allowing other channels, indicate in the logs that we are disabling everything else
+	UE_LOG(LogChaosVDRuntime, Log, TEXT("[%s] Disabling additional trace channels..."), ANSI_TO_TCHAR(__FUNCTION__));
+
+#if UE_TRACE_ENABLED
+	OriginalTraceChannelsState.Reset();
+
+	// Disable any enabled additional channel
+	UE::Trace::EnumerateChannels([](const ANSICHAR* ChannelName, bool bEnabled, void* SavedTraceChannelsPtr)
+	{
+		TMap<FString, bool>* SavedTraceChannels = static_cast<TMap<FString, bool>*>(SavedTraceChannelsPtr);
+		FString ChannelNameFString(ChannelName);
+		SavedTraceChannels->Add(ChannelNameFString, bEnabled);
+		if (bEnabled)
+		{
+			UE::Trace::ToggleChannel(ChannelNameFString.GetCharArray().GetData(), false);
+		}
+	}
+	, &OriginalTraceChannelsState);
+#endif
+}
+
+void FChaosVDRuntimeModule::RestoreTraceChannelsToPreRecordingState()
+{
+#if UE_TRACE_ENABLED
+	UE_LOG(LogChaosVDRuntime, Log, TEXT("[%s] Restoring trace channels state..."), ANSI_TO_TCHAR(__FUNCTION__));
+
+	for (const TPair<FString, bool>& ChannelWithState : OriginalTraceChannelsState)
+	{
+		UE::Trace::ToggleChannel(GetData(ChannelWithState.Key), ChannelWithState.Value); 
+	}
+	
+	OriginalTraceChannelsState.Reset();
+#endif
+}
+
+void FChaosVDRuntimeModule::EnableRequiredTraceChannels()
+{
+#if UE_TRACE_ENABLED
+	UE::Trace::ToggleChannel(TEXT("ChaosVDChannel"), true); 
+	UE::Trace::ToggleChannel(TEXT("Frame"), true);
+	UE::Trace::ToggleChannel(TEXT("Log"), true);
+#endif
+}
+
 #undef LOCTEXT_NAMESPACE 
diff --git a/Engine/Source/Runtime/Experimental/ChaosVisualDebugger/Public/ChaosVDRuntimeModule.h b/Engine/Source/Runtime/Experimental/ChaosVisualDebugger/Public/ChaosVDRuntimeModule.h
--- a/Engine/Source/Runtime/Experimental/ChaosVisualDebugger/Public/ChaosVDRuntimeModule.h
+++ b/Engine/Source/Runtime/Experimental/ChaosVisualDebugger/Public/ChaosVDRuntimeModule.h
@@ -119,6 +119,12 @@ private:
 	 */
 	void HandleTraceStopRequest(FTraceAuxiliary::EConnectionType TraceType, const FString& TraceDestination);
 
+	bool WaitForTraceSessionDisconnect();
+
+	void EnableRequiredTraceChannels();
+	void SaveAndDisabledCurrentEnabledTraceChannels();
+	void RestoreTraceChannelsToPreRecordingState();
+
 	bool bIsRecording = false;
 	bool bRequestedStop = false;
 
@@ -136,5 +142,7 @@ private:
 
 	FString LastRecordingFileNamePath;
 
+	TMap<FString, bool> OriginalTraceChannelsState;
+
 	static FRWLock DelegatesRWLock;
 };
