Iris
- Fixed issue where only client would print default state of cdo used to initialize network state.

diff --git a/Engine/Source/Runtime/Experimental/Iris/Core/Private/Iris/ReplicationSystem/ObjectReplicationBridge.cpp b/Engine/Source/Runtime/Experimental/Iris/Core/Private/Iris/ReplicationSystem/ObjectReplicationBridge.cpp
--- a/Engine/Source/Runtime/Experimental/Iris/Core/Private/Iris/ReplicationSystem/ObjectReplicationBridge.cpp
+++ b/Engine/Source/Runtime/Experimental/Iris/Core/Private/Iris/ReplicationSystem/ObjectReplicationBridge.cpp
@@ -46,7 +46,7 @@
 
 namespace UE::Net::Private::ObjectBridgeDebugging
 {
-	extern void RemoteProtocolMismatchDetected(UReplicationSystem* , uint32, const FReplicationFragments&, const UObject*, const UObject*);
+	extern void RemoteProtocolMismatchDetected(TMap<FObjectKey, bool>& ArchetypesAlreadyPrinted, UReplicationSystem* , uint32, const FReplicationFragments&, const UObject*, const UObject*);
 }
 
 #ifndef UE_IRIS_VALIDATE_PROTOCOLS
@@ -908,7 +908,7 @@ FReplicationBridgeCreateNetRefHandleResult UObjectReplicationBridge::CreateNetRe
 
 		if (UE_LOG_ACTIVE(LogIris, Error))
 		{
-			UE::Net::Private::ObjectBridgeDebugging::RemoteProtocolMismatchDetected(ReplicationSystem, Context.ConnectionId, RegisteredFragments, ArchetypeOrCDOUsedAsKey, InstancePtr);
+			UE::Net::Private::ObjectBridgeDebugging::RemoteProtocolMismatchDetected(ArchetypesAlreadyPrinted, ReplicationSystem, Context.ConnectionId, RegisteredFragments, ArchetypeOrCDOUsedAsKey, InstancePtr);
 		}
 
 		FIrisDelegates::GetCriticalErrorDetectedDelegate().Broadcast(ReplicationSystem);
@@ -2294,7 +2294,7 @@ void UObjectReplicationBridge::OnProtocolMismatchReported(FNetRefHandle RefHandl
 			Fragments.Emplace(MoveTemp(FragmentInfo));
 		}
 
-		UE::Net::Private::ObjectBridgeDebugging::RemoteProtocolMismatchDetected(ReplicationSystem, 0 /*TODO: Local ConnectionId*/, Fragments, ObjArchetype, ObjInstance);
+		UE::Net::Private::ObjectBridgeDebugging::RemoteProtocolMismatchDetected(ArchetypesAlreadyPrinted, ReplicationSystem, 0 /*TODO: Local ConnectionId*/, Fragments, ObjArchetype, ObjInstance);
 	}
 }
 
diff --git a/Engine/Source/Runtime/Experimental/Iris/Core/Private/Iris/ReplicationSystem/ObjectReplicationBridgeDebugging.cpp b/Engine/Source/Runtime/Experimental/Iris/Core/Private/Iris/ReplicationSystem/ObjectReplicationBridgeDebugging.cpp
--- a/Engine/Source/Runtime/Experimental/Iris/Core/Private/Iris/ReplicationSystem/ObjectReplicationBridgeDebugging.cpp
+++ b/Engine/Source/Runtime/Experimental/Iris/Core/Private/Iris/ReplicationSystem/ObjectReplicationBridgeDebugging.cpp
@@ -150,14 +150,13 @@ void PrintDefaultNetObjectState(UReplicationSystem* ReplicationSystem, uint32 Co
 	NetSerializationContext.SetLocalConnectionId(ConnectionId);
 
 	FReplicationInstanceOperations::OutputInternalDefaultStateToString(NetSerializationContext, StringBuilder, RegisteredFragments);
+	FReplicationInstanceOperations::OutputInternalDefaultStateMemberHashesToString(ReplicationSystem, StringBuilder, RegisteredFragments);
 }
 
-void RemoteProtocolMismatchDetected(UReplicationSystem* ReplicationSystem, uint32 ConnectionId, const FReplicationFragments& RegisteredFragments, const UObject* ArchetypeOrCDOKey, const UObject* InstancePtr)
+void RemoteProtocolMismatchDetected(TMap<FObjectKey, bool>& ArchetypesAlreadyPrinted, UReplicationSystem* ReplicationSystem, uint32 ConnectionId, const FReplicationFragments& RegisteredFragments, const UObject* ArchetypeOrCDOKey, const UObject* InstancePtr)
 {
 	if (UE_LOG_ACTIVE(LogIris, Error))
 	{
-		static TMap<FObjectKey, bool> ArchetypesAlreadyPrinted;
-
 		// Only print the CDO state once
 		if (ArchetypesAlreadyPrinted.Find(FObjectKey(ArchetypeOrCDOKey)) == nullptr)
 		{
diff --git a/Engine/Source/Runtime/Experimental/Iris/Core/Private/Iris/ReplicationSystem/ReplicationOperations.cpp b/Engine/Source/Runtime/Experimental/Iris/Core/Private/Iris/ReplicationSystem/ReplicationOperations.cpp
--- a/Engine/Source/Runtime/Experimental/Iris/Core/Private/Iris/ReplicationSystem/ReplicationOperations.cpp
+++ b/Engine/Source/Runtime/Experimental/Iris/Core/Private/Iris/ReplicationSystem/ReplicationOperations.cpp
@@ -5,6 +5,7 @@
 #include "Iris/Core/IrisLog.h"
 #include "Net/Core/NetBitArray.h"
 #include "Iris/ReplicationState/ReplicationStateUtil.h"
+#include "Iris/ReplicationState/InternalReplicationStateDescriptorUtils.h"
 #include "Iris/ReplicationSystem/ObjectReferenceCache.h"
 #include "Iris/ReplicationSystem/ReplicationFragment.h"
 #include "Iris/ReplicationSystem/ReplicationOperationsInternal.h"
@@ -27,6 +28,9 @@
 #include "HAL/IConsoleManager.h"
 #include "UObject/CoreNetTypes.h"
 
+#include "CoreTypes.h"
+
+
 namespace UE::Net::Private
 {
 
@@ -266,6 +270,102 @@ bool FReplicationStateOperations::Validate(FNetSerializationContext& Context, co
 	return true;
 }
 
+static void OutputDescriptorMemberHashToString(FNetSerializationContext Context, uint32* HashBuffer, uint32 Indentation, FStringBuilderBase& StringBuilder, const uint8* RESTRICT SrcInternalBuffer, const FReplicationStateDescriptor* Descriptor)
+{
+	const FReplicationStateMemberDescriptor* MemberDescriptors = Descriptor->MemberDescriptors;
+	const FReplicationStateMemberSerializerDescriptor* MemberSerializerDescriptors = Descriptor->MemberSerializerDescriptors;
+	const FReplicationStateMemberDebugDescriptor* MemberDebugDescriptors = Descriptor->MemberDebugDescriptors;
+	const uint32 MemberCount = Descriptor->MemberCount;
+
+	FNetBitStreamWriter& Writer = *Context.GetBitStreamWriter();
+
+	auto AppendMemberAndIndentation = [](FStringBuilderBase& StringBuilder, uint32 MemberIndex, uint32 Indentation)
+	{
+		while (Indentation--)
+		{
+			StringBuilder.Append(TEXT("    "));
+		};
+		StringBuilder.Appendf(TEXT("Member: %u "), MemberIndex);
+	};
+
+	for (uint32 MemberIt = 0; MemberIt < MemberCount; ++MemberIt)
+	{
+		const FReplicationStateMemberDescriptor& MemberDescriptor = MemberDescriptors[MemberIt];
+		const FReplicationStateMemberSerializerDescriptor& MemberSerializerDescriptor = MemberSerializerDescriptors[MemberIt];
+
+		AppendMemberAndIndentation(StringBuilder, MemberIt, Indentation);
+
+		// Traverse struct members as well.
+		if (IsUsingStructNetSerializer(MemberSerializerDescriptor))
+		{
+			const FStructNetSerializerConfig* StructConfig = static_cast<const FStructNetSerializerConfig*>(MemberSerializerDescriptor.SerializerConfig);
+			const FReplicationStateDescriptor* StructDescriptor = StructConfig->StateDescriptor;
+
+			StringBuilder.Appendf(TEXT("%s\n"), ToCStr(MemberDebugDescriptors[MemberIt].DebugName));
+			OutputDescriptorMemberHashToString(Context, HashBuffer, Indentation + 1, StringBuilder, SrcInternalBuffer + MemberDescriptor.InternalMemberOffset, StructDescriptor);
+			return;	
+		}
+		else
+		{
+			// Reset writer as we want to capture this per member
+			Writer.Seek(0U);
+
+			FNetSerializeArgs Args;
+			Args.Version = 0;
+			Args.NetSerializerConfig = MemberSerializerDescriptor.SerializerConfig;
+			Args.Source = reinterpret_cast<NetSerializerValuePointer>(SrcInternalBuffer + MemberDescriptor.InternalMemberOffset);
+
+			MemberSerializerDescriptor.Serializer->Serialize(Context, Args);
+
+			// Make sure last byte has well defined data as we include it in hash
+			if (const uint32 BitsToFill = (8U - (Writer.GetPosBits() & 7U)) & 7U)
+			{
+				Writer.WriteBits(0U, BitsToFill);
+			}
+
+			Writer.CommitWrites();
+
+			// Calculate and print hash 
+			// Hash the serialized buffer
+			const uint64 MemberDefaultHash = CityHash64((const char*)HashBuffer, Writer.GetPosBytes());
+		
+			StringBuilder.Appendf(TEXT("%s Serializer: %s Hash: 0x%" UINT64_x_FMT "\n"), ToCStr(MemberDebugDescriptors[MemberIt].DebugName), ToCStr(MemberSerializerDescriptor.Serializer->Name), MemberDefaultHash);
+		}
+	}
+};
+
+void FReplicationStateOperations::OutputDefaultStateMembersHashToString(UReplicationSystem* ReplicationSystem, FStringBuilderBase& StringBuilder, const FReplicationStateDescriptor* Descriptor)
+{
+	if (!ReplicationSystem)
+	{
+		return;
+	}
+
+	if (!Descriptor || !Descriptor->DefaultStateBuffer)
+	{
+		return;
+	}
+
+	// Setup a temporary buffer for this.
+	constexpr uint32 BufferSize = 4096;
+	TArray<uint32> TempBuffer;
+	TempBuffer.SetNum(BufferSize/sizeof(uint32));
+
+	FNetBitStreamWriter Writer;
+	Writer.InitBytes(TempBuffer.GetData(), BufferSize);
+
+	// Setup context to serialize default state
+	FNetSerializationContext Context(&Writer);
+	Private::FInternalNetSerializationContext InternalContext(ReplicationSystem);
+	Context.SetInternalContext(&InternalContext);
+
+	// Tell serializers we are serializing default state. It allows serializers to opt out of being part of the checksum by simply not serializing any data.
+	Context.SetIsInitializingDefaultState(true);
+
+	OutputDescriptorMemberHashToString(Context, TempBuffer.GetData(), 0, StringBuilder, Descriptor->DefaultStateBuffer, Descriptor);
+}
+
+
 void FReplicationStateOperations::Serialize(FNetSerializationContext& Context, const uint8* RESTRICT SrcInternalBuffer, const FReplicationStateDescriptor* Descriptor)
 {
 	const FReplicationStateMemberDescriptor* MemberDescriptors = Descriptor->MemberDescriptors;
@@ -805,6 +905,22 @@ void FReplicationInstanceOperations::OutputInternalStateToString(FNetSerializati
 	}
 }
 
+void FReplicationInstanceOperations::OutputInternalDefaultStateMemberHashesToString(UReplicationSystem* ReplicationSystem, FStringBuilderBase& StringBuilder, const FReplicationFragments& Fragments)
+{
+	// Iterate over fragments
+	for (int32 Index=0; Index < Fragments.Num(); ++Index)
+	{
+		const FReplicationFragmentInfo& FragmentInfo = Fragments[Index];
+
+		const FReplicationStateDescriptor* CurrentDescriptor = FragmentInfo.Descriptor;
+		const FReplicationFragment* CurrentFragment = FragmentInfo.Fragment;
+
+		StringBuilder.Appendf(TEXT("Printing default state hash for members of [%d/%d] Fragment: %s DescriptorId: 0x%" UINT64_x_FMT " DefaultStateHash: 0x%" UINT64_x_FMT "\n"), 
+			Index+1, Fragments.Num(), ToCStr(CurrentDescriptor->DebugName), CurrentDescriptor->DescriptorIdentifier.Value, CurrentDescriptor->DescriptorIdentifier.DefaultStateHash);
+		FReplicationStateOperations::OutputDefaultStateMembersHashToString(ReplicationSystem, StringBuilder, CurrentDescriptor);
+	}	
+}
+
 void FReplicationInstanceOperations::OutputInternalDefaultStateToString(FNetSerializationContext& NetSerializationContext, FStringBuilderBase& StringBuilder, const FReplicationFragments& Fragments)
 {
 	// Iterate over fragments
diff --git a/Engine/Source/Runtime/Experimental/Iris/Core/Public/Iris/ReplicationSystem/ObjectReplicationBridge.h b/Engine/Source/Runtime/Experimental/Iris/Core/Public/Iris/ReplicationSystem/ObjectReplicationBridge.h
--- a/Engine/Source/Runtime/Experimental/Iris/Core/Public/Iris/ReplicationSystem/ObjectReplicationBridge.h
+++ b/Engine/Source/Runtime/Experimental/Iris/Core/Public/Iris/ReplicationSystem/ObjectReplicationBridge.h
@@ -497,4 +497,6 @@ private:
 
 protected:
 	bool bSuppressCreateInstanceFailedEnsure = false;
+private:
+	TMap<FObjectKey, bool> ArchetypesAlreadyPrinted;
 };
diff --git a/Engine/Source/Runtime/Experimental/Iris/Core/Public/Iris/ReplicationSystem/ReplicationOperations.h b/Engine/Source/Runtime/Experimental/Iris/Core/Public/Iris/ReplicationSystem/ReplicationOperations.h
--- a/Engine/Source/Runtime/Experimental/Iris/Core/Public/Iris/ReplicationSystem/ReplicationOperations.h
+++ b/Engine/Source/Runtime/Experimental/Iris/Core/Public/Iris/ReplicationSystem/ReplicationOperations.h
@@ -47,6 +47,9 @@ struct FReplicationStateOperations
 	/** Compare two quantized states return false if they are different */
 	static IRISCORE_API bool IsEqualQuantizedState(FNetSerializationContext& Context, const uint8* RESTRICT Source0, const uint8* RESTRICT Source1, const FReplicationStateDescriptor* Descriptor);
 
+	/** Debug method to output the per member defaultstate hash */
+	static IRISCORE_API void OutputDefaultStateMembersHashToString(UReplicationSystem* ReplicationSystem, FStringBuilderBase& StringBuilder, const FReplicationStateDescriptor* Descriptor);
+
 	/** Serialize a Replication state from internal buffer to BitStream */
 	static IRISCORE_API void Serialize(FNetSerializationContext& Context, const uint8* RESTRICT SrcInternalBuffer, const FReplicationStateDescriptor* Descriptor);
 
@@ -114,6 +117,10 @@ struct FReplicationInstanceOperations
 
 	/** Dequantize the default state for a replicated object with a given protocol and output the state to string. */
 	static IRISCORE_API void OutputInternalDefaultStateToString(FNetSerializationContext& NetSerializationContext, FStringBuilderBase& StringBuilder, const FReplicationFragments& Fragments);
+
+	/** Serialize and output per member default state hashes to string. */
+	static IRISCORE_API void OutputInternalDefaultStateMemberHashesToString(UReplicationSystem* ReplicationSystem, FStringBuilderBase& StringBuilder, const FReplicationFragments& Fragments);
+
 };
 
 struct FReplicationProtocolOperations
