Fixed bug with NetTokenStroreState not being properly reset if reused.

diff --git a/Engine/Source/Runtime/Experimental/Iris/Core/Private/Iris/ReplicationSystem/NetTokenStore.cpp b/Engine/Source/Runtime/Experimental/Iris/Core/Private/Iris/ReplicationSystem/NetTokenStore.cpp
--- a/Engine/Source/Runtime/Experimental/Iris/Core/Private/Iris/ReplicationSystem/NetTokenStore.cpp
+++ b/Engine/Source/Runtime/Experimental/Iris/Core/Private/Iris/ReplicationSystem/NetTokenStore.cpp
@@ -41,10 +41,10 @@ public:
 
 	void Reset()
 	{
-		// We reserve the first token as an invalid token
+		// We reserve the first token for each type as an invalid token
 		for (TArray<FNetTokenStoreKey>& TokenInfos : TokenInfoArray)
 		{
-			TokenInfos.Add(FNetTokenStoreKey());
+			TokenInfos.SetNum(1);
 		}
 	}
 
