Fixing bug where ReturnQuantized would automatically convert the ValueType of any FModuleInputValue to MAxis3D
- OutValue was constructed from Value which would run the MAxis3D constructor since Value is an FVector. This change constructs it from both Value and ValueType to preserve the type

diff --git a/Engine/Source/Runtime/Experimental/ChaosVehicles/ChaosVehiclesCore/Public/SimModule/ModuleInput.h b/Engine/Source/Runtime/Experimental/ChaosVehicles/ChaosVehiclesCore/Public/SimModule/ModuleInput.h
--- a/Engine/Source/Runtime/Experimental/ChaosVehicles/ChaosVehiclesCore/Public/SimModule/ModuleInput.h
+++ b/Engine/Source/Runtime/Experimental/ChaosVehicles/ChaosVehiclesCore/Public/SimModule/ModuleInput.h
@@ -175,7 +175,7 @@ public:
 
 	FModuleInputValue ReturnQuantized() const
 	{
-		FModuleInputValue OutValue = Value;
+		FModuleInputValue OutValue = FModuleInputValue(ValueType, Value);
 
 		switch (ValueType)
 		{
