Iris - Added delegates for when a valid connection is added and removed to the ReplicationSystem

diff --git a/Engine/Source/Runtime/Experimental/Iris/Core/Private/Iris/ReplicationSystem/ReplicationSystem.cpp b/Engine/Source/Runtime/Experimental/Iris/Core/Private/Iris/ReplicationSystem/ReplicationSystem.cpp
--- a/Engine/Source/Runtime/Experimental/Iris/Core/Private/Iris/ReplicationSystem/ReplicationSystem.cpp
+++ b/Engine/Source/Runtime/Experimental/Iris/Core/Private/Iris/ReplicationSystem/ReplicationSystem.cpp
@@ -20,6 +20,7 @@
 #include "Iris/ReplicationSystem/ReplicationBridge.h"
 #include "Iris/ReplicationSystem/ReplicationOperationsInternal.h"
 #include "Iris/ReplicationSystem/ReplicationReader.h"
+#include "Iris/ReplicationSystem/ReplicationSystemDelegates.h"
 #include "Iris/ReplicationSystem/ReplicationSystemInternal.h"
 #include "Iris/ReplicationSystem/ReplicationSystemTypes.h"
 #include "Iris/ReplicationSystem/ReplicationTypes.h"
@@ -59,6 +60,7 @@ public:
 	UReplicationSystem* ReplicationSystem;
 	FNetTokenStore* NetTokenStore;
 	FReplicationSystemInternal ReplicationSystemInternal;
+	FReplicationSystemDelegates Delegates;
 	uint64 IrisDebugHelperDummy = 0U;
 	FNetObjectGroupHandle NotReplicatedNetObjectGroupHandle;
 	FNetObjectGroupHandle NetGroupOwnerNetObjectGroupHandle;
@@ -572,6 +574,20 @@ public:
 		LLM_SCOPE_BYTAG(IrisConnection);
 
 		FReplicationConnections& Connections = ReplicationSystemInternal.GetConnections();
+		// Invalid or out of bounds ID?
+		if (ConnectionId == UE::Net::InvalidConnectionId || ConnectionId >= Connections.GetMaxConnectionCount())
+		{
+			UE_LOG(LogIris, Error, TEXT("UReplicationSystem::AddConnection called with a bad connection Id: %u. Max connection count: %u."), ConnectionId, Connections.GetMaxConnectionCount());
+			return;
+		}
+
+		// Already registered?
+		if (Connections.IsValidConnection(ConnectionId))
+		{
+			UE_LOG(LogIris, Error, TEXT("UReplicationSystem::AddConnection called with already added connection Id: %u."), ConnectionId);
+			return;
+		}
+
 		{
 			Connections.AddConnection(ConnectionId);
 
@@ -621,10 +637,19 @@ public:
 			FDeltaCompressionBaselineManager& DC = ReplicationSystemInternal.GetDeltaCompressionBaselineManager();
 			DC.AddConnection(ConnectionId);
 		}
+
+		FConnectionHandle ConnectionHandle(ConnectionId);
+		Delegates.ConnectionAddedDelegate.Broadcast(ConnectionHandle);
 	}
 
 	void RemoveConnection(uint32 ConnectionId)
 	{
+		FReplicationConnections& Connections = ReplicationSystemInternal.GetConnections();
+		if (!Connections.IsValidConnection(ConnectionId))
+		{
+			UE_LOG(LogIris, Error, TEXT("UReplicationSystem::RemoveConnection called for connection ID that isn't added: %u."), ConnectionId);
+		}
+
 		{
 			FDeltaCompressionBaselineManager& DC = ReplicationSystemInternal.GetDeltaCompressionBaselineManager();
 			DC.RemoveConnection(ConnectionId);
@@ -646,9 +671,11 @@ public:
 		}
 
 		{
-			FReplicationConnections& Connections = ReplicationSystemInternal.GetConnections();
 			Connections.RemoveConnection(ConnectionId);
 		}
+
+		FConnectionHandle ConnectionHandle(ConnectionId);
+		Delegates.ConnectionRemovedDelegate.Broadcast(ConnectionHandle);
 	}
 
 	void UpdateUnresolvableReferenceTracking()
@@ -1811,6 +1838,10 @@ void UReplicationSystem::CollectNetMetrics(UE::Net::FNetMetrics& OutNetMetrics)
 	Impl->CollectNetMetrics(OutNetMetrics);
 }
 
+UE::Net::FReplicationSystemDelegates& UReplicationSystem::GetDelegates()
+{
+	return Impl->Delegates;
+}
 
 #pragma region ReplicationSystemFactory
 
diff --git a/Engine/Source/Runtime/Experimental/Iris/Core/Public/Iris/ReplicationSystem/ReplicationSystem.h b/Engine/Source/Runtime/Experimental/Iris/Core/Public/Iris/ReplicationSystem/ReplicationSystem.h
--- a/Engine/Source/Runtime/Experimental/Iris/Core/Public/Iris/ReplicationSystem/ReplicationSystem.h
+++ b/Engine/Source/Runtime/Experimental/Iris/Core/Public/Iris/ReplicationSystem/ReplicationSystem.h
@@ -32,6 +32,7 @@ namespace UE::Net
 	class FNetCullDistanceOverrides;
 	class FNetObjectAttachment;	
 	class FNetObjectReference;
+	class FReplicationSystemDelegates;
 	class FReplicationSystemFactory;
 	class FNetTokenStore;
 	class FStringTokenStore;
@@ -694,6 +695,8 @@ public:
 	/** Collect relevant metrics that could be sent to analytics */
 	IRISCORE_API void CollectNetMetrics(UE::Net::FNetMetrics& OutNetMetrics) const;
 
+	IRISCORE_API UE::Net::FReplicationSystemDelegates& GetDelegates();
+
 public:
 	// For internal use and not exported.
 	UE::Net::Private::FReplicationSystemInternal* GetReplicationSystemInternal();
diff --git a/Engine/Source/Runtime/Experimental/Iris/Core/Public/Iris/ReplicationSystem/ReplicationSystemDelegates.h b/Engine/Source/Runtime/Experimental/Iris/Core/Public/Iris/ReplicationSystem/ReplicationSystemDelegates.h
new file mode 100644
--- /dev/null
+++ b/Engine/Source/Runtime/Experimental/Iris/Core/Public/Iris/ReplicationSystem/ReplicationSystemDelegates.h
@@ -0,0 +1,54 @@
+// Copyright Epic Games, Inc. All Rights Reserved.
+
+#pragma once
+
+#include "HAL/Platform.h"
+#include "Delegates/Delegate.h"
+#include "Net/Core/Connection/ConnectionHandle.h"
+
+namespace UE::Net::Private
+{
+	class FReplicationSystemImpl;
+}
+
+namespace UE::Net
+{
+
+class FReplicationSystemDelegates
+{
+public:
+	using FConnectionAddedDelegate = TMulticastDelegate<void(FConnectionHandle ConnectionHandle)>;
+	using FConnectionRemovedDelegate = TMulticastDelegate<void(FConnectionHandle ConnectionHandle)>;
+
+	/** 
+	 * Returns a delegate registration instance allowing the caller to register their FConnectionAddedDelegate. 
+	 * The delegate will be called when a valid and not previously added connection is registered via a FReplicationSystem::AddConnection call.
+	 * Currently only parent connections will call the delegates.
+	 */
+	FConnectionAddedDelegate::RegistrationType& OnConnectionAdded(); 
+
+	/**
+	 * Returns a delegate registration instance allowing the caller to register their FConnectionRemovedDelegate.
+	 * The delegate will be called when a previously successfully added connection is removed via a FReplicationSystem::RemoveConnection call. 
+	 * Currently only parent connections will call the delegates.
+	 */
+	FConnectionAddedDelegate::RegistrationType& OnConnectionRemoved(); 
+
+private:
+	friend UE::Net::Private::FReplicationSystemImpl;
+
+	FConnectionAddedDelegate ConnectionAddedDelegate;
+	FConnectionRemovedDelegate ConnectionRemovedDelegate;
+};
+
+inline FReplicationSystemDelegates::FConnectionAddedDelegate::RegistrationType& FReplicationSystemDelegates::OnConnectionAdded()
+{
+	return ConnectionAddedDelegate;
+}
+
+inline FReplicationSystemDelegates::FConnectionRemovedDelegate::RegistrationType& FReplicationSystemDelegates::OnConnectionRemoved()
+{
+	return ConnectionRemovedDelegate;
+}
+
+}
