Dataflow: support copy-pasting option pins

We were already doing what we needed to do to fixup nodes after serialization in their virtual Serialize methods.
Added a PostSerialize virtual method and moved this code there. Call this after Serialize, and also after pasting copied nodes.

diff --git a/Engine/Plugins/ChaosClothAssetEditor/Source/ChaosClothAssetDataflowNodes/Private/ChaosClothAsset/MergeClothCollectionsNode.cpp b/Engine/Plugins/ChaosClothAssetEditor/Source/ChaosClothAssetDataflowNodes/Private/ChaosClothAsset/MergeClothCollectionsNode.cpp
--- a/Engine/Plugins/ChaosClothAssetEditor/Source/ChaosClothAssetDataflowNodes/Private/ChaosClothAsset/MergeClothCollectionsNode.cpp
+++ b/Engine/Plugins/ChaosClothAssetEditor/Source/ChaosClothAssetDataflowNodes/Private/ChaosClothAsset/MergeClothCollectionsNode.cpp
@@ -474,7 +474,7 @@ void FChaosClothAssetMergeClothCollectionsNode_v2::OnPinRemoved(const UE::Datafl
 	return Super::OnPinRemoved(Pin);
 }
 
-void FChaosClothAssetMergeClothCollectionsNode_v2::Serialize(FArchive& Ar)
+void FChaosClothAssetMergeClothCollectionsNode_v2::PostSerialize(const FArchive& Ar)
 {
 	if (Ar.IsLoading())
 	{
@@ -755,7 +755,7 @@ PRAGMA_ENABLE_DEPRECATION_WARNINGS
 	}
 }
 
-void FChaosClothAssetMergeClothCollectionsNode::Serialize(FArchive& Ar)
+void FChaosClothAssetMergeClothCollectionsNode::PostSerialize(const FArchive& Ar)
 {
 	if (Ar.IsLoading())
 	{
diff --git a/Engine/Plugins/ChaosClothAssetEditor/Source/ChaosClothAssetDataflowNodes/Private/ChaosClothAsset/ProxyDeformerNode.cpp b/Engine/Plugins/ChaosClothAssetEditor/Source/ChaosClothAssetDataflowNodes/Private/ChaosClothAsset/ProxyDeformerNode.cpp
--- a/Engine/Plugins/ChaosClothAssetEditor/Source/ChaosClothAssetDataflowNodes/Private/ChaosClothAsset/ProxyDeformerNode.cpp
+++ b/Engine/Plugins/ChaosClothAssetEditor/Source/ChaosClothAssetDataflowNodes/Private/ChaosClothAsset/ProxyDeformerNode.cpp
@@ -523,7 +523,7 @@ void FChaosClothAssetProxyDeformerNode_v2::OnPinRemoved(const UE::Dataflow::FPin
 	return Super::OnPinRemoved(Pin);
 }
 
-void FChaosClothAssetProxyDeformerNode_v2::Serialize(FArchive& Ar)
+void FChaosClothAssetProxyDeformerNode_v2::PostSerialize(const FArchive& Ar)
 {
 	// Restore the pins when re-loading so they can get properly reconnected
 	if (Ar.IsLoading())
@@ -719,7 +719,7 @@ void FChaosClothAssetProxyDeformerNode::OnPinRemoved(const UE::Dataflow::FPin& P
 	return Super::OnPinRemoved(Pin);
 }
 
-void FChaosClothAssetProxyDeformerNode::Serialize(FArchive& Ar)
+void FChaosClothAssetProxyDeformerNode::PostSerialize(const FArchive& Ar)
 {
 	// Restore the pins when re-loading so they can get properly reconnected
 	if (Ar.IsLoading())
diff --git a/Engine/Plugins/ChaosClothAssetEditor/Source/ChaosClothAssetDataflowNodes/Private/ChaosClothAsset/TerminalNode.cpp b/Engine/Plugins/ChaosClothAssetEditor/Source/ChaosClothAssetDataflowNodes/Private/ChaosClothAsset/TerminalNode.cpp
--- a/Engine/Plugins/ChaosClothAssetEditor/Source/ChaosClothAssetDataflowNodes/Private/ChaosClothAsset/TerminalNode.cpp
+++ b/Engine/Plugins/ChaosClothAssetEditor/Source/ChaosClothAssetDataflowNodes/Private/ChaosClothAsset/TerminalNode.cpp
@@ -328,7 +328,7 @@ UE::Dataflow::TConnectionReference<FManagedArrayCollection> FChaosClothAssetTerm
 	return { &CollectionLods[Index], Index, &CollectionLods };
 }
 
-void FChaosClothAssetTerminalNode_v2::Serialize(FArchive& Ar)
+void FChaosClothAssetTerminalNode_v2::PostSerialize(const FArchive& Ar)
 {
 	// because we add pins we need to make sure we restore them when loading
 	// to make sure they can get properly reconnected
@@ -696,7 +696,7 @@ TArray<TSharedRef<FManagedArrayCollection>> FChaosClothAssetTerminalNode::GetCle
 	return CollectionLodValues;
 }
 
-void FChaosClothAssetTerminalNode::Serialize(FArchive& Ar)
+void FChaosClothAssetTerminalNode::PostSerialize(const FArchive& Ar)
 {
 	// because we add pins we need to make sure we restore them when loading
 	// to make sure they can get properly reconnected
diff --git a/Engine/Plugins/ChaosClothAssetEditor/Source/ChaosClothAssetDataflowNodes/Public/ChaosClothAsset/MergeClothCollectionsNode.h b/Engine/Plugins/ChaosClothAssetEditor/Source/ChaosClothAssetDataflowNodes/Public/ChaosClothAsset/MergeClothCollectionsNode.h
--- a/Engine/Plugins/ChaosClothAssetEditor/Source/ChaosClothAssetDataflowNodes/Public/ChaosClothAsset/MergeClothCollectionsNode.h
+++ b/Engine/Plugins/ChaosClothAssetEditor/Source/ChaosClothAssetDataflowNodes/Public/ChaosClothAsset/MergeClothCollectionsNode.h
@@ -37,7 +37,7 @@ private:
 	virtual bool CanRemovePin() const override { return Collections.Num() > NumInitialOptionalInputs; }
 	virtual TArray<UE::Dataflow::FPin> GetPinsToRemove() const override;
 	virtual void OnPinRemoved(const UE::Dataflow::FPin& Pin) override;
-	virtual void Serialize(FArchive& Ar) override;
+	virtual void PostSerialize(const FArchive& Ar) override;
 	//~ End FDataflowNode interface
 
 
@@ -93,7 +93,7 @@ private:
 	virtual bool CanRemovePin() const override { return NumInputs > NumInitialOptionalInputs; }
 	virtual TArray<UE::Dataflow::FPin> GetPinsToRemove() const override;
 	virtual void OnPinRemoved(const UE::Dataflow::FPin& Pin) override;
-	virtual void Serialize(FArchive& Ar) override;
+	virtual void PostSerialize(const FArchive& Ar) override;
 	//~ End FDataflowNode interface
 
 	TArray<const FManagedArrayCollection*> GetCollections() const;
diff --git a/Engine/Plugins/ChaosClothAssetEditor/Source/ChaosClothAssetDataflowNodes/Public/ChaosClothAsset/ProxyDeformerNode.h b/Engine/Plugins/ChaosClothAssetEditor/Source/ChaosClothAssetDataflowNodes/Public/ChaosClothAsset/ProxyDeformerNode.h
--- a/Engine/Plugins/ChaosClothAssetEditor/Source/ChaosClothAssetDataflowNodes/Public/ChaosClothAsset/ProxyDeformerNode.h
+++ b/Engine/Plugins/ChaosClothAssetEditor/Source/ChaosClothAssetDataflowNodes/Public/ChaosClothAsset/ProxyDeformerNode.h
@@ -63,7 +63,7 @@ private:
 	virtual bool CanRemovePin() const override { return SelectionFilterSets.Num() > NumInitialSelectionFilterSets; }
 	virtual TArray<UE::Dataflow::FPin> GetPinsToRemove() const override;
 	virtual void OnPinRemoved(const UE::Dataflow::FPin& Pin) override;
-	virtual void Serialize(FArchive& Ar) override;
+	virtual void PostSerialize(const FArchive& Ar) override;
 	//~ End FDataflowNode interface
 
 	TArray<TPair<FName, FName>> GetSelectionFilterNames(UE::Dataflow::FContext& Context) const;
@@ -179,7 +179,7 @@ PRAGMA_ENABLE_DEPRECATION_WARNINGS
 	virtual bool CanRemovePin() const override { return NumFilterSets > NumInitialOptionalInputs; }
 	virtual TArray<UE::Dataflow::FPin> GetPinsToRemove() const override;
 	virtual void OnPinRemoved(const UE::Dataflow::FPin& Pin) override;
-	virtual void Serialize(FArchive& Ar) override;
+	virtual void PostSerialize(const FArchive& Ar) override;
 	//~ End FDataflowNode interface
 
 	TArray<FName> GetSelectionFilterNames(UE::Dataflow::FContext& Context) const;
diff --git a/Engine/Plugins/ChaosClothAssetEditor/Source/ChaosClothAssetDataflowNodes/Public/ChaosClothAsset/TerminalNode.h b/Engine/Plugins/ChaosClothAssetEditor/Source/ChaosClothAssetDataflowNodes/Public/ChaosClothAsset/TerminalNode.h
--- a/Engine/Plugins/ChaosClothAssetEditor/Source/ChaosClothAssetDataflowNodes/Public/ChaosClothAsset/TerminalNode.h
+++ b/Engine/Plugins/ChaosClothAssetEditor/Source/ChaosClothAssetDataflowNodes/Public/ChaosClothAsset/TerminalNode.h
@@ -54,7 +54,7 @@ private:
 	virtual bool CanRemovePin() const override { return CollectionLods.Num() > NumInitialCollectionLods; }
 	virtual TArray<UE::Dataflow::FPin> GetPinsToRemove() const override;
 	virtual void OnPinRemoved(const UE::Dataflow::FPin& Pin) override;
-	virtual void Serialize(FArchive& Ar) override;
+	virtual void PostSerialize(const FArchive& Ar) override;
 	//~ End FDataflowNode interface
 
 	TArray<TSharedRef<FManagedArrayCollection>> GetCleanedCollectionLodValues(UE::Dataflow::FContext& Context) const;
@@ -125,7 +125,7 @@ private:
 	virtual bool CanRemovePin() const override { return NumLods > 1; }
 	virtual TArray<UE::Dataflow::FPin> GetPinsToRemove() const override;
 	virtual void OnPinRemoved(const UE::Dataflow::FPin& Pin) override;
-	virtual void Serialize(FArchive& Ar) override;
+	virtual void PostSerialize(const FArchive& Ar) override;
 	//~ End FDataflowNode interface
 
 	TArray<const FManagedArrayCollection*> GetCollectionLods() const;
diff --git a/Engine/Plugins/Experimental/Dataflow/Source/DataflowEditor/Private/Dataflow/DataflowSNode.cpp b/Engine/Plugins/Experimental/Dataflow/Source/DataflowEditor/Private/Dataflow/DataflowSNode.cpp
--- a/Engine/Plugins/Experimental/Dataflow/Source/DataflowEditor/Private/Dataflow/DataflowSNode.cpp
+++ b/Engine/Plugins/Experimental/Dataflow/Source/DataflowEditor/Private/Dataflow/DataflowSNode.cpp
@@ -579,6 +579,10 @@ static UDataflowEdNode* CreateNodeFromPaste(UDataflow* Dataflow, UEdGraphPin* Fr
 				{
 					DataflowNode->TypedScriptStruct()->ImportText(*NodeProperties, DataflowNode.Get(), nullptr, EPropertyPortFlags::PPF_None, nullptr, DataflowNode->TypedScriptStruct()->GetName(), true);
 				}
+				// Do any post-import fixup.
+				FArchive Ar;
+				Ar.SetIsLoading(true);
+				DataflowNode->PostSerialize(Ar);
 
 				EdNode->CreateNewGuid();
 				EdNode->PostPlacedNewNode();
diff --git a/Engine/Source/Runtime/Experimental/Dataflow/Core/Private/Dataflow/DataflowCoreNodes.cpp b/Engine/Source/Runtime/Experimental/Dataflow/Core/Private/Dataflow/DataflowCoreNodes.cpp
--- a/Engine/Source/Runtime/Experimental/Dataflow/Core/Private/Dataflow/DataflowCoreNodes.cpp
+++ b/Engine/Source/Runtime/Experimental/Dataflow/Core/Private/Dataflow/DataflowCoreNodes.cpp
@@ -181,7 +181,7 @@ void FDataflowSelectNode::OnPinRemoved(const UE::Dataflow::FPin& Pin)
 	return Super::OnPinRemoved(Pin);
 }
 
-void FDataflowSelectNode::Serialize(FArchive& Ar)
+void FDataflowSelectNode::PostSerialize(const FArchive& Ar)
 {
 	if (Ar.IsLoading())
 	{
diff --git a/Engine/Source/Runtime/Experimental/Dataflow/Core/Public/Dataflow/DataflowCoreNodes.h b/Engine/Source/Runtime/Experimental/Dataflow/Core/Public/Dataflow/DataflowCoreNodes.h
--- a/Engine/Source/Runtime/Experimental/Dataflow/Core/Public/Dataflow/DataflowCoreNodes.h
+++ b/Engine/Source/Runtime/Experimental/Dataflow/Core/Public/Dataflow/DataflowCoreNodes.h
@@ -93,7 +93,7 @@ private:
 	virtual bool CanRemovePin() const override { return Inputs.Num() > NumInitialInputs; }
 	virtual TArray<UE::Dataflow::FPin> GetPinsToRemove() const override;
 	virtual void OnPinRemoved(const UE::Dataflow::FPin& Pin) override;
-	virtual void Serialize(FArchive& Ar) override;
+	virtual void PostSerialize(const FArchive& Ar) override;
 
 	UE::Dataflow::TConnectionReference<FDataflowAnyType> GetConnectionReference(int32 Index) const;
 
diff --git a/Engine/Source/Runtime/Experimental/Dataflow/Core/Public/Dataflow/DataflowNode.h b/Engine/Source/Runtime/Experimental/Dataflow/Core/Public/Dataflow/DataflowNode.h
--- a/Engine/Source/Runtime/Experimental/Dataflow/Core/Public/Dataflow/DataflowNode.h
+++ b/Engine/Source/Runtime/Experimental/Dataflow/Core/Public/Dataflow/DataflowNode.h
@@ -204,6 +204,9 @@ PRAGMA_ENABLE_DEPRECATION_WARNINGS
 
 	/** Override this method to provide custom serialization for this node. */
 	virtual void Serialize(FArchive& Ar) {}
+	/** Override this method to provide custom post-serialization for this node. This method will be called after Serialize. It is also called after copy-paste with
+	 ArchiveState IsLoading. */
+	virtual void PostSerialize(const FArchive& Ar) {}
 
 	/** Override this method to provide custom reconnections when a node inputs has been deprecated and removed. */
 	virtual FDataflowInput* RedirectSerializedInput(const FName& MissingInputName) { return nullptr; }
@@ -678,7 +681,8 @@ public:																				\
 		UScriptStruct* const Struct = TYPE::StaticStruct();							\
 		Struct->SerializeTaggedProperties(Ar, (uint8*)this,							\
 		Struct, nullptr);															\
-		Serialize(Ar);}																\
+		Serialize(Ar);																\
+		PostSerialize(Ar);}															\
 	virtual FName GetDisplayName() const override { return TYPE::StaticDisplay(); }	\
 	virtual FName GetCategory() const override { return TYPE::StaticCategory(); }	\
 	virtual FString GetTags() const override { return TYPE::StaticTags(); }			\
