ChaosFlesh: CIS fix
backward compatibility for per-vertex kinematic attribute

diff --git a/Engine/Source/Runtime/Experimental/Chaos/Public/GeometryCollection/Facades/CollectionVertexBoneWeightsFacade.h b/Engine/Source/Runtime/Experimental/Chaos/Public/GeometryCollection/Facades/CollectionVertexBoneWeightsFacade.h
--- a/Engine/Source/Runtime/Experimental/Chaos/Public/GeometryCollection/Facades/CollectionVertexBoneWeightsFacade.h
+++ b/Engine/Source/Runtime/Experimental/Chaos/Public/GeometryCollection/Facades/CollectionVertexBoneWeightsFacade.h
@@ -67,8 +67,12 @@ namespace GeometryCollection::Facades
 		const TManagedArray< TArray<int32> >* FindBoneIndices()  const { return BoneIndexAttribute.Find(); }
 		const TManagedArray< TArray<int32> >& GetBoneIndices() const { return BoneIndexAttribute.Get(); }
 
-		/** Return if the vertex is kinematic */
-		CHAOS_API bool IsKinematicVertex(int32 VertexIndex) const { return KinematicAttribute.IsValidIndex(VertexIndex) && KinematicAttribute.Get()[VertexIndex] && BoneIndexAttribute.Get()[VertexIndex].Num();};
+		/** Return if the vertex is kinematic
+		Pre 5.5 we did not have per-vertex kinematic attribute. 
+		This supports defining kinematics without per-vertex kinematic attribute. */
+		CHAOS_API bool IsKinematicVertex(int32 VertexIndex) const { return (!KinematicAttribute.IsValid() || (KinematicAttribute.IsValidIndex(VertexIndex) && KinematicAttribute.Get()[VertexIndex]))
+			&& BoneIndexAttribute.IsValid() && BoneIndexAttribute.IsValidIndex(VertexIndex) && BoneIndexAttribute.Get()[VertexIndex].Num()
+			&& BoneWeightAttribute.IsValid() && BoneWeightAttribute.IsValidIndex(VertexIndex) && BoneWeightAttribute.Get()[VertexIndex].Num(); };
 
 		/** Return number of vertices */
 		CHAOS_API int32 NumVertices() const { return VerticesAttribute.Num(); };
