thread safety fix for Ias cache Materialize function

#rnx

diff --git a/Engine/Source/Runtime/Experimental/IoStore/OnDemand/Private/JournaledCache.cpp b/Engine/Source/Runtime/Experimental/IoStore/OnDemand/Private/JournaledCache.cpp
--- a/Engine/Source/Runtime/Experimental/IoStore/OnDemand/Private/JournaledCache.cpp
+++ b/Engine/Source/Runtime/Experimental/IoStore/OnDemand/Private/JournaledCache.cpp
@@ -2235,9 +2235,10 @@ void FJournaledCache::Materialize(
 
 	FServiceThread& ServiceThread = FServiceThread::Get();
 
-	uint32 ReadId = ServiceThread.ClaimReadId();
+	uint32 ReadId;
 	{
 		FScopeLock _(&Lock);
+		ReadId = ServiceThread.ClaimReadId();
 		PendingMaterializes.Add(ReadId, { MoveTemp(DoneEvent), &Status });
 	}
 
