Iris
* AActor::PostNetInit will only be called for dynamic actors
* NetObjectFactory now must return a flag if the instantiated object is dynamic.

diff --git a/Engine/Source/Runtime/Engine/Private/Net/Iris/ReplicationSystem/NetActorFactory.cpp b/Engine/Source/Runtime/Engine/Private/Net/Iris/ReplicationSystem/NetActorFactory.cpp
--- a/Engine/Source/Runtime/Engine/Private/Net/Iris/ReplicationSystem/NetActorFactory.cpp
+++ b/Engine/Source/Runtime/Engine/Private/Net/Iris/ReplicationSystem/NetActorFactory.cpp
@@ -248,7 +248,6 @@ TUniquePtr<UE::Net::FNetObjectCreationHeader> UNetActorFactory::CreateAndFillHea
 		}
 		else
 		{
-			//OutHeader.SpawnInfo = FactoryGetDefaultActorReplicationBridgeSpawnInfo();
 			check(!Header->SpawnInfo.Location.ContainsNaN());
 		}
 	}
@@ -472,8 +471,12 @@ void UNetActorFactory::PostInstantiation(const FPostInstantiationContext& Contex
 
 void UNetActorFactory::PostInit(const FPostInitContext& Context)
 {
-	AActor* Actor = CastChecked<AActor>(Context.Instance);
-	Actor->PostNetInit();
+	// PostNetInit is only called for dynamic actors
+	if (Context.Handle.IsDynamic())
+	{
+		AActor* Actor = CastChecked<AActor>(Context.Instance);
+		Actor->PostNetInit();
+	}
 }
 
  namespace UE::Net
diff --git a/Engine/Source/Runtime/Experimental/Iris/Core/Private/Iris/ReplicationSystem/ReplicationReader.cpp b/Engine/Source/Runtime/Experimental/Iris/Core/Private/Iris/ReplicationSystem/ReplicationReader.cpp
--- a/Engine/Source/Runtime/Experimental/Iris/Core/Private/Iris/ReplicationSystem/ReplicationReader.cpp
+++ b/Engine/Source/Runtime/Experimental/Iris/Core/Private/Iris/ReplicationSystem/ReplicationReader.cpp
@@ -101,6 +101,8 @@ static FAutoConsoleVariableRef CvarResolvedObjectsDispatchDebugging(
 
 static const FName NetError_FailedToFindAttachmentQueue("Failed to find attachment queue");
 
+//------------------------------------------------------------------------
+
 class FResolveAndCollectUnresolvedAndResolvedReferenceCollector
 {
 public:
@@ -152,6 +154,24 @@ private:
 	bool bIncludeInitState = false;
 };
 
+//------------------------------------------------------------------------
+
+struct FReplicationReader::FDispatchObjectInfo
+{
+	FInternalNetRefIndex InternalIndex = FNetRefHandleManager::InvalidInternalIndex;
+	FChangeMaskStorageOrPointer ChangeMaskOrPointer;
+	uint32 bIsInitialState : 1 = false;
+	uint32 bHasState : 1 = false;
+	uint32 bHasAttachments : 1 = false;
+	uint32 bDestroy : 1 = false;
+	uint32 bTearOff : 1 = false;
+	uint32 bDeferredEndReplication : 1 = false;
+	uint32 bShouldCallSubObjectCreatedFromReplication : 1 = false;
+	uint32 bDynamicObjectCreated : 1 = false;
+};
+
+//------------------------------------------------------------------------
+
 // Helper class to deal with management of ObjectsToDispatch allocations from our temporary allocator
 class FReplicationReader::FObjectsToDispatchArray
 {
@@ -204,6 +224,8 @@ private:
 	uint32 Capacity;		
 };
 
+//------------------------------------------------------------------------
+
 FReplicationReader::FReplicatedObjectInfo::FReplicatedObjectInfo()
 : InternalIndex(FNetRefHandleManager::InvalidInternalIndex)
 , Value(0U)
diff --git a/Engine/Source/Runtime/Experimental/Iris/Core/Private/Iris/ReplicationSystem/ReplicationReader.h b/Engine/Source/Runtime/Experimental/Iris/Core/Private/Iris/ReplicationSystem/ReplicationReader.h
--- a/Engine/Source/Runtime/Experimental/Iris/Core/Private/Iris/ReplicationSystem/ReplicationReader.h
+++ b/Engine/Source/Runtime/Experimental/Iris/Core/Private/Iris/ReplicationSystem/ReplicationReader.h
@@ -115,19 +115,7 @@ private:
 	};
 
 	// Temporary Data to dispatch
-	struct FDispatchObjectInfo
-	{
-		uint32 InternalIndex;
-		FChangeMaskStorageOrPointer ChangeMaskOrPointer;
-		uint32 bIsInitialState : 1;
-		uint32 bHasState : 1;
-		uint32 bHasAttachments : 1;
-		uint32 bDestroy : 1;
-		uint32 bTearOff : 1;
-		uint32 bDeferredEndReplication : 1;
-		uint32 bShouldCallSubObjectCreatedFromReplication : 1;
-	};
-
+	struct FDispatchObjectInfo;
 
 	enum : uint32
 	{
diff --git a/Engine/Source/Runtime/Experimental/Iris/Core/Public/Iris/ReplicationSystem/ReplicationBridgeTypes.h b/Engine/Source/Runtime/Experimental/Iris/Core/Public/Iris/ReplicationSystem/ReplicationBridgeTypes.h
--- a/Engine/Source/Runtime/Experimental/Iris/Core/Public/Iris/ReplicationSystem/ReplicationBridgeTypes.h
+++ b/Engine/Source/Runtime/Experimental/Iris/Core/Public/Iris/ReplicationSystem/ReplicationBridgeTypes.h
@@ -30,11 +30,12 @@ IRISCORE_API FString LexToString(EEndReplicationFlags EndReplicationFlags);
 
 //------------------------------------------------------------------------
 
-enum class EReplicationBridgeCreateNetRefHandleResultFlags : unsigned
+enum class EReplicationBridgeCreateNetRefHandleResultFlags : uint32
 {
 	None = 0U,
 	/** Whether the instance may be destroyed due to the remote peer requesting the object to be destroyed. If not then the object itself must not be destroyed. */
 	AllowDestroyInstanceFromRemote = 1U << 0U,
+	/** Set this flag if you created a subobject and want the RootObject to be notified of the subobject's creation. */
 	ShouldCallSubObjectCreatedFromReplication = AllowDestroyInstanceFromRemote << 1U,
 };
 ENUM_CLASS_FLAGS(EReplicationBridgeCreateNetRefHandleResultFlags);
@@ -49,7 +50,7 @@ struct FReplicationBridgeCreateNetRefHandleResult
 
 //------------------------------------------------------------------------
 
-enum class EReplicationBridgeDestroyInstanceReason : unsigned
+enum class EReplicationBridgeDestroyInstanceReason : uint32
 {
 	DoNotDestroy,
 	TearOff,
@@ -59,7 +60,7 @@ IRISCORE_API const TCHAR* LexToString(EReplicationBridgeDestroyInstanceReason Re
 
 //------------------------------------------------------------------------
 
-enum class EReplicationBridgeDestroyInstanceFlags : unsigned
+enum class EReplicationBridgeDestroyInstanceFlags : uint32
 {
 	None = 0U,
 	/** Whether the instance may be destroyed when instructed from the remote peer. This flag applies when the destroy reason is TearOff and torn off actors are to be destroyed as well as regular Destroy. */
