Iris: Fix off-by-one in max group limit

Max is 65535, not 65534.

#rnx
#tests Tested with DefaultEngine.ini [/Script/OnlineSubsystemUtils.IpNetDriver] ReplicationSystemConfigServer=(MaxNetObjectGroupCount=65534)

diff --git a/Engine/Source/Runtime/Experimental/Iris/Core/Private/Iris/ReplicationSystem/Filtering/NetObjectGroups.cpp b/Engine/Source/Runtime/Experimental/Iris/Core/Private/Iris/ReplicationSystem/Filtering/NetObjectGroups.cpp
--- a/Engine/Source/Runtime/Experimental/Iris/Core/Private/Iris/ReplicationSystem/Filtering/NetObjectGroups.cpp
+++ b/Engine/Source/Runtime/Experimental/Iris/Core/Private/Iris/ReplicationSystem/Filtering/NetObjectGroups.cpp
@@ -56,7 +56,7 @@ void FNetObjectGroups::Init(const FNetObjectGroupInitParams& Params)
 {
 	NetRefHandleManager = Params.NetRefHandleManager;
 
-	ensureMsgf(Params.MaxGroupCount < std::numeric_limits<FNetObjectGroupHandle::FGroupIndexType>::max(), TEXT("MaxGroupCount cannot exceed %u"), std::numeric_limits<FNetObjectGroupHandle::FGroupIndexType>::max());
+	ensureMsgf(Params.MaxGroupCount <= std::numeric_limits<FNetObjectGroupHandle::FGroupIndexType>::max(), TEXT("MaxGroupCount cannot exceed %u"), std::numeric_limits<FNetObjectGroupHandle::FGroupIndexType>::max());
 	MaxGroupCount = FMath::Clamp<uint32>(Params.MaxGroupCount, 0U, std::numeric_limits<FNetObjectGroupHandle::FGroupIndexType>::max());
 
 	// Reserve first as invalid group
diff --git a/Engine/Source/Runtime/Experimental/Iris/Core/Private/Iris/ReplicationSystem/Filtering/ReplicationFiltering.cpp b/Engine/Source/Runtime/Experimental/Iris/Core/Private/Iris/ReplicationSystem/Filtering/ReplicationFiltering.cpp
--- a/Engine/Source/Runtime/Experimental/Iris/Core/Private/Iris/ReplicationSystem/Filtering/ReplicationFiltering.cpp
+++ b/Engine/Source/Runtime/Experimental/Iris/Core/Private/Iris/ReplicationSystem/Filtering/ReplicationFiltering.cpp
@@ -208,7 +208,7 @@ void FReplicationFiltering::Init(FReplicationFilteringInitParams& Params)
 	// Group filtering
 	{
 		const uint32 InMaxGroupCount = Params.MaxGroupCount;
-		check(InMaxGroupCount < std::numeric_limits<FNetObjectGroupHandle::FGroupIndexType>::max());
+		check(InMaxGroupCount <= std::numeric_limits<FNetObjectGroupHandle::FGroupIndexType>::max());
 		MaxGroupCount = InMaxGroupCount;
 		GroupInfos.SetNumZeroed(InMaxGroupCount);
 		ExclusionFilterGroups.Init(InMaxGroupCount);
