Chaos::Fix instability issue in simd code of joint cached solver.

diff --git a/Engine/Source/Runtime/Experimental/Chaos/Private/Chaos/Joint/PBDJointCachedSolverGaussSeidel.cpp b/Engine/Source/Runtime/Experimental/Chaos/Private/Chaos/Joint/PBDJointCachedSolverGaussSeidel.cpp
--- a/Engine/Source/Runtime/Experimental/Chaos/Private/Chaos/Joint/PBDJointCachedSolverGaussSeidel.cpp
+++ b/Engine/Source/Runtime/Experimental/Chaos/Private/Chaos/Joint/PBDJointCachedSolverGaussSeidel.cpp
@@ -1639,7 +1639,7 @@ void FPBDJointCachedSolver::ApplyRotationSoftConstraintsSimd(
 		const FRealSingle SolverStiffnessf = FRealSingle(SolverStiffness);
 		VectorRegister4Float Stiffness = VectorLoadFloat1(&SolverStiffnessf);
 		VectorRegister4Float DeltaLambda = VectorMultiply(Stiffness, VectorDivide(VectorSubtract(VectorMultiply(RotationConstraints.Simd.ConstraintSoftStiffness, DeltaAngle),
-			VectorSubtract(VectorMultiply(RotationConstraints.Simd.ConstraintSoftDamping, AngVelDt), RotationConstraints.Simd.ConstraintLambda)), RotationConstraints.Simd.ConstraintSoftIM));
+			VectorAdd(VectorMultiply(RotationConstraints.Simd.ConstraintSoftDamping, AngVelDt), RotationConstraints.Simd.ConstraintLambda)), RotationConstraints.Simd.ConstraintSoftIM));
 
 		DeltaLambda = VectorSelect(NeedsSolve, DeltaLambda, VectorZeroFloat());
 		RotationConstraints.Simd.ConstraintLambda = VectorAdd(RotationConstraints.Simd.ConstraintLambda, DeltaLambda);
