FOnDemandInstallCache write-exclusive reads for platfroms that can't support both at the same time.
- using a pipe only for the current block

diff --git a/Engine/Source/Runtime/Experimental/IoStore/OnDemand/Private/OnDemandInstallCache.cpp b/Engine/Source/Runtime/Experimental/IoStore/OnDemand/Private/OnDemandInstallCache.cpp
--- a/Engine/Source/Runtime/Experimental/IoStore/OnDemand/Private/OnDemandInstallCache.cpp
+++ b/Engine/Source/Runtime/Experimental/IoStore/OnDemand/Private/OnDemandInstallCache.cpp
@@ -31,6 +31,14 @@
 #include <catch2/generators/catch_generators.hpp>
 #endif
 
+#ifndef UE_ONDEMANDINSTALLCACHE_EXCLUSIVE_WRITE
+#define UE_ONDEMANDINSTALLCACHE_EXCLUSIVE_WRITE (0)
+#endif
+
+#if UE_ONDEMANDINSTALLCACHE_EXCLUSIVE_WRITE
+#include "Tasks/Pipe.h"
+#endif // UE_ONDEMANDINSTALLCACHE_EXCLUSIVE_WRITE
+
 namespace UE::IoStore
 {
 
@@ -965,7 +973,8 @@ private:
 	FIoStatus									Purge(const FCasBlockInfoMap& BlockInfo, uint64 TotalBytesToPurge, uint64& OutTotalPurgedBytes);
 	bool										Resolve(FIoRequestImpl* Request);
 	void										CompleteRequest(FIoRequestImpl* Request, bool bFileReadWasCancelled);
-	FIoStatus									FlushPendingChunks(FPendingChunks& Block);
+	FIoStatus									FlushPendingChunks(FPendingChunks& Chunks);
+	FIoStatus									FlushPendingChunksImpl(FPendingChunks& Chunks);
 	FString										GetJournalFilename() const { return CacheDirectory / TEXT("cas.jrn"); }
 
 	FOnDemandIoStore&			IoStore;
@@ -977,6 +986,10 @@ private:
 	FIoRequestList				CompletedRequests;
 	UE::FMutex					Mutex;
 	uint64						MaxCacheSize{ 0 };
+
+#if UE_ONDEMANDINSTALLCACHE_EXCLUSIVE_WRITE
+	UE::Tasks::FPipe			ExclusivePipe{ UE_SOURCE_LOCATION };
+#endif // UE_ONDEMANDINSTALLCACHE_EXCLUSIVE_WRITE
 };
 
 ///////////////////////////////////////////////////////////////////////////////
@@ -1214,6 +1227,68 @@ bool FOnDemandInstallCache::Resolve(FIoRequestImpl* Request)
 	TRACE_IOSTORE_BACKEND_REQUEST_STARTED(Request, this);
 	Cas.TrackAccess(CasLoc.BlockId);
 
+#if UE_ONDEMANDINSTALLCACHE_EXCLUSIVE_WRITE
+	const bool bIsLocationInCurrentBlock = CasLoc.BlockId == CurrentBlock;
+	if (bIsLocationInCurrentBlock)
+	{
+		// The current block may have open writes which may cause async reads to fail
+		// on some platforms. Schedule the reads to happen on the same pipe as writes
+
+		// The internal request parameters are attached/owned by the I/O request via
+		// the backend data parameter. The chunk request is deleted in GetCompletedRequests
+		FChunkRequest::Attach(*Request, new FChunkRequest(
+			FSharedAsyncFileHandle(),
+			Request,
+			MoveTemp(ChunkInfo),
+			ChunkRange.ConsumeValueOrDie(),
+			RequestSize));
+
+		ExclusivePipe.Launch(UE_SOURCE_LOCATION, [this, Request, CasLoc]
+		{
+			FChunkRequest& ChunkRequest = FChunkRequest::GetRef(*Request);
+			bool bOk = false;
+
+			FFileOpenResult FileOpenResult = Cas.OpenRead(CasLoc.BlockId);
+			if (FileOpenResult.IsValid())
+			{
+				TUniquePtr<IFileHandle> FileHandle = FileOpenResult.StealValue();
+				const int64 CasBlockOffset = CasLoc.BlockOffset + ChunkRequest.ChunkRange.GetOffset();
+				if (Request->IsCancelled() == false && FileHandle->Seek(CasBlockOffset))
+				{
+					bOk = FileHandle->Read(ChunkRequest.EncodedChunk.GetData(), ChunkRequest.EncodedChunk.GetSize());
+					if (bOk == false)
+					{
+						const FString Filename = Cas.GetBlockFilename(CasLoc.BlockId);
+						UE_LOG(LogIoStoreOnDemand, Error, TEXT("Failed to read %llu bytes at offset %lld in CAS block '%s'"),
+							ChunkRequest.EncodedChunk.GetSize(),
+							CasBlockOffset,
+							*Filename);
+					}
+				}
+				else
+				{
+					const FString Filename = Cas.GetBlockFilename(CasLoc.BlockId);
+					UE_LOG(LogIoStoreOnDemand, Error, TEXT("Failed to seek to offset %lld in CAS block '%s'"), CasBlockOffset, *Filename);
+				}
+			}
+			else
+			{
+				const FString Filename = Cas.GetBlockFilename(CasLoc.BlockId);
+				UE_LOG(LogIoStoreOnDemand, Error, TEXT("Failed to open CAS block '%s' for reading, reason '%s'"),
+					*Filename, *FileOpenResult.GetError().GetMessage());
+			}
+
+			const bool bWasCancelled = bOk == false;
+			UE::Tasks::Launch(UE_SOURCE_LOCATION, [this, Request, bWasCancelled]
+			{
+				CompleteRequest(Request, bWasCancelled);
+			});
+		}, UE::Tasks::ETaskPriority::BackgroundHigh);
+
+		return true;
+	}
+#endif // UE_ONDEMANDINSTALLCACHE_EXCLUSIVE_WRITE
+
 	FSharedFileOpenResult FileOpenResult = Cas.OpenAsyncRead(CasLoc.BlockId);
 	if (FileOpenResult.HasError())
 	{
@@ -1510,6 +1585,28 @@ FOnDemandInstallCacheStorageUsage FOnDemandInstallCache::GetStorageUsage()
 }
 
 FIoStatus FOnDemandInstallCache::FlushPendingChunks(FPendingChunks& Chunks)
+{
+	if (Chunks.IsEmpty())
+	{
+		return FIoStatus::Ok;
+	}
+
+#if UE_ONDEMANDINSTALLCACHE_EXCLUSIVE_WRITE
+	UE::Tasks::TTask<FIoStatus> Task = ExclusivePipe.Launch(UE_SOURCE_LOCATION, [this, &Chunks]
+	{
+		return FlushPendingChunksImpl(Chunks);
+	}, UE::Tasks::ETaskPriority::BackgroundHigh);
+
+	Task.Wait();
+
+	return Task.GetResult();
+
+#else
+	return FlushPendingChunksImpl(Chunks);
+#endif // UE_ONDEMANDINSTALLCACHE_EXCLUSIVE_WRITE
+}
+
+FIoStatus FOnDemandInstallCache::FlushPendingChunksImpl(FPendingChunks& Chunks)
 {
 	ON_SCOPE_EXIT { Chunks.Reset(); };
 
