Fix crash when Header creation failed

#rnx

diff --git a/Engine/Source/Runtime/Experimental/Iris/Core/Private/Iris/ReplicationSystem/ObjectReplicationBridge.cpp b/Engine/Source/Runtime/Experimental/Iris/Core/Private/Iris/ReplicationSystem/ObjectReplicationBridge.cpp
--- a/Engine/Source/Runtime/Experimental/Iris/Core/Private/Iris/ReplicationSystem/ObjectReplicationBridge.cpp
+++ b/Engine/Source/Runtime/Experimental/Iris/Core/Private/Iris/ReplicationSystem/ObjectReplicationBridge.cpp
@@ -694,6 +694,12 @@ bool UObjectReplicationBridge::WriteNetRefHandleCreationInfo(FReplicationBridgeS
 
 	TUniquePtr<FNetObjectCreationHeader> Header = Factory->CreateHeader(Handle, Protocol->ProtocolIdentifier);
 
+	if (!Header.IsValid())
+	{
+		UE_LOG(LogIrisBridge, Error, TEXT("WriteNetRefHandleCreationInfo: CreateHeader via %s failed for %s"), *GetNameSafe(Factory), *Handle.ToString());
+		return false;
+	}
+
 	UE_LOG(LogIrisBridge, VeryVerbose, TEXT("WriteNetRefHandleCreationInfo (new) for %s : Factory: %s (id:%u) Header: %s"), *Handle.ToString(), *GetNameSafe(Factory), FactoryId, *Header->ToString());
 
 	return Factory->WriteHeader(Handle, Context.SerializationContext, Header.Get());
