Iris - Remove MaxLocalIndex from ObjectScopeHysteresisUpdater in favor of using bit array size instead. Fixes late join crash.

#rnx

diff --git a/Engine/Source/Runtime/Experimental/Iris/Core/Private/Iris/ReplicationSystem/Filtering/ObjectScopeHysteresisUpdater.cpp b/Engine/Source/Runtime/Experimental/Iris/Core/Private/Iris/ReplicationSystem/Filtering/ObjectScopeHysteresisUpdater.cpp
--- a/Engine/Source/Runtime/Experimental/Iris/Core/Private/Iris/ReplicationSystem/Filtering/ObjectScopeHysteresisUpdater.cpp
+++ b/Engine/Source/Runtime/Experimental/Iris/Core/Private/Iris/ReplicationSystem/Filtering/ObjectScopeHysteresisUpdater.cpp
@@ -85,7 +85,7 @@ void FObjectScopeHysteresisUpdater::Update(uint8 FramesSinceLastUpdate, TArray<F
 	TArray<FInternalNetRefIndex, TInlineAllocator<32>> ObjectsToRemoveFromUpdate;
 
 	const WordType* LocalIndicesData = UsedLocalIndices.GetData();
-	for (FLocalIndex ObjectIt = 0, ObjectEndIt = MaxLocalIndex, IndexOffset = 0; ObjectIt < ObjectEndIt; ObjectIt += WordBitCount, ++LocalIndicesData, IndexOffset += 32U)
+	for (FLocalIndex ObjectIt = 0, ObjectEndIt = UsedLocalIndices.GetNumBits(), IndexOffset = 0; ObjectIt < ObjectEndIt; ObjectIt += WordBitCount, ++LocalIndicesData, IndexOffset += 32U)
 	{
 		// Skip ranges with no objects to update
 		WordType LocalIndicesWord = *LocalIndicesData;
@@ -167,9 +167,6 @@ FObjectScopeHysteresisUpdater::FLocalIndex FObjectScopeHysteresisUpdater::GetOrC
 	LocalIndexToNetRefIndex[LocalIndex] = NetRefIndex;
 	NetRefIndexToLocalIndex.Add(NetRefIndex, LocalIndex);
 	ObjectsToUpdate.SetBit(NetRefIndex);
-	MaxLocalIndex = FPlatformMath::Max(MaxLocalIndex, LocalIndex + 1U);
-
-	ensureMsgf(MaxLocalIndex <= UsedLocalIndices.GetNumBits(), TEXT("MaxLocalIndex has grown beyond supported size. MaxLocalIndex: %u LocalIndex: %u Array sizes: %u"), MaxLocalIndex, LocalIndex, UsedLocalIndices.GetNumBits());
 
 	return LocalIndex;
 }
diff --git a/Engine/Source/Runtime/Experimental/Iris/Core/Private/Iris/ReplicationSystem/Filtering/ObjectScopeHysteresisUpdater.h b/Engine/Source/Runtime/Experimental/Iris/Core/Private/Iris/ReplicationSystem/Filtering/ObjectScopeHysteresisUpdater.h
--- a/Engine/Source/Runtime/Experimental/Iris/Core/Private/Iris/ReplicationSystem/Filtering/ObjectScopeHysteresisUpdater.h
+++ b/Engine/Source/Runtime/Experimental/Iris/Core/Private/Iris/ReplicationSystem/Filtering/ObjectScopeHysteresisUpdater.h
@@ -59,8 +59,6 @@ private:
 	FLocalIndex GetOrCreateLocalIndex(FInternalNetRefIndex NetRefIndex);
 	void FreeLocalIndex(FLocalIndex LocalIndex);
 
-	// Internally there's a mapping between NetRefIndex and a LocalIndex to keep updates in as small a range as possible. This is the highest local index + 1 that needs to be updated.
-	FLocalIndex MaxLocalIndex = 0;
 	// Per LocalIndex how many frames left to update before filtering out the object.
 	TArray<uint16> FrameCounters;
 	// Lookup table for local index to InternalNetRefIndex.
