Chaos:: Split PBD joint constraints into a linear and a nonlinear one to reduce memory allocation inside FPBDJointContainerSolver::Reset.

diff --git a/Engine/Source/Programs/HeadlessChaos/Private/HeadlessChaosTestConstraints.cpp b/Engine/Source/Programs/HeadlessChaos/Private/HeadlessChaosTestConstraints.cpp
--- a/Engine/Source/Programs/HeadlessChaos/Private/HeadlessChaosTestConstraints.cpp
+++ b/Engine/Source/Programs/HeadlessChaos/Private/HeadlessChaosTestConstraints.cpp
@@ -38,7 +38,7 @@ namespace ChaosTest {
 			TArray<FVec3> Positions = { FVec3(0) };
 			FPBDPositionConstraints PositionConstraints(MoveTemp(Positions), MoveTemp(Dynamics), 1.f);
 			InitEvolutionSettings(Evolution);
-			Evolution.GetJointConstraints().SetUseSimd(bUseSimd);
+			Evolution.GetJointConstraints().LinearConstraints.SetUseSimd(bUseSimd);
 
 			Evolution.AddConstraintContainer(PositionConstraints);
 			Evolution.AdvanceOneTimeStep(0.1);
@@ -58,7 +58,7 @@ namespace ChaosTest {
 			TArray<FVec3> Positions = { FVec3(1) };
 			FPBDPositionConstraints PositionConstraints(MoveTemp(Positions), MoveTemp(Dynamics), 0.5f);
 			Evolution.AddConstraintContainer(PositionConstraints);
-			Evolution.GetJointConstraints().SetUseSimd(bUseSimd);
+			Evolution.GetJointConstraints().LinearConstraints.SetUseSimd(bUseSimd);
 
 			// The effect of stiffness parameter (which is set to 0.5 above) is iteration depeendent
 			Evolution.SetNumPositionIterations(1);
diff --git a/Engine/Source/Programs/HeadlessChaos/Private/HeadlessChaosTestJoint.h b/Engine/Source/Programs/HeadlessChaos/Private/HeadlessChaosTestJoint.h
--- a/Engine/Source/Programs/HeadlessChaos/Private/HeadlessChaosTestJoint.h
+++ b/Engine/Source/Programs/HeadlessChaos/Private/HeadlessChaosTestJoint.h
@@ -33,7 +33,21 @@ namespace ChaosTest
 
 		FPBDJointConstraintHandle* AddJoint(const TVec2<TGeometryParticleHandle<FReal, 3>*>& InConstrainedParticleIndices, const int32 JointIndex)
 		{
-			FPBDJointConstraintHandle* Joint = Evolution.GetJointConstraints().AddConstraint(InConstrainedParticleIndices, FRigidTransform3(JointPositions[JointIndex], FRotation3::FromIdentity()));
+			
+			FPBDJointConstraintHandle* Joint = nullptr;
+				
+			//Unless someone sets the joint to use nonlinear solver beforehand, otherwise it will default to linear solver. 
+			if (JointIndex < JointSettings.Num())
+			{
+				if (!JointSettings[JointIndex].bUseLinearSolver)
+				{
+					Joint = Evolution.GetJointConstraints().NonLinearConstraints.AddConstraint(InConstrainedParticleIndices, FRigidTransform3(JointPositions[JointIndex], FRotation3::FromIdentity()));
+				}
+			}
+			if (!Joint)
+			{
+				Joint = Evolution.GetJointConstraints().LinearConstraints.AddConstraint(InConstrainedParticleIndices, FRigidTransform3(JointPositions[JointIndex], FRotation3::FromIdentity()));
+			}
 
 			// @todo(chaos): this indicates we need to change the AddConstraint API (since ConnectorTransforms were added to Settings). 
 			// Calling AddConstraint followed by SetSettings will overwrite the ConnectorTransforms
@@ -49,7 +63,15 @@ namespace ChaosTest
 
 		FPBDJointConstraintHandle* GetJoint(const int32 JointIndex)
 		{
-			return Evolution.GetJointConstraints().GetConstConstraintHandles()[JointIndex];
+			if (JointIndex < JointSettings.Num())
+			{
+				if (!JointSettings[JointIndex].bUseLinearSolver)
+				{
+					return Evolution.GetJointConstraints().NonLinearConstraints.GetConstConstraintHandles()[JointIndex];
+				}
+			}
+
+			return Evolution.GetJointConstraints().LinearConstraints.GetConstConstraintHandles()[JointIndex];
 		}
 
 		void Create()
diff --git a/Engine/Source/Programs/HeadlessChaos/Private/HeadlessChaosTestJointBreak.cpp b/Engine/Source/Programs/HeadlessChaos/Private/HeadlessChaosTestJointBreak.cpp
--- a/Engine/Source/Programs/HeadlessChaos/Private/HeadlessChaosTestJointBreak.cpp
+++ b/Engine/Source/Programs/HeadlessChaos/Private/HeadlessChaosTestJointBreak.cpp
@@ -21,10 +21,10 @@ namespace ChaosTest {
 		FJointChainTest<TEvolution> Test(NumIterations, Gravity);
 		Test.InitChain(2, FVec3(0,0,-1));
 		Test.Create();
-		Test.Evolution.GetJointConstraints().SetUseSimd(bUseSimd);
+		Test.Evolution.GetJointConstraints().LinearConstraints.SetUseSimd(bUseSimd);
 
 		bool bBrokenCallbackCalled = false;
-		Test.Evolution.GetJointConstraints().SetBreakCallback([&bBrokenCallbackCalled](FPBDJointConstraintHandle* Constraint)
+		Test.Evolution.GetJointConstraints().LinearConstraints.SetBreakCallback([&bBrokenCallbackCalled](FPBDJointConstraintHandle* Constraint)
 			{
 				bBrokenCallbackCalled = true;
 			});
@@ -39,14 +39,14 @@ namespace ChaosTest {
 
 		// Nothing should have broken
 		EXPECT_FALSE(bBrokenCallbackCalled);
-		EXPECT_TRUE(Test.Evolution.GetJointConstraints().IsConstraintEnabled(0));
+		EXPECT_TRUE(Test.Evolution.GetJointConstraints().LinearConstraints.IsConstraintEnabled(0));
 
 		// Manually break the constraints
-		Test.Evolution.GetJointConstraints().BreakConstraint(0);
+		Test.Evolution.GetJointConstraints().LinearConstraints.BreakConstraint(0);
 
 		// Check that it worked
 		EXPECT_TRUE(bBrokenCallbackCalled);
-		EXPECT_FALSE(Test.Evolution.GetJointConstraints().IsConstraintEnabled(0));
+		EXPECT_FALSE(Test.Evolution.GetJointConstraints().LinearConstraints.IsConstraintEnabled(0));
 
 		// Run the sim - body should fall
 		for (int32 i = 0; i < NumIts; ++i)
@@ -76,7 +76,7 @@ namespace ChaosTest {
 
 		FJointChainTest<TEvolution> Test(NumIterations, Gravity);
 		Test.InitChain(2, FVec3(0, 0, -1));
-		Test.Evolution.GetJointConstraints().SetUseSimd(bUseSimd);
+		Test.Evolution.GetJointConstraints().LinearConstraints.SetUseSimd(bUseSimd);
 		// Joint should break only if Threshold < MG
 		// So not in this test
 		Test.JointSettings[0].LinearBreakForce = 1.1f * Test.ParticleMasses[1] * Gravity;
@@ -84,7 +84,7 @@ namespace ChaosTest {
 		Test.Create();
 
 		bool bBrokenCallbackCalled = false;
-		Test.Evolution.GetJointConstraints().SetBreakCallback([&bBrokenCallbackCalled](FPBDJointConstraintHandle* Constraint)
+		Test.Evolution.GetJointConstraints().LinearConstraints.SetBreakCallback([&bBrokenCallbackCalled](FPBDJointConstraintHandle* Constraint)
 			{
 				bBrokenCallbackCalled = true;
 			});
@@ -98,7 +98,7 @@ namespace ChaosTest {
 
 		// Nothing should have broken
 		EXPECT_FALSE(bBrokenCallbackCalled);
-		EXPECT_TRUE(Test.Evolution.GetJointConstraints().IsConstraintEnabled(0));
+		EXPECT_TRUE(Test.Evolution.GetJointConstraints().LinearConstraints.IsConstraintEnabled(0));
 	}
 
 	GTEST_TEST(AllEvolutions, JointBreakTests_TestUnderLinearThreshold)
@@ -119,7 +119,7 @@ namespace ChaosTest {
 
 		FJointChainTest<TEvolution> Test(NumIterations, Gravity);
 		Test.InitChain(3, FVec3(0, 0, -1));
-		Test.Evolution.GetJointConstraints().SetUseSimd(bUseSimd);
+		Test.Evolution.GetJointConstraints().LinearConstraints.SetUseSimd(bUseSimd);
 
 		// Joint should break only if Threshold < MG
 		// So not in this test
@@ -130,7 +130,7 @@ namespace ChaosTest {
 		Test.Create();
 
 		bool bBrokenCallbackCalled = false;
-		Test.Evolution.GetJointConstraints().SetBreakCallback([&bBrokenCallbackCalled](FPBDJointConstraintHandle* Constraint)
+		Test.Evolution.GetJointConstraints().LinearConstraints.SetBreakCallback([&bBrokenCallbackCalled](FPBDJointConstraintHandle* Constraint)
 			{
 				bBrokenCallbackCalled = true;
 			});
@@ -144,8 +144,8 @@ namespace ChaosTest {
 
 		// Nothing should have broken
 		EXPECT_FALSE(bBrokenCallbackCalled);
-		EXPECT_TRUE(Test.Evolution.GetJointConstraints().IsConstraintEnabled(0));
-		EXPECT_TRUE(Test.Evolution.GetJointConstraints().IsConstraintEnabled(1));
+		EXPECT_TRUE(Test.Evolution.GetJointConstraints().LinearConstraints.IsConstraintEnabled(0));
+		EXPECT_TRUE(Test.Evolution.GetJointConstraints().LinearConstraints.IsConstraintEnabled(1));
 	}
 
 	GTEST_TEST(AllEvolutions, JointBreakTests_TestUnderLinearThreshold2)
@@ -166,7 +166,7 @@ namespace ChaosTest {
 
 		FJointChainTest<TEvolution> Test(NumIterations, Gravity);
 		Test.InitChain(2, FVec3(0,0,-1));
-		Test.Evolution.GetJointConstraints().SetUseSimd(bUseSimd);
+		Test.Evolution.GetJointConstraints().LinearConstraints.SetUseSimd(bUseSimd);
 		// Joint should break only if Threshold < MG
 		// So yes in this test
 		Test.JointSettings[0].LinearBreakForce = 0.9f * Test.ParticleMasses[1] * Gravity;
@@ -174,7 +174,7 @@ namespace ChaosTest {
 		Test.Create();
 
 		bool bBrokenCallbackCalled = false;
-		Test.Evolution.GetJointConstraints().SetBreakCallback([&bBrokenCallbackCalled](FPBDJointConstraintHandle* Constraint)
+		Test.Evolution.GetJointConstraints().LinearConstraints.SetBreakCallback([&bBrokenCallbackCalled](FPBDJointConstraintHandle* Constraint)
 			{
 				bBrokenCallbackCalled = true;
 			});
@@ -188,7 +188,7 @@ namespace ChaosTest {
 
 		// Constraint should have broken
 		EXPECT_TRUE(bBrokenCallbackCalled);
-		EXPECT_FALSE(Test.Evolution.GetJointConstraints().IsConstraintEnabled(0));
+		EXPECT_FALSE(Test.Evolution.GetJointConstraints().LinearConstraints.IsConstraintEnabled(0));
 	}
 
 	GTEST_TEST(AllEvolutions, JointBreakTests_TestOverLinearThreshold)
@@ -210,7 +210,7 @@ namespace ChaosTest {
 
 		FJointChainTest<TEvolution> Test(NumIterations, Gravity);
 		Test.InitChain(3, FVec3(0, 0, -1));
-		Test.Evolution.GetJointConstraints().SetUseSimd(bUseSimd);
+		Test.Evolution.GetJointConstraints().LinearConstraints.SetUseSimd(bUseSimd);
 
 		// Joint should break only if Threshold < MG
 		// So no in this test
@@ -220,7 +220,7 @@ namespace ChaosTest {
 		Test.Create();
 
 		bool bBrokenCallbackCalled = false;
-		Test.Evolution.GetJointConstraints().SetBreakCallback([&bBrokenCallbackCalled](FPBDJointConstraintHandle* Constraint)
+		Test.Evolution.GetJointConstraints().LinearConstraints.SetBreakCallback([&bBrokenCallbackCalled](FPBDJointConstraintHandle* Constraint)
 			{
 				bBrokenCallbackCalled = true;
 			});
@@ -234,8 +234,8 @@ namespace ChaosTest {
 
 		// Constraint should not have broken
 		EXPECT_FALSE(bBrokenCallbackCalled);
-		EXPECT_TRUE(Test.Evolution.GetJointConstraints().IsConstraintEnabled(0));
-		EXPECT_TRUE(Test.Evolution.GetJointConstraints().IsConstraintEnabled(1));
+		EXPECT_TRUE(Test.Evolution.GetJointConstraints().LinearConstraints.IsConstraintEnabled(0));
+		EXPECT_TRUE(Test.Evolution.GetJointConstraints().LinearConstraints.IsConstraintEnabled(1));
 	}
 
 	GTEST_TEST(AllEvolutions, JointBreakTests_TestUnderLinearThreshold3)
@@ -258,7 +258,7 @@ namespace ChaosTest {
 
 		FJointChainTest<TEvolution> Test(NumIterations, Gravity);
 		Test.InitChain(2, FVec3(0, 0, -1));
-		Test.Evolution.GetJointConstraints().SetUseSimd(bUseSimd);
+		Test.Evolution.GetJointConstraints().LinearConstraints.SetUseSimd(bUseSimd);
 		// Joint should break only if Threshold < MGL
 		// So not in this test
 		Test.JointSettings[0].AngularBreakTorque = 1.1f * Torque.X;
@@ -267,7 +267,7 @@ namespace ChaosTest {
 		Test.Create();
 
 		bool bBrokenCallbackCalled = false;
-		Test.Evolution.GetJointConstraints().SetBreakCallback([&bBrokenCallbackCalled](FPBDJointConstraintHandle* Constraint)
+		Test.Evolution.GetJointConstraints().LinearConstraints.SetBreakCallback([&bBrokenCallbackCalled](FPBDJointConstraintHandle* Constraint)
 			{
 				bBrokenCallbackCalled = true;
 			});
@@ -283,7 +283,7 @@ namespace ChaosTest {
 
 		// Nothing should have broken
 		EXPECT_FALSE(bBrokenCallbackCalled);
-		EXPECT_TRUE(Test.Evolution.GetJointConstraints().IsConstraintEnabled(0));
+		EXPECT_TRUE(Test.Evolution.GetJointConstraints().LinearConstraints.IsConstraintEnabled(0));
 	}
 
 	GTEST_TEST(AllEvolutions, JointBreakTests_TestUnderAngularThreshold)
@@ -305,7 +305,7 @@ namespace ChaosTest {
 
 		FJointChainTest<TEvolution> Test(NumSolverIterations, Gravity);
 		Test.InitChain(2, FVec3(0, 0, -1));
-		Test.Evolution.GetJointConstraints().SetUseSimd(bUseSimd);
+		Test.Evolution.GetJointConstraints().LinearConstraints.SetUseSimd(bUseSimd);
 		// Joint should break only if Threshold < MGL
 		// So not in this test
 		Test.JointSettings[0].AngularBreakTorque = 0.9f * Torque.X;
@@ -314,7 +314,7 @@ namespace ChaosTest {
 		Test.Create();
 
 		bool bBrokenCallbackCalled = false;
-		Test.Evolution.GetJointConstraints().SetBreakCallback([&bBrokenCallbackCalled](FPBDJointConstraintHandle* Constraint)
+		Test.Evolution.GetJointConstraints().LinearConstraints.SetBreakCallback([&bBrokenCallbackCalled](FPBDJointConstraintHandle* Constraint)
 			{
 				bBrokenCallbackCalled = true;
 			});
@@ -330,7 +330,7 @@ namespace ChaosTest {
 
 		// Nothing should have broken
 		EXPECT_TRUE(bBrokenCallbackCalled);
-		EXPECT_FALSE(Test.Evolution.GetJointConstraints().IsConstraintEnabled(0));
+		EXPECT_FALSE(Test.Evolution.GetJointConstraints().LinearConstraints.IsConstraintEnabled(0));
 	}
 
 	GTEST_TEST(AllEvolutions, JointBreakTests_TestOverAngularThreshold)
diff --git a/Engine/Source/Programs/HeadlessChaos/Private/HeadlessChaosTestJointForces.cpp b/Engine/Source/Programs/HeadlessChaos/Private/HeadlessChaosTestJointForces.cpp
--- a/Engine/Source/Programs/HeadlessChaos/Private/HeadlessChaosTestJointForces.cpp
+++ b/Engine/Source/Programs/HeadlessChaos/Private/HeadlessChaosTestJointForces.cpp
@@ -25,7 +25,7 @@ namespace ChaosTest {
 
 			FJointChainTest<TEvolution> Test(NumSolverIterations, Gravity);
 			Test.InitChain(NumBodies, FVec3(0, 0, -1));
-			Test.Evolution.GetJointConstraints().SetUseSimd(bUseSimd);
+			Test.Evolution.GetJointConstraints().LinearConstraints.SetUseSimd(bUseSimd);
 			Test.Create();
 
 			// Run the sim
@@ -35,7 +35,7 @@ namespace ChaosTest {
 				Test.Evolution.EndFrame(Dt);
 			}
 
-			for (int32 JointIndex = 0; JointIndex < Test.Evolution.GetJointConstraints().NumConstraints(); ++JointIndex)
+			for (int32 JointIndex = 0; JointIndex < Test.Evolution.GetJointConstraints().LinearConstraints.NumConstraints(); ++JointIndex)
 			{
 				FReal ChildMass = 0.0f;
 				for (int32 ChildBodyIndex = JointIndex + 1; ChildBodyIndex < Test.ParticleMasses.Num(); ++ChildBodyIndex)
@@ -44,7 +44,7 @@ namespace ChaosTest {
 				}
 
 				FVec3 ExpectedLinearImpulse = FVec3(0.0f, 0.0f, -ChildMass * Gravity * Dt);
-				FVec3 LinearImpulse = Test.Evolution.GetJointConstraints().GetConstraintLinearImpulse(JointIndex);
+				FVec3 LinearImpulse = Test.Evolution.GetJointConstraints().LinearConstraints.GetConstraintLinearImpulse(JointIndex);
 				EXPECT_NEAR(LinearImpulse.X, ExpectedLinearImpulse.X, ExpectedLinearImpulse.Size() / 100.0f);
 				EXPECT_NEAR(LinearImpulse.Y, ExpectedLinearImpulse.Y, ExpectedLinearImpulse.Size() / 100.0f);
 				EXPECT_NEAR(LinearImpulse.Z, ExpectedLinearImpulse.Z, ExpectedLinearImpulse.Size() / 100.0f);
@@ -73,7 +73,7 @@ namespace ChaosTest {
 
 			FJointChainTest<TEvolution> Test(NumSolverIterations, Gravity);
 			Test.InitChain(NumBodies, FVec3(1, 0, 0));
-			Test.Evolution.GetJointConstraints().SetUseSimd(bUseSimd);
+			Test.Evolution.GetJointConstraints().LinearConstraints.SetUseSimd(bUseSimd);
 
 			// Set all joints to fixed angular
 			for (int32 JointIndex = 0; JointIndex < Test.JointSettings.Num(); ++JointIndex)
@@ -105,13 +105,13 @@ namespace ChaosTest {
 				}
 
 				FVec3 ExpectedLinearImpulse = FVec3(0.0f, 0.0f, -ChildMass * Gravity * Dt);
-				FVec3 LinearImpulse = Test.Evolution.GetJointConstraints().GetConstraintLinearImpulse(JointIndex);
+				FVec3 LinearImpulse = Test.Evolution.GetJointConstraints().LinearConstraints.GetConstraintLinearImpulse(JointIndex);
 				EXPECT_NEAR(LinearImpulse.X, ExpectedLinearImpulse.X, ExpectedLinearImpulse.Size() / 100.0f);
 				EXPECT_NEAR(LinearImpulse.Y, ExpectedLinearImpulse.Y, ExpectedLinearImpulse.Size() / 100.0f);
 				EXPECT_NEAR(LinearImpulse.Z, ExpectedLinearImpulse.Z, ExpectedLinearImpulse.Size() / 100.0f);
 
 				FVec3 ExpectedAngularImpulse = FVec3(0.0f, ChildMoment * Gravity * Dt, 0.0f);
-				FVec3 AngularImpulse = Test.Evolution.GetJointConstraints().GetConstraintAngularImpulse(JointIndex);
+				FVec3 AngularImpulse = Test.Evolution.GetJointConstraints().LinearConstraints.GetConstraintAngularImpulse(JointIndex);
 				EXPECT_NEAR(AngularImpulse.X, ExpectedAngularImpulse.X, ExpectedAngularImpulse.Size() / 100.0f);
 				EXPECT_NEAR(AngularImpulse.Y, ExpectedAngularImpulse.Y, ExpectedAngularImpulse.Size() / 100.0f);
 				EXPECT_NEAR(AngularImpulse.Z, ExpectedAngularImpulse.Z, ExpectedAngularImpulse.Size() / 100.0f);
@@ -142,7 +142,7 @@ namespace ChaosTest {
 
 			FJointChainTest<TEvolution> Test(NumSolverIterations, Gravity);
 			Test.InitChain(NumBodies, FVec3(1, 0, 0));
-			Test.Evolution.GetJointConstraints().SetUseSimd(bUseSimd);
+			Test.Evolution.GetJointConstraints().LinearConstraints.SetUseSimd(bUseSimd);
 
 			// Set all joints to fixed angular and disable projection
 			for (int32 JointIndex = 0; JointIndex < Test.JointSettings.Num(); ++JointIndex)
@@ -171,7 +171,7 @@ namespace ChaosTest {
 				}
 
 				FVec3 ExpectedAngularImpulse = Torque * Dt;
-				FVec3 AngularImpulse = Test.Evolution.GetJointConstraints().GetConstraintAngularImpulse(JointIndex);
+				FVec3 AngularImpulse = Test.Evolution.GetJointConstraints().LinearConstraints.GetConstraintAngularImpulse(JointIndex);
 				EXPECT_NEAR(AngularImpulse.X, ExpectedAngularImpulse.X, ExpectedAngularImpulse.Size() / 100.0f);
 				EXPECT_NEAR(AngularImpulse.Y, ExpectedAngularImpulse.Y, ExpectedAngularImpulse.Size() / 100.0f);
 				EXPECT_NEAR(AngularImpulse.Z, ExpectedAngularImpulse.Z, ExpectedAngularImpulse.Size() / 100.0f);
@@ -203,7 +203,7 @@ namespace ChaosTest {
 
 			FJointChainTest<FPBDRigidsEvolutionGBF> Test(NumSolverIterations, Gravity);
 			Test.InitChain(NumBodies, FVec3(0, 0, -1));
-			Test.Evolution.GetJointConstraints().SetUseSimd(bUseSimd);
+			Test.Evolution.GetJointConstraints().LinearConstraints.SetUseSimd(bUseSimd);
 
 			// Disable all limits
 			Test.JointSettings[0].LinearMotionTypes = { EJointMotionType::Free, EJointMotionType::Free, EJointMotionType::Free };
@@ -241,7 +241,7 @@ namespace ChaosTest {
 			}
 
 			// Check the joint forces agree
-			const FReal ForceZ = -Test.Evolution.GetJointConstraints().GetConstraintLinearImpulse(0).Z / Dt;
+			const FReal ForceZ = -Test.Evolution.GetJointConstraints().LinearConstraints.GetConstraintLinearImpulse(0).Z / Dt;
 			EXPECT_NEAR(ForceZ, ExpectedForceZ, 0.01);
 		}
 	}
@@ -265,9 +265,9 @@ namespace ChaosTest {
 
 			FJointChainTest<FPBDRigidsEvolutionGBF> Test(NumSolverIterations, Gravity);
 			Test.InitChain(NumBodies, FVec3(0, 0, -1));
-			Test.Evolution.GetJointConstraints().SetUseSimd(bUseSimd);
+			Test.Evolution.GetJointConstraints().LinearConstraints.SetUseSimd(bUseSimd);
 
-			if (!Test.Evolution.GetJointConstraints().GetSettings().bUsePositionBasedDrives)
+			if (!Test.Evolution.GetJointConstraints().LinearConstraints.GetSettings().bUsePositionBasedDrives)
 			{
 				Test.Evolution.SetNumVelocityIterations(NumSolverIterations);
 			}
@@ -308,7 +308,7 @@ namespace ChaosTest {
 			}
 
 			// Check the joint forces agree
-			const FReal ForceZ = -Test.Evolution.GetJointConstraints().GetConstraintLinearImpulse(0).Z / Dt;
+			const FReal ForceZ = -Test.Evolution.GetJointConstraints().LinearConstraints.GetConstraintLinearImpulse(0).Z / Dt;
 			EXPECT_NEAR(ForceZ, ExpectedForceZ, 1);
 		}
 	}
@@ -329,9 +329,9 @@ namespace ChaosTest {
 
 			FJointChainTest<FPBDRigidsEvolutionGBF> Test(NumSolverIterations, Gravity);
 			Test.InitChain(NumBodies, FVec3(0, 0, -1));
-			Test.Evolution.GetJointConstraints().SetUseSimd(bUseSimd);
+			Test.Evolution.GetJointConstraints().LinearConstraints.SetUseSimd(bUseSimd);
 
-			if (!Test.Evolution.GetJointConstraints().GetSettings().bUsePositionBasedDrives)
+			if (!Test.Evolution.GetJointConstraints().LinearConstraints.GetSettings().bUsePositionBasedDrives)
 			{
 				Test.Evolution.SetNumVelocityIterations(NumSolverIterations);
 			}
@@ -371,7 +371,7 @@ namespace ChaosTest {
 			}
 
 			// Check the joint forces agree
-			const FReal TorqueZ = -Test.Evolution.GetJointConstraints().GetConstraintAngularImpulse(0).Z / Dt;
+			const FReal TorqueZ = -Test.Evolution.GetJointConstraints().LinearConstraints.GetConstraintAngularImpulse(0).Z / Dt;
 			EXPECT_NEAR(TorqueZ, ExpectedTorqueZ, 0.1);
 		}
 	}
@@ -393,9 +393,9 @@ namespace ChaosTest {
 
 			FJointChainTest<FPBDRigidsEvolutionGBF> Test(NumSolverIterations, Gravity);
 			Test.InitChain(NumBodies, FVec3(0, 0, -1));
-			Test.Evolution.GetJointConstraints().SetUseSimd(bUseSimd);
+			Test.Evolution.GetJointConstraints().LinearConstraints.SetUseSimd(bUseSimd);
 
-			if (!Test.Evolution.GetJointConstraints().GetSettings().bUsePositionBasedDrives)
+			if (!Test.Evolution.GetJointConstraints().LinearConstraints.GetSettings().bUsePositionBasedDrives)
 			{
 				Test.Evolution.SetNumVelocityIterations(NumSolverIterations);
 			}
@@ -449,9 +449,9 @@ namespace ChaosTest {
 
 			FJointChainTest<FPBDRigidsEvolutionGBF> Test(NumSolverIterations, Gravity);
 			Test.InitChain(NumBodies, FVec3(0, 0, -1));
-			Test.Evolution.GetJointConstraints().SetUseSimd(bUseSimd);
+			Test.Evolution.GetJointConstraints().LinearConstraints.SetUseSimd(bUseSimd);
 
-			if (!Test.Evolution.GetJointConstraints().GetSettings().bUsePositionBasedDrives)
+			if (!Test.Evolution.GetJointConstraints().LinearConstraints.GetSettings().bUsePositionBasedDrives)
 			{
 				Test.Evolution.SetNumVelocityIterations(NumSolverIterations);
 			}
@@ -480,7 +480,7 @@ namespace ChaosTest {
 				Test.Evolution.AdvanceOneTimeStep(Dt);
 				Test.Evolution.EndFrame(Dt);
 
-				const FReal ForceZ = -Test.Evolution.GetJointConstraints().GetConstraintLinearImpulse(0).Z / Dt;
+				const FReal ForceZ = -Test.Evolution.GetJointConstraints().LinearConstraints.GetConstraintLinearImpulse(0).Z / Dt;
 				EXPECT_LT(ForceZ, Test.JointSettings[0].LinearDriveMaxForce.Z + 0.1);
 			}
 		}
@@ -505,9 +505,9 @@ namespace ChaosTest {
 
 			FJointChainTest<FPBDRigidsEvolutionGBF> Test(NumSolverIterations, Gravity);
 			Test.InitChain(NumBodies, FVec3(0, 0, -1));
-			Test.Evolution.GetJointConstraints().SetUseSimd(bUseSimd);
+			Test.Evolution.GetJointConstraints().LinearConstraints.SetUseSimd(bUseSimd);
 
-			if (!Test.Evolution.GetJointConstraints().GetSettings().bUsePositionBasedDrives)
+			if (!Test.Evolution.GetJointConstraints().LinearConstraints.GetSettings().bUsePositionBasedDrives)
 			{
 				Test.Evolution.SetNumVelocityIterations(NumSolverIterations);
 			}
@@ -564,9 +564,9 @@ namespace ChaosTest {
 
 			FJointChainTest<FPBDRigidsEvolutionGBF> Test(NumSolverIterations, Gravity);
 			Test.InitChain(NumBodies, FVec3(0, 0, -1));
-			Test.Evolution.GetJointConstraints().SetUseSimd(bUseSimd);
+			Test.Evolution.GetJointConstraints().LinearConstraints.SetUseSimd(bUseSimd);
 
-			if (!Test.Evolution.GetJointConstraints().GetSettings().bUsePositionBasedDrives)
+			if (!Test.Evolution.GetJointConstraints().LinearConstraints.GetSettings().bUsePositionBasedDrives)
 			{
 				Test.Evolution.SetNumVelocityIterations(NumSolverIterations);
 			}
@@ -596,7 +596,7 @@ namespace ChaosTest {
 				Test.Evolution.AdvanceOneTimeStep(Dt);
 				Test.Evolution.EndFrame(Dt);
 
-				const FReal ForceZ = -Test.Evolution.GetJointConstraints().GetConstraintLinearImpulse(0).Z / Dt;
+				const FReal ForceZ = -Test.Evolution.GetJointConstraints().LinearConstraints.GetConstraintLinearImpulse(0).Z / Dt;
 				EXPECT_LT(ForceZ, Test.JointSettings[0].LinearDriveMaxForce.Z * Test.ParticleMasses[1] + 0.1);
 			}
 		}
diff --git a/Engine/Source/Programs/HeadlessChaos/Private/HeadlessChaosTestJointPlasticity.cpp b/Engine/Source/Programs/HeadlessChaos/Private/HeadlessChaosTestJointPlasticity.cpp
--- a/Engine/Source/Programs/HeadlessChaos/Private/HeadlessChaosTestJointPlasticity.cpp
+++ b/Engine/Source/Programs/HeadlessChaos/Private/HeadlessChaosTestJointPlasticity.cpp
@@ -38,7 +38,7 @@ namespace ChaosTest {
 		Test.Create();
 		Test.AddParticleBox(FVec3(0, 30, 50), FRotation3::Identity, FVec3(10.f), 100.f);
 
-		FReal Angle = Test.Evolution.GetJointConstraints().GetConstraintSettings(0).AngularDrivePositionTarget.GetAngle() * (180. / PI);
+		FReal Angle = Test.Evolution.GetJointConstraints().LinearConstraints.GetConstraintSettings(0).AngularDrivePositionTarget.GetAngle() * (180. / PI);
 		EXPECT_TRUE(FMath::IsNearlyEqual(Angle, (FReal)0.));
 
 		// Run the sim
@@ -52,7 +52,7 @@ namespace ChaosTest {
 			//FVec3 Pos = Test.SOAs.GetDynamicParticles().X(0);
 			//std::cout << "["<< Angle <<"]" << Pos.X << "," << Pos.Y << "," << Pos.Z << std::endl;
 		}
-		Angle = Test.Evolution.GetJointConstraints().GetConstraintSettings(0).AngularDrivePositionTarget.GetAngle() * (180. / PI);
+		Angle = Test.Evolution.GetJointConstraints().LinearConstraints.GetConstraintSettings(0).AngularDrivePositionTarget.GetAngle() * (180. / PI);
 
 		// Nothing should have been reset
 		EXPECT_TRUE(FMath::IsNearlyEqual(Angle, (FReal)0.));
@@ -94,7 +94,7 @@ namespace ChaosTest {
 		Test.Create();
 		Test.AddParticleBox(FVec3(0, 30, 50), FRotation3::Identity, FVec3(10.f), 100.f);
 
-		FReal Angle = Test.Evolution.GetJointConstraints().GetConstraintSettings(0).AngularDrivePositionTarget.GetAngle() * (180. / PI);
+		FReal Angle = Test.Evolution.GetJointConstraints().LinearConstraints.GetConstraintSettings(0).AngularDrivePositionTarget.GetAngle() * (180. / PI);
 		EXPECT_NEAR(Angle, (FReal)0., (FReal)KINDA_SMALL_NUMBER);
 
 		// Run the sim
@@ -112,7 +112,7 @@ namespace ChaosTest {
 		Test.Evolution.AdvanceOneTimeStep(Dt);
 		Test.Evolution.EndFrame(Dt);
 
-		Angle = Test.Evolution.GetJointConstraints().GetConstraintSettings(0).AngularDrivePositionTarget.GetAngle() * (180. / PI);
+		Angle = Test.Evolution.GetJointConstraints().LinearConstraints.GetConstraintSettings(0).AngularDrivePositionTarget.GetAngle() * (180. / PI);
 
 		// The angle should have reset. 
 		EXPECT_GE(Angle, PlasticityAngle);
diff --git a/Engine/Source/Programs/HeadlessChaos/Private/HeadlessChaosTestJointRestitution.cpp b/Engine/Source/Programs/HeadlessChaos/Private/HeadlessChaosTestJointRestitution.cpp
--- a/Engine/Source/Programs/HeadlessChaos/Private/HeadlessChaosTestJointRestitution.cpp
+++ b/Engine/Source/Programs/HeadlessChaos/Private/HeadlessChaosTestJointRestitution.cpp
@@ -20,7 +20,7 @@ namespace ChaosTest {
 		FJointChainTest<TEvolution> Test(NumSolverIterations, Gravity);
 
 		FPBDJointSolverSettings JointSolverSettings;
-		Test.Evolution.GetJointConstraints().SetSettings(JointSolverSettings);
+		Test.Evolution.GetJointConstraints().LinearConstraints.SetSettings(JointSolverSettings);
 
 		Test.InitChain(2, FVec3(1,0,0));
 
diff --git a/Engine/Source/Programs/HeadlessChaos/Private/HeadlessChaosTestJoints.cpp b/Engine/Source/Programs/HeadlessChaos/Private/HeadlessChaosTestJoints.cpp
--- a/Engine/Source/Programs/HeadlessChaos/Private/HeadlessChaosTestJoints.cpp
+++ b/Engine/Source/Programs/HeadlessChaos/Private/HeadlessChaosTestJoints.cpp
@@ -46,9 +46,21 @@ namespace ChaosTest {
 
 		FPBDJointConstraintHandle* AddJoint(const TVec2<FGeometryParticleHandle*>& InConstrainedParticles, const FVec3& InLocation)
 		{
-			FPBDJointConstraintHandle* Constraint = Evolution.GetJointConstraints().AddConstraint(InConstrainedParticles, FRigidTransform3(InLocation, FRotation3::FromIdentity()));
+			FPBDJointSettings SettingsTemp;
+			SetTestJointSetting(SettingsTemp);
 
-			FPBDJointSettings Settings = Constraint->GetSettings();
+			FPBDJointConstraintHandle* Constraint = nullptr;
+
+			if (SettingsTemp.bUseLinearSolver)
+			{
+				Constraint = Evolution.GetJointConstraints().LinearConstraints.AddConstraint(InConstrainedParticles, FRigidTransform3(InLocation, FRotation3::FromIdentity()));
+			}
+			else
+			{
+				Constraint = Evolution.GetJointConstraints().NonLinearConstraints.AddConstraint(InConstrainedParticles, FRigidTransform3(InLocation, FRotation3::FromIdentity()));
+			}
+
+			FPBDJointSettings Settings = Constraint->GetSettings();			
 			Settings.LinearProjection = 1.0;
 			SetTestJointSetting(Settings);
 			Constraint->SetSettings(Settings);
diff --git a/Engine/Source/Runtime/Engine/Private/PhysicsEngine/ImmediatePhysics/ImmediatePhysicsChaos/ImmediatePhysicsSimulation_Chaos.cpp b/Engine/Source/Runtime/Engine/Private/PhysicsEngine/ImmediatePhysics/ImmediatePhysicsChaos/ImmediatePhysicsSimulation_Chaos.cpp
--- a/Engine/Source/Runtime/Engine/Private/PhysicsEngine/ImmediatePhysics/ImmediatePhysicsChaos/ImmediatePhysicsSimulation_Chaos.cpp
+++ b/Engine/Source/Runtime/Engine/Private/PhysicsEngine/ImmediatePhysics/ImmediatePhysicsChaos/ImmediatePhysicsSimulation_Chaos.cpp
@@ -9,6 +9,7 @@
 
 #include "PhysicalMaterials/PhysicalMaterial.h"
 #include "PhysicsEngine/BodyInstance.h"
+#include "PhysicsEngine/ConstraintInstance.h"
 
 #include "Chaos/Collision/BasicCollisionDetector.h"
 #include "Chaos/DebugDrawQueue.h"
@@ -303,7 +304,8 @@ namespace ImmediatePhysics_Chaos
 			Particles.GetParticleHandles().AddArray(&ParticlePrevRs);
 
 			Evolution.AddConstraintContainer(Collisions, ChaosImmediate_Collision_Priority);
-			Evolution.AddConstraintContainer(Joints, 0);
+			Evolution.AddConstraintContainer(Joints.LinearConstraints, 0);
+			Evolution.AddConstraintContainer(Joints.NonLinearConstraints, 0);
 		}
 
 		~FImplementation()
@@ -321,7 +323,7 @@ namespace ImmediatePhysics_Chaos
 
 		FSimpleParticleUniqueIndices UniqueIndices;
 		Chaos::FPBDRigidsSOAs Particles;
-		Chaos::FPBDJointConstraints Joints;
+		Chaos::FPBDJointCombinedConstraints Joints;
 		Chaos::FPBDCollisionConstraints Collisions;
 		Chaos::FBasicBroadPhase BroadPhase;
 		Chaos::FBasicCollisionDetector CollisionDetector;
@@ -539,16 +541,34 @@ namespace ImmediatePhysics_Chaos
 
 	FJointHandle* FSimulation::CreateJoint(FConstraintInstance* ConstraintInstance, FActorHandle* Body1, FActorHandle* Body2)
 	{
-		FJointHandle* JointHandle = new FJointHandle(&Implementation->Joints, ConstraintInstance, Body1, Body2);
-		Implementation->JointHandles.Add(JointHandle);
-		return JointHandle;
+		if (ConstraintInstance->ProfileInstance.bUseLinearJointSolver) 
+		{
+			FJointHandle* JointHandle = new FJointHandle(&(Implementation->Joints.LinearConstraints), ConstraintInstance, Body1, Body2);
+			Implementation->JointHandles.Add(JointHandle);
+			return JointHandle;
+		}
+		else
+		{
+			FJointHandle* JointHandle = new FJointHandle(&(Implementation->Joints.NonLinearConstraints), ConstraintInstance, Body1, Body2);
+			Implementation->JointHandles.Add(JointHandle);
+			return JointHandle;
+		}
 	}
 
 	FJointHandle* FSimulation::CreateJoint(const Chaos::FPBDJointSettings& ConstraintSettings, FActorHandle* const Body1, FActorHandle* const Body2)
 	{
-		FJointHandle* JointHandle = new FJointHandle(&Implementation->Joints, ConstraintSettings, Body1, Body2);
-		Implementation->JointHandles.Add(JointHandle);
-		return JointHandle;
+		if (ConstraintSettings.bUseLinearSolver)
+		{
+			FJointHandle* JointHandle = new FJointHandle(&(Implementation->Joints.LinearConstraints), ConstraintSettings, Body1, Body2);
+			Implementation->JointHandles.Add(JointHandle);
+			return JointHandle;
+		}
+		else
+		{
+			FJointHandle* JointHandle = new FJointHandle(&(Implementation->Joints.NonLinearConstraints), ConstraintSettings, Body1, Body2);
+			Implementation->JointHandles.Add(JointHandle);
+			return JointHandle;
+		}
 	}
 
 	void FSimulation::DestroyJoint(FJointHandle* JointHandle)
@@ -818,11 +838,21 @@ namespace ImmediatePhysics_Chaos
 
 		if (UseLinearJointSolver >= 0)
 		{
-			for (int32 ConstraintIndex = 0; ConstraintIndex < Implementation->Joints.GetNumConstraints(); ConstraintIndex++)
+			if (UseLinearJointSolver != 0)
 			{
-				Implementation->Joints.SetUseLinearSolver(ConstraintIndex, UseLinearJointSolver != 0);
+				if (Implementation->Joints.NonLinearConstraints.NumConstraints() > 0)
+				{
+					Implementation->Joints.MergeNonLinearConstraintsIntoLinearConstraints();
+				}
+			}
+			else
+			{
+				if (Implementation->Joints.LinearConstraints.NumConstraints() > 0)
+				{
+					Implementation->Joints.MergeLinearConstraintsIntoNonLinearConstraints();
+				}
 			}
-			Implementation->Joints.SetUseSimd(bChaosImmediate_UseSimdForLinearSolver); // RBAN solver can use safely the SIMD version
+			Implementation->Joints.LinearConstraints.SetUseSimd(bChaosImmediate_UseSimdForLinearSolver); 
 		}
 
 		if (bUseManifoilds >= 0)
@@ -887,7 +917,7 @@ namespace ImmediatePhysics_Chaos
 				ChaosImmediate_Evolution_ProjectionIterations,
 				ChaosImmediate_Collision_UseManifolds);
 
-			FPBDJointSolverSettings JointsSettings = Implementation->Joints.GetSettings();
+			FPBDJointSolverSettings JointsSettings = Implementation->Joints.LinearConstraints.GetSettings();
 			JointsSettings.SwingTwistAngleTolerance = ChaosImmediate_Joint_SwingTwistAngleTolerance;
 			JointsSettings.PositionTolerance = ChaosImmediate_Joint_PositionTolerance;
 			JointsSettings.AngleTolerance = ChaosImmediate_Joint_AngleTolerance;
@@ -915,7 +945,8 @@ namespace ImmediatePhysics_Chaos
 			JointsSettings.AngularDriveDampingOverride = ChaosImmediate_Joint_AngularDriveDamping;
 			JointsSettings.NumShockPropagationIterations = ChaosImmediate_Joint_NumShockPropagationIterations;
 			JointsSettings.bSortEnabled = true;
-			Implementation->Joints.SetSettings(JointsSettings);
+			Implementation->Joints.LinearConstraints.SetSettings(JointsSettings);
+			Implementation->Joints.NonLinearConstraints.SetSettings(JointsSettings);
 
 			Implementation->Collisions.SetPositionFrictionIterations(ChaosImmediate_Collision_NumPositionFrictionIterations);
 			Implementation->Collisions.SetVelocityFrictionIterations(ChaosImmediate_Collision_NumVelocityFrictionIterations);
@@ -1145,7 +1176,8 @@ namespace ImmediatePhysics_Chaos
 			}
 			if (bChaosImmediate_DebugDrawJoints)
 			{
-				DebugDraw::DrawJointConstraints(Implementation->SimulationSpace.Transform, Implementation->Joints, 1.0f, ChaosImmediate_DebugDrawJointFeatures, &ChaosImmPhysDebugDebugDrawSettings);
+				DebugDraw::DrawJointConstraints(Implementation->SimulationSpace.Transform, Implementation->Joints.LinearConstraints, 1.0f, ChaosImmediate_DebugDrawJointFeatures, &ChaosImmPhysDebugDebugDrawSettings);
+				DebugDraw::DrawJointConstraints(Implementation->SimulationSpace.Transform, Implementation->Joints.NonLinearConstraints, 1.0f, ChaosImmediate_DebugDrawJointFeatures, &ChaosImmPhysDebugDebugDrawSettings);
 			}
 		}
 #endif
diff --git a/Engine/Source/Runtime/Engine/Public/Physics/ImmediatePhysics/ImmediatePhysicsChaos/ImmediatePhysicsCore_Chaos.h b/Engine/Source/Runtime/Engine/Public/Physics/ImmediatePhysics/ImmediatePhysicsChaos/ImmediatePhysicsCore_Chaos.h
--- a/Engine/Source/Runtime/Engine/Public/Physics/ImmediatePhysics/ImmediatePhysicsChaos/ImmediatePhysicsCore_Chaos.h
+++ b/Engine/Source/Runtime/Engine/Public/Physics/ImmediatePhysics/ImmediatePhysicsChaos/ImmediatePhysicsCore_Chaos.h
@@ -17,6 +17,7 @@ namespace Chaos
 	class FKinematicTarget;
 	class FPBDJointConstraintHandle;
 	class FPBDJointConstraints;
+	struct FPBDJointCombinedConstraints;
 	class FPerShapeData;
 	template<class T> class TArrayCollectionArray;
 	struct FKinematicGeometryParticleParameters;
diff --git a/Engine/Source/Runtime/Experimental/Chaos/Private/Chaos/Joint/PBDJointContainerSolver.cpp b/Engine/Source/Runtime/Experimental/Chaos/Private/Chaos/Joint/PBDJointContainerSolver.cpp
--- a/Engine/Source/Runtime/Experimental/Chaos/Private/Chaos/Joint/PBDJointContainerSolver.cpp
+++ b/Engine/Source/Runtime/Experimental/Chaos/Private/Chaos/Joint/PBDJointContainerSolver.cpp
@@ -109,31 +109,27 @@ namespace Chaos
 		///////////////////////////////////////////////////////////////////////////////////////////
 		///////////////////////////////////////////////////////////////////////////////////////////
 
-
-		FPBDJointContainerSolver::FPBDJointContainerSolver(FPBDJointConstraints& InConstraintContainer, const int32 InPriority)
+		template <typename JointSolverType>
+		TPBDJointContainerSolver<JointSolverType>::TPBDJointContainerSolver(FPBDJointConstraints& InConstraintContainer, const int32 InPriority)
 			: FConstraintContainerSolver(InPriority)
 			, ConstraintContainer(InConstraintContainer)
 		{
 		}
 
-		FPBDJointContainerSolver::~FPBDJointContainerSolver()
+		template <typename JointSolverType>
+		TPBDJointContainerSolver<JointSolverType>::~TPBDJointContainerSolver()
 		{
 		}
 
-		bool FPBDJointContainerSolver::UseLinearSolver() const 
+		template <typename JointSolverType>
+		void TPBDJointContainerSolver<JointSolverType>::Reset(const int32 InMaxConstraints)
 		{
-			return NonLinearConstraintSolvers.Num() == 0;
+			ConstraintSolvers.SetNum(InMaxConstraints);
+			ContainerIndices.Reset(InMaxConstraints);
 		}
 
-		void FPBDJointContainerSolver::Reset(const int32 InMaxConstraints)
-		{
-			LinearConstraintSolvers.SetNum(InMaxConstraints);
-			NonLinearConstraintSolvers.SetNum(InMaxConstraints);
-			ContainerLinearConstraintGlobalIndices.Reset(InMaxConstraints);
-			ContainerNonLinearConstraintGlobalIndices.Reset(InMaxConstraints);
-		}
-
-		void FPBDJointContainerSolver::AddConstraints()
+		template <typename JointSolverType>
+		void TPBDJointContainerSolver<JointSolverType>::AddConstraints()
 		{
 			Reset(ConstraintContainer.GetNumConstraints());
 
@@ -142,30 +138,29 @@ namespace Chaos
 			{
 				if (ConstraintContainer.IsConstraintEnabled(ContainerConstraintIndex))
 				{
-					AddConstraint(ContainerConstraintIndex, ConstraintContainer.GetConstraintSettings(ContainerConstraintIndex).bUseLinearSolver);
+					AddConstraint(ContainerConstraintIndex);
 				}
 			}
 		}
 
-		void FPBDJointContainerSolver::AddConstraints(const TArrayView<Private::FPBDIslandConstraint*>& IslandConstraints)
+		template <typename JointSolverType>
+		void TPBDJointContainerSolver<JointSolverType>::AddConstraints(const TArrayView<Private::FPBDIslandConstraint*>& IslandConstraints)
 		{
-			int32 LinearSolverIndex = 0;
-			int32 NonLinearSolverIndex = 0;
 			int32 CurrentIndex = 0;
 			for (Private::FPBDIslandConstraint* IslandConstraint : IslandConstraints)
 			{
 				// We will only ever be given constraints from our container (asserts in non-shipping)
 				const int32 ContainerConstraintIndex = IslandConstraint->GetConstraint()->AsUnsafe<FPBDJointConstraintHandle>()->GetConstraintIndex();
 
-				AddConstraint(ContainerConstraintIndex, ConstraintContainer.GetConstraintSettings(ContainerConstraintIndex).bUseLinearSolver);
+				AddConstraint(ContainerConstraintIndex);
 			}
 		}
 
-		void FPBDJointContainerSolver::AddConstraint(const int32 InContainerConstraintIndex, const bool bUseLinearSolver)
+		template <typename JointSolverType>
+		void TPBDJointContainerSolver<JointSolverType>::AddConstraint(const int32 InContainerConstraintIndex)
 		{
 			// If this triggers, Reset was called with the wrong constraint count
-			check(ContainerLinearConstraintGlobalIndices.Num() <= ContainerLinearConstraintGlobalIndices.Max());
-			check(ContainerNonLinearConstraintGlobalIndices.Num() <= ContainerNonLinearConstraintGlobalIndices.Max());
+			check(ContainerIndices.Num() <= ConstraintSolvers.Num());
 
 			// Only add a constraint if it is working on at least one dynamic body
 			const FGenericParticleHandle Particle0 = GetJointParticle(ConstraintContainer, InContainerConstraintIndex, 0);
@@ -173,19 +168,12 @@ namespace Chaos
 
 			if (Particle0->IsDynamic() || Particle1->IsDynamic())
 			{
-				if (bUseLinearSolver)
-				{
-					ContainerLinearConstraintGlobalIndices.Add(InContainerConstraintIndex);
-				}
-				else
-				{
-					ContainerNonLinearConstraintGlobalIndices.Add(InContainerConstraintIndex);
-				}
+				ContainerIndices.Add(InContainerConstraintIndex);
 			}
 		}
 
 		template <typename SolverType>
-		void AddBodiesImpl(const FPBDJointContainerSolver& Container, const TArray<int32>& SolverGlobalIndices, FSolverBodyContainer& SolverBodyContainer, TArray<SolverType>& Solvers)
+		void AddBodiesImpl(const TPBDJointContainerSolver<SolverType>& Container, const TArray<int32>& SolverGlobalIndices, FSolverBodyContainer& SolverBodyContainer, TArray<SolverType>& Solvers)
 		{
 			for (int32 i = 0; i < SolverGlobalIndices.Num(); i++)
 			{
@@ -201,27 +189,28 @@ namespace Chaos
 
 		}
 
-		void FPBDJointContainerSolver::AddBodies(FSolverBodyContainer& SolverBodyContainer)
+		template <typename JointSolverType>
+		void TPBDJointContainerSolver<JointSolverType>::AddBodies(FSolverBodyContainer& SolverBodyContainer)
 		{
-			AddBodiesImpl(*this, ContainerLinearConstraintGlobalIndices, SolverBodyContainer, LinearConstraintSolvers);
-			AddBodiesImpl(*this, ContainerNonLinearConstraintGlobalIndices, SolverBodyContainer, NonLinearConstraintSolvers);
+			AddBodiesImpl(*this, ContainerIndices, SolverBodyContainer, ConstraintSolvers);
 		}
 
-		void FPBDJointContainerSolver::GatherInput(const FReal Dt)
+		template <typename JointSolverType>
+		void TPBDJointContainerSolver<JointSolverType>::GatherInput(const FReal Dt)
 		{
 			GatherInput(Dt, 0, GetNumConstraints());
 		}
 
 		template<typename SolverType>
-		void GatherInputImpl(const FPBDJointContainerSolver& Container, TArray<SolverType>& Solvers, const TArray<int32>& SolverGlobalIndices, const FReal Dt, const int32 SolverConstraintBeginIndex, const int32 SolverConstraintEndIndex, const bool bUseLinearSolver)
+		void GatherInputImpl(const TPBDJointContainerSolver<SolverType>& Container, TArray<SolverType>& Solvers, const TArray<int32>& SolverGlobalIndices, const FReal Dt, const int32 SolverConstraintBeginIndex, const int32 SolverConstraintEndIndex, const bool bUseLinearSolver)
 		{
 			const FPBDJointSolverSettings& SolverSettings = Container.GetSettings();
 
 			for (int32 SolverConstraintIndex = FMath::Max(0, SolverConstraintBeginIndex); SolverConstraintIndex < FMath::Min(SolverConstraintEndIndex, SolverGlobalIndices.Num()); ++SolverConstraintIndex)
 			{
-				const FPBDJointSettings& JointSettings = Container.GetConstraintSettings(SolverConstraintIndex, bUseLinearSolver);
+				const FPBDJointSettings& JointSettings = Container.GetConstraintSettings(SolverConstraintIndex);
 
-				const int32 ContainerConstraintIndex = Container.GetContainerConstraintIndex(SolverConstraintIndex, bUseLinearSolver);
+				const int32 ContainerConstraintIndex = Container.GetContainerConstraintIndex(SolverConstraintIndex);
 				FGenericParticleHandle Particle0 = GetJointParticle(Container.GetContainer(), ContainerConstraintIndex, 0);
 				FGenericParticleHandle Particle1 = GetJointParticle(Container.GetContainer(), ContainerConstraintIndex, 1);
 				const FRigidTransform3& Frame0 = GetJointFrame(Container.GetContainer(), ContainerConstraintIndex, 0);
@@ -230,23 +219,24 @@ namespace Chaos
 			}
 		}
 
-		void FPBDJointContainerSolver::GatherInput(const FReal Dt, const int32 ConstraintBeginIndex, const int32 ConstraintEndIndex)
+		template<typename SolverType>
+		void TPBDJointContainerSolver<SolverType>::GatherInput(const FReal Dt, const int32 ConstraintBeginIndex, const int32 ConstraintEndIndex)
 		{
-			GatherInputImpl(*this, LinearConstraintSolvers, ContainerLinearConstraintGlobalIndices, Dt, ConstraintBeginIndex, ConstraintEndIndex, true);
-			GatherInputImpl(*this, NonLinearConstraintSolvers, ContainerNonLinearConstraintGlobalIndices, Dt, ConstraintBeginIndex - ContainerLinearConstraintGlobalIndices.Num(), ConstraintEndIndex - -ContainerLinearConstraintGlobalIndices.Num(), false);
+			GatherInputImpl(*this, ConstraintSolvers, ContainerIndices, Dt, ConstraintBeginIndex, ConstraintEndIndex, true);
 		}
 
-		void FPBDJointContainerSolver::ScatterOutput(const FReal Dt)
+		template<typename SolverType>
+		void TPBDJointContainerSolver<SolverType>::ScatterOutput(const FReal Dt)
 		{
-			ScatterOutput(Dt, 0, ContainerLinearConstraintGlobalIndices.Num() + ContainerNonLinearConstraintGlobalIndices.Num());
+			ScatterOutput(Dt, 0, GetNumConstraints());
 		}
 
 		template<typename SolverType>
-		void ScatterOutputImpl(const FPBDJointContainerSolver& Container, TArray<SolverType>& Solvers, const TArray<int32>& SolverGlobalIndices, const FReal Dt, const int32 SolverConstraintBeginIndex, const int32 SolverConstraintEndIndex, const bool bUseLinearSolver)
+		void ScatterOutputImpl(const TPBDJointContainerSolver<SolverType>& Container, TArray<SolverType>& Solvers, const TArray<int32>& SolverGlobalIndices, const FReal Dt, const int32 SolverConstraintBeginIndex, const int32 SolverConstraintEndIndex, const bool bUseLinearSolver)
 		{
 			for (int32 SolverConstraintIndex = FMath::Max(0, SolverConstraintBeginIndex); SolverConstraintIndex < FMath::Min(SolverConstraintEndIndex, SolverGlobalIndices.Num()); ++SolverConstraintIndex)
 			{
-				const int32 GlobalContainerConstraintIndex = Container.GetContainerConstraintIndex(SolverConstraintIndex, bUseLinearSolver);
+				const int32 GlobalContainerConstraintIndex = Container.GetContainerConstraintIndex(SolverConstraintIndex);
 				if (Dt > UE_SMALL_NUMBER)
 				{
 					SolverType& Solver = Solvers[SolverConstraintIndex];
@@ -273,23 +263,22 @@ namespace Chaos
 			}
 		}
 
-		void FPBDJointContainerSolver::ScatterOutput(const FReal Dt, const int32 ConstraintBeginIndex, const int32 ConstraintEndIndex)
+		template<typename SolverType>
+		void TPBDJointContainerSolver<SolverType>::ScatterOutput(const FReal Dt, const int32 ConstraintBeginIndex, const int32 ConstraintEndIndex)
 		{
-			ScatterOutputImpl(*this, LinearConstraintSolvers, ContainerLinearConstraintGlobalIndices, Dt, ConstraintBeginIndex, ConstraintEndIndex, true);
-			ScatterOutputImpl(*this, NonLinearConstraintSolvers, ContainerNonLinearConstraintGlobalIndices, Dt, ConstraintBeginIndex - ContainerLinearConstraintGlobalIndices.Num(), ConstraintEndIndex - ContainerLinearConstraintGlobalIndices.Num(), false);
+			ScatterOutputImpl(*this, ConstraintSolvers, ContainerIndices, Dt, ConstraintBeginIndex, ConstraintEndIndex, true);
 		}
 
-		void FPBDJointContainerSolver::ResizeSolverArrays()
+		template<typename SolverType>
+		void TPBDJointContainerSolver<SolverType>::ResizeSolverArrays()
 		{
-			check(LinearConstraintSolvers.Num() >= ContainerLinearConstraintGlobalIndices.Num());
-			LinearConstraintSolvers.SetNum(ContainerLinearConstraintGlobalIndices.Num());
-			check(NonLinearConstraintSolvers.Num() >= ContainerNonLinearConstraintGlobalIndices.Num());
-			NonLinearConstraintSolvers.SetNum(ContainerNonLinearConstraintGlobalIndices.Num());
+			check(ConstraintSolvers.Num() >= ContainerIndices.Num());
+			ConstraintSolvers.SetNum(ContainerIndices.Num());
 		}
 
 		// Apply position constraints for linear or non-linear solvers
 		template<typename SolverType>
-		void ApplyPositionConstraintsImpl(const FPBDJointContainerSolver& Container, TArray<SolverType>& Solvers, const FReal Dt, const int32 It, const int32 NumIts, const bool bUseLinearSolver)
+		void ApplyPositionConstraintsImpl(const TPBDJointContainerSolver<SolverType>& Container, TArray<SolverType>& Solvers, const FReal Dt, const int32 It, const int32 NumIts)
 		{
 			const FPBDJointSolverSettings& Settings = Container.GetSettings();
 			const FReal IterationStiffness = GetJointIterationStiffness(Settings, It, NumIts);
@@ -302,7 +291,7 @@ namespace Chaos
 					continue;
 				}
 
-				const FPBDJointSettings& JointSettings = Container.GetConstraintSettings(SolverConstraintIndex, bUseLinearSolver);
+				const FPBDJointSettings& JointSettings = Container.GetConstraintSettings(SolverConstraintIndex);
 				Solver.Update(Dt, Settings, JointSettings);
 
 				// Set parent inverse mass scale based on current shock propagation state
@@ -321,19 +310,16 @@ namespace Chaos
 			}
 		}
 
-		void FPBDJointContainerSolver::ApplyPositionConstraints(const FReal Dt, const int32 It, const int32 NumIts)
+		template<typename SolverType>
+		void TPBDJointContainerSolver<SolverType>::ApplyPositionConstraints(const FReal Dt, const int32 It, const int32 NumIts)
 		{
 			ResizeSolverArrays();
-
-			constexpr bool bUseSolveLinear = true;
-
-			ApplyPositionConstraintsImpl(*this, LinearConstraintSolvers, Dt, It, NumIts, bUseSolveLinear);
-			ApplyPositionConstraintsImpl(*this, NonLinearConstraintSolvers, Dt, It, NumIts, !bUseSolveLinear);
+			ApplyPositionConstraintsImpl(*this, ConstraintSolvers, Dt, It, NumIts);
 		}
 
 		// Apply velocity constraints for linear or non-linear solvers
 		template<typename SolverType>
-		void ApplyVelocityConstraintsImpl(const FPBDJointContainerSolver& Container, TArray<SolverType>& Solvers, const FReal Dt, const int32 It, const int32 NumIts, const bool bUseLinearSolver)
+		void ApplyVelocityConstraintsImpl(const TPBDJointContainerSolver<SolverType>& Container, TArray<SolverType>& Solvers, const FReal Dt, const int32 It, const int32 NumIts)
 		{
 			const FPBDJointSolverSettings& Settings = Container.GetSettings();
 			const FReal IterationStiffness = GetJointIterationStiffness(Settings, It, NumIts);
@@ -346,7 +332,7 @@ namespace Chaos
 					continue;
 				}
 
-				const FPBDJointSettings& JointSettings = Container.GetConstraintSettings(SolverConstraintIndex, bUseLinearSolver);
+				const FPBDJointSettings& JointSettings = Container.GetConstraintSettings(SolverConstraintIndex);
 				Solver.Update(Dt, Settings, JointSettings);
 
 				// Set parent inverse mass scale based on current shock propagation state
@@ -359,20 +345,39 @@ namespace Chaos
 			}
 		}
 
-		void FPBDJointContainerSolver::ApplyVelocityConstraints(const FReal Dt, const int32 It, const int32 NumIts)
+		template<typename SolverType>
+		void TPBDJointContainerSolver<SolverType>::ApplyVelocityConstraints(const FReal Dt, const int32 It, const int32 NumIts)
 		{
-			constexpr bool bUseSolveLinear = true;
-			ApplyVelocityConstraintsImpl(*this, LinearConstraintSolvers, Dt, It, NumIts, bUseSolveLinear);
-			ApplyVelocityConstraintsImpl(*this, NonLinearConstraintSolvers, Dt, It, NumIts, !bUseSolveLinear);
+			ApplyVelocityConstraintsImpl(*this, ConstraintSolvers, Dt, It, NumIts);
 		}
 
-		void FPBDJointContainerSolver::ApplyProjectionConstraints(const FReal Dt, const int32 It, const int32 NumIts)
+		template<typename SolverType>
+		void TPBDJointContainerSolver<SolverType>::ApplyProjectionConstraints(const FReal Dt, const int32 It, const int32 NumIts)
+		{
+			if (UseLinearSolver())
+			{
+				ApplyLinearProjectionConstraints(Dt, It, NumIts);
+			}
+			else
+			{
+				ApplyNonLinearProjectionConstraints(Dt, It, NumIts);
+			}
+		}
+
+		template <>
+		bool TPBDJointContainerSolver<FPBDJointCachedSolver>::UseLinearSolver() const
 		{
-			ApplyLinearProjectionConstraints(Dt, It, NumIts);
-			ApplyNonLinearProjectionConstraints(Dt, It, NumIts);
+			return true;
 		}
 
-		void FPBDJointContainerSolver::ApplyLinearProjectionConstraints(const FReal Dt, const int32 It, const int32 NumIts)
+		template <>
+		bool TPBDJointContainerSolver<FPBDJointSolver>::UseLinearSolver() const
+		{
+			return false;
+		}
+
+		template<>
+		void TPBDJointContainerSolver<FPBDJointCachedSolver>::ApplyLinearProjectionConstraints(const FReal Dt, const int32 It, const int32 NumIts)
 		{
 			constexpr bool bUseSolveLinear = true;
 			const FPBDJointSolverSettings& Settings = GetSettings();
@@ -383,15 +388,15 @@ namespace Chaos
 				// This must happen for all joints before we project any joints so the the initial state for each joint is not polluted by any earlier projections.
 				// @todo(chaos): if we ever support projection on other constraint types, we will need a PrepareProjection phase so that all constraint types
 				// can initialize correctly before any constraints apply their projection. For now we can just check the iteration count is zero.
-				for (int32 SolverConstraintIndex = 0; SolverConstraintIndex < LinearConstraintSolvers.Num(); ++SolverConstraintIndex)
+				for (int32 SolverConstraintIndex = 0; SolverConstraintIndex < ConstraintSolvers.Num(); ++SolverConstraintIndex)
 				{
-					FPBDJointCachedSolver& Solver = LinearConstraintSolvers[SolverConstraintIndex];
+					FPBDJointCachedSolver& Solver = ConstraintSolvers[SolverConstraintIndex];
 					if (!Solver.RequiresSolve())
 					{
 						continue;
 					}
 
-					const FPBDJointSettings& JointSettings = GetConstraintSettings(SolverConstraintIndex, bUseSolveLinear);
+					const FPBDJointSettings& JointSettings = GetConstraintSettings(SolverConstraintIndex);
 					if (!JointSettings.bProjectionEnabled)
 					{
 						continue;
@@ -401,15 +406,15 @@ namespace Chaos
 				}
 			}
 
-			for (int32 SolverConstraintIndex = 0; SolverConstraintIndex < LinearConstraintSolvers.Num(); ++SolverConstraintIndex)
+			for (int32 SolverConstraintIndex = 0; SolverConstraintIndex < ConstraintSolvers.Num(); ++SolverConstraintIndex)
 			{
-				FPBDJointCachedSolver& Solver = LinearConstraintSolvers[SolverConstraintIndex];
+				FPBDJointCachedSolver& Solver = ConstraintSolvers[SolverConstraintIndex];
 				if (!Solver.RequiresSolve())
 				{
 					continue;
 				}
 
-				const FPBDJointSettings& JointSettings = GetConstraintSettings(SolverConstraintIndex, bUseSolveLinear);
+				const FPBDJointSettings& JointSettings = GetConstraintSettings(SolverConstraintIndex);
 				if (!JointSettings.bProjectionEnabled)
 				{
 					continue;
@@ -425,20 +430,30 @@ namespace Chaos
 			}
 		}
 
-		void FPBDJointContainerSolver::ApplyNonLinearProjectionConstraints(const FReal Dt, const int32 It, const int32 NumIts)
+		template<>
+		void TPBDJointContainerSolver<FPBDJointSolver>::ApplyLinearProjectionConstraints(const FReal Dt, const int32 It, const int32 NumIts)
+		{
+		}
+
+		template<>
+		void TPBDJointContainerSolver<FPBDJointCachedSolver>::ApplyNonLinearProjectionConstraints(const FReal Dt, const int32 It, const int32 NumIts)
+		{
+		}
+
+		template<>
+		void TPBDJointContainerSolver<FPBDJointSolver>::ApplyNonLinearProjectionConstraints(const FReal Dt, const int32 It, const int32 NumIts)
 		{
-			constexpr bool bUseSolveLinear = false;
 			const FPBDJointSolverSettings& Settings = GetSettings();
 
-			for (int32 SolverConstraintIndex = 0; SolverConstraintIndex < NonLinearConstraintSolvers.Num(); ++SolverConstraintIndex)
+			for (int32 SolverConstraintIndex = 0; SolverConstraintIndex < ConstraintSolvers.Num(); ++SolverConstraintIndex)
 			{
-				FPBDJointSolver& Solver = NonLinearConstraintSolvers[SolverConstraintIndex];
+				FPBDJointSolver& Solver = ConstraintSolvers[SolverConstraintIndex];
 				if (!Solver.RequiresSolve())
 				{
 					continue;
 				}
 
-				const FPBDJointSettings& JointSettings = GetConstraintSettings(SolverConstraintIndex, bUseSolveLinear);
+				const FPBDJointSettings& JointSettings = GetConstraintSettings(SolverConstraintIndex);
 				if (!JointSettings.bProjectionEnabled)
 				{
 					continue;
@@ -457,6 +472,8 @@ namespace Chaos
 				Solver.ApplyProjections(Dt, Settings, JointSettings, bLastIteration);
 			}
 		}
-
 	}	// namespace Private
-}	// namespace Chaos
\ No newline at end of file
+}	// namespace Chaos
+
+template class Chaos::Private::TPBDJointContainerSolver<Chaos::FPBDJointCachedSolver>;
+template class Chaos::Private::TPBDJointContainerSolver<Chaos::FPBDJointSolver>;
\ No newline at end of file
diff --git a/Engine/Source/Runtime/Experimental/Chaos/Private/Chaos/PBDJointConstraints.cpp b/Engine/Source/Runtime/Experimental/Chaos/Private/Chaos/PBDJointConstraints.cpp
--- a/Engine/Source/Runtime/Experimental/Chaos/Private/Chaos/PBDJointConstraints.cpp
+++ b/Engine/Source/Runtime/Experimental/Chaos/Private/Chaos/PBDJointConstraints.cpp
@@ -565,12 +565,26 @@ namespace Chaos
 
 	TUniquePtr<FConstraintContainerSolver> FPBDJointConstraints::CreateSceneSolver(const int32 Priority)
 	{
-		return MakeUnique<Private::FPBDJointContainerSolver>(*this, Priority);
+		if (bUseLinearSolver)
+		{
+			return MakeUnique<Private::TPBDJointContainerSolver<Chaos::FPBDJointCachedSolver>>(*this, Priority);
+		}
+		else
+		{
+			return MakeUnique<Private::TPBDJointContainerSolver<Chaos::FPBDJointSolver>>(*this, Priority);
+		}
 	}
 
 	TUniquePtr<FConstraintContainerSolver> FPBDJointConstraints::CreateGroupSolver(const int32 Priority)
 	{
-		return MakeUnique<Private::FPBDJointContainerSolver>(*this, Priority);
+		if (bUseLinearSolver)
+		{
+			return MakeUnique<Private::TPBDJointContainerSolver<Chaos::FPBDJointCachedSolver>>(*this, Priority);
+		}
+		else
+		{
+			return MakeUnique<Private::TPBDJointContainerSolver<Chaos::FPBDJointSolver>>(*this, Priority);
+		}
 	}
 	
 	const FPBDJointSolverSettings& FPBDJointConstraints::GetSettings() const
@@ -709,7 +723,7 @@ namespace Chaos
 	{
 		// Sort constraints so that constraints with lower level (closer to a kinematic joint) are first
 		// @todo(chaos): this is only needed for RBAN and should be moved to an RBAN specific solver
-		// (so we should split FPBDJointContainerSolver into Scene and Group versions)
+		// (so we should split TPBDJointContainerSolver into Scene and Group versions)
 		SCOPE_CYCLE_COUNTER(STAT_Joints_Sort);
 
 		TArray<int32> SortedIndices;
@@ -1335,6 +1349,79 @@ namespace Chaos
 		}
 	}
 
+	FPBDJointCombinedConstraints::FPBDJointCombinedConstraints()
+	{
+		LinearConstraints.SetUseLinearSolver(true);
+		NonLinearConstraints.SetUseLinearSolver(false);
+	}
+
+	//This function moves one constraint info from linear constraint into nonlinear constraints
+	void FPBDJointCombinedConstraints::ConvertLinearConstraintToNonLinear(const int32 LinearConstraintIndex)
+	{
+		if (!ensureMsgf(LinearConstraintIndex > INDEX_NONE && LinearConstraintIndex < LinearConstraints.NumConstraints(), TEXT("Linear constraint accessed is out of bounds!")))
+		{
+			return;
+		}
+		NonLinearConstraints.bJointsDirty = true;
+
+		int NonLinearConstraintsNum = NonLinearConstraints.Handles.Num();
+		NonLinearConstraints.Handles.Add(MoveTemp(LinearConstraints.Handles[LinearConstraintIndex]));
+		NonLinearConstraints.ConstraintParticles.Add(MoveTemp(LinearConstraints.ConstraintParticles[LinearConstraintIndex]));
+		NonLinearConstraints.ConstraintStates.Add(MoveTemp(LinearConstraints.ConstraintStates[LinearConstraintIndex]));
+		NonLinearConstraints.ConstraintSettings.Add(MoveTemp(LinearConstraints.ConstraintSettings[LinearConstraintIndex]));
+		NonLinearConstraints.ConstraintSettings[NonLinearConstraintsNum].bUseLinearSolver = false;
+		NonLinearConstraints.Handles[NonLinearConstraintsNum]->SetConstraintIndex(NonLinearConstraintsNum);
+		NonLinearConstraints.Handles[NonLinearConstraintsNum]->SetConstraintContainer(&NonLinearConstraints);
+
+		LinearConstraints.Handles.RemoveAtSwap(LinearConstraintIndex);
+		LinearConstraints.ConstraintParticles.RemoveAtSwap(LinearConstraintIndex);
+		LinearConstraints.ConstraintStates.RemoveAtSwap(LinearConstraintIndex);
+		LinearConstraints.ConstraintSettings.RemoveAtSwap(LinearConstraintIndex);
+		LinearConstraints.bJointsDirty = true;
+	}
+
+	//This function moves one constraint info from nonlinear constraint into linear constraints
+	void FPBDJointCombinedConstraints::ConvertNonLinearConstraintToLinear(const int32 NonLinearConstraintIndex)
+	{
+		if (!ensureMsgf(NonLinearConstraintIndex > INDEX_NONE && NonLinearConstraintIndex < NonLinearConstraints.NumConstraints(), TEXT("NonLinear constraint accessed is out of bounds!")))
+		{
+			return; 
+		}
+
+		LinearConstraints.bJointsDirty = true;
+		int LinearConstraintsNum = LinearConstraints.Handles.Num();
+		LinearConstraints.Handles.Add(MoveTemp(NonLinearConstraints.Handles[NonLinearConstraintIndex]));
+		LinearConstraints.ConstraintParticles.Add(MoveTemp(NonLinearConstraints.ConstraintParticles[NonLinearConstraintIndex]));
+		LinearConstraints.ConstraintStates.Add(MoveTemp(NonLinearConstraints.ConstraintStates[NonLinearConstraintIndex]));
+		LinearConstraints.ConstraintSettings.Add(MoveTemp(NonLinearConstraints.ConstraintSettings[NonLinearConstraintIndex]));
+		LinearConstraints.ConstraintSettings[LinearConstraintsNum].bUseLinearSolver = true;
+		LinearConstraints.Handles[LinearConstraintsNum]->SetConstraintIndex(LinearConstraintsNum);
+		LinearConstraints.Handles[LinearConstraintsNum]->SetConstraintContainer(&LinearConstraints);
+
+		NonLinearConstraints.Handles.RemoveAtSwap(NonLinearConstraintIndex);
+		NonLinearConstraints.ConstraintParticles.RemoveAtSwap(NonLinearConstraintIndex);
+		NonLinearConstraints.ConstraintStates.RemoveAtSwap(NonLinearConstraintIndex);
+		NonLinearConstraints.ConstraintSettings.RemoveAtSwap(NonLinearConstraintIndex);
+		NonLinearConstraints.bJointsDirty = true;
+	}
+
+	void FPBDJointCombinedConstraints::MergeLinearConstraintsIntoNonLinearConstraints()
+	{
+		while (LinearConstraints.NumConstraints() > 0)
+		{
+			ConvertLinearConstraintToNonLinear(0);
+		}
+	}
+
+	void FPBDJointCombinedConstraints::MergeNonLinearConstraintsIntoLinearConstraints()
+	{
+		while (NonLinearConstraints.NumConstraints() > 0)
+		{
+			ConvertNonLinearConstraintToLinear(0);
+		}
+	}
+
+
 }
 
 namespace Chaos
diff --git a/Engine/Source/Runtime/Experimental/Chaos/Private/Chaos/PBDRigidsEvolutionGBF.cpp b/Engine/Source/Runtime/Experimental/Chaos/Private/Chaos/PBDRigidsEvolutionGBF.cpp
--- a/Engine/Source/Runtime/Experimental/Chaos/Private/Chaos/PBDRigidsEvolutionGBF.cpp
+++ b/Engine/Source/Runtime/Experimental/Chaos/Private/Chaos/PBDRigidsEvolutionGBF.cpp
@@ -783,7 +783,8 @@ void FPBDRigidsEvolutionGBF::AdvanceOneTimeStepImpl(const FReal Dt, const FSubSt
 			CVD_TRACE_STEP_MID_PHASES_FROM_COLLISION_CONSTRAINTS(CVDDC_EndOfEvolutionCollisionConstraints, GetCollisionConstraints());
 		}
 
-		CVD_TRACE_JOINT_CONSTRAINTS(CVDDC_JointConstraints, JointConstraints);
+		CVD_TRACE_JOINT_CONSTRAINTS(CVDDC_JointLinearConstraints, JointConstraints.LinearConstraints);
+		CVD_TRACE_JOINT_CONSTRAINTS(CVDDC_JointNonLinearConstraints, JointConstraints.NonLinearConstraints);
 		CVD_TRACE_CHARACTER_GROUND_CONSTRAINTS(CVDDC_CharacterGroundConstraints, CharacterGroundConstraints);
 
 		CVD_TRACE_PARTICLES_SOA(Particles, &Clustering);
@@ -1153,7 +1154,7 @@ FPBDRigidsEvolutionGBF::FPBDRigidsEvolutionGBF(
 
 	CollisionConstraints.SetCullDistance(DefaultCollisionCullDistance);
 
-	JointConstraints.SetUseSimd(bChaosRigids_UseSimdForJointsSolver);
+	JointConstraints.LinearConstraints.SetUseSimd(bChaosRigids_UseSimdForJointsSolver); //NonLinear Joints don't have simd code so it's skipped. 
 
 	GetIslandManager().SetMaterialContainers(&PhysicsMaterials, &PerParticlePhysicsMaterials, &SolverPhysicsMaterials);
 	GetIslandManager().SetGravityForces(&GravityForces);
@@ -1178,7 +1179,8 @@ FPBDRigidsEvolutionGBF::FPBDRigidsEvolutionGBF(
 
 	AddConstraintContainer(SuspensionConstraints, ChaosSolverSuspensionPriority);
 	AddConstraintContainer(CollisionConstraints, ChaosSolverCollisionPriority);
-	AddConstraintContainer(JointConstraints, ChaosSolverJointPriority);
+	AddConstraintContainer(JointConstraints.LinearConstraints, ChaosSolverJointPriority);
+	AddConstraintContainer(JointConstraints.NonLinearConstraints, ChaosSolverJointPriority);
 	AddConstraintContainer(CharacterGroundConstraints, ChaosSolverCharacterGroundConstraintPriority);
 
 	SetInternalParticleInitilizationFunction([](const FGeometryParticleHandle*, const FGeometryParticleHandle*) {});
diff --git a/Engine/Source/Runtime/Experimental/Chaos/Private/ChaosVisualDebugger/ChaosVDOptionalDataChannel.cpp b/Engine/Source/Runtime/Experimental/Chaos/Private/ChaosVisualDebugger/ChaosVDOptionalDataChannel.cpp
--- a/Engine/Source/Runtime/Experimental/Chaos/Private/ChaosVisualDebugger/ChaosVDOptionalDataChannel.cpp
+++ b/Engine/Source/Runtime/Experimental/Chaos/Private/ChaosVisualDebugger/ChaosVDOptionalDataChannel.cpp
@@ -83,7 +83,8 @@ CVD_DEFINE_OPTIONAL_DATA_CHANNEL(PostConstraintSolve, EChaosVDDataChannelInitial
 CVD_DEFINE_OPTIONAL_DATA_CHANNEL(EndOfEvolutionCollisionConstraints, EChaosVDDataChannelInitializationFlags::StartEnabled | EChaosVDDataChannelInitializationFlags::CanChangeEnabledState);
 CVD_DEFINE_OPTIONAL_DATA_CHANNEL(EvolutionEnd, EChaosVDDataChannelInitializationFlags::StartEnabled); // Intentionally not allowing changing the enabled state because otherwise there might be nothing to visualize
 CVD_DEFINE_OPTIONAL_DATA_CHANNEL(SceneQueries, EChaosVDDataChannelInitializationFlags::StartEnabled | EChaosVDDataChannelInitializationFlags::CanChangeEnabledState);
-CVD_DEFINE_OPTIONAL_DATA_CHANNEL(JointConstraints, EChaosVDDataChannelInitializationFlags::CanChangeEnabledState);
+CVD_DEFINE_OPTIONAL_DATA_CHANNEL(JointLinearConstraints, EChaosVDDataChannelInitializationFlags::CanChangeEnabledState);
+CVD_DEFINE_OPTIONAL_DATA_CHANNEL(JointNonLinearConstraints, EChaosVDDataChannelInitializationFlags::CanChangeEnabledState);
 CVD_DEFINE_OPTIONAL_DATA_CHANNEL(CharacterGroundConstraints, EChaosVDDataChannelInitializationFlags::CanChangeEnabledState);
 CVD_DEFINE_OPTIONAL_DATA_CHANNEL(AccelerationStructures, EChaosVDDataChannelInitializationFlags::CanChangeEnabledState);
 CVD_DEFINE_OPTIONAL_DATA_CHANNEL(ClusterParticlesChildData, EChaosVDDataChannelInitializationFlags::CanChangeEnabledState);
diff --git a/Engine/Source/Runtime/Experimental/Chaos/Private/ChaosVisualDebugger/ChaosVisualDebuggerTrace.cpp b/Engine/Source/Runtime/Experimental/Chaos/Private/ChaosVisualDebugger/ChaosVisualDebuggerTrace.cpp
--- a/Engine/Source/Runtime/Experimental/Chaos/Private/ChaosVisualDebugger/ChaosVisualDebuggerTrace.cpp
+++ b/Engine/Source/Runtime/Experimental/Chaos/Private/ChaosVisualDebugger/ChaosVisualDebuggerTrace.cpp
@@ -495,7 +495,16 @@ void FChaosVisualDebuggerTrace::TraceConstraintsContainer(TConstArrayView<Chaos:
 		{
 			if (ConstraintContainer->GetConstraintHandleType().IsA(Chaos::FPBDJointConstraintHandle::StaticType()))
 			{
-				CVD_TRACE_JOINT_CONSTRAINTS(CVDDC_JointConstraints, *static_cast<Chaos::FPBDJointConstraints*>(ConstraintContainer));
+				Chaos::FPBDJointConstraints* const JointConstraintPtr = static_cast<Chaos::FPBDJointConstraints*>(ConstraintContainer);
+				if (JointConstraintPtr->GetUseLinearSolver())
+				{
+					CVD_TRACE_JOINT_CONSTRAINTS(CVDDC_JointLinearConstraints, *JointConstraintPtr);
+				}
+				else
+				{
+					CVD_TRACE_JOINT_CONSTRAINTS(CVDDC_JointNonLinearConstraints, *JointConstraintPtr);
+				}
+				
 			}
 			else if (ConstraintContainer->GetConstraintHandleType().IsA(Chaos::FPBDCollisionConstraint::StaticType()))
 			{
diff --git a/Engine/Source/Runtime/Experimental/Chaos/Private/PBDRigidsSolver.cpp b/Engine/Source/Runtime/Experimental/Chaos/Private/PBDRigidsSolver.cpp
--- a/Engine/Source/Runtime/Experimental/Chaos/Private/PBDRigidsSolver.cpp
+++ b/Engine/Source/Runtime/Experimental/Chaos/Private/PBDRigidsSolver.cpp
@@ -2580,7 +2580,8 @@ TRACE_COUNTER_SET(ChaosTraceCounter_##Name, Value)
 		}
 		if (ChaosSolverDrawJoints == 1)
 		{
-			DebugDraw::DrawJointConstraints(FRigidTransform3(), MEvolution->GetJointConstraints(), 1.0f, ChaosSolverDrawJointFeatures, &ChaosSolverDebugDebugDrawSettings);
+			DebugDraw::DrawJointConstraints(FRigidTransform3(), MEvolution->GetJointConstraints().LinearConstraints, 1.0f, ChaosSolverDrawJointFeatures, &ChaosSolverDebugDebugDrawSettings);
+			DebugDraw::DrawJointConstraints(FRigidTransform3(), MEvolution->GetJointConstraints().NonLinearConstraints, 1.0f, ChaosSolverDrawJointFeatures, &ChaosSolverDebugDebugDrawSettings);
 		}
 		if (ChaosSolverDrawCharacterGroundConstraints == 1)
 		{
@@ -2817,7 +2818,7 @@ TRACE_COUNTER_SET(ChaosTraceCounter_##Name, Value)
 		CollisionDetectorSettings.bAllowMACD = bChaosUseMACD;
 		MEvolution->GetCollisionConstraints().SetDetectorSettings(CollisionDetectorSettings);
 
-		FPBDJointSolverSettings JointsSettings = MEvolution->GetJointConstraints().GetSettings();
+		FPBDJointSolverSettings JointsSettings = MEvolution->GetJointConstraints().LinearConstraints.GetSettings();
 		JointsSettings.MinSolverStiffness = ChaosSolverJointMinSolverStiffness;
 		JointsSettings.MaxSolverStiffness = ChaosSolverJointMaxSolverStiffness;
 		JointsSettings.NumIterationsAtMaxSolverStiffness = ChaosSolverJointNumIterationsAtMaxSolverStiffness;
@@ -2830,7 +2831,8 @@ TRACE_COUNTER_SET(ChaosTraceCounter_##Name, Value)
 		JointsSettings.NumShockPropagationIterations = ChaosSolverJointNumShockProagationIterations;
 		JointsSettings.ShockPropagationOverride = ChaosSolverJointShockPropagation;
 		JointsSettings.bSortEnabled = false;
-		MEvolution->GetJointConstraints().SetSettings(JointsSettings);
+		MEvolution->GetJointConstraints().LinearConstraints.SetSettings(JointsSettings);
+		MEvolution->GetJointConstraints().NonLinearConstraints.SetSettings(JointsSettings);
 
 		// Apply CVAR overrides if set
 		{
diff --git a/Engine/Source/Runtime/Experimental/Chaos/Private/PhysicsProxy/JointConstraintProxy.cpp b/Engine/Source/Runtime/Experimental/Chaos/Private/PhysicsProxy/JointConstraintProxy.cpp
--- a/Engine/Source/Runtime/Experimental/Chaos/Private/PhysicsProxy/JointConstraintProxy.cpp
+++ b/Engine/Source/Runtime/Experimental/Chaos/Private/PhysicsProxy/JointConstraintProxy.cpp
@@ -128,7 +128,14 @@ void FJointConstraintPhysicsProxy::InitializeOnPhysicsThread(FPBDRigidsSolver* I
 
 					FixConnectorTransformsForRoot(OriginalParticleHandles_PT, RootHandles, ConnectorTransforms);
 
-					Constraint_PT = JointConstraints.AddConstraint(RootHandles, ConnectorTransforms);
+					if (JointSettings->bUseLinearSolver)
+					{
+						Constraint_PT = JointConstraints.LinearConstraints.AddConstraint(RootHandles, ConnectorTransforms);
+					}
+					else
+					{
+						Constraint_PT = JointConstraints.NonLinearConstraints.AddConstraint(RootHandles, ConnectorTransforms);
+					}
 
 					RootHandles[0]->AddConstraintHandle(Constraint_PT);
 					RootHandles[1]->AddConstraintHandle(Constraint_PT);
@@ -175,7 +182,14 @@ void FJointConstraintPhysicsProxy::DestroyOnPhysicsThread(FPBDRigidsSolver* InSo
 		}
 
 		FPBDRigidsSolver::FJointConstraints& JointConstraints = InSolver->GetJointConstraints();
-		JointConstraints.RemoveConstraint(Constraint_PT->GetConstraintIndex());
+		if (Settings.bUseLinearSolver)
+		{
+			JointConstraints.LinearConstraints.RemoveConstraint(Constraint_PT->GetConstraintIndex());
+		}
+		else
+		{
+			JointConstraints.NonLinearConstraints.RemoveConstraint(Constraint_PT->GetConstraintIndex());
+		}
 		Constraint_PT = nullptr;
 		OriginalParticleHandles_PT = { nullptr, nullptr };
 	}
diff --git a/Engine/Source/Runtime/Experimental/Chaos/Public/Chaos/Joint/PBDJointContainerSolver.h b/Engine/Source/Runtime/Experimental/Chaos/Public/Chaos/Joint/PBDJointContainerSolver.h
--- a/Engine/Source/Runtime/Experimental/Chaos/Public/Chaos/Joint/PBDJointContainerSolver.h
+++ b/Engine/Source/Runtime/Experimental/Chaos/Public/Chaos/Joint/PBDJointContainerSolver.h
@@ -15,20 +15,21 @@ namespace Chaos
 		/**
 		 * Runs the solvers for a set of constraints belonging to a JointConstraints container.
 		 * 
-		 * For the main scene, each IslandGroup owns a FPBDJointContainerSolver and the list of constraints to be solved
-		 * and the order in which they are solved is determined by the constraint graph.
+		 * For the main scene, each IslandGroup owns two TPBDJointContainerSolvers (one for linear and one for nonlinear) and the list of constraints to be solved
+		 * and the order in which they are solved is determined by the constraint graph. The two TPBDJointContainerSolvers are grouped under struct FPBDJointCombinedConstraints.
 		 * 
-		 * For RBAN, there is one FPBDJointContainerSolver that solves all joints in the simulation in the order that
+		 * For RBAN, there is one FPBDJointCombinedConstraints with two TPBDJointContainerSolvers (linear & nonlinear) that solves all joints in the simulation in the order that
 		 * they occur in the container.
 		*/
-		class FPBDJointContainerSolver : public FConstraintContainerSolver
+		template <typename JointSolverType>
+		class TPBDJointContainerSolver : public FConstraintContainerSolver
 		{
 		public:
-			FPBDJointContainerSolver(FPBDJointConstraints& InConstraintContainer, const int32 InPriority);
-			~FPBDJointContainerSolver();
+			TPBDJointContainerSolver(FPBDJointConstraints& InConstraintContainer, const int32 InPriority);
+			~TPBDJointContainerSolver();
 
 			// FConstraintContainerSolver impl
-			virtual int32 GetNumConstraints() const override final { return ContainerLinearConstraintGlobalIndices.Num() + ContainerNonLinearConstraintGlobalIndices.Num();}
+			virtual int32 GetNumConstraints() const override final { return ContainerIndices.Num();}
 			virtual void Reset(const int32 InMaxCollisions) override final;
 			virtual void AddConstraints() override final;
 			virtual void AddConstraints(const TArrayView<Private::FPBDIslandConstraint*>& IslandConstraints) override final;
@@ -43,53 +44,22 @@ namespace Chaos
 
 			FPBDJointConstraints& GetContainer() const { return ConstraintContainer; }
 			const FPBDJointSolverSettings& GetSettings() const { return ConstraintContainer.GetSettings(); }
-			const FPBDJointSettings& GetLinearConstraintSettings(const int32 InConstraintIndex) const { return ConstraintContainer.GetConstraintSettings(ContainerLinearConstraintGlobalIndices[InConstraintIndex]);}
-			const FPBDJointSettings& GetNonLinearConstraintSettings(const int32 InConstraintIndex) const { return ConstraintContainer.GetConstraintSettings(ContainerNonLinearConstraintGlobalIndices[InConstraintIndex]); }
-			const FPBDJointSettings& GetConstraintSettings(const int32 InConstraintIndex, const bool bUseLinearSolver = true) const 
-			{
-				if (bUseLinearSolver)
-				{
-					return ConstraintContainer.GetConstraintSettings(ContainerLinearConstraintGlobalIndices[InConstraintIndex]);
-				}
-				else
-				{
-					return ConstraintContainer.GetConstraintSettings(ContainerNonLinearConstraintGlobalIndices[InConstraintIndex]);
-				}
-			}
-			int32 GetContainerConstraintIndex(const int32 InConstraintIndex, const bool bUseLinearSolver) const 
-			{
-				if (bUseLinearSolver)
-				{
-					return ContainerLinearConstraintGlobalIndices[InConstraintIndex];
-				}
-				else
-				{
-					return ContainerNonLinearConstraintGlobalIndices[InConstraintIndex];
-				}
-			}
-			int32 GetContainerLinearConstraintGlobalIndex(const int32 InLocalConstraintIndex) const { return ContainerLinearConstraintGlobalIndices[InLocalConstraintIndex];}
-			int32 GetContainerNonLinearConstraintGlobalIndex(const int32 InLocalConstraintIndex) const { return ContainerNonLinearConstraintGlobalIndices[InLocalConstraintIndex]; }
-			int32 GetContainerLinearGlobalContraintIndex(const int32 InSolverIndex) const { return ContainerLinearConstraintGlobalIndices[InSolverIndex]; }
-			int32 GetContainerNonLinearGlobalContraintIndex(const int32 InSolverIndex) const { return ContainerNonLinearConstraintGlobalIndices[InSolverIndex]; }
-
-
+			const FPBDJointSettings& GetConstraintSettings(const int32 InConstraintIndex) const { return ConstraintContainer.GetConstraintSettings(ContainerIndices[InConstraintIndex]);}
+			int32 GetContainerConstraintIndex(const int32 InConstraintIndex) const { return ContainerIndices[InConstraintIndex];}
 		private:
 			bool UseLinearSolver() const;
-			void AddConstraint(const int32 InContainerConstraintIndex, const bool bUseLinearSolver = true);
+			void AddConstraint(const int32 InContainerConstraintIndex);
 			void ResizeSolverArrays();
 			void ApplyLinearProjectionConstraints(const FReal Dt, const int32 It, const int32 NumIts);
 			void ApplyNonLinearProjectionConstraints(const FReal Dt, const int32 It, const int32 NumIts);
 
 			FPBDJointConstraints& ConstraintContainer;
 
-			// The linear and non-linear joint solvers. One for each joint we wish to solve in the order that they are solved
-			TArray<FPBDJointCachedSolver> LinearConstraintSolvers;
-			TArray<FPBDJointSolver> NonLinearConstraintSolvers;
+			TArray<JointSolverType> ConstraintSolvers;
 
 			// Index remapping from respective internal index of each array to index in the joint container [0,NumJointsInWorld)
-			TArray<int32> ContainerLinearConstraintGlobalIndices;
-			TArray<int32> ContainerNonLinearConstraintGlobalIndices;
+			TArray<int32> ContainerIndices;
 		};
-
+		
 	}
 }
\ No newline at end of file
diff --git a/Engine/Source/Runtime/Experimental/Chaos/Public/Chaos/PBDJointConstraints.h b/Engine/Source/Runtime/Experimental/Chaos/Public/Chaos/PBDJointConstraints.h
--- a/Engine/Source/Runtime/Experimental/Chaos/Public/Chaos/PBDJointConstraints.h
+++ b/Engine/Source/Runtime/Experimental/Chaos/Public/Chaos/PBDJointConstraints.h
@@ -16,7 +16,8 @@ namespace Chaos
 {
 	namespace Private
 	{
-		class FPBDJointContainerSolver;
+		template <typename JointSolverType>
+		class TPBDJointContainerSolver;
 	}
 
 	/**
@@ -115,6 +116,8 @@ namespace Chaos
 		UE_DEPRECATED(5.2, "No longer used")
 		CHAOS_API int32 GetConstraintColor() const;
 
+		CHAOS_API void SetConstraintContainer(FPBDConstraintContainer* NewConstraintContainer) { ConstraintContainer = NewConstraintContainer;  }
+
 	protected:
 		using Base::ConstraintIndex;
 		using Base::ConcreteContainer;
@@ -169,7 +172,7 @@ namespace Chaos
 		CHAOS_API virtual ~FPBDJointConstraints();
 
 		/**
-		 * Get the solver settings (used by FPBDJointContainerSolver)
+		 * Get the solver settings (used by TPBDJointContainerSolver)
 		*/
 		CHAOS_API const FPBDJointSolverSettings& GetSettings() const;
 
@@ -393,39 +396,45 @@ namespace Chaos
 		CHAOS_API int32 GetConstraintLevel(int32 ConstraintIndex) const;
 		CHAOS_API int32 GetConstraintColor(int32 ConstraintIndex) const;
 
-		// @todo(chaos): functionality no longer supported here (see FPBDJointContainerSolver)
-		UE_DEPRECATED(5.2, "Joint Solver API moved to FPBDJointContainerSolver")
+		// @todo(chaos): functionality no longer supported here (see TPBDJointContainerSolver)
+		UE_DEPRECATED(5.2, "Joint Solver API moved to TPBDJointContainerSolver")
 		void AddBodies(FSolverBodyContainer& SolverBodyContainer) {}
-		UE_DEPRECATED(5.2, "Joint Solver API moved to FPBDJointContainerSolver")
+		UE_DEPRECATED(5.2, "Joint Solver API moved to TPBDJointContainerSolver")
 		void GatherInput(const FReal Dt){}
-		UE_DEPRECATED(5.2, "Joint Solver API moved to FPBDJointContainerSolver")
+		UE_DEPRECATED(5.2, "Joint Solver API moved to TPBDJointContainerSolver")
 		void ScatterOutput(const FReal Dt){}
-		UE_DEPRECATED(5.2, "Joint Solver API moved to FPBDJointContainerSolver")
+		UE_DEPRECATED(5.2, "Joint Solver API moved to TPBDJointContainerSolver")
 		void ApplyPositionConstraints(const FReal Dt, const int32 It, const int32 NumIts) {}
-		UE_DEPRECATED(5.2, "Joint Solver API moved to FPBDJointContainerSolver")
+		UE_DEPRECATED(5.2, "Joint Solver API moved to TPBDJointContainerSolver")
 		void ApplyVelocityConstraints(const FReal Dt, const int32 It, const int32 NumIts) {}
-		UE_DEPRECATED(5.2, "Joint Solver API moved to FPBDJointContainerSolver")
+		UE_DEPRECATED(5.2, "Joint Solver API moved to TPBDJointContainerSolver")
 		void ApplyProjectionConstraints(const FReal Dt, const int32 It, const int32 NumIts) {}
-		UE_DEPRECATED(5.2, "Joint Solver API moved to FPBDJointContainerSolver")
+		UE_DEPRECATED(5.2, "Joint Solver API moved to TPBDJointContainerSolver")
 		void AddBodies(const TArrayView<int32>& ConstraintIndices, FSolverBodyContainer& SolverBodyContainer) {}
-		UE_DEPRECATED(5.2, "Joint Solver API moved to FPBDJointContainerSolver")
+		UE_DEPRECATED(5.2, "Joint Solver API moved to TPBDJointContainerSolver")
 		void GatherInput(const TArrayView<int32>& ConstraintIndices, const FReal Dt) {}
-		UE_DEPRECATED(5.2, "Joint Solver API moved to FPBDJointContainerSolver")
+		UE_DEPRECATED(5.2, "Joint Solver API moved to TPBDJointContainerSolver")
 		void ScatterOutput(const TArrayView<int32>& ConstraintIndices, const FReal Dt) {}
-		UE_DEPRECATED(5.2, "Joint Solver API moved to FPBDJointContainerSolver")
+		UE_DEPRECATED(5.2, "Joint Solver API moved to TPBDJointContainerSolver")
 		void ApplyPositionConstraints(const TArrayView<int32>& ConstraintIndices, const FReal Dt, const int32 It, const int32 NumIts) {}
-		UE_DEPRECATED(5.2, "Joint Solver API moved to FPBDJointContainerSolver")
+		UE_DEPRECATED(5.2, "Joint Solver API moved to TPBDJointContainerSolver")
 		void ApplyVelocityConstraints(const TArrayView<int32>& ConstraintIndices, const FReal Dt, const int32 It, const int32 NumIts) {}
-		UE_DEPRECATED(5.2, "Joint Solver API moved to FPBDJointContainerSolver")
+		UE_DEPRECATED(5.2, "Joint Solver API moved to TPBDJointContainerSolver")
 		void ApplyProjectionConstraints(const TArrayView<int32>& ConstraintIndices, const FReal Dt, const int32 It, const int32 NumIts) {}
-		UE_DEPRECATED(5.2, "Joint Solver API moved to FPBDJointContainerSolver")
+		UE_DEPRECATED(5.2, "Joint Solver API moved to TPBDJointContainerSolver")
 		void GatherInput(const int32 ConstraintIndex, const FReal Dt) {}
-		UE_DEPRECATED(5.2, "Joint Solver API moved to FPBDJointContainerSolver")
+		UE_DEPRECATED(5.2, "Joint Solver API moved to TPBDJointContainerSolver")
 		void ScatterOutput(const int32 ConstraintIndex, const FReal Dt) {}
 
+		CHAOS_API bool GetUseLinearSolver() { return bUseLinearSolver; }
+
+		CHAOS_API void SetUseLinearSolver(const bool bUseLinearSolverIn) { bUseLinearSolver = bUseLinearSolverIn; }
+
 	private:
 		friend class FPBDJointConstraintHandle;
 
+		friend struct FPBDJointCombinedConstraints;
+
 		CHAOS_API void GetConstrainedParticleIndices(const int32 ConstraintIndex, int32& Index0, int32& Index1) const;
 		CHAOS_API void CalculateConstraintSpace(int32 ConstraintIndex, FVec3& OutX0, FMatrix33& OutR0, FVec3& OutX1, FMatrix33& OutR1) const;
 		
@@ -450,6 +459,23 @@ namespace Chaos
 		bool bJointsDirty;
 
 		FJointBreakCallback BreakCallback;
+
+		bool bUseLinearSolver;
+
+	};
+
+	struct FPBDJointCombinedConstraints
+	{
+		CHAOS_API FPBDJointCombinedConstraints();
+		void ConvertLinearConstraintToNonLinear(const int32 LinearConstraintIndex);
+		void ConvertNonLinearConstraintToLinear(const int32 NonLinearConstraintIndex);
+		CHAOS_API void MergeLinearConstraintsIntoNonLinearConstraints();
+		CHAOS_API void MergeNonLinearConstraintsIntoLinearConstraints();
+		int32 NumConstraints() { return LinearConstraints.NumConstraints() + NonLinearConstraints.NumConstraints(); }
+
+		FPBDJointConstraints LinearConstraints;
+		FPBDJointConstraints NonLinearConstraints;
 	};
+	
 
 }
diff --git a/Engine/Source/Runtime/Experimental/Chaos/Public/Chaos/PBDRigidsEvolutionGBF.h b/Engine/Source/Runtime/Experimental/Chaos/Public/Chaos/PBDRigidsEvolutionGBF.h
--- a/Engine/Source/Runtime/Experimental/Chaos/Public/Chaos/PBDRigidsEvolutionGBF.h
+++ b/Engine/Source/Runtime/Experimental/Chaos/Public/Chaos/PBDRigidsEvolutionGBF.h
@@ -54,7 +54,7 @@ namespace Chaos
 		using FCollisionConstraints = FPBDCollisionConstraints;
 		using FCollisionDetector = FSpatialAccelerationCollisionDetector;
 		using FExternalForces = FPerParticleExternalForces;
-		using FJointConstraints = FPBDJointConstraints;
+		using FJointConstraints = FPBDJointCombinedConstraints;
 
 		// Default settings for FChaosSolverConfiguration
 		static constexpr int32 DefaultNumPositionIterations = 8;
@@ -290,7 +290,7 @@ namespace Chaos
 
 		FRigidClustering Clustering;
 
-		FPBDJointConstraints JointConstraints;
+		FPBDJointCombinedConstraints JointConstraints;
 		FPBDSuspensionConstraints SuspensionConstraints;
 		FCharacterGroundConstraintContainer CharacterGroundConstraints;
 
diff --git a/Engine/Source/Runtime/Experimental/Chaos/Public/ChaosVisualDebugger/ChaosVDOptionalDataChannel.h b/Engine/Source/Runtime/Experimental/Chaos/Public/ChaosVisualDebugger/ChaosVDOptionalDataChannel.h
--- a/Engine/Source/Runtime/Experimental/Chaos/Public/ChaosVisualDebugger/ChaosVDOptionalDataChannel.h
+++ b/Engine/Source/Runtime/Experimental/Chaos/Public/ChaosVisualDebugger/ChaosVDOptionalDataChannel.h
@@ -151,7 +151,8 @@ CVD_DECLARE_OPTIONAL_DATA_CHANNEL(EndOfEvolutionCollisionConstraints);
 CVD_DECLARE_OPTIONAL_DATA_CHANNEL(PostConstraintSolve);						// Particle positions corrected by the constraint solver (collisions, joints, etc.), but before user callbacks or destruction handling.
 CVD_DECLARE_OPTIONAL_DATA_CHANNEL(EvolutionEnd);							// Final particle positions, including destruction and user callbacks.
 CVD_DECLARE_OPTIONAL_DATA_CHANNEL(SceneQueries);
-CVD_DECLARE_OPTIONAL_DATA_CHANNEL(JointConstraints);
+CVD_DECLARE_OPTIONAL_DATA_CHANNEL(JointLinearConstraints);
+CVD_DECLARE_OPTIONAL_DATA_CHANNEL(JointNonLinearConstraints);
 CVD_DECLARE_OPTIONAL_DATA_CHANNEL(CharacterGroundConstraints);
 CVD_DECLARE_OPTIONAL_DATA_CHANNEL(AccelerationStructures);
 
diff --git a/Engine/Source/Runtime/Experimental/Chaos/Public/PBDRigidsSolver.h b/Engine/Source/Runtime/Experimental/Chaos/Public/PBDRigidsSolver.h
--- a/Engine/Source/Runtime/Experimental/Chaos/Public/PBDRigidsSolver.h
+++ b/Engine/Source/Runtime/Experimental/Chaos/Public/PBDRigidsSolver.h
@@ -113,7 +113,7 @@ namespace Chaos
 		typedef FPBDRigidDynamicSpringConstraints FRigidDynamicSpringConstraints;
 		typedef FPBDPositionConstraints FPositionConstraints;
 
-		using FJointConstraints = FPBDJointConstraints;
+		using FJointConstraints = FPBDJointCombinedConstraints;
 		//
 		// Execution API
 		//
