Fixes the sign of twist calculated from FPBDJointUtilities::GetSwingTwistAngles

diff --git a/Engine/Source/Runtime/Experimental/Chaos/Private/Chaos/PBDJointConstraintUtilities.cpp b/Engine/Source/Runtime/Experimental/Chaos/Private/Chaos/PBDJointConstraintUtilities.cpp
--- a/Engine/Source/Runtime/Experimental/Chaos/Private/Chaos/PBDJointConstraintUtilities.cpp
+++ b/Engine/Source/Runtime/Experimental/Chaos/Private/Chaos/PBDJointConstraintUtilities.cpp
@@ -167,7 +167,7 @@ namespace Chaos
 	{
 		FRotation3 R01Twist, R01Swing;
 		FPBDJointUtilities::DecomposeSwingTwistLocal(R0, R1, R01Swing, R01Twist);
-		TwistAngle = R01Twist.GetAngle();
+		TwistAngle = GetTwistAngle(R01Twist);
 		// @todo(ccaulfield): makes assumptions about Swing1 and Swing2 axes - fix this
 		Swing1Angle = GetSwingAngleZ(R01Swing);
 		Swing2Angle = GetSwingAngleY(R01Swing);
