Enable aerofoil to update the vehicle animation instance.

#rnx

diff --git a/Engine/Source/Runtime/Experimental/ChaosVehicles/ChaosVehiclesCore/Private/SimModule/AerofoilModule.cpp b/Engine/Source/Runtime/Experimental/ChaosVehicles/ChaosVehiclesCore/Private/SimModule/AerofoilModule.cpp
--- a/Engine/Source/Runtime/Experimental/ChaosVehicles/ChaosVehiclesCore/Private/SimModule/AerofoilModule.cpp
+++ b/Engine/Source/Runtime/Experimental/ChaosVehicles/ChaosVehiclesCore/Private/SimModule/AerofoilModule.cpp
@@ -158,6 +158,45 @@ namespace Chaos
 
 	}
 
+
+	void FAerofoilOutputData::FillOutputState(const ISimulationModuleBase* SimModule)
+	{
+		check(SimModule->IsSimType<class FAerofoilSimModule>());
+
+		FSimOutputData::FillOutputState(SimModule);
+
+		if (const FAerofoilSimModule* Sim = static_cast<const FAerofoilSimModule*>(SimModule))
+		{
+			AnimFlags = EAnimationFlags::AnimateRotation;
+
+			switch (Sim->Setup().Type)
+			{
+			case EAerofoil::Wing:
+			case EAerofoil::Elevator:
+				AnimationRotOffset.Pitch = Sim->ControlSurfaceAngle * Sim->Setup().AnimationMagnitudeMultiplier;
+				break;
+
+			case EAerofoil::Rudder:
+				AnimationRotOffset.Yaw = Sim->ControlSurfaceAngle * Sim->Setup().AnimationMagnitudeMultiplier;
+				break;
+			}
+
+		}
+	}
+
+	void FAerofoilOutputData::Lerp(const FSimOutputData& InCurrent, const FSimOutputData& InNext, float Alpha)
+	{
+		FSimOutputData::Lerp(InCurrent, InNext, Alpha);
+	}
+
+#if !(UE_BUILD_SHIPPING || UE_BUILD_TEST)
+	FString FAerofoilOutputData::ToString()
+	{
+		return  FString::Printf(TEXT("%s")
+			, *DebugString);
+	}
+#endif
+
 } // namespace Chaos
 
 #if VEHICLE_DEBUGGING_ENABLED
diff --git a/Engine/Source/Runtime/Experimental/ChaosVehicles/ChaosVehiclesCore/Public/SimModule/AerofoilModule.h b/Engine/Source/Runtime/Experimental/ChaosVehicles/ChaosVehiclesCore/Public/SimModule/AerofoilModule.h
--- a/Engine/Source/Runtime/Experimental/ChaosVehicles/ChaosVehiclesCore/Public/SimModule/AerofoilModule.h
+++ b/Engine/Source/Runtime/Experimental/ChaosVehicles/ChaosVehiclesCore/Public/SimModule/AerofoilModule.h
@@ -18,6 +18,21 @@ namespace Chaos
 		Elevator
 	};
 
+	struct CHAOSVEHICLESCORE_API FAerofoilOutputData
+		: public FSimOutputData
+		, public Chaos::TSimulationModuleTypeable<class FAerofoilSimModule, FAerofoilOutputData>
+	{
+		virtual FSimOutputData* MakeNewData() override { return FAerofoilOutputData::MakeNew(); }
+		static FSimOutputData* MakeNew() { return new FAerofoilOutputData(); }
+
+		virtual void FillOutputState(const ISimulationModuleBase* SimModule) override;
+		virtual void Lerp(const FSimOutputData& InCurrent, const FSimOutputData& InNext, float Alpha) override;
+
+#if !(UE_BUILD_SHIPPING || UE_BUILD_TEST)
+		virtual FString ToString() override;
+#endif
+	};
+
 	struct CHAOSVEHICLESCORE_API FAerofoilSettings
 	{
 		FAerofoilSettings()
@@ -52,6 +67,8 @@ namespace Chaos
 
 	class CHAOSVEHICLESCORE_API FAerofoilSimModule : public ISimulationModuleBase, public TSimModuleSettings<FAerofoilSettings>, public TSimulationModuleTypeable<FAerofoilSimModule>
 	{
+		friend FAerofoilOutputData;
+
 	public:
 		DEFINE_CHAOSSIMTYPENAME(FAerofoilSimModule);
 		FAerofoilSimModule(const FAerofoilSettings& Settings);
@@ -60,6 +77,11 @@ namespace Chaos
 
 		virtual TSharedPtr<FModuleNetData> GenerateNetData(const int32 NodeArrayIndex) const override { return nullptr; }
 
+		virtual FSimOutputData* GenerateOutputData() const override
+		{
+			return FAerofoilOutputData::MakeNew();
+		}
+
 		virtual const FString GetDebugName() const { return TEXT("Aerofoil"); }
 
 		virtual bool IsBehaviourType(eSimModuleTypeFlags InType) const override { return (InType & Velocity); }
