Dataflow Editor: fix for Skeletal Mesh being cleared from the Preview scene when changing Animations, and vice-versa. Removes the Skeleton property from UDataflowSkeletalContent

diff --git a/Engine/Plugins/ChaosClothAsset/Source/ChaosClothAssetEngine/Private/ChaosClothAsset/ClothAsset.cpp b/Engine/Plugins/ChaosClothAsset/Source/ChaosClothAssetEngine/Private/ChaosClothAsset/ClothAsset.cpp
--- a/Engine/Plugins/ChaosClothAsset/Source/ChaosClothAssetEngine/Private/ChaosClothAsset/ClothAsset.cpp
+++ b/Engine/Plugins/ChaosClothAsset/Source/ChaosClothAssetEngine/Private/ChaosClothAsset/ClothAsset.cpp
@@ -1064,7 +1064,6 @@ void UChaosClothAsset::WriteDataflowContent(const TObjectPtr<UDataflowBaseConten
 	{
 		SkeletalContent->SetDataflowAsset(DataflowAsset);
 		SkeletalContent->SetDataflowTerminal(DataflowTerminal);
-		SkeletalContent->SetSkeleton(nullptr);
 
 #if WITH_EDITORONLY_DATA
 		SkeletalContent->SetAnimationAsset(GetPreviewSceneAnimation());
diff --git a/Engine/Plugins/Experimental/ChaosFlesh/Source/ChaosFleshEngine/Private/ChaosFlesh/FleshAsset.cpp b/Engine/Plugins/Experimental/ChaosFlesh/Source/ChaosFleshEngine/Private/ChaosFlesh/FleshAsset.cpp
--- a/Engine/Plugins/Experimental/ChaosFlesh/Source/ChaosFleshEngine/Private/ChaosFlesh/FleshAsset.cpp
+++ b/Engine/Plugins/Experimental/ChaosFlesh/Source/ChaosFleshEngine/Private/ChaosFlesh/FleshAsset.cpp
@@ -94,7 +94,6 @@ void UFleshAsset::WriteDataflowContent(const TObjectPtr<UDataflowBaseContent>& D
 		SkeletalContent->SetDataflowTerminal(DataflowTerminal);
 		
 		SkeletalContent->SetSkeletalMesh(SkeletalMesh, true);
-		SkeletalContent->SetSkeleton(Skeleton);
 
 #if WITH_EDITORONLY_DATA
 		SkeletalContent->SetAnimationAsset(PreviewAnimationAsset.Get());
diff --git a/Engine/Plugins/Experimental/Dataflow/Source/DataflowEditor/Private/Dataflow/DataflowEditor.cpp b/Engine/Plugins/Experimental/Dataflow/Source/DataflowEditor/Private/Dataflow/DataflowEditor.cpp
--- a/Engine/Plugins/Experimental/Dataflow/Source/DataflowEditor/Private/Dataflow/DataflowEditor.cpp
+++ b/Engine/Plugins/Experimental/Dataflow/Source/DataflowEditor/Private/Dataflow/DataflowEditor.cpp
@@ -79,17 +79,6 @@ void UDataflowEditor::Initialize(const TArray<TObjectPtr<UObject>>& InObjects, c
 
 void UDataflowEditor::UpdateEditorContent()
 {
-	if(const TObjectPtr<UDataflowSkeletalContent> SkeletalContent = Cast<UDataflowSkeletalContent>(EditorContent))
-	{
-		// Add a skeleton mesh in case we don't have one
-		if(!SkeletalContent->GetSkeleton())
-		{
-			const FName SkeletonName = MakeUniqueObjectName(SkeletalContent->GetDataflowAsset(), UDataflow::StaticClass(), FName("USkeleton"));
-			USkeleton* Skeleton = NewObject<USkeleton>(SkeletalContent->GetDataflowAsset(), SkeletonName);
-
-			SkeletalContent->SetSkeleton(Skeleton);
-		}
-	}
 	if(EditorContent && EditorContent->GetDataflowAsset())
 	{
 		EditorContent->GetDataflowAsset()->Schema = UDataflowSchema::StaticClass();
diff --git a/Engine/Source/Runtime/Experimental/Dataflow/Engine/Private/Dataflow/DataflowContent.cpp b/Engine/Source/Runtime/Experimental/Dataflow/Engine/Private/Dataflow/DataflowContent.cpp
--- a/Engine/Source/Runtime/Experimental/Dataflow/Engine/Private/Dataflow/DataflowContent.cpp
+++ b/Engine/Source/Runtime/Experimental/Dataflow/Engine/Private/Dataflow/DataflowContent.cpp
@@ -427,13 +427,6 @@ UDataflowSkeletalContent::UDataflowSkeletalContent() : Super()
 void UDataflowSkeletalContent::SetSkeletalMesh(const TObjectPtr<USkeletalMesh>& SkeletalMeshAsset, const bool bHideAsset)
 {
 	SkeletalMesh = SkeletalMeshAsset;
-	if(SkeletalMesh)
-	{
-		if(SkeletalMesh && (SkeletalMesh->GetSkeleton() != Skeleton))
-		{
-			SetSkeleton(SkeletalMesh->GetSkeleton());
-		}
-	}
 	bHideSkeletalMesh = bHideAsset;
 	SetConstructionDirty(true);
 	SetSimulationDirty(true);
@@ -442,30 +435,11 @@ void UDataflowSkeletalContent::SetSkeletalMesh(const TObjectPtr<USkeletalMesh>&
 void UDataflowSkeletalContent::SetAnimationAsset(const TObjectPtr<UAnimationAsset>& SkeletalAnimationAsset, const bool bHideAsset)
 {
 	AnimationAsset = SkeletalAnimationAsset;
-	if(AnimationAsset && (AnimationAsset->GetSkeleton()) != Skeleton)
-	{
-		SetSkeleton(AnimationAsset->GetSkeleton());
-	}
 	bHideAnimationAsset = bHideAsset;
 	SetConstructionDirty(true);
 	SetSimulationDirty(true);
 }
 
-void UDataflowSkeletalContent::SetSkeleton(const TObjectPtr<USkeleton>& SkeletonAsset)
-{
-	Skeleton = SkeletonAsset;
-	if(SkeletalMesh && (SkeletalMesh->GetSkeleton() != Skeleton))
-	{
-		SetSkeletalMesh(nullptr);
-	}
-	if(AnimationAsset && (AnimationAsset->GetSkeleton() != Skeleton))
-	{
-		SetAnimationAsset(nullptr);
-	}
-	SetConstructionDirty(true);
-	SetSimulationDirty(true);
-}
-
 #if WITH_EDITOR
 
 void UDataflowSkeletalContent::PostEditChangeProperty(FPropertyChangedEvent& PropertyChangedEvent) 
@@ -480,10 +454,6 @@ void UDataflowSkeletalContent::PostEditChangeProperty(FPropertyChangedEvent& Pro
 	{
 		SetAnimationAsset(AnimationAsset);
 	}
-	if (PropertyName == GET_MEMBER_NAME_CHECKED(UDataflowSkeletalContent, Skeleton))
-	{
-		SetSkeleton(Skeleton);
-	}
 	Super::PostEditChangeProperty(PropertyChangedEvent);
 }
 
@@ -517,7 +487,6 @@ void UDataflowSkeletalContent::AddContentObjects(FReferenceCollector& Collector)
 void UDataflowSkeletalContent::AddReferencedObjects(UObject* InThis, FReferenceCollector& Collector)
 {
 	UDataflowSkeletalContent* This = CastChecked<UDataflowSkeletalContent>(InThis);
-	Collector.AddReferencedObject(This->Skeleton);
 	Collector.AddReferencedObject(This->SkeletalMesh);
 	Collector.AddReferencedObject(This->AnimationAsset);
 	Super::AddReferencedObjects(InThis, Collector);
diff --git a/Engine/Source/Runtime/Experimental/Dataflow/Engine/Public/Dataflow/DataflowContent.h b/Engine/Source/Runtime/Experimental/Dataflow/Engine/Public/Dataflow/DataflowContent.h
--- a/Engine/Source/Runtime/Experimental/Dataflow/Engine/Public/Dataflow/DataflowContent.h
+++ b/Engine/Source/Runtime/Experimental/Dataflow/Engine/Public/Dataflow/DataflowContent.h
@@ -224,10 +224,6 @@ public:
 	void SetSkeletalMesh(const TObjectPtr<USkeletalMesh>& InMesh, const bool bHideAsset = false);
 	const TObjectPtr<USkeletalMesh>& GetSkeletalMesh() const { return SkeletalMesh; }
 
-	/** Data flow skeleton accessors */
-	void SetSkeleton(const TObjectPtr<USkeleton>& InSkeleton);
-	const TObjectPtr<USkeleton>& GetSkeleton() const { return Skeleton; }
-
 	/** Data flow animation asset accessors */
 	void SetAnimationAsset(const TObjectPtr<UAnimationAsset>& InAnimation, const bool bHideAsset = false);
 	const TObjectPtr<UAnimationAsset>& GetAnimationAsset() const { return AnimationAsset; }
@@ -243,11 +239,7 @@ protected:
 	/** Data flow skeletal mesh*/
 	UPROPERTY(EditAnywhere, Category = "Preview", Transient, SkipSerialization)
 	TObjectPtr<USkeletalMesh> SkeletalMesh = nullptr;
-	
-	/** Data flow skeleton */
-	UPROPERTY(Transient, SkipSerialization)
-	TObjectPtr<USkeleton> Skeleton = nullptr;
-	
+
 	/** Animation asset to be used to preview simulation */
 	UPROPERTY(EditAnywhere, Category = "Preview", Transient, SkipSerialization)
 	TObjectPtr<UAnimationAsset> AnimationAsset = nullptr;
