NetObjectFactory now holds TWeakPtrs instead of raw pointers

#rnx

diff --git a/Engine/Source/Runtime/Engine/Private/UnrealEngine.cpp b/Engine/Source/Runtime/Engine/Private/UnrealEngine.cpp
--- a/Engine/Source/Runtime/Engine/Private/UnrealEngine.cpp
+++ b/Engine/Source/Runtime/Engine/Private/UnrealEngine.cpp
@@ -357,7 +357,11 @@ void FEngineModule::StartupModule()
 	UE::Anim::FSkeletonRemappingRegistry::Init();
 
 #if UE_WITH_IRIS
-	UE::Net::InitEngineNetObjectFactories();
+	FDelayedAutoRegisterHelper(EDelayedRegisterRunPhase::ObjectSystemReady,
+	[]
+	{
+		UE::Net::InitEngineNetObjectFactories();
+	});
 #endif
 
 	IUniversalObjectLocatorModule& UolModule = FModuleManager::Get().LoadModuleChecked<IUniversalObjectLocatorModule>("UniversalObjectLocator");
diff --git a/Engine/Source/Runtime/Experimental/Iris/Core/Private/Iris/ReplicationSystem/NetObjectFactoryRegistry.cpp b/Engine/Source/Runtime/Experimental/Iris/Core/Private/Iris/ReplicationSystem/NetObjectFactoryRegistry.cpp
--- a/Engine/Source/Runtime/Experimental/Iris/Core/Private/Iris/ReplicationSystem/NetObjectFactoryRegistry.cpp
+++ b/Engine/Source/Runtime/Experimental/Iris/Core/Private/Iris/ReplicationSystem/NetObjectFactoryRegistry.cpp
@@ -54,7 +54,7 @@ void FNetObjectFactoryRegistry::RegisterFactory(UClass* FactoryClass, FName Fact
 	{
 		if (FactoryData.Name == FactoryName)
 		{
-			UE_LOG(LogIris, Error, TEXT("FNetObjectFactoryRegistry::RegisterFactory cannot register factory: %s with name: %s. This name is already used by factory: %s id: %u"), *GetNameSafe(FactoryClass), *FactoryName.ToString(), *GetNameSafe(FactoryData.NetFactoryClass), FactoryData.Id);
+			UE_LOG(LogIris, Error, TEXT("FNetObjectFactoryRegistry::RegisterFactory cannot register factory: %s with name: %s. This name is already used by factory: %s id: %u"), *GetNameSafe(FactoryClass), *FactoryName.ToString(), *GetNameSafe(FactoryData.NetFactoryClass.Get()), FactoryData.Id);
 			ensure(false);
 			return;
 		}
@@ -75,7 +75,7 @@ void FNetObjectFactoryRegistry::RegisterFactory(UClass* FactoryClass, FName Fact
 		});
 	NetFactories[Index].Id = IntCastChecked<FNetObjectFactoryId>(Index);
 
-	UE_LOG(LogIris, Log, TEXT("FNetObjectFactoryRegistry::UnregisterFactory is unregistering factory: %s name: %s id: %u"), *GetNameSafe(NetFactories[Index].NetFactoryClass), *NetFactories[Index].Name.ToString(), NetFactories[Index].Id);
+	UE_LOG(LogIris, Log, TEXT("FNetObjectFactoryRegistry::UnregisterFactory is unregistering factory: %s name: %s id: %u"), *GetNameSafe(NetFactories[Index].NetFactoryClass.Get()), *NetFactories[Index].Name.ToString(), NetFactories[Index].Id);
 #endif
 }
 
@@ -97,7 +97,7 @@ void FNetObjectFactoryRegistry::UnregisterFactory(FName FactoryName)
 		const FFactoryData& Data = NetFactories[Index];
 		if (Data.Name == FactoryName)
 		{
-			UE_LOG(LogIris, Log, TEXT("FNetObjectFactoryRegistry::UnregisterFactory is unregistering factory: %s name: %s id: %u"), *GetNameSafe(Data.NetFactoryClass), *Data.Name.ToString(), Data.Id);
+			UE_LOG(LogIris, Log, TEXT("FNetObjectFactoryRegistry::UnregisterFactory is unregistering factory: %s name: %s id: %u"), *GetNameSafe(Data.NetFactoryClass.Get()), *Data.Name.ToString(), Data.Id);
 			NetFactories.RemoveAt(Index, EAllowShrinking::No);
 			return;
 		}
diff --git a/Engine/Source/Runtime/Experimental/Iris/Core/Private/Iris/ReplicationSystem/ObjectReplicationBridge.cpp b/Engine/Source/Runtime/Experimental/Iris/Core/Private/Iris/ReplicationSystem/ObjectReplicationBridge.cpp
--- a/Engine/Source/Runtime/Experimental/Iris/Core/Private/Iris/ReplicationSystem/ObjectReplicationBridge.cpp
+++ b/Engine/Source/Runtime/Experimental/Iris/Core/Private/Iris/ReplicationSystem/ObjectReplicationBridge.cpp
@@ -198,9 +198,9 @@ void UObjectReplicationBridge::InitNetObjectFactories()
 	{
 		const FNetObjectFactoryRegistry::FFactoryData& Data = RegisteredFactories[Index];
 		check(Data.Id == (FNetObjectFactoryId)Index);
-		check(Data.NetFactoryClass);
+		check(Data.NetFactoryClass.Get());
 
-		UNetObjectFactory* Factory = NewObject<UNetObjectFactory>(this, Data.NetFactoryClass);
+		UNetObjectFactory* Factory = NewObject<UNetObjectFactory>(this, Data.NetFactoryClass.Get());
 		NetObjectFactories.Add(Factory);
 
 		Factory->Init(Data.Id, this);
diff --git a/Engine/Source/Runtime/Experimental/Iris/Core/Public/Iris/ReplicationSystem/NetObjectFactoryRegistry.h b/Engine/Source/Runtime/Experimental/Iris/Core/Public/Iris/ReplicationSystem/NetObjectFactoryRegistry.h
--- a/Engine/Source/Runtime/Experimental/Iris/Core/Public/Iris/ReplicationSystem/NetObjectFactoryRegistry.h
+++ b/Engine/Source/Runtime/Experimental/Iris/Core/Public/Iris/ReplicationSystem/NetObjectFactoryRegistry.h
@@ -3,7 +3,9 @@
 #pragma once
 
 #include "Containers/ArrayView.h"
+#include "UObject/Class.h"
 #include "UObject/NameTypes.h"
+#include "UObject/WeakObjectPtrTemplates.h"
 
 // Define this to increase the maximum amount of NetObjectFactories that can exist
 #ifndef UE_IRIS_MAX_NETOBJECT_FACTORIES
@@ -48,7 +50,7 @@ public:
 		/** Name associated with this factory class */
 		FName Name;
 		/** Class representing a concrete UNetObjectFactory */
-		UClass* NetFactoryClass = nullptr;
+		TWeakObjectPtr<UClass> NetFactoryClass;
 		/** FactoryId assigned to this factory */
 		FNetObjectFactoryId Id = InvalidNetObjectFactoryId;
 	};
