Dataflow - Fixed crashes in Geometry Collection nodes caused by cached objects to be garbage collected.

#rnx

diff --git a/Engine/Source/Runtime/Experimental/Dataflow/Core/Private/Dataflow/DataflowNodeParameters.cpp b/Engine/Source/Runtime/Experimental/Dataflow/Core/Private/Dataflow/DataflowNodeParameters.cpp
--- a/Engine/Source/Runtime/Experimental/Dataflow/Core/Private/Dataflow/DataflowNodeParameters.cpp
+++ b/Engine/Source/Runtime/Experimental/Dataflow/Core/Private/Dataflow/DataflowNodeParameters.cpp
@@ -148,6 +148,8 @@ namespace UE::Dataflow
 			for (TPair<FContextCacheKey, TUniquePtr<FContextCacheElementBase>>& Elem : Pairs)
 			{
 				// note : we only serialize typed cache element and ignore the reference ones ( since they don't hold data per say )
+				// Also UObject pointers aren't serialized, as there are no ways to differentiate the objects owned by the cache 
+				// from the ones own by any other owners for now.
 				if (Elem.Value && Elem.Value->Property && Elem.Value->Type == FContextCacheElementBase::EType::CacheElementTyped)
 				{
 					FProperty* Property = (FProperty*)Elem.Value->Property;
diff --git a/Engine/Source/Runtime/Experimental/Dataflow/Core/Public/Dataflow/DataflowContextCachingFactory.h b/Engine/Source/Runtime/Experimental/Dataflow/Core/Public/Dataflow/DataflowContextCachingFactory.h
--- a/Engine/Source/Runtime/Experimental/Dataflow/Core/Public/Dataflow/DataflowContextCachingFactory.h
+++ b/Engine/Source/Runtime/Experimental/Dataflow/Core/Public/Dataflow/DataflowContextCachingFactory.h
@@ -54,7 +54,7 @@ namespace UE::Dataflow
 			if (InElement && InElement->GetType() == FContextCacheElementBase::EType::CacheElementTyped)
 			{
 				// it is assumed the type requested matches the cache entry 
-				const FContextCacheElement<T>* TypedElement = static_cast<const FContextCacheElement<T>*>(InElement);
+				const TContextCacheElement<T>* TypedElement = static_cast<const TContextCacheElement<T>*>(InElement);
 				return TypedElement->GetDataDirect();
 			}
 			return Default;
@@ -63,7 +63,7 @@ namespace UE::Dataflow
 		template<class T>
 		static FContextCacheElementBase* NewTypedElement(T&& Data)
 		{
-			return new FContextCacheElement<T>(FGuid(), (FProperty*)nullptr, MoveTemp(Data), (uint32)0, FTimestamp::Invalid);
+			return new TContextCacheElement<T>(FGuid(), (FProperty*)nullptr, MoveTemp(Data), (uint32)0, FTimestamp::Invalid);
 		}
 
 
diff --git a/Engine/Source/Runtime/Experimental/Dataflow/Core/Public/Dataflow/DataflowNodeParameters.h b/Engine/Source/Runtime/Experimental/Dataflow/Core/Public/Dataflow/DataflowNodeParameters.h
--- a/Engine/Source/Runtime/Experimental/Dataflow/Core/Public/Dataflow/DataflowNodeParameters.h
+++ b/Engine/Source/Runtime/Experimental/Dataflow/Core/Public/Dataflow/DataflowNodeParameters.h
@@ -22,6 +22,14 @@ namespace UE::Dataflow
 
 	typedef uint32 FContextCacheKey;
 
+	/** Trait used to select the UObject* or TObjectPtr cache element path code. */
+	template <typename T>
+	struct TIsUObjectPtrElement
+	{
+		typedef typename TDecay<T>::Type Type;
+		static constexpr bool Value = (std::is_pointer_v<Type> && std::is_convertible_v<Type, const UObjectBase*>) || TIsTObjectPtr_V<Type>;
+	};
+
 	struct FTimestamp
 	{
 		typedef uint64 Type;
@@ -59,7 +67,8 @@ namespace UE::Dataflow
 		{
 			CacheElementTyped,
 			CacheElementReference,
-			CacheElementNull
+			CacheElementNull,
+			CacheElementUObject
 		};
 
 		FContextCacheElementBase(EType CacheElementType, FGuid InNodeGuid = FGuid(), const FProperty* InProperty = nullptr, uint32 InNodeHash = 0, FTimestamp InTimestamp = FTimestamp::Invalid)
@@ -75,7 +84,7 @@ namespace UE::Dataflow
 		virtual TUniquePtr<FContextCacheElementBase> CreateReference(FContextCacheKey InReferenceDataKey) const = 0;
 
 		template<typename T>
-		const T& GetTypedData(FContext& Context, const FProperty* PropertyIn, const T& Default) const;
+		inline const T& GetTypedData(FContext& Context, const FProperty* PropertyIn, const T& Default) const;
 		
 		EType GetType() const {	return Type; }
 
@@ -102,19 +111,19 @@ namespace UE::Dataflow
 	//--------------------------------------------------------------------
 	// Value storing context cache entry - strongly typed
 	//--------------------------------------------------------------------
-	template<class T>
-	struct FContextCacheElement : public FContextCacheElementBase 
+	template<class T UE_REQUIRES(!TIsUObjectPtrElement<T>::Value)>
+	struct TContextCacheElement : public FContextCacheElementBase 
 	{
-		FContextCacheElement(FGuid InNodeGuid, const FProperty* InProperty, T&& InData, uint32 InNodeHash, FTimestamp Timestamp)
+		TContextCacheElement(FGuid InNodeGuid, const FProperty* InProperty, T&& InData, uint32 InNodeHash, FTimestamp Timestamp)
 			: FContextCacheElementBase(EType::CacheElementTyped, InNodeGuid, InProperty, InNodeHash, Timestamp)
 			, Data(Forward<T>(InData))
 		{}
 		
-		const T& GetData(FContext& Context, const FProperty* PropertyIn, const T& Default) const;
+		inline const T& GetData(FContext& Context, const FProperty* PropertyIn, const T& Default) const;
 
 		const T& GetDataDirect() const { return Data; }
 
-		virtual TUniquePtr<FContextCacheElementBase> CreateReference(FContextCacheKey InReferenceDataKey) const override;
+		inline virtual TUniquePtr<FContextCacheElementBase> CreateReference(FContextCacheKey InReferenceDataKey) const override;
 
 	private:
 		typedef typename TDecay<T>::Type FDataType;  // Using universal references here means T could be either const& or an rvalue reference
@@ -125,16 +134,16 @@ namespace UE::Dataflow
 	// Reference to another context cache entry 
 	//--------------------------------------------------------------------
 	template<class T>
-	struct FContextCacheElementReference : public FContextCacheElementBase
+	struct TContextCacheElementReference : public FContextCacheElementBase
 	{
-		FContextCacheElementReference(FGuid InNodeGuid, const FProperty* InProperty, FContextCacheKey InDataKey, uint32 InNodeHash, FTimestamp Timestamp)
+		TContextCacheElementReference(FGuid InNodeGuid, const FProperty* InProperty, FContextCacheKey InDataKey, uint32 InNodeHash, FTimestamp Timestamp)
 			: FContextCacheElementBase(EType::CacheElementReference, InNodeGuid, InProperty, InNodeHash, Timestamp)
 			, DataKey(InDataKey)
 		{}
 
-		const T& GetData(FContext& Context, const FProperty* PropertyIn, const T& Default) const;
+		inline const T& GetData(FContext& Context, const FProperty* PropertyIn, const T& Default) const;
 
-		virtual TUniquePtr<FContextCacheElementBase> CreateReference(FContextCacheKey InReferenceDataKey) const override;
+		inline virtual TUniquePtr<FContextCacheElementBase> CreateReference(FContextCacheKey InReferenceDataKey) const override;
 
 	private:
 		const FContextCacheKey DataKey; // this is a key to another cache element
@@ -149,7 +158,32 @@ namespace UE::Dataflow
 			: FContextCacheElementBase(EType::CacheElementNull, InNodeGuid, InProperty, InNodeHash, Timestamp)
 		{}
 
-		virtual TUniquePtr<FContextCacheElementBase> CreateReference(FContextCacheKey InReferenceDataKey) const override;
+		inline virtual TUniquePtr<FContextCacheElementBase> CreateReference(FContextCacheKey InReferenceDataKey) const override;
+	};
+
+	//--------------------------------------------------------------------
+	// UObject cache element, prevents the object from being garbage collected while in the cache
+	//--------------------------------------------------------------------
+	template<class T UE_REQUIRES(TIsUObjectPtrElement<T>::Value)>
+	struct TContextCacheElementUObject : public FContextCacheElementBase, public FGCObject
+	{
+		TContextCacheElementUObject(FGuid InNodeGuid, const FProperty* InProperty, T&& InObject, uint32 InNodeHash, FTimestamp Timestamp)
+			: FContextCacheElementBase(EType::CacheElementUObject, InNodeGuid, InProperty, InNodeHash, Timestamp)
+			, Object(InObject)
+		{}
+
+		const T& GetData(FContext& /*Context*/, const FProperty* /*PropertyIn*/, const T& /*Default*/) const { return Object; }
+
+		inline virtual TUniquePtr<FContextCacheElementBase> CreateReference(FContextCacheKey InReferenceDataKey) const override;
+
+		//~ Begin FGCObject interface
+		virtual void AddReferencedObjects(FReferenceCollector& Collector) override { Collector.AddReferencedObject(Object); }
+		virtual FString GetReferencerName() const override { return TEXT("TContextCacheElementUObject"); }
+		//~ End FGCObject interface
+
+	private:
+		typedef typename TDecay<T>::Type FDataType;  // Using universal references here means T could be either const& or an rvalue reference
+		FDataType Object;                            // Decaying T removes any reference and gets the correct underlying storage data type
 	};
 
 	// cache element method implementation 
@@ -159,16 +193,26 @@ namespace UE::Dataflow
 		// check(PropertyIn->IsA<T>()); // @todo(dataflow) compile error for non-class T; find alternatives
 		if (Type == EType::CacheElementTyped)
 		{
-			return static_cast<const FContextCacheElement<T>&>(*this).GetData(Context, PropertyIn, Default);
+			if constexpr (!TIsUObjectPtrElement<T>::Value)
+			{
+				return static_cast<const TContextCacheElement<T>&>(*this).GetData(Context, PropertyIn, Default);
+			}
 		}
 		if (Type == EType::CacheElementReference)
 		{
-			return static_cast<const FContextCacheElementReference<T>&>(*this).GetData(Context, PropertyIn, Default);
+			return static_cast<const TContextCacheElementReference<T>&>(*this).GetData(Context, PropertyIn, Default);
 		}
 		if (Type == EType::CacheElementNull)
 		{
 			return Default; 
 		}
+		if (Type == EType::CacheElementUObject)
+		{
+			if constexpr (TIsUObjectPtrElement<T>::Value)
+			{
+				return static_cast<const TContextCacheElementUObject<T>&>(*this).GetData(Context, PropertyIn, Default);
+			}
+		}
 		check(false); // should never happen
 		return Default;
 	}
@@ -234,7 +278,15 @@ namespace UE::Dataflow
 		template<typename T>
 		void SetData(FContextCacheKey InKey, const FProperty* InProperty, T&& InValue, const FGuid& InNodeGuid, uint32 InNodeHash, const FTimestamp& InTimestamp)
 		{
-			TUniquePtr<FContextCacheElement<T>> DataStoreEntry = MakeUnique<FContextCacheElement<T>>(InNodeGuid, InProperty, Forward<T>(InValue), InNodeHash, InTimestamp);
+			TUniquePtr<FContextCacheElementBase> DataStoreEntry;
+			if constexpr (TIsUObjectPtrElement<T>::Value)
+			{
+				DataStoreEntry = MakeUnique<TContextCacheElementUObject<T>>(InNodeGuid, InProperty, Forward<T>(InValue), InNodeHash, InTimestamp);
+			}
+			else
+			{
+				DataStoreEntry = MakeUnique<TContextCacheElement<T>>(InNodeGuid, InProperty, Forward<T>(InValue), InNodeHash, InTimestamp);
+			}
 
 			SetDataImpl(InKey, MoveTemp(DataStoreEntry));
 		}
@@ -421,28 +473,28 @@ namespace UE::Dataflow
 	// cache classes implemetation 
 	// this needs to be after the FContext definition because they access its methods
 
-	template<class T>
-	const T& FContextCacheElement<T>::GetData(FContext& Context, const FProperty* PropertyIn, const T& Default) const
+	template<class T UE_REQUIRES(!TIsUObjectPtrElement<T>::Value)>
+	const T& TContextCacheElement<T>::GetData(FContext& Context, const FProperty* PropertyIn, const T& Default) const
 	{
 		return Data;
 	}
 
-	template<class T>
-	TUniquePtr<FContextCacheElementBase> FContextCacheElement<T>::CreateReference(FContextCacheKey InReferenceDataKey) const
+	template<class T UE_REQUIRES(!TIsUObjectPtrElement<T>::Value)>
+	TUniquePtr<FContextCacheElementBase> TContextCacheElement<T>::CreateReference(FContextCacheKey InReferenceDataKey) const
 	{
-		return MakeUnique<FContextCacheElementReference<T>>(GetNodeGuid(), GetProperty(), InReferenceDataKey, GetNodeHash(), GetTimestamp());
+		return MakeUnique<TContextCacheElementReference<T>>(GetNodeGuid(), GetProperty(), InReferenceDataKey, GetNodeHash(), GetTimestamp());
 	}
 
 	template<class T>
-	const T& FContextCacheElementReference<T>::GetData(FContext& Context, const FProperty* PropertyIn, const T& Default) const
+	const T& TContextCacheElementReference<T>::GetData(FContext& Context, const FProperty* PropertyIn, const T& Default) const
 	{
 		return Context.GetData(DataKey, PropertyIn, Default);
 	}
 
 	template<class T>
-	TUniquePtr<FContextCacheElementBase> FContextCacheElementReference<T>::CreateReference(FContextCacheKey InReferenceDataKey) const
+	TUniquePtr<FContextCacheElementBase> TContextCacheElementReference<T>::CreateReference(FContextCacheKey InReferenceDataKey) const
 	{
-		return MakeUnique<FContextCacheElementReference<T>>(GetNodeGuid(), GetProperty(), InReferenceDataKey, GetNodeHash(), GetTimestamp());
+		return MakeUnique<TContextCacheElementReference<T>>(GetNodeGuid(), GetProperty(), InReferenceDataKey, GetNodeHash(), GetTimestamp());
 	}
 
 	inline TUniquePtr<FContextCacheElementBase> FContextCacheElementNull::CreateReference(FContextCacheKey InReferenceDataKey) const
@@ -450,6 +502,12 @@ namespace UE::Dataflow
 		// a null entry will always return a null entry as a reference
 		return MakeUnique<FContextCacheElementNull>(GetNodeGuid(), GetProperty(), InReferenceDataKey, GetNodeHash(), GetTimestamp());
 	}
+
+	template<class T UE_REQUIRES(TIsUObjectPtrElement<T>::Value)>
+	TUniquePtr<FContextCacheElementBase> TContextCacheElementUObject<T>::CreateReference(FContextCacheKey InReferenceDataKey) const
+	{
+		return MakeUnique<TContextCacheElementReference<T>>(GetNodeGuid(), GetProperty(), InReferenceDataKey, GetNodeHash(), GetTimestamp());
+	}
 }
 
 
