CIS - Fixed issue Issue 621500: Compile and NonUnity errors in Module.FractureEditor.cpp

#rnx

diff --git a/Engine/Source/Runtime/Experimental/Dataflow/Core/Public/Dataflow/DataflowNodeParameters.h b/Engine/Source/Runtime/Experimental/Dataflow/Core/Public/Dataflow/DataflowNodeParameters.h
--- a/Engine/Source/Runtime/Experimental/Dataflow/Core/Public/Dataflow/DataflowNodeParameters.h
+++ b/Engine/Source/Runtime/Experimental/Dataflow/Core/Public/Dataflow/DataflowNodeParameters.h
@@ -111,7 +111,7 @@ namespace UE::Dataflow
 	//--------------------------------------------------------------------
 	// Value storing context cache entry - strongly typed
 	//--------------------------------------------------------------------
-	template<class T UE_REQUIRES(!TIsUObjectPtrElement<T>::Value)>
+	template<class T>
 	struct TContextCacheElement : public FContextCacheElementBase 
 	{
 		TContextCacheElement(FGuid InNodeGuid, const FProperty* InProperty, T&& InData, uint32 InNodeHash, FTimestamp Timestamp)
@@ -164,7 +164,7 @@ namespace UE::Dataflow
 	//--------------------------------------------------------------------
 	// UObject cache element, prevents the object from being garbage collected while in the cache
 	//--------------------------------------------------------------------
-	template<class T UE_REQUIRES(TIsUObjectPtrElement<T>::Value)>
+	template<class T>
 	struct TContextCacheElementUObject : public FContextCacheElementBase, public FGCObject
 	{
 		TContextCacheElementUObject(FGuid InNodeGuid, const FProperty* InProperty, T&& InObject, uint32 InNodeHash, FTimestamp Timestamp)
@@ -473,13 +473,13 @@ namespace UE::Dataflow
 	// cache classes implemetation 
 	// this needs to be after the FContext definition because they access its methods
 
-	template<class T UE_REQUIRES(!TIsUObjectPtrElement<T>::Value)>
+	template<class T>
 	const T& TContextCacheElement<T>::GetData(FContext& Context, const FProperty* PropertyIn, const T& Default) const
 	{
 		return Data;
 	}
 
-	template<class T UE_REQUIRES(!TIsUObjectPtrElement<T>::Value)>
+	template<class T>
 	TUniquePtr<FContextCacheElementBase> TContextCacheElement<T>::CreateReference(FContextCacheKey InReferenceDataKey) const
 	{
 		return MakeUnique<TContextCacheElementReference<T>>(GetNodeGuid(), GetProperty(), InReferenceDataKey, GetNodeHash(), GetTimestamp());
@@ -503,7 +503,7 @@ namespace UE::Dataflow
 		return MakeUnique<FContextCacheElementNull>(GetNodeGuid(), GetProperty(), InReferenceDataKey, GetNodeHash(), GetTimestamp());
 	}
 
-	template<class T UE_REQUIRES(TIsUObjectPtrElement<T>::Value)>
+	template<class T>
 	TUniquePtr<FContextCacheElementBase> TContextCacheElementUObject<T>::CreateReference(FContextCacheKey InReferenceDataKey) const
 	{
 		return MakeUnique<TContextCacheElementReference<T>>(GetNodeGuid(), GetProperty(), InReferenceDataKey, GetNodeHash(), GetTimestamp());
