Minor clean up of unused code.

#rnx

diff --git a/Engine/Source/Runtime/Experimental/IoStore/OnDemand/Public/IO/IoStoreOnDemand.h b/Engine/Source/Runtime/Experimental/IoStore/OnDemand/Public/IO/IoStoreOnDemand.h
--- a/Engine/Source/Runtime/Experimental/IoStore/OnDemand/Public/IO/IoStoreOnDemand.h
+++ b/Engine/Source/Runtime/Experimental/IoStore/OnDemand/Public/IO/IoStoreOnDemand.h
@@ -2,6 +2,7 @@
 
 #pragma once
 
+#include "Containers/Map.h"
 #include "Containers/UnrealString.h"
 #include "Containers/StringFwd.h"
 #include "Containers/SharedString.h"
@@ -18,11 +19,6 @@
 
 #include <atomic>
 
-#if (IS_PROGRAM || WITH_EDITOR)
-#include "Containers/Map.h"
-#include "Misc/AES.h"
-#endif // (IS_PROGRAM || WITH_EDITOR)
-
 #define UE_API IOSTOREONDEMAND_API
 
 class FArchive;
