Chaos::Add velocity from rotation in destruction when a child is freed from its parent.

diff --git a/Engine/Source/Runtime/Experimental/Chaos/Private/Chaos/PBDRigidClustering.cpp b/Engine/Source/Runtime/Experimental/Chaos/Private/Chaos/PBDRigidClustering.cpp
--- a/Engine/Source/Runtime/Experimental/Chaos/Private/Chaos/PBDRigidClustering.cpp
+++ b/Engine/Source/Runtime/Experimental/Chaos/Private/Chaos/PBDRigidClustering.cpp
@@ -907,13 +907,13 @@ namespace Chaos
 			Child->SetP(Child->GetX());
 			Child->SetQf(Child->GetRf());
 
-			//todo(ocohen): for now just inherit velocity at new COM. This isn't quite right for rotation
 			//todo(ocohen): in the presence of collisions, this will leave all children with the post-collision
 			// velocity. This should be controlled by material properties so we can allow the broken pieces to
 			// maintain the clusters pre-collision velocity.
-			Child->SetVf(Child->GetVf() + ClusteredParent->GetVf());
+			const FVec3f PDiff = Child->GetP() - ClusteredParent->GetP();
+			Child->SetVf(Child->GetVf() + ClusteredParent->GetVf() + FVec3f::CrossProduct(ClusteredParent->GetWf(), PDiff));
 			Child->SetWf(Child->GetWf() + ClusteredParent->GetWf());
-			Child->SetPreVf(Child->GetPreVf() + ClusteredParent->GetPreVf());
+			Child->SetPreVf(Child->GetPreVf() + ClusteredParent->GetPreVf() + FVec3f::CrossProduct(ClusteredParent->GetPreWf(), PDiff));
 			Child->SetPreWf(Child->GetPreWf() + ClusteredParent->GetPreWf());
 
 			// We also need to do cluster book-keeping on the parent.
