ChaosFlesh: CIS fix
Fixed BoneWeightsFacade.IsValid() by allowing KinematicAttribute to be invalid for backward compatibility

diff --git a/Engine/Source/Runtime/Experimental/Chaos/Private/GeometryCollection/Facades/CollectionVertexBoneWeightsFacade.cpp b/Engine/Source/Runtime/Experimental/Chaos/Private/GeometryCollection/Facades/CollectionVertexBoneWeightsFacade.cpp
--- a/Engine/Source/Runtime/Experimental/Chaos/Private/GeometryCollection/Facades/CollectionVertexBoneWeightsFacade.cpp
+++ b/Engine/Source/Runtime/Experimental/Chaos/Private/GeometryCollection/Facades/CollectionVertexBoneWeightsFacade.cpp
@@ -55,7 +55,7 @@ namespace GeometryCollection::Facades
 
 	bool FVertexBoneWeightsFacade::IsValid() const
 	{
-		return BoneIndexAttribute.IsValid() && BoneWeightAttribute.IsValid() && KinematicAttribute.IsValid() && ParentAttribute.IsValid() && VerticesAttribute.IsValid();
+		return BoneIndexAttribute.IsValid() && BoneWeightAttribute.IsValid() && ParentAttribute.IsValid() && VerticesAttribute.IsValid();
 	}
 
 
@@ -105,7 +105,7 @@ namespace GeometryCollection::Facades
 
 	void FVertexBoneWeightsFacade::SetVertexKinematic(int32 VertexIndex)
 	{
-		if (KinematicAttribute.IsValidIndex(VertexIndex))
+		if (KinematicAttribute.IsValid() && KinematicAttribute.IsValidIndex(VertexIndex))
 		{
 			KinematicAttribute.ModifyAt(VertexIndex, true);
 		}
@@ -113,12 +113,15 @@ namespace GeometryCollection::Facades
 
 	void FVertexBoneWeightsFacade::SetVertexArrayKinematic(const TArray<int32>& VertexIndices)
 	{
-		TManagedArray<bool>& KinematicArray = KinematicAttribute.Modify();
-		for (const int32& VertexIndex : VertexIndices)
+		if (KinematicAttribute.IsValid())
 		{
-			if (KinematicArray.IsValidIndex(VertexIndex))
+			TManagedArray<bool>& KinematicArray = KinematicAttribute.Modify();
+			for (const int32& VertexIndex : VertexIndices)
 			{
-				KinematicArray[VertexIndex] = true;
+				if (KinematicArray.IsValidIndex(VertexIndex))
+				{
+					KinematicArray[VertexIndex] = true;
+				}
 			}
 		}
 	}
