Iris
- Fixed disconnect issue when receiving RPCs to no longer replicated dynamic object that has not yet been garbage collected.

diff --git a/Engine/Source/Runtime/Experimental/Iris/Core/Private/Iris/ReplicationSystem/NetBlob/NetRPC.cpp b/Engine/Source/Runtime/Experimental/Iris/Core/Private/Iris/ReplicationSystem/NetBlob/NetRPC.cpp
--- a/Engine/Source/Runtime/Experimental/Iris/Core/Private/Iris/ReplicationSystem/NetBlob/NetRPC.cpp
+++ b/Engine/Source/Runtime/Experimental/Iris/Core/Private/Iris/ReplicationSystem/NetBlob/NetRPC.cpp
@@ -700,10 +700,10 @@ static bool NetRPC_GetFunctionAndObject(FNetSerializationContext& Context, const
 	}
 
 	const FReplicationProtocol* Protocol = ReplicationSystem->GetReplicationProtocol(ObjectReference.GetRefHandle());
-	if (!ensureMsgf(Protocol != nullptr, TEXT("ReplicationProtocol doesn't exist for %s (Connection %u). Ignoring RPC (%u|%u)"), 
-		*GetNameSafe(RefObject), Context.GetLocalConnectionId(), FunctionLocator.DescriptorIndex, FunctionLocator.FunctionIndex))
+	if (!Protocol)
 	{
-		Context.SetError(NetError_InvalidNetObjectReference);
+		// Ignore this RPC and continue processing the rest of the data, Note: this might for example occur if we have incoming RPC data from client to an object that has been destroyed on server.
+  		UE_LOG(LogIrisRpc, Verbose, TEXT("ReplicationProtocol doesn't exist for %s (Connection %u). Ignoring RPC (%u|%u), this is most likely due to object %s no longer being replicated."), *GetNameSafe(RefObject), Context.GetLocalConnectionId(), FunctionLocator.DescriptorIndex, FunctionLocator.FunctionIndex, *ObjectReference.ToString());
 		return false;
 	}
 
