Chaos Debug Draw - enable new debug draw API for particles

diff --git a/Engine/Source/Runtime/Experimental/Chaos/Private/Chaos/ChaosDebugDraw.cpp b/Engine/Source/Runtime/Experimental/Chaos/Private/Chaos/ChaosDebugDraw.cpp
--- a/Engine/Source/Runtime/Experimental/Chaos/Private/Chaos/ChaosDebugDraw.cpp
+++ b/Engine/Source/Runtime/Experimental/Chaos/Private/Chaos/ChaosDebugDraw.cpp
@@ -7,6 +7,7 @@
 #include "Chaos/Collision/CollisionConstraintAllocator.h"
 #include "Chaos/Collision/ParticlePairMidPhase.h"
 #include "Chaos/Convex.h"
+#include "Chaos/DebugDraw/DebugDrawPrimitives.h"
 #include "Chaos/HeightField.h"
 #include "Chaos/Triangle.h"
 #include "Chaos/TriangleMesh.h"
@@ -35,6 +36,8 @@
 
 namespace Chaos
 {
+	extern bool bChaosDebugDraw_UseNewQueue;
+
 	namespace CVars
 	{
 		extern int32 ChaosOneWayInteractionPairCollisionMode;
@@ -137,6 +140,171 @@ namespace Chaos
 		FAutoConsoleVariableRef CVarChaosDebugDebugDrawBVHShapes(TEXT("p.Chaos.DebugDraw.BVHShapes"), bChaosDebugDebugDrawBVHShapes, TEXT("Whether to draw the shapes in each BVH node (colorized)"));
 		FAutoConsoleVariableRef CVarChaosDebugDebugDrawBVHBounds(TEXT("p.Chaos.DebugDraw.BVHBounds"), bChaosDebugDebugDrawBVHBounds, TEXT("Whether to draw the BVH node bounds (colorized)"));
 
+		//
+		// TEMP - helpers to switch between new and old debug draw systems until we drop the old one
+		// NOTE: This is switching to FChaosDDPrimitives, which misses out on the benefits of the new
+		// system which are reduced data size and render-time style, but it will do for a start.
+		// The end game is to replace use of FChaosDDPrimitives with specialized debugh draw
+		// commands like FChaosDDParticles::DrawShapes() etc.
+		//
+
+		void DrawDebugPoint(const FVector& Position, const FColor& Color, bool bPersistentLines, float LifeTime, uint8 DepthPriority, float Thickness)
+		{
+			if (bChaosDebugDraw_UseNewQueue)
+			{
+				FChaosDDPrimitives::DrawPoint(Position, Color, Thickness, LifeTime);
+			}
+			else
+			{
+				FDebugDrawQueue::GetInstance().DrawDebugPoint(Position, Color, bPersistentLines, LifeTime, DepthPriority, Thickness);
+			}
+		}
+
+		void DrawDebugLine(const FVector& LineStart, const FVector& LineEnd, const FColor& Color, bool bPersistentLines, float LifeTime, uint8 DepthPriority, float Thickness)
+		{
+			if (bChaosDebugDraw_UseNewQueue)
+			{
+				FChaosDDPrimitives::DrawLine(LineStart, LineEnd, Color, Thickness, LifeTime);
+			}
+			else
+			{
+				FDebugDrawQueue::GetInstance().DrawDebugLine(LineStart, LineEnd, Color, bPersistentLines, LifeTime, DepthPriority, Thickness);
+			}
+		}
+
+		void DrawDebugSphere(FVector const& Center, FReal Radius, int32 Segments, const FColor& Color, bool bPersistentLines, float LifeTime, uint8 DepthPriority, float Thickness)
+		{
+			if (bChaosDebugDraw_UseNewQueue)
+			{
+				FChaosDDPrimitives::DrawSphere(Center, float(Radius), Color, Thickness, LifeTime);
+			}
+			else
+			{
+				FDebugDrawQueue::GetInstance().DrawDebugSphere(Center, Radius, Segments, Color, bPersistentLines, LifeTime, DepthPriority, Thickness);
+			}
+		}
+
+		void DrawDebugBox(FVector const& Center, FVector const& Extent, const FQuat& Rotation, FColor const& Color, bool bPersistentLines, float LifeTime, uint8 DepthPriority, float Thickness)
+		{
+			if (bChaosDebugDraw_UseNewQueue)
+			{
+				FChaosDDPrimitives::DrawBox(Center, Rotation, Extent, Color, Thickness, LifeTime);
+			}
+			else
+			{
+				FDebugDrawQueue::GetInstance().DrawDebugBox(Center, Extent, Rotation, Color, bPersistentLines, LifeTime, DepthPriority, Thickness);
+			}
+		}
+
+		void DrawDebugString(FVector const& TextLocation, const FString& Text, class AActor* TestBaseActor, FColor const& Color, float LifeTime, bool bDrawShadow, float FontScale)
+		{
+			if (bChaosDebugDraw_UseNewQueue)
+			{
+				FChaosDDPrimitives::DrawString(TextLocation, Text, Color, FontScale, bDrawShadow, LifeTime);
+			}
+			else
+			{
+				FDebugDrawQueue::GetInstance().DrawDebugString(TextLocation, Text, TestBaseActor, Color, LifeTime, bDrawShadow, FontScale);
+			}
+		}
+
+		void DrawDebugCircle(const FVector& Center, FReal Radius, int32 Segments, const FColor& Color, bool bPersistentLines, float LifeTime, uint8 DepthPriority, float Thickness, const FMatrix& Axes, bool bDrawAxis)
+		{
+			if (bChaosDebugDraw_UseNewQueue)
+			{
+				FChaosDDPrimitives::DrawCircle(Center, Axes, float(Radius), Color, Thickness, LifeTime);
+			}
+			else
+			{
+				FDebugDrawQueue::GetInstance().DrawDebugCircle(Center, Radius, Segments, Color, bPersistentLines, LifeTime, DepthPriority, Thickness, Axes.GetUnitAxis(EAxis::Y), Axes.GetUnitAxis(EAxis::Z), bDrawAxis);
+			}
+		}
+
+		void DrawDebugCapsule(const FVector& Center, FReal HalfHeight, FReal Radius, const FQuat& Rotation, const FColor& Color, bool bPersistentLines, float LifeTime, uint8 DepthPriority, float Thickness)
+		{
+			if (bChaosDebugDraw_UseNewQueue)
+			{
+				FChaosDDPrimitives::DrawCapsule(Center, Rotation, float(HalfHeight), float(Radius), Color, Thickness, LifeTime);
+			}
+			else
+			{
+				FDebugDrawQueue::GetInstance().DrawDebugCapsule(Center, HalfHeight, Radius, Rotation, Color, bPersistentLines, LifeTime, DepthPriority, Thickness);
+			}
+		}
+
+		void DrawDebugDirectionalArrow(const FVector& LineStart, const FVector& LineEnd, float ArrowSize, const FColor& Color, bool bPersistentLines, float LifeTime, uint8 DepthPriority, float Thickness)
+		{
+			if (bChaosDebugDraw_UseNewQueue)
+			{
+				FChaosDDPrimitives::DrawArrow(LineStart, LineEnd, ArrowSize, Color, Thickness, LifeTime);
+			}
+			else
+			{
+				FDebugDrawQueue::GetInstance().DrawDebugDirectionalArrow(LineStart, LineEnd, ArrowSize, Color, bPersistentLines, LifeTime, DepthPriority, Thickness);
+			}
+		}
+
+		FVector GetCenterOfInterest()
+		{
+			if (bChaosDebugDraw_UseNewQueue)
+			{
+				return ChaosDD::Private::FChaosDDContext::GetWriter().GetDrawRegion().Center;
+			}
+			else
+			{
+				return FDebugDrawQueue::GetInstance().GetCenterOfInterest();
+			}
+		}
+
+		FReal GetRadiusOfInterest()
+		{
+			if (bChaosDebugDraw_UseNewQueue)
+			{
+				return ChaosDD::Private::FChaosDDContext::GetWriter().GetDrawRegion().W;
+			}
+			else
+			{
+				return FDebugDrawQueue::GetInstance().GetRadiusOfInterest();
+			}
+		}
+
+		bool IsInRegionOfInterest(FVector Pos)
+		{
+			if (bChaosDebugDraw_UseNewQueue)
+			{
+				return ChaosDD::Private::FChaosDDContext::GetWriter().IsInDrawRegion(Pos);
+			}
+			else
+			{
+				return FDebugDrawQueue::GetInstance().IsInRegionOfInterest(Pos);
+			}
+		}
+
+		bool IsInRegionOfInterest(FVector Pos, FReal Radius)
+		{
+			if (bChaosDebugDraw_UseNewQueue)
+			{
+				return ChaosDD::Private::FChaosDDContext::GetWriter().IsInDrawRegion(FSphere3d(Pos, Radius));
+			}
+			else
+			{
+				return FDebugDrawQueue::GetInstance().IsInRegionOfInterest(Pos, Radius);
+			}
+		}
+
+		bool IsInRegionOfInterest(FAABB3 Bounds)
+		{
+			if (bChaosDebugDraw_UseNewQueue)
+			{
+				return ChaosDD::Private::FChaosDDContext::GetWriter().IsInDrawRegion(FBox3d(Bounds.Min(), Bounds.Max()));
+			}
+			else
+			{
+				return FDebugDrawQueue::GetInstance().IsInRegionOfInterest(Bounds);
+			}
+		}
+
+
 		// NOTE: These settings should never really be used - they are the fallback defaults
 		// if the user does not specify settings in the debug draw call.
 		// See PBDRigidsColver.cpp and ImmediatePhysicsSimulation_Chaos.cpp for example.
@@ -529,7 +697,7 @@ namespace Chaos
 						const FVec3 ExplodeDelta = ChaosDebugDrawConvexExplodeDistance * N0;
 
 						// Outer shape
-						FDebugDrawQueue::GetInstance().DrawDebugLine(OuterP0 + ExplodeDelta, OuterP1 + ExplodeDelta, Color, false, Duration, uint8(Settings.DrawPriority), Settings.ShapeThicknesScale * Settings.LineThickness);
+						DrawDebugLine(OuterP0 + ExplodeDelta, OuterP1 + ExplodeDelta, Color, false, Duration, uint8(Settings.DrawPriority), Settings.ShapeThicknesScale * Settings.LineThickness);
 
 						// Core shape and lines connecting core to outer
 						if (Margin > 0.0f)
@@ -537,14 +705,14 @@ namespace Chaos
 							const FRealSingle LineThickness = 0.5f * Settings.ShapeThicknesScale * Settings.LineThickness;
 							const FVec3 InnerP0 = ShapeTransform.TransformPositionNoScale(Shape->GetMarginAdjustedVertexScaled(VertexIndex0, Margin, ShapeTransform.GetScale3D(), nullptr));
 							const FVec3 InnerP1 = ShapeTransform.TransformPositionNoScale(Shape->GetMarginAdjustedVertexScaled(VertexIndex1, Margin, ShapeTransform.GetScale3D(), nullptr));
-							FDebugDrawQueue::GetInstance().DrawDebugLine(InnerP0, InnerP1, FColor::Blue, false, Duration, uint8(Settings.DrawPriority), LineThickness);
-							FDebugDrawQueue::GetInstance().DrawDebugLine(InnerP0, OuterP0, FColor::Black, false, Duration, uint8(Settings.DrawPriority), LineThickness);
+							DrawDebugLine(InnerP0, InnerP1, FColor::Blue, false, Duration, uint8(Settings.DrawPriority), LineThickness);
+							DrawDebugLine(InnerP0, OuterP0, FColor::Black, false, Duration, uint8(Settings.DrawPriority), LineThickness);
 						}
 
 						// Vertex and face normal
 						if (bChaosDebugDebugDrawConvexVertices)
 						{
-							FDebugDrawQueue::GetInstance().DrawDebugLine(OuterP0 + ExplodeDelta, OuterP0 + ExplodeDelta + Settings.DrawScale * 20.0f * N0, FColor::Green, false, Duration, uint8(Settings.DrawPriority), Settings.LineThickness);
+							DrawDebugLine(OuterP0 + ExplodeDelta, OuterP0 + ExplodeDelta + Settings.DrawScale * 20.0f * N0, FColor::Green, false, Duration, uint8(Settings.DrawPriority), Settings.LineThickness);
 						}
 					}
 				}
@@ -553,8 +721,8 @@ namespace Chaos
 
 		void DrawShapesHeightFieldImpl(const TGeometryParticleHandle<FReal, 3>* Particle, const FRigidTransform3& ShapeTransform, const FHeightField* Shape, const FColor& Color, const FRealSingle Duration, const FChaosDebugDrawSettings& Settings)
 		{
-			const FVec3& WorldQueryCenter = FDebugDrawQueue::GetInstance().GetCenterOfInterest();
-			const FReal WorldQueryRadius = FDebugDrawQueue::GetInstance().GetRadiusOfInterest();
+			const FVec3& WorldQueryCenter = GetCenterOfInterest();
+			const FReal WorldQueryRadius = GetRadiusOfInterest();
 			const FAABB3 WorldQueryBounds = FAABB3(WorldQueryCenter - FVec3(WorldQueryRadius), WorldQueryCenter + FVec3(WorldQueryRadius));
 			const FAABB3 LocalQueryBounds = WorldQueryBounds.InverseTransformedAABB(ShapeTransform);
 
@@ -563,16 +731,16 @@ namespace Chaos
 				const FVec3& A = Tri[0];
 				const FVec3& B = Tri[1];
 				const FVec3& C = Tri[2];
-				FDebugDrawQueue::GetInstance().DrawDebugLine(A, B, Color, false, Duration, 0, Settings.ShapeThicknesScale * Settings.LineThickness);
-				FDebugDrawQueue::GetInstance().DrawDebugLine(B, C, Color, false, Duration, 0, Settings.ShapeThicknesScale * Settings.LineThickness);
-				FDebugDrawQueue::GetInstance().DrawDebugLine(C, A, Color, false, Duration, 0, Settings.ShapeThicknesScale * Settings.LineThickness);
+				DrawDebugLine(A, B, Color, false, Duration, 0, Settings.ShapeThicknesScale * Settings.LineThickness);
+				DrawDebugLine(B, C, Color, false, Duration, 0, Settings.ShapeThicknesScale * Settings.LineThickness);
+				DrawDebugLine(C, A, Color, false, Duration, 0, Settings.ShapeThicknesScale * Settings.LineThickness);
 			});
 		}
 
 		void DrawShapesTriangleMeshImpl(const TGeometryParticleHandle<FReal, 3>* Particle, const FRigidTransform3& ShapeTransform, const FTriangleMeshImplicitObject* Shape, const FColor& Color, const FRealSingle Duration, const FChaosDebugDrawSettings& Settings)
 		{
-			const FVec3& WorldQueryCenter = FDebugDrawQueue::GetInstance().GetCenterOfInterest();
-			const FReal WorldQueryRadius = FDebugDrawQueue::GetInstance().GetRadiusOfInterest();
+			const FVec3& WorldQueryCenter = GetCenterOfInterest();
+			const FReal WorldQueryRadius = GetRadiusOfInterest();
 			const FAABB3 WorldQueryBounds = FAABB3(WorldQueryCenter - FVec3(WorldQueryRadius), WorldQueryCenter + FVec3(WorldQueryRadius));
 			const FAABB3 LocalQueryBounds = WorldQueryBounds.InverseTransformedAABB(ShapeTransform);
 
@@ -581,9 +749,9 @@ namespace Chaos
 				const FVec3& A = Tri[0];
 				const FVec3& B = Tri[1];
 				const FVec3& C = Tri[2];
-				FDebugDrawQueue::GetInstance().DrawDebugLine(A, B, Color, false, Duration, 0, Settings.ShapeThicknesScale * Settings.LineThickness);
-				FDebugDrawQueue::GetInstance().DrawDebugLine(B, C, Color, false, Duration, 0, Settings.ShapeThicknesScale * Settings.LineThickness);
-				FDebugDrawQueue::GetInstance().DrawDebugLine(C, A, Color, false, Duration, 0, Settings.ShapeThicknesScale * Settings.LineThickness);
+				DrawDebugLine(A, B, Color, false, Duration, 0, Settings.ShapeThicknesScale * Settings.LineThickness);
+				DrawDebugLine(B, C, Color, false, Duration, 0, Settings.ShapeThicknesScale * Settings.LineThickness);
+				DrawDebugLine(C, A, Color, false, Duration, 0, Settings.ShapeThicknesScale * Settings.LineThickness);
 			});
 		}
 
@@ -607,7 +775,7 @@ namespace Chaos
 					for (int32 ParticleIndex = 0; ParticleIndex < (int32)CollisionParticles->Size(); ++ParticleIndex)
 					{
 						const FVec3 P = ShapeTransform.TransformPosition(CollisionParticles->GetX(ParticleIndex));
-						FDebugDrawQueue::GetInstance().DrawDebugPoint(P, Color, false, Duration, 0, Settings.PointSize);
+						DrawDebugPoint(P, Color, false, Duration, 0, Settings.PointSize);
 					}
 				}
 			}
@@ -739,7 +907,7 @@ namespace Chaos
 			const EImplicitObjectType InnerType = GetInnerType(Implicit->GetType());
 
 			// Are we within the region of interest?
-			if (Implicit->HasBoundingBox() && !FDebugDrawQueue::GetInstance().IsInRegionOfInterest(Implicit->BoundingBox().TransformedAABB(ShapeTransform)))
+			if (Implicit->HasBoundingBox() && !IsInRegionOfInterest(Implicit->BoundingBox().TransformedAABB(ShapeTransform)))
 			{
 				return;
 			}
@@ -901,7 +1069,7 @@ namespace Chaos
 				const TSphere<FReal, 3>* Sphere = Implicit->template GetObject<TSphere<FReal, 3>>();
 				const FVec3 P = ShapeTransform.TransformPosition(Sphere->GetCenter());
 				const FVec3 Scale = ShapeTransform.GetScale3D();
-				FDebugDrawQueue::GetInstance().DrawDebugSphere(P, Sphere->GetRadius() * Scale.Z, 8, ShapeColor, false, Duration, uint8(Settings.DrawPriority), Settings.ShapeThicknesScale * Settings.LineThickness);
+				DrawDebugSphere(P, Sphere->GetRadius() * Scale.Z, 8, ShapeColor, false, Duration, uint8(Settings.DrawPriority), Settings.ShapeThicknesScale * Settings.LineThickness);
 				break;
 			}
 			case ImplicitObjectType::Box:
@@ -909,7 +1077,7 @@ namespace Chaos
 				const TBox<FReal, 3>* Box = Implicit->template GetObject<TBox<FReal, 3>>();
 				const FVec3 P = ShapeTransform.TransformPosition(Box->GetCenter());
 				const FVec3 Scale = ShapeTransform.GetScale3D();
-				FDebugDrawQueue::GetInstance().DrawDebugBox(P, (FReal)0.5 * Box->Extents() * Scale, ShapeTransform.GetRotation(), ShapeColor, false, Duration, uint8(Settings.DrawPriority), Settings.ShapeThicknesScale * Settings.LineThickness);
+				DrawDebugBox(P, (FReal)0.5 * Box->Extents() * Scale, ShapeTransform.GetRotation(), ShapeColor, false, Duration, uint8(Settings.DrawPriority), Settings.ShapeThicknesScale * Settings.LineThickness);
 				break;
 			}
 			case ImplicitObjectType::Plane:
@@ -920,7 +1088,7 @@ namespace Chaos
 				const FVec3 P = ShapeTransform.TransformPosition(Capsule->GetCenter());
 				const FRotation3 Q = ShapeTransform.GetRotation() * FRotationMatrix::MakeFromZ(Capsule->GetAxis());
 				const FVec3 Scale = ShapeTransform.GetScale3D();
-				FDebugDrawQueue::GetInstance().DrawDebugCapsule(P, ((FReal)0.5 * Capsule->GetHeight() + Capsule->GetRadius()) * Scale.Z, Capsule->GetRadius() * Scale.X, Q, ShapeColor, false, Duration, uint8(Settings.DrawPriority), Settings.ShapeThicknesScale * Settings.LineThickness);
+				DrawDebugCapsule(P, ((FReal)0.5 * Capsule->GetHeight() + Capsule->GetRadius()) * Scale.Z, Capsule->GetRadius() * Scale.X, Q, ShapeColor, false, Duration, uint8(Settings.DrawPriority), Settings.ShapeThicknesScale * Settings.LineThickness);
 				break;
 			}
 			case ImplicitObjectType::LevelSet:
@@ -985,7 +1153,7 @@ namespace Chaos
 						for (int32 ParticleIndex = 0; ParticleIndex < (int32)Particles->Size(); ++ParticleIndex)
 						{
 							FVec3 P = ShapeTransform.TransformPosition(Particles->GetX(ParticleIndex));
-							FDebugDrawQueue::GetInstance().DrawDebugPoint(P, ShapeColor, false, Duration, uint8(Settings.DrawPriority), Settings.PointSize);
+							DrawDebugPoint(P, ShapeColor, false, Duration, uint8(Settings.DrawPriority), Settings.PointSize);
 						}
 					}
 				}
@@ -996,7 +1164,7 @@ namespace Chaos
 				const FColor ShapeBoundsColor = FColor::Orange;
 				const FAABB3& ShapeBounds = Implicit->BoundingBox();
 				const FVec3 ShapeBoundsPos = ShapeTransform.TransformPosition(ShapeBounds.Center());
-				FDebugDrawQueue::GetInstance().DrawDebugBox(ShapeBoundsPos, 0.5f * ShapeBounds.Extents(), ShapeTransform.GetRotation(), ShapeBoundsColor, false, Duration, uint8(Settings.DrawPriority), Settings.LineThickness);
+				DrawDebugBox(ShapeBoundsPos, 0.5f * ShapeBounds.Extents(), ShapeTransform.GetRotation(), ShapeBoundsColor, false, Duration, uint8(Settings.DrawPriority), Settings.LineThickness);
 			}
 		}
 		
@@ -1065,7 +1233,7 @@ namespace Chaos
 						{
 							const FVec3 NodeBoundsPos = ShapeTransform.TransformPosition(FVec3(NodeBounds.Center()));
 							const FVec3 NodeBoundsExtents = 0.5f * FVec3(NodeBounds.Extents()) - FVec3(BoundsDrawShrink);
-							FDebugDrawQueue::GetInstance().DrawDebugBox(NodeBoundsPos, NodeBoundsExtents, ShapeTransform.GetRotation(), Color, false, Duration, uint8(Settings.DrawPriority), Settings.LineThickness);
+							DrawDebugBox(NodeBoundsPos, NodeBoundsExtents, ShapeTransform.GetRotation(), Color, false, Duration, uint8(Settings.DrawPriority), Settings.LineThickness);
 						}
 
 						if (bChaosDebugDebugDrawBVHShapes)
@@ -1150,7 +1318,7 @@ namespace Chaos
 				}
 			}
 
-			FDebugDrawQueue::GetInstance().DrawDebugBox(P, 0.5f * Box.Extents(), Q, Color, false, UE_KINDA_SMALL_NUMBER, uint8(Settings.DrawPriority), Settings.LineThickness);
+			DrawDebugBox(P, 0.5f * Box.Extents(), Q, Color, false, UE_KINDA_SMALL_NUMBER, uint8(Settings.DrawPriority), Settings.LineThickness);
 
 			for (const auto& Shape : InParticle->ShapesArray())
 			{
@@ -1163,7 +1331,7 @@ namespace Chaos
 					const FRotation3 ShapeQ = SpaceTransform.GetRotation();
 					const FColor ShapeColor = (bChaosDebugDebugDrawColorBoundsByShapeType) ? Settings.BoundsColorsPerShapeType.GetColorFromShapeType(ShapeType) : Color;
 
-					FDebugDrawQueue::GetInstance().DrawDebugBox(ShapeP, 0.5f * ShapeBox.Extents(), ShapeQ, ShapeColor, false, UE_KINDA_SMALL_NUMBER, uint8(Settings.DrawPriority), Settings.LineThickness);
+					DrawDebugBox(ShapeP, 0.5f * ShapeBox.Extents(), ShapeQ, ShapeColor, false, UE_KINDA_SMALL_NUMBER, uint8(Settings.DrawPriority), Settings.LineThickness);
 				}
 			}
 		}
@@ -1199,29 +1367,29 @@ namespace Chaos
 			FVec3 PCOM = SpaceCOM.GetLocation();
 			FRotation3 QCOM = SpaceCOM.GetRotation();
 			FMatrix33 QCOMm = QCOM.ToMatrix();
-			FDebugDrawQueue::GetInstance().DrawDebugDirectionalArrow(PCOM, PCOM + Settings.DrawScale * Settings.BodyAxisLen * QCOMm.GetAxis(0), Settings.DrawScale * Settings.ArrowSize, Red, false, UE_KINDA_SMALL_NUMBER, uint8(Settings.DrawPriority), Settings.LineThickness);
-			FDebugDrawQueue::GetInstance().DrawDebugDirectionalArrow(PCOM, PCOM + Settings.DrawScale * Settings.BodyAxisLen * QCOMm.GetAxis(1), Settings.DrawScale * Settings.ArrowSize, Green, false, UE_KINDA_SMALL_NUMBER, uint8(Settings.DrawPriority), Settings.LineThickness);
-			FDebugDrawQueue::GetInstance().DrawDebugDirectionalArrow(PCOM, PCOM + Settings.DrawScale * Settings.BodyAxisLen * QCOMm.GetAxis(2), Settings.DrawScale * Settings.ArrowSize, Blue, false, UE_KINDA_SMALL_NUMBER, uint8(Settings.DrawPriority), Settings.LineThickness);
+			DrawDebugDirectionalArrow(PCOM, PCOM + Settings.DrawScale * Settings.BodyAxisLen * QCOMm.GetAxis(0), Settings.DrawScale * Settings.ArrowSize, Red, false, UE_KINDA_SMALL_NUMBER, uint8(Settings.DrawPriority), Settings.LineThickness);
+			DrawDebugDirectionalArrow(PCOM, PCOM + Settings.DrawScale * Settings.BodyAxisLen * QCOMm.GetAxis(1), Settings.DrawScale * Settings.ArrowSize, Green, false, UE_KINDA_SMALL_NUMBER, uint8(Settings.DrawPriority), Settings.LineThickness);
+			DrawDebugDirectionalArrow(PCOM, PCOM + Settings.DrawScale * Settings.BodyAxisLen * QCOMm.GetAxis(2), Settings.DrawScale * Settings.ArrowSize, Blue, false, UE_KINDA_SMALL_NUMBER, uint8(Settings.DrawPriority), Settings.LineThickness);
 
 			FColor Black = FColor::Black;
 			FColor Grey = (ColorScale * FColor(64, 64, 64)).ToFColor(false);
 			FVec3 PActor = SpaceTransform.TransformPosition(FParticleUtilities::GetActorWorldTransform(Particle).GetTranslation());
-			FDebugDrawQueue::GetInstance().DrawDebugPoint(PActor, Black, false, UE_KINDA_SMALL_NUMBER, uint8(Settings.DrawPriority), Settings.DrawScale * Settings.PointSize);
-			FDebugDrawQueue::GetInstance().DrawDebugLine(PCOM, PActor, Grey, false, UE_KINDA_SMALL_NUMBER, uint8(Settings.DrawPriority), Settings.LineThickness);
+			DrawDebugPoint(PActor, Black, false, UE_KINDA_SMALL_NUMBER, uint8(Settings.DrawPriority), Settings.DrawScale * Settings.PointSize);
+			DrawDebugLine(PCOM, PActor, Grey, false, UE_KINDA_SMALL_NUMBER, uint8(Settings.DrawPriority), Settings.LineThickness);
 		
 			if (Index >= 0)
 			{
-				//FDebugDrawQueue::GetInstance().DrawDebugString(PCOM + FontHeight * FVec3(0, 0, 1), FString::Format(TEXT("{0}{1}"), { Particle->IsKinematic()? TEXT("K"): TEXT("D"), Index }), nullptr, FColor::White, KINDA_SMALL_NUMBER, false, FontScale);
+				//DrawDebugString(PCOM + FontHeight * FVec3(0, 0, 1), FString::Format(TEXT("{0}{1}"), { Particle->IsKinematic()? TEXT("K"): TEXT("D"), Index }), nullptr, FColor::White, KINDA_SMALL_NUMBER, false, FontScale);
 			}
 
 			if ((Settings.VelScale > 0.0f) && (Particle->V().Size() > UE_KINDA_SMALL_NUMBER))
 			{
-				FDebugDrawQueue::GetInstance().DrawDebugLine(PCOM, PCOM + SpaceTransform.TransformVector(Particle->V()) * Settings.VelScale, Red, false, UE_KINDA_SMALL_NUMBER, uint8(Settings.DrawPriority), Settings.LineThickness);
-				//FDebugDrawQueue::GetInstance().DrawDebugLine(PCOM, PCOM + SpaceTransform.TransformVector(Particle->VSmooth()) * Settings.VelScale, Blue, false, KINDA_SMALL_NUMBER, uint8(Settings.DrawPriority), Settings.LineThickness);
+				DrawDebugLine(PCOM, PCOM + SpaceTransform.TransformVector(Particle->V()) * Settings.VelScale, Red, false, UE_KINDA_SMALL_NUMBER, uint8(Settings.DrawPriority), Settings.LineThickness);
+				//DrawDebugLine(PCOM, PCOM + SpaceTransform.TransformVector(Particle->VSmooth()) * Settings.VelScale, Blue, false, KINDA_SMALL_NUMBER, uint8(Settings.DrawPriority), Settings.LineThickness);
 			}
 			if ((Settings.AngVelScale > 0.0f) && (Particle->W().Size() > UE_KINDA_SMALL_NUMBER))
 			{
-				FDebugDrawQueue::GetInstance().DrawDebugLine(PCOM, PCOM + SpaceTransform.TransformVector(Particle->W()) * Settings.AngVelScale, Green, false, UE_KINDA_SMALL_NUMBER, uint8(Settings.DrawPriority), Settings.LineThickness);
+				DrawDebugLine(PCOM, PCOM + SpaceTransform.TransformVector(Particle->W()) * Settings.AngVelScale, Green, false, UE_KINDA_SMALL_NUMBER, uint8(Settings.DrawPriority), Settings.LineThickness);
 			}
 
 			if (Settings.InertiaScale > 0.0f)
@@ -1237,7 +1405,7 @@ namespace Chaos
 							if (Utilities::BoxFromInertia(Rigid->I(), Rigid->M(), EquivalentBoxCenter, EquivalentBoxSize))
 							{
 								const FVec3 BoxCoM = SpaceTransform.TransformPositionNoScale(PCOM + QCOM * EquivalentBoxCenter);
-								FDebugDrawQueue::GetInstance().DrawDebugBox(BoxCoM, 0.5f * Settings.InertiaScale * EquivalentBoxSize, QCOM, FColor::Black, false, 0.0f, 0, Settings.LineThickness);
+								DrawDebugBox(BoxCoM, 0.5f * Settings.InertiaScale * EquivalentBoxSize, QCOM, FColor::Black, false, 0.0f, 0, Settings.LineThickness);
 							}
 						}
 
@@ -1246,7 +1414,7 @@ namespace Chaos
 						if (Utilities::BoxFromInertia(Rigid->ConditionedI(), Rigid->M(), EquivalentBoxConditionedCenter, EquivalentBoxConditionedSize))
 						{
 							const FVec3 BoxConditionedCoM = SpaceTransform.TransformPositionNoScale(PCOM + QCOM * EquivalentBoxConditionedCenter);
-							FDebugDrawQueue::GetInstance().DrawDebugBox(BoxConditionedCoM, 0.5f * Settings.InertiaScale * EquivalentBoxConditionedSize, QCOM, FColor::Magenta, false, 0.0f, 0, Settings.LineThickness);
+							DrawDebugBox(BoxConditionedCoM, 0.5f * Settings.InertiaScale * EquivalentBoxConditionedSize, QCOM, FColor::Magenta, false, 0.0f, 0, Settings.LineThickness);
 						}
 					}
 				}
@@ -1268,11 +1436,11 @@ namespace Chaos
 				const FRigidTransform3 WorldActorTransform1 = FParticleUtilities::GetActorWorldTransform(Particle1);
 
 				// Are we within the region of interest?
-				if (Contact.GetParticle0()->HasBounds() && !FDebugDrawQueue::GetInstance().IsInRegionOfInterest(Contact.GetParticle0()->WorldSpaceInflatedBounds().TransformedAABB(SpaceTransform)))
+				if (Contact.GetParticle0()->HasBounds() && !IsInRegionOfInterest(Contact.GetParticle0()->WorldSpaceInflatedBounds().TransformedAABB(SpaceTransform)))
 				{
 					return;
 				}
-				if (Contact.GetParticle1()->HasBounds() && !FDebugDrawQueue::GetInstance().IsInRegionOfInterest(Contact.GetParticle1()->WorldSpaceInflatedBounds().TransformedAABB(SpaceTransform)))
+				if (Contact.GetParticle1()->HasBounds() && !IsInRegionOfInterest(Contact.GetParticle1()->WorldSpaceInflatedBounds().TransformedAABB(SpaceTransform)))
 				{
 					return;
 				}
@@ -1350,31 +1518,31 @@ namespace Chaos
 					if ((Settings.PushOutScale > 0) && ManifoldPointResult.bIsValid && !ManifoldPointResult.NetPushOut.IsNearlyZero())
 					{
 						FColor Color = PushOutImpusleColor;
-						FDebugDrawQueue::GetInstance().DrawDebugLine(WorldPointPlaneLocation, WorldPointPlaneLocation + Settings.DrawScale * Settings.PushOutScale * SpaceTransform.TransformVectorNoScale(FVec3(ManifoldPointResult.NetPushOut)), Color, false, Duration, uint8(Settings.DrawPriority), Settings.LineThickness);
+						DrawDebugLine(WorldPointPlaneLocation, WorldPointPlaneLocation + Settings.DrawScale * Settings.PushOutScale * SpaceTransform.TransformVectorNoScale(FVec3(ManifoldPointResult.NetPushOut)), Color, false, Duration, uint8(Settings.DrawPriority), Settings.LineThickness);
 					}
 					if ((Settings.ImpulseScale > 0) && ManifoldPointResult.bIsValid && !ManifoldPointResult.NetImpulse.IsNearlyZero())
 					{
 						FColor Color = ImpulseColor;
-						FDebugDrawQueue::GetInstance().DrawDebugLine(WorldPointPlaneLocation, WorldPointPlaneLocation + Settings.DrawScale * Settings.ImpulseScale * SpaceTransform.TransformVectorNoScale(FVec3(ManifoldPointResult.NetImpulse)), Color, false, Duration, uint8(Settings.DrawPriority), Settings.LineThickness);
+						DrawDebugLine(WorldPointPlaneLocation, WorldPointPlaneLocation + Settings.DrawScale * Settings.ImpulseScale * SpaceTransform.TransformVectorNoScale(FVec3(ManifoldPointResult.NetImpulse)), Color, false, Duration, uint8(Settings.DrawPriority), Settings.LineThickness);
 					}
 
 					// Manifold point and normal
 					if (Settings.ContactWidth > 0)
 					{
 						FColor C0 = DiscColor;
-						FDebugDrawQueue::GetInstance().DrawDebugCircle(WorldPlaneLocation, Settings.DrawScale * Settings.ContactWidth, 12, C0, false, Duration, uint8(Settings.DrawPriority), Settings.LineThickness, Axes.GetUnitAxis(EAxis::Y), Axes.GetUnitAxis(EAxis::Z), false);
-						FDebugDrawQueue::GetInstance().DrawDebugCircle(WorldPointLocation, 0.5f * Settings.DrawScale * Settings.ContactWidth, 12, C0, false, Duration, uint8(Settings.DrawPriority), Settings.LineThickness, Axes.GetUnitAxis(EAxis::Y), Axes.GetUnitAxis(EAxis::Z), false);
+						DrawDebugCircle(WorldPlaneLocation, Settings.DrawScale * Settings.ContactWidth, 12, C0, false, Duration, uint8(Settings.DrawPriority), Settings.LineThickness, Axes, false);
+						DrawDebugCircle(WorldPointLocation, 0.5f * Settings.DrawScale * Settings.ContactWidth, 12, C0, false, Duration, uint8(Settings.DrawPriority), Settings.LineThickness, Axes, false);
 						if (ManifoldPoint.InitialPhi != 0)
 						{
 							FColor C3 = FColor::Silver;
-							FDebugDrawQueue::GetInstance().DrawDebugCircle(WorldPlaneLocation + ManifoldPoint.InitialPhi * WorldPlaneNormal, 0.25f * Settings.DrawScale * Settings.ContactWidth, 12, C3, false, Duration, uint8(Settings.DrawPriority), Settings.LineThickness, Axes.GetUnitAxis(EAxis::Y), Axes.GetUnitAxis(EAxis::Z), false);
+							DrawDebugCircle(WorldPlaneLocation + ManifoldPoint.InitialPhi * WorldPlaneNormal, 0.25f * Settings.DrawScale * Settings.ContactWidth, 12, C3, false, Duration, uint8(Settings.DrawPriority), Settings.LineThickness, Axes, false);
 						}
 					}
 					if (Settings.ContactLen > 0)
 					{
 						FColor NormalColor = ((ManifoldPoint.ContactPoint.ContactType != EContactPointType::EdgeEdge) ? PlaneNormalColor : EdgeNormalColor);
 						FColor C1 = NormalColor;
-						FDebugDrawQueue::GetInstance().DrawDebugLine(WorldPlaneLocation, WorldPlaneLocation + Settings.DrawScale * Settings.ContactLen * ContactLenScale * WorldPlaneNormal, C1, false, Duration, uint8(Settings.DrawPriority), Settings.LineThickness);
+						DrawDebugLine(WorldPlaneLocation, WorldPlaneLocation + Settings.DrawScale * Settings.ContactLen * ContactLenScale * WorldPlaneNormal, C1, false, Duration, uint8(Settings.DrawPriority), Settings.LineThickness);
 					}
 
 					// Whether restored
@@ -1383,17 +1551,17 @@ namespace Chaos
 						if (Contact.WasManifoldRestored())
 						{
 							const FReal BoxScale = Settings.DrawScale * Settings.ContactInfoWidth;
-							FDebugDrawQueue::GetInstance().DrawDebugBox(WorldPlaneLocation, FVec3(BoxScale, BoxScale, FReal(0.01)), FRotation3(FRotationMatrix::MakeFromZ(WorldPlaneNormal)), FColor::Blue, false, Duration, uint8(Settings.DrawPriority), 0.5f * Settings.LineThickness);
+							DrawDebugBox(WorldPlaneLocation, FVec3(BoxScale, BoxScale, FReal(0.01)), FRotation3(FRotationMatrix::MakeFromZ(WorldPlaneNormal)), FColor::Blue, false, Duration, uint8(Settings.DrawPriority), 0.5f * Settings.LineThickness);
 						}
 						else if (ManifoldPoint.Flags.bWasRestored)
 						{
 							const FReal BoxScale = Settings.DrawScale * Settings.ContactInfoWidth;
-							FDebugDrawQueue::GetInstance().DrawDebugBox(WorldPlaneLocation, FVec3(BoxScale, BoxScale, FReal(0.01)), FRotation3(FRotationMatrix::MakeFromZ(WorldPlaneNormal)), FColor::Purple, false, Duration, uint8(Settings.DrawPriority), 0.5f * Settings.LineThickness);
+							DrawDebugBox(WorldPlaneLocation, FVec3(BoxScale, BoxScale, FReal(0.01)), FRotation3(FRotationMatrix::MakeFromZ(WorldPlaneNormal)), FColor::Purple, false, Duration, uint8(Settings.DrawPriority), 0.5f * Settings.LineThickness);
 						}
 						else if (ManifoldPoint.Flags.bWasReplaced)
 						{
 							const FReal BoxScale = Settings.DrawScale * Settings.ContactInfoWidth;
-							FDebugDrawQueue::GetInstance().DrawDebugBox(WorldPlaneLocation, FVec3(BoxScale, BoxScale, FReal(0.01)), FRotation3(FRotationMatrix::MakeFromZ(WorldPlaneNormal)), FColor::Orange, false, Duration, uint8(Settings.DrawPriority), 0.5f * Settings.LineThickness);
+							DrawDebugBox(WorldPlaneLocation, FVec3(BoxScale, BoxScale, FReal(0.01)), FRotation3(FRotationMatrix::MakeFromZ(WorldPlaneNormal)), FColor::Orange, false, Duration, uint8(Settings.DrawPriority), 0.5f * Settings.LineThickness);
 						}
 					}
 
@@ -1403,7 +1571,7 @@ namespace Chaos
 						if (Contact.IsSleeping())
 						{
 							const FReal BoxScale = Settings.DrawScale * Settings.ContactInfoWidth * 1.1f;
-							FDebugDrawQueue::GetInstance().DrawDebugBox(WorldPlaneLocation, FVec3(BoxScale, BoxScale, FReal(0.01)), FRotation3(FRotationMatrix::MakeFromZ(WorldPlaneNormal)), FColor::Black, false, Duration, uint8(Settings.DrawPriority), 0.5f * Settings.LineThickness);
+							DrawDebugBox(WorldPlaneLocation, FVec3(BoxScale, BoxScale, FReal(0.01)), FRotation3(FRotationMatrix::MakeFromZ(WorldPlaneNormal)), FColor::Black, false, Duration, uint8(Settings.DrawPriority), 0.5f * Settings.LineThickness);
 						}
 					}
 				}
@@ -1413,7 +1581,7 @@ namespace Chaos
 				{
 					FColor Color = FColor::White;
 					const FVec3 ImpulsePos = SpaceTransform.TransformPosition(WorldActorTransform0.GetTranslation());
-					FDebugDrawQueue::GetInstance().DrawDebugLine(ImpulsePos, ImpulsePos + Settings.DrawScale * Settings.ImpulseScale * SpaceTransform.TransformVectorNoScale(Contact.GetAccumulatedImpulse()), Color, false, Duration, uint8(Settings.DrawPriority), Settings.LineThickness);
+					DrawDebugLine(ImpulsePos, ImpulsePos + Settings.DrawScale * Settings.ImpulseScale * SpaceTransform.TransformVectorNoScale(Contact.GetAccumulatedImpulse()), Color, false, Duration, uint8(Settings.DrawPriority), Settings.LineThickness);
 				}
 
 				// Static Friction Anchors (bets viewed with DebugDrawPostIntegrateCollisions)
@@ -1430,7 +1598,7 @@ namespace Chaos
 						const FVec3 PlaneLocation = PlaneTransform.TransformPosition(FVec3(SavedManifoldPoint.ShapeContactPoints[ContactPlaneOwner]));
 				
 						FColor Color = FColor::Black;
-						FDebugDrawQueue::GetInstance().DrawDebugLine(PointLocation, PlaneLocation, Color, false, Duration, uint8(Settings.DrawPriority), Settings.LineThickness);
+						DrawDebugLine(PointLocation, PlaneLocation, Color, false, Duration, uint8(Settings.DrawPriority), Settings.LineThickness);
 					}
 				}
 
@@ -1450,8 +1618,8 @@ namespace Chaos
 				const FMatrix Axes = FRotationMatrix::MakeFromX(Normal);
 				const FVec3 P0 = SpaceTransform.TransformPosition(Contact.GetShape0()->GetWorldSpaceShapeBounds().GetCenter());
 				const FVec3 P1 = SpaceTransform.TransformPosition(Contact.GetShape1()->GetWorldSpaceShapeBounds().GetCenter());
-				FDebugDrawQueue::GetInstance().DrawDebugLine(Location, P0, C3, false, UE_KINDA_SMALL_NUMBER, uint8(Settings.DrawPriority), Settings.LineThickness * 0.5f);
-				FDebugDrawQueue::GetInstance().DrawDebugLine(Location, P1, C3, false, UE_KINDA_SMALL_NUMBER, uint8(Settings.DrawPriority), Settings.LineThickness * 0.5f);
+				DrawDebugLine(Location, P0, C3, false, UE_KINDA_SMALL_NUMBER, uint8(Settings.DrawPriority), Settings.LineThickness * 0.5f);
+				DrawDebugLine(Location, P1, C3, false, UE_KINDA_SMALL_NUMBER, uint8(Settings.DrawPriority), Settings.LineThickness * 0.5f);
 			}
 		}
 		
@@ -1470,7 +1638,7 @@ namespace Chaos
 					if (!Impulse.IsNearlyZero() && (Settings.ImpulseScale > 0))
 					{
 						const FVec3 Pos0 = SpaceTransform.TransformPosition(P0->PCom());
-						FDebugDrawQueue::GetInstance().DrawDebugLine(Pos0, Pos0 + Settings.ImpulseScale * Impulse, FColor::Red, false, Duration, uint8(Settings.DrawPriority), Settings.LineThickness * 0.5f);
+						DrawDebugLine(Pos0, Pos0 + Settings.ImpulseScale * Impulse, FColor::Red, false, Duration, uint8(Settings.DrawPriority), Settings.LineThickness * 0.5f);
 					}
 				}
 			}
@@ -1604,7 +1772,7 @@ namespace Chaos
 					}
 					const FReal RoundedMass = 0.1 * FMath::RoundToDouble(Mass * FormatMultiplier * 10.0);
 					const FString MassStr = FString::Format(FormatString, { FString::SanitizeFloat(RoundedMass, 1) });
-					FDebugDrawQueue::GetInstance().DrawDebugString(PCOM, MassStr, nullptr, Color, UE_KINDA_SMALL_NUMBER, false, Settings.FontScale);
+					DrawDebugString(PCOM, MassStr, nullptr, Color, UE_KINDA_SMALL_NUMBER, false, Settings.FontScale);
 				}
 			}
 		}
@@ -1701,10 +1869,10 @@ namespace Chaos
 					{
 						Sb = FMath::Lerp(Pb, Xb, FMath::Clamp<FReal>(CoMSize / Lena, 0., 1.));
 					}
-					FDebugDrawQueue::GetInstance().DrawDebugLine(Pa, Sa, FColor::White, false, UE_KINDA_SMALL_NUMBER, uint8(Settings.DrawPriority), ConnectorThickness);
-					FDebugDrawQueue::GetInstance().DrawDebugLine(Pb, Sb, FColor::White, false, UE_KINDA_SMALL_NUMBER, uint8(Settings.DrawPriority), ConnectorThickness);
-					FDebugDrawQueue::GetInstance().DrawDebugLine(Sa, Xa, R, false, UE_KINDA_SMALL_NUMBER, uint8(Settings.DrawPriority), ConnectorThickness);
-					FDebugDrawQueue::GetInstance().DrawDebugLine(Sb, Xb, C, false, UE_KINDA_SMALL_NUMBER, uint8(Settings.DrawPriority), ConnectorThickness);
+					DrawDebugLine(Pa, Sa, FColor::White, false, UE_KINDA_SMALL_NUMBER, uint8(Settings.DrawPriority), ConnectorThickness);
+					DrawDebugLine(Pb, Sb, FColor::White, false, UE_KINDA_SMALL_NUMBER, uint8(Settings.DrawPriority), ConnectorThickness);
+					DrawDebugLine(Sa, Xa, R, false, UE_KINDA_SMALL_NUMBER, uint8(Settings.DrawPriority), ConnectorThickness);
+					DrawDebugLine(Sb, Xb, C, false, UE_KINDA_SMALL_NUMBER, uint8(Settings.DrawPriority), ConnectorThickness);
 				}
 				if (FeatureMask.bCoMConnector)
 				{
@@ -1723,24 +1891,24 @@ namespace Chaos
 					{
 						Sb = FMath::Lerp(Cb, Xb, FMath::Clamp<FReal>(CoMSize / Lena, 0., 1.));
 					}
-					FDebugDrawQueue::GetInstance().DrawDebugLine(Ca, Sa, FColor::Black, false, UE_KINDA_SMALL_NUMBER, uint8(Settings.DrawPriority), ConnectorThickness);
-					FDebugDrawQueue::GetInstance().DrawDebugLine(Cb, Sb, FColor::Black, false, UE_KINDA_SMALL_NUMBER, uint8(Settings.DrawPriority), ConnectorThickness);
-					FDebugDrawQueue::GetInstance().DrawDebugLine(Sa, Xa, R, false, UE_KINDA_SMALL_NUMBER, uint8(Settings.DrawPriority), ConnectorThickness);
-					FDebugDrawQueue::GetInstance().DrawDebugLine(Sb, Xb, C, false, UE_KINDA_SMALL_NUMBER, uint8(Settings.DrawPriority), ConnectorThickness);
+					DrawDebugLine(Ca, Sa, FColor::Black, false, UE_KINDA_SMALL_NUMBER, uint8(Settings.DrawPriority), ConnectorThickness);
+					DrawDebugLine(Cb, Sb, FColor::Black, false, UE_KINDA_SMALL_NUMBER, uint8(Settings.DrawPriority), ConnectorThickness);
+					DrawDebugLine(Sa, Xa, R, false, UE_KINDA_SMALL_NUMBER, uint8(Settings.DrawPriority), ConnectorThickness);
+					DrawDebugLine(Sb, Xb, C, false, UE_KINDA_SMALL_NUMBER, uint8(Settings.DrawPriority), ConnectorThickness);
 				}
 				if (FeatureMask.bStretch)
 				{
 					const FRealSingle StretchThickness = 3.0f * Settings.LineThickness;
-					FDebugDrawQueue::GetInstance().DrawDebugLine(Xa, Xb, M, false, UE_KINDA_SMALL_NUMBER, uint8(Settings.DrawPriority), StretchThickness);
+					DrawDebugLine(Xa, Xb, M, false, UE_KINDA_SMALL_NUMBER, uint8(Settings.DrawPriority), StretchThickness);
 				}
 				if (FeatureMask.bAxes)
 				{
-					FDebugDrawQueue::GetInstance().DrawDebugDirectionalArrow(Xa, Xa + Settings.DrawScale * Settings.ConstraintAxisLen * SpaceTransform.TransformVector(Ra.GetAxis(0)), Settings.DrawScale * Settings.ArrowSize, R, false, UE_KINDA_SMALL_NUMBER, uint8(Settings.DrawPriority), Settings.LineThickness);
-					FDebugDrawQueue::GetInstance().DrawDebugDirectionalArrow(Xa, Xa + Settings.DrawScale * Settings.ConstraintAxisLen * SpaceTransform.TransformVector(Ra.GetAxis(1)), Settings.DrawScale * Settings.ArrowSize, G, false, UE_KINDA_SMALL_NUMBER, uint8(Settings.DrawPriority), Settings.LineThickness);
-					FDebugDrawQueue::GetInstance().DrawDebugDirectionalArrow(Xa, Xa + Settings.DrawScale * Settings.ConstraintAxisLen * SpaceTransform.TransformVector(Ra.GetAxis(2)), Settings.DrawScale * Settings.ArrowSize, B, false, UE_KINDA_SMALL_NUMBER, uint8(Settings.DrawPriority), Settings.LineThickness);
-					FDebugDrawQueue::GetInstance().DrawDebugDirectionalArrow(Xb, Xb + Settings.DrawScale * Settings.ConstraintAxisLen * SpaceTransform.TransformVector(Rb.GetAxis(0)), Settings.DrawScale * Settings.ArrowSize, C, false, UE_KINDA_SMALL_NUMBER, uint8(Settings.DrawPriority), Settings.LineThickness);
-					FDebugDrawQueue::GetInstance().DrawDebugDirectionalArrow(Xb, Xb + Settings.DrawScale * Settings.ConstraintAxisLen * SpaceTransform.TransformVector(Rb.GetAxis(1)), Settings.DrawScale * Settings.ArrowSize, M, false, UE_KINDA_SMALL_NUMBER, uint8(Settings.DrawPriority), Settings.LineThickness);
-					FDebugDrawQueue::GetInstance().DrawDebugDirectionalArrow(Xb, Xb + Settings.DrawScale * Settings.ConstraintAxisLen * SpaceTransform.TransformVector(Rb.GetAxis(2)), Settings.DrawScale * Settings.ArrowSize, Y, false, UE_KINDA_SMALL_NUMBER, uint8(Settings.DrawPriority), Settings.LineThickness);
+					DrawDebugDirectionalArrow(Xa, Xa + Settings.DrawScale * Settings.ConstraintAxisLen * SpaceTransform.TransformVector(Ra.GetAxis(0)), Settings.DrawScale * Settings.ArrowSize, R, false, UE_KINDA_SMALL_NUMBER, uint8(Settings.DrawPriority), Settings.LineThickness);
+					DrawDebugDirectionalArrow(Xa, Xa + Settings.DrawScale * Settings.ConstraintAxisLen * SpaceTransform.TransformVector(Ra.GetAxis(1)), Settings.DrawScale * Settings.ArrowSize, G, false, UE_KINDA_SMALL_NUMBER, uint8(Settings.DrawPriority), Settings.LineThickness);
+					DrawDebugDirectionalArrow(Xa, Xa + Settings.DrawScale * Settings.ConstraintAxisLen * SpaceTransform.TransformVector(Ra.GetAxis(2)), Settings.DrawScale * Settings.ArrowSize, B, false, UE_KINDA_SMALL_NUMBER, uint8(Settings.DrawPriority), Settings.LineThickness);
+					DrawDebugDirectionalArrow(Xb, Xb + Settings.DrawScale * Settings.ConstraintAxisLen * SpaceTransform.TransformVector(Rb.GetAxis(0)), Settings.DrawScale * Settings.ArrowSize, C, false, UE_KINDA_SMALL_NUMBER, uint8(Settings.DrawPriority), Settings.LineThickness);
+					DrawDebugDirectionalArrow(Xb, Xb + Settings.DrawScale * Settings.ConstraintAxisLen * SpaceTransform.TransformVector(Rb.GetAxis(1)), Settings.DrawScale * Settings.ArrowSize, M, false, UE_KINDA_SMALL_NUMBER, uint8(Settings.DrawPriority), Settings.LineThickness);
+					DrawDebugDirectionalArrow(Xb, Xb + Settings.DrawScale * Settings.ConstraintAxisLen * SpaceTransform.TransformVector(Rb.GetAxis(2)), Settings.DrawScale * Settings.ArrowSize, Y, false, UE_KINDA_SMALL_NUMBER, uint8(Settings.DrawPriority), Settings.LineThickness);
 				}
 
 				// NOTE: GetLinearImpulse is the positional impulse (pushout)
@@ -1748,7 +1916,7 @@ namespace Chaos
 				{
 					FColor PushOutImpusleColor = FColor(0, 250, 250);
 					FColor Color = PushOutImpusleColor;
-					FDebugDrawQueue::GetInstance().DrawDebugLine(Xa, Xa + Settings.DrawScale * Settings.PushOutScale * SpaceTransform.TransformVectorNoScale(FVec3(ConstraintHandle->GetLinearImpulse())), Color, false, 0, uint8(Settings.DrawPriority), Settings.LineThickness);
+					DrawDebugLine(Xa, Xa + Settings.DrawScale * Settings.PushOutScale * SpaceTransform.TransformVectorNoScale(FVec3(ConstraintHandle->GetLinearImpulse())), Color, false, 0, uint8(Settings.DrawPriority), Settings.LineThickness);
 				}
 			}
 		}
@@ -1784,14 +1952,14 @@ namespace Chaos
 			const FVec3 Pos = SimSpace.Transform.GetLocation();
 			const FRotation3& Rot = SimSpace.Transform.GetRotation();
 			const FMatrix33 Rotm = Rot.ToMatrix();
-			FDebugDrawQueue::GetInstance().DrawDebugDirectionalArrow(Pos, Pos + Settings.DrawScale * Settings.BodyAxisLen * Rotm.GetAxis(0), Settings.DrawScale * Settings.ArrowSize, FColor::Red, false, UE_KINDA_SMALL_NUMBER, uint8(Settings.DrawPriority), Settings.LineThickness);
-			FDebugDrawQueue::GetInstance().DrawDebugDirectionalArrow(Pos, Pos + Settings.DrawScale * Settings.BodyAxisLen * Rotm.GetAxis(1), Settings.DrawScale * Settings.ArrowSize, FColor::Green, false, UE_KINDA_SMALL_NUMBER, uint8(Settings.DrawPriority), Settings.LineThickness);
-			FDebugDrawQueue::GetInstance().DrawDebugDirectionalArrow(Pos, Pos + Settings.DrawScale * Settings.BodyAxisLen * Rotm.GetAxis(2), Settings.DrawScale * Settings.ArrowSize, FColor::Blue, false, UE_KINDA_SMALL_NUMBER, uint8(Settings.DrawPriority), Settings.LineThickness);
+			DrawDebugDirectionalArrow(Pos, Pos + Settings.DrawScale * Settings.BodyAxisLen * Rotm.GetAxis(0), Settings.DrawScale * Settings.ArrowSize, FColor::Red, false, UE_KINDA_SMALL_NUMBER, uint8(Settings.DrawPriority), Settings.LineThickness);
+			DrawDebugDirectionalArrow(Pos, Pos + Settings.DrawScale * Settings.BodyAxisLen * Rotm.GetAxis(1), Settings.DrawScale * Settings.ArrowSize, FColor::Green, false, UE_KINDA_SMALL_NUMBER, uint8(Settings.DrawPriority), Settings.LineThickness);
+			DrawDebugDirectionalArrow(Pos, Pos + Settings.DrawScale * Settings.BodyAxisLen * Rotm.GetAxis(2), Settings.DrawScale * Settings.ArrowSize, FColor::Blue, false, UE_KINDA_SMALL_NUMBER, uint8(Settings.DrawPriority), Settings.LineThickness);
 
-			FDebugDrawQueue::GetInstance().DrawDebugLine(Pos, Pos + Settings.VelScale * SimSpace.LinearVelocity, FColor::Cyan, false, UE_KINDA_SMALL_NUMBER, uint8(Settings.DrawPriority), Settings.LineThickness);
-			FDebugDrawQueue::GetInstance().DrawDebugLine(Pos, Pos + Settings.AngVelScale * SimSpace.AngularVelocity, FColor::Cyan, false, UE_KINDA_SMALL_NUMBER, uint8(Settings.DrawPriority), Settings.LineThickness);
-			FDebugDrawQueue::GetInstance().DrawDebugLine(Pos, Pos + 0.01f * Settings.VelScale * SimSpace.LinearAcceleration, FColor::Yellow, false, UE_KINDA_SMALL_NUMBER, uint8(Settings.DrawPriority), Settings.LineThickness);
-			FDebugDrawQueue::GetInstance().DrawDebugLine(Pos, Pos + 0.01f * Settings.AngVelScale * SimSpace.AngularAcceleration, FColor::Orange, false, UE_KINDA_SMALL_NUMBER, uint8(Settings.DrawPriority), Settings.LineThickness);
+			DrawDebugLine(Pos, Pos + Settings.VelScale * SimSpace.LinearVelocity, FColor::Cyan, false, UE_KINDA_SMALL_NUMBER, uint8(Settings.DrawPriority), Settings.LineThickness);
+			DrawDebugLine(Pos, Pos + Settings.AngVelScale * SimSpace.AngularVelocity, FColor::Cyan, false, UE_KINDA_SMALL_NUMBER, uint8(Settings.DrawPriority), Settings.LineThickness);
+			DrawDebugLine(Pos, Pos + 0.01f * Settings.VelScale * SimSpace.LinearAcceleration, FColor::Yellow, false, UE_KINDA_SMALL_NUMBER, uint8(Settings.DrawPriority), Settings.LineThickness);
+			DrawDebugLine(Pos, Pos + 0.01f * Settings.AngVelScale * SimSpace.AngularAcceleration, FColor::Orange, false, UE_KINDA_SMALL_NUMBER, uint8(Settings.DrawPriority), Settings.LineThickness);
 		}
 
 		void DrawConstraintGraphImpl(const FRigidTransform3& SpaceTransform, const Private::FPBDIslandManager& Graph, const FChaosDebugDrawSettings& Settings)
@@ -1818,30 +1986,30 @@ namespace Chaos
 				if ((bChaosDebugDebugDrawContactGraphUsed && bIsUsed) || (bChaosDebugDebugDrawContactGraphUnused && !bIsUsed))
 				{
 					FColor Color = bIsUsed ? FColor::Green : FColor::Red;
-					FDebugDrawQueue::GetInstance().DrawDebugLine(Transform0.GetLocation(), Transform1.GetLocation(), Color, false, UE_KINDA_SMALL_NUMBER, uint8(Settings.DrawPriority), Settings.LineThickness);
+					DrawDebugLine(Transform0.GetLocation(), Transform1.GetLocation(), Color, false, UE_KINDA_SMALL_NUMBER, uint8(Settings.DrawPriority), Settings.LineThickness);
 				}
 
 				if (bChaosDebugDebugDrawContactGraph)
 				{
 					if (FConstGenericParticleHandle(Constraint->GetParticle0())->IsDynamic())
 					{
-						FDebugDrawQueue::GetInstance().DrawDebugLine(Transform0.GetLocation(), ContactPos, FColor::Red, false, UE_KINDA_SMALL_NUMBER, uint8(Settings.DrawPriority), Settings.LineThickness);
+						DrawDebugLine(Transform0.GetLocation(), ContactPos, FColor::Red, false, UE_KINDA_SMALL_NUMBER, uint8(Settings.DrawPriority), Settings.LineThickness);
 					}
 					if (FConstGenericParticleHandle(Constraint->GetParticle1())->IsDynamic())
 					{
-						FDebugDrawQueue::GetInstance().DrawDebugLine(Transform1.GetLocation(), ContactPos, FColor::White, false, UE_KINDA_SMALL_NUMBER, uint8(Settings.DrawPriority), Settings.LineThickness);
+						DrawDebugLine(Transform1.GetLocation(), ContactPos, FColor::White, false, UE_KINDA_SMALL_NUMBER, uint8(Settings.DrawPriority), Settings.LineThickness);
 					}
 					if (ColorIndex >= 0)
 					{
-						FDebugDrawQueue::GetInstance().DrawDebugString(ContactPos, FString::Format(TEXT("{0} L{1} C{2}"), { OrderIndex, LevelIndex, ColorIndex }), nullptr, FColor::Yellow, UE_KINDA_SMALL_NUMBER, false, Settings.FontScale);
+						DrawDebugString(ContactPos, FString::Format(TEXT("{0} L{1} C{2}"), { OrderIndex, LevelIndex, ColorIndex }), nullptr, FColor::Yellow, UE_KINDA_SMALL_NUMBER, false, Settings.FontScale);
 					}
 					else if (LevelIndex >= 0)
 					{
-						FDebugDrawQueue::GetInstance().DrawDebugString(ContactPos, FString::Format(TEXT("{0} L{1}"), { OrderIndex, LevelIndex }), nullptr, FColor::Yellow, UE_KINDA_SMALL_NUMBER, false, Settings.FontScale);
+						DrawDebugString(ContactPos, FString::Format(TEXT("{0} L{1}"), { OrderIndex, LevelIndex }), nullptr, FColor::Yellow, UE_KINDA_SMALL_NUMBER, false, Settings.FontScale);
 					}
 					else
 					{
-						FDebugDrawQueue::GetInstance().DrawDebugString(ContactPos, FString::Format(TEXT("{0}"), { OrderIndex }), nullptr, FColor::Yellow, UE_KINDA_SMALL_NUMBER, false, Settings.FontScale);
+						DrawDebugString(ContactPos, FString::Format(TEXT("{0}"), { OrderIndex }), nullptr, FColor::Yellow, UE_KINDA_SMALL_NUMBER, false, Settings.FontScale);
 					}
 				}
 			};
@@ -1861,7 +2029,7 @@ namespace Chaos
 
 					const FColor IslandColor = GetIslandColor(IslandIndex, !Graph.GetIsland(IslandIndex)->IsSleeping());
 					const FAABB3 Bounds = IslandAABB.TransformedAABB(SpaceTransform);
-					FDebugDrawQueue::GetInstance().DrawDebugBox(Bounds.Center(), 0.5f * Bounds.Extents(), SpaceTransform.GetRotation(), IslandColor, false, UE_KINDA_SMALL_NUMBER, uint8(Settings.DrawPriority), 3.0f * Settings.LineThickness);
+					DrawDebugBox(Bounds.Center(), 0.5f * Bounds.Extents(), SpaceTransform.GetRotation(), IslandColor, false, UE_KINDA_SMALL_NUMBER, uint8(Settings.DrawPriority), 3.0f * Settings.LineThickness);
 				}
 			}
 
@@ -1972,13 +2140,13 @@ namespace Chaos
 
 					FColor Color = (Connection.Handle0->Parent() == Connection.Handle1->Parent()) ? FColor::Green : FColor::Blue; // FMath::Lerp(FColor::Green, FColor::Red, Connection.Strain / MaxStrain);
 
-					FDebugDrawQueue::GetInstance().DrawDebugLine(Pos0, Pos1, Color, false, UE_KINDA_SMALL_NUMBER, uint8(Settings.DrawPriority), Settings.LineThickness);
+					DrawDebugLine(Pos0, Pos1, Color, false, UE_KINDA_SMALL_NUMBER, uint8(Settings.DrawPriority), Settings.LineThickness);
 
 					if (bChaosDebugDrawConnectionGraphShowAreas)
 					{
 						const FReal Radius = FMath::Sqrt(FMath::Abs(Connection.Strain / UE_PI)); // represent area as a circle
 						const FMatrix Axes = FRotationMatrix::MakeFromX(Normal);
-						FDebugDrawQueue::GetInstance().DrawDebugCircle((Pos0 + Pos1) * 0.5, Radius, 24, Color, false, UE_KINDA_SMALL_NUMBER, uint8(Settings.DrawPriority), Settings.LineThickness, Axes.GetUnitAxis(EAxis::Y), Axes.GetUnitAxis(EAxis::Z), false);
+						DrawDebugCircle((Pos0 + Pos1) * 0.5, Radius, 24, Color, false, UE_KINDA_SMALL_NUMBER, uint8(Settings.DrawPriority), Settings.LineThickness, Axes, false);
 					}
 
 
@@ -2060,11 +2228,11 @@ namespace Chaos
 						{
 							PointColor = Settings.ShapesColorsPerState.KinematicColor;
 						}
-						FDebugDrawQueue::GetInstance().DrawDebugPoint(GetClusterTransform(Child).GetLocation(), PointColor, false, UE_KINDA_SMALL_NUMBER, uint8(Settings.DrawPriority), Settings.LineThickness * 10);
+						DrawDebugPoint(GetClusterTransform(Child).GetLocation(), PointColor, false, UE_KINDA_SMALL_NUMBER, uint8(Settings.DrawPriority), Settings.LineThickness * 10);
 
 						if (bChaosDebugDrawConnectionGraphShowInternalStrains)
 						{
-							FDebugDrawQueue::GetInstance().DrawDebugString(Child->GetX(), FString::Printf(TEXT("%.1f"), Child->GetInternalStrains()), nullptr, FColor::White, UE_KINDA_SMALL_NUMBER, false, 1.0);
+							DrawDebugString(Child->GetX(), FString::Printf(TEXT("%.1f"), Child->GetInternalStrains()), nullptr, FColor::White, UE_KINDA_SMALL_NUMBER, false, 1.0);
 						}
 
 						const FConnectivityEdgeArray& Edges = Child->ConnectivityEdges();
@@ -2102,7 +2270,7 @@ namespace Chaos
 			const FVec3 HardStopPushOutWorld = ConstraintResults.HardStopNetPushOut;
 			const FVec3 HardStopImpulseWorld = ConstraintResults.HardStopNetImpulse;
 
-			FDebugDrawQueue::GetInstance().DrawDebugLine(
+			DrawDebugLine(
 				SpaceTransform.TransformPosition(PWorld), 
 				SpaceTransform.TransformPosition(PWorld + AxisLen * AxisWorld), 
 				FColor::Green, false, UE_KINDA_SMALL_NUMBER, uint8(Settings.DrawPriority), Settings.LineThickness);
@@ -2118,12 +2286,12 @@ namespace Chaos
 				const FColor HardStopPushOutColor = FColor(0, 200, 100);
 				const FColor SpringPushOutColor = FColor(0, 200, 0);
 
-				FDebugDrawQueue::GetInstance().DrawDebugLine(
+				DrawDebugLine(
 					SpaceTransform.TransformPosition(PWorld),
 					SpaceTransform.TransformPosition(PWorld + Settings.PushOutScale * PushOutWorld),
 					SpringPushOutColor, false, UE_KINDA_SMALL_NUMBER, uint8(Settings.DrawPriority), SpringWidthScale * Settings.LineThickness);
 
-				FDebugDrawQueue::GetInstance().DrawDebugLine(
+				DrawDebugLine(
 					SpaceTransform.TransformPosition(PWorld),
 					SpaceTransform.TransformPosition(PWorld + Settings.PushOutScale * HardStopPushOutWorld),
 					HardStopPushOutColor, false, UE_KINDA_SMALL_NUMBER, uint8(Settings.DrawPriority), HardStopWidthScale * Settings.LineThickness);
@@ -2133,7 +2301,7 @@ namespace Chaos
 			{
 				const FColor HardStopImpulseColor = FColor(200, 200, 0);
 
-				FDebugDrawQueue::GetInstance().DrawDebugLine(
+				DrawDebugLine(
 					SpaceTransform.TransformPosition(PWorld),
 					SpaceTransform.TransformPosition(PWorld + Settings.PushOutScale * HardStopImpulseWorld),
 					HardStopImpulseColor, false, UE_KINDA_SMALL_NUMBER, uint8(Settings.DrawPriority), Settings.LineThickness * FRealSingle(0.5));
@@ -2449,12 +2617,12 @@ namespace Chaos
 
 			virtual void Box(const FAABB3& InBox, const TVec3<FReal>& InLinearColor, float InThickness) override
 			{
-				FDebugDrawQueue::GetInstance().DrawDebugBox(InBox.Center(), InBox.Extents() * FReal(0.5), FQuat::Identity, FLinearColor(InLinearColor).ToFColor(false), false, -1.f, uint8(Settings.DrawPriority), InThickness);
+				DrawDebugBox(InBox.Center(), InBox.Extents() * FReal(0.5), FQuat::Identity, FLinearColor(InLinearColor).ToFColor(false), false, -1.f, uint8(Settings.DrawPriority), InThickness);
 			}
 
 			virtual void Line(const TVec3<FReal>& InBegin, const TVec3<FReal>& InEnd, const TVec3<FReal>& InLinearColor, float InThickness) override
 			{
-				FDebugDrawQueue::GetInstance().DrawDebugLine(InBegin, InEnd, FLinearColor(InLinearColor).ToFColor(false), false, -1.f, uint8(Settings.DrawPriority), InThickness);
+				DrawDebugLine(InBegin, InEnd, FLinearColor(InLinearColor).ToFColor(false), false, -1.f, uint8(Settings.DrawPriority), InThickness);
 			}
 		private:
 			FChaosDebugDrawSettings Settings;
@@ -2526,9 +2694,9 @@ namespace Chaos
 				AxisLocalDelta[AxisIndex] = AbsLocalDelta[AxisIndex];
 				const FVector AxisExtent = R.RotateVector(AxisLocalExtent);
 				const FVector AxisDelta = R.RotateVector(AxisLocalDelta);
-				FDebugDrawQueue::GetInstance().DrawDebugDirectionalArrow(X, X + AxisExtent, ArrowSize, AxisColor, false, -1.f, -1, AxisThickness);
-				FDebugDrawQueue::GetInstance().DrawDebugDirectionalArrow(X, X + AxisDelta, ArrowSize, DeltaColor, false, -1.f, -1, DeltaThickness);
-				FDebugDrawQueue::GetInstance().DrawDebugString(X + AxisExtent, FString::Printf(TEXT("%s: %.2f%%"), *AxisStrings[AxisIndex], AxisPercents[AxisIndex] * 100.f), nullptr, AxisColor, -1.f, true, 1.f);
+				DrawDebugDirectionalArrow(X, X + AxisExtent, ArrowSize, AxisColor, false, -1.f, -1, AxisThickness);
+				DrawDebugDirectionalArrow(X, X + AxisDelta, ArrowSize, DeltaColor, false, -1.f, -1, DeltaThickness);
+				DrawDebugString(X + AxisExtent, FString::Printf(TEXT("%s: %.2f%%"), *AxisStrings[AxisIndex], AxisPercents[AxisIndex] * 100.f), nullptr, AxisColor, -1.f, true, 1.f);
 			};
 
 			// Get the index of the largest element
@@ -2547,8 +2715,8 @@ namespace Chaos
 			// and end position of the CoM over this frame (gray = no ccd, orange = ccd enabled)
 			// Also draw a position delta vector (red) and the threshold delta vector (green)
 			const FVector HalfExtents = AxisThreshold * .5f;
-			FDebugDrawQueue::GetInstance().DrawDebugBox(X, HalfExtents, R, BoundsColor, false, -1.f, -1, BoundsThickness);
-			FDebugDrawQueue::GetInstance().DrawDebugBox(X + DeltaX, HalfExtents, R, BoundsColor, false, -1.f, -1, BoundsThickness);
+			DrawDebugBox(X, HalfExtents, R, BoundsColor, false, -1.f, -1, BoundsThickness);
+			DrawDebugBox(X + DeltaX, HalfExtents, R, BoundsColor, false, -1.f, -1, BoundsThickness);
 			const float ScaleArray[] = { -1.f, 1.f };
 			for (uint8 IX = 0; IX < 2; ++IX)
 			{
@@ -2558,7 +2726,7 @@ namespace Chaos
 					{
 						const FVector LocalOffset = HalfExtents * FVector(ScaleArray[IX], ScaleArray[IY], ScaleArray[IZ]);
 						const FVector Offset = R.RotateVector(LocalOffset);
-						FDebugDrawQueue::GetInstance().DrawDebugLine(X + Offset, X + Offset + DeltaX, BoundsColor, false, -1.f, -1, BoundsThickness);
+						DrawDebugLine(X + Offset, X + Offset + DeltaX, BoundsColor, false, -1.f, -1, BoundsThickness);
 					}
 				}
 			}
diff --git a/Engine/Source/Runtime/Experimental/Chaos/Private/PBDRigidsSolver.cpp b/Engine/Source/Runtime/Experimental/Chaos/Private/PBDRigidsSolver.cpp
--- a/Engine/Source/Runtime/Experimental/Chaos/Private/PBDRigidsSolver.cpp
+++ b/Engine/Source/Runtime/Experimental/Chaos/Private/PBDRigidsSolver.cpp
@@ -30,6 +30,7 @@
 #include "Chaos/ConvexOptimizer.h"
 #include "Chaos/AsyncInitBodyHelper.h"
 
+#include "Chaos/DebugDraw/DebugDrawParticle.h"
 #include "ChaosDebugDraw/ChaosDDContext.h"
 #include "ChaosDebugDraw/ChaosDDScene.h"
 #include "ChaosDebugDraw/ChaosDDTimeline.h"
@@ -114,6 +115,8 @@ TRACE_DECLARE_INT_COUNTER(ChaosTraceCounter_NumProjectionIterations, TEXT("Chaos
 
 namespace Chaos
 {
+	extern CHAOS_API bool bChaosDebugDraw_UseNewQueue;
+
 	namespace DebugDraw
 	{
 		extern const FChaosDebugDrawColorsByState& GetDefaultShapesColorsPreIntegrate();
@@ -2461,9 +2464,22 @@ TRACE_COUNTER_SET(ChaosTraceCounter_##Name, Value)
 
 		if (ChaosSolverDebugDrawPreIntegrationShapes == 1)
 		{
-			ChaosSolverDebugDebugDrawSettings.ShapesColorsPerState = DebugDraw::GetDefaultShapesColorsPreIntegrate();
-
-			DebugDrawShapes(false, !!ChaosSolverDrawShapesShowKinematic, !!ChaosSolverDrawShapesShowDynamic);
+			if (bChaosDebugDraw_UseNewQueue)
+			{
+				for (FTransientGeometryParticleHandle& Particle : Particles.GetActiveKinematicParticlesView())
+				{
+					FChaosDDParticle::DrawShapes(Particle.Handle(), FColor::Black);
+				}
+				for (FTransientGeometryParticleHandle& Particle : Particles.GetNonDisabledDynamicView())
+				{
+					FChaosDDParticle::DrawShapes(Particle.Handle(), FColor::Black);
+				}
+			}
+			else
+			{
+				ChaosSolverDebugDebugDrawSettings.ShapesColorsPerState = DebugDraw::GetDefaultShapesColorsPreIntegrate();
+				DebugDrawShapes(false, !!ChaosSolverDrawShapesShowKinematic, !!ChaosSolverDrawShapesShowDynamic);
+			}
 		}
 		if (ChaosSolverDebugDrawPreIntegrationCollisions == 1)
 		{
@@ -2479,9 +2495,22 @@ TRACE_COUNTER_SET(ChaosTraceCounter_##Name, Value)
 
 		if (ChaosSolverDebugDrawPostIntegrationShapes == 1)
 		{
-			ChaosSolverDebugDebugDrawSettings.ShapesColorsPerState = DebugDraw::GetDefaultShapesColorsPostIntegrate();
-
-			DebugDrawShapes(!!ChaosSolverDrawShapesShowStatic, !!ChaosSolverDrawShapesShowKinematic, !!ChaosSolverDrawShapesShowDynamic);
+			if (bChaosDebugDraw_UseNewQueue)
+			{
+				for (FTransientGeometryParticleHandle& Particle : Particles.GetActiveKinematicParticlesView())
+				{
+					FChaosDDParticle::DrawShapes(Particle.Handle(), FColor::Blue);
+				}
+				for (FTransientGeometryParticleHandle& Particle : Particles.GetNonDisabledDynamicView())
+				{
+					FChaosDDParticle::DrawShapes(Particle.Handle(), FColor::Yellow);
+				}
+			}
+			else
+			{
+				ChaosSolverDebugDebugDrawSettings.ShapesColorsPerState = DebugDraw::GetDefaultShapesColorsPostIntegrate();
+				DebugDrawShapes(!!ChaosSolverDrawShapesShowStatic, !!ChaosSolverDrawShapesShowKinematic, !!ChaosSolverDrawShapesShowDynamic);
+			}
 		}
 		if (ChaosSolverDebugDrawPostIntegrationCollisions == 1)
 		{
@@ -2514,22 +2543,32 @@ TRACE_COUNTER_SET(ChaosTraceCounter_##Name, Value)
 
 		if (ChaosSolverDebugDrawShapes == 1)
 		{
-			if (ChaosSolverDebugDrawColorShapeByClientServer)
+			if (bChaosDebugDraw_UseNewQueue)
 			{
-				if (bIsServer)
+				for (FTransientGeometryParticleHandle& Particle : Particles.GetNonDisabledView())
 				{
-					ChaosSolverDebugDebugDrawSettings.ShapesColorsPerState = GetSolverShapesColorsByState_Server();
-				}
-				else
-				{
-					ChaosSolverDebugDebugDrawSettings.ShapesColorsPerState = GetSolverShapesColorsByState_Client();
+					FChaosDDParticle::DrawShapes(Particle.Handle());
 				}
 			}
 			else
 			{
-				ChaosSolverDebugDebugDrawSettings.ShapesColorsPerState = DebugDraw::GetDefaultShapesColorsByState();
+				if (ChaosSolverDebugDrawColorShapeByClientServer)
+				{
+					if (bIsServer)
+					{
+						ChaosSolverDebugDebugDrawSettings.ShapesColorsPerState = GetSolverShapesColorsByState_Server();
+					}
+					else
+					{
+						ChaosSolverDebugDebugDrawSettings.ShapesColorsPerState = GetSolverShapesColorsByState_Client();
+					}
+				}
+				else
+				{
+					ChaosSolverDebugDebugDrawSettings.ShapesColorsPerState = DebugDraw::GetDefaultShapesColorsByState();
+				}
+				DebugDrawShapes(!!ChaosSolverDrawShapesShowStatic, !!ChaosSolverDrawShapesShowKinematic, !!ChaosSolverDrawShapesShowDynamic);
 			}
-			DebugDrawShapes(!!ChaosSolverDrawShapesShowStatic, !!ChaosSolverDrawShapesShowKinematic, !!ChaosSolverDrawShapesShowDynamic);
 		}
 		if (ChaosSolverDebugDrawCollisions == 1)
 		{
