GJK clean up
Make sure not to evaluate branching condition at runtime on GJK

diff --git a/Engine/Source/Runtime/Experimental/Chaos/Public/Chaos/SimplexVectorized.h b/Engine/Source/Runtime/Experimental/Chaos/Public/Chaos/SimplexVectorized.h
--- a/Engine/Source/Runtime/Experimental/Chaos/Public/Chaos/SimplexVectorized.h
+++ b/Engine/Source/Runtime/Experimental/Chaos/Public/Chaos/SimplexVectorized.h
@@ -32,7 +32,7 @@ namespace Chaos
 
 		Simplex[0] = VectorSelect(IsX1, Simplex[1], Simplex[0]);
 
-		if (CalculatExtraInformation)
+		if constexpr(CalculatExtraInformation)
 		{
 			A[0] = VectorSelect(IsX1, A[1], A[0]);
 			B[0] = VectorSelect(IsX1, B[1], B[0]);
@@ -53,7 +53,7 @@ namespace Chaos
 		const T IsX0OrX1 = VectorBitwiseOr(IsX0, IsX1);
 		NumVerts =  VectorMaskBits(IsX0OrX1) ? 1 : NumVerts;
 
-		if (CalculatExtraInformation)
+		if constexpr (CalculatExtraInformation)
 		{
 			OutBarycentric = VectorSelect(IsX0OrX1, OutBarycentricIfX0OrX1, OutBarycentricOtherwise);
 		}
@@ -100,7 +100,7 @@ namespace Chaos
 		if (VectorMaskBits(IsA))
 		{
 			NumVerts = 1;
-			if (CalculatExtraInformation)
+			if constexpr (CalculatExtraInformation)
 			{
 				OutBarycentric = TMakeVectorRegisterConstant<T>(1, 0, 0, 0);
 			}
@@ -119,12 +119,12 @@ namespace Chaos
 		if (VectorMaskBits(IsB))
 		{
 			NumVerts = 1;
-			if (CalculatExtraInformation)
+			if constexpr (CalculatExtraInformation)
 			{
 				OutBarycentric = TMakeVectorRegisterConstant<T>(1, 0, 0, 0); 
 			}
 			Simplex[0] = Simplex[1];
-			if (CalculatExtraInformation)
+			if constexpr (CalculatExtraInformation)
 			{
 				As[0] = As[1];
 				Bs[0] = Bs[1];
@@ -150,7 +150,7 @@ namespace Chaos
 			const T v = VectorDivide(d1, NormalizationDenominatorAB);
 			const T OneMinusV = VectorSubtract(TMakeVectorRegisterConstant<T>(1, 1, 1, 1), v);
 			// b0	a1	a2	a3
-			if (CalculatExtraInformation)
+			if constexpr (CalculatExtraInformation)
 			{
 				OutBarycentric = VectorUnpackLo(OneMinusV, v);
 			}
@@ -168,13 +168,13 @@ namespace Chaos
 		if (VectorMaskBits(IsC))
 		{
 			NumVerts = 1;
-			if (CalculatExtraInformation)
+			if constexpr (CalculatExtraInformation)
 			{
 				OutBarycentric = TMakeVectorRegisterConstant<T>(1, 0, 0, 0);
 			}
 
 			Simplex[0] = Simplex[2];
-			if (CalculatExtraInformation)
+			if constexpr (CalculatExtraInformation)
 			{
 				As[0] = As[2];
 				Bs[0] = Bs[2];
@@ -199,12 +199,12 @@ namespace Chaos
 			NumVerts = 2;
 			const T OneMinusW = VectorSubtract(TMakeVectorRegisterConstant<T>(1, 1, 1, 1), w);
 			// b0	a1	a2	a3
-			if (CalculatExtraInformation)
+			if constexpr (CalculatExtraInformation)
 			{
 				OutBarycentric = VectorUnpackLo(OneMinusW, w);
 			}
 			Simplex[1] = Simplex[2];
-			if (CalculatExtraInformation)
+			if constexpr (CalculatExtraInformation)
 			{
 				As[1] = As[2];
 				Bs[1] = Bs[2];
@@ -229,7 +229,7 @@ namespace Chaos
 		{
 			NumVerts = 2;
 			const T w = VectorDivide(d4MinusD3, NormalizationDenominatorBC);
-			if (CalculatExtraInformation)
+			if constexpr (CalculatExtraInformation)
 			{
 				const T OneMinusW = VectorSubtract(TMakeVectorRegisterConstant<T>(1, 1, 1, 1), w);
 				// b0	a1	a2	a3
@@ -239,7 +239,7 @@ namespace Chaos
 			const T Result = VectorMultiplyAdd(w, CMinusB, B);
 			Simplex[0] = Simplex[1];
 			Simplex[1] = Simplex[2];
-			if (CalculatExtraInformation)
+			if constexpr (CalculatExtraInformation)
 			{
 				As[0] = As[1];
 				Bs[0] = Bs[1];
@@ -255,7 +255,7 @@ namespace Chaos
 		T w = VectorMultiply(vc, denom);
 		NumVerts = 3;
 
-		if (CalculatExtraInformation)
+		if constexpr (CalculatExtraInformation)
 		{
 			const T OneMinusVMinusW = VectorSubtract(VectorSubtract(TMakeVectorRegisterConstant<T>(1, 1, 1, 1), v), w);
 			// b0	a1	a2	a3
@@ -316,7 +316,7 @@ namespace Chaos
 		T SubSimplices[4][3] = { {Simplex[1], Simplex[2], Simplex[3]}, {Simplex[0], Simplex[2], Simplex[3]}, {Simplex[0], Simplex[1], Simplex[3]}, {Simplex[0], Simplex[1], Simplex[2]} };
 		T SubAs[4][3];
 		T SubBs[4][3];
-		if (CalculatExtraInformation)
+		if constexpr (CalculatExtraInformation)
 		{
 			//SubAs = { {A[1], A[2], A[3]}, {A[0], A[2], A[3]}, {A[0], A[1], A[3]}, {A[0], A[1], A[2]} };
 			//SubBs = { {B[1], B[2], B[3]}, {B[0], B[2], B[3]}, {B[0], B[1], B[3]}, {B[0], B[1], B[2]} };
@@ -362,7 +362,7 @@ namespace Chaos
 
 		if (bInside)
 		{
-			if (CalculatExtraInformation)
+			if constexpr (CalculatExtraInformation)
 			{
 				T OutBarycentricVectors[4];
 				const T  InvDetM = VectorDivide(TMakeVectorRegisterConstant<T>(1, 1, 1, 1), DetM);
@@ -381,7 +381,7 @@ namespace Chaos
 		}
 
 		NumVerts = SubNumVerts[ClosestTriangleIdx];
-		if (CalculatExtraInformation)
+		if constexpr (CalculatExtraInformation)
 		{
 			OutBarycentric = SubBarycentric[ClosestTriangleIdx];
 		}
@@ -389,7 +389,7 @@ namespace Chaos
 		for (int i = 0; i < 3; i++)
 		{
 			Simplex[i] = SubSimplices[ClosestTriangleIdx][i];
-			if (CalculatExtraInformation)
+			if constexpr (CalculatExtraInformation)
 			{
 				A[i] = SubAs[ClosestTriangleIdx][i];
 				B[i] = SubBs[ClosestTriangleIdx][i];
@@ -407,7 +407,7 @@ namespace Chaos
 		switch (NumVerts)
 		{
 		case 1:
-			if (CalculatExtraInformation)
+			if constexpr (CalculatExtraInformation)
 			{
 				OutBarycentric = TMakeVectorRegisterConstant<T>(1, 0, 0 , 0);
 			}
