Chaos Cloth: Fix FluidDensity to properly scale with WorldScale on VelocityFields in the constructor.
During typical initialization, one call to Update will re-read the FluidDensity and properly set the value, but this does not happen when recreating cloth actors.

diff --git a/Engine/Source/Runtime/Experimental/Chaos/Public/Chaos/VelocityField.h b/Engine/Source/Runtime/Experimental/Chaos/Public/Chaos/VelocityField.h
--- a/Engine/Source/Runtime/Experimental/Chaos/Public/Chaos/VelocityField.h
+++ b/Engine/Source/Runtime/Experimental/Chaos/Public/Chaos/VelocityField.h
@@ -63,7 +63,7 @@ public:
 		, Drag(GetWeightedFloatDrag(PropertyCollection, (FSolverReal)0.).ClampAxes(MinCoefficient, MaxCoefficient))
 		, OuterDrag(GetWeightedFloatOuterDrag(PropertyCollection, FSolverVec2(Drag.GetLow(), Drag.GetHigh())).ClampAxes(MinCoefficient, MaxCoefficient))
 		, Pressure(GetWeightedFloatPressure(PropertyCollection, (FSolverReal)0.)/WorldScale)
-		, Rho(FMath::Max(GetFluidDensity(PropertyCollection, (FSolverReal)0.), (FSolverReal)0.))
+		, Rho(FMath::Max(GetFluidDensity(PropertyCollection, (FSolverReal)0.)/FMath::Cube(WorldScale), (FSolverReal)0.))
 		, QuarterRho(Rho* (FSolverReal)0.25f)
 		, DragIndex(PropertyCollection)
 		, OuterDragIndex(PropertyCollection)
