Fix for constraint container crash on joint disconnect due to multiple containers being present and not owning all constraints. Now on disconnect the container will make sure it is bound to the constraint before continuing with the removal process.

diff --git a/Engine/Source/Runtime/Experimental/Chaos/Private/Chaos/PBDJointConstraints.cpp b/Engine/Source/Runtime/Experimental/Chaos/Private/Chaos/PBDJointConstraints.cpp
--- a/Engine/Source/Runtime/Experimental/Chaos/Private/Chaos/PBDJointConstraints.cpp
+++ b/Engine/Source/Runtime/Experimental/Chaos/Private/Chaos/PBDJointConstraints.cpp
@@ -699,6 +699,11 @@ namespace Chaos
 			{
 				if (FPBDJointConstraintHandle* JointHandle = ConstraintHandle->As<FPBDJointConstraintHandle>())
 				{
+					if (!JointHandle->IsInContainer(this))
+					{
+						continue;
+					}
+
 					JointHandle->SetEnabled(false); // constraint lifespan is managed by the proxy
 
 					int ConstraintIndex = JointHandle->GetConstraintIndex();
diff --git a/Engine/Source/Runtime/Experimental/Chaos/Public/Chaos/PBDJointConstraints.h b/Engine/Source/Runtime/Experimental/Chaos/Public/Chaos/PBDJointConstraints.h
--- a/Engine/Source/Runtime/Experimental/Chaos/Public/Chaos/PBDJointConstraints.h
+++ b/Engine/Source/Runtime/Experimental/Chaos/Public/Chaos/PBDJointConstraints.h
@@ -117,6 +117,7 @@ namespace Chaos
 		CHAOS_API int32 GetConstraintColor() const;
 
 		CHAOS_API void SetConstraintContainer(FPBDConstraintContainer* NewConstraintContainer) { ConstraintContainer = NewConstraintContainer;  }
+		CHAOS_API bool IsInContainer(FPBDConstraintContainer* TestContainer) { return TestContainer == ConstraintContainer;  }
 
 	protected:
 		using Base::ConstraintIndex;
