Iris - Disconnect instead of crashing if an invalid NetObjectFactory ID is received.

#rnx

diff --git a/Engine/Source/Runtime/Experimental/Iris/Core/Private/Iris/ReplicationSystem/ObjectReplicationBridge.cpp b/Engine/Source/Runtime/Experimental/Iris/Core/Private/Iris/ReplicationSystem/ObjectReplicationBridge.cpp
--- a/Engine/Source/Runtime/Experimental/Iris/Core/Private/Iris/ReplicationSystem/ObjectReplicationBridge.cpp
+++ b/Engine/Source/Runtime/Experimental/Iris/Core/Private/Iris/ReplicationSystem/ObjectReplicationBridge.cpp
@@ -219,8 +219,12 @@ void UObjectReplicationBridge::DeinitNetObjectFactories()
 
 UNetObjectFactory* UObjectReplicationBridge::GetNetFactory(UE::Net::FNetObjectFactoryId FactoryId) const
 {
-	check(FactoryId != UE::Net::InvalidNetObjectFactoryId);
-	return NetObjectFactories[FactoryId];
+	if (ensureMsgf(FactoryId < NetObjectFactories.Num() && FactoryId != UE::Net::InvalidNetObjectFactoryId, TEXT("GetNetFactory invalid factory ID %u, num factories: %d"), FactoryId, NetObjectFactories.Num()))
+	{
+		return NetObjectFactories[FactoryId];
+	}
+
+	return nullptr;
 }
 
 void UObjectReplicationBridge::OnMaxInternalNetRefIndexIncreased(UE::Net::Private::FInternalNetRefIndex NewMaxInternalIndex)
@@ -822,9 +826,13 @@ FReplicationBridgeCreateNetRefHandleResult UObjectReplicationBridge::CreateNetRe
 	FNetBitStreamReader* Reader = Context.SerializationContext.GetBitStreamReader();
 
 	const FNetObjectFactoryId FactoryId = IntCastChecked<FNetObjectFactoryId>(Reader->ReadBits(FNetObjectFactoryRegistry::GetMaxBits()));
-	check(FactoryId != InvalidNetObjectFactoryId);
-
 	UNetObjectFactory* Factory = GetNetFactory(FactoryId);
+	if (Factory == nullptr)
+	{
+		// There must be a valid factory. We're reading a creation header so a bitstream error seems likely and we should disconnect.
+		Context.SerializationContext.SetError(GNetError_InvalidValue);
+		return FReplicationBridgeCreateNetRefHandleResult();
+	}
 
 	TUniquePtr<FNetObjectCreationHeader> Header = Factory->ReadHeader(WantedNetHandle, Context.SerializationContext);
 		
