Iris
- Fix for protocol mismatch due to bad descriptor lookup

diff --git a/Engine/Source/Runtime/Experimental/Iris/Core/Private/Iris/ReplicationState/ReplicationStateDescriptorRegistry.cpp b/Engine/Source/Runtime/Experimental/Iris/Core/Private/Iris/ReplicationState/ReplicationStateDescriptorRegistry.cpp
--- a/Engine/Source/Runtime/Experimental/Iris/Core/Private/Iris/ReplicationState/ReplicationStateDescriptorRegistry.cpp
+++ b/Engine/Source/Runtime/Experimental/Iris/Core/Private/Iris/ReplicationState/ReplicationStateDescriptorRegistry.cpp
@@ -4,10 +4,16 @@
 #include "Iris/ReplicationState/ReplicationStateDescriptor.h"
 #include "Iris/ReplicationSystem/ReplicationProtocolManager.h"
 #include "Iris/Core/IrisProfiler.h"
+#include "Iris/Core/IrisLog.h"
+#include "HAL/ConsoleManager.h"
 
 namespace UE::Net::Private
 {
 
+static bool bPruneReplicationStateDescriptorsWithArchetype = true;
+static FAutoConsoleVariableRef CVarbPruneReplicationStateDescriptorsWithArchetype(TEXT("net.Iris.PruneReplicationStateDescriptorsWithArchetype"), bPruneReplicationStateDescriptorsWithArchetype, 
+	TEXT("If true, we will invalidate registered descriptors if archetype is no longer resolvable, otherwise we will keep them around until CDO is no longer valid."));
+
 FReplicationStateDescriptorRegistry::FReplicationStateDescriptorRegistry()
 : ProtocolManager(nullptr)
 {
@@ -26,25 +32,28 @@ void FReplicationStateDescriptorRegistry::Register(const FFieldVariant& Object,
 
 	if (Entry)
 	{
-		// if we found the entry and the WeakPtrForPruning is also a match we are trying to register descriptors for an already registered class
+		// We do not want to overwrite descriptors for valid objects.
 		if (Entry->WeakPtrForPruning.Get() == ObjectForPruning)
 		{
-			checkf(false, TEXT("FReplicationStateDescriptorRegistry::Trying to register descriptors for the same UObject %s"), ToCStr(ObjectForPruning->GetName()));
-			return;
+			if (Entry->OwnerKey == FObjectKey(GetObjectForPruning(Object)))
+			{
+				checkf(false, TEXT("FReplicationStateDescriptorRegistry::Trying to register descriptors for the same UObject %s"), ToCStr(ObjectForPruning->GetName()));
+				return;
+			}
 		}
-		else
-		{
-			// Notify protocol manager about pruned descriptors
-			InvalidateDescriptors(Entry->Descriptors);
 
-			RegisteredDescriptorsMap.Remove(Object);
-		}
+		// We found an invalid entry, invalidate it before registering new descriptors.
+		UE_LOG(LogIris, VeryVerbose, TEXT("FReplicationStateDescriptorRegistry::Register invalidate descriptors for ptr: 0x%p"), Object.GetRawPointer());
+		// Notify protocol manager about pruned descriptors
+		InvalidateDescriptors(Entry->Descriptors);
+		RegisteredDescriptorsMap.Remove(Object);
 	}
 
+	// We found an invalid entry, invalidate it before registering new descriptors.
 	FRegisteredDescriptors NewEntry;
+	NewEntry.OwnerKey = FObjectKey(GetObjectForPruning(Object));
 	NewEntry.WeakPtrForPruning = TWeakObjectPtr<const UObject>(ObjectForPruning);
 	NewEntry.Descriptors = Descriptors;
-
 	RegisteredDescriptorsMap.Add(Object, NewEntry);
 }
 
@@ -56,22 +65,25 @@ void FReplicationStateDescriptorRegistry::Register(const FFieldVariant& Object,
 
 	if (Entry)
 	{
-		// if we found the entry and the WeakPtrForPruning is also a match we are trying to register descriptors for an already registered class
+		// We do not want to overwrite descriptors for valid objects.
 		if (Entry->WeakPtrForPruning.Get() == ObjectForPruning)
 		{
-			checkf(false, TEXT("FReplicationStateDescriptorRegistry::Trying to register descriptors for the same UObject %s"), ToCStr(ObjectForPruning->GetName()));
-			return;
+			if (Entry->OwnerKey == FObjectKey(GetObjectForPruning(Object)))
+			{
+				checkf(false, TEXT("FReplicationStateDescriptorRegistry::Trying to register descriptor for the same UObject %s"), ToCStr(ObjectForPruning->GetName()));
+				return;
+			}
 		}
-		else
-		{
-			// Notify protocol manager about pruned descriptors
-			InvalidateDescriptors(Entry->Descriptors);
 
-			RegisteredDescriptorsMap.Remove(Object);
-		}
+		// We found an invalid entry, invalidate it before registering new descriptors.
+		UE_LOG(LogIris, VeryVerbose, TEXT("FReplicationStateDescriptorRegistry::Register invalidate descriptor for ptr: 0x%p"), Object.GetRawPointer());
+		// Notify protocol manager about pruned descriptors
+		InvalidateDescriptors(Entry->Descriptors);
+		RegisteredDescriptorsMap.Remove(Object);
 	}
 
 	FRegisteredDescriptors& NewEntry = RegisteredDescriptorsMap.Emplace(Object);
+	NewEntry.OwnerKey = FObjectKey(GetObjectForPruning(Object));
 	NewEntry.WeakPtrForPruning = TWeakObjectPtr<const UObject>(ObjectForPruning);
 	NewEntry.Descriptors.Add(Descriptor);
 }
@@ -82,8 +94,15 @@ const FReplicationStateDescriptorRegistry::FDescriptors* FReplicationStateDescri
 
 	check(ObjectForPruning != nullptr);
 
-	if (Entry && Entry->WeakPtrForPruning.Get() == ObjectForPruning)
+	if (Entry && (Entry->WeakPtrForPruning.Get() == ObjectForPruning))
 	{
+		// Archetype might have been reused, we will clear this up when registering
+		if (Object.IsUObject() && !Entry->OwnerKey.ResolveObjectPtr())
+		{
+			UE_LOG(LogIris, VeryVerbose, TEXT("FReplicationStateDescriptorRegistry Found invalidated entry ptr: 0x%p"), Object.GetRawPointer());
+			return nullptr;
+		}
+
 		return &Entry->Descriptors;
 	}
 	else
@@ -100,11 +119,13 @@ void FReplicationStateDescriptorRegistry::PruneStaleDescriptors()
 	for (auto It = RegisteredDescriptorsMap.CreateIterator(); It; ++It)
 	{
 		const FRegisteredDescriptors& RegisteredDescriptors = It.Value();
-		if (!RegisteredDescriptors.WeakPtrForPruning.IsValid())
+		const bool bPruneDueToWeakPtrForPruningBeingStale = !RegisteredDescriptors.WeakPtrForPruning.IsValid();
+		if (bPruneDueToWeakPtrForPruningBeingStale || (bPruneReplicationStateDescriptorsWithArchetype && (RegisteredDescriptors.OwnerKey.ResolveObjectPtr() == nullptr)))
 		{
+			UE_LOG(LogIris, VeryVerbose, TEXT("FReplicationStateDescriptorRegistry Pruning descriptors for ptr: 0x%p due to %s"), It.Key().GetRawPointer(), (bPruneDueToWeakPtrForPruningBeingStale ? TEXT("invalidated CDO") : TEXT("invalidated Key/Archetype")));
+
 			// Notify protocol manager about pruned descriptors
 			InvalidateDescriptors(RegisteredDescriptors.Descriptors);
-
 			It.RemoveCurrent();
 		}
 	}
diff --git a/Engine/Source/Runtime/Experimental/Iris/Core/Private/Iris/ReplicationState/ReplicationStateDescriptorRegistry.h b/Engine/Source/Runtime/Experimental/Iris/Core/Private/Iris/ReplicationState/ReplicationStateDescriptorRegistry.h
--- a/Engine/Source/Runtime/Experimental/Iris/Core/Private/Iris/ReplicationState/ReplicationStateDescriptorRegistry.h
+++ b/Engine/Source/Runtime/Experimental/Iris/Core/Private/Iris/ReplicationState/ReplicationStateDescriptorRegistry.h
@@ -8,6 +8,7 @@
 #include "Templates/RefCounting.h"
 #include "UObject/Field.h"
 #include "UObject/WeakObjectPtr.h"
+#include "UObject/ObjectKey.h"
 
 namespace UE::Net
 {
@@ -66,6 +67,7 @@ private:
 	{
 		TWeakObjectPtr<const UObject> WeakPtrForPruning;
 		FDescriptors Descriptors;
+		FObjectKey OwnerKey;
 	};
 
 	typedef TMap<FFieldVariant, FRegisteredDescriptors> FClassToDescriptorMap;
