Move "InLine" flag out from the free bits of the Data pointer as that upsets ARM's MTE.

#rnx

diff --git a/Engine/Source/Runtime/Experimental/IoStore/HttpClient/Private/Misc.inl b/Engine/Source/Runtime/Experimental/IoStore/HttpClient/Private/Misc.inl
--- a/Engine/Source/Runtime/Experimental/IoStore/HttpClient/Private/Misc.inl
+++ b/Engine/Source/Runtime/Experimental/IoStore/HttpClient/Private/Misc.inl
@@ -514,17 +514,17 @@ public:
 private:
 	char*						GetDataPtr();
 	void						Extend(uint32 AtLeast, uint32 PageSize);
+	UPTRINT						Data;
+	uint32						Max = 0;
 	union
 	{
 		struct
 		{
-			UPTRINT				Data : 63;
-			UPTRINT				Inline : 1;
+			uint32				Used : 31;
+			uint32				Inline : 1;
 		};
-		UPTRINT					DataInline = 0;
+		uint32					UsedInline = 0;
 	};
-	uint32						Max = 0;
-	uint32						Used = 0;
 
 private:
 								FBuffer(FBuffer&&) = delete;
@@ -535,9 +535,10 @@ private:
 ////////////////////////////////////////////////////////////////////////////////
 FBuffer::FBuffer(char* InData, uint32 InMax)
 : Data(UPTRINT(InData))
-, Inline(1)
 , Max(InMax)
 {
+	check(Used == 0);
+	Inline = 1;
 }
 
 ////////////////////////////////////////////////////////////////////////////////
@@ -552,9 +553,9 @@ FBuffer::~FBuffer()
 ////////////////////////////////////////////////////////////////////////////////
 FBuffer& FBuffer::operator = (FBuffer&& Rhs)
 {
-	Swap(DataInline, Rhs.DataInline);
+	Swap(Data, Rhs.Data);
 	Swap(Max, Rhs.Max);
-	Swap(Used, Rhs.Used);
+	Swap(UsedInline, Rhs.UsedInline);
 	return *this;
 }
 
