Conciliate static analysis that can't follow the union.

#rnx

diff --git a/Engine/Source/Runtime/Experimental/IoStore/HttpClient/Private/Misc.inl b/Engine/Source/Runtime/Experimental/IoStore/HttpClient/Private/Misc.inl
--- a/Engine/Source/Runtime/Experimental/IoStore/HttpClient/Private/Misc.inl
+++ b/Engine/Source/Runtime/Experimental/IoStore/HttpClient/Private/Misc.inl
@@ -537,7 +537,6 @@ FBuffer::FBuffer(char* InData, uint32 InMax)
 : Data(UPTRINT(InData))
 , Max(InMax)
 {
-	check(Used == 0);
 	Inline = 1;
 }
 
