Dataflow Editor: add support for materials in the Construction viewport. This requires adding UVs and Materials to the RenderCollection.
Also implement getting Render Mesh materials from a ClothCollection into the RenderCollection, which enables material rendering for cloth in the construction view

Review: 36645514

diff --git a/Engine/Plugins/ChaosClothAssetEditor/Source/ChaosClothAssetDataflowNodes/Private/ChaosClothAsset/ChaosClothAssetDataflowNodesModule.cpp b/Engine/Plugins/ChaosClothAssetEditor/Source/ChaosClothAssetDataflowNodes/Private/ChaosClothAsset/ChaosClothAssetDataflowNodesModule.cpp
--- a/Engine/Plugins/ChaosClothAssetEditor/Source/ChaosClothAssetDataflowNodes/Private/ChaosClothAsset/ChaosClothAssetDataflowNodesModule.cpp
+++ b/Engine/Plugins/ChaosClothAssetEditor/Source/ChaosClothAssetDataflowNodes/Private/ChaosClothAsset/ChaosClothAssetDataflowNodesModule.cpp
@@ -194,6 +194,8 @@ PRAGMA_ENABLE_DEPRECATION_WARNINGS
 					TArray<FVector3f> Vertices;
 					TArray<FVector3f> Normals;
 					TArray<FIntVector> Indices;
+					TArray<TArray<FVector2f>> UVs;
+					TArray<int32> RenderMaterialIDs;
 
 					if (State.GetViewMode().GetName() == FCloth3DSimViewMode::Name)
 					{
@@ -226,14 +228,39 @@ PRAGMA_ENABLE_DEPRECATION_WARNINGS
 						Vertices = ClothFacade.GetRenderPosition();
 						Indices = ClothFacade.GetRenderIndices();
 						Normals = ClothFacade.GetRenderNormal();
-						// TODO: Get materials from the render mesh
+						
+						const TConstArrayView<TArray<FVector2f>> AllClothUVs = ClothFacade.GetRenderUVs();
+						check(AllClothUVs.Num() == Vertices.Num());
+						UVs = TArray<TArray<FVector2f>>(AllClothUVs);
+
+						RenderMaterialIDs.SetNumUninitialized(Indices.Num());
+
+						for (int32 PatternIndex = 0; PatternIndex < ClothFacade.GetNumRenderPatterns(); ++PatternIndex)
+						{
+							const FCollectionClothRenderPatternConstFacade& Pattern = ClothFacade.GetRenderPattern(PatternIndex);
+							const int32 PatternFaceIndexOffset = Pattern.GetRenderFacesOffset();
+							for (int32 PatternFaceIndex = 0; PatternFaceIndex < Pattern.GetNumRenderFaces(); ++PatternFaceIndex)
+							{
+								RenderMaterialIDs[PatternFaceIndexOffset + PatternFaceIndex] = PatternIndex;
+							}
+						}
 					}
 
 					TArray<FLinearColor> Colors;
 					Colors.Init(FLinearColor::Gray, Vertices.Num());	// TODO: Choose a vertex color
 
 					const int32 GeometryIndex = RenderCollection.StartGeometryGroup(State.GetGuid().ToString());
-					RenderCollection.AddSurface(MoveTemp(Vertices), MoveTemp(Indices), MoveTemp(Normals), MoveTemp(Colors));
+
+					if (State.GetViewMode().GetName() == FClothRenderViewMode::Name)
+					{
+						TArray<FString> MaterialArray(ClothFacade.GetRenderMaterialPathName());
+						RenderCollection.AddSurface(MoveTemp(Vertices), MoveTemp(Indices), MoveTemp(Normals), MoveTemp(Colors), MoveTemp(UVs), MoveTemp(RenderMaterialIDs), MoveTemp(MaterialArray));
+					}
+					else
+					{
+						RenderCollection.AddSurface(MoveTemp(Vertices), MoveTemp(Indices), MoveTemp(Normals), MoveTemp(Colors));
+					}
+
 					RenderCollection.EndGeometryGroup(GeometryIndex);
 				}
 			}
diff --git a/Engine/Plugins/Experimental/Dataflow/Source/DataflowAssetTools/Private/Dataflow/CollectionRenderingPatternUtility.cpp b/Engine/Plugins/Experimental/Dataflow/Source/DataflowAssetTools/Private/Dataflow/CollectionRenderingPatternUtility.cpp
--- a/Engine/Plugins/Experimental/Dataflow/Source/DataflowAssetTools/Private/Dataflow/CollectionRenderingPatternUtility.cpp
+++ b/Engine/Plugins/Experimental/Dataflow/Source/DataflowAssetTools/Private/Dataflow/CollectionRenderingPatternUtility.cpp
@@ -24,8 +24,10 @@ namespace UE::Dataflow
 			{
 				int32 StartTriangles = 0;
 				int32 StartVertices = 0;
+				int32 StartMaterials = 0;
 				int32 NumTriangles = Facade.NumTriangles();
 				int32 NumVertices = Facade.NumVertices();
+				int32 NumMaterials = Facade.NumMaterials();
 
 				if (InMeshIndex != INDEX_NONE)
 				{
@@ -33,8 +35,10 @@ namespace UE::Dataflow
 					{
 						StartTriangles = Facade.GetIndicesStart()[InMeshIndex];
 						StartVertices = Facade.GetVertexStart()[InMeshIndex];
+						StartMaterials = Facade.GetMaterialStart()[InMeshIndex];
 						NumTriangles = Facade.GetIndicesCount()[InMeshIndex];
 						NumVertices = Facade.GetVertexCount()[InMeshIndex];
+						NumMaterials = Facade.GetMaterialCount()[InMeshIndex];
 					}
 				}
 
@@ -43,6 +47,8 @@ namespace UE::Dataflow
 				const TManagedArray<FVector3f>& Positions = Facade.GetVertices();
 				const TManagedArray<FVector3f>& Normals = Facade.GetNormals();
 				const TManagedArray<FLinearColor>& Colors = Facade.GetVertexColor();
+				const TManagedArray<TArray<FVector2f>>& UVs = Facade.GetVertexUV();
+				const TManagedArray<int32>& FacadeMaterialID = Facade.GetMaterialID();
 
 				int32 LastVertexIndex = StartVertices + NumVertices;
 				for (int32 VertexIndex = StartVertices; VertexIndex < LastVertexIndex; ++VertexIndex)
@@ -88,6 +94,44 @@ namespace UE::Dataflow
 				{
 					SetColorsFromWeights(TriangleID);
 				}
+
+				DynamicMesh.Attributes()->EnableMaterialID();
+				FDynamicMeshMaterialAttribute* const MaterialIDAttrib = DynamicMesh.Attributes()->GetMaterialID();
+				for (const int TriangleID : DynamicMesh.TriangleIndicesItr())
+				{
+					MaterialIDAttrib->SetValue(TriangleID, FacadeMaterialID[TriangleID] - StartMaterials);
+				}
+
+				const int32 NumUVLayers = UVs[0].Num();
+
+				if (NumUVLayers > 0)
+				{
+					for (const TArray<FVector2f>& VertexUVs : UVs)
+					{
+						check(NumUVLayers == VertexUVs.Num());
+					}
+
+					DynamicMesh.Attributes()->SetNumUVLayers(NumUVLayers);
+
+					for (int32 UVLayerIndex = 0; UVLayerIndex < NumUVLayers; ++UVLayerIndex)
+					{
+						FDynamicMeshUVOverlay* const UVLayer = DynamicMesh.Attributes()->GetUVLayer(UVLayerIndex);
+						UVLayer->CreateFromPredicate([](int ParentVID, int TriIDA, int TriIDB) { return true; }, 0.f);
+
+						for (const int TriangleID : DynamicMesh.TriangleIndicesItr())
+						{
+							const FIndex3i Tri = DynamicMesh.GetTriangle(TriangleID);
+							const FIndex3i UVElementTri = UVLayer->GetTriangle(TriangleID);
+
+							for (int TriVertIndex = 0; TriVertIndex < 3; ++TriVertIndex)
+							{
+								const FVector2f UV(UVs[Remapping[Tri[TriVertIndex]]][UVLayerIndex]);
+								UVLayer->SetElement(UVElementTri[TriVertIndex], UV);
+							}
+						}
+					}
+				}
+
 			}
 		}
 
diff --git a/Engine/Plugins/Experimental/Dataflow/Source/DataflowEditor/Private/Dataflow/DataflowConstructionScene.cpp b/Engine/Plugins/Experimental/Dataflow/Source/DataflowEditor/Private/Dataflow/DataflowConstructionScene.cpp
--- a/Engine/Plugins/Experimental/Dataflow/Source/DataflowEditor/Private/Dataflow/DataflowConstructionScene.cpp
+++ b/Engine/Plugins/Experimental/Dataflow/Source/DataflowEditor/Private/Dataflow/DataflowConstructionScene.cpp
@@ -166,7 +166,20 @@ void FDataflowConstructionScene::UpdateDynamicMeshComponents()
 								EditorContent->SetRenderCollection(RenderCollection);
 							}
 							const FString MeshName = Facade.GetGeometryName()[MeshIndex];
-							AddDynamicMeshComponent({Target, MeshIndex}, MeshName, MoveTemp(DynamicMesh), {});
+
+							const TManagedArray<FString>& MaterialPaths = Facade.GetMaterialPaths();
+							const int32 MaterialStart = Facade.GetMaterialStart()[MeshIndex];
+							const int32 MaterialCount = Facade.GetMaterialCount()[MeshIndex];
+
+							TArray<UMaterialInterface*> Materials;
+							for (int32 MaterialIndex = MaterialStart; MaterialIndex < MaterialStart + MaterialCount; ++MaterialIndex)
+							{
+								const FString& Path = MaterialPaths[MaterialIndex];
+								UMaterialInterface* const Material = LoadObject<UMaterialInterface>(nullptr, *Path);
+								Materials.Add(Material);
+							}
+
+							AddDynamicMeshComponent({Target, MeshIndex}, MeshName, MoveTemp(DynamicMesh), Materials);
 						}
 					}
 				}
@@ -278,30 +291,44 @@ TObjectPtr<UDynamicMeshComponent>& FDataflowConstructionScene::AddDynamicMeshCom
 	DynamicMeshComponent->Node = InKey.Key;
 	DynamicMeshComponent->SetMesh(MoveTemp(DynamicMesh));
 	
-	// @todo(Material) This is just to have a material, we should transfer the materials from the assets if they have them. 
-	const  TObjectPtr<UDataflowBaseContent>& EditorContent = GetEditorContent();
-	if (EditorContent && EditorContent->GetDataflowAsset() && EditorContent->GetDataflowAsset()->Material)
+	if (MaterialSet.Num() == 0)
 	{
-		DynamicMeshComponent->ConfigureMaterialSet({ EditorContent->GetDataflowAsset()->Material });
+		const TObjectPtr<UDataflowBaseContent>& EditorContent = GetEditorContent();
+		if (EditorContent && EditorContent->GetDataflowAsset() && EditorContent->GetDataflowAsset()->Material)
+		{
+			DynamicMeshComponent->ConfigureMaterialSet({ EditorContent->GetDataflowAsset()->Material });
+		}
+		else
+		{
+			ensure(FDataflowEditorStyle::Get().DefaultTwoSidedMaterial);
+			DynamicMeshComponent->SetOverrideRenderMaterial(FDataflowEditorStyle::Get().DefaultTwoSidedMaterial);
+			DynamicMeshComponent->SetShadowsEnabled(false);
+		}
 	}
 	else
 	{
-		ensure(FDataflowEditorStyle::Get().DefaultTwoSidedMaterial);
-		DynamicMeshComponent->SetOverrideRenderMaterial(FDataflowEditorStyle::Get().DefaultTwoSidedMaterial);
-		DynamicMeshComponent->SetShadowsEnabled(false);
+		DynamicMeshComponent->ConfigureMaterialSet(MaterialSet);
 	}
-	//else if (FDataflowEditorStyle::Get().DefaultMaterial)
-	//{
-	//	DynamicMeshComponent->ConfigureMaterialSet({ FDataflowEditorStyle::Get().DefaultMaterial });
-	//}
-	//else
-	//{
-	//	DynamicMeshComponent->ValidateMaterialSlots(true, false);
-	//}
 
 	DynamicMeshComponent->SelectionOverrideDelegate = UPrimitiveComponent::FSelectionOverride::CreateRaw(this, &FDataflowPreviewSceneBase::IsComponentSelected);
 	DynamicMeshComponent->UpdateBounds();
 
+	// Fix up any triangles without valid material IDs
+	int32 DefaultMaterialID = INDEX_NONE;
+	for (const int32 TriID : DynamicMeshComponent->GetMesh()->TriangleIndicesItr())
+	{
+		const int32 MaterialID = DynamicMeshComponent->GetMesh()->Attributes()->GetMaterialID()->GetValue(TriID);
+		if (!DynamicMeshComponent->GetMaterial(MaterialID))
+		{
+			if (DefaultMaterialID == INDEX_NONE)
+			{
+				DefaultMaterialID = DynamicMeshComponent->GetNumMaterials();
+				DynamicMeshComponent->SetMaterial(DefaultMaterialID, FDataflowEditorStyle::Get().VertexMaterial);
+			}
+			DynamicMeshComponent->GetMesh()->Attributes()->GetMaterialID()->SetValue(TriID, DefaultMaterialID);
+		}
+	}
+
 	AddComponent(DynamicMeshComponent, DynamicMeshComponent->GetRelativeTransform());	
 	DynamicMeshComponents.Emplace(InKey, DynamicMeshComponent);
 	return DynamicMeshComponents[InKey];
diff --git a/Engine/Source/Runtime/Experimental/Chaos/Private/GeometryCollection/Facades/CollectionRenderingFacade.cpp b/Engine/Source/Runtime/Experimental/Chaos/Private/GeometryCollection/Facades/CollectionRenderingFacade.cpp
--- a/Engine/Source/Runtime/Experimental/Chaos/Private/GeometryCollection/Facades/CollectionRenderingFacade.cpp
+++ b/Engine/Source/Runtime/Experimental/Chaos/Private/GeometryCollection/Facades/CollectionRenderingFacade.cpp
@@ -19,15 +19,19 @@ namespace GeometryCollection::Facades
 		, VertexHitProxyIndexAttribute(InCollection, "HitIndex", FGeometryCollection::VerticesGroup)
 		, VertexNormalAttribute(InCollection, "Normal", FGeometryCollection::VerticesGroup)
 		, VertexColorAttribute(InCollection, "Color", FGeometryCollection::VerticesGroup)
+		, VertexUVAttribute(InCollection, "UV", FGeometryCollection::VerticesGroup)
 		, IndicesAttribute(InCollection, "Indices", FGeometryCollection::FacesGroup, FGeometryCollection::VerticesGroup)
 		, MaterialIDAttribute(InCollection, "MaterialID", FGeometryCollection::FacesGroup)
 		, TriangleSectionAttribute(InCollection, "Sections", FGeometryCollection::MaterialGroup)
+		, MaterialPathAttribute(InCollection, "MaterialPath", FGeometryCollection::MaterialGroup)
 		, GeometryNameAttribute(InCollection, "Name", FGeometryCollection::GeometryGroup)
 		, GeometryHitProxyIndexAttribute(InCollection, "HitIndex", FGeometryCollection::GeometryGroup)
 		, VertexStartAttribute(InCollection, "VertexStart", FGeometryCollection::GeometryGroup)
 		, VertexCountAttribute(InCollection, "VertexCount", FGeometryCollection::GeometryGroup)
 		, IndicesStartAttribute(InCollection, "IndicesStart", FGeometryCollection::GeometryGroup)
 		, IndicesCountAttribute(InCollection, "IndicesCount", FGeometryCollection::GeometryGroup)
+		, MaterialStartAttribute(InCollection, "MaterialsStart", FGeometryCollection::GeometryGroup, FGeometryCollection::MaterialGroup)
+		, MaterialCountAttribute(InCollection, "MaterialsCount", FGeometryCollection::GeometryGroup)
 		, GeometrySelectionAttribute(InCollection, "SelectionState", FGeometryCollection::GeometryGroup)
 
 	{}
@@ -41,15 +45,19 @@ namespace GeometryCollection::Facades
 		, VertexHitProxyIndexAttribute(InCollection, "HitIndex", FGeometryCollection::VerticesGroup)
 		, VertexNormalAttribute(InCollection, "Normal", FGeometryCollection::VerticesGroup)
 		, VertexColorAttribute(InCollection, "Color", FGeometryCollection::VerticesGroup)
+		, VertexUVAttribute(InCollection, "UV", FGeometryCollection::VerticesGroup)
 		, IndicesAttribute(InCollection, "Indices", FGeometryCollection::FacesGroup, FGeometryCollection::VerticesGroup)
 		, MaterialIDAttribute(InCollection, "MaterialID", FGeometryCollection::FacesGroup)
 		, TriangleSectionAttribute(InCollection, "Sections", FGeometryCollection::MaterialGroup)
+		, MaterialPathAttribute(InCollection, "MaterialPath", FGeometryCollection::MaterialGroup)
 		, GeometryNameAttribute(InCollection, "Name", FGeometryCollection::GeometryGroup)
 		, GeometryHitProxyIndexAttribute(InCollection, "HitIndex", FGeometryCollection::GeometryGroup)
 		, VertexStartAttribute(InCollection, "VertexStart", FGeometryCollection::GeometryGroup, FGeometryCollection::VerticesGroup)
 		, VertexCountAttribute(InCollection, "VertexCount", FGeometryCollection::GeometryGroup)
 		, IndicesStartAttribute(InCollection, "IndicesStart", FGeometryCollection::GeometryGroup, FGeometryCollection::FacesGroup)
 		, IndicesCountAttribute(InCollection, "IndicesCount", FGeometryCollection::GeometryGroup)
+		, MaterialStartAttribute(InCollection, "MaterialsStart", FGeometryCollection::GeometryGroup, FGeometryCollection::MaterialGroup)
+		, MaterialCountAttribute(InCollection, "MaterialsCount", FGeometryCollection::GeometryGroup)
 		, GeometrySelectionAttribute(InCollection, "SelectionState", FGeometryCollection::GeometryGroup)
 	{}
 
@@ -66,15 +74,19 @@ namespace GeometryCollection::Facades
 		VertexHitProxyIndexAttribute.Add();
 		VertexNormalAttribute.Add();
 		VertexColorAttribute.Add();
+		VertexUVAttribute.Add();
 		IndicesAttribute.Add();
 		MaterialIDAttribute.Add();
 		TriangleSectionAttribute.Add();
+		MaterialPathAttribute.Add();
 		GeometryNameAttribute.Add();
 		GeometryHitProxyIndexAttribute.Add();
 		VertexStartAttribute.Add();
 		VertexCountAttribute.Add();
 		IndicesStartAttribute.Add();
 		IndicesCountAttribute.Add();
+		MaterialStartAttribute.Add();
+		MaterialCountAttribute.Add();
 		GeometrySelectionAttribute.Add();
 	}
 
@@ -88,12 +100,15 @@ namespace GeometryCollection::Facades
 		return VertexAttribute.IsValid() && VertexToGeometryIndexAttribute.IsValid() &&
 			VertexSelectionAttribute.IsValid() && VertexHitProxyIndexAttribute.IsValid() &&
 			IndicesAttribute.IsValid() &&
-			MaterialIDAttribute.IsValid() && TriangleSectionAttribute.IsValid() &&
+			MaterialIDAttribute.IsValid() && 
+			TriangleSectionAttribute.IsValid() &&
 			GeometryNameAttribute.IsValid() && GeometryHitProxyIndexAttribute.IsValid() &&
 			VertexStartAttribute.IsValid() && VertexCountAttribute.IsValid() &&
 			IndicesStartAttribute.IsValid() && IndicesCountAttribute.IsValid() &&
 			GeometrySelectionAttribute.IsValid() && VertexColorAttribute.IsValid() &&
-			VertexNormalAttribute.IsValid();
+			VertexNormalAttribute.IsValid() &&
+			VertexUVAttribute.IsValid() &&
+			MaterialPathAttribute.IsValid();
 	}
 
 	int32 FRenderingFacade::NumTriangles() const
@@ -168,6 +183,40 @@ namespace GeometryCollection::Facades
 		}
 	}
 
+	void FRenderingFacade::AddSurface(TArray<FVector3f>&& InVertices, TArray<FIntVector>&& InIndices, TArray<FVector3f>&& InNormals, TArray<FLinearColor>&& InColors, TArray<TArray<FVector2f>>&& InUVs, TArray<int32>&& InMaterialIDs, TArray<FString>&& InMaterialPaths)
+	{
+		check(!IsConst());
+		if (IsValid())
+		{
+			const int32 IndicesStart = IndicesAttribute.Num();
+			const int32 VertexStart = VertexAttribute.Num();
+
+			AddSurface(MoveTemp(InVertices), MoveTemp(InIndices), MoveTemp(InNormals), MoveTemp(InColors));
+
+			// Add Vertex UVs
+			check(InUVs.Num() == InVertices.Num());
+			TManagedArray<TArray<FVector2f>>& DestVertexUVs = VertexUVAttribute.Modify();		
+			for (int32 VertexID = 0; VertexID < InUVs.Num(); ++VertexID)
+			{
+				DestVertexUVs[VertexStart + VertexID] = MoveTemp(InUVs[VertexID]);
+			}
+
+			// Add Material paths
+			const int32 MaterialOffset = MaterialPathAttribute.AddElements(InMaterialPaths.Num());
+			for (int32 InMaterialIndex = 0; InMaterialIndex < InMaterialPaths.Num(); ++InMaterialIndex)
+			{
+				MaterialPathAttribute.Modify()[MaterialOffset + InMaterialIndex] = InMaterialPaths[InMaterialIndex];
+			}
+
+			// Add per-Triangle MaterialIDs
+			check(InMaterialIDs.Num() == InIndices.Num());
+			TManagedArray<int32>& DestTriangleMaterialIDs = MaterialIDAttribute.Modify();
+			for (int32 FaceID = 0; FaceID < InMaterialIDs.Num(); ++FaceID)
+			{
+				DestTriangleMaterialIDs[IndicesStart + FaceID] = InMaterialIDs[FaceID] + MaterialOffset;
+			}
+		}
+	}
 
 	TArray<FRenderingFacade::FTriangleSection> 
 	FRenderingFacade::BuildMeshSections(const TArray<FIntVector>& InputIndices, TArray<int32> BaseMeshOriginalIndicesIndex, TArray<FIntVector>& RetIndices) const
@@ -192,6 +241,8 @@ namespace GeometryCollection::Facades
 			IndicesStartAttribute.Modify()[GeomIndex] = IndicesAttribute.Num();
 			IndicesCountAttribute.Modify()[GeomIndex] = 0;
 			GeometrySelectionAttribute.Modify()[GeomIndex] = 0;
+			MaterialStartAttribute.Modify()[GeomIndex] = MaterialPathAttribute.Num();
+			MaterialCountAttribute.Modify()[GeomIndex] = 0;
 		}
 		return GeomIndex;
 	}
@@ -226,6 +277,15 @@ namespace GeometryCollection::Facades
 			{
 				IndicesStartAttribute.Modify()[InGeomIndex] = IndicesAttribute.Num();
 			}
+
+			if (MaterialStartAttribute.Get()[InGeomIndex] < MaterialPathAttribute.Num())
+			{
+				MaterialCountAttribute.Modify()[InGeomIndex] = MaterialPathAttribute.Num() - MaterialStartAttribute.Get()[InGeomIndex];
+			}
+			else
+			{
+				MaterialStartAttribute.Modify()[InGeomIndex] = MaterialPathAttribute.Num();
+			}
 		}
 	}
 
diff --git a/Engine/Source/Runtime/Experimental/Chaos/Public/GeometryCollection/Facades/CollectionRenderingFacade.h b/Engine/Source/Runtime/Experimental/Chaos/Public/GeometryCollection/Facades/CollectionRenderingFacade.h
--- a/Engine/Source/Runtime/Experimental/Chaos/Public/GeometryCollection/Facades/CollectionRenderingFacade.h
+++ b/Engine/Source/Runtime/Experimental/Chaos/Public/GeometryCollection/Facades/CollectionRenderingFacade.h
@@ -54,7 +54,10 @@ namespace GeometryCollection::Facades
 
 		/** Add a surface to the rendering view.*/
 		CHAOS_API void AddSurface(TArray<FVector3f>&& InVertices, TArray<FIntVector>&& InIndices, TArray<FVector3f>&& InNormals, TArray<FLinearColor>&& InColors);
-		
+
+		/** Add a surface with UV layers, Material IDs and material paths to the rendering view.*/
+		CHAOS_API void AddSurface(TArray<FVector3f>&& InVertices, TArray<FIntVector>&& InIndices, TArray<FVector3f>&& InNormals, TArray<FLinearColor>&& InColors, TArray<TArray<FVector2f>>&& InUVs, TArray<int32>&& InMaterialIDs, TArray<FString>&& MaterialPaths);
+
 		/** GetIndices */
 		const TManagedArray< FIntVector >& GetIndices() const { return IndicesAttribute.Get(); }
 
@@ -99,6 +102,10 @@ namespace GeometryCollection::Facades
 		const TManagedArray<FLinearColor>& GetVertexColor() const { return VertexColorAttribute.Get(); }
 		TManagedArray<FLinearColor>& ModifyVertexColor() { check(!IsConst()); return VertexColorAttribute.Modify(); }
 
+		/** GetVertexUVAttribute */
+		const TManagedArray<TArray<FVector2f>>& GetVertexUV() const { return VertexUVAttribute.Get(); }
+		TManagedArray<TArray<FVector2f>>& ModifyVertexUV() { check(!IsConst()); return VertexUVAttribute.Modify(); }
+
 		//
 		// Geometry Group Attributes
 		//
@@ -140,6 +147,18 @@ namespace GeometryCollection::Facades
 		/** GetGeometryNameToIndexMap */
 		CHAOS_API FStringIntMap GetGeometryNameToIndexMap() const;
 
+		/** GetMaterialPathAttribute */
+		CHAOS_API const TManagedArray<FString>& GetMaterialPaths() const { return MaterialPathAttribute.Get(); }
+
+		/** GetMaterialStart */
+		const TManagedArray< int32 >& GetMaterialStart() const { return MaterialStartAttribute.Get(); }
+
+		/** GetMaterialCount */
+		const TManagedArray< int32 >& GetMaterialCount() const { return MaterialCountAttribute.Get(); }
+
+		/** Total NumMaterials */
+		int32 NumMaterials() const { return MaterialPathAttribute.Num(); }
+
 	private : 
 		const FManagedArrayCollection& ConstCollection;
 		FManagedArrayCollection* Collection = nullptr;
@@ -150,10 +169,13 @@ namespace GeometryCollection::Facades
 		TManagedArrayAccessor<int32> VertexHitProxyIndexAttribute;
 		TManagedArrayAccessor<FVector3f> VertexNormalAttribute;
 		TManagedArrayAccessor<FLinearColor> VertexColorAttribute;
+		TManagedArrayAccessor<TArray<FVector2f>> VertexUVAttribute;
 
 		TManagedArrayAccessor<FIntVector> IndicesAttribute;
 		TManagedArrayAccessor<int32> MaterialIDAttribute;
+
 		TManagedArrayAccessor<FTriangleSection> TriangleSectionAttribute;
+		TManagedArrayAccessor<FString> MaterialPathAttribute;
 
 		TManagedArrayAccessor<FString> GeometryNameAttribute;
 		TManagedArrayAccessor<int32> GeometryHitProxyIndexAttribute;
@@ -161,6 +183,8 @@ namespace GeometryCollection::Facades
 		TManagedArrayAccessor<int32> VertexCountAttribute;
 		TManagedArrayAccessor<int32> IndicesStartAttribute;
 		TManagedArrayAccessor<int32> IndicesCountAttribute;
+		TManagedArrayAccessor<int32> MaterialStartAttribute;
+		TManagedArrayAccessor<int32> MaterialCountAttribute;
 		TManagedArrayAccessor<int32> GeometrySelectionAttribute;
 	};
 
