Fix complaints from Android's MTE and finish how peer are waited on.

The underlying OS objects for each peer need to be gathered so they can collectively waited on. These are exposed as waitable objects so the underlying object state cannot be modified, and waitables are non-copyable to discourage persistence. Waitables are wrapped prior to waiting so the level that knows what to wait for can add appropriate metadata.

#rnx

diff --git a/Engine/Source/Runtime/Experimental/IoStore/HttpClient/Private/Loop.inl b/Engine/Source/Runtime/Experimental/IoStore/HttpClient/Private/Loop.inl
--- a/Engine/Source/Runtime/Experimental/IoStore/HttpClient/Private/Loop.inl
+++ b/Engine/Source/Runtime/Experimental/IoStore/HttpClient/Private/Loop.inl
@@ -737,9 +737,8 @@ class FPeerGroup
 public:
 						FPeerGroup() = default;
 						~FPeerGroup();
-	bool				operator == (FSocket* Rhs) const { return (const FSocket*)Peer == Rhs; }
 	void				Unwait()			{ check(bWaiting); bWaiting = false; }
-	FSocket::FWaiter	GetWaiter() const;
+	FWaiter				GetWaiter() const;
 	bool				Tick(FTickState& State);
 	void				TickSend(FTickState& State, FHost& Host);
 	void				Fail(FTickState& State, const char* Reason);
@@ -767,18 +766,18 @@ FPeerGroup::~FPeerGroup()
 }
 
 ////////////////////////////////////////////////////////////////////////////////
-FSocket::FWaiter FPeerGroup::GetWaiter() const
+FWaiter FPeerGroup::GetWaiter() const
 {
 	if (!bWaiting)
 	{
-		return FSocket::FWaiter();
+		return FWaiter();
 	}
 
-	const auto& S = (const FSocket&)Peer;
+	FWaitable Waitable = Peer.GetWaitable();
 
-	using EWhat = FSocket::FWaiter::EWhat;
-	EWhat What = (Recv != nullptr) ? EWhat::Recv : EWhat::Send;
-	return FSocket::FWaiter(S, What);
+	FWaiter Waiter(MoveTemp(Waitable));
+	Waiter.WaitFor((Recv != nullptr) ? FWaiter::EWhat::Recv : FWaiter::EWhat::Send);
+	return Waiter;
 }
 
 ////////////////////////////////////////////////////////////////////////////////
@@ -1073,14 +1072,17 @@ FHostGroup::FHostGroup(FHost& InHost)
 int32 FHostGroup::Wait(const FTickState& State)
 {
 	// Collect groups that are waiting on something
-	TArray<FSocket::FWaiter, TFixedAllocator<64>> Waiters;
-	for (FPeerGroup& Group : PeerGroups)
+	TArray<FWaiter, TFixedAllocator<64>> Waiters;
+	for (uint32 i = 0, n = PeerGroups.Num(); i < n; ++i)
 	{
-		FSocket::FWaiter Waiter = Group.GetWaiter();
-		if (Waiter.IsValid())
+		FWaiter Waiter = PeerGroups[i].GetWaiter();
+		if (!Waiter.IsValid())
 		{
-			Waiters.Add(Waiter);
+			continue;
 		}
+
+		Waiter.SetIndex(i);
+		Waiters.Add(MoveTemp(Waiter));
 	}
 
 	if (Waiters.IsEmpty())
@@ -1100,7 +1102,7 @@ int32 FHostGroup::Wait(const FTickState& State)
 	}
 
 	// Actually do the wait
-	int32 Result = FSocket::Wait(Waiters, PollTimeoutMs);
+	int32 Result = FWaiter::Wait(Waiters, PollTimeoutMs);
 	if (Result <= 0)
 	{
 		// If the user opts to not block then we don't accumulate wait time and
@@ -1121,16 +1123,14 @@ int32 FHostGroup::Wait(const FTickState& State)
 	int32 Count = 0;
 	for (int32 i = 0, n = Waiters.Num(); i < n; ++i)
 	{
-		if (Waiters[i].Ready == 0)
+		if (!Waiters[i].IsReady())
 		{
 			continue;
 		}
 
-		auto* Candidate = (FSocket*)(Waiters[i].Candidate);
-		auto Pred = [Candidate] (auto& Lhs) { return Lhs == Candidate; };
-		FPeerGroup* Group = PeerGroups.FindByPredicate(Pred);
-		check(Group != nullptr);
-		Group->Unwait();
+		uint32 Index = Waiters[i].GetIndex();
+		check(Index < uint32(PeerGroups.Num()));
+		PeerGroups[Index].Unwait();
 
 		Waiters.RemoveAtSwap(i, EAllowShrinking::No);
 		--n, --i, ++Count;
diff --git a/Engine/Source/Runtime/Experimental/IoStore/HttpClient/Private/Peer.inl b/Engine/Source/Runtime/Experimental/IoStore/HttpClient/Private/Peer.inl
--- a/Engine/Source/Runtime/Experimental/IoStore/HttpClient/Private/Peer.inl
+++ b/Engine/Source/Runtime/Experimental/IoStore/HttpClient/Private/Peer.inl
@@ -159,19 +159,13 @@ class FPeer
 public:
 				FPeer() = default;
 				FPeer(FSocket InSocket);
+	FWaitable	GetWaitable() const					{ return Socket.GetWaitable(); }
 	FOutcome	Send(const char* Data, int32 Size)	{ return Socket.Send(Data, Size); }
 	FOutcome	Recv(char* Out, int32 MaxSize)		{ return Socket.Recv(Out, MaxSize); }
 	bool		IsValid() const						{ return Socket.IsValid(); }
 
 private:
 	FSocket		Socket;
-
-
-
-	// TODO: !!!!!! remove this !!!!!
-public:
-	explicit operator const FSocket& () const { return Socket; }
-	explicit operator const FSocket* () const { return &Socket; }
 };
 
 ////////////////////////////////////////////////////////////////////////////////
diff --git a/Engine/Source/Runtime/Experimental/IoStore/HttpClient/Private/Socket.inl b/Engine/Source/Runtime/Experimental/IoStore/HttpClient/Private/Socket.inl
--- a/Engine/Source/Runtime/Experimental/IoStore/HttpClient/Private/Socket.inl
+++ b/Engine/Source/Runtime/Experimental/IoStore/HttpClient/Private/Socket.inl
@@ -76,27 +76,192 @@
 	static_assert(EWOULDBLOCK == EAGAIN);
 #endif
 
+
+
 ////////////////////////////////////////////////////////////////////////////////
+class FWaitable
+{
+public:
+				FWaitable() = default;
+				FWaitable(FWaitable&&) = default;
+				FWaitable& operator = (FWaitable&&) = default;
+
+private:
+	friend		class FSocket;
+	friend		class FWaiter;
+				FWaitable(SocketType InSocket);
+	SocketType	Socket = InvalidSocket;
+
+private:
+				FWaitable(const FWaitable&) = delete;
+				FWaitable& operator = (const FWaitable&) = delete;
+};
+
+////////////////////////////////////////////////////////////////////////////////
+FWaitable::FWaitable(SocketType InSocket)
+: Socket(InSocket)
+{
+}
+
+
+
+////////////////////////////////////////////////////////////////////////////////
+class FWaiter
+{
+public:
+	enum class EWhat : uint8
+	{
+		None = 0b00,
+		Send = 0b01,
+		Recv = 0b10,
+		Both = Send|Recv
+	};
+
+	static int32	Wait(TArrayView<FWaiter> Waiters, int32 TimeoutMs);
+
+					FWaiter() = default;
+					FWaiter(FWaitable&& Waitable);
+					FWaiter(FWaiter&&) = default;
+	FWaiter&		operator = (FWaiter&&) = default;
+	bool			IsValid() const				{ return Socket != 0; }
+	bool			IsReady() const				{ return Ready != 0; }
+	void			WaitFor(EWhat What)			{ WaitOn = What; }
+	void			SetIndex(uint32 InIndex)	{ Index = uint16(InIndex); }
+	uint32			GetIndex() const			{ return Index; }
+
+private:
+	SocketType		Socket = 0;
+	uint16			Index = 0;
+	EWhat			WaitOn = EWhat::None;
+	uint8			Ready = 0;
+
+private:
+					FWaiter(const FWaiter&) = delete;
+	FWaiter&		operator = (const FWaiter&) = delete;
+};
+
+////////////////////////////////////////////////////////////////////////////////
+FWaiter::FWaiter(FWaitable&& Waitable)
+: Socket(Waitable.Socket)
+{
+}
+
+////////////////////////////////////////////////////////////////////////////////
+int32 FWaiter::Wait(TArrayView<FWaiter> Waiters, int32 TimeoutMs)
+{
+#if !defined(IAS_HTTP_USE_POLL)
+	struct FSelect
+	{
+		SocketType	fd;
+		int32		events;
+		int32		revents;
+	};
+	static const int32 POLLIN  = 1 << 0;
+	static const int32 POLLOUT = 1 << 1;
+	static const int32 POLLERR = 1 << 2;
+	static const int32 POLLHUP = POLLERR;
+	static const int32 POLLNVAL= POLLERR;
+#else
+	using FSelect = pollfd;
+#endif
+
+	// The following looks odd because POLLFD varies subtly from one platform
+	// to the next. To cleanly set members to zero and to not get narrowing
+	// warnings from the compiler, we list-init and don't assume POD types.
+	using PollEventType = decltype(FSelect::events);
+	PollEventType Events[] = { POLLERR, POLLOUT, POLLIN, POLLOUT|POLLOUT };
+	TArray<FSelect, TFixedAllocator<64>> Selects;
+	for (FWaiter& Waiter : Waiters)
+	{
+		Selects.Emplace_GetRef() = {
+			Waiter.Socket,
+			Events[uint32(Waiter.WaitOn)],
+			/* 0 */
+		};
+	}
+
+	// Poll the sockets
+#if defined(IAS_HTTP_USE_POLL)
+	int32 Result = poll(Selects.GetData(), Selects.Num(), TimeoutMs);
+	if (Result <= 0)
+	{
+		return Result;
+	}
+#else
+	timeval TimeVal = {};
+	timeval* TimeValPtr = (TimeoutMs >= 0 ) ? &TimeVal : nullptr;
+	if (TimeoutMs > 0)
+	{
+		TimeVal = { TimeoutMs >> 10, TimeoutMs & ((1 << 10) - 1) };
+	}
+
+	fd_set FdSetRead;	FD_ZERO(&FdSetRead);
+	fd_set FdSetWrite;	FD_ZERO(&FdSetWrite);
+	fd_set FdSetExcept; FD_ZERO(&FdSetExcept);
+
+	SocketType MaxFd = 0;
+	for (FSelect& Select : Selects)
+	{
+		fd_set* RwSet = (Select.events & POLLIN) ? &FdSetRead : &FdSetWrite;
+		FD_SET(Select.fd, RwSet);
+		FD_SET(Select.fd, &FdSetExcept);
+		MaxFd = FMath::Max(Select.fd, MaxFd);
+	}
+
+	int32 Result = select(int32(MaxFd + 1), &FdSetRead, &FdSetWrite, &FdSetExcept, TimeValPtr);
+	if (Result <= 0)
+	{
+		return Result;
+	}
+
+	for (FSelect& Select : Selects)
+	{
+		if (FD_ISSET(Select.fd, &FdSetExcept))
+		{
+			Select.revents = POLLERR;
+			continue;
+		}
+
+		fd_set* RwSet = (Select.events & POLLIN) ? &FdSetRead : &FdSetWrite;
+		if (FD_ISSET(Select.fd, RwSet))
+		{
+			Select.revents = Select.events;
+		}
+	}
+#endif // IAS_HTTP_USE_POLL
+
+	// Transfer poll results to the input sockets. We don't transfer across error
+	// states. Subsequent sockets ops can take care of that instead.
+	static const auto TestBits = POLLIN|POLLOUT|POLLERR|POLLHUP|POLLNVAL;
+	for (uint32 i = 0, n = Waiters.Num(); i < n; ++i)
+	{
+		auto RetEvents = Selects[i].revents;
+		if (!(RetEvents & TestBits))
+		{
+			continue;
+		}
+
+		uint8 Value = 0;
+		if (!!(RetEvents & POLLOUT)) Value |= uint8(EWhat::Send);
+		if (!!(RetEvents & POLLIN))	 Value |= uint8(EWhat::Recv);
+		Waiters[i].Ready = Value ? Value : uint8(EWhat::Both);
+	}
+
+	return Result;
+}
+
+
+
+////////////////////////////////////////////////////////////////////////////////
 class FSocket
 {
 public:
-	struct FWaiter
-	{
-		enum class EWhat { Send = 0b01, Recv = 0b10, Both = Send|Recv };
-				FWaiter() { std::memset(this, 0, sizeof(*this)); }
-				FWaiter(const FSocket& Socket, EWhat InWaitOn);
-		bool	IsValid() const { UPTRINT x{0}; return std::memcmp(this, &x, sizeof(*this)) != 0; }
-		bool	operator == (FSocket& Rhs) const { return UPTRINT(&Rhs) == Candidate; }
-		UPTRINT	Candidate : 60;
-		UPTRINT	WaitOn : 2;
-		UPTRINT	Ready : 2;
-	};
-
 				FSocket() = default;
 				~FSocket()					{ Destroy(); }
 				FSocket(FSocket&& Rhs)		{ Move(MoveTemp(Rhs)); }
 	FSocket&	operator = (FSocket&& Rhs)	{ Move(MoveTemp(Rhs)); return *this; }
 	bool		IsValid() const				{ return Socket != InvalidSocket; }
+	FWaitable	GetWaitable() const			{ return { Socket }; }
 	bool		Create();
 	void		Destroy();
 	FOutcome	Connect(uint32 Ip, uint32 Port);
@@ -106,7 +271,6 @@
 	bool		SetBlocking(bool bBlocking);
 	bool		SetSendBufSize(int32 Size);
 	bool		SetRecvBufSize(int32 Size);
-	static int	Wait(TArrayView<FWaiter> Waiters, int32 TimeoutMs);
 
 private:
 	void		Move(FSocket&& Rhs);
@@ -117,14 +281,5 @@
 };
 
 ////////////////////////////////////////////////////////////////////////////////
-FSocket::FWaiter::FWaiter(const FSocket& Socket, EWhat InWaitOn)
-: Candidate(UPTRINT(&Socket))
-, WaitOn(UPTRINT(InWaitOn))
-, Ready(0)
-{
-	static_assert(sizeof(*this) == sizeof(void*));
-}
-
-////////////////////////////////////////////////////////////////////////////////
 void FSocket::Move(FSocket&& Rhs)
 {
@@ -277,110 +432,6 @@
 	return 0 == setsockopt(Socket, SOL_SOCKET, SO_RCVBUF, &(char&)Size, sizeof(Size));
 }
 
-////////////////////////////////////////////////////////////////////////////////
-int32 FSocket::Wait(TArrayView<FWaiter> Waiters, int32 TimeoutMs)
-{
-#if !defined(IAS_HTTP_USE_POLL)
-	struct FSelect
-	{
-		SocketType	fd;
-		int32		events;
-		int32		revents;
-	};
-	static const int32 POLLIN  = 1 << 0;
-	static const int32 POLLOUT = 1 << 1;
-	static const int32 POLLERR = 1 << 2;
-	static const int32 POLLHUP = POLLERR;
-	static const int32 POLLNVAL= POLLERR;
-#else
-	using FSelect = pollfd;
-#endif
-
-	// The following looks odd because POLLFD varies subtly from one platform
-	// to the next. To cleanly set members to zero and to not get narrowing
-	// warnings from the compiler, we list-init and don't assume POD types.
-	using PollEventType = decltype(FSelect::events);
-	PollEventType Events[] = { POLLERR, POLLOUT, POLLIN, POLLOUT|POLLOUT };
-	TArray<FSelect, TFixedAllocator<64>> Selects;
-	for (FWaiter& Waiter : Waiters)
-	{
-		Selects.Emplace_GetRef() = {
-			((FSocket*)Waiter.Candidate)->Socket,
-			Events[Waiter.WaitOn],
-			/* 0 */
-		};
-	}
-
-	// Poll the sockets
-#if defined(IAS_HTTP_USE_POLL)
-	int32 Result = poll(Selects.GetData(), Selects.Num(), TimeoutMs);
-	if (Result <= 0)
-	{
-		return Result;
-	}
-#else
-	timeval TimeVal = {};
-	timeval* TimeValPtr = (TimeoutMs >= 0 ) ? &TimeVal : nullptr;
-	if (TimeoutMs > 0)
-	{
-		TimeVal = { TimeoutMs >> 10, TimeoutMs & ((1 << 10) - 1) };
-	}
-
-	fd_set FdSetRead;	FD_ZERO(&FdSetRead);
-	fd_set FdSetWrite;	FD_ZERO(&FdSetWrite);
-	fd_set FdSetExcept; FD_ZERO(&FdSetExcept);
-
-	SocketType MaxFd = 0;
-	for (FSelect& Select : Selects)
-	{
-		fd_set* RwSet = (Select.events & POLLIN) ? &FdSetRead : &FdSetWrite;
-		FD_SET(Select.fd, RwSet);
-		FD_SET(Select.fd, &FdSetExcept);
-		MaxFd = FMath::Max(Select.fd, MaxFd);
-	}
-
-	int32 Result = select(int32(MaxFd + 1), &FdSetRead, &FdSetWrite, &FdSetExcept, TimeValPtr);
-	if (Result <= 0)
-	{
-		return Result;
-	}
-
-	for (FSelect& Select : Selects)
-	{
-		if (FD_ISSET(Select.fd, &FdSetExcept))
-		{
-			Select.revents = POLLERR;
-			continue;
-		}
-
-		fd_set* RwSet = (Select.events & POLLIN) ? &FdSetRead : &FdSetWrite;
-		if (FD_ISSET(Select.fd, RwSet))
-		{
-			Select.revents = Select.events;
-		}
-	}
-#endif // IAS_HTTP_USE_POLL
-
-	// Transfer poll results to the input sockets. We don't transfer across error
-	// states. Subsequent sockets ops can take care of that instead.
-	static const auto TestBits = POLLIN|POLLOUT|POLLERR|POLLHUP|POLLNVAL;
-	for (uint32 i = 0, n = Waiters.Num(); i < n; ++i)
-	{
-		auto RetEvents = Selects[i].revents;
-		if (!(RetEvents & TestBits))
-		{
-			continue;
-		}
-
-		uint32 Value = 0;
-		if (!!(RetEvents & POLLOUT)) Value |= uint32(FSocket::FWaiter::EWhat::Send);
-		if (!!(RetEvents & POLLIN))	 Value |= uint32(FSocket::FWaiter::EWhat::Recv);
-		Waiters[i].Ready = Value ? Value : uint32(FSocket::FWaiter::EWhat::Both);
-	}
-
-	return Result;
-}
 
 
-
 // }}}
