Iris
- Improved detection of protocol issues by adding MemberCount, FunctionCount and ChangeMaskBitCount to ReplicatoinStateDescriptorId
- Added error if we fail to find expected property culldistancesquared when building
- Added natvis for ReadJournal

diff --git a/Engine/Source/Runtime/Experimental/Iris/Core/Private/Iris/ReplicationState/ReplicationStateDescriptorBuilder.cpp b/Engine/Source/Runtime/Experimental/Iris/Core/Private/Iris/ReplicationState/ReplicationStateDescriptorBuilder.cpp
--- a/Engine/Source/Runtime/Experimental/Iris/Core/Private/Iris/ReplicationState/ReplicationStateDescriptorBuilder.cpp
+++ b/Engine/Source/Runtime/Experimental/Iris/Core/Private/Iris/ReplicationState/ReplicationStateDescriptorBuilder.cpp
@@ -1812,6 +1812,8 @@ FPropertyReplicationStateDescriptorBuilder::Build(const FString& StateName, FRep
 			Seed = CityHash64((const char*)BuildParams.PathName->GetData(), PathLen * sizeof(FStringBuilderBase::ElementType));
 		}
 		FString UpperCaseNameForHash(StateName);
+		// Append MemberCount, FunctionCount and ChangeMaskBitcount to get some validation of the descriptor.
+		UpperCaseNameForHash.Appendf(TEXT("%u%u%u"), Descriptor->MemberCount, Descriptor->FunctionCount, Descriptor->ChangeMaskBitCount);
 		UpperCaseNameForHash.ToUpperInline();
 		Descriptor->DescriptorIdentifier.Value = CityHash64WithSeed((const char*)UpperCaseNameForHash.GetCharArray().GetData(), sizeof(TCHAR) * UpperCaseNameForHash.GetCharArray().Num(), Seed);
 
@@ -1827,6 +1829,7 @@ FPropertyReplicationStateDescriptorBuilder::Build(const FString& StateName, FRep
 			{
 				// Store the default state checksum for later use when verifying protocol
 				Descriptor->DescriptorIdentifier.DefaultStateHash = DefaultStateHash;
+				UE_LOG_DESCRIPTORBUILDER(Verbose, TEXT("DefaultStateHash for identifier (0x%" UINT64_x_FMT ") is 0x%" UINT64_x_FMT), Descriptor->DescriptorIdentifier.Value, Descriptor->DescriptorIdentifier.DefaultStateHash);
 			}
 		}
 
@@ -2593,6 +2596,12 @@ SIZE_T FReplicationStateDescriptorBuilder::CreateDescriptorsForClass(FResult& Cr
 
 	IRIS_PROFILER_SCOPE(FReplicationStateDescriptorBuilder_CreateDescriptorForClass);
 
+	if (!IsValid(InObjectClass))
+	{
+		UE_LOG(LogIris, Error, TEXT("Trying to create descriptor for invalid class %s"), *GetNameSafe(InObjectClass));
+		return 0U;
+	}
+
 	// Just to make sure we don't mess up the parameters when building descriptor for single property
 	//check(Parameters.SinglePropertyIndex == -1 || Parameters.DescriptorRegistry == nullptr)
 
@@ -2603,6 +2612,7 @@ SIZE_T FReplicationStateDescriptorBuilder::CreateDescriptorsForClass(FResult& Cr
 	{
 		if (const FReplicationStateDescriptorRegistry::FDescriptors* Result = Parameters.DescriptorRegistry->Find(ObjectClassOrArchetypeUsedAsKey, InObjectClass))
 		{
+			UE_LOG_DESCRIPTORBUILDER(VeryVerbose, TEXT("Found existing descriptors for Class %s : ArcheTypeOrKey: %s ptr: 0x%p"), *InObjectClass->GetName(), *ObjectClassOrArchetypeUsedAsKey->GetPathName(), ObjectClassOrArchetypeUsedAsKey);
 			CreatedDescriptors.Insert(*Result, 0);
 			return CreatedDescriptors.Num();
 		}
@@ -2749,6 +2759,11 @@ SIZE_T FReplicationStateDescriptorBuilder::CreateDescriptorsForClass(FResult& Cr
 
 				Builders[LifetimeConditionalsReplicationStateBuilderIndex].AddMemberProperty(NetCullDistanceSqrMemberProperty);
 			}
+			else
+			{
+				UE_LOG(LogIris, Error, TEXT("FPropertyReplicationStateDescriptorBuilder: Failed to find expected property %s.%s"), ToCStr(InObjectClass->GetName()), *FName(ReplicationStateDescriptorBuilder_NAME_NetCullDistanceSquared).ToString());
+				ensure(false);
+			}
 		}
 	}
 
@@ -2862,7 +2877,7 @@ SIZE_T FReplicationStateDescriptorBuilder::CreateDescriptorsForClass(FResult& Cr
 	// If we have a registry, register the created descriptors
 	if (bRegisterCreatedDescriptors && CreatedDescriptors.Num())
 	{
-		UE_LOG_DESCRIPTORBUILDER(Verbose, TEXT("Registering descriptors for  Class %s : ArcheTypeOrKey: %s"), *InObjectClass->GetName(), *ObjectClassOrArchetypeUsedAsKey->GetPathName());
+		UE_LOG_DESCRIPTORBUILDER(VeryVerbose, TEXT("Registering descriptors for  Class %s : ArcheTypeOrKey: %s ptr: 0x%p"), *InObjectClass->GetName(), *ObjectClassOrArchetypeUsedAsKey->GetPathName(), ObjectClassOrArchetypeUsedAsKey);
 		Parameters.DescriptorRegistry->Register(ObjectClassOrArchetypeUsedAsKey, InObjectClass, CreatedDescriptors);
 	}
 	
diff --git a/Engine/Source/Runtime/Experimental/Iris/Core/Private/Iris/ReplicationSystem/ObjectReplicationBridge.cpp b/Engine/Source/Runtime/Experimental/Iris/Core/Private/Iris/ReplicationSystem/ObjectReplicationBridge.cpp
--- a/Engine/Source/Runtime/Experimental/Iris/Core/Private/Iris/ReplicationSystem/ObjectReplicationBridge.cpp
+++ b/Engine/Source/Runtime/Experimental/Iris/Core/Private/Iris/ReplicationSystem/ObjectReplicationBridge.cpp
@@ -923,7 +923,7 @@ FReplicationBridgeCreateNetRefHandleResult UObjectReplicationBridge::CreateNetRe
 
 	if (!ReplicationProtocol)
 	{
-		UE_LOG(LogIris, Error, TEXT("Protocol mismatch prevents binding %s to instanced object %s (CDO: %s)."), *WantedNetHandle.ToString(), *GetNameSafe(InstancePtr), *GetNameSafe(ArchetypeOrCDOUsedAsKey));
+		UE_LOG(LogIris, Error, TEXT("Protocol mismatch prevents binding %s to instanced object %s (CDO: %s)."), *WantedNetHandle.ToString(), *GetNameSafe(InstancePtr), *GetPathNameSafe(ArchetypeOrCDOUsedAsKey));
 
 		if (UE_LOG_ACTIVE(LogIris, Error))
 		{
@@ -2281,7 +2281,7 @@ void UObjectReplicationBridge::OnProtocolMismatchReported(FNetRefHandle RefHandl
 	UObject* ObjInstance = NetRefHandleManager->GetReplicatedObjectInstance(ObjectInternalIndex);
 	UObject* ObjArchetype = ObjInstance ? ObjInstance->GetArchetype() : nullptr;
 
-	UE_LOG(LogIris, Error, TEXT("OnProtocolMismatchReported from client:%s when instancing %s. CDO:%s ReplicatedObject:%s NetObject:%s"), *PrintConnectionInfo(ConnectionId), *RefHandle.ToString(), *GetNameSafe(ObjArchetype), *GetNameSafe(ObjInstance), *NetRefHandleManager->PrintObjectFromIndex(ObjectInternalIndex));
+	UE_LOG(LogIris, Error, TEXT("OnProtocolMismatchReported from client:%s when instancing %s. CDO:%s ReplicatedObject:%s NetObject:%s"), *PrintConnectionInfo(ConnectionId), *RefHandle.ToString(), *GetPathNameSafe(ObjArchetype), *GetNameSafe(ObjInstance), *NetRefHandleManager->PrintObjectFromIndex(ObjectInternalIndex));
 
 	if (UE_LOG_ACTIVE(LogIris, Error))
 	{
@@ -2313,7 +2313,7 @@ void UObjectReplicationBridge::OnProtocolMismatchReported(FNetRefHandle RefHandl
 			Fragments.Emplace(MoveTemp(FragmentInfo));
 		}
 
-		UE::Net::Private::ObjectBridgeDebugging::RemoteProtocolMismatchDetected(ArchetypesAlreadyPrinted, ReplicationSystem, 0 /*TODO: Local ConnectionId*/, Fragments, ObjArchetype, ObjInstance);
+		UE::Net::Private::ObjectBridgeDebugging::RemoteProtocolMismatchDetected(ArchetypesAlreadyPrinted, ReplicationSystem, ConnectionId, Fragments, ObjArchetype, ObjInstance);
 	}
 }
 
diff --git a/Engine/Source/Runtime/Experimental/Iris/Core/Private/Iris/ReplicationSystem/ReplicationProtocolManager.cpp b/Engine/Source/Runtime/Experimental/Iris/Core/Private/Iris/ReplicationSystem/ReplicationProtocolManager.cpp
--- a/Engine/Source/Runtime/Experimental/Iris/Core/Private/Iris/ReplicationSystem/ReplicationProtocolManager.cpp
+++ b/Engine/Source/Runtime/Experimental/Iris/Core/Private/Iris/ReplicationSystem/ReplicationProtocolManager.cpp
@@ -225,7 +225,7 @@ const FReplicationProtocol* FReplicationProtocolManager::CreateReplicationProtoc
 		const FReplicationProtocolIdentifier NewProtocolId = CalculateProtocolIdentifier(Fragments);
 		if (NewProtocolId != ProtocolId)
 		{
-			UE_LOG(LogIris, Warning, TEXT("FReplicationProtocolManager::CreateReplicationProtocol Id mismatch when creating protocol named %s with ProtocolId:0x%x"), DebugName, ProtocolId);
+			UE_LOG(LogIris, Warning, TEXT("FReplicationProtocolManager::CreateReplicationProtocol Id mismatch when creating protocol named %s with in ProtocolId:0x%x Calculated ProtocolId:0x%x"), DebugName, ProtocolId, NewProtocolId);
  #if UE_NET_ENABLE_PROTOCOLMANAGER_LOG
  			if (UE_LOG_ACTIVE(LogIris, Warning))
  			{
diff --git a/Engine/Source/Runtime/Experimental/Iris/Core/Public/Iris/Serialization/NetJournal.h b/Engine/Source/Runtime/Experimental/Iris/Core/Public/Iris/Serialization/NetJournal.h
--- a/Engine/Source/Runtime/Experimental/Iris/Core/Public/Iris/Serialization/NetJournal.h
+++ b/Engine/Source/Runtime/Experimental/Iris/Core/Public/Iris/Serialization/NetJournal.h
@@ -27,7 +27,7 @@ namespace UE::Net
 // Simple journal to track last few entries of read data
 class FNetJournal
 {
-	static constexpr uint32 JournalSize = 8U;
+	static constexpr uint32 JournalSize = 32U;
 	static constexpr uint32 JournalMask = JournalSize - 1U;
 
 public:
