Removed unnecessary UE_CHAOS_ASYNC_INITBODY_WRITESCOPELOCK in FChaosMarshallingManager::AddSimCallbackInputData_External as it is only called from the Gamethread.

diff --git a/Engine/Source/Runtime/Experimental/Chaos/Public/Chaos/ChaosMarshallingManager.h b/Engine/Source/Runtime/Experimental/Chaos/Public/Chaos/ChaosMarshallingManager.h
--- a/Engine/Source/Runtime/Experimental/Chaos/Public/Chaos/ChaosMarshallingManager.h
+++ b/Engine/Source/Runtime/Experimental/Chaos/Public/Chaos/ChaosMarshallingManager.h
@@ -348,7 +348,6 @@ public:
 
 	void AddSimCallbackInputData_External(ISimCallbackObject* SimCallbackObject, FSimCallbackInput* InputData)
 	{
-		UE_CHAOS_ASYNC_INITBODY_WRITESCOPELOCK(MarshallingManagerLock);
 		ProducerData->SimCallbackInputs.Add(FSimCallbackInputAndObject{ SimCallbackObject, InputData });
 	}
 	/** Step forward using the external delta time. Should only be called by external thread */
