Iris
- Bumped size of FGroupIndexType to 32 bits
- Now allows GroupIndex up to 24 bits

diff --git a/Engine/Source/Runtime/Experimental/Iris/Core/Private/Iris/ReplicationSystem/Filtering/NetObjectGroups.cpp b/Engine/Source/Runtime/Experimental/Iris/Core/Private/Iris/ReplicationSystem/Filtering/NetObjectGroups.cpp
--- a/Engine/Source/Runtime/Experimental/Iris/Core/Private/Iris/ReplicationSystem/Filtering/NetObjectGroups.cpp
+++ b/Engine/Source/Runtime/Experimental/Iris/Core/Private/Iris/ReplicationSystem/Filtering/NetObjectGroups.cpp
@@ -22,8 +22,12 @@ static int32 CVarEnsureIfNumGroupMembershipsExceedsNum = 128;
 static FAutoConsoleVariableRef CVarEnsureIfNumGroupMembershipsExceeds(TEXT("net.Iris.EnsureIfNumGroupMembershipsExceeds"), CVarEnsureIfNumGroupMembershipsExceedsNum, TEXT("If set to a positive number we will warn and ensure if an object is added to a high number of groups."), ECVF_Default );
 
 FNetObjectGroups::FNetObjectGroups()
-: CurrentEpoch(++NetObjectGroupsInternal::NextEpoch)
+: CurrentEpoch((++NetObjectGroupsInternal::NextEpoch) & FNetObjectGroupHandle::EpochMask)
 {
+	if (CurrentEpoch == 0U)
+	{
+		UE_LOG(LogIris, Warning, TEXT("FNetObjectGroups::Epoch wraparound detected."));
+	}
 }
 
 FNetObjectGroups::~FNetObjectGroups()
@@ -56,8 +60,8 @@ void FNetObjectGroups::Init(const FNetObjectGroupInitParams& Params)
 {
 	NetRefHandleManager = Params.NetRefHandleManager;
 
-	ensureMsgf(Params.MaxGroupCount <= std::numeric_limits<FNetObjectGroupHandle::FGroupIndexType>::max(), TEXT("MaxGroupCount cannot exceed %u"), std::numeric_limits<FNetObjectGroupHandle::FGroupIndexType>::max());
-	MaxGroupCount = FMath::Clamp<uint32>(Params.MaxGroupCount, 0U, std::numeric_limits<FNetObjectGroupHandle::FGroupIndexType>::max());
+	ensureMsgf(Params.MaxGroupCount <= FNetObjectGroupHandle::MaxGroupIndexCount, TEXT("MaxGroupCount cannot exceed %u"), FNetObjectGroupHandle::MaxGroupIndexCount);
+	MaxGroupCount = FMath::Clamp<uint32>(Params.MaxGroupCount, 0U, FNetObjectGroupHandle::MaxGroupIndexCount);
 
 	// Reserve first as invalid group
 	Groups.Add(FNetObjectGroup());
diff --git a/Engine/Source/Runtime/Experimental/Iris/Core/Public/Iris/ReplicationSystem/NetObjectGroupHandle.h b/Engine/Source/Runtime/Experimental/Iris/Core/Public/Iris/ReplicationSystem/NetObjectGroupHandle.h
--- a/Engine/Source/Runtime/Experimental/Iris/Core/Public/Iris/ReplicationSystem/NetObjectGroupHandle.h
+++ b/Engine/Source/Runtime/Experimental/Iris/Core/Public/Iris/ReplicationSystem/NetObjectGroupHandle.h
@@ -21,7 +21,12 @@ namespace UE::Net
 class FNetObjectGroupHandle
 {
 public:
-	typedef uint16 FGroupIndexType;
+	typedef uint32 FGroupIndexType;
+
+	enum { GroupIndexBits = 24U };
+	enum { EpochBits = 8U };
+	enum { EpochMask = (1U << EpochBits) - 1U };
+	enum { MaxGroupIndexCount = 1U << GroupIndexBits };
 
 	/** Reserved group indices */
 	enum : FGroupIndexType
@@ -77,8 +82,8 @@ private:
 		}
 		else
 		{
-			Index = IndexIn;
-			Epoch = EpochIn;
+			Index = (uint32)IndexIn;
+			Epoch = (uint32)EpochIn;
 			UniqueId = InUniqueId;
 		}
 	}
@@ -88,8 +93,8 @@ private:
 		uint64 Value;
 		struct
 		{
-			FGroupIndexType Index;
-			FGroupIndexType Epoch;
+			uint32 Index : GroupIndexBits;
+			uint32 Epoch : EpochBits;
 			uint32 UniqueId;
 		};
 	};
