Iris: Lower log level of "IrisGroup: Created GroupHandle named" message

diff --git a/Engine/Source/Runtime/Experimental/Iris/Core/Private/Iris/ReplicationSystem/Filtering/NetObjectGroups.cpp b/Engine/Source/Runtime/Experimental/Iris/Core/Private/Iris/ReplicationSystem/Filtering/NetObjectGroups.cpp
--- a/Engine/Source/Runtime/Experimental/Iris/Core/Private/Iris/ReplicationSystem/Filtering/NetObjectGroups.cpp
+++ b/Engine/Source/Runtime/Experimental/Iris/Core/Private/Iris/ReplicationSystem/Filtering/NetObjectGroups.cpp
@@ -109,7 +109,7 @@ FNetObjectGroupHandle FNetObjectGroups::CreateGroup(FName InGroupName)
 	
 	const uint32 Index = (uint32)Groups.Emplace(FNetObjectGroup{.GroupName = InGroupName, .GroupId = NewGroupId});
 
-	UE_LOG(LogIrisGroup, Log, TEXT("Created GroupHandle named: %s Index: %u Id: %u"), *InGroupName.ToString(), Index, NewGroupId);
+	UE_LOG(LogIrisGroup, Verbose, TEXT("Created GroupHandle named: %s Index: %u Id: %u"), *InGroupName.ToString(), Index, NewGroupId);
 
 	FNetObjectGroupHandle GroupHandle;
 	GroupHandle.Index = static_cast<FNetObjectGroupHandle::FGroupIndexType>(Index);
