Fix server assert
Fix issue in Geometry Collection when synching GT and PT particles when destruction occcurs

diff --git a/Engine/Source/Runtime/Experimental/Chaos/Private/PhysicsProxy/GeometryCollectionPhysicsProxy.cpp b/Engine/Source/Runtime/Experimental/Chaos/Private/PhysicsProxy/GeometryCollectionPhysicsProxy.cpp
--- a/Engine/Source/Runtime/Experimental/Chaos/Private/PhysicsProxy/GeometryCollectionPhysicsProxy.cpp
+++ b/Engine/Source/Runtime/Experimental/Chaos/Private/PhysicsProxy/GeometryCollectionPhysicsProxy.cpp
@@ -1708,8 +1708,10 @@ void FGeometryCollectionPhysicsProxy::SyncParticles_External()
 			FParticle* GTParticle = GTParticles[Index].Get();
 			check(GTParticle != nullptr);
 			FClusterHandle* Handle = SolverParticleHandles[Index];
-			check(Handle != nullptr);
-			Handle->GTGeometryParticle() = GTParticle;
+			if (Handle != nullptr)
+			{
+				Handle->GTGeometryParticle() = GTParticle;
+			}
 		}
 	});
 }
