Iris
- Fixed issue with client assigned NetObjectReferences with a relative patch not reloading package on server.

diff --git a/Engine/Source/Runtime/Experimental/Iris/Core/Private/Iris/ReplicationSystem/ObjectReferenceCache.cpp b/Engine/Source/Runtime/Experimental/Iris/Core/Private/Iris/ReplicationSystem/ObjectReferenceCache.cpp
--- a/Engine/Source/Runtime/Experimental/Iris/Core/Private/Iris/ReplicationSystem/ObjectReferenceCache.cpp
+++ b/Engine/Source/Runtime/Experimental/Iris/Core/Private/Iris/ReplicationSystem/ObjectReferenceCache.cpp
@@ -1050,10 +1050,11 @@ ENetObjectReferenceResolveResult FObjectReferenceCache::ResolveObjectReference(c
 			constexpr bool bReading = true;
 			RenamePathForPie(ResolveContext.ConnectionId, ObjectPath, bReading);
 
-			// If both path and handle is valid this is a reference to a dynamic object with a relative path
+			// If both path and handle is valid this is a reference to an object with a relative path.
 			if (Reference.GetRefHandle().IsValid())
 			{
-				UObject* ReplicatedOuter = GetObjectFromReferenceHandle(Reference.GetRefHandle());
+				bool bLocalMustBeMapped = false;
+				UObject* ReplicatedOuter = ResolveObjectReferenceHandleInternal(Reference.GetRefHandle(), ResolveContext, bLocalMustBeMapped);
 				ResolvedObject = ReplicatedOuter ? StaticFindObject(UObject::StaticClass(), ReplicatedOuter, ToCStr(ObjectPath), false) : nullptr;
 
 				if (ResolvedObject == nullptr)
