Dataflow - Fixed greedy Dataflow evaluations.
- Removed unnecessary disconnections when adding connections from an already connected output.
- Disabled invalidations when breaking input connections prior to reconnecting them to a new output.
- Added connection logging.

[RN] Dataflow - Fixed redundant evaluations when connecting an output to multiple inputs, and when replacing an already connected input connection.

diff --git a/Engine/Plugins/Experimental/Dataflow/Source/DataflowEditor/Private/Dataflow/DataflowEditorCommands.cpp b/Engine/Plugins/Experimental/Dataflow/Source/DataflowEditor/Private/Dataflow/DataflowEditorCommands.cpp
--- a/Engine/Plugins/Experimental/Dataflow/Source/DataflowEditor/Private/Dataflow/DataflowEditorCommands.cpp
+++ b/Engine/Plugins/Experimental/Dataflow/Source/DataflowEditor/Private/Dataflow/DataflowEditorCommands.cpp
@@ -148,7 +148,7 @@ bool FDataflowEditorCommands::IsRegistered()
 const FDataflowNode* FDataflowEditorCommands::EvaluateNode(UE::Dataflow::FContext& Context, UE::Dataflow::FTimestamp& InOutLastNodeTimestamp,
 	const UDataflow* Dataflow, const FDataflowNode* Node, const FDataflowOutput* Output, const FString& NodeName, UObject* Asset)
 {
-	UE_LOG(LogChaosDataflow, Verbose, TEXT("FDataflowEditorCommands::EvaluateNode(): Node [%s], NodeName [%s] Output [%s]"), Node ? *Node->GetName().ToString() : TEXT("nullptr"), *NodeName, Output ? *Output->GetName().ToString() : TEXT("nullptr"));
+	UE_LOG(LogChaosDataflow, VeryVerbose, TEXT("FDataflowEditorCommands::EvaluateNode(): Node [%s], NodeName [%s] Output [%s]"), Node ? *Node->GetName().ToString() : TEXT("nullptr"), *NodeName, Output ? *Output->GetName().ToString() : TEXT("nullptr"));
 
 	if (!Node && Dataflow)
 	{
diff --git a/Engine/Plugins/Experimental/Dataflow/Source/DataflowEditor/Private/Dataflow/DataflowEditorToolkit.cpp b/Engine/Plugins/Experimental/Dataflow/Source/DataflowEditor/Private/Dataflow/DataflowEditorToolkit.cpp
--- a/Engine/Plugins/Experimental/Dataflow/Source/DataflowEditor/Private/Dataflow/DataflowEditorToolkit.cpp
+++ b/Engine/Plugins/Experimental/Dataflow/Source/DataflowEditor/Private/Dataflow/DataflowEditorToolkit.cpp
@@ -1166,7 +1166,7 @@ TStatId FDataflowEditorToolkit::GetStatId() const
 
 void FDataflowEditorToolkit::EvaluateNode(const FDataflowNode* Node, const FDataflowOutput* Output, UE::Dataflow::FTimestamp& InOutTimestamp)
 {
-	UE_LOG(LogChaosDataflow, Verbose, TEXT("FDataflowEditorToolkit::EvaluateNode(): Node [%s], Output [%s]"), Node ? *Node->GetName().ToString() : TEXT("nullptr"), Output ? *Output->GetName().ToString() : TEXT("nullptr"));
+	UE_LOG(LogChaosDataflow, VeryVerbose, TEXT("FDataflowEditorToolkit::EvaluateNode(): Node [%s], Output [%s]"), Node ? *Node->GetName().ToString() : TEXT("nullptr"), Output ? *Output->GetName().ToString() : TEXT("nullptr"));
 
 	const bool bIsInPIEOrSimulate = GEditor->PlayWorld || GEditor->bIsSimulatingInEditor;
 	if (!bIsInPIEOrSimulate)  // TODO: make this test optional for some toolkit
diff --git a/Engine/Plugins/Experimental/Dataflow/Source/DataflowEditor/Private/Dataflow/DataflowGraphEditor.cpp b/Engine/Plugins/Experimental/Dataflow/Source/DataflowEditor/Private/Dataflow/DataflowGraphEditor.cpp
--- a/Engine/Plugins/Experimental/Dataflow/Source/DataflowEditor/Private/Dataflow/DataflowGraphEditor.cpp
+++ b/Engine/Plugins/Experimental/Dataflow/Source/DataflowEditor/Private/Dataflow/DataflowGraphEditor.cpp
@@ -179,7 +179,7 @@ TSharedPtr<UE::Dataflow::FContext> SDataflowGraphEditor::GetDataflowContext() co
 
 void SDataflowGraphEditor::EvaluateNode()
 {
-	UE_LOG(LogChaosDataflow, Verbose, TEXT("SDataflowGraphEditor::EvaluateNode(): Nodes [%s]"),
+	UE_LOG(LogChaosDataflow, VeryVerbose, TEXT("SDataflowGraphEditor::EvaluateNode(): Nodes [%s]"),
 		*FString::JoinBy(GetSelectedNodes().Array(), TEXT(", "), [](const UObject* SelectedNode)
 			{
 				return Cast<UDataflowEdNode>(SelectedNode) && Cast<UDataflowEdNode>(SelectedNode)->GetDataflowNode() ? 
diff --git a/Engine/Plugins/Experimental/Dataflow/Source/DataflowEditor/Private/Dataflow/DataflowSchema.cpp b/Engine/Plugins/Experimental/Dataflow/Source/DataflowEditor/Private/Dataflow/DataflowSchema.cpp
--- a/Engine/Plugins/Experimental/Dataflow/Source/DataflowEditor/Private/Dataflow/DataflowSchema.cpp
+++ b/Engine/Plugins/Experimental/Dataflow/Source/DataflowEditor/Private/Dataflow/DataflowSchema.cpp
@@ -344,6 +344,29 @@ void UDataflowSchema::BreakPinLinks(UEdGraphPin& TargetPin, bool bSendsNodeNotif
 	Super::BreakPinLinks(TargetPin, bSendsNodeNotifcation);
 }
 
+bool UDataflowSchema::TryCreateConnection(UEdGraphPin* PinA, UEdGraphPin* PinB) const
+{
+	check(PinA && PinB);
+	UDataflowEdNode* const DataflowEdNodeA = CastChecked<UDataflowEdNode>(PinA->GetOwningNodeUnchecked());
+	UDataflowEdNode* const DataflowEdNodeB = CastChecked<UDataflowEdNode>(PinB->GetOwningNodeUnchecked());
+	if (ensure(DataflowEdNodeA->IsBound() && DataflowEdNodeB->IsBound()))
+	{
+		const TSharedPtr<FDataflowNode> DataflowNodeA = DataflowEdNodeA->GetDataflowNode();
+		const TSharedPtr<FDataflowNode> DataflowNodeB = DataflowEdNodeB->GetDataflowNode();
+		if (ensure(DataflowNodeA && DataflowNodeB))
+		{
+			// Pausing invalidations is a quick hack while sorting the invalidation callbacks that are causing multiple evaluations
+			DataflowNodeA->PauseInvalidations();
+			DataflowNodeB->PauseInvalidations();
+			const bool bModified = Super::TryCreateConnection(PinA, PinB);
+			DataflowNodeA->ResumeInvalidations();
+			DataflowNodeB->ResumeInvalidations();
+			return bModified;
+		}
+	}
+	return Super::TryCreateConnection(PinA, PinB);
+}
+
 FConnectionDrawingPolicy* UDataflowSchema::CreateConnectionDrawingPolicy(int32 InBackLayerID, int32 InFrontLayerID, float InZoomFactor, const FSlateRect& InClippingRect, class FSlateWindowElementList& InDrawElements, class UEdGraph* InGraphObj) const
 {
 	return new FDataflowConnectionDrawingPolicy(InBackLayerID, InFrontLayerID, InZoomFactor, InClippingRect, InDrawElements, InGraphObj);
diff --git a/Engine/Plugins/Experimental/Dataflow/Source/DataflowEditor/Public/Dataflow/DataflowSchema.h b/Engine/Plugins/Experimental/Dataflow/Source/DataflowEditor/Public/Dataflow/DataflowSchema.h
--- a/Engine/Plugins/Experimental/Dataflow/Source/DataflowEditor/Public/Dataflow/DataflowSchema.h
+++ b/Engine/Plugins/Experimental/Dataflow/Source/DataflowEditor/Public/Dataflow/DataflowSchema.h
@@ -17,7 +17,7 @@ class DATAFLOWEDITOR_API UDataflowSchema : public UEdGraphSchema
 public:
 	UDataflowSchema();
 
-	//~ Begin EdGraphSchema Interface
+	//~ Begin UEdGraphSchema Interface
 	virtual const FPinConnectionResponse CanCreateConnection(const UEdGraphPin* PinA, const UEdGraphPin* PinB) const override;
 	virtual void GetContextMenuActions(class UToolMenu* Menu, class UGraphNodeContextMenuContext* Context) const override;
 	virtual void GetGraphContextActions(FGraphContextMenuBuilder& ContextMenuBuilder) const override;
@@ -26,7 +26,8 @@ public:
 	virtual FConnectionDrawingPolicy* CreateConnectionDrawingPolicy(int32 InBackLayerID, int32 InFrontLayerID, float InZoomFactor, const FSlateRect& InClippingRect, class FSlateWindowElementList& InDrawElements, class UEdGraph* InGraphObj) const override;
 	virtual void OnPinConnectionDoubleCicked(UEdGraphPin* PinA, UEdGraphPin* PinB, const FVector2D& GraphPosition) const override;
 	virtual void BreakPinLinks(UEdGraphPin& TargetPin, bool bSendsNodeNotifcation) const override;
-	//~ End EdGraphSchema Interface
+	virtual bool TryCreateConnection(UEdGraphPin* PinA, UEdGraphPin* PinB) const override;
+	//~ End UEdGraphSchema Interface
 
 	static FLinearColor GetTypeColor(const FName& Type);
 };
diff --git a/Engine/Source/Runtime/Experimental/Dataflow/Core/Private/Dataflow/DataflowGraph.cpp b/Engine/Source/Runtime/Experimental/Dataflow/Core/Private/Dataflow/DataflowGraph.cpp
--- a/Engine/Source/Runtime/Experimental/Dataflow/Core/Private/Dataflow/DataflowGraph.cpp
+++ b/Engine/Source/Runtime/Experimental/Dataflow/Core/Private/Dataflow/DataflowGraph.cpp
@@ -116,22 +116,49 @@ namespace UE::Dataflow
 	{
 		if (ensure(OutputConnection && InputConnection))
 		{
-			OutputConnection->AddConnection(InputConnection);
-			InputConnection->AddConnection(OutputConnection);
-			Connections.Add(FLink(
-				OutputConnection->GetOwningNode()->GetGuid(), OutputConnection->GetGuid(),
-				InputConnection->GetOwningNode()->GetGuid(), InputConnection->GetGuid()));
+			FDataflowOutput* const OldOutputConnection = InputConnection->GetConnection();
+			if (OldOutputConnection != OutputConnection)
+			{
+				if (OldOutputConnection)
+				{
+					UE_LOG(LogChaosDataflow, Verbose, TEXT("FGraph::Connect(): Disconnecting output [%s:%s] from input [%s:%s]"),
+						OldOutputConnection && OldOutputConnection->GetOwningNode() ? *OldOutputConnection->GetOwningNode()->GetName().ToString() : TEXT("Invalid"),
+						OldOutputConnection ? *OldOutputConnection->GetName().ToString() : TEXT("Invalid"),
+						InputConnection->GetOwningNode() ? *InputConnection->GetOwningNode()->GetName().ToString() : TEXT("Invalid"),
+						*InputConnection->GetName().ToString());
+					// Note: Do not remove the expired connection from the input to avoid an unnecessary invalidation.
+					//       Simply clobber it with calling AddConnection() on the input instead.
+					OldOutputConnection->RemoveConnection(InputConnection);
+					Connections.RemoveSwap(FLink(
+						OldOutputConnection->GetOwningNode()->GetGuid(), OldOutputConnection->GetGuid(),
+						InputConnection->GetOwningNode()->GetGuid(), InputConnection->GetGuid()));
+				}
+				UE_LOG(LogChaosDataflow, Verbose, TEXT("FGraph::Connect(): Connecting output [%s:%s] to input [%s:%s]"),
+					OutputConnection->GetOwningNode() ? *OutputConnection->GetOwningNode()->GetName().ToString() : TEXT("Invalid"),
+					*OutputConnection->GetName().ToString(),
+					InputConnection->GetOwningNode() ? *InputConnection->GetOwningNode()->GetName().ToString() : TEXT("Invalid"),
+					*InputConnection->GetName().ToString());
+				OutputConnection->AddConnection(InputConnection);
+				InputConnection->AddConnection(OutputConnection);
+				Connections.Add(FLink(
+					OutputConnection->GetOwningNode()->GetGuid(), OutputConnection->GetGuid(),
+					InputConnection->GetOwningNode()->GetGuid(), InputConnection->GetGuid()));
+			}
 		}
 	}
 
 	void FGraph::Disconnect(FDataflowOutput* OutputConnection, FDataflowInput* InputConnection)
 	{
+		UE_LOG(LogChaosDataflow, Verbose, TEXT("FGraph::Diconnect(): Disconnecting output [%s:%s] from input [%s:%s]"),
+			OutputConnection->GetOwningNode() ? *OutputConnection->GetOwningNode()->GetName().ToString() : TEXT("Invalid"),
+			*OutputConnection->GetName().ToString(),
+			InputConnection->GetOwningNode() ? *InputConnection->GetOwningNode()->GetName().ToString() : TEXT("Invalid"),
+			*InputConnection->GetName().ToString());
 		OutputConnection->RemoveConnection(InputConnection);
 		InputConnection->RemoveConnection(OutputConnection);
 		Connections.RemoveSwap(FLink(
 			OutputConnection->GetOwningNode()->GetGuid(), OutputConnection->GetGuid(),
 			InputConnection->GetOwningNode()->GetGuid(), InputConnection->GetGuid()));
-
 	}
 
 	void FGraph::AddReferencedObjects(FReferenceCollector& Collector)
diff --git a/Engine/Source/Runtime/Experimental/Dataflow/Core/Private/Dataflow/DataflowInputOutput.cpp b/Engine/Source/Runtime/Experimental/Dataflow/Core/Private/Dataflow/DataflowInputOutput.cpp
--- a/Engine/Source/Runtime/Experimental/Dataflow/Core/Private/Dataflow/DataflowInputOutput.cpp
+++ b/Engine/Source/Runtime/Experimental/Dataflow/Core/Private/Dataflow/DataflowInputOutput.cpp
@@ -250,7 +250,7 @@ bool FDataflowOutput::Evaluate(UE::Dataflow::FContext& Context) const
 
 bool FDataflowOutput::EvaluateImpl(UE::Dataflow::FContext& Context) const
 {
-	UE_LOG(LogChaosDataflow, Verbose, TEXT("FDataflowOutput::EvaluateImpl(): Node [%s], Output [%s]"), *GetOwningNode()->GetName().ToString(), *GetName().ToString());
+	UE_LOG(LogChaosDataflow, VeryVerbose, TEXT("FDataflowOutput::EvaluateImpl(): Node [%s], Output [%s]"), *GetOwningNode()->GetName().ToString(), *GetName().ToString());
 
 	UE::Dataflow::FContextScopedCallstack Callstack(Context, this);
 	if (Callstack.IsLoopDetected())
@@ -262,7 +262,7 @@ bool FDataflowOutput::EvaluateImpl(UE::Dataflow::FContext& Context) const
 	// check if the cache has a valid version
 	if(Context.HasData(CacheKey(), OwningNode->GetTimestamp()))
 	{
-		UE_LOG(LogChaosDataflow, Verbose, TEXT("FDataflowOutput::EvaluateImpl(): Context has data, NodeTimestamp [%ul], CacheTimestamp [%lu]"), OwningNode->GetTimestamp().Value, Context.GetTimestamp(CacheKey()).Value);
+		UE_LOG(LogChaosDataflow, VeryVerbose, TEXT("FDataflowOutput::EvaluateImpl(): Context has data, NodeTimestamp [%ul], CacheTimestamp [%lu]"), OwningNode->GetTimestamp().Value, Context.GetTimestamp(CacheKey()).Value);
 		return true;
 	}
 	// if not, evaluate
diff --git a/Engine/Source/Runtime/Experimental/Dataflow/Core/Private/Dataflow/DataflowNodeParameters.cpp b/Engine/Source/Runtime/Experimental/Dataflow/Core/Private/Dataflow/DataflowNodeParameters.cpp
--- a/Engine/Source/Runtime/Experimental/Dataflow/Core/Private/Dataflow/DataflowNodeParameters.cpp
+++ b/Engine/Source/Runtime/Experimental/Dataflow/Core/Private/Dataflow/DataflowNodeParameters.cpp
@@ -119,7 +119,7 @@ namespace UE::Dataflow
 
 	bool FContextSingle::Evaluate(const FDataflowOutput& Connection)
 	{
-		UE_LOG(LogChaosDataflow, Verbose, TEXT("FContextSingle::Evaluate(): Node [%s], Output [%s]"), *Connection.GetOwningNode()->GetName().ToString(), *Connection.GetName().ToString());
+		UE_LOG(LogChaosDataflow, VeryVerbose, TEXT("FContextSingle::Evaluate(): Node [%s], Output [%s]"), *Connection.GetOwningNode()->GetName().ToString(), *Connection.GetName().ToString());
 		return Connection.EvaluateImpl(*this);
 	}
 
@@ -132,7 +132,7 @@ namespace UE::Dataflow
 
 	bool FContextThreaded::Evaluate(const FDataflowOutput& Connection)
 	{
-		UE_LOG(LogChaosDataflow, Verbose, TEXT("FContextThreaded::Evaluate(): Node [%s], Output [%s]"), *Connection.GetOwningNode()->GetName().ToString(), *Connection.GetName().ToString());
+		UE_LOG(LogChaosDataflow, VeryVerbose, TEXT("FContextThreaded::Evaluate(): Node [%s], Output [%s]"), *Connection.GetOwningNode()->GetName().ToString(), *Connection.GetName().ToString());
 		Connection.OutputLock->Lock(); ON_SCOPE_EXIT{ Connection.OutputLock->Unlock(); };
 		return Connection.EvaluateImpl(*this);
 	}
diff --git a/Engine/Source/Runtime/Experimental/Dataflow/Engine/Private/Dataflow/DataflowEdNode.cpp b/Engine/Source/Runtime/Experimental/Dataflow/Engine/Private/Dataflow/DataflowEdNode.cpp
--- a/Engine/Source/Runtime/Experimental/Dataflow/Engine/Private/Dataflow/DataflowEdNode.cpp
+++ b/Engine/Source/Runtime/Experimental/Dataflow/Engine/Private/Dataflow/DataflowEdNode.cpp
@@ -336,20 +336,43 @@ void UDataflowEdNode::PinConnectionListChanged(UEdGraphPin* Pin)
 	{
 		if (TSharedPtr<FDataflowNode> DataflowNode = DataflowGraph->FindBaseNode(DataflowNodeGuid))
 		{
+			check(Pin);
+
+			// Return whether a pin in the UEdGraph matches the specified connection in the Dataflow::FGraph
+			auto MatchesConnection = [this](const UEdGraphPin* const Pin, const FDataflowConnection* const Connection) -> bool
+				{
+					if (const UDataflowEdNode* const LinkedNode = Cast<UDataflowEdNode>(Pin->GetOwningNode()))
+					{
+						if (ensure(LinkedNode->IsBound()))
+						{
+							if (const TSharedPtr<FDataflowNode> LinkedDataflowNode = DataflowGraph->FindBaseNode(LinkedNode->GetDataflowNodeGuid()))
+							{
+								return
+									(Pin->Direction == EEdGraphPinDirection::EGPD_Input &&
+										Connection == LinkedDataflowNode->FindInput(FName(Pin->GetName()))) ||
+									(Pin->Direction == EEdGraphPinDirection::EGPD_Output &&
+										Connection == LinkedDataflowNode->FindOutput(FName(Pin->GetName())));
+							}
+						}
+					}
+					return false;
+				};
+
 			if (Pin->Direction == EEdGraphPinDirection::EGPD_Input)
 			{
-				if (FDataflowInput* ConnectionInput = DataflowNode->FindInput(FName(Pin->GetName())))
+				if (FDataflowInput* const ConnectionInput = DataflowNode->FindInput(FName(Pin->GetName())))
 				{
-					DataflowGraph->ClearConnections(ConnectionInput);
-					for (UEdGraphPin* LinkedCon : Pin->LinkedTo)
+					// Add any newly added connections
+					for (const UEdGraphPin* const LinkedCon : Pin->LinkedTo)
 					{
-						if (UDataflowEdNode* LinkedNode = Cast<UDataflowEdNode>(LinkedCon->GetOwningNode()))
+						const UDataflowEdNode* const LinkedNode = Cast<UDataflowEdNode>(LinkedCon->GetOwningNode());
+						if (ensure(LinkedNode && LinkedNode->IsBound()))
 						{
-							if (ensure(LinkedNode->IsBound()))
+							if (const TSharedPtr<FDataflowNode> LinkedDataflowNode = DataflowGraph->FindBaseNode(LinkedNode->GetDataflowNodeGuid()))
 							{
-								if (TSharedPtr<FDataflowNode> LinkedDataflowNode = DataflowGraph->FindBaseNode(LinkedNode->GetDataflowNodeGuid()))
+								if (FDataflowOutput* const LinkedConOutput = LinkedDataflowNode->FindOutput(FName(LinkedCon->GetName())))
 								{
-									if (FDataflowOutput* LinkedConOutput = LinkedDataflowNode->FindOutput(FName(LinkedCon->GetName())))
+									if (ConnectionInput->GetConnectedOutputs().Find(LinkedConOutput) == INDEX_NONE)
 									{
 										const bool bTypeChanged = DataflowNode->TrySetConnectionType(ConnectionInput, LinkedConOutput->GetType());
 										DataflowGraph->Connect(LinkedConOutput, ConnectionInput);
@@ -358,39 +381,58 @@ void UDataflowEdNode::PinConnectionListChanged(UEdGraphPin* Pin)
 										{
 											UpdatePinsFromDataflowNode();
 										}
-
 									}
 								}
 							}
 						}
 					}
+					// Clear any defunct connection
+					if (FDataflowOutput* const ConnectedOutput = ConnectionInput->GetConnection())
+					{
+						if (!Pin->LinkedTo.FindByPredicate(
+							[this, ConnectedOutput, &MatchesConnection](const UEdGraphPin* const LinkedCon) -> bool
+							{
+								return MatchesConnection(LinkedCon, ConnectedOutput);
+							}))
+						{
+							DataflowGraph->Disconnect(ConnectedOutput, ConnectionInput);
+						}
+					}
 				}
 			}
 			else if (Pin->Direction == EEdGraphPinDirection::EGPD_Output)
 			{
-				if (FDataflowOutput* ConnectionOutput = DataflowNode->FindOutput(FName(Pin->GetName())))
+				// Update newly added connections from the input pins
+				for (UEdGraphPin* const LinkedPin : Pin->LinkedTo)
 				{
-					DataflowGraph->ClearConnections(ConnectionOutput);
-					for (UEdGraphPin* LinkedCon : Pin->LinkedTo)
+					LinkedPin->GetOwningNode()->PinConnectionListChanged(LinkedPin);
+				}
+
+				// Remove any remaining defunct connections
+				if (FDataflowOutput* const ConnectionOutput = DataflowNode->FindOutput(FName(Pin->GetName())))
+				{
+					const TArray<FDataflowInput*> InputsToDisconnect =
+						ConnectionOutput->GetConnections().FilterByPredicate(
+							[this, &Pin, &MatchesConnection](FDataflowInput* const ConnectedInput) -> bool
+							{
+								return ensure(ConnectedInput) && !Pin->LinkedTo.FindByPredicate(
+									[this, ConnectedInput, &MatchesConnection](const UEdGraphPin* const LinkedCon) -> bool
+									{
+										return MatchesConnection(LinkedCon, ConnectedInput);
+									});
+							});
+					UDataflow* const EdGraph = CastChecked<UDataflow>(GetGraph());
+					for (FDataflowInput* const ConnectedInput : InputsToDisconnect)
 					{
-						if (UDataflowEdNode* LinkedNode = Cast<UDataflowEdNode>(LinkedCon->GetOwningNode()))
+						const TObjectPtr<UDataflowEdNode> InputEdNode = EdGraph->FindEdNodeByDataflowNodeGuid(ConnectedInput->GetOwningNodeGuid());
+						if (ensure(InputEdNode))
 						{
-							if (ensure(LinkedNode->IsBound()))
+							UEdGraphPin* const InputPin = InputEdNode->FindPin(ConnectedInput->GetName(), EEdGraphPinDirection::EGPD_Input);
+							if (ensure(InputPin))
 							{
-								if (TSharedPtr<FDataflowNode> LinkedDataflowNode = DataflowGraph->FindBaseNode(LinkedNode->GetDataflowNodeGuid()))
-								{
-									if (FDataflowInput* LinkedConInput = LinkedDataflowNode->FindInput(FName(LinkedCon->GetName())))
-									{
-										const bool bTypeChanged = DataflowNode->TrySetConnectionType(ConnectionOutput, LinkedConInput->GetType());
-										DataflowGraph->Connect(ConnectionOutput, LinkedConInput);
-
-										if (bTypeChanged)
-										{
-											UpdatePinsFromDataflowNode();
-										}
-									}
-								}
-								
+								// To avoid double invalidations, instead of disconnecting, update connections by calling PinConnectionListChanged on the input pin
+								// This means PinConnectionListChanged might be called twice on the input, with the second time resulting in no actions
+								InputEdNode->PinConnectionListChanged(InputPin);
 							}
 						}
 					}
diff --git a/Engine/Source/Runtime/Experimental/Dataflow/Engine/Public/Dataflow/DataflowEdNode.h b/Engine/Source/Runtime/Experimental/Dataflow/Engine/Public/Dataflow/DataflowEdNode.h
--- a/Engine/Source/Runtime/Experimental/Dataflow/Engine/Public/Dataflow/DataflowEdNode.h
+++ b/Engine/Source/Runtime/Experimental/Dataflow/Engine/Public/Dataflow/DataflowEdNode.h
@@ -57,7 +57,7 @@ public:
 	DATAFLOWENGINE_API void Serialize(FArchive& Ar);
 	// End UObject interface
 
-	bool IsBound() { return DataflowGraph && DataflowNodeGuid.IsValid(); }
+	bool IsBound() const { return DataflowGraph && DataflowNodeGuid.IsValid(); }
 
 	TSharedPtr<UE::Dataflow::FGraph> GetDataflowGraph() { return DataflowGraph; }
 	TSharedPtr<const UE::Dataflow::FGraph> GetDataflowGraph() const { return DataflowGraph; }
