Dataflow - Fixed the undo of connections that didn't sync the Dataflow::FGraph despite the UEdGraph showing (or not showing) the connections.

[RN] Dataflow - Fixed the undo of connections that would leave connections still connected despite not showing it in the graph.

diff --git a/Engine/Source/Runtime/Experimental/Dataflow/Core/Private/Dataflow/DataflowGraph.cpp b/Engine/Source/Runtime/Experimental/Dataflow/Core/Private/Dataflow/DataflowGraph.cpp
--- a/Engine/Source/Runtime/Experimental/Dataflow/Core/Private/Dataflow/DataflowGraph.cpp
+++ b/Engine/Source/Runtime/Experimental/Dataflow/Core/Private/Dataflow/DataflowGraph.cpp
@@ -149,7 +149,7 @@ namespace UE::Dataflow
 
 	void FGraph::Disconnect(FDataflowOutput* OutputConnection, FDataflowInput* InputConnection)
 	{
-		UE_LOG(LogChaosDataflow, Verbose, TEXT("FGraph::Diconnect(): Disconnecting output [%s:%s] from input [%s:%s]"),
+		UE_LOG(LogChaosDataflow, Verbose, TEXT("FGraph::Disconnect(): Disconnecting output [%s:%s] from input [%s:%s]"),
 			OutputConnection->GetOwningNode() ? *OutputConnection->GetOwningNode()->GetName().ToString() : TEXT("Invalid"),
 			*OutputConnection->GetName().ToString(),
 			InputConnection->GetOwningNode() ? *InputConnection->GetOwningNode()->GetName().ToString() : TEXT("Invalid"),
diff --git a/Engine/Source/Runtime/Experimental/Dataflow/Engine/Private/Dataflow/DataflowEdNode.cpp b/Engine/Source/Runtime/Experimental/Dataflow/Engine/Private/Dataflow/DataflowEdNode.cpp
--- a/Engine/Source/Runtime/Experimental/Dataflow/Engine/Private/Dataflow/DataflowEdNode.cpp
+++ b/Engine/Source/Runtime/Experimental/Dataflow/Engine/Private/Dataflow/DataflowEdNode.cpp
@@ -700,6 +700,12 @@ void UDataflowEdNode::PostEditUndo()
 	{
 		ParentGraph->NotifyGraphChanged();
 	}
+
+	// Make sure to re-sync the Dataflow connections
+	for (UEdGraphPin* const Pin : GetAllPins())
+	{
+		PinConnectionListChanged(Pin);
+	}
 }
 
 void UDataflowEdNode::HideAllInputPins()
