ChaosFlseh: support static kinematic particles

diff --git a/Engine/Plugins/Experimental/ChaosFlesh/Source/ChaosFleshNodes/Private/Dataflow/ChaosFleshKinematicConstraintNode.cpp b/Engine/Plugins/Experimental/ChaosFlesh/Source/ChaosFleshNodes/Private/Dataflow/ChaosFleshKinematicConstraintNode.cpp
--- a/Engine/Plugins/Experimental/ChaosFlesh/Source/ChaosFleshNodes/Private/Dataflow/ChaosFleshKinematicConstraintNode.cpp
+++ b/Engine/Plugins/Experimental/ChaosFlesh/Source/ChaosFleshNodes/Private/Dataflow/ChaosFleshKinematicConstraintNode.cpp
@@ -19,22 +19,14 @@ void FKinematicInitializationDataflowNode::Evaluate(UE::Dataflow::FContext& Cont
 		{
 			if (FindInput(&VertexIndicesIn) && FindInput(&VertexIndicesIn)->GetConnection())
 			{
-				TArray<int32> BoundVerts;
-				TArray<float> BoundWeights;
-
+				GeometryCollection::Facades::FVertexBoneWeightsFacade VertexBoneWeightsFacade(InCollection);
 				for (int32 SelectionIndex : GetValue<TArray<int32>>(Context, &VertexIndicesIn))
 				{
-					if (0 <= SelectionIndex && SelectionIndex < Vertices->Num())
+					if (Vertices->IsValidIndex(SelectionIndex))
 					{
-						BoundVerts.Add(SelectionIndex);
+						VertexBoneWeightsFacade.SetVertexKinematic(SelectionIndex);
 					}
 				}
-				if (BoundVerts.Num())
-				{
-					BoundWeights.Init(1.0, BoundVerts.Num());
-					GeometryCollection::Facades::FKinematicBindingFacade Kinematics(InCollection);
-					Kinematics.AddKinematicBinding(Kinematics.SetBoneBindings(INDEX_NONE, BoundVerts, BoundWeights));
-				}
 			}
 			else if (TObjectPtr<USkeletalMesh> SkeletalMesh = GetValue<TObjectPtr<USkeletalMesh>>(Context, &SkeletalMeshIn))
 			{
diff --git a/Engine/Source/Runtime/Experimental/Chaos/Private/Chaos/Deformable/ChaosDeformableSolver.cpp b/Engine/Source/Runtime/Experimental/Chaos/Private/Chaos/Deformable/ChaosDeformableSolver.cpp
--- a/Engine/Source/Runtime/Experimental/Chaos/Private/Chaos/Deformable/ChaosDeformableSolver.cpp
+++ b/Engine/Source/Runtime/Experimental/Chaos/Private/Chaos/Deformable/ChaosDeformableSolver.cpp
@@ -561,8 +561,6 @@ namespace Chaos::Softs
 		if (Property.bEnableKinematics)
 		{
 			GeometryCollection::Facades::FVertexBoneWeightsFacade VertexBoneWeightsFacade(Rest);
-			typedef GeometryCollection::Facades::FKinematicBindingFacade FKinematics;
-			FKinematics Kinematics(Rest);
 			for (int32 VertexIdx = 0; VertexIdx < VertexBoneWeightsFacade.NumVertices(); ++VertexIdx)
 			{
 				if (VertexBoneWeightsFacade.IsKinematicVertex(VertexIdx))
@@ -572,6 +570,40 @@ namespace Chaos::Softs
 					Evolution->Particles().PAndInvM(ParticleIndex).InvM = 0.f;
 				}
 			}
+			//Supports backward compatibility for pre-5.5 nodes that uses bone-based bindings
+			//To be removed post-5.6
+			typedef GeometryCollection::Facades::FKinematicBindingFacade FKinematics;
+			FKinematics Kinematics(Rest);
+			if (Kinematics.IsValid())
+			{
+				// Add Kinematics Node
+				bool bHavePrintedLog = false;
+				for (int i = Kinematics.NumKinematicBindings() - 1; i >= 0; i--)
+				{
+					FKinematics::FBindingKey Key = Kinematics.GetKinematicBindingKey(i);
+
+					int32 BoneIndex = INDEX_NONE;
+					TArray<int32> BoundVerts;
+					TArray<float> BoundWeights;
+					Kinematics.GetBoneBindings(Key, BoneIndex, BoundVerts, BoundWeights);
+
+					for (int32 vdx : BoundVerts)
+					{
+						if (!VertexBoneWeightsFacade.IsKinematicVertex(vdx))
+						{
+							if (!bHavePrintedLog)
+							{
+								bHavePrintedLog = true;
+								UE_LOG(LogChaosDeformableSolver, Warning, TEXT("Detected deprecated kinematic initialization, reevaluate input asset"));
+							}
+							int32 ParticleIndex = Range[0] + vdx;
+							Evolution->Particles().InvM(ParticleIndex) = 0.f;
+							Evolution->Particles().PAndInvM(ParticleIndex).InvM = 0.f;
+						}
+					}
+				}
+			}
+
 		}
 	}
 
diff --git a/Engine/Source/Runtime/Experimental/Chaos/Private/GeometryCollection/Facades/CollectionVertexBoneWeightsFacade.cpp b/Engine/Source/Runtime/Experimental/Chaos/Private/GeometryCollection/Facades/CollectionVertexBoneWeightsFacade.cpp
--- a/Engine/Source/Runtime/Experimental/Chaos/Private/GeometryCollection/Facades/CollectionVertexBoneWeightsFacade.cpp
+++ b/Engine/Source/Runtime/Experimental/Chaos/Private/GeometryCollection/Facades/CollectionVertexBoneWeightsFacade.cpp
@@ -126,6 +126,19 @@ namespace GeometryCollection::Facades
 		}
 	}
 
+	bool FVertexBoneWeightsFacade::IsKinematicVertex(int32 VertexIndex) const
+	{
+		if (KinematicAttribute.IsValid())
+		{
+			return KinematicAttribute.IsValidIndex(VertexIndex) && KinematicAttribute.Get()[VertexIndex];
+		}
+		else //backward compatibility for KinematicAttribute added in 5.5
+		{
+			return BoneIndexAttribute.IsValid() && BoneIndexAttribute.IsValidIndex(VertexIndex) && BoneIndexAttribute.Get()[VertexIndex].Num()
+				&& BoneWeightAttribute.IsValid() && BoneWeightAttribute.IsValidIndex(VertexIndex) && BoneWeightAttribute.Get()[VertexIndex].Num();
+		}
+	};
+
 	//
 	//  Add Weights from Selection 
 	//
diff --git a/Engine/Source/Runtime/Experimental/Chaos/Public/GeometryCollection/Facades/CollectionVertexBoneWeightsFacade.h b/Engine/Source/Runtime/Experimental/Chaos/Public/GeometryCollection/Facades/CollectionVertexBoneWeightsFacade.h
--- a/Engine/Source/Runtime/Experimental/Chaos/Public/GeometryCollection/Facades/CollectionVertexBoneWeightsFacade.h
+++ b/Engine/Source/Runtime/Experimental/Chaos/Public/GeometryCollection/Facades/CollectionVertexBoneWeightsFacade.h
@@ -70,9 +70,7 @@ namespace GeometryCollection::Facades
 		/** Return if the vertex is kinematic
 		Pre 5.5 we did not have per-vertex kinematic attribute. 
 		This supports defining kinematics without per-vertex kinematic attribute. */
-		CHAOS_API bool IsKinematicVertex(int32 VertexIndex) const { return (!KinematicAttribute.IsValid() || (KinematicAttribute.IsValidIndex(VertexIndex) && KinematicAttribute.Get()[VertexIndex]))
-			&& BoneIndexAttribute.IsValid() && BoneIndexAttribute.IsValidIndex(VertexIndex) && BoneIndexAttribute.Get()[VertexIndex].Num()
-			&& BoneWeightAttribute.IsValid() && BoneWeightAttribute.IsValidIndex(VertexIndex) && BoneWeightAttribute.Get()[VertexIndex].Num(); };
+		CHAOS_API bool IsKinematicVertex(int32 VertexIndex) const;
 
 		/** Return number of vertices */
 		CHAOS_API int32 NumVertices() const { return VerticesAttribute.Num(); };
