Iris- Fixed FieldOfViewNetObjectPrioritizer issue causing incorrect priorities at view positions other than origin.

diff --git a/Engine/Plugins/Runtime/ReplicationSystemTestPlugin/Source/Private/Tests/ReplicationSystem/Prioritization/TestFieldOfViewNetObjectPrioritizer.cpp b/Engine/Plugins/Runtime/ReplicationSystemTestPlugin/Source/Private/Tests/ReplicationSystem/Prioritization/TestFieldOfViewNetObjectPrioritizer.cpp
--- a/Engine/Plugins/Runtime/ReplicationSystemTestPlugin/Source/Private/Tests/ReplicationSystem/Prioritization/TestFieldOfViewNetObjectPrioritizer.cpp
+++ b/Engine/Plugins/Runtime/ReplicationSystemTestPlugin/Source/Private/Tests/ReplicationSystem/Prioritization/TestFieldOfViewNetObjectPrioritizer.cpp
@@ -215,9 +215,11 @@ UE_NET_TEST_FIXTURE(FTestFieldOfViewNetObjectPrioritizer, VisualizeFieldOfViewNe
 
 	// The produced image needs to be viewed in a debugger capable of it or stored to a file and viewed in an external viewer.
 	FVisualizationParams VisualizationParams;
-	VisualizationParams.View = MakeReplicationView(FVector::ZeroVector, FVector::ForwardVector, SixtyDegreeFovInRadians);
-	VisualizationParams.PrioritizationBox.Min = FVector(-1, -1, 0)*(PrioritizerConfig->OuterSphereRadius + 100.0f);
-	VisualizationParams.PrioritizationBox.Max = FVector(PrioritizerConfig->ConeLength + 100.0f, PrioritizerConfig->OuterSphereRadius + 100.0f, 0.0f);
+	const FVector ViewPos(-15000, 10000, 5000);
+	const FVector ViewDir = FVector::ForwardVector;
+	VisualizationParams.View = MakeReplicationView(ViewPos, ViewDir, SixtyDegreeFovInRadians);
+	VisualizationParams.PrioritizationBox.Min = ViewPos + FVector(-1, -1, 0)*(PrioritizerConfig->OuterSphereRadius + 100.0f);
+	VisualizationParams.PrioritizationBox.Max = ViewPos + FVector(PrioritizerConfig->ConeLength + 100.0f, PrioritizerConfig->OuterSphereRadius + 100.0f, 0.0f);
 	// 20 units per pixel means 100/20=5 pixels per meter.
 	VisualizationParams.UnitsPerPixel = 20.0f;
 	const FNetObjectPrioritizerImage& Image = Super::Visualize(FoVPrioritizerHandle, VisualizationParams);
diff --git a/Engine/Source/Runtime/Experimental/Iris/Core/Private/Iris/ReplicationSystem/Prioritization/FieldOfViewNetObjectPrioritizer.cpp b/Engine/Source/Runtime/Experimental/Iris/Core/Private/Iris/ReplicationSystem/Prioritization/FieldOfViewNetObjectPrioritizer.cpp
--- a/Engine/Source/Runtime/Experimental/Iris/Core/Private/Iris/ReplicationSystem/Prioritization/FieldOfViewNetObjectPrioritizer.cpp
+++ b/Engine/Source/Runtime/Experimental/Iris/Core/Private/Iris/ReplicationSystem/Prioritization/FieldOfViewNetObjectPrioritizer.cpp
@@ -112,6 +112,7 @@ void UFieldOfViewNetObjectPrioritizer::PrioritizeBatch(FBatchParams& BatchParams
 		{
 			const VectorRegister ViewPos = ViewPositions[ViewIt];
 			const VectorRegister ViewDir = ViewDirs[ViewIt];
+			const VectorRegister ReverseViewDir = VectorNegate(ViewDir);
 
 			// Object directions
 			const VectorRegister ObjectDir0 = VectorSubtract(Positions[ObjIt + 0], ViewPos);
@@ -139,10 +140,12 @@ void UFieldOfViewNetObjectPrioritizer::PrioritizeBatch(FBatchParams& BatchParams
 
 			// $IRIS TODO Worth doing VectorMaskBits on ConeDistInRangeMask before doing the more expensive cone checks?
 			// Calculate the distance to the center axis at the cone distances
-			VectorRegister DistSqrToConeCenterAxis0 = VectorSubtract(Positions[ObjIt + 0], VectorMultiply(ConeDist0, ViewDir));
-			VectorRegister DistSqrToConeCenterAxis1 = VectorSubtract(Positions[ObjIt + 1], VectorMultiply(ConeDist1, ViewDir));
-			VectorRegister DistSqrToConeCenterAxis2 = VectorSubtract(Positions[ObjIt + 2], VectorMultiply(ConeDist2, ViewDir));
-			VectorRegister DistSqrToConeCenterAxis3 = VectorSubtract(Positions[ObjIt + 3], VectorMultiply(ConeDist3, ViewDir));
+			// Simplified VectorSubtract(Positions[ObjIt + N], VectorMultiplyAdd(ConeDistN, ViewDir, ViewPos)) to VectorMultiplyAdd(ConeDistN, -ViewDir, ObjectDirN)
+			VectorRegister DistSqrToConeCenterAxis0 = VectorMultiplyAdd(ConeDist0, ReverseViewDir, ObjectDir0);
+			VectorRegister DistSqrToConeCenterAxis1 = VectorMultiplyAdd(ConeDist1, ReverseViewDir, ObjectDir1);
+			VectorRegister DistSqrToConeCenterAxis2 = VectorMultiplyAdd(ConeDist2, ReverseViewDir, ObjectDir2);
+			VectorRegister DistSqrToConeCenterAxis3 = VectorMultiplyAdd(ConeDist3, ReverseViewDir, ObjectDir3);
+
 			DistSqrToConeCenterAxis0 = VectorDot4(DistSqrToConeCenterAxis0, DistSqrToConeCenterAxis0);
 			DistSqrToConeCenterAxis1 = VectorDot4(DistSqrToConeCenterAxis1, DistSqrToConeCenterAxis1);
 			DistSqrToConeCenterAxis2 = VectorDot4(DistSqrToConeCenterAxis2, DistSqrToConeCenterAxis2);
