ChaosFlesh: added attribute check during initialization for FVertexBoneWeightsFacade

diff --git a/Engine/Source/Runtime/Experimental/Chaos/Private/GeometryCollection/Facades/CollectionVertexBoneWeightsFacade.cpp b/Engine/Source/Runtime/Experimental/Chaos/Private/GeometryCollection/Facades/CollectionVertexBoneWeightsFacade.cpp
--- a/Engine/Source/Runtime/Experimental/Chaos/Private/GeometryCollection/Facades/CollectionVertexBoneWeightsFacade.cpp
+++ b/Engine/Source/Runtime/Experimental/Chaos/Private/GeometryCollection/Facades/CollectionVertexBoneWeightsFacade.cpp
@@ -5,9 +5,13 @@
 =============================================================================*/
 
 #include "GeometryCollection/Facades/CollectionVertexBoneWeightsFacade.h"
+
+#include "CoreMinimal.h"
 #include "GeometryCollection/Facades/CollectionKinematicBindingFacade.h"
 #include "GeometryCollection/GeometryCollection.h"
 
+DEFINE_LOG_CATEGORY_STATIC(LogDataflowVertexBoneWeightsFacade, Log, All);
+
 namespace GeometryCollection::Facades
 {
 
@@ -51,6 +55,22 @@ namespace GeometryCollection::Facades
 		BoneWeightAttribute.Add();
 		KinematicAttribute.AddAndFill(false);
 		ParentAttribute.Add();
+		if (!BoneIndexAttribute.IsValid())
+		{
+			UE_LOG(LogDataflowVertexBoneWeightsFacade, Warning, TEXT("FVertexBoneWeightsFacade failed to initialize because '%s' attribute from '%s' group expected type TArray<int32>."), *FVertexBoneWeightsFacade::BoneIndexAttributeName.ToString(), *FGeometryCollection::VerticesGroup.ToString());
+		}
+		if (!BoneWeightAttribute.IsValid())
+		{
+			UE_LOG(LogDataflowVertexBoneWeightsFacade, Warning, TEXT("FVertexBoneWeightsFacade failed to initialize because '%s' attribute from '%s' group expected type TArray<float>."), *FVertexBoneWeightsFacade::BoneWeightAttributeName.ToString(), *FGeometryCollection::VerticesGroup.ToString());
+		}
+		if (!KinematicAttribute.IsValid())
+		{
+			UE_LOG(LogDataflowVertexBoneWeightsFacade, Warning, TEXT("FVertexBoneWeightsFacade failed to initialize because '%s' attribute from '%s' group expected type <bool>."), *FVertexBoneWeightsFacade::KinematicAttributeName.ToString(), *FGeometryCollection::VerticesGroup.ToString());
+		}
+		if (!ParentAttribute.IsValid())
+		{
+			UE_LOG(LogDataflowVertexBoneWeightsFacade, Warning, TEXT("FVertexBoneWeightsFacade failed to initialize because '%s' attribute from '%s' group expected type <int32>."), *FTransformCollection::ParentAttribute.ToString(), *FGeometryCollection::TransformGroup.ToString());
+		}
 	}
 
 	bool FVertexBoneWeightsFacade::IsValid() const
