Dataflow : Deserialization crash when connecting nodes.

diff --git a/Engine/Source/Runtime/Experimental/Dataflow/Core/Private/Dataflow/DataflowGraph.cpp b/Engine/Source/Runtime/Experimental/Dataflow/Core/Private/Dataflow/DataflowGraph.cpp
--- a/Engine/Source/Runtime/Experimental/Dataflow/Core/Private/Dataflow/DataflowGraph.cpp
+++ b/Engine/Source/Runtime/Experimental/Dataflow/Core/Private/Dataflow/DataflowGraph.cpp
@@ -471,11 +471,15 @@ namespace UE::Dataflow
 			{
 				if (ConnectionGuidMap.Contains(Con.Input) && ConnectionGuidMap.Contains(Con.Output))
 				{
-					FDataflowOutput* Output = static_cast<FDataflowOutput*>(ConnectionGuidMap[Con.Output]);
-					FDataflowInput* Input = static_cast<FDataflowInput*>(ConnectionGuidMap[Con.Input]);
-					if (Input->GetType() == Output->GetType())
+					if (ConnectionGuidMap[Con.Output] && ConnectionGuidMap[Con.Output]->Direction == FPin::EDirection::OUTPUT &&
+						ConnectionGuidMap[Con.Input] && ConnectionGuidMap[Con.Input]->Direction == FPin::EDirection::INPUT)
 					{
-						InGraph->Connect(Output, Input);
+						FDataflowOutput* Output = static_cast<FDataflowOutput*>(ConnectionGuidMap[Con.Output]);
+						FDataflowInput* Input = static_cast<FDataflowInput*>(ConnectionGuidMap[Con.Input]);
+						if (Input->GetType() == Output->GetType())
+						{
+							InGraph->Connect(Output, Input);
+						}
 					}
 				}
 			}
diff --git a/Engine/Source/Runtime/Experimental/Dataflow/Core/Public/Dataflow/DataflowInputOutput.h b/Engine/Source/Runtime/Experimental/Dataflow/Core/Public/Dataflow/DataflowInputOutput.h
--- a/Engine/Source/Runtime/Experimental/Dataflow/Core/Public/Dataflow/DataflowInputOutput.h
+++ b/Engine/Source/Runtime/Experimental/Dataflow/Core/Public/Dataflow/DataflowInputOutput.h
@@ -45,7 +45,7 @@ struct FDataflowInput : public FDataflowConnection
 
 	friend struct FDataflowConnection;
 
-	FDataflowOutput* Connection;
+	FDataflowOutput* Connection = nullptr ;
 
 protected:
 	friend struct FDataflowOutput;
