Use SpscQueue instead of Queue
Using Queue for Singl;e Producer Single Consumer, will be deprecated and trigger some TSAN errors . Using SpscQueue is replacing it.

diff --git a/Engine/Source/Runtime/Experimental/Chaos/Private/Chaos/Collision/CollisionConstraintFlags.cpp b/Engine/Source/Runtime/Experimental/Chaos/Private/Chaos/Collision/CollisionConstraintFlags.cpp
--- a/Engine/Source/Runtime/Experimental/Chaos/Private/Chaos/Collision/CollisionConstraintFlags.cpp
+++ b/Engine/Source/Runtime/Experimental/Chaos/Private/Chaos/Collision/CollisionConstraintFlags.cpp
@@ -175,8 +175,8 @@ namespace Chaos
 
 	void FIgnoreCollisionManager::PopStorageData_Internal(int32 ExternalTimestamp)
 	{
-		FStorageData* StorageData;
-		while(StorageDataQueue.Peek(StorageData) && StorageData->ExternalTimestamp <= ExternalTimestamp)
+		FStorageData* StorageData = nullptr;
+		while((StorageData = (StorageDataQueue.IsEmpty() ? nullptr : *StorageDataQueue.Peek())) && StorageData->ExternalTimestamp <= ExternalTimestamp)
 		{
 			for (auto& Elem : StorageData->PendingActivations)
 			{
@@ -203,7 +203,7 @@ namespace Chaos
 				PendingDeactivations.Add(Item);
 			}
 
-			StorageDataQueue.Pop();
+			StorageDataQueue.Dequeue();
 			ReleaseStorageData(StorageData);
 		}
 	}
diff --git a/Engine/Source/Runtime/Experimental/Chaos/Public/Chaos/ChaosMarshallingManager.h b/Engine/Source/Runtime/Experimental/Chaos/Public/Chaos/ChaosMarshallingManager.h
--- a/Engine/Source/Runtime/Experimental/Chaos/Public/Chaos/ChaosMarshallingManager.h
+++ b/Engine/Source/Runtime/Experimental/Chaos/Public/Chaos/ChaosMarshallingManager.h
@@ -1,7 +1,7 @@
 // Copyright Epic Games, Inc. All Rights Reserved.
 #pragma once
 
-#include "Containers/Queue.h"
+#include "Containers/SpscQueue.h"
 #include "Chaos/Defines.h"
 #include "Chaos/ParticleDirtyFlags.h"
 #include "Chaos/Framework/PhysicsProxyBase.h"
@@ -411,14 +411,14 @@ private:
 	//push
 	FPushPhysicsData* ProducerData;
 	TArray<FPushPhysicsData*> ExternalQueue;	//the data pushed from external thread with a time stamp
-	TQueue<FPushPhysicsData*,EQueueMode::Spsc> PushDataPool;	//pool to grab more push data from to avoid expensive reallocs
+	TSpscQueue<FPushPhysicsData*> PushDataPool;	//pool to grab more push data from to avoid expensive reallocs
 	TArray<TUniquePtr<FPushPhysicsData>> BackingBuffer;	//all push data is cleaned up by this
 	TArray<FPushPhysicsData*> HistoryQueue_Internal;	//all push data still needed for potential rewinds. Latest data comes first
 
 	//pull
 	FPullPhysicsData* CurPullData;	//the current pull data sim is writing to
-	TQueue<FPullPhysicsData*,EQueueMode::Spsc> PullDataQueue;	//the results the simulation has written to. Consumed by external thread
-	TQueue<FPullPhysicsData*,EQueueMode::Spsc> PullDataPool;	//the pull data pool to avoid reallocs. Pushed by external thread, popped by internal
+	TSpscQueue<FPullPhysicsData*> PullDataQueue;	//the results the simulation has written to. Consumed by external thread
+	TSpscQueue<FPullPhysicsData*> PullDataPool;	//the pull data pool to avoid reallocs. Pushed by external thread, popped by internal
 	TArray<TUniquePtr<FPullPhysicsData>> BackingPullBuffer;		//all pull data is cleaned up by this
 
 	std::atomic<int32> Delay;
diff --git a/Engine/Source/Runtime/Experimental/Chaos/Public/Chaos/Collision/CollisionConstraintFlags.h b/Engine/Source/Runtime/Experimental/Chaos/Public/Chaos/Collision/CollisionConstraintFlags.h
--- a/Engine/Source/Runtime/Experimental/Chaos/Public/Chaos/Collision/CollisionConstraintFlags.h
+++ b/Engine/Source/Runtime/Experimental/Chaos/Public/Chaos/Collision/CollisionConstraintFlags.h
@@ -6,7 +6,7 @@
 #include "Chaos/GeometryParticlesfwd.h"
 #include "Chaos/ParticleHandleFwd.h"
 #include "Chaos/PBDRigidsEvolutionFwd.h"
-#include "Containers/Queue.h"
+#include "Containers/SpscQueue.h"
 
 #ifndef WITH_TODO_COLLISION_DISABLE
 #define WITH_TODO_COLLISION_DISABLE 0
@@ -202,8 +202,8 @@ namespace Chaos
 		// Producer storage data, pending changes written here until pushed into queue.
 		FStorageData* StorageDataProducer;
 
-		TQueue<FStorageData*, EQueueMode::Spsc> StorageDataQueue; // Queue of storage data being passed to physics thread
-		TQueue<FStorageData*,EQueueMode::Spsc> StorageDataFreePool;	//free pool of storage data
+		TSpscQueue<FStorageData*> StorageDataQueue; // Queue of storage data being passed to physics thread
+		TSpscQueue<FStorageData*> StorageDataFreePool;	//free pool of storage data
 		TArray<TUniquePtr<FStorageData>> StorageDataBackingBuffer;	// Holds unique ptrs for storage data allocation
 	};
 
diff --git a/Engine/Source/Runtime/Experimental/Chaos/Public/Chaos/PBDRigidsEvolution.h b/Engine/Source/Runtime/Experimental/Chaos/Public/Chaos/PBDRigidsEvolution.h
--- a/Engine/Source/Runtime/Experimental/Chaos/Public/Chaos/PBDRigidsEvolution.h
+++ b/Engine/Source/Runtime/Experimental/Chaos/Public/Chaos/PBDRigidsEvolution.h
@@ -928,11 +928,11 @@ protected:
 	FRewindData* MRewindData = nullptr;
 
 	//internal thread will push into this and external thread will consume
-	TQueue<FAccelerationStructure*,EQueueMode::Spsc> ExternalStructuresQueue;
+	TSpscQueue<FAccelerationStructure*> ExternalStructuresQueue;
 
 	//external thread will push into this when done with structure
 	//internal thread will pop from this to generate new structure
-	TQueue<FAccelerationStructure*,EQueueMode::Spsc> ExternalStructuresPool;
+	TSpscQueue<FAccelerationStructure*> ExternalStructuresPool;
 
 	//the backing buffer for all acceleration structures
 	TArray<TUniquePtr<FAccelerationStructure>> AccelerationBackingBuffer;
diff --git a/Engine/Source/Runtime/Experimental/Chaos/Public/Chaos/SimCallbackObject.h b/Engine/Source/Runtime/Experimental/Chaos/Public/Chaos/SimCallbackObject.h
--- a/Engine/Source/Runtime/Experimental/Chaos/Public/Chaos/SimCallbackObject.h
+++ b/Engine/Source/Runtime/Experimental/Chaos/Public/Chaos/SimCallbackObject.h
@@ -1,7 +1,7 @@
 // Copyright Epic Games, Inc. All Rights Reserved.
 #pragma once
 
-#include "Containers/Queue.h"
+#include "Containers/SpscQueue.h"
 #include "Chaos/Defines.h"
 #include "Chaos/SimCallbackInput.h"
 #include "Chaos/CollisionResolutionTypes.h"
@@ -611,7 +611,7 @@ private:
 	}
 
 	template <typename T>
-	T* NewDataHelper(TArray<TUniquePtr<T>>& Backing, TQueue<T*, EQueueMode::Spsc>& Queue)
+	T* NewDataHelper(TArray<TUniquePtr<T>>& Backing, TSpscQueue<T*>& Queue)
 	{
 		T* Result;
 		if (!Queue.Dequeue(Result))
@@ -645,12 +645,12 @@ private:
 		OutputQueue.Enqueue(static_cast<TOutputType*>(BaseOutput));
 	}
 
-	TQueue<TInputType*, EQueueMode::Spsc> InputPool;
+	TSpscQueue<TInputType*> InputPool;
 	TArray<TUniquePtr<TInputType>> InputBacking;
 
-	TQueue<TOutputType*, EQueueMode::Spsc> OutputPool;
+	TSpscQueue<TOutputType*> OutputPool;
 	TArray<TUniquePtr<TOutputType>> OutputBacking;
-	TQueue<TOutputType*, EQueueMode::Spsc> OutputQueue;	//holds the outputs in order
+	TSpscQueue<TOutputType*> OutputQueue;	//holds the outputs in order
 
 	TOutputType* CurrentOutput_External;	//the earliest output we can consume
 };
