expose half edge functionality on Convex

diff --git a/Engine/Source/Runtime/Experimental/Chaos/Public/Chaos/Convex.h b/Engine/Source/Runtime/Experimental/Chaos/Public/Chaos/Convex.h
--- a/Engine/Source/Runtime/Experimental/Chaos/Public/Chaos/Convex.h
+++ b/Engine/Source/Runtime/Experimental/Chaos/Public/Chaos/Convex.h
@@ -477,6 +477,30 @@ namespace Chaos
 			return INDEX_NONE;
 		}
 
+		// Get the half edge index for a given plane
+		inline int32 GetPlaneHalfEdge(int32 PlaneIndex, int32 PlaneEdgeIndex) const
+		{
+			if (StructureData.IsValid())
+			{
+				return StructureData.GetPlaneHalfEdge(PlaneIndex, PlaneEdgeIndex);
+			}
+			return INDEX_NONE;
+		}
+
+		// Get the half edges associated with an edge
+		inline void GetHalfEdges(int32 EdgeIndex, int32& OutHalfEdge0, int32& OutHalfEdge1) const
+		{
+			if (StructureData.IsValid())
+			{
+				StructureData.GetHalfEdges(EdgeIndex, OutHalfEdge0, OutHalfEdge1);
+				return;
+			}
+
+			OutHalfEdge0 = INDEX_NONE;
+			OutHalfEdge1 = INDEX_NONE;
+			return;
+		}
+
 		inline int32 NumPlanes() const
 		{
 			return Planes.Num();
@@ -1049,7 +1073,7 @@ namespace Chaos
 
 		FVec3 GetCenter() const
 		{
-			return FVec3(0);
+			return GetCenterOfMass();
 		}
 
 #if INTEL_ISPC
diff --git a/Engine/Source/Runtime/Experimental/Chaos/Public/Chaos/ConvexHalfEdgeStructureData.h b/Engine/Source/Runtime/Experimental/Chaos/Public/Chaos/ConvexHalfEdgeStructureData.h
--- a/Engine/Source/Runtime/Experimental/Chaos/Public/Chaos/ConvexHalfEdgeStructureData.h
+++ b/Engine/Source/Runtime/Experimental/Chaos/Public/Chaos/ConvexHalfEdgeStructureData.h
@@ -138,6 +138,7 @@ namespace Chaos
 		const FPlaneData& GetPlane(int32 PlaneIndex) const { return Planes[PlaneIndex]; }
 		FHalfEdgeData& GetHalfEdge(int32 HalfEdgeIndex) { return HalfEdges[HalfEdgeIndex]; }
 		const FHalfEdgeData& GetHalfEdge(int32 HalfEdgeIndex) const { return HalfEdges[HalfEdgeIndex]; }
+		int32 GetHalfEdgeIndex(int32 EdgeIndex) const { return Edges[EdgeIndex]; }
 		FVertexData& GetVertex(int32 VertexIndex) { return Vertices[VertexIndex]; }
 		const FVertexData& GetVertex(int32 VertexIndex) const { return Vertices[VertexIndex]; }
 
diff --git a/Engine/Source/Runtime/Experimental/Chaos/Public/Chaos/ConvexStructureData.h b/Engine/Source/Runtime/Experimental/Chaos/Public/Chaos/ConvexStructureData.h
--- a/Engine/Source/Runtime/Experimental/Chaos/Public/Chaos/ConvexStructureData.h
+++ b/Engine/Source/Runtime/Experimental/Chaos/Public/Chaos/ConvexStructureData.h
@@ -237,6 +237,33 @@ namespace Chaos
 				[]() { return (int32)INDEX_NONE; });
 		}
 
+		inline int32 GetPlaneHalfEdge(int32 PlaneIndex, int32 FaceVertexIndex) const
+		{
+			return ConstDataOp(
+				[PlaneIndex, FaceVertexIndex](const auto& ConcreteData)
+				{
+					return ConcreteData.GetPlaneHalfEdge(PlaneIndex, FaceVertexIndex);
+				},
+				[]() { return (int32)INDEX_NONE; });
+		}
+
+		inline void GetHalfEdges(int32 EdgeIndex, int32 &OutEdgeIndex0, int32 &OutEdgeIndex1) const
+		{
+			return ConstDataOp(
+				[EdgeIndex, &OutEdgeIndex0, &OutEdgeIndex1](const auto& ConcreteData)
+				{
+					OutEdgeIndex0 = ConcreteData.GetHalfEdgeIndex(EdgeIndex);
+					OutEdgeIndex1 = ConcreteData.GetTwinHalfEdge(OutEdgeIndex0);
+					return; 
+				},
+				[&OutEdgeIndex0, &OutEdgeIndex1]()
+				{ 
+					OutEdgeIndex0 = INDEX_NONE;
+					OutEdgeIndex1 = INDEX_NONE;
+					return; 
+				});
+		}
+
 		// Initialize the structure data from the set of vertices for each face of the convex
 		bool SetPlaneVertices(const TArray<TArray<int32>>& InPlaneVertices, int32 NumVerts, const bool bRegularDatas = false)
 		{
