Dataflow : add support for generic selection types

diff --git a/Engine/Source/Runtime/Experimental/Dataflow/Core/Public/Dataflow/DataflowAnyType.h b/Engine/Source/Runtime/Experimental/Dataflow/Core/Public/Dataflow/DataflowAnyType.h
--- a/Engine/Source/Runtime/Experimental/Dataflow/Core/Public/Dataflow/DataflowAnyType.h
+++ b/Engine/Source/Runtime/Experimental/Dataflow/Core/Public/Dataflow/DataflowAnyType.h
@@ -99,4 +99,18 @@ struct FDataflowUObjectConvertibleTypes : public FDataflowAnyType
 
 	UPROPERTY(EditAnywhere, Category = Value)
 	TObjectPtr<UObject> Value;
+};
+
+
+/** Selection types */
+USTRUCT()
+struct FDataflowSelectionTypes : public FDataflowAnyType
+{
+	using FPolicyType = FDataflowSelectionTypePolicy;
+	using FStorageType = FDataflowSelection;
+
+	GENERATED_USTRUCT_BODY()
+
+	UPROPERTY()
+	FDataflowSelection Value;
 };
\ No newline at end of file
diff --git a/Engine/Source/Runtime/Experimental/Dataflow/Core/Public/Dataflow/DataflowTypePolicy.h b/Engine/Source/Runtime/Experimental/Dataflow/Core/Public/Dataflow/DataflowTypePolicy.h
--- a/Engine/Source/Runtime/Experimental/Dataflow/Core/Public/Dataflow/DataflowTypePolicy.h
+++ b/Engine/Source/Runtime/Experimental/Dataflow/Core/Public/Dataflow/DataflowTypePolicy.h
@@ -6,6 +6,7 @@
 #include "Misc/TVariant.h"
 #include "UObject/Object.h"
 #include "Math/MathFwd.h"
+#include "Dataflow/DataflowSelection.h"
 #include <string>
 
 class UObject;
@@ -39,6 +40,7 @@ UE_DATAFLOW_POLICY_DECLARE_TYPENAME(double)
 UE_DATAFLOW_POLICY_DECLARE_TYPENAME(FName)
 UE_DATAFLOW_POLICY_DECLARE_TYPENAME(FString)
 UE_DATAFLOW_POLICY_DECLARE_TYPENAME(TObjectPtr<UObject>)
+
 UE_DATAFLOW_POLICY_DECLARE_TYPENAME(FVector2D)
 UE_DATAFLOW_POLICY_DECLARE_TYPENAME(FVector)
 UE_DATAFLOW_POLICY_DECLARE_TYPENAME(FVector4)
@@ -46,6 +48,13 @@ UE_DATAFLOW_POLICY_DECLARE_TYPENAME(FVector2f)
 UE_DATAFLOW_POLICY_DECLARE_TYPENAME(FVector3f)
 UE_DATAFLOW_POLICY_DECLARE_TYPENAME(FVector4f)
 
+UE_DATAFLOW_POLICY_DECLARE_TYPENAME(FDataflowTransformSelection)
+UE_DATAFLOW_POLICY_DECLARE_TYPENAME(FDataflowVertexSelection)
+UE_DATAFLOW_POLICY_DECLARE_TYPENAME(FDataflowFaceSelection)
+UE_DATAFLOW_POLICY_DECLARE_TYPENAME(FDataflowGeometrySelection)
+UE_DATAFLOW_POLICY_DECLARE_TYPENAME(FDataflowMaterialSelection)
+
+
 #undef UE_DATAFLOW_POLICY_DECLARE_TYPENAME
 
 struct IDataflowTypePolicy
@@ -175,6 +184,11 @@ struct FDataflowStringTypePolicy :
 {
 };
 
+struct FDataflowSelectionTypePolicy :
+	public TDataflowMultiTypePolicy<FDataflowTransformSelection, FDataflowVertexSelection, FDataflowFaceSelection, FDataflowGeometrySelection, FDataflowMaterialSelection>
+{
+};
+
 /**
 * string comvertible types
 * - FString / Fname
@@ -404,4 +418,21 @@ struct FDataflowConverter<FVector4>
 			To = From;
 		}
 	}
+};
+
+
+template <>
+struct FDataflowConverter<FDataflowSelection>
+{
+	template <typename TFromType>
+	static void From(const TFromType& From, FDataflowSelection& To)
+	{
+		To.Initialize(From);
+	}
+
+	template <typename TToType>
+	static void To(const FDataflowSelection& From, TToType& To)
+	{
+		To.Initialize(From);
+	}
 };
\ No newline at end of file
