Iris
- Fixed bug where subobjects spawned by replication with static root object would not be destroyed if root was filtered out.
- This could lead to leaking subobjects on clients.

diff --git a/Engine/Source/Runtime/Experimental/Iris/Core/Private/Iris/ReplicationSystem/ReplicationBridge.cpp b/Engine/Source/Runtime/Experimental/Iris/Core/Private/Iris/ReplicationSystem/ReplicationBridge.cpp
--- a/Engine/Source/Runtime/Experimental/Iris/Core/Private/Iris/ReplicationSystem/ReplicationBridge.cpp
+++ b/Engine/Source/Runtime/Experimental/Iris/Core/Private/Iris/ReplicationSystem/ReplicationBridge.cpp
@@ -46,6 +46,13 @@ static FAutoConsoleVariableRef CVarAllowDestroyToCancelFlushAndTearOff(
 	TEXT("When true issuing a EndReplication on an object that is already Tear-off or pending endreplication will cancel destroy/flush and destroy the replicated objects.")
 );
 
+static bool bAlwaysDestroyDynamicSubObjectInstancesOnDetachFromStaticRoot = true;
+static FAutoConsoleVariableRef CVarAlwaysDestroyDynamicSubObjectInstancesOnDetachFromStaticRoot(
+	TEXT("net.Iris.AlwaysDestroyDynamicSubObjectInstancesOnDetachFromStaticRoot"),
+	bAlwaysDestroyDynamicSubObjectInstancesOnDetachFromStaticRoot,
+	TEXT("When true, We will always destroy instance for dynamic subobjects during EndReplication of a static rootobject.")
+);
+
 /**
  * ReplicationBridge Implementation
  */
@@ -150,6 +157,23 @@ void UReplicationBridge::DetachSubObjectInstancesFromRemote(FNetRefHandle OwnerH
 			{
 				EnumRemoveFlags(SubObjectDestroyFlags, EReplicationBridgeDestroyInstanceFlags::AllowDestroyInstanceFromRemote);
 			}
+
+			if (DestroyReason == EReplicationBridgeDestroyInstanceReason::DoNotDestroy && SubObjectData.bAllowDestroyInstanceFromRemote)
+			{
+				// When ending replication of static objects without destroying them, we should always destroy instances spawned from replication
+				// as they will be recreated when the static object is scoped again.
+				if (bAlwaysDestroyDynamicSubObjectInstancesOnDetachFromStaticRoot)
+				{
+					EnumAddFlags(SubObjectDestroyFlags, EReplicationBridgeDestroyInstanceFlags::AllowDestroyInstanceFromRemote);
+					CallDetachInstanceFromRemote(SubObjectHandle, EReplicationBridgeDestroyInstanceReason::Destroy, SubObjectDestroyFlags);
+					continue;
+				}
+				else
+				{
+					UE_LOG(LogIrisBridge, Warning, TEXT("Detaching Subobject %s with reason DoNotDestroy even though it is dynamic"), *SubObjectHandle.ToString());
+				}
+			}
+
 			CallDetachInstanceFromRemote(SubObjectHandle, DestroyReason, SubObjectDestroyFlags);
 		}
 	}
