Iris
- Changed FReplicationFiltering::PerObjectInfoIndexType to uint32

diff --git a/Engine/Source/Runtime/Experimental/Iris/Core/Private/Iris/ReplicationSystem/Filtering/ReplicationFiltering.h b/Engine/Source/Runtime/Experimental/Iris/Core/Private/Iris/ReplicationSystem/Filtering/ReplicationFiltering.h
--- a/Engine/Source/Runtime/Experimental/Iris/Core/Private/Iris/ReplicationSystem/Filtering/ReplicationFiltering.h
+++ b/Engine/Source/Runtime/Experimental/Iris/Core/Private/Iris/ReplicationSystem/Filtering/ReplicationFiltering.h
@@ -180,7 +180,7 @@ private:
 	};
 
 	static constexpr uint32 UsedPerObjectInfoStorageGrowSize = 32; // 256 bytes, 1024 indices
-	typedef uint16 PerObjectInfoIndexType;
+	typedef uint32 PerObjectInfoIndexType;
 
 	struct FPerGroupInfo
 	{
