Geometry collection : fix comment for ClusterGroupindex

diff --git a/Engine/Source/Runtime/Experimental/GeometryCollectionEngine/Public/GeometryCollection/GeometryCollectionComponent.h b/Engine/Source/Runtime/Experimental/GeometryCollectionEngine/Public/GeometryCollection/GeometryCollectionComponent.h
--- a/Engine/Source/Runtime/Experimental/GeometryCollectionEngine/Public/GeometryCollection/GeometryCollectionComponent.h
+++ b/Engine/Source/Runtime/Experimental/GeometryCollectionEngine/Public/GeometryCollection/GeometryCollectionComponent.h
@@ -991,7 +991,7 @@ public:
 	UPROPERTY(EditAnywhere, BlueprintReadWrite, Category = "ChaosPhysics|Clustering")
 	bool EnableClustering;
 
-	/** Maximum level for cluster breaks. */
+	/** Cluster group index - Geometry Collections using the same index will simulate as one single cluster. */
 	UPROPERTY(EditAnywhere, BlueprintReadWrite, Category = "ChaosPhysics|Clustering")
 	int32 ClusterGroupIndex;
 
diff --git a/Engine/Source/Runtime/Experimental/GeometryCollectionEngine/Public/GeometryCollection/GeometryCollectionObject.h b/Engine/Source/Runtime/Experimental/GeometryCollectionEngine/Public/GeometryCollection/GeometryCollectionObject.h
--- a/Engine/Source/Runtime/Experimental/GeometryCollectionEngine/Public/GeometryCollection/GeometryCollectionObject.h
+++ b/Engine/Source/Runtime/Experimental/GeometryCollectionEngine/Public/GeometryCollection/GeometryCollectionObject.h
@@ -519,7 +519,7 @@ public:
 	UPROPERTY(EditAnywhere, Category = "Clustering")
 	bool EnableClustering;
 
-	/** Maximum level for cluster breaks. */
+	/** Cluster group index - Geometry Collections using the same index will simulate as one single cluster. */
 	UPROPERTY(EditAnywhere, Category = "Clustering")
 	int32 ClusterGroupIndex;
 
