[HWRT] Support ray tracing in FNaniteGeometryCollectionSceneProxy
- moved relevant logic from FGeometryCollectionSceneProxy to FGeometryCollectionSceneProxyBase

diff --git a/Engine/Source/Runtime/Experimental/GeometryCollectionEngine/Private/GeometryCollection/GeometryCollectionComponent.cpp b/Engine/Source/Runtime/Experimental/GeometryCollectionEngine/Private/GeometryCollection/GeometryCollectionComponent.cpp
--- a/Engine/Source/Runtime/Experimental/GeometryCollectionEngine/Private/GeometryCollection/GeometryCollectionComponent.cpp
+++ b/Engine/Source/Runtime/Experimental/GeometryCollectionEngine/Private/GeometryCollection/GeometryCollectionComponent.cpp
@@ -4808,9 +4808,9 @@ void UGeometryCollectionComponent::SendDynamicDataToSceneProxy()
 			const FMatrix RenderMatrix{ GetRenderMatrix() };
 			FNaniteGeometryCollectionSceneProxy* GeometryCollectionSceneProxy = static_cast<FNaniteGeometryCollectionSceneProxy*>(SceneProxy);
 			ENQUEUE_RENDER_COMMAND(SendRenderDynamicData)(
-				[GeometryCollectionSceneProxy, DynamicData, RenderMatrix](FRHICommandListBase&)
+				[GeometryCollectionSceneProxy, DynamicData, RenderMatrix](FRHICommandListBase& RHICmdList)
 				{
-					GeometryCollectionSceneProxy->SetDynamicData_RenderThread(DynamicData, RenderMatrix);
+					GeometryCollectionSceneProxy->SetDynamicData_RenderThread(RHICmdList, DynamicData, RenderMatrix);
 				}
 			);
 		}
diff --git a/Engine/Source/Runtime/Experimental/GeometryCollectionEngine/Private/GeometryCollection/GeometryCollectionSceneProxy.cpp b/Engine/Source/Runtime/Experimental/GeometryCollectionEngine/Private/GeometryCollection/GeometryCollectionSceneProxy.cpp
--- a/Engine/Source/Runtime/Experimental/GeometryCollectionEngine/Private/GeometryCollection/GeometryCollectionSceneProxy.cpp
+++ b/Engine/Source/Runtime/Experimental/GeometryCollectionEngine/Private/GeometryCollection/GeometryCollectionSceneProxy.cpp
@@ -95,6 +95,23 @@
 
 FGeometryCollectionDynamicDataPool GDynamicDataPool;
 
+static void UpdateLooseParameter(
+	FGeometryCollectionVertexFactory& VertexFactory,
+	FRHIShaderResourceView* BoneTransformSRV,
+	FRHIShaderResourceView* BonePrevTransformSRV,
+	FRHIShaderResourceView* BoneMapSRV)
+{
+	FGCBoneLooseParameters LooseParameters;
 
+	LooseParameters.VertexFetch_BoneTransformBuffer = BoneTransformSRV;
+	LooseParameters.VertexFetch_BonePrevTransformBuffer = BonePrevTransformSRV;
+	LooseParameters.VertexFetch_BoneMapBuffer = BoneMapSRV;
+
+	EUniformBufferUsage UniformBufferUsage = VertexFactory.EnableLooseParameter ? UniformBuffer_SingleFrame : UniformBuffer_MultiFrame;
+
+	VertexFactory.LooseParameterUniformBuffer = FGCBoneLooseParametersRef::CreateUniformBufferImmediate(LooseParameters, UniformBufferUsage);
+}
+
+
 class FGeometryCollectionMeshCollectorResources : public FOneFrameResource
 {
@@ -115,95 +132,43 @@
 };
 
 
-FGeometryCollectionSceneProxy::FGeometryCollectionSceneProxy(UGeometryCollectionComponent* Component)
-	: FPrimitiveSceneProxy(Component)
-	, MaterialRelevance(Component->GetMaterialRelevance(GetScene().GetFeatureLevel()))
+FGeometryCollectionSceneProxyBase::FGeometryCollectionSceneProxyBase(UGeometryCollectionComponent* Component, bool bInIsNanite)
+	: bIsNanite(bInIsNanite)
+	, FeatureLevel(Component->GetScene()->GetFeatureLevel())
 	, MeshResource(Component->GetRestCollection()->RenderData->MeshResource)
-	, MeshDescription(Component->GetRestCollection()->RenderData->MeshDescription)
-	, VertexFactory(GetScene().GetFeatureLevel())
-	, bSupportsManualVertexFetch(VertexFactory.SupportsManualVertexFetch(GetScene().GetFeatureLevel()))
+	, MaterialRelevance(Component->GetMaterialRelevance(FeatureLevel))
+	, VertexFactory(FeatureLevel)
+	, bSupportsManualVertexFetch(VertexFactory.SupportsManualVertexFetch(FeatureLevel))
 	, bSupportsTripleBufferVertexUpload(GRHISupportsMapWriteNoOverwrite)
-#if WITH_EDITOR
-	, bShowBoneColors(Component->GetShowBoneColors())
-	, bSuppressSelectionMaterial(Component->GetSuppressSelectionMaterial())
-	, VertexFactoryDebugColor(GetScene().GetFeatureLevel())
+#if RHI_RAYTRACING
+	, bHasRayTracingRepresentation(IsRayTracingEnabled() && Component->GetRestCollection()->RenderData->bHasMeshData && Component->GetRestCollection()->RenderData->MeshDescription.NumVertices != 0)
 #endif
 {
-	if (Component->GetRestCollection())
+	if (!bIsNanite || bHasRayTracingRepresentation)
 	{
-		GeometryCollection = Component->GetRestCollection()->GetGeometryCollection();
-	}
+		MeshDescription = Component->GetRestCollection()->RenderData->MeshDescription; // TODO: use const-reference instead?
 
-	EnableGPUSceneSupportFlags();
+		const TSharedPtr<FGeometryCollection, ESPMode::ThreadSafe> Collection = Component->GetRestCollection()->GetGeometryCollection();
 
-	Materials.Empty();
-	const int32 NumMaterials = Component->GetNumMaterials();
-	for (int MaterialIndex = 0; MaterialIndex < NumMaterials; ++MaterialIndex)
-	{
-		Materials.Push(Component->GetMaterial(MaterialIndex));
+		NumTransforms = Collection ? Collection->NumElements(FTransformCollection::TransformGroup) : 0;
 
-		if (Materials[MaterialIndex] == nullptr || !Materials[MaterialIndex]->CheckMaterialUsage_Concurrent(MATUSAGE_GeometryCollections))
+		Materials.Empty();
+		const int32 NumMaterials = Component->GetNumMaterials();
+		for (int MaterialIndex = 0; MaterialIndex < NumMaterials; ++MaterialIndex)
 		{
-			Materials[MaterialIndex] = UMaterial::GetDefaultMaterial(MD_Surface);
+			Materials.Push(Component->GetMaterial(MaterialIndex));
+
+			if (Materials[MaterialIndex] == nullptr || !Materials[MaterialIndex]->CheckMaterialUsage_Concurrent(MATUSAGE_GeometryCollections))
+			{
+				Materials[MaterialIndex] = UMaterial::GetDefaultMaterial(MD_Surface);
+			}
 		}
 	}
 
-	NumTransforms = GeometryCollection? GeometryCollection->NumElements(FTransformCollection::TransformGroup) : 0;
-
-#if GEOMETRYCOLLECTION_EDITOR_SELECTION
-	// Render by SubSection if we are in the rigid body picker.
-	bUsesSubSections = Component->GetIsTransformSelectionMode() && MeshDescription.SubSections.Num();
-	// Enable bone hit selection proxies if we are in the rigid body picker or in the fracture modes.
-	bEnableBoneSelection = Component->GetEnableBoneSelection();
-
-	if (bEnableBoneSelection || bUsesSubSections)
-	{
-		for (int32 TransformIndex = 0; TransformIndex < NumTransforms; ++TransformIndex)
-		{
-			HGeometryCollection* HitProxy = new HGeometryCollection(Component, TransformIndex);
-			HitProxies.Add(HitProxy);
-		}
-	}
-#endif
-
-#if WITH_EDITOR
-	if (bShowBoneColors || bEnableBoneSelection)
-	{
-		Component->GetBoneColors(BoneColors);
-		ColorVertexBuffer.InitFromColorArray(BoneColors);
-
-		if (Component->GetRestCollection())
-		{
-			BoneSelectedMaterial = Component->GetRestCollection()->GetBoneSelectedMaterial();
-		}
-		if (BoneSelectedMaterial && !BoneSelectedMaterial->CheckMaterialUsage_Concurrent(MATUSAGE_GeometryCollections))
-		{
-			// If we have an invalid BoneSelectedMaterial, switch it back to null to skip its usage in GetDynamicMeshElements below
-			BoneSelectedMaterial = nullptr;
-		}
-
-		// Make sure the vertex color material has the usage flag for rendering geometry collections
-		if (GEngine->VertexColorMaterial)
-		{
-			GEngine->VertexColorMaterial->CheckMaterialUsage_Concurrent(MATUSAGE_GeometryCollections);
-		}
-	}
-
-#endif
-
-	// #todo(dmp): This flag means that when motion blur is turned on, it will always render geometry collections into the
-	// velocity buffer.  Note that the way around this is to loop through the global matrices and test whether they have
-	// changed from the prev to curr frame, but this is expensive.  We should revisit this if the draw calls for velocity
-	// rendering become a problem. One solution could be to use internal solver sleeping state to drive motion blur.
-	bAlwaysHasVelocity = true;
-
 	DynamicData = Component->InitDynamicData(true);
-
-	SetWireframeColor(Component->GetWireframeColorForSceneProxy());
-	CollisionResponse = Component->GetCollisionResponseToChannels();
 }
 
-FGeometryCollectionSceneProxy::~FGeometryCollectionSceneProxy()
+FGeometryCollectionSceneProxyBase::~FGeometryCollectionSceneProxyBase()
 {
 	if (DynamicData != nullptr)
 	{
@@ -212,33 +177,92 @@
 	}
 }
 
-SIZE_T FGeometryCollectionSceneProxy::GetTypeHash() const
+void FGeometryCollectionSceneProxyBase::CreateRenderThreadResources(FRHICommandListBase& RHICmdList)
 {
-	static size_t UniquePointer;
-	return reinterpret_cast<size_t>(&UniquePointer);
+	if (!bIsNanite || bHasRayTracingRepresentation)
+	{
+		if (bSupportsManualVertexFetch)
+		{
+			// Initialize transform buffers and upload rest transforms.
+			TransformBuffers.AddDefaulted(1);
+
+			TransformBuffers[0].NumTransforms = NumTransforms;
+			TransformBuffers[0].InitResource(RHICmdList);
+
+			const bool bLocalGeometryCollectionTripleBufferUploads = (GGeometryCollectionTripleBufferUploads != 0) && bSupportsTripleBufferVertexUpload;
+			const EResourceLockMode LockMode = bLocalGeometryCollectionTripleBufferUploads ? RLM_WriteOnly_NoOverwrite : RLM_WriteOnly;
+
+			FGeometryCollectionTransformBuffer& TransformBuffer = GetCurrentTransformBuffer();
+			TransformBuffer.UpdateDynamicData(RHICmdList, DynamicData->Transforms, LockMode);
+		}
+		else
+		{
+			// Initialize CPU skinning buffer with rest transforms.
+			SkinnedPositionVertexBuffer.Init(MeshResource.PositionVertexBuffer.GetNumVertices(), false);
+			SkinnedPositionVertexBuffer.InitResource(RHICmdList);
+			UpdateSkinnedPositions(RHICmdList, DynamicData->Transforms);
+		}
+
+		SetupVertexFactory(RHICmdList, VertexFactory);
+	}
+
+#if RHI_RAYTRACING
+	if (bHasRayTracingRepresentation)
+	{
+		FRayTracingGeometryInitializer Initializer;
+		Initializer.DebugName = FName(TEXT("GeometryCollection"));
+		Initializer.GeometryType = RTGT_Triangles;
+		Initializer.bFastBuild = true;
+		Initializer.bAllowUpdate = false;
+		Initializer.TotalPrimitiveCount = 0;
+
+		RayTracingGeometry.SetInitializer(Initializer);
+		RayTracingGeometry.InitResource(RHICmdList);
+
+		bRayTracingGeometryDirty = true;
+	}
+#endif
+
+	bRenderResourcesCreated = true;
 }
 
-static void UpdateLooseParameter(
-	FGeometryCollectionVertexFactory& VertexFactory,
-	FRHIShaderResourceView* BoneTransformSRV,
-	FRHIShaderResourceView* BonePrevTransformSRV,
-	FRHIShaderResourceView* BoneMapSRV)
+void FGeometryCollectionSceneProxyBase::DestroyRenderThreadResources()
 {
-	FGCBoneLooseParameters LooseParameters;
+	bRenderResourcesCreated = false;
 
-	LooseParameters.VertexFetch_BoneTransformBuffer = BoneTransformSRV;
-	LooseParameters.VertexFetch_BonePrevTransformBuffer = BonePrevTransformSRV;
-	LooseParameters.VertexFetch_BoneMapBuffer = BoneMapSRV;
+	if (!bIsNanite || bHasRayTracingRepresentation)
+	{
+		if (bSupportsManualVertexFetch)
+		{
+			for (int32 i = 0; i < TransformBuffers.Num(); i++)
+			{
+				TransformBuffers[i].ReleaseResource();
+			}
+			TransformBuffers.Reset();
+		}
+		else
+		{
+			SkinnedPositionVertexBuffer.ReleaseResource();
+		}
+	}
 
-	EUniformBufferUsage UniformBufferUsage = VertexFactory.EnableLooseParameter ? UniformBuffer_SingleFrame : UniformBuffer_MultiFrame;
+	VertexFactory.ReleaseResource();
 
-	VertexFactory.LooseParameterUniformBuffer = FGCBoneLooseParametersRef::CreateUniformBufferImmediate(LooseParameters, UniformBufferUsage);
+#if RHI_RAYTRACING
+	if (bHasRayTracingRepresentation)
+	{
+		RayTracingGeometry.ReleaseResource();
+		RayTracingDynamicVertexBuffer.Release();
+	}
+#endif
 }
 
-void FGeometryCollectionSceneProxy::SetupVertexFactory(FRHICommandListBase& RHICmdList, FGeometryCollectionVertexFactory& GeometryCollectionVertexFactory, FColorVertexBuffer* ColorOverride) const
+void FGeometryCollectionSceneProxyBase::SetupVertexFactory(FRHICommandListBase& RHICmdList, FGeometryCollectionVertexFactory& GeometryCollectionVertexFactory, FColorVertexBuffer* ColorOverride) const
 {
+	checkf(GeometryCollectionVertexFactory.SupportsManualVertexFetch(FeatureLevel) == bSupportsManualVertexFetch, TEXT("Setting up vertex factory for manual vertex fetch but provided type doesn't support it."));
+
 	FGeometryCollectionVertexFactory::FDataType Data;
-	
+
 	FPositionVertexBuffer const& PositionVB = bSupportsManualVertexFetch ? MeshResource.PositionVertexBuffer : SkinnedPositionVertexBuffer;
 	PositionVB.BindPositionVertexBuffer(&GeometryCollectionVertexFactory, Data);
 
@@ -252,8 +276,8 @@
 	if (bSupportsManualVertexFetch)
 	{
 		Data.BoneMapSRV = MeshResource.BoneMapVertexBuffer.GetSRV();
-		Data.BoneTransformSRV = TransformBuffers[CurrentTransformBufferIndex].VertexBufferSRV;
-		Data.BonePrevTransformSRV = Data.BoneTransformSRV; // setup : both prev and crrent are the same 
+		Data.BoneTransformSRV = GetCurrentTransformBuffer().VertexBufferSRV;
+		Data.BonePrevTransformSRV = GetCurrentPrevTransformBuffer().VertexBufferSRV;
 	}
 	else
 	{
@@ -275,31 +299,379 @@
 	}
 }
 
+/** Called on render thread to setup dynamic geometry for rendering */
+void FGeometryCollectionSceneProxyBase::SetDynamicData_RenderThread(FRHICommandListBase& RHICmdList, FGeometryCollectionDynamicData* NewDynamicData)
+{
+	if (NewDynamicData != DynamicData)
+	{
+		if (DynamicData)
+		{
+			GDynamicDataPool.Release(DynamicData);
+			DynamicData = nullptr;
+		}
+		DynamicData = NewDynamicData;
+	}
+
+	if (MeshDescription.NumVertices == 0 || !DynamicData || !bRenderResourcesCreated)
+	{
+		return;
+	}
+
+	if (!bIsNanite || bHasRayTracingRepresentation)
+	{
+		if (bSupportsManualVertexFetch)
+		{
+			const bool bLocalGeometryCollectionTripleBufferUploads = (GGeometryCollectionTripleBufferUploads != 0) && bSupportsTripleBufferVertexUpload;
+
+			if (bLocalGeometryCollectionTripleBufferUploads && TransformBuffers.Num() == 1)
+			{
+				TransformBuffers.AddDefaulted(2);
+				check(TransformBuffers.Num() == 3);
+
+				for (int32 i = 1; i < TransformBuffers.Num(); i++)
+				{
+					TransformBuffers[i].NumTransforms = NumTransforms;
+					TransformBuffers[i].InitResource(RHICmdList);
+				}
+			}
+
+			// Copy the transform data over to the vertex buffer	
+			{
+				const EResourceLockMode LockMode = bLocalGeometryCollectionTripleBufferUploads ? RLM_WriteOnly_NoOverwrite : RLM_WriteOnly;
+
+				CycleTransformBuffers(bLocalGeometryCollectionTripleBufferUploads);
+
+				FGeometryCollectionTransformBuffer& TransformBuffer = GetCurrentTransformBuffer();
+				const FGeometryCollectionTransformBuffer& PrevTransformBuffer = GetCurrentPrevTransformBuffer();
+
+				VertexFactory.SetBoneTransformSRV(TransformBuffer.VertexBufferSRV);
+				VertexFactory.SetBonePrevTransformSRV(PrevTransformBuffer.VertexBufferSRV);
+
+				TransformBuffer.UpdateDynamicData(RHICmdList, DynamicData->Transforms, LockMode);
+
+				UpdateLooseParameter(VertexFactory, TransformBuffer.VertexBufferSRV, PrevTransformBuffer.VertexBufferSRV, MeshResource.BoneMapVertexBuffer.GetSRV());
+			}
+		}
+		else
+		{
+			UpdateSkinnedPositions(RHICmdList, DynamicData->Transforms);
+		}
+	}
+
+#if RHI_RAYTRACING
+	if (bHasRayTracingRepresentation)
+	{
+		bRayTracingGeometryDirty = true;
+	}
+#endif
+}
+
+void FGeometryCollectionSceneProxyBase::UpdateSkinnedPositions(FRHICommandListBase& RHICmdList, TArray<FMatrix44f> const& Transforms)
+{
+	const int32 VertexStride = SkinnedPositionVertexBuffer.GetStride();
+	const int32 VertexCount = SkinnedPositionVertexBuffer.GetNumVertices();
+	check(VertexCount == MeshDescription.NumVertices)
+
+		void* VertexBufferData = RHICmdList.LockBuffer(SkinnedPositionVertexBuffer.VertexBufferRHI, 0, VertexCount * VertexStride, RLM_WriteOnly);
+	check(VertexBufferData != nullptr);
+
+	FPositionVertexBuffer const& SourcePositionVertexBuffer = MeshResource.PositionVertexBuffer;
+	FBoneMapVertexBuffer const& SourceBoneMapVertexBuffer = MeshResource.BoneMapVertexBuffer;
+
+	bool bParallelGeometryCollection = true;
+	int32 ParallelGeometryCollectionBatchSize = CVarParallelGeometryCollectionBatchSize.GetValueOnRenderThread();
+
+	int32 NumBatches = (VertexCount / ParallelGeometryCollectionBatchSize);
+
+	if (VertexCount != ParallelGeometryCollectionBatchSize)
+	{
+		NumBatches++;
+	}
+
+	// Batch too small, don't bother with parallel
+	if (ParallelGeometryCollectionBatchSize > VertexCount)
+	{
+		bParallelGeometryCollection = false;
+		ParallelGeometryCollectionBatchSize = VertexCount;
+	}
+
+	auto GeometryCollectionBatch([&](int32 BatchNum)
+		{
+			uint32 IndexOffset = ParallelGeometryCollectionBatchSize * BatchNum;
+			uint32 ThisBatchSize = ParallelGeometryCollectionBatchSize;
+
+			// Check for final batch
+			if (IndexOffset + ParallelGeometryCollectionBatchSize > MeshDescription.NumVertices)
+			{
+				ThisBatchSize = VertexCount - IndexOffset;
+			}
+
+			if (ThisBatchSize > 0)
+			{
+				const FMatrix44f* RESTRICT BoneTransformsPtr = Transforms.GetData();
+
+				if (bGeometryCollection_SetDynamicData_ISPC_Enabled)
+				{
+#if INTEL_ISPC
+					uint8* VertexBufferOffset = (uint8*)VertexBufferData + (IndexOffset * VertexStride);
+					ispc::SetDynamicData_RenderThread(
+						(ispc::FVector3f*)VertexBufferOffset,
+						ThisBatchSize,
+						VertexStride,
+						&SourceBoneMapVertexBuffer.BoneIndex(IndexOffset),
+						(ispc::FMatrix44f*)BoneTransformsPtr,
+						(ispc::FVector3f*)&SourcePositionVertexBuffer.VertexPosition(IndexOffset));
+#endif
+				}
+				else
+				{
+					for (uint32 i = IndexOffset; i < IndexOffset + ThisBatchSize; i++)
+					{
+						FVector3f Transformed = BoneTransformsPtr[SourceBoneMapVertexBuffer.BoneIndex(i)].TransformPosition(SourcePositionVertexBuffer.VertexPosition(i));
+						FMemory::Memcpy((uint8*)VertexBufferData + (i * VertexStride), &Transformed, sizeof(FVector3f));
+					}
+				}
+			}
+		});
+
+	ParallelFor(NumBatches, GeometryCollectionBatch, !bParallelGeometryCollection);
+
+	RHICmdList.UnlockBuffer(SkinnedPositionVertexBuffer.VertexBufferRHI);
+}
+
+#if RHI_RAYTRACING
+void FGeometryCollectionSceneProxyBase::UpdatingRayTracingGeometry_RenderingThread(TArray<FGeometryCollectionMeshElement> const& InSectionArray)
+{
+	FRHICommandList& RHICmdList = FRHICommandListImmediate::Get();
+
+	if (bRayTracingGeometryDirty)
+	{
+		RayTracingGeometry.Initializer.Segments.Empty();
+		RayTracingGeometry.Initializer.TotalPrimitiveCount = 0;
+
+		for (int SectionIndex = 0; SectionIndex < InSectionArray.Num(); ++SectionIndex)
+		{
+			const FGeometryCollectionMeshElement& Section = InSectionArray[SectionIndex];
+			FRayTracingGeometrySegment Segment;
+			Segment.FirstPrimitive = Section.TriangleStart;
+			Segment.VertexBuffer = MeshResource.PositionVertexBuffer.VertexBufferRHI;
+			Segment.NumPrimitives = Section.TriangleCount;
+			Segment.MaxVertices = Section.VertexEnd;
+			RayTracingGeometry.Initializer.Segments.Add(Segment);
+			RayTracingGeometry.Initializer.TotalPrimitiveCount += Section.TriangleCount;
+		}
+
+		if (RayTracingGeometry.Initializer.TotalPrimitiveCount > 0)
+		{
+			RayTracingGeometry.Initializer.IndexBuffer = MeshResource.IndexBuffer.IndexBufferRHI;
+			// Create the ray tracing geometry but delay the acceleration structure build.
+			RayTracingGeometry.CreateRayTracingGeometry(RHICmdList, ERTAccelerationStructureBuildPriority::Skip);
+		}
+
+		bRayTracingGeometryDirty = false;
+	}
+}
+
+void FGeometryCollectionSceneProxyBase::GetDynamicRayTracingInstances(FRayTracingInstanceCollector& Collector, const FMatrix& LocalToWorld, FRHIUniformBuffer* UniformBuffer)
+{
+	checkf(bHasRayTracingRepresentation, TEXT("Shouldn't try to get ray tracing instances from proxy that doesn't have a ray tracing representation."));
+
+	if (GRayTracingGeometryCollectionProxyMeshes == 0)
+	{
+		return;
+	}
+
+	QUICK_SCOPE_CYCLE_COUNTER(STAT_GeometryCollectionSceneProxyBase_GetDynamicRayTracingInstances);
+
+	const uint32 LODIndex = 0;
+	const bool bWireframe = false;
+
+	//Loose parameter needs to be updated every frame
+	FGeometryCollectionMeshCollectorResources* CollectorResources;
+	CollectorResources = &Collector.AllocateOneFrameResource<FGeometryCollectionMeshCollectorResources>(FeatureLevel);
+	FGeometryCollectionVertexFactory& GeometryCollectionVertexFactory = CollectorResources->GetVertexFactory();
+
+	// Render dynamic objects
+	if (!GeometryCollectionVertexFactory.GetType()->SupportsRayTracingDynamicGeometry())
+	{
+		return;
+	}
+
+	SetupVertexFactory(Collector.GetRHICommandList(), GeometryCollectionVertexFactory);
+
+	TArray<FGeometryCollectionMeshElement> const& SectionArray = MeshDescription.Sections;
+
+	UpdatingRayTracingGeometry_RenderingThread(SectionArray);
+
+	if (SectionArray.Num() && RayTracingGeometry.IsValid())
+	{
+		FRayTracingInstance RayTracingInstance;
+		RayTracingInstance.Geometry = &RayTracingGeometry;
+		RayTracingInstance.InstanceTransforms.Emplace(LocalToWorld);
+
+		// Grab the material proxies we'll be using for each section
+		TArray<FMaterialRenderProxy*, TInlineAllocator<32>> MaterialProxies;
+
+		for (int32 SectionIndex = 0; SectionIndex < SectionArray.Num(); ++SectionIndex)
+		{
+			const FGeometryCollectionMeshElement& Section = SectionArray[SectionIndex];
+
+			//TODO: Add BoneColor support in Path/Ray tracing?
+			FMaterialRenderProxy* MaterialProxy = Materials[Section.MaterialIndex]->GetRenderProxy();
+
+			if (MaterialProxy == nullptr)
+			{
+				MaterialProxy = UMaterial::GetDefaultMaterial(MD_Surface)->GetRenderProxy();
+			}
+
+			MaterialProxies.Add(MaterialProxy);
+		}
+
+		uint32 MaxVertexIndex = 0;
+		for (int32 SectionIndex = 0; SectionIndex < SectionArray.Num(); ++SectionIndex)
+		{
+			const FGeometryCollectionMeshElement& Section = SectionArray[SectionIndex];
+
+			// Draw the mesh
+			FMeshBatch& Mesh = RayTracingInstance.Materials.AddDefaulted_GetRef();
+			Mesh.bWireframe = bWireframe;
+			Mesh.SegmentIndex = SectionIndex;
+			Mesh.VertexFactory = &GeometryCollectionVertexFactory;
+			Mesh.MaterialRenderProxy = MaterialProxies[SectionIndex];
+			Mesh.LODIndex = LODIndex;
+			Mesh.bDisableBackfaceCulling = true;
+			Mesh.Type = PT_TriangleList;
+			Mesh.DepthPriorityGroup = SDPG_World;
+			Mesh.bCanApplyViewModeOverrides = true;
+
+			FMeshBatchElement& BatchElement = Mesh.Elements[0];
+			BatchElement.IndexBuffer = &MeshResource.IndexBuffer;
+			BatchElement.PrimitiveUniformBuffer = UniformBuffer;
+			BatchElement.FirstIndex = Section.TriangleStart * 3;
+			BatchElement.NumPrimitives = Section.TriangleCount;
+			BatchElement.MinVertexIndex = Section.VertexStart;
+			BatchElement.MaxVertexIndex = Section.VertexEnd;
+			BatchElement.NumInstances = 1;
+
+			MaxVertexIndex = std::max(Section.VertexEnd, MaxVertexIndex);
+
+			//#TODO: bone color, bone selection and render bound?
+		}
+
+		FRWBuffer* VertexBuffer = RayTracingDynamicVertexBuffer.NumBytes > 0 ? &RayTracingDynamicVertexBuffer : nullptr;
+
+		const uint32 VertexCount = MaxVertexIndex + 1;
+		Collector.AddRayTracingGeometryUpdate(
+			FRayTracingDynamicGeometryUpdateParams
+			{
+				RayTracingInstance.Materials,
+				false,
+				VertexCount,
+				VertexCount * (uint32)sizeof(FVector3f),
+				RayTracingGeometry.Initializer.TotalPrimitiveCount,
+				&RayTracingGeometry,
+				VertexBuffer,
+				true
+			}
+		);
+
+		Collector.AddRayTracingInstance(MoveTemp(RayTracingInstance));
+	}
+}
+#endif
+
+uint32 FGeometryCollectionSceneProxyBase::GetAllocatedSize() const
+{
+	return Materials.GetAllocatedSize()
+		+ MeshDescription.Sections.GetAllocatedSize()
+		+ MeshDescription.SubSections.GetAllocatedSize()
+		+ (SkinnedPositionVertexBuffer.GetAllowCPUAccess() ? SkinnedPositionVertexBuffer.GetStride() * SkinnedPositionVertexBuffer.GetNumVertices() : 0)
+#if RHI_RAYTRACING
+		+ RayTracingGeometry.RawData.GetAllocatedSize()
+#endif
+		;
+}
+
+FGeometryCollectionSceneProxy::FGeometryCollectionSceneProxy(UGeometryCollectionComponent* Component)
+	: FPrimitiveSceneProxy(Component)
+	, FGeometryCollectionSceneProxyBase(Component, false)
+#if WITH_EDITOR
+	, bShowBoneColors(Component->GetShowBoneColors())
+	, bSuppressSelectionMaterial(Component->GetSuppressSelectionMaterial())
+	, VertexFactoryDebugColor(GetScene().GetFeatureLevel())
+#endif
+{
+	if (Component->GetRestCollection())
+	{
+		GeometryCollection = Component->GetRestCollection()->GetGeometryCollection();
+	}
+
+	EnableGPUSceneSupportFlags();
+
+#if GEOMETRYCOLLECTION_EDITOR_SELECTION
+	// Render by SubSection if we are in the rigid body picker.
+	bUsesSubSections = Component->GetIsTransformSelectionMode() && MeshDescription.SubSections.Num();
+	// Enable bone hit selection proxies if we are in the rigid body picker or in the fracture modes.
+	bEnableBoneSelection = Component->GetEnableBoneSelection();
+
+	if (bEnableBoneSelection || bUsesSubSections)
+	{
+		for (int32 TransformIndex = 0; TransformIndex < NumTransforms; ++TransformIndex)
+		{
+			HGeometryCollection* HitProxy = new HGeometryCollection(Component, TransformIndex);
+			HitProxies.Add(HitProxy);
+		}
+	}
+#endif
+
+#if WITH_EDITOR
+	if (bShowBoneColors || bEnableBoneSelection)
+	{
+		Component->GetBoneColors(BoneColors);
+		ColorVertexBuffer.InitFromColorArray(BoneColors);
+
+		if (Component->GetRestCollection())
+		{
+			BoneSelectedMaterial = Component->GetRestCollection()->GetBoneSelectedMaterial();
+		}
+		if (BoneSelectedMaterial && !BoneSelectedMaterial->CheckMaterialUsage_Concurrent(MATUSAGE_GeometryCollections))
+		{
+			// If we have an invalid BoneSelectedMaterial, switch it back to null to skip its usage in GetDynamicMeshElements below
+			BoneSelectedMaterial = nullptr;
+		}
+
+		// Make sure the vertex color material has the usage flag for rendering geometry collections
+		if (GEngine->VertexColorMaterial)
+		{
+			GEngine->VertexColorMaterial->CheckMaterialUsage_Concurrent(MATUSAGE_GeometryCollections);
+		}
+	}
+
+#endif
+
+	// #todo(dmp): This flag means that when motion blur is turned on, it will always render geometry collections into the
+	// velocity buffer.  Note that the way around this is to loop through the global matrices and test whether they have
+	// changed from the prev to curr frame, but this is expensive.  We should revisit this if the draw calls for velocity
+	// rendering become a problem. One solution could be to use internal solver sleeping state to drive motion blur.
+	bAlwaysHasVelocity = true;
+
+	SetWireframeColor(Component->GetWireframeColorForSceneProxy());
+	CollisionResponse = Component->GetCollisionResponseToChannels();
+}
+
+FGeometryCollectionSceneProxy::~FGeometryCollectionSceneProxy()
+{
+}
+
+SIZE_T FGeometryCollectionSceneProxy::GetTypeHash() const
+{
+	static size_t UniquePointer;
+	return reinterpret_cast<size_t>(&UniquePointer);
+}
+
 void FGeometryCollectionSceneProxy::CreateRenderThreadResources(FRHICommandListBase& RHICmdList)
 {
-	if (bSupportsManualVertexFetch)
-	{
-		// Initialize transform buffers and upload rest transforms.
-		TransformBuffers.AddDefaulted(1);
-
-		TransformBuffers[0].NumTransforms = NumTransforms;
-		TransformBuffers[0].InitResource(RHICmdList);
-
-		const bool bLocalGeometryCollectionTripleBufferUploads = (GGeometryCollectionTripleBufferUploads != 0) && bSupportsTripleBufferVertexUpload;
-		const EResourceLockMode LockMode = bLocalGeometryCollectionTripleBufferUploads ? RLM_WriteOnly_NoOverwrite : RLM_WriteOnly;
-
-		FGeometryCollectionTransformBuffer& TransformBuffer = GetCurrentTransformBuffer();
-		TransformBuffer.UpdateDynamicData(RHICmdList, DynamicData->Transforms, LockMode);
-	}
-	else
-	{
-		// Initialize CPU skinning buffer with rest transforms.
-		SkinnedPositionVertexBuffer.Init(MeshResource.PositionVertexBuffer.GetNumVertices(), false);
-		SkinnedPositionVertexBuffer.InitResource(RHICmdList);
-		UpdateSkinnedPositions(RHICmdList, DynamicData->Transforms);
-	}
-
-	SetupVertexFactory(RHICmdList, VertexFactory);
+	FGeometryCollectionSceneProxyBase::CreateRenderThreadResources(RHICmdList);
 
 #if WITH_EDITOR
 	if (bShowBoneColors || bEnableBoneSelection)
@@ -334,43 +706,13 @@
 	}
 #endif
 
-#if RHI_RAYTRACING
-	if (IsRayTracingEnabled())
-	{
-		FRayTracingGeometryInitializer Initializer;
-		Initializer.DebugName = FName(TEXT("GeometryCollection"));
-		Initializer.GeometryType = RTGT_Triangles;
-		Initializer.bFastBuild = true;
-		Initializer.bAllowUpdate = false;
-		Initializer.TotalPrimitiveCount = 0;
-
-		RayTracingGeometry.SetInitializer(Initializer);
-		RayTracingGeometry.InitResource(RHICmdList);
-
-		bGeometryResourceUpdated = true;
-	}
-#endif
-
 	bRenderResourcesCreated = true;
 	SetDynamicData_RenderThread(RHICmdList, DynamicData);
 }
 
 void FGeometryCollectionSceneProxy::DestroyRenderThreadResources()
 {
-	if (bSupportsManualVertexFetch)
-	{
-		for (int32 i = 0; i < TransformBuffers.Num(); i++)
-		{
-			TransformBuffers[i].ReleaseResource();
-		}
-		TransformBuffers.Reset();
-	}
-	else
-	{
-		SkinnedPositionVertexBuffer.ReleaseResource();
-	}
+	FGeometryCollectionSceneProxyBase::DestroyRenderThreadResources();
 
-	VertexFactory.ReleaseResource();
-
 #if WITH_EDITOR
 	VertexFactoryDebugColor.ReleaseResource();
@@ -380,27 +722,11 @@
 #if GEOMETRYCOLLECTION_EDITOR_SELECTION
 	HitProxyIdBuffer.ReleaseResource();
 #endif
-
-#if RHI_RAYTRACING
-	if (IsRayTracingAllowed())
-	{
-		RayTracingGeometry.ReleaseResource();
-		RayTracingDynamicVertexBuffer.Release();
-	}
-#endif
 }
 
 void FGeometryCollectionSceneProxy::SetDynamicData_RenderThread(FRHICommandListBase& RHICmdList, FGeometryCollectionDynamicData* NewDynamicData)
 {
-	if (NewDynamicData != DynamicData)
-	{
-		if (DynamicData)
-		{
-			GDynamicDataPool.Release(DynamicData);
-			DynamicData = nullptr;
-		}
-		DynamicData = NewDynamicData;
-	}
+	FGeometryCollectionSceneProxyBase::SetDynamicData_RenderThread(RHICmdList, NewDynamicData);
 
 	if (MeshDescription.NumVertices == 0 || !DynamicData || !bRenderResourcesCreated)
 	{
@@ -409,130 +735,18 @@
 		
 	if (bSupportsManualVertexFetch)
 	{
-		const bool bLocalGeometryCollectionTripleBufferUploads = (GGeometryCollectionTripleBufferUploads != 0) && bSupportsTripleBufferVertexUpload;
+		const FGeometryCollectionTransformBuffer& TransformBuffer = GetCurrentTransformBuffer();
+		const FGeometryCollectionTransformBuffer& PrevTransformBuffer = GetCurrentPrevTransformBuffer();
 
-		if (bLocalGeometryCollectionTripleBufferUploads && TransformBuffers.Num() == 1)
-		{
-			TransformBuffers.AddDefaulted(2);
-			check(TransformBuffers.Num() == 3);
-
-			for (int32 i = 1; i < TransformBuffers.Num(); i++)
-			{
-				TransformBuffers[i].NumTransforms = NumTransforms;
-				TransformBuffers[i].InitResource(RHICmdList);
-			}
-		}
-
-		// Copy the transform data over to the vertex buffer	
-		{
-			const EResourceLockMode LockMode = bLocalGeometryCollectionTripleBufferUploads ? RLM_WriteOnly_NoOverwrite : RLM_WriteOnly;
-
-			CycleTransformBuffers(bLocalGeometryCollectionTripleBufferUploads);
-
-			FGeometryCollectionTransformBuffer& TransformBuffer = GetCurrentTransformBuffer();
-			FGeometryCollectionTransformBuffer& PrevTransformBuffer = GetCurrentPrevTransformBuffer();
-
-			VertexFactory.SetBoneTransformSRV(TransformBuffer.VertexBufferSRV);
-			VertexFactory.SetBonePrevTransformSRV(PrevTransformBuffer.VertexBufferSRV);
-	
-			TransformBuffer.UpdateDynamicData(RHICmdList, DynamicData->Transforms, LockMode);
-
-			UpdateLooseParameter(VertexFactory, TransformBuffer.VertexBufferSRV, PrevTransformBuffer.VertexBufferSRV, MeshResource.BoneMapVertexBuffer.GetSRV());
-
 #if WITH_EDITOR
-			if (bShowBoneColors || bEnableBoneSelection)
-			{
-				VertexFactoryDebugColor.SetBoneTransformSRV(TransformBuffer.VertexBufferSRV);
-				VertexFactoryDebugColor.SetBonePrevTransformSRV(PrevTransformBuffer.VertexBufferSRV);
-				UpdateLooseParameter(VertexFactoryDebugColor, TransformBuffer.VertexBufferSRV, PrevTransformBuffer.VertexBufferSRV, MeshResource.BoneMapVertexBuffer.GetSRV());
-			}
+		if (bShowBoneColors || bEnableBoneSelection)
+		{
+			VertexFactoryDebugColor.SetBoneTransformSRV(TransformBuffer.VertexBufferSRV);
+			VertexFactoryDebugColor.SetBonePrevTransformSRV(PrevTransformBuffer.VertexBufferSRV);
+			UpdateLooseParameter(VertexFactoryDebugColor, TransformBuffer.VertexBufferSRV, PrevTransformBuffer.VertexBufferSRV, MeshResource.BoneMapVertexBuffer.GetSRV());
+		}
 #endif
-		}
 	}
-	else
-	{
-		UpdateSkinnedPositions(RHICmdList, DynamicData->Transforms);
-	}
-
-#if RHI_RAYTRACING
-	if (IsRayTracingEnabled())
-	{
-		bGeometryResourceUpdated = true;
-	}
-#endif
 }
 
-void FGeometryCollectionSceneProxy::UpdateSkinnedPositions(FRHICommandListBase& RHICmdList, TArray<FMatrix44f> const& Transforms)
-{
-	const int32 VertexStride = SkinnedPositionVertexBuffer.GetStride();
-	const int32 VertexCount = SkinnedPositionVertexBuffer.GetNumVertices();
-	check (VertexCount == MeshDescription.NumVertices)
-
-	void* VertexBufferData = RHICmdList.LockBuffer(SkinnedPositionVertexBuffer.VertexBufferRHI, 0, VertexCount * VertexStride, RLM_WriteOnly);
-	check(VertexBufferData != nullptr);
-
-	FPositionVertexBuffer const& SourcePositionVertexBuffer = MeshResource.PositionVertexBuffer;
-	FBoneMapVertexBuffer const& SourceBoneMapVertexBuffer = MeshResource.BoneMapVertexBuffer;
-
-	bool bParallelGeometryCollection = true;
-	int32 ParallelGeometryCollectionBatchSize = CVarParallelGeometryCollectionBatchSize.GetValueOnRenderThread();
-
-	int32 NumBatches = (VertexCount / ParallelGeometryCollectionBatchSize);
-
-	if (VertexCount != ParallelGeometryCollectionBatchSize)
-	{
-		NumBatches++;
-	}
-
-	// Batch too small, don't bother with parallel
-	if (ParallelGeometryCollectionBatchSize > VertexCount)
-	{
-		bParallelGeometryCollection = false;
-		ParallelGeometryCollectionBatchSize = VertexCount;
-	}
-
-	auto GeometryCollectionBatch([&](int32 BatchNum)
-	{
-		uint32 IndexOffset = ParallelGeometryCollectionBatchSize * BatchNum;
-		uint32 ThisBatchSize = ParallelGeometryCollectionBatchSize;
-
-		// Check for final batch
-		if (IndexOffset + ParallelGeometryCollectionBatchSize > MeshDescription.NumVertices)
-		{
-			ThisBatchSize = VertexCount - IndexOffset;
-		}
-
-		if (ThisBatchSize > 0)
-		{
-			const FMatrix44f* RESTRICT BoneTransformsPtr = Transforms.GetData();
-
-			if (bGeometryCollection_SetDynamicData_ISPC_Enabled)
-			{
-#if INTEL_ISPC
-				uint8* VertexBufferOffset = (uint8*)VertexBufferData + (IndexOffset * VertexStride);
-				ispc::SetDynamicData_RenderThread(
-					(ispc::FVector3f*)VertexBufferOffset,
-					ThisBatchSize,
-					VertexStride,
-					&SourceBoneMapVertexBuffer.BoneIndex(IndexOffset),
-					(ispc::FMatrix44f*)BoneTransformsPtr,
-					(ispc::FVector3f*)&SourcePositionVertexBuffer.VertexPosition(IndexOffset));
-#endif
-			}
-			else
-			{
-				for (uint32 i = IndexOffset; i < IndexOffset + ThisBatchSize; i++)
-				{
-					FVector3f Transformed = BoneTransformsPtr[SourceBoneMapVertexBuffer.BoneIndex(i)].TransformPosition(SourcePositionVertexBuffer.VertexPosition(i));
-					FMemory::Memcpy((uint8*)VertexBufferData + (i * VertexStride), &Transformed, sizeof(FVector3f));
-				}
-			}
-		}
-	});
-
-	ParallelFor(NumBatches, GeometryCollectionBatch, !bParallelGeometryCollection);
-
-	RHICmdList.UnlockBuffer(SkinnedPositionVertexBuffer.VertexBufferRHI);
-}
-
 FMaterialRenderProxy* FGeometryCollectionSceneProxy::GetMaterial(FMeshElementCollector& Collector, int32 MaterialIndex) const
@@ -748,145 +962,6 @@
 	}
 }
 
-#if RHI_RAYTRACING
-void FGeometryCollectionSceneProxy::GetDynamicRayTracingInstances(FRayTracingInstanceCollector& Collector)
-{
-	if (GRayTracingGeometryCollectionProxyMeshes == 0 || MeshDescription.NumVertices == 0)
-	{
-		return;
-	}
-
-	QUICK_SCOPE_CYCLE_COUNTER(STAT_GeometryCollectionSceneProxy_GetDynamicRayTracingInstances);
-
-	const uint32 LODIndex = 0;
-	const bool bWireframe = false;
-		
-	//Loose parameter needs to be updated every frame
-	FGeometryCollectionMeshCollectorResources* CollectorResources;
-	CollectorResources = &Collector.AllocateOneFrameResource<FGeometryCollectionMeshCollectorResources>(GetScene().GetFeatureLevel());
-	FGeometryCollectionVertexFactory& GeometryCollectionVertexFactory = CollectorResources->GetVertexFactory();
-		
-	// Render dynamic objects
-	if (!GeometryCollectionVertexFactory.GetType()->SupportsRayTracingDynamicGeometry())
-	{
-		return;
-	}
-
-	SetupVertexFactory(Collector.GetRHICommandList(), GeometryCollectionVertexFactory);
-
-	TArray<FGeometryCollectionMeshElement> const& SectionArray = MeshDescription.Sections;
-
-	UpdatingRayTracingGeometry_RenderingThread(SectionArray);
-
-	if (SectionArray.Num() && RayTracingGeometry.IsValid())
-	{
-		FRayTracingInstance RayTracingInstance;
-		RayTracingInstance.Geometry = &RayTracingGeometry;
-		RayTracingInstance.InstanceTransforms.Emplace(GetLocalToWorld());
-
-		// Grab the material proxies we'll be using for each section
-		TArray<FMaterialRenderProxy*, TInlineAllocator<32>> MaterialProxies;
-
-		for (int32 SectionIndex = 0; SectionIndex < SectionArray.Num(); ++SectionIndex)
-		{
-			const FGeometryCollectionMeshElement& Section = SectionArray[SectionIndex];
-				
-			//TODO: Add BoneColor support in Path/Ray tracing?
-			FMaterialRenderProxy* MaterialProxy= Materials[Section.MaterialIndex]->GetRenderProxy();
-
-			if (MaterialProxy == nullptr)
-			{
-				MaterialProxy = UMaterial::GetDefaultMaterial(MD_Surface)->GetRenderProxy();
-			}
-
-			MaterialProxies.Add(MaterialProxy);
-		}
-
-		uint32 MaxVertexIndex = 0;
-		for (int32 SectionIndex = 0; SectionIndex < SectionArray.Num(); ++SectionIndex)
-		{
-			const FGeometryCollectionMeshElement& Section = SectionArray[SectionIndex];
-
-			// Draw the mesh
-			FMeshBatch& Mesh = RayTracingInstance.Materials.AddDefaulted_GetRef();
-			Mesh.bWireframe = bWireframe;
-			Mesh.SegmentIndex = SectionIndex;
-			Mesh.VertexFactory = &GeometryCollectionVertexFactory;
-			Mesh.MaterialRenderProxy = MaterialProxies[SectionIndex];
-			Mesh.LODIndex = LODIndex;
-			Mesh.bDisableBackfaceCulling = true;
-			Mesh.Type = PT_TriangleList;
-			Mesh.DepthPriorityGroup = SDPG_World;
-			Mesh.bCanApplyViewModeOverrides = true;
-
-			FMeshBatchElement& BatchElement = Mesh.Elements[0];
-			BatchElement.IndexBuffer = &MeshResource.IndexBuffer;
-			BatchElement.PrimitiveUniformBuffer = GetUniformBuffer();
-			BatchElement.FirstIndex = Section.TriangleStart * 3;
-			BatchElement.NumPrimitives = Section.TriangleCount;
-			BatchElement.MinVertexIndex = Section.VertexStart;
-			BatchElement.MaxVertexIndex = Section.VertexEnd; 
-			BatchElement.NumInstances = 1;
-
-			MaxVertexIndex = std::max(Section.VertexEnd, MaxVertexIndex);
-
-			//#TODO: bone color, bone selection and render bound?
-		}
-
-		FRWBuffer* VertexBuffer = RayTracingDynamicVertexBuffer.NumBytes > 0 ? &RayTracingDynamicVertexBuffer : nullptr;
-
-		const uint32 VertexCount = MaxVertexIndex + 1;
-		Collector.AddRayTracingGeometryUpdate(
-			FRayTracingDynamicGeometryUpdateParams
-			{
-				RayTracingInstance.Materials,
-				false,
-				VertexCount,
-				VertexCount * (uint32)sizeof(FVector3f),
-				RayTracingGeometry.Initializer.TotalPrimitiveCount,
-				&RayTracingGeometry,
-				VertexBuffer,
-				true
-			}
-		);
-
-		Collector.AddRayTracingInstance(MoveTemp(RayTracingInstance));
-	}
-}
-
-void FGeometryCollectionSceneProxy::UpdatingRayTracingGeometry_RenderingThread(TArray<FGeometryCollectionMeshElement> const& InSectionArray)
-{
-	FRHICommandList& RHICmdList = FRHICommandListImmediate::Get();
-
-	if (bGeometryResourceUpdated)
-	{
-		RayTracingGeometry.Initializer.Segments.Empty();
-		RayTracingGeometry.Initializer.TotalPrimitiveCount = 0;
-
-		for (int SectionIndex = 0; SectionIndex < InSectionArray.Num(); ++SectionIndex)
-		{
-			const FGeometryCollectionMeshElement& Section = InSectionArray[SectionIndex];
-			FRayTracingGeometrySegment Segment;
-			Segment.FirstPrimitive = Section.TriangleStart;
-			Segment.VertexBuffer = MeshResource.PositionVertexBuffer.VertexBufferRHI;
-			Segment.NumPrimitives = Section.TriangleCount;
-			Segment.MaxVertices = Section.VertexEnd;
-			RayTracingGeometry.Initializer.Segments.Add(Segment);
-			RayTracingGeometry.Initializer.TotalPrimitiveCount += Section.TriangleCount;
-		}
-				
-		if (RayTracingGeometry.Initializer.TotalPrimitiveCount > 0)
-		{
-			RayTracingGeometry.Initializer.IndexBuffer = MeshResource.IndexBuffer.IndexBufferRHI;
-			// Create the ray tracing geometry but delay the acceleration structure build.
-			RayTracingGeometry.CreateRayTracingGeometry(RHICmdList, ERTAccelerationStructureBuildPriority::Skip);
-		}
-
-		bGeometryResourceUpdated = false;
-	}
-}
-#endif // RHI_RAYTRACING
-
 FPrimitiveViewRelevance FGeometryCollectionSceneProxy::GetViewRelevance(const FSceneView* View) const
 {
 	FPrimitiveViewRelevance Result;
@@ -921,10 +996,7 @@
 uint32 FGeometryCollectionSceneProxy::GetAllocatedSize() const
 {
 	return FPrimitiveSceneProxy::GetAllocatedSize()
-		+ Materials.GetAllocatedSize()
-		+ MeshDescription.Sections.GetAllocatedSize()
-		+ MeshDescription.SubSections.GetAllocatedSize()
-		+ (SkinnedPositionVertexBuffer.GetAllowCPUAccess() ? SkinnedPositionVertexBuffer.GetStride() * SkinnedPositionVertexBuffer.GetNumVertices() : 0)
+		+ FGeometryCollectionSceneProxyBase::GetAllocatedSize()
 #if WITH_EDITOR
 		+ BoneColors.GetAllocatedSize()
 		+ (ColorVertexBuffer.GetAllowCPUAccess() ? ColorVertexBuffer.GetStride() * ColorVertexBuffer.GetNumVertices() : 0)
@@ -933,8 +1005,5 @@
 		+ HitProxies.GetAllocatedSize()
 		+ (HitProxyIdBuffer.GetAllowCPUAccess() ? HitProxyIdBuffer.GetStride() * HitProxyIdBuffer.GetNumVertices() : 0)
 #endif
-#if RHI_RAYTRACING
-		+ RayTracingGeometry.RawData.GetAllocatedSize()
-#endif
 		;
 }
@@ -942,6 +1011,7 @@
 
 FNaniteGeometryCollectionSceneProxy::FNaniteGeometryCollectionSceneProxy(UGeometryCollectionComponent* Component)
 : Nanite::FSceneProxyBase(Component)
+, FGeometryCollectionSceneProxyBase(Component, true)
 , GeometryCollection(Component->GetRestCollection())
 , bRequiresGPUSceneUpdate(false)
 , bEnableBoneSelection(false)
@@ -953,8 +1023,6 @@
 	checkSlow(DoesPlatformSupportNanite(GMaxRHIShaderPlatform));
 	checkSlow(GeometryCollection->HasNaniteData());
 
-	MaterialRelevance = Component->GetMaterialRelevance(Component->GetScene()->GetFeatureLevel());
-
 #if GEOMETRYCOLLECTION_EDITOR_SELECTION
 	bEnableBoneSelection = Component->GetEnableBoneSelection();
 #endif
@@ -1080,15 +1148,15 @@
 		
 		const TManagedArray<FBox>& BoundingBoxes = Collection->GetAttribute<FBox>("BoundingBox", FGeometryCollection::TransformGroup);
 		const TManagedArray<int32>& TransformToGeometry = Collection->GetAttribute<int32>("TransformToGeometryIndex", FGeometryCollection::TransformGroup);
-		const int32 NumTransforms = TransformToGeometry.Num();
-		for (int32 TransformIndex = 0; TransformIndex < NumTransforms; ++TransformIndex)
+		const int32 NumTransformToGeometry = TransformToGeometry.Num();
+		for (int32 TransformToGeometryIndex = 0; TransformToGeometryIndex < NumTransformToGeometry; ++TransformToGeometryIndex)
 		{
-			const int32 GeometryIndex = TransformToGeometry[TransformIndex];
+			const int32 GeometryIndex = TransformToGeometry[TransformToGeometryIndex];
 			if (GeometryIndex > INDEX_NONE)
 			{
 				FGeometryNaniteData& Instance = GeometryNaniteData[GeometryIndex];
 				Instance.HierarchyOffset = GeometryCollection->GetNaniteHierarchyOffset(GeometryIndex);
-				Instance.LocalBounds = BoundingBoxes[TransformIndex];
+				Instance.LocalBounds = BoundingBoxes[TransformToGeometryIndex];
 			}
 		}
 	}
@@ -1127,25 +1195,29 @@
 
 	CollisionResponse = Component->GetCollisionResponseToChannels();
 
-	FGeometryCollectionDynamicData* InitDynamicData = Component->InitDynamicData(true);
-	SetDynamicData_RenderThread(InitDynamicData, Component->GetRenderMatrix());
+	UpdateInstanceSceneDataBuffers(Component->GetRenderMatrix());
 }
 
 FNaniteGeometryCollectionSceneProxy::~FNaniteGeometryCollectionSceneProxy()
 {
-	if (DynamicData != nullptr)
-	{
-		GDynamicDataPool.Release(DynamicData);
-		DynamicData = nullptr;
-	}
+
 }
 
 void FNaniteGeometryCollectionSceneProxy::CreateRenderThreadResources(FRHICommandListBase& RHICmdList)
 {
+	FGeometryCollectionSceneProxyBase::CreateRenderThreadResources(RHICmdList);
+
 	// Should have valid Nanite data at this point.
 	NaniteResourceID = GeometryCollection->GetNaniteResourceID();
 	NaniteHierarchyOffset = GeometryCollection->GetNaniteHierarchyOffset();
 	check(NaniteResourceID != INDEX_NONE && NaniteHierarchyOffset != INDEX_NONE);
+
+	FGeometryCollectionSceneProxyBase::SetDynamicData_RenderThread(RHICmdList, DynamicData);
 }
 
+void FNaniteGeometryCollectionSceneProxy::DestroyRenderThreadResources()
+{
+	FGeometryCollectionSceneProxyBase::DestroyRenderThreadResources();
+}
+
 SIZE_T FNaniteGeometryCollectionSceneProxy::GetTypeHash() const
@@ -1204,9 +1276,10 @@
 	DrawStaticElementsInternal(PDI, LCI);
 }
 
-uint32 FNaniteGeometryCollectionSceneProxy::GetMemoryFootprint() const
+uint32 FNaniteGeometryCollectionSceneProxy::GetAllocatedSize() const
 {
-	return sizeof(*this) + GetAllocatedSize();
+	return FPrimitiveSceneProxy::GetAllocatedSize()
+		+ FGeometryCollectionSceneProxyBase::GetAllocatedSize();
 }
 
 void FNaniteGeometryCollectionSceneProxy::GetNaniteResourceInfo(uint32& ResourceID, uint32& HierarchyOffset, uint32& ImposterIndex) const
@@ -1237,39 +1310,29 @@
 	return MoveTemp(OutInfo);
 }
 
-void FNaniteGeometryCollectionSceneProxy::SetDynamicData_RenderThread(FGeometryCollectionDynamicData* NewDynamicData, const FMatrix &PrimitiveLocalToWorld)
+void FNaniteGeometryCollectionSceneProxy::UpdateInstanceSceneDataBuffers(const FMatrix& PrimitiveLocalToWorld)
 {
-	if (NewDynamicData != DynamicData)
-	{
-		if (DynamicData)
-		{
-			GDynamicDataPool.Release(DynamicData);
-			DynamicData = nullptr;
-		}
-		DynamicData = NewDynamicData;
-	}
-
 	FInstanceSceneDataBuffers::FAccessTag AccessTag(PointerHash(this));
 	FInstanceSceneDataBuffers::FWriteView ProxyData = InstanceSceneDataBuffersImpl.BeginWriteAccess(AccessTag);
 	InstanceSceneDataBuffersImpl.SetPrimitiveLocalToWorld(PrimitiveLocalToWorld, AccessTag);
 
 	const TSharedPtr<FGeometryCollection, ESPMode::ThreadSafe> Collection = GeometryCollection->GetGeometryCollection();
-	const TManagedArray<int32>& TransformToGeometryIndices	 = Collection->TransformToGeometryIndex;
-	const TManagedArray<TSet<int32>>& TransformChildren		 = Collection->Children;
-	const TManagedArray<int32>& SimulationType				 = Collection->SimulationType;
+	const TManagedArray<int32>& TransformToGeometryIndices = Collection->TransformToGeometryIndex;
+	const TManagedArray<TSet<int32>>& TransformChildren = Collection->Children;
+	const TManagedArray<int32>& SimulationType = Collection->SimulationType;
 
-	const int32 TransformCount = NewDynamicData->Transforms.Num();
+	const int32 TransformCount = DynamicData->Transforms.Num();
 	check(TransformCount == TransformToGeometryIndices.Num());
 	check(TransformCount == TransformChildren.Num());
 
 	// set the prev by copying the last current
 	ProxyData.PrevInstanceToPrimitiveRelative = ProxyData.InstanceToPrimitiveRelative;
-	bCanSkipRedundantTransformUpdates = false; // shoudl we compare the transform to better decide about this ? 
+	bCanSkipRedundantTransformUpdates = false; // should we compare the transform to better decide about this ? 
 
 	ProxyData.InstanceToPrimitiveRelative.Reset(TransformCount);
 	ProxyData.InstanceLocalBounds.Reset(TransformCount);
 	ProxyData.InstanceHierarchyOffset.Reset(TransformCount);
-		
+
 #if GEOMETRYCOLLECTION_EDITOR_SELECTION
 	ProxyData.InstanceEditorData.Reset(bEnableBoneSelection ? TransformCount : 0);
 #endif
@@ -1287,11 +1350,11 @@
 		}
 
 		const FGeometryNaniteData& NaniteData = GeometryNaniteData[TransformToGeometryIndex];
-		const FRenderTransform& InstanceToPrimitiveRelative = ProxyData.InstanceToPrimitiveRelative.Emplace_GetRef(InstanceSceneDataBuffersImpl.ComputeInstanceToPrimitiveRelative(NewDynamicData->Transforms[TransformIndex], AccessTag));
+		const FRenderTransform& InstanceToPrimitiveRelative = ProxyData.InstanceToPrimitiveRelative.Emplace_GetRef(InstanceSceneDataBuffersImpl.ComputeInstanceToPrimitiveRelative(DynamicData->Transforms[TransformIndex], AccessTag));
 
 		ProxyData.InstanceLocalBounds.Emplace(PadInstanceLocalBounds(NaniteData.LocalBounds));
 		ProxyData.InstanceHierarchyOffset.Emplace(NaniteData.HierarchyOffset);
-			
+
 #if GEOMETRYCOLLECTION_EDITOR_SELECTION
 		if (bEnableBoneSelection)
 		{
@@ -1311,6 +1374,13 @@
 	InstanceSceneDataBuffersImpl.EndWriteAccess(AccessTag);
 }
 
+void FNaniteGeometryCollectionSceneProxy::SetDynamicData_RenderThread(FRHICommandListBase& RHICmdList, FGeometryCollectionDynamicData* NewDynamicData, const FMatrix &PrimitiveLocalToWorld)
+{
+	FGeometryCollectionSceneProxyBase::SetDynamicData_RenderThread(RHICmdList, NewDynamicData);
+
+	UpdateInstanceSceneDataBuffers(PrimitiveLocalToWorld);
+}
+
 void FNaniteGeometryCollectionSceneProxy::ResetPreviousTransforms_RenderThread()
 {
 	FInstanceSceneDataBuffers::FAccessTag AccessTag(PointerHash(this));
diff --git a/Engine/Source/Runtime/Experimental/GeometryCollectionEngine/Private/GeometryCollection/GeometryCollectionSceneProxy.h b/Engine/Source/Runtime/Experimental/GeometryCollectionEngine/Private/GeometryCollection/GeometryCollectionSceneProxy.h
--- a/Engine/Source/Runtime/Experimental/GeometryCollectionEngine/Private/GeometryCollection/GeometryCollectionSceneProxy.h
+++ b/Engine/Source/Runtime/Experimental/GeometryCollectionEngine/Private/GeometryCollection/GeometryCollectionSceneProxy.h
@@ -125,6 +125,98 @@
 };
 
 
+class FGeometryCollectionSceneProxyBase
+{
+protected:
+	FGeometryCollectionSceneProxyBase(UGeometryCollectionComponent* Component, bool bInIsNanite);
+	virtual ~FGeometryCollectionSceneProxyBase();
+
+	void CreateRenderThreadResources(FRHICommandListBase& RHICmdList);
+	void DestroyRenderThreadResources();
+
+	/** Setup a geometry collection vertex factory. */
+	void SetupVertexFactory(FRHICommandListBase& RHICmdList, FGeometryCollectionVertexFactory& GeometryCollectionVertexFactory, FColorVertexBuffer* ColorOverride = nullptr) const;
+
+	/** Called on render thread to setup dynamic geometry for rendering */
+	void SetDynamicData_RenderThread(FRHICommandListBase& RHICmdList, FGeometryCollectionDynamicData* NewDynamicData);
+
+	FGeometryCollectionTransformBuffer& GetCurrentTransformBuffer()
+	{
+		return TransformBuffers[CurrentTransformBufferIndex];
+	}
+
+	const FGeometryCollectionTransformBuffer& GetCurrentTransformBuffer() const
+	{
+		return TransformBuffers[CurrentTransformBufferIndex];
+	}
+
+	const FGeometryCollectionTransformBuffer& GetCurrentPrevTransformBuffer() const
+	{
+		const int32 NumBuffers = TransformBuffers.Num();
+		const int32 PreviousIndex = (CurrentTransformBufferIndex + NumBuffers - 1) % NumBuffers;
+		return TransformBuffers[PreviousIndex];
+	}
+
+	void CycleTransformBuffers(bool bCycle)
+	{
+		if (bCycle)
+		{
+			CurrentTransformBufferIndex = (CurrentTransformBufferIndex + 1) % TransformBuffers.Num();
+		}
+	}
+
+#if RHI_RAYTRACING
+	void UpdatingRayTracingGeometry_RenderingThread(TArray<FGeometryCollectionMeshElement> const& InSectionArray);
+
+	bool IsRayTracingRelevant() const { return true; }
+	bool IsRayTracingStaticRelevant() const { return false; }
+	bool HasRayTracingRepresentation() const { return bHasRayTracingRepresentation; }
+	void GetDynamicRayTracingInstances(FRayTracingInstanceCollector& Collector, const FMatrix& LocalToWorld, FRHIUniformBuffer* UniformBuffer);
+#endif
+
+	uint32 GetAllocatedSize() const;
+
+protected:
+
+	const bool bIsNanite;
+	ERHIFeatureLevel::Type FeatureLevel;
+
+	const FGeometryCollectionMeshResources& MeshResource;
+	FGeometryCollectionMeshDescription MeshDescription;
+	FMaterialRelevance MaterialRelevance;
+
+	FGeometryCollectionVertexFactory VertexFactory;
+	const bool bSupportsManualVertexFetch;
+	const bool bSupportsTripleBufferVertexUpload;
+
+	TArray<UMaterialInterface*> Materials;
+
+	int32 NumTransforms = 0;
+
+	FGeometryCollectionDynamicData* DynamicData = nullptr;
+
+private:
+
+	/** Update skinned position buffer used by mobile CPU skinning path. */
+	void UpdateSkinnedPositions(FRHICommandListBase& RHICmdList, TArray<FMatrix44f> const& Transforms);
+
+	FPositionVertexBuffer SkinnedPositionVertexBuffer;
+
+	int32 CurrentTransformBufferIndex = 0;
+	TArray<FGeometryCollectionTransformBuffer, TInlineAllocator<3>> TransformBuffers;
+
+	const bool bHasRayTracingRepresentation = false;
+
+#if RHI_RAYTRACING
+	bool bRayTracingGeometryDirty = true;
+	FRayTracingGeometry RayTracingGeometry;
+	FRWBuffer RayTracingDynamicVertexBuffer;
+#endif
+
+	bool bRenderResourcesCreated = false;
+};
+
+
 /***
 *   FGeometryCollectionSceneProxy
 *    
@@ -136,33 +228,15 @@
 *   - @todo double buffer - The double buffering of the FGeometryCollectionDynamicData.
 *   - @todo GPU skin : Make the skinning use the GpuVertexShader
 */
-class FGeometryCollectionSceneProxy final : public FPrimitiveSceneProxy
+class FGeometryCollectionSceneProxy final : public FPrimitiveSceneProxy, private FGeometryCollectionSceneProxyBase
 {
-	TArray<UMaterialInterface*> Materials;
-
-	FMaterialRelevance MaterialRelevance;
-
-	FGeometryCollectionMeshResources const& MeshResource;
-	FGeometryCollectionMeshDescription MeshDescription;
-
-	int32 NumTransforms = 0;
 	TSharedPtr<FGeometryCollection, ESPMode::ThreadSafe> GeometryCollection;
 
 	FCollisionResponseContainer CollisionResponse;
 
 	FBoxSphereBounds PreSkinnedBounds;
 
-	FGeometryCollectionVertexFactory VertexFactory;
-	
-	bool bSupportsManualVertexFetch;
-	FPositionVertexBuffer SkinnedPositionVertexBuffer;
-
-	int32 CurrentTransformBufferIndex = 0;
-	bool bSupportsTripleBufferVertexUpload = false;
 	bool bRenderResourcesCreated = false;
-	TArray<FGeometryCollectionTransformBuffer, TInlineAllocator<3>> TransformBuffers;
 
-	FGeometryCollectionDynamicData* DynamicData = nullptr;
-
 #if WITH_EDITOR
 	bool bShowBoneColors = false;
@@ -180,12 +254,6 @@
 	FColorVertexBuffer HitProxyIdBuffer;
 #endif
 
-#if RHI_RAYTRACING
-	bool bGeometryResourceUpdated = false;
-	FRayTracingGeometry RayTracingGeometry;
-	FRWBuffer RayTracingDynamicVertexBuffer;
-#endif
-
 public:
 	FGeometryCollectionSceneProxy(UGeometryCollectionComponent* Component);
 	virtual ~FGeometryCollectionSceneProxy();
@@ -193,7 +261,7 @@
 	/** Called on render thread to setup dynamic geometry for rendering */
 	void SetDynamicData_RenderThread(FRHICommandListBase& RHICmdList, FGeometryCollectionDynamicData* NewDynamicData);
 
-	uint32 GetMemoryFootprint() const override { return sizeof(*this) + GetAllocatedSize(); }
+	virtual uint32 GetMemoryFootprint() const override { return sizeof(*this) + GetAllocatedSize(); }
 	uint32 GetAllocatedSize() const;
 
 	SIZE_T GetTypeHash() const override;
@@ -214,50 +282,24 @@
 #endif
 
 #if RHI_RAYTRACING
-	bool IsRayTracingRelevant() const override { return true; }
-	bool IsRayTracingStaticRelevant() const override { return false; }
-	void GetDynamicRayTracingInstances(FRayTracingInstanceCollector& Collector) override;
+	bool IsRayTracingRelevant() const override { return FGeometryCollectionSceneProxyBase::IsRayTracingRelevant(); }
+	bool IsRayTracingStaticRelevant() const override { return FGeometryCollectionSceneProxyBase::IsRayTracingStaticRelevant(); }
+	bool HasRayTracingRepresentation() const override { return FGeometryCollectionSceneProxyBase::HasRayTracingRepresentation(); }
+	void GetDynamicRayTracingInstances(FRayTracingInstanceCollector& Collector) override { return FGeometryCollectionSceneProxyBase::GetDynamicRayTracingInstances(Collector, GetLocalToWorld(), GetUniformBuffer()); }
 #endif
 
 protected:
-	/** Setup a geometry collection vertex factory. */
-	void SetupVertexFactory(FRHICommandListBase& RHICmdList, FGeometryCollectionVertexFactory& GeometryCollectionVertexFactory, FColorVertexBuffer* ColorOverride = nullptr) const;
-	/** Update skinned position buffer used by mobile CPU skinning path. */
-	void UpdateSkinnedPositions(FRHICommandListBase& RHICmdList, TArray<FMatrix44f> const& Transforms);
 	/** Get material proxy from material ID */
 	FMaterialRenderProxy* GetMaterial(FMeshElementCollector& Collector, int32 MaterialIndex) const;
 	/** Get the standard or debug vertex factory dependent on current state. */
 	FVertexFactory const* GetVertexFactory() const;
 
-	FGeometryCollectionTransformBuffer& GetCurrentTransformBuffer()
-	{
-		return TransformBuffers[CurrentTransformBufferIndex];
-	}
-	FGeometryCollectionTransformBuffer& GetCurrentPrevTransformBuffer()
-	{
-		const int32 NumBuffers = TransformBuffers.Num();
-		const int32 PreviousIndex = (CurrentTransformBufferIndex + NumBuffers - 1) % NumBuffers;
-		return TransformBuffers[PreviousIndex];
-	}
-
-	void CycleTransformBuffers(bool bCycle)
-	{
-		if (bCycle)
-		{
-			CurrentTransformBufferIndex = (CurrentTransformBufferIndex + 1) % TransformBuffers.Num();
-		}
-	}
-
-#if RHI_RAYTRACING
-	void UpdatingRayTracingGeometry_RenderingThread(TArray<FGeometryCollectionMeshElement> const& InSectionArray);
-#endif
-
 private:
 	bool ShowCollisionMeshes(const FEngineShowFlags& EngineShowFlags) const;
 };
 
 
-class FNaniteGeometryCollectionSceneProxy : public Nanite::FSceneProxyBase
+class FNaniteGeometryCollectionSceneProxy : public Nanite::FSceneProxyBase, private FGeometryCollectionSceneProxyBase
 {
 public:
 	using Super = Nanite::FSceneProxyBase;
@@ -268,6 +310,7 @@
 public:
 	// FPrimitiveSceneProxy interface.
 	virtual void CreateRenderThreadResources(FRHICommandListBase& RHICmdList) override;
+	virtual void DestroyRenderThreadResources() override;
 	virtual SIZE_T GetTypeHash() const override;
 	virtual FPrimitiveViewRelevance	GetViewRelevance(const FSceneView* View) const override;
 	virtual void GetDynamicMeshElements(const TArray<const FSceneView*>& Views, const FSceneViewFamily& ViewFamily, uint32 VisibilityMap, FMeshElementCollector& Collector) const override;
@@ -276,7 +319,8 @@
 #endif
 	virtual void DrawStaticElements(FStaticPrimitiveDrawInterface* PDI) override;
 
-	virtual uint32 GetMemoryFootprint() const override;
+	virtual uint32 GetMemoryFootprint() const override { return sizeof(*this) + GetAllocatedSize(); }
+	uint32 GetAllocatedSize() const;
 
 	// FSceneProxyBase interface.
 	virtual void GetNaniteResourceInfo(uint32& ResourceID, uint32& HierarchyOffset, uint32& ImposterIndex) const override;
@@ -284,7 +328,7 @@
 	virtual Nanite::FResourceMeshInfo GetResourceMeshInfo() const override;
 
 	/** Called on render thread to setup dynamic geometry for rendering */
-	void SetDynamicData_RenderThread(FGeometryCollectionDynamicData* NewDynamicData, const FMatrix &PrimitiveLocalToWorld);
+	void SetDynamicData_RenderThread(FRHICommandListBase& RHICmdList, FGeometryCollectionDynamicData* NewDynamicData, const FMatrix &PrimitiveLocalToWorld);
 
 	void ResetPreviousTransforms_RenderThread();
 
@@ -306,6 +350,13 @@
 		LCIs.Add(LCI);
 	}
 
+#if RHI_RAYTRACING
+	bool IsRayTracingRelevant() const override { return FGeometryCollectionSceneProxyBase::IsRayTracingRelevant(); }
+	bool IsRayTracingStaticRelevant() const override { return FGeometryCollectionSceneProxyBase::IsRayTracingStaticRelevant(); }
+	bool HasRayTracingRepresentation() const override { return FGeometryCollectionSceneProxyBase::HasRayTracingRepresentation(); }
+	void GetDynamicRayTracingInstances(FRayTracingInstanceCollector& Collector) override { return FGeometryCollectionSceneProxyBase::GetDynamicRayTracingInstances(Collector, GetLocalToWorld(), GetUniformBuffer()); }
+#endif
+
 protected:
 	// TODO : Copy required data from UObject instead of using unsafe object pointer.
 	const UGeometryCollection* GeometryCollection = nullptr;
@@ -321,10 +372,6 @@
 	uint32 NaniteResourceID = INDEX_NONE;
 	uint32 NaniteHierarchyOffset = INDEX_NONE;
 
-	// TODO: Should probably calculate this on the materials array above instead of on the component
-	//       Null and !Opaque are assigned default material unlike the component material relevance.
-	FMaterialRelevance MaterialRelevance;
-
 	uint32 bCastShadow : 1;
 	uint32 bReverseCulling : 1;
 	uint32 bHasMaterialErrors : 1;
@@ -337,8 +384,6 @@
 
 	FInstanceSceneDataBuffers InstanceSceneDataBuffersImpl;
 
-	FGeometryCollectionDynamicData* DynamicData = nullptr;
-
 	// Geometry collection doesn't currently support baked light maps, so we use this simple empty light cache info for all nanite geometry collection proxies
 	class FEmptyLightCacheInfo : public FLightCacheInterface
 	{
@@ -349,6 +394,8 @@
 	};
 
 private:
+	void UpdateInstanceSceneDataBuffers(const FMatrix& PrimitiveLocalToWorld);
+
 	bool ShowCollisionMeshes(const FEngineShowFlags& EngineShowFlags) const;
 
 private:
diff --git a/Engine/Source/Runtime/Experimental/GeometryCollectionEngine/Public/GeometryCollection/GeometryCollectionComponent.h b/Engine/Source/Runtime/Experimental/GeometryCollectionEngine/Public/GeometryCollection/GeometryCollectionComponent.h
--- a/Engine/Source/Runtime/Experimental/GeometryCollectionEngine/Public/GeometryCollection/GeometryCollectionComponent.h
+++ b/Engine/Source/Runtime/Experimental/GeometryCollectionEngine/Public/GeometryCollection/GeometryCollectionComponent.h
@@ -591,8 +591,7 @@ class UGeometryCollectionComponent : public UMeshComponent, public IChaosNotifyH
 	friend class FScopedColorEdit;
 #endif
 	friend class FGeometryCollectionCommands;
-	friend class FGeometryCollectionSceneProxy;
-	friend class FNaniteGeometryCollectionSceneProxy;
+	friend class FGeometryCollectionSceneProxyBase;
 
 public:
 
