[HWRT] Cache RayTracingGeometry of Geometry Collections when only update when geometry is modified or if it uses WPO materials.

diff --git a/Engine/Source/Runtime/Experimental/GeometryCollectionEngine/Private/GeometryCollection/GeometryCollectionSceneProxy.cpp b/Engine/Source/Runtime/Experimental/GeometryCollectionEngine/Private/GeometryCollection/GeometryCollectionSceneProxy.cpp
--- a/Engine/Source/Runtime/Experimental/GeometryCollectionEngine/Private/GeometryCollection/GeometryCollectionSceneProxy.cpp
+++ b/Engine/Source/Runtime/Experimental/GeometryCollectionEngine/Private/GeometryCollection/GeometryCollectionSceneProxy.cpp
@@ -78,6 +78,11 @@ FAutoConsoleVariableRef CVarRayTracingGeometryCollectionProxyMeshes(
 	ECVF_RenderThreadSafe
 );
 
+static TAutoConsoleVariable<int32> CVarRayTracingGeometryCollectionForceUpdate(
+	TEXT("r.RayTracing.Geometry.GeometryCollection.ForceUpdate"),
+	0,
+	TEXT("Forces ray tracing representation for geometry collections meshes to be updated every frame"));
+
 
 #if !defined(CHAOS_GEOMETRY_COLLECTION_SET_DYNAMIC_DATA_ISPC_ENABLED_DEFAULT)
 #define CHAOS_GEOMETRY_COLLECTION_SET_DYNAMIC_DATA_ISPC_ENABLED_DEFAULT 1
@@ -472,7 +477,7 @@ void FGeometryCollectionSceneProxyBase::UpdatingRayTracingGeometry_RenderingThre
 	}
 }
 
-void FGeometryCollectionSceneProxyBase::GetDynamicRayTracingInstances(FRayTracingInstanceCollector& Collector, const FMatrix& LocalToWorld, FRHIUniformBuffer* UniformBuffer)
+void FGeometryCollectionSceneProxyBase::GetDynamicRayTracingInstances(FRayTracingInstanceCollector& Collector, const FMatrix& LocalToWorld, FRHIUniformBuffer* UniformBuffer, bool bAnyMaterialHasWorldPositionOffset)
 {
 	checkf(bHasRayTracingRepresentation, TEXT("Shouldn't try to get ray tracing instances from proxy that doesn't have a ray tracing representation."));
 
@@ -558,22 +563,39 @@ void FGeometryCollectionSceneProxyBase::GetDynamicRayTracingInstances(FRayTracin
 			//#TODO: bone color, bone selection and render bound?
 		}
 
-		FRWBuffer* VertexBuffer = RayTracingDynamicVertexBuffer.NumBytes > 0 ? &RayTracingDynamicVertexBuffer : nullptr;
+		const bool bAlwaysUpdate = bAnyMaterialHasWorldPositionOffset || (CVarRayTracingGeometryCollectionForceUpdate.GetValueOnRenderThread() != 0);
+		const bool bNeedsUpdate = bAlwaysUpdate
+			|| (!bAlwaysUpdate && RayTracingGeometry.DynamicGeometrySharedBufferGenerationID != FRayTracingGeometry::NonSharedVertexBuffers) // was using shared VB but won't use it anymore so update once
+			|| !RayTracingGeometry.IsValid()
+			|| RayTracingGeometry.IsEvicted()
+			|| RayTracingGeometry.GetRequiresBuild();
 
-		const uint32 VertexCount = MaxVertexIndex + 1;
-		Collector.AddRayTracingGeometryUpdate(
-			FRayTracingDynamicGeometryUpdateParams
-			{
-				RayTracingInstance.Materials,
-				false,
-				VertexCount,
-				VertexCount * (uint32)sizeof(FVector3f),
-				RayTracingGeometry.Initializer.TotalPrimitiveCount,
-				&RayTracingGeometry,
-				VertexBuffer,
-				true
-			}
-		);
+		FRWBuffer* VertexBuffer = &RayTracingDynamicVertexBuffer;
+
+		if (bAlwaysUpdate)
+		{
+			// if updating every frame release memory and use shared VB
+			VertexBuffer->Release();
+			VertexBuffer = nullptr;
+		}
+
+		if (bNeedsUpdate)
+		{
+			const uint32 VertexCount = MaxVertexIndex + 1;
+			Collector.AddRayTracingGeometryUpdate(
+				FRayTracingDynamicGeometryUpdateParams
+				{
+					RayTracingInstance.Materials,
+					false,
+					VertexCount,
+					VertexCount * (uint32)sizeof(FVector3f),
+					RayTracingGeometry.Initializer.TotalPrimitiveCount,
+					&RayTracingGeometry,
+					VertexBuffer,
+					true
+				}
+			);
+		}
 
 		Collector.AddRayTracingInstance(MoveTemp(RayTracingInstance));
 	}
diff --git a/Engine/Source/Runtime/Experimental/GeometryCollectionEngine/Private/GeometryCollection/GeometryCollectionSceneProxy.h b/Engine/Source/Runtime/Experimental/GeometryCollectionEngine/Private/GeometryCollection/GeometryCollectionSceneProxy.h
--- a/Engine/Source/Runtime/Experimental/GeometryCollectionEngine/Private/GeometryCollection/GeometryCollectionSceneProxy.h
+++ b/Engine/Source/Runtime/Experimental/GeometryCollectionEngine/Private/GeometryCollection/GeometryCollectionSceneProxy.h
@@ -171,7 +171,7 @@ protected:
 	bool IsRayTracingRelevant() const { return true; }
 	bool IsRayTracingStaticRelevant() const { return false; }
 	bool HasRayTracingRepresentation() const { return bHasRayTracingRepresentation; }
-	void GetDynamicRayTracingInstances(FRayTracingInstanceCollector& Collector, const FMatrix& LocalToWorld, FRHIUniformBuffer* UniformBuffer);
+	void GetDynamicRayTracingInstances(FRayTracingInstanceCollector& Collector, const FMatrix& LocalToWorld, FRHIUniformBuffer* UniformBuffer, bool bAnyMaterialHasWorldPositionOffset);
 #endif
 
 	uint32 GetAllocatedSize() const;
@@ -285,7 +285,7 @@ public:
 	bool IsRayTracingRelevant() const override { return FGeometryCollectionSceneProxyBase::IsRayTracingRelevant(); }
 	bool IsRayTracingStaticRelevant() const override { return FGeometryCollectionSceneProxyBase::IsRayTracingStaticRelevant(); }
 	bool HasRayTracingRepresentation() const override { return FGeometryCollectionSceneProxyBase::HasRayTracingRepresentation(); }
-	void GetDynamicRayTracingInstances(FRayTracingInstanceCollector& Collector) override { return FGeometryCollectionSceneProxyBase::GetDynamicRayTracingInstances(Collector, GetLocalToWorld(), GetUniformBuffer()); }
+	void GetDynamicRayTracingInstances(FRayTracingInstanceCollector& Collector) override { return FGeometryCollectionSceneProxyBase::GetDynamicRayTracingInstances(Collector, GetLocalToWorld(), GetUniformBuffer(), bAnyMaterialHasWorldPositionOffset); }
 #endif
 
 protected:
@@ -354,7 +354,7 @@ public:
 	bool IsRayTracingRelevant() const override { return FGeometryCollectionSceneProxyBase::IsRayTracingRelevant(); }
 	bool IsRayTracingStaticRelevant() const override { return FGeometryCollectionSceneProxyBase::IsRayTracingStaticRelevant(); }
 	bool HasRayTracingRepresentation() const override { return FGeometryCollectionSceneProxyBase::HasRayTracingRepresentation(); }
-	void GetDynamicRayTracingInstances(FRayTracingInstanceCollector& Collector) override { return FGeometryCollectionSceneProxyBase::GetDynamicRayTracingInstances(Collector, GetLocalToWorld(), GetUniformBuffer()); }
+	void GetDynamicRayTracingInstances(FRayTracingInstanceCollector& Collector) override { return FGeometryCollectionSceneProxyBase::GetDynamicRayTracingInstances(Collector, GetLocalToWorld(), GetUniformBuffer(), bAnyMaterialHasWorldPositionOffset); }
 #endif
 
 protected:
