Removed unused function in IoStoreOnDemand header.

#rnx

diff --git a/Engine/Source/Runtime/Experimental/IoStore/OnDemand/Public/IO/IoStoreOnDemand.h b/Engine/Source/Runtime/Experimental/IoStore/OnDemand/Public/IO/IoStoreOnDemand.h
--- a/Engine/Source/Runtime/Experimental/IoStore/OnDemand/Public/IO/IoStoreOnDemand.h
+++ b/Engine/Source/Runtime/Experimental/IoStore/OnDemand/Public/IO/IoStoreOnDemand.h
@@ -28,7 +28,6 @@ struct FKeyChain;
 struct FAnalyticsEventAttribute;
 struct FIoContainerSettings;
 struct FIoStoreWriterSettings;
-namespace UE::IoStore { struct FOnDemandEndpoint; }
 namespace UE::IoStore { class FOnDemandInternalContentHandle; }
 namespace UE::IoStore { class FOnDemandIoStore; }
 namespace UE::IoStore { class IOnDemandIoDispatcherBackend; }
@@ -50,10 +49,6 @@ namespace UE::IoStore
 
 int64 ParseSizeParam(FStringView Value);
 
-////////////////////////////////////////////////////////////////////////////////
-
-bool TryParseConfigFile(const FString& ConfigPath, FOnDemandEndpoint& OutEndpoint);
-
 ////////////////////////////////////////////////////////////////////////////////
 enum class EOnDemandTocVersion : uint32
 {
