Geometry collection : fix issue were clusters using remove on break with cluster crumbling could make their hierarhy to disappear for a short time , causing rendering glitches
- this is a knock on bug from CL 30290839 that was aiming at particles with no simulated children
- the fix adds a check to makes sure we only reconcile decays state and decay value when the particle has no children ( non-clustered )
- the fix is also behind a cvar ( on by default ) p.Chaos.GC.RemoveOnBreakDecayFix

diff --git a/Engine/Source/Runtime/Experimental/GeometryCollectionEngine/Private/GeometryCollection/GeometryCollectionComponent.cpp b/Engine/Source/Runtime/Experimental/GeometryCollectionEngine/Private/GeometryCollection/GeometryCollectionComponent.cpp
--- a/Engine/Source/Runtime/Experimental/GeometryCollectionEngine/Private/GeometryCollection/GeometryCollectionComponent.cpp
+++ b/Engine/Source/Runtime/Experimental/GeometryCollectionEngine/Private/GeometryCollection/GeometryCollectionComponent.cpp
@@ -144,6 +144,11 @@ FAutoConsoleVariableRef CVarGeometryCollectionNetAwakeningMode(TEXT("p.Chaos.GC.
 bool bGeometryCollectionCustomRendererHiddenActorFix = true;
 FAutoConsoleVariableRef CVarGeometryCollectionCustomRendererHiddenActorFix(TEXT("p.Chaos.GC.CustomRendererHiddenActorFix"), bGeometryCollectionCustomRendererHiddenActorFix, TEXT("When true custom renderer will account for the actor hidden flag"));
 
+// this cvar enables by default the fix for the render issue happening when a cluster with children breaks from a remove on break cluster crumbling event 
+// this should remain true unless an issue is found with it ( note this was introduce right after 5.5 )
+bool bGeometryCollectionRemoveOnBreakDecayFix = true;
+FAutoConsoleVariableRef CVarGeometryCollectionRemoveOnBreakDecayFix(TEXT("p.Chaos.GC.RemoveOnBreakDecayFix"), bGeometryCollectionRemoveOnBreakDecayFix, TEXT("When true, fix the issue causing a remove on break cluster crumbling with simulating children from disappearing for several frames"));
+
 DEFINE_LOG_CATEGORY_STATIC(UGCC_LOG, Error, All);
 DEFINE_LOG_CATEGORY_STATIC(LogGeometryCollectionComponent, Warning, All);
 
@@ -4436,8 +4441,28 @@ void UGeometryCollectionComponent::UpdateRemovalIfNeeded()
 					{
 						if (BrokenAndDecayedStates.GetHasDecayed(TransformIndex))
 						{
-							Decay = 1.0f;
-							DecayFacade.SetDecay(TransformIndex, Decay);
+							// this logic should only apply to end nodes that do not have simulated children 
+							// otherwise this may cause clusters with children to disappear for several frames
+							// ( see Cl 30290839 and UpdateBrokenAndDecayedStates() method logic )
+							bool bHasAnyChildWithValidParticle = false;
+							if (bGeometryCollectionRemoveOnBreakDecayFix)
+							{
+								DynamicCollection->IterateThroughChildren(TransformIndex,
+									[&bHasAnyChildWithValidParticle, this](int32 ChildTransformIdx) -> bool
+									{
+										if (PhysicsProxy->GetParticleByIndex_External(ChildTransformIdx))
+										{
+											bHasAnyChildWithValidParticle = true;
+											return false; // do no continue iterating
+										}
+										return true; // continue iterating
+									});
+							}
+							if (!bHasAnyChildWithValidParticle)
+							{
+								Decay = 1.0f;
+								DecayFacade.SetDecay(TransformIndex, Decay);
+							}
 						}
 					}
 				}
