null check archetype, there is a single object that has no archetype, and that is the CDO of UObject.

#rnx

diff --git a/Engine/Source/Runtime/Experimental/JsonObjectGraph/Private/JsonStringifyImpl.cpp b/Engine/Source/Runtime/Experimental/JsonObjectGraph/Private/JsonStringifyImpl.cpp
--- a/Engine/Source/Runtime/Experimental/JsonObjectGraph/Private/JsonStringifyImpl.cpp
+++ b/Engine/Source/Runtime/Experimental/JsonObjectGraph/Private/JsonStringifyImpl.cpp
@@ -397,7 +397,7 @@ void FJsonStringifyImpl::WriteObjectToJson(const UObject* Object)
 
 	// loop properties, writing any that have changed:
 	const UObject* const Archetype = Object->GetArchetype();
-	const UClass* const ArchetypeClass = Archetype->GetClass();
+	const UClass* const ArchetypeClass = Archetype ? Archetype->GetClass() : nullptr; // UObject's CDO itself has no archetype
 	Writer->WriteObjectStartInline();
 	// Write native UObject data - name, type, flags, native user serialize, etc
 	WriteNativeObjectData();
@@ -405,7 +405,7 @@ void FJsonStringifyImpl::WriteObjectToJson(const UObject* Object)
 	for (TFieldIterator<FProperty> FieldIt(Object->GetClass()); FieldIt; ++FieldIt)
 	{
 		FProperty* Property = *FieldIt;
-		const UObject* ValidatedArchtype = ArchetypeClass->IsChildOf(Property->GetOwnerClass()) ?
+		const UObject* ValidatedArchtype = ArchetypeClass && ArchetypeClass->IsChildOf(Property->GetOwnerClass()) ?
 			Archetype : nullptr;
 
 		WriteIdentifierAndValueToJson(Object, ValidatedArchtype, Property);
