Defensive code to not crash against StaticMesh nullptr.
Still need to understand how that can happen since the StaticMesh is already guaranteed non null lower in calstack at InitMergedMeshFromGeometryCollection(). (But it is then stored an retrieved from a weak object pointer.)

diff --git a/Engine/Source/Runtime/Experimental/ISMPool/Private/ISMPool/ISMPoolComponent.cpp b/Engine/Source/Runtime/Experimental/ISMPool/Private/ISMPool/ISMPoolComponent.cpp
--- a/Engine/Source/Runtime/Experimental/ISMPool/Private/ISMPool/ISMPoolComponent.cpp
+++ b/Engine/Source/Runtime/Experimental/ISMPool/Private/ISMPool/ISMPoolComponent.cpp
@@ -136,10 +136,11 @@ void FISMPoolISM::InitISM(const FISMPoolStaticMeshInstance& InMeshInstance, bool
 
 	UStaticMesh* StaticMesh = MeshInstance.StaticMesh.Get();
 	// We should only get here for valid static mesh objects.
-	check(StaticMesh != nullptr);
+	ensureMsgf(StaticMesh != nullptr, TEXT("StaticMesh is not valid."));
 
 #if WITH_EDITOR
-	const FName ISMName = MakeUniqueObjectName(ISMComponent->GetOwner(), UInstancedStaticMeshComponent::StaticClass(), StaticMesh->GetFName());
+	const FName MeshName = StaticMesh != nullptr ? StaticMesh->GetFName() : NAME_None;
+	const FName ISMName = MakeUniqueObjectName(ISMComponent->GetOwner(), UInstancedStaticMeshComponent::StaticClass(), MeshName);
 	const FString ISMNameString = ISMName.ToString();
 	ISMComponent->Rename(*ISMNameString);
 #endif
