ChaosFlesh: fixed crash in AppendSkeletalMesh when SKM does not have a Skeleton connected

diff --git a/Engine/Source/Runtime/Experimental/GeometryCollectionEngine/Private/GeometryCollection/GeometryCollectionEngineConversion.cpp b/Engine/Source/Runtime/Experimental/GeometryCollectionEngine/Private/GeometryCollection/GeometryCollectionEngineConversion.cpp
--- a/Engine/Source/Runtime/Experimental/GeometryCollectionEngine/Private/GeometryCollection/GeometryCollectionEngineConversion.cpp
+++ b/Engine/Source/Runtime/Experimental/GeometryCollectionEngine/Private/GeometryCollection/GeometryCollectionEngineConversion.cpp
@@ -1142,7 +1142,7 @@ bool FGeometryCollectionEngineConversion::AppendSkeletalMesh(const USkeletalMesh
 	//UE_LOG(UGeometryCollectionConversionLogging, Log, TEXT("FGeometryCollectionEngineConversion::AppendSkeletalMesh()"));
 #if WITH_EDITOR
 	int LODIndex = 0;
-	if (!InManagedArrayCollection)
+	if (!InManagedArrayCollection || !InSkeletalMesh)
 	{
 		return false;
 	}
@@ -1163,11 +1163,11 @@ bool FGeometryCollectionEngineConversion::AppendSkeletalMesh(const USkeletalMesh
 	int32 RootIndex = INDEX_NONE;
 	int32 TransformBaseIndex = INDEX_NONE;
 	const USkeleton* Skeleton = InSkeletalMesh->GetSkeleton();
-	const TArray<FTransform>& RestArray = Skeleton->GetRefLocalPoses();
 	const FReferenceSkeleton& ReferenceSkeleton = InSkeletalMesh->GetRefSkeleton();
 
-	if (ReferenceSkeleton.GetNum())
+	if (Skeleton && ReferenceSkeleton.GetNum())
 	{
+		const TArray<FTransform>& RestArray = ReferenceSkeleton.GetRefBonePose();
 		TransformBaseIndex = GeometryCollection.AddElements(ReferenceSkeleton.GetNum(), FGeometryCollection::TransformGroup);
 		RootIndex = TransformBaseIndex;
 
