[ChaosVD] CVD Extension System (Initial Implementation - Part 1)

diff --git a/Engine/Source/Runtime/Experimental/ChaosVDData/Public/ChaosVDRecording.h b/Engine/Source/Runtime/Experimental/ChaosVDData/Public/ChaosVDRecording.h
--- a/Engine/Source/Runtime/Experimental/ChaosVDData/Public/ChaosVDRecording.h
+++ b/Engine/Source/Runtime/Experimental/ChaosVDData/Public/ChaosVDRecording.h
@@ -7,6 +7,7 @@
 #include "DataWrappers/ChaosVDCollisionDataWrappers.h"
 #include "DataWrappers/ChaosVDParticleDataWrapper.h"
 #include "UObject/ObjectMacros.h"
+#include "UObject/StructOnScope.h"
 #include "Containers/UnrealString.h"
 #include "Chaos/ImplicitFwd.h"
 #include "Chaos/ImplicitObject.h"
@@ -38,7 +39,8 @@ enum class EChaosVDSolverStageFlags : uint8
 
 ENUM_CLASS_FLAGS(EChaosVDSolverStageFlags)
 
-struct FChaosVDStepData
+
+struct FChaosVDFrameStageData
 {
 	FString StepName;
 	TArray<TSharedPtr<FChaosVDParticleDataWrapper>> RecordedParticlesData;
@@ -52,6 +54,9 @@ struct FChaosVDStepData
 	EChaosVDSolverStageFlags StageFlags = EChaosVDSolverStageFlags::None;
 };
 
+UE_DEPRECATED(5.6, "FChaosVDStepData is deprecated and will be removed. Please use FChaosVDFrameStageData instead.")
+typedef FChaosVDFrameStageData FChaosVDStepData;
+
 struct FChaosVDTrackedLocation
 {
 	FString DebugName;
@@ -75,7 +80,105 @@ enum class EChaosVDNetworkSyncDataRequirements
 
 ENUM_CLASS_FLAGS(EChaosVDNetworkSyncDataRequirements)
 
-typedef TArray<FChaosVDStepData, TInlineAllocator<16>> FChaosVDStepsContainer;
+/** Handle to user-defined data types in a CVD recorded frame.
+ * These handles share ownership of the data
+ */
+struct CHAOSVDDATA_API FChaosVDCustomUserDataHandle
+{
+	/** Creates a handle for the provided data instance. This handle will hold a reference to the data
+	 * @param DataPtr Shared ref pointer to the data we want to rpovide access to
+	 */
+	template <typename DataStructType>
+	static FChaosVDCustomUserDataHandle MakeHandle(const TSharedRef<DataStructType>& DataPtr);
+
+	/** Returns the FName of the UStruct type that this handle represents */
+	FName GetTypeName() const;
+
+	/** Returns a raw ptr to the data this handle provides access to.*/
+	template <typename DataStructType>
+	DataStructType* GetData() const;
+
+	/** Returns a shared ptr to the data this handle provides access to */
+	template <typename DataStructType>
+	TSharedPtr<DataStructType> GetDataAsShared() const;
+private:
+
+	/** Checks if this handle is of another type, using the UStruct data we parsed on creation */
+	template <typename DataStructType>
+	bool IsA_Internal(const TSharedPtr<FStructOnScope>& InStructOnScope) const;
+	
+	TSharedPtr<FStructOnScope> DataStruct;
+	TSharedPtr<void> DataSharedPtr;
+};
+
+template <typename DataStructType>
+TSharedPtr<DataStructType> FChaosVDCustomUserDataHandle::GetDataAsShared() const
+{
+	if (IsA_Internal<DataStructType>(DataStruct))
+	{
+		return StaticCastSharedPtr<DataStructType>(DataSharedPtr);
+	}
+	return nullptr;
+}
+
+template <typename DataStructType>
+FChaosVDCustomUserDataHandle FChaosVDCustomUserDataHandle::MakeHandle(const TSharedRef<DataStructType>& DataPtr)
+{
+	FChaosVDCustomUserDataHandle NewCustomDataHandle;
+	
+	NewCustomDataHandle.DataStruct = MakeShared<FStructOnScope>(DataStructType::StaticStruct(), reinterpret_cast<uint8*>(&DataPtr.Get()));
+	NewCustomDataHandle.DataSharedPtr = DataPtr;
+
+	return NewCustomDataHandle;
+}
+
+template <typename DataStructType>
+DataStructType* FChaosVDCustomUserDataHandle::GetData() const
+{
+	if (IsA_Internal<DataStructType>(DataStruct))
+	{
+		return reinterpret_cast<DataStructType*>(DataStruct->GetStructMemory());
+	}
+
+	return nullptr;
+}
+
+template <typename DataStructType>
+bool FChaosVDCustomUserDataHandle::IsA_Internal(const TSharedPtr<FStructOnScope>& InStructOnScope) const
+{
+	if (InStructOnScope)
+	{
+		const UStruct* HandleStruct = InStructOnScope->GetStruct();
+		return HandleStruct && (DataStructType::StaticStruct() == HandleStruct || HandleStruct->IsChildOf(DataStructType::StaticStruct()));
+	}
+
+	return false;
+}
+
+struct CHAOSVDDATA_API FChaosVDCustomFrameData
+{
+	void AddData(const FChaosVDCustomUserDataHandle& InData);
+
+	/** Returns the shared ptr to a custom data instance, if such data type was added to this frame.
+	 * If the data type was not loaded, the ptr will be null
+	 */
+	template<typename DataType>
+	TSharedPtr<DataType> GetData() const;
+
+	/** Returns the shared ptr to a custom data instance, if such data type was added to this frame.
+	 * If the data type was not loaded, a new instance will be created and automatically added to the frame's data.
+	 */
+	template<typename DataType>
+	TSharedPtr<DataType> GetOrAddDefaultData();
+
+private:
+	TMap<FName, FChaosVDCustomUserDataHandle> CustomDataHandlesByType;
+};
+
+UE_DEPRECATED(5.6, "This FChaosVDStepsContainer is deprecated and it will be removed in the future. Use FChaosVDStagesContainer instead.")
+typedef TArray<FChaosVDFrameStageData, TInlineAllocator<16>> FChaosVDStepsContainer;
+
+typedef TArray<FChaosVDFrameStageData, TInlineAllocator<16>> FChaosVDFrameStagesContainer;
 
 struct CHAOSVDDATA_API FChaosVDSolverFrameData
 {
@@ -98,7 +201,7 @@ PRAGMA_ENABLE_DEPRECATION_WARNINGS
 	Chaos::FRigidTransform3 SimulationTransform;
 	bool bIsKeyFrame = false;
 	bool bIsResimulated = false;
-	FChaosVDStepsContainer SolverSteps;
+	FChaosVDFrameStagesContainer SolverSteps;
 	TSet<int32> ParticlesDestroyedIDs;
 	double StartTime = -1.0;
 	double EndTime = -1.0;
@@ -140,6 +243,19 @@ PRAGMA_ENABLE_DEPRECATION_WARNINGS
 	{
 		return NetworkTickOffset >= 0 ? NetworkTickOffset : 0;
 	}
+
+	FChaosVDCustomFrameData& GetCustomData()
+	{
+		return CustomData;
+	}
+	
+	const FChaosVDCustomFrameData& GetCustomData() const
+	{
+		return CustomData;
+	}
+
+private:
+	FChaosVDCustomFrameData CustomData;
 };
 
 struct FChaosVDGameFrameData
@@ -171,9 +287,15 @@ PRAGMA_ENABLE_DEPRECATION_WARNINGS
 		return EndTime - StartTime;
 	}
 
-	bool IsDirty() const { return bIsDirty; };
+	bool IsDirty() const
+	{
+		return bIsDirty;
+	}
 
-	void MarkDirty() { bIsDirty = true; }
+	void MarkDirty()
+	{
+		bIsDirty = true;
+	}
 
 	TMap<FName, FChaosVDTrackedLocation> RecordedNonSolverLocationsByID;
 	TMap<FName, FChaosVDTrackedTransform> RecordedNonSolverTransformsByID;
@@ -183,14 +305,31 @@ PRAGMA_ENABLE_DEPRECATION_WARNINGS
 	TMap<int32, TSharedPtr<FChaosVDQueryDataWrapper>> RecordedSceneQueries;
 
 	TMap<int32, TSharedPtr<FChaosVDQueryDataWrapper>> RecordedSceneQueriesByQueryID;
+
+	UE_DEPRECATED(5.6, "RecordedDebugDrawBoxesBySolverID is deprecated and will be removed in a future release, use GetCustomDataContainer instead.")
 	TMap<int32, TArray<TSharedPtr<FChaosVDAABBTreeDataWrapper>>> RecordedAABBTreesBySolverID;
 
+	UE_DEPRECATED(5.6, "RecordedDebugDrawBoxesBySolverID is deprecated and will be removed in a future release, use GetCustomDataContainer instead.")
 	TMap<int32,TArray<TSharedPtr<FChaosVDDebugDrawBoxDataWrapper>>> RecordedDebugDrawBoxesBySolverID;
+	UE_DEPRECATED(5.6, "RecordedDebugDrawLinesBySolverID is deprecated and will be removed in a future release, use GetCustomDataContainer instead.")
 	TMap<int32,TArray<TSharedPtr<FChaosVDDebugDrawLineDataWrapper>>> RecordedDebugDrawLinesBySolverID;
+	UE_DEPRECATED(5.6, "RecordedDebugDrawSpheresBySolverID is deprecated and will be removed in a future release, use GetCustomDataContainer instead.")
 	TMap<int32,TArray<TSharedPtr<FChaosVDDebugDrawSphereDataWrapper>>> RecordedDebugDrawSpheresBySolverID;
+	UE_DEPRECATED(5.6, "RecordedDebugDrawImplicitObjectsBySolverID is deprecated and will be removed in a future release, use GetCustomDataContainer instead.")
 	TMap<int32,TArray<TSharedPtr<FChaosVDDebugDrawImplicitObjectDataWrapper>>> RecordedDebugDrawImplicitObjectsBySolverID;
 
+	FChaosVDCustomFrameData& GetCustomDataHandler()
+	{
+		return CustomData;
+	}
+	
+	const FChaosVDCustomFrameData& GetCustomDataHandler() const
+	{
+		return CustomData;
+	}
+
 private:
+	FChaosVDCustomFrameData CustomData;
 	bool bIsDirty = false;
 };
 
@@ -207,17 +346,26 @@ struct CHAOSVDDATA_API FChaosVDRecording
 	static constexpr int32 CommonTrackCount = 16;
 
 	/** Returns the current available recorded solvers number */
-	int32 GetAvailableSolversNumber_AssumesLocked() const { return RecordedFramesDataPerSolver.Num(); }
+	int32 GetAvailableSolversNumber_AssumesLocked() const
+	{
+		return RecordedFramesDataPerSolver.Num();
+	}
 	
 	/** Returns the current available Game Frames */
 	int32 GetAvailableGameFramesNumber() const;
 	int32 GetAvailableGameFramesNumber_AssumesLocked() const;
 
 	/** Returns a reference to the array holding all the available game frames */
-	const TArray<FChaosVDGameFrameData>& GetAvailableGameFrames_AssumesLocked() const { return GameFrames; }
+	const TArray<FChaosVDGameFrameData>& GetAvailableGameFrames_AssumesLocked() const
+	{
+		return GameFrames;
+	}
 
 	/** Returns a reference to the map containing the available solver data */
-	const TMap<int32, TArray<FChaosVDSolverFrameData>>& GetAvailableSolvers_AssumesLocked() const { return RecordedFramesDataPerSolver; }
+	const TMap<int32, TArray<FChaosVDSolverFrameData>>& GetAvailableSolvers_AssumesLocked() const
+	{
+		return RecordedFramesDataPerSolver;
+	}
 
 	/**
 	 * Returns the number of available frame data for the specified solver ID
@@ -232,8 +380,11 @@ struct CHAOSVDDATA_API FChaosVDRecording
 	 */
 	FName GetSolverFName(int32 SolverID);
 
-	UE_DEPRECATED(5.5, "Please use the GetSolverFName instead")
-	FString GetSolverName(int32 SolverID) { return TEXT(""); }
+	UE_DEPRECATED(5.6, "Please use the GetSolverFName instead")
+	FString GetSolverName(int32 SolverID)
+	{
+		return TEXT("");
+	}
 
 	/**
 	 * Returns the name of the specified solver id. Must be called from within a ReadLock
@@ -244,8 +395,11 @@ struct CHAOSVDDATA_API FChaosVDRecording
 	bool IsServerSolver_AssumesLocked(int32 SolverID);
 	bool IsServerSolver(int32 SolverID);
 
-	UE_DEPRECATED(5.5, "Please use the GetSolverFName_AssumedLocked instead")
-	FString GetSolverName_AssumedLocked(int32 SolverID) { return TEXT(""); }
+	UE_DEPRECATED(5.6, "Please use the GetSolverFName_AssumedLocked instead")
+	FString GetSolverName_AssumedLocked(int32 SolverID)
+	{
+		return TEXT("");
+	}
 
 	/**
 	 * Return a ptr to the existing solver frame data from the specified ID and Frame number
@@ -309,7 +463,10 @@ struct CHAOSVDDATA_API FChaosVDRecording
 	void AddGameFrameData(const FChaosVDGameFrameData& InFrameData);
 
 	/** Called each time new geometry data becomes available in the recording - Mainly when a new frame is added from the Trace analysis */
-	FChaosVDGeometryDataLoaded& OnGeometryDataLoaded() { return GeometryDataLoaded; };
+	FChaosVDGeometryDataLoaded& OnGeometryDataLoaded()
+	{
+		return GeometryDataLoaded;
+	}
 
 	/**
 	 * Searches for a recorded Game frame at the specified cycle 
@@ -360,7 +517,10 @@ struct CHAOSVDDATA_API FChaosVDRecording
 	void CollapseSolverFramesRange_AssumesLocked(int32 SolverID, int32 StartFrame, int32 EndFrame, FChaosVDSolverFrameData& OutCollapsedFrameData);
 
 	/** Returns a reference to the GeometryID-ImplicitObject map of this recording */
-	const TMap<uint32, Chaos::FConstImplicitObjectPtr>& GetGeometryMap() const { return ImplicitObjects; };
+	const TMap<uint32, Chaos::FConstImplicitObjectPtr>& GetGeometryMap() const
+	{
+		return ImplicitObjects;
+	}
 
 	UE_DEPRECATED(5.4, "Please use GetGeometryMap instead")
 	const TMap<uint32, TSharedPtr<const Chaos::FImplicitObject>>& GetGeometryDataMap() const
@@ -368,7 +528,7 @@ struct CHAOSVDDATA_API FChaosVDRecording
 		check(false);
 		static TMap<uint32, TSharedPtr<const Chaos::FImplicitObject>> DummyMap;
 		return DummyMap;
-	};
+	}
 
 	/** Adds a shared Implicit Object to the recording */
 	void AddImplicitObject(const uint32 ID, const Chaos::FImplicitObjectPtr& InImplicitObject);
@@ -379,30 +539,54 @@ struct CHAOSVDDATA_API FChaosVDRecording
 	/** Session name of the trace session used to re-build this recording */
 	FString SessionName;
 
-	FRWLock& GetRecordingDataLock() { return RecordingDataLock; }
+	FRWLock& GetRecordingDataLock()
+	{
+		return RecordingDataLock;
+	}
 
 	/** Returns true if this recording is being populated from a live session */
-	bool IsLive() const { return bIsLive; }
+	bool IsLive() const
+	{
+		return bIsLive;
+	}
 
 	/** Sets if this recording is being populated from a live session */
-	void SetIsLive(bool bNewIsLive) { bIsLive = bNewIsLive; }
+	void SetIsLive(bool bNewIsLive)
+	{
+		bIsLive = bNewIsLive;
+	}
 
 	/** Returns the name table instances used to de-duplicate strings serialization */
-	TSharedPtr<Chaos::VisualDebugger::FChaosVDSerializableNameTable> GetNameTableInstance() const { return NameTable; }
+	TSharedPtr<Chaos::VisualDebugger::FChaosVDSerializableNameTable> GetNameTableInstance() const
+	{
+		return NameTable;
+	}
 
 	/** Returns the FArchive header used to read the serialized binary data */
-	const Chaos::VisualDebugger::FChaosVDArchiveHeader& GetHeaderData() const { return HeaderData; }
+	const Chaos::VisualDebugger::FChaosVDArchiveHeader& GetHeaderData() const
+	{
+		return HeaderData;
+	}
 	
 	/** Sets the FArchive header used to read the serialized binary data */
-	void SetHeaderData(const Chaos::VisualDebugger::FChaosVDArchiveHeader& InNewHeader) { HeaderData = InNewHeader; }
+	void SetHeaderData(const Chaos::VisualDebugger::FChaosVDArchiveHeader& InNewHeader)
+	{
+		HeaderData = InNewHeader;
+	}
 
 	/** Returns true if this recording does not have any usable data */
 	bool IsEmpty() const;
 
 	/** Returns the last Platform Cycle on which this recording was updated (A new frame was added) */
-	uint64 GetLastUpdatedTimeAsCycle() { return LastUpdatedTimeAsCycle; }
+	uint64 GetLastUpdatedTimeAsCycle()
+	{
+		return LastUpdatedTimeAsCycle;
+	}
 
-	TSharedPtr<FChaosVDCollisionChannelsInfoContainer> GetCollisionChannelsInfoContainer(){ return CollisionChannelsInfoContainer; }
+	TSharedPtr<FChaosVDCollisionChannelsInfoContainer> GetCollisionChannelsInfoContainer()
+	{
+		return CollisionChannelsInfoContainer;
+	}
 	
 	void SetCollisionChannelsInfoContainer(const TSharedPtr<FChaosVDCollisionChannelsInfoContainer>& InCollisionChannelsInfo);
 
@@ -451,6 +635,35 @@ protected:
 	friend class FChaosVDTraceImplicitObjectProcessor;
 };
 
+template <typename DataType>
+TSharedPtr<DataType> FChaosVDCustomFrameData::GetData() const
+{
+	if (const FChaosVDCustomUserDataHandle* FoundDataHandle = CustomDataHandlesByType.Find(DataType::StaticStruct()->GetFName()))
+	{
+		return FoundDataHandle->GetDataAsShared<DataType>();
+	}
+
+	return nullptr;
+}
+
+template <typename DataType>
+TSharedPtr<DataType> FChaosVDCustomFrameData::GetOrAddDefaultData()
+{
+	TSharedPtr<DataType> DataSharedPtr = GetData<DataType>();
+	if (!DataSharedPtr)
+	{
+		TSharedPtr<DataType> CustomData = nullptr;
+
+		CustomData = MakeShared<DataType>();
+
+		AddData(FChaosVDCustomUserDataHandle::MakeHandle<DataType>(CustomData.ToSharedRef()));
+
+		return CustomData;
+	}
+
+	return DataSharedPtr;
+}
+
 template <typename TAllocator>
 void FChaosVDRecording::GetAvailableSolverIDsAtGameFrameNumber(int32 FrameNumber, TArray<int32, TAllocator>& OutSolversID)
 {
diff --git a/Engine/Source/Runtime/Experimental/ChaosVisualDebugger/Public/DataWrappers/ChaosVDAccelerationStructureDataWrappers.h b/Engine/Source/Runtime/Experimental/ChaosVisualDebugger/Public/DataWrappers/ChaosVDAccelerationStructureDataWrappers.h
--- a/Engine/Source/Runtime/Experimental/ChaosVisualDebugger/Public/DataWrappers/ChaosVDAccelerationStructureDataWrappers.h
+++ b/Engine/Source/Runtime/Experimental/ChaosVisualDebugger/Public/DataWrappers/ChaosVDAccelerationStructureDataWrappers.h
@@ -147,6 +147,14 @@ struct FChaosVDAABBTreeLeafDataWrapper : public FChaosVDWrapperDataBase
 CVD_IMPLEMENT_SERIALIZER(FChaosVDAABBTreeLeafDataWrapper)
 #pragma warning(default : 4996)
 
+USTRUCT()
+struct FChaosVDAccelerationStructureContainer
+{
+	GENERATED_BODY()
+
+	TMap<int32, TArray<TSharedPtr<FChaosVDAABBTreeDataWrapper>>> RecordedAABBTreesBySolverID;
+};
+
 USTRUCT(DisplayName="AABB Tree Data")
 struct FChaosVDAABBTreeDataWrapper : public FChaosVDAccelerationStructureBase
 {
diff --git a/Engine/Source/Runtime/Experimental/ChaosVisualDebugger/Public/DataWrappers/ChaosVDDebugShapeDataWrapper.h b/Engine/Source/Runtime/Experimental/ChaosVisualDebugger/Public/DataWrappers/ChaosVDDebugShapeDataWrapper.h
--- a/Engine/Source/Runtime/Experimental/ChaosVisualDebugger/Public/DataWrappers/ChaosVDDebugShapeDataWrapper.h
+++ b/Engine/Source/Runtime/Experimental/ChaosVisualDebugger/Public/DataWrappers/ChaosVDDebugShapeDataWrapper.h
@@ -6,6 +6,17 @@
 
 #include "ChaosVDDebugShapeDataWrapper.generated.h"
 
+USTRUCT()
+struct FChaosVDDebugShapeDataContainer
+{
+	GENERATED_BODY()
+
+	TMap<int32,TArray<TSharedPtr<FChaosVDDebugDrawBoxDataWrapper>>> RecordedDebugDrawBoxesBySolverID;
+	TMap<int32,TArray<TSharedPtr<FChaosVDDebugDrawLineDataWrapper>>> RecordedDebugDrawLinesBySolverID;
+	TMap<int32,TArray<TSharedPtr<FChaosVDDebugDrawSphereDataWrapper>>> RecordedDebugDrawSpheresBySolverID;
+	TMap<int32,TArray<TSharedPtr<FChaosVDDebugDrawImplicitObjectDataWrapper>>> RecordedDebugDrawImplicitObjectsBySolverID;
+};
+
 USTRUCT()
 struct FChaosVDDebugDrawShapeBase : public FChaosVDWrapperDataBase
 {
diff --git a/Engine/Source/Runtime/Experimental/ChaosVisualDebugger/Public/DataWrappers/ChaosVDParticleDataWrapper.h b/Engine/Source/Runtime/Experimental/ChaosVisualDebugger/Public/DataWrappers/ChaosVDParticleDataWrapper.h
--- a/Engine/Source/Runtime/Experimental/ChaosVisualDebugger/Public/DataWrappers/ChaosVDParticleDataWrapper.h
+++ b/Engine/Source/Runtime/Experimental/ChaosVisualDebugger/Public/DataWrappers/ChaosVDParticleDataWrapper.h
@@ -42,7 +42,7 @@ enum class EChaosVDObjectStateType: int8
 
 /** Base struct that declares the interface to be used for any ParticleData Viewer */
 USTRUCT()
-struct FChaosVDWrapperDataBase
+struct CHAOSVDRUNTIME_API FChaosVDWrapperDataBase
 {
 	GENERATED_BODY()
 	virtual ~FChaosVDWrapperDataBase() = default;
