Make EventManager use TransactionallySafe RWLocks
Put a temp hack in TryWriteLock, needs to be implemented but return false for now

diff --git a/Engine/Source/Runtime/Core/Public/Misc/TransactionallySafeRWLock.h b/Engine/Source/Runtime/Core/Public/Misc/TransactionallySafeRWLock.h
--- a/Engine/Source/Runtime/Core/Public/Misc/TransactionallySafeRWLock.h
+++ b/Engine/Source/Runtime/Core/Public/Misc/TransactionallySafeRWLock.h
@@ -169,6 +169,11 @@ struct FTransactionallySafeRWLockDefinition
 		}
 	}
 
+	bool TryWriteLock()
+	{
+		return false;
+	}
+
 private:
 	UE_NONCOPYABLE(FTransactionallySafeRWLockDefinition)
 
diff --git a/Engine/Source/Runtime/Experimental/Chaos/Public/EventManager.h b/Engine/Source/Runtime/Experimental/Chaos/Public/EventManager.h
--- a/Engine/Source/Runtime/Experimental/Chaos/Public/EventManager.h
+++ b/Engine/Source/Runtime/Experimental/Chaos/Public/EventManager.h
@@ -8,6 +8,7 @@
 #include "Chaos/PBDRigidsEvolutionFwd.h"
 #include "Misc/ScopeLock.h"
 #include "Chaos/Defines.h"
+#include "Misc/TransactionallySafeRWLock.h"
 
 namespace Chaos
 {
@@ -396,11 +397,10 @@ private:
 		 */
 		TMap<void*, FEventHandlerPtr> HandlerMap;
 
-		FRWLock HandlerLock; // protect access ProxyOwnerToHandlerMap, HandlersNotInProxyOwnerMap, HandlerMap
+		FTransactionallySafeRWLock HandlerLock; // protect access ProxyOwnerToHandlerMap, HandlersNotInProxyOwnerMap, HandlerMap
 		TArray<FEventHandlerPtr> DeferredHandlers; // Store handler to register, they couldn't registered to avoid reentrant lock
 		TArray<const void*> DeferredUnregisterHandlers; // Store handler to unregister, they couldn't unregistered to avoid reentrant lock
-		FRWLock DeferredHandlerLock; // protect access to DeferredHandlers and DeferredUnregisterHandlers
-
+		FTransactionallySafeRWLock DeferredHandlerLock; // protect access to DeferredHandlers and DeferredUnregisterHandlers
 	};
 
 	/**
@@ -531,8 +531,8 @@ private:
 
 		Chaos::EMultiBufferMode BufferMode;			// specifies the buffer type to be constructed, single, double, triple
 		TArray<FEventContainerBasePtr> EventContainers;	// Array of event types
-		FRWLock ResourceLock;
-		FRWLock ContainerLock;
+		FTransactionallySafeRWLock ResourceLock;
+		FTransactionallySafeRWLock ContainerLock;
 	};
 
 }
