[Chaos] Disable Chaos Debug Draw rendering in preview worlds.

The old debug draw queue was excluding Preview worlds, but the new one doesn't. This cases that Chaos Debug draw, when enabled, it also renders in any asset editor you might have open.

This CL disables debug draw the creation of the debug draw subsystem for any world that is a Preview world.

In the future, we could look into exposing the debug raw controls on a per viewport bases, as a new Show menu category

diff --git a/Engine/Source/Runtime/Experimental/ChaosSolverEngine/Private/Chaos/ChaosDebugDrawSubsystem.cpp b/Engine/Source/Runtime/Experimental/ChaosSolverEngine/Private/Chaos/ChaosDebugDrawSubsystem.cpp
--- a/Engine/Source/Runtime/Experimental/ChaosSolverEngine/Private/Chaos/ChaosDebugDrawSubsystem.cpp
+++ b/Engine/Source/Runtime/Experimental/ChaosSolverEngine/Private/Chaos/ChaosDebugDrawSubsystem.cpp
@@ -260,7 +260,7 @@ bool UChaosDebugDrawSubsystem::ShouldCreateSubsystem(UObject* Outer) const
 #if CHAOS_DEBUG_DRAW
 	if (UWorld* World = Cast<UWorld>(Outer))
 	{
-		const bool bCreateDebugDraw = !IsRunningCommandlet();
+		const bool bCreateDebugDraw = !IsRunningCommandlet() && DoesSupportWorldType(World->WorldType);
 
 		UE_LOG(LogChaosDD, Log, TEXT("%s Chaos Debug Draw Scene for world %s"), bCreateDebugDraw ? TEXT("Creating") : TEXT("Not creating"), *World->GetName());
 		
