[NetPhysics] Fix property initializing

diff --git a/Engine/Source/Runtime/Experimental/ChaosVehicles/ChaosVehiclesCore/Public/SimModule/ModuleInput.h b/Engine/Source/Runtime/Experimental/ChaosVehicles/ChaosVehiclesCore/Public/SimModule/ModuleInput.h
--- a/Engine/Source/Runtime/Experimental/ChaosVehicles/ChaosVehiclesCore/Public/SimModule/ModuleInput.h
+++ b/Engine/Source/Runtime/Experimental/ChaosVehicles/ChaosVehiclesCore/Public/SimModule/ModuleInput.h
@@ -347,7 +347,7 @@ protected:
 	EModuleInputValueType ValueType = EModuleInputValueType::MBoolean;
 
 	UPROPERTY()
-	bool bApplyInputDecay;
+	bool bApplyInputDecay = false;
 };
 
 // Supported getter specializations
@@ -461,11 +461,13 @@ struct CHAOSVEHICLESCORE_API FModuleInputSetup
 	FModuleInputSetup()
 	{
 		Type = EModuleInputValueType::MBoolean;
+		bApplyInputDecay = false;
 	}
 
 	FModuleInputSetup(const FName& InName, const EModuleInputValueType& InType)
 		: Name(InName)
 		, Type(InType)
+		, bApplyInputDecay(false)
 	{
 	}
 
