Iris - AlwaysRelevantNetObjectFilter. When you want to be more explicit about something being always relevant rather than configuring the class to a None filter.

diff --git a/Engine/Source/Runtime/Experimental/Iris/Core/Private/Iris/ReplicationSystem/Filtering/AlwaysRelevantNetObjectFilter.cpp b/Engine/Source/Runtime/Experimental/Iris/Core/Private/Iris/ReplicationSystem/Filtering/AlwaysRelevantNetObjectFilter.cpp
new file mode 100644
--- /dev/null
+++ b/Engine/Source/Runtime/Experimental/Iris/Core/Private/Iris/ReplicationSystem/Filtering/AlwaysRelevantNetObjectFilter.cpp
@@ -0,0 +1,35 @@
+// Copyright Epic Games, Inc. All Rights Reserved.
+
+#include "Iris/ReplicationSystem/Filtering/AlwaysRelevantNetObjectFilter.h"
+
+void UAlwaysRelevantNetObjectFilter::OnInit(const FNetObjectFilterInitParams& Params)
+{
+}
+
+void UAlwaysRelevantNetObjectFilter::OnDeinit()
+{
+}
+
+bool UAlwaysRelevantNetObjectFilter::AddObject(uint32 ObjectIndex, FNetObjectFilterAddObjectParams& Params)
+{
+	return true;
+}
+
+void UAlwaysRelevantNetObjectFilter::RemoveObject(uint32 ObjectIndex, const FNetObjectFilteringInfo& Info)
+{
+}
+
+void UAlwaysRelevantNetObjectFilter::UpdateObjects(FNetObjectFilterUpdateParams& Params)
+{
+}
+
+void UAlwaysRelevantNetObjectFilter::Filter(FNetObjectFilteringParams& Params)
+{
+	// Allow everything
+	Params.OutAllowedObjects.SetAllBits();
+}
+
+void UAlwaysRelevantNetObjectFilter::OnMaxInternalNetRefIndexIncreased(uint32 NewMaxInternalIndex)
+{
+}
+
diff --git a/Engine/Source/Runtime/Experimental/Iris/Core/Public/Iris/ReplicationSystem/Filtering/AlwaysRelevantNetObjectFilter.h b/Engine/Source/Runtime/Experimental/Iris/Core/Public/Iris/ReplicationSystem/Filtering/AlwaysRelevantNetObjectFilter.h
new file mode 100644
--- /dev/null
+++ b/Engine/Source/Runtime/Experimental/Iris/Core/Public/Iris/ReplicationSystem/Filtering/AlwaysRelevantNetObjectFilter.h
@@ -0,0 +1,28 @@
+// Copyright Epic Games, Inc. All Rights Reserved.
+
+#pragma once
+
+#include "Iris/ReplicationSystem/Filtering/NetObjectFilter.h"
+#include "AlwaysRelevantNetObjectFilter.generated.h"
+
+UCLASS(transient, MinimalAPI)
+class UAlwaysReleantNetObjectFilterConfig final : public UNetObjectFilterConfig
+{
+	GENERATED_BODY()
+};
+
+UCLASS()
+class UAlwaysRelevantNetObjectFilter final : public UNetObjectFilter
+{
+	GENERATED_BODY()
+
+protected:
+	// UNetObjectFilter interface
+	IRISCORE_API virtual void OnInit(const FNetObjectFilterInitParams&) override;
+	IRISCORE_API virtual void OnDeinit() override;
+	IRISCORE_API virtual bool AddObject(uint32 ObjectIndex, FNetObjectFilterAddObjectParams&) override;
+	IRISCORE_API virtual void RemoveObject(uint32 ObjectIndex, const FNetObjectFilteringInfo&) override;
+	IRISCORE_API virtual void UpdateObjects(FNetObjectFilterUpdateParams&) override;
+	IRISCORE_API virtual void Filter(FNetObjectFilteringParams&) override;
+	IRISCORE_API virtual void OnMaxInternalNetRefIndexIncreased(uint32 NewMaxInternalIndex) override;
+};
