Iris - Log error when ArrayPropertyNetSerializer detects invalid array sizes.

#rnx

diff --git a/Engine/Source/Runtime/Experimental/Iris/Core/Private/Iris/Serialization/ArrayPropertyNetSerializer.cpp b/Engine/Source/Runtime/Experimental/Iris/Core/Private/Iris/Serialization/ArrayPropertyNetSerializer.cpp
--- a/Engine/Source/Runtime/Experimental/Iris/Core/Private/Iris/Serialization/ArrayPropertyNetSerializer.cpp
+++ b/Engine/Source/Runtime/Experimental/Iris/Core/Private/Iris/Serialization/ArrayPropertyNetSerializer.cpp
@@ -1,12 +1,14 @@
 // Copyright Epic Games, Inc. All Rights Reserved.
 
 #include "InternalNetSerializers.h"
-#include "Iris/ReplicationState/ReplicationStateDescriptor.h"
+#include "Iris/Core/IrisLog.h"
 #include "Iris/Serialization/NetBitStreamReader.h"
 #include "Iris/Serialization/NetBitStreamWriter.h"
 #include "Iris/Serialization/InternalNetSerializationContext.h"
 #include "Iris/Serialization/NetReferenceCollector.h"
 #include "Iris/ReplicationSystem/ReplicationOperationsInternal.h"
+#include "Iris/ReplicationState/ReplicationStateDescriptor.h"
+#include "Logging/StructuredLog.h"
 #include "Net/Core/NetBitArray.h"
 #include "Net/Core/Trace/NetTrace.h"
 
@@ -211,6 +213,7 @@ void FArrayPropertyNetSerializer::Deserialize(FNetSerializationContext& Context,
 
 	if (NewElementCount > Config->MaxElementCount)
 	{
+		UE_LOGFMT(LogIris, Error, "FArrayPropertyNetSerializer::Deserialize Element count exceeds Config->MaxElementCount. {NewElementCount} > {MaxElementCount}. Config->ElementCountBitCount: {ElementCountBitCount}", NewElementCount, Config->MaxElementCount, Config->ElementCountBitCount);
 		Context.SetError(GNetError_ArraySizeTooLarge);
 		return;
 	}
@@ -378,6 +381,7 @@ void FArrayPropertyNetSerializer::DeserializeDelta(FNetSerializationContext& Con
 
 	if (ElementCount > Config->MaxElementCount)
 	{
+		UE_LOGFMT(LogIris, Error, "FArrayPropertyNetSerializer::DeserializeDelta Element count exceeds Config->MaxElementCount. {ElementCount} > {MaxElementCount}. Config->ElementCountBitCount: {ElementCountBitCount} bSameSizeArray: {bSameSizeArray}", ElementCount, Config->MaxElementCount, Config->ElementCountBitCount, bSameSizeArray);
 		Context.SetError(GNetError_ArraySizeTooLarge);
 		return;
 	}
@@ -472,6 +476,7 @@ void FArrayPropertyNetSerializer::Quantize(FNetSerializationContext& Context, co
 
 	if (ElementCount > Config->MaxElementCount)
 	{
+		UE_LOGFMT(LogIris, Error, "FArrayPropertyNetSerializer::Quantize Element count exceeds Config->MaxElementCount. {ElementCount} > {MaxElementCount}.", ElementCount, Config->MaxElementCount);
 		Context.SetError(GNetError_ArraySizeTooLarge);
 		return;
 	}
