Setting CCD enabled and MACD enabled on a Cluster union now working

diff --git a/Engine/Source/Runtime/Engine/Private/PhysicsEngine/ClusterUnionComponent.cpp b/Engine/Source/Runtime/Engine/Private/PhysicsEngine/ClusterUnionComponent.cpp
--- a/Engine/Source/Runtime/Engine/Private/PhysicsEngine/ClusterUnionComponent.cpp
+++ b/Engine/Source/Runtime/Engine/Private/PhysicsEngine/ClusterUnionComponent.cpp
@@ -715,6 +715,8 @@ void UClusterUnionComponent::OnCreatePhysicsState()
 	InitData.ComponentId = GetUniqueID();
 	InitData.InitialTransform = GetComponentTransform();
 	InitData.GravityGroupOverride = GravityGroupIndexOverride;
+	InitData.bCCDEnabled = BodyInstance.bUseCCD;
+	InitData.bMACDEnabled = BodyInstance.GetUseMACD();
 
 	// Client needs to be set to unbreakable so the server is authoritative.
 	InitData.bUnbreakable = !bHasAuthority;
diff --git a/Engine/Source/Runtime/Experimental/Chaos/Private/PhysicsProxy/ClusterUnionPhysicsProxy.cpp b/Engine/Source/Runtime/Experimental/Chaos/Private/PhysicsProxy/ClusterUnionPhysicsProxy.cpp
--- a/Engine/Source/Runtime/Experimental/Chaos/Private/PhysicsProxy/ClusterUnionPhysicsProxy.cpp
+++ b/Engine/Source/Runtime/Experimental/Chaos/Private/PhysicsProxy/ClusterUnionPhysicsProxy.cpp
@@ -92,6 +92,8 @@ namespace Chaos
 		Particle_External->SetUserData(InitData.UserData);
 		Particle_External->SetX(InitData.InitialTransform.GetTranslation());
 		Particle_External->SetR(InitData.InitialTransform.GetRotation());
+		Particle_External->SetCCDEnabled(InitData.bCCDEnabled);
+		Particle_External->SetMACDEnabled(InitData.bMACDEnabled);
 
 		// NO DIRTY FLAGS ALLOWED. We must strictly manage the dirty flags on the particle.
 		// Setting the particle's XR on the particle will set the XR dirty flag but that isn't
@@ -143,6 +145,8 @@ namespace Chaos
 			Particle_Internal->SetQf(Particle_Internal->GetRf());
 			Particle_Internal->SetCenterOfMass(FVector3f::ZeroVector);
 			Particle_Internal->SetRotationOfMass(FQuat::Identity);
+			Particle_Internal->SetCCDEnabled(InitData.bCCDEnabled);
+			Particle_Internal->SetMACDEnabled(InitData.bMACDEnabled);
 
 #if CHAOS_DEBUG_NAME
 			if (Particle_External.IsValid())
@@ -904,4 +908,4 @@ namespace Chaos
 			}
 		);
 	}
-}
\ No newline at end of file
+}
diff --git a/Engine/Source/Runtime/Experimental/Chaos/Public/PhysicsProxy/ClusterUnionPhysicsProxy.h b/Engine/Source/Runtime/Experimental/Chaos/Public/PhysicsProxy/ClusterUnionPhysicsProxy.h
--- a/Engine/Source/Runtime/Experimental/Chaos/Public/PhysicsProxy/ClusterUnionPhysicsProxy.h
+++ b/Engine/Source/Runtime/Experimental/Chaos/Public/PhysicsProxy/ClusterUnionPhysicsProxy.h
@@ -36,6 +36,8 @@ namespace Chaos
 		bool bUnbreakable = false;
 		bool bGenerateConnectivityEdges = true;
 		int32 GravityGroupOverride = INDEX_NONE;
+		bool bCCDEnabled = false;
+		bool bMACDEnabled = false;
 #if CHAOS_DEBUG_NAME
 		TSharedPtr<FString, ESPMode::ThreadSafe> DebugName;
 #endif
