Chaos::Fix the issue of particles not going into sleep from external SetObjectState(Sleeping) call in async mode.

diff --git a/Engine/Source/Runtime/Experimental/Chaos/Private/Chaos/PBDRigidsEvolutionGBF.cpp b/Engine/Source/Runtime/Experimental/Chaos/Private/Chaos/PBDRigidsEvolutionGBF.cpp
--- a/Engine/Source/Runtime/Experimental/Chaos/Private/Chaos/PBDRigidsEvolutionGBF.cpp
+++ b/Engine/Source/Runtime/Experimental/Chaos/Private/Chaos/PBDRigidsEvolutionGBF.cpp
@@ -1409,14 +1409,14 @@ void FPBDRigidsEvolutionGBF::DestroyTransientConstraints()
 	}
 }
 
-void FPBDRigidsEvolutionGBF::SetParticleTransform(FGeometryParticleHandle* InParticle, const FVec3& InPos, const FRotation3& InRot, const bool bIsTeleport)
+void FPBDRigidsEvolutionGBF::SetParticleTransform(FGeometryParticleHandle* InParticle, const FVec3& InPos, const FRotation3& InRot, const bool bIsTeleport, const bool bWakeUp)
 {
 	const FVec3 PrevX = InParticle->GetX();
 	const FRotation3 PrevR = InParticle->GetR();
 
 	FGenericParticleHandle(InParticle)->SetTransform(InPos, InRot);
 
-	OnParticleMoved(InParticle, PrevX, PrevR, bIsTeleport);
+	OnParticleMoved(InParticle, PrevX, PrevR, bIsTeleport, bWakeUp);
 
 	// This is a bit overkill, if this becomes an issue we will need to look into adding support
 	// to record partial updates of particles, only the transform in this case
@@ -1455,7 +1455,7 @@ void FPBDRigidsEvolutionGBF::SetParticleTransformSwept(FGeometryParticleHandle*
 	if (TOI > 0)
 	{
 		const FVec3 NewPos = FMath::Lerp(SweepStart, InPos, TOI);
-		SetParticleTransform(InParticle, NewPos, InRot, bIsTeleport);
+		SetParticleTransform(InParticle, NewPos, InRot, bIsTeleport, true);
 	}
 }
 
@@ -1486,7 +1486,7 @@ void FPBDRigidsEvolutionGBF::SetParticleKinematicTarget(FGeometryParticleHandle*
 	}
 }
 
-void FPBDRigidsEvolutionGBF::OnParticleMoved(FGeometryParticleHandle* InParticle, const FVec3& PrevX, const FRotation3& PrevR, const bool bIsTeleport)
+void FPBDRigidsEvolutionGBF::OnParticleMoved(FGeometryParticleHandle* InParticle, const FVec3& PrevX, const FRotation3& PrevR, const bool bIsTeleport, const bool bWakeUp)
 {
 	// When a particle is moved, we need to 
 	// - tell the collisions because they cache friction state and would attempt to undo small translations within the friction cone
@@ -1497,7 +1497,10 @@ void FPBDRigidsEvolutionGBF::OnParticleMoved(FGeometryParticleHandle* InParticle
 	const FReal CollisionRotationTolerance = FReal(1.e-6);
 	if (!FVec3::IsNearlyEqual(PrevX, InParticle->GetX(), CollisionPositionTolerance) || !FRotation3::IsNearlyEqual(PrevR, InParticle->GetR(), CollisionRotationTolerance))
 	{
-		GetIslandManager().WakeParticleIslands(InParticle);
+		if (bWakeUp)
+		{
+			GetIslandManager().WakeParticleIslands(InParticle);
+		}
 
 		InParticle->ParticleCollisions().VisitCollisions([this, InParticle](FPBDCollisionConstraint& Collision)
 			{
diff --git a/Engine/Source/Runtime/Experimental/Chaos/Private/PhysicsProxy/SingleParticlePhysicsProxy.cpp b/Engine/Source/Runtime/Experimental/Chaos/Private/PhysicsProxy/SingleParticlePhysicsProxy.cpp
--- a/Engine/Source/Runtime/Experimental/Chaos/Private/PhysicsProxy/SingleParticlePhysicsProxy.cpp
+++ b/Engine/Source/Runtime/Experimental/Chaos/Private/PhysicsProxy/SingleParticlePhysicsProxy.cpp
@@ -70,6 +70,7 @@ void PushToPhysicsStateImp(const Chaos::FDirtyPropertiesManager& Manager, Chaos:
 	{
 		auto NewXR = ParticleData.FindXR(Manager, DataIdx);
 		auto NewNonFrequentData = ParticleData.FindNonFrequentData(Manager, DataIdx);
+		auto NewDynamicMiscData = ParticleData.FindDynamicMisc(Manager, DataIdx);
 
 		if(NewXR)
 		{
@@ -77,7 +78,13 @@ void PushToPhysicsStateImp(const Chaos::FDirtyPropertiesManager& Manager, Chaos:
 			// For now we just set bIsTeleport to true since that's the no-impact option for SetParticleTransform
 			// (there would be issues if we report a non-teleport move for an initial-position a long way from the origin)
 			const bool bIsTeleport = true;
-			Evolution.SetParticleTransform(Handle, NewXR->X(), NewXR->R(), bIsTeleport);
+			bool bWakeUp = true;
+			if (NewDynamicMiscData && NewDynamicMiscData->ObjectState() == EObjectStateType::Sleeping)
+			{
+				bWakeUp = false;
+			}
+
+			Evolution.SetParticleTransform(Handle, NewXR->X(), NewXR->R(), bIsTeleport, bWakeUp);
 		}
 
 		if(NewNonFrequentData)
@@ -124,9 +131,9 @@ void PushToPhysicsStateImp(const Chaos::FDirtyPropertiesManager& Manager, Chaos:
 				Evolution.ResetVSmoothFromForces(*RigidHandle);
 			}
 
-			if(auto NewData = ParticleData.FindDynamicMisc(Manager,DataIdx))
+			if(NewDynamicMiscData)
 			{
-				Solver.SetParticleDynamicMisc(RigidHandle, *NewData);
+				Solver.SetParticleDynamicMisc(RigidHandle, *NewDynamicMiscData);
 			}
 		}
 
diff --git a/Engine/Source/Runtime/Experimental/Chaos/Public/Chaos/PBDRigidsEvolutionGBF.h b/Engine/Source/Runtime/Experimental/Chaos/Public/Chaos/PBDRigidsEvolutionGBF.h
--- a/Engine/Source/Runtime/Experimental/Chaos/Public/Chaos/PBDRigidsEvolutionGBF.h
+++ b/Engine/Source/Runtime/Experimental/Chaos/Public/Chaos/PBDRigidsEvolutionGBF.h
@@ -161,7 +161,7 @@ namespace Chaos
 		 * Does not change velocity.
 		 * Will wake the particle if this is a move (i.e., bIsTeleport is false)
 		 */
-		CHAOS_API void SetParticleTransform(FGeometryParticleHandle* InParticle, const FVec3& InPos, const FRotation3& InRot, const bool bIsTeleport);
+		CHAOS_API void SetParticleTransform(FGeometryParticleHandle* InParticle, const FVec3& InPos, const FRotation3& InRot, const bool bIsTeleport, const bool bWakeUp = true);
 
 		/**
 		 * Move a particle to a new location with a sweep and stop and the first opposing contact.
@@ -200,7 +200,7 @@ namespace Chaos
 		/**
 		 * Called when a particle is moved. We need to reset some friction properties, sleeping properties, etc
 		 */
-		CHAOS_API void OnParticleMoved(FGeometryParticleHandle* InParticle, const FVec3& PrevX, const FRotation3& PrevR, const bool bIsTeleport);
+		CHAOS_API void OnParticleMoved(FGeometryParticleHandle* InParticle, const FVec3& PrevX, const FRotation3& PrevR, const bool bIsTeleport, const bool bWakeUp = true);
 
 		/**
 		 * User has changed particle velocity or angular velocity
