IAS Csv Stats: Pass MB/KB size stats as float instead of int for accuracy

diff --git a/Engine/Source/Runtime/Experimental/IoStore/OnDemand/Private/Statistics.cpp b/Engine/Source/Runtime/Experimental/IoStore/OnDemand/Private/Statistics.cpp
--- a/Engine/Source/Runtime/Experimental/IoStore/OnDemand/Private/Statistics.cpp
+++ b/Engine/Source/Runtime/Experimental/IoStore/OnDemand/Private/Statistics.cpp
@@ -344,6 +344,8 @@ static FDelegateHandle GStatisticsOnScreenDelegateHandle;
 FOnDemandIoBackendStats::FOnDemandIoBackendStats(EStatsFlags InFlags)
 	: Flags(InFlags)
 {
+	static constexpr float OneOver1024 = 1.0f / 1024.0f;
+
 	check(GStatistics == nullptr);
 	GStatistics = this;
 	GStatisticsEndFrameDelegateHandle = FCoreDelegates::OnEndFrame.AddLambda([this]()
@@ -354,19 +356,19 @@ FOnDemandIoBackendStats::FOnDemandIoBackendStats(EStatsFlags InFlags)
 		int32 CPutCount = int32(GCachePutCount.Get());
 		int32 CPutExistingCount = int32(GCachePutExistingCount.Get());
 		int32 CPutRejectCount = int32(GCachePutRejectCount.Get());
-		int32 CCachedKiB = BytesToApproxKB(GCacheCachedBytes.Get());
-		int32 CWrittenKiB = BytesToApproxKB(GCacheWrittenBytes.Get());
-		int32 CReadKiB = BytesToApproxKB(GCacheReadBytes.Get());
-		int32 CRejectedKiB = BytesToApproxKB(GCacheRejectBytes.Get());
+		float CCachedKiB = (float)GCacheCachedBytes.Get() * OneOver1024;
+		float CWrittenKiB = (float)GCacheWrittenBytes.Get() * OneOver1024;
+		float CReadKiB = (float)GCacheReadBytes.Get() * OneOver1024;
+		float CRejectedKiB = (float)GCacheRejectBytes.Get() * OneOver1024;
 		CSV_CUSTOM_STAT_DEFINED(CacheGetCount, CGetCount, ECsvCustomStatOp::Set);
 		CSV_CUSTOM_STAT_DEFINED(CacheErrorCount, CErrorCount, ECsvCustomStatOp::Set);
 		CSV_CUSTOM_STAT_DEFINED(CachePutCount, CPutCount, ECsvCustomStatOp::Set);
 		CSV_CUSTOM_STAT_DEFINED(CachePutExistingCount, CPutExistingCount, ECsvCustomStatOp::Set);
 		CSV_CUSTOM_STAT_DEFINED(CachePutRejectCount, CPutRejectCount, ECsvCustomStatOp::Set);
-		CSV_CUSTOM_STAT_DEFINED(CacheCachedMB, CCachedKiB >> 10, ECsvCustomStatOp::Set);
-		CSV_CUSTOM_STAT_DEFINED(CacheWrittenMB, CWrittenKiB >> 10, ECsvCustomStatOp::Set);
-		CSV_CUSTOM_STAT_DEFINED(CacheReadMB, CReadKiB >> 10, ECsvCustomStatOp::Set);
-		CSV_CUSTOM_STAT_DEFINED(CacheRejectedMB, CRejectedKiB >> 10, ECsvCustomStatOp::Set);
+		CSV_CUSTOM_STAT_DEFINED(CacheCachedMB, CCachedKiB * OneOver1024, ECsvCustomStatOp::Set);
+		CSV_CUSTOM_STAT_DEFINED(CacheWrittenMB, CWrittenKiB * OneOver1024, ECsvCustomStatOp::Set);
+		CSV_CUSTOM_STAT_DEFINED(CacheReadMB, CReadKiB * OneOver1024, ECsvCustomStatOp::Set);
+		CSV_CUSTOM_STAT_DEFINED(CacheRejectedMB, CRejectedKiB * OneOver1024, ECsvCustomStatOp::Set);
 
 		// http stat totals
 		int32 HGetCount = int32(GHttpGetCount.Get());
@@ -374,7 +376,7 @@ FOnDemandIoBackendStats::FOnDemandIoBackendStats(EStatsFlags InFlags)
 		int32 HCancelCount = int32(GHttpCancelCount.Get());
 		int32 HErrorCount = int32(GHttpErrorCount.Get());
 		int32 HPendingCount = int32(GHttpPendingCount.Get());
-		int32 HDownloadedKiB = BytesToApproxKB(GHttpDownloadedBytes.Get());
+		float HDownloadedKiB = (float)GHttpDownloadedBytes.Get() * OneOver1024;
 		int32 HBandwidthMpbs = int32(GHttpBandwidthMpbs.Get());
 		int32 HDurationMsAvg = int32(GHttpDurationMsAvg.Get());
 		int32 HDurationMsMax = int32(GHttpDurationMsMax.Get());
@@ -382,7 +384,7 @@ FOnDemandIoBackendStats::FOnDemandIoBackendStats(EStatsFlags InFlags)
 		CSV_CUSTOM_STAT_DEFINED(HttpCancelCount, HCancelCount, ECsvCustomStatOp::Set);
 		CSV_CUSTOM_STAT_DEFINED(HttpErrorCount, HErrorCount, ECsvCustomStatOp::Set);
 		CSV_CUSTOM_STAT_DEFINED(HttpPendingCount, HPendingCount, ECsvCustomStatOp::Set);
-		CSV_CUSTOM_STAT_DEFINED(HttpDownloadedMB, HDownloadedKiB >> 10, ECsvCustomStatOp::Set);
+		CSV_CUSTOM_STAT_DEFINED(HttpDownloadedMB, HDownloadedKiB * OneOver1024, ECsvCustomStatOp::Set);
 		CSV_CUSTOM_STAT_DEFINED(HttpBandwidthMpbs, HBandwidthMpbs, ECsvCustomStatOp::Set);
 		CSV_CUSTOM_STAT_DEFINED(HttpDurationMsAvg, HDurationMsAvg, ECsvCustomStatOp::Set);
 		CSV_CUSTOM_STAT_DEFINED(HttpDurationMsMax, HDurationMsMax, ECsvCustomStatOp::Set);
@@ -395,7 +397,7 @@ FOnDemandIoBackendStats::FOnDemandIoBackendStats(EStatsFlags InFlags)
 				if (!EnumHasAllFlags(Flags,EStatsFlags::CachingDisabled))
 				{
 					UE_LOG(LogIas, Log, TEXT("CacheStats: CachedKiB=%d, WrittenKiB=%d, ReadKiB=%d, RejectedKiB=%d, Get=%d, Error=%d, Put=%d, PutReject=%d, PutExisting=%d"),
-						CCachedKiB, CWrittenKiB, CReadKiB, CRejectedKiB, CGetCount, CErrorCount, CPutCount, CPutRejectCount, CPutExistingCount);
+						(int32)CCachedKiB, (int32)CWrittenKiB, (int32)CReadKiB, (int32)CRejectedKiB, CGetCount, CErrorCount, CPutCount, CPutRejectCount, CPutExistingCount);
 				}
 				else
 				{
@@ -403,7 +405,7 @@ FOnDemandIoBackendStats::FOnDemandIoBackendStats(EStatsFlags InFlags)
 				}
 
 				UE_LOG(LogIas, Log, TEXT("HttpStats: DownloadedKiB=%d, Get=%d, Retry=%d, Cancel=%d, Error=%d, CurPending=%d, CurDurationMsAvg=%d, CurDurationMsMax=%d"),
-					HDownloadedKiB, HGetCount, HRetryCount, HCancelCount, HErrorCount, HPendingCount, HDurationMsAvg, HDurationMsMax);
+					(int32)HDownloadedKiB, HGetCount, HRetryCount, HCancelCount, HErrorCount, HPendingCount, HDurationMsAvg, HDurationMsMax);
 				LastLogTime = Time;
 			}
 		}
