Add flag to disable delta encoding when stringifying an object with JsonObjectGraph

#rnx

diff --git a/Engine/Source/Runtime/Experimental/JsonObjectGraph/Private/JsonStringifyImpl.cpp b/Engine/Source/Runtime/Experimental/JsonObjectGraph/Private/JsonStringifyImpl.cpp
--- a/Engine/Source/Runtime/Experimental/JsonObjectGraph/Private/JsonStringifyImpl.cpp
+++ b/Engine/Source/Runtime/Experimental/JsonObjectGraph/Private/JsonStringifyImpl.cpp
@@ -627,7 +627,7 @@ void FJsonStringifyImpl::WriteIdentifierAndValueToJson(const void* Container, co
 
 	// this is inefficient for structs... but it papers over some
 	// problematic identical implementations. We will improve it
-	bool bMatchesDefault = DefaultContainer != nullptr;
+	bool bMatchesDefault = DefaultContainer != nullptr && IsDeltaEncoding();
 	for (int32 Idx = 0; Idx < Property->ArrayDim && bMatchesDefault; ++Idx)
 	{
 		const void* Value = Property->ContainerPtrToValuePtr<void>(Container, Idx);
@@ -952,6 +952,11 @@ TArray<uint8> FJsonStringifyImpl::StructuredDataToJson(const UObject* Object, in
 }
 #endif // WITH_TEXT_ARCHIVE_SUPPORT
 
+bool FJsonStringifyImpl::IsDeltaEncoding() const
+{
+	return !EnumHasAnyFlags(WriteOptions.Flags, EJsonStringifyFlags::DisableDeltaEncoding);
+}
+
 FJsonStringifyImpl::FPendingScope::FPendingScope(FJsonStringifyImpl* To, TFunctionRef<void()> Prefix)
 	: Owner(To)
 	, Outer(To->CurrentScope)
diff --git a/Engine/Source/Runtime/Experimental/JsonObjectGraph/Private/JsonStringifyImpl.h b/Engine/Source/Runtime/Experimental/JsonObjectGraph/Private/JsonStringifyImpl.h
--- a/Engine/Source/Runtime/Experimental/JsonObjectGraph/Private/JsonStringifyImpl.h
+++ b/Engine/Source/Runtime/Experimental/JsonObjectGraph/Private/JsonStringifyImpl.h
@@ -67,6 +67,8 @@ private:
 	TArray<uint8> StructuredDataToJson(const UObject* Object, int32 InitialIndentLevel);
 #endif// WITH_TEXT_ARCHIVE_SUPPORT
 
+	bool IsDeltaEncoding() const;
+
 	const FJsonStringifyOptions& WriteOptions;
 
 	TArray<uint8> Result;
diff --git a/Engine/Source/Runtime/Experimental/JsonObjectGraph/Public/JsonObjectGraph/Stringify.h b/Engine/Source/Runtime/Experimental/JsonObjectGraph/Public/JsonObjectGraph/Stringify.h
--- a/Engine/Source/Runtime/Experimental/JsonObjectGraph/Public/JsonObjectGraph/Stringify.h
+++ b/Engine/Source/Runtime/Experimental/JsonObjectGraph/Public/JsonObjectGraph/Stringify.h
@@ -16,6 +16,8 @@ enum class EJsonStringifyFlags : uint8
 	Default = 0,
 	/** Filters editor only data such that it is not written to the Json */
 	FilterEditorOnlyData = 1 << 1,
+	/** Disables delta encoding such that all properties are encoded in the Json, rather than only changes from the objects' archetypes */
+	DisableDeltaEncoding = 1 << 2,
 };
 ENUM_CLASS_FLAGS(EJsonStringifyFlags);
 
