[AutoRTFM] Disable memory validation for hot opens, optimize small write hashing

Disable the memory validator for frequently called open blocks.
These have been successfully passing the validator, and are presumed safe.

diff --git a/Engine/Source/Runtime/Core/Private/AutoRTFM/Transaction.cpp b/Engine/Source/Runtime/Core/Private/AutoRTFM/Transaction.cpp
--- a/Engine/Source/Runtime/Core/Private/AutoRTFM/Transaction.cpp
+++ b/Engine/Source/Runtime/Core/Private/AutoRTFM/Transaction.cpp
@@ -441,18 +441,54 @@ FTransaction::TWriteHash FTransaction::CalculateNestedWriteHash(size_t NumWriteE
 		}
 
 		const FWriteLogEntry& Entry = *Iter;
-		const void* Original = Entry.GetOriginal();
-		size_t Size = Entry.GetSize();
-		if (Size <= sizeof(FSmallValue))
+		const uint8* Original = reinterpret_cast<const uint8*>(Entry.GetOriginal());
+		const size_t Size = Entry.GetSize();
+		if (Size > sizeof(FSmallValue))
 		{
-			FSmallValue Value = 0;
-			memcpy(&Value, Original, Size);
-			SmallValue = SmallValue * 31 + Value;
+			Hasher.Update(Original, Size);
+			continue; // Next write
 		}
-		else
+
+		const uintptr_t Address = reinterpret_cast<uintptr_t>(Original);
+		switch (Size)
 		{
-			Hasher.Update(Original, Size);
+			case 8:
+				if (LIKELY((Address & 7) == 0))
+				{
+					uint64 Load = *reinterpret_cast<uint64 const*>(Original);
+					SmallValue = (SmallValue * 31) ^ static_cast<FSmallValue>(Load);
+				}
+				break;
+
+			case 4:
+				if (LIKELY((Address & 3) == 0))
+				{
+					uint32 Load = *reinterpret_cast<uint32 const*>(Original);
+					SmallValue = (SmallValue * 31) ^ static_cast<FSmallValue>(Load);
+				}
+				break;
+
+			case 2:
+				if (LIKELY((Address & 1) == 0))
+				{
+					uint16 Load = *reinterpret_cast<uint16 const*>(Original);
+					SmallValue = (SmallValue * 31) ^ static_cast<FSmallValue>(Load);
+				}
+				break;
+
+			case 1:
+				SmallValue = (SmallValue * 31) ^ static_cast<FSmallValue>(Original[0]);
+				break;
+
+			default:
+				// Unaligned / non-POT write.
+				for (size_t I = 0; I < Size; I++)
+				{
+					SmallValue = (SmallValue * 31) ^ static_cast<FSmallValue>(Original[I]);
+				}
+				break;
 		}
+
 	}
 	Hasher.Update(&SmallValue, sizeof(SmallValue));
 
diff --git a/Engine/Source/Runtime/Core/Private/Delegates/DelegateHandle.cpp b/Engine/Source/Runtime/Core/Private/Delegates/DelegateHandle.cpp
--- a/Engine/Source/Runtime/Core/Private/Delegates/DelegateHandle.cpp
+++ b/Engine/Source/Runtime/Core/Private/Delegates/DelegateHandle.cpp
@@ -15,7 +15,8 @@ uint64 FDelegateHandle::GenerateNewID()
 	// Just increment a counter to generate an ID.
 	uint64 Result = 0; // Initialize just to silence static analysis.
 	
-	UE_AUTORTFM_OPEN
+	// Memory validation disabled as this is hot code.
+	UE_AUTORTFM_OPEN_NO_VALIDATION
 	{
 		Result = ++UE::Delegates::Private::GNextID;
 
diff --git a/Engine/Source/Runtime/Core/Private/UObject/UnrealNames.cpp b/Engine/Source/Runtime/Core/Private/UObject/UnrealNames.cpp
--- a/Engine/Source/Runtime/Core/Private/UObject/UnrealNames.cpp
+++ b/Engine/Source/Runtime/Core/Private/UObject/UnrealNames.cpp
@@ -3130,7 +3130,8 @@ private:
 	static FName MakeInternal(FNameStringView View, EFindName FindType, int32 InternalNumber)
 	{
 		FNameEntryIds Ids;
-		UE_AUTORTFM_OPEN
+		// Memory validation disabled as this is hot code.
+		UE_AUTORTFM_OPEN_NO_VALIDATION
 		{
 			Ids = FindOrStoreString(View, FindType);
 		};
diff --git a/Engine/Source/Runtime/Core/Public/HAL/FMemory.inl b/Engine/Source/Runtime/Core/Public/HAL/FMemory.inl
--- a/Engine/Source/Runtime/Core/Public/HAL/FMemory.inl
+++ b/Engine/Source/Runtime/Core/Public/HAL/FMemory.inl
@@ -169,7 +169,8 @@ FMEMORY_INLINE_FUNCTION_DECORATOR void FMemory::Free(void* Original)
 FMEMORY_INLINE_FUNCTION_DECORATOR SIZE_T FMemory::GetAllocSize(void* Original)
 {
 	SIZE_T Result;
-	UE_AUTORTFM_OPEN
+	// Memory validation disabled as this is hot code.
+	UE_AUTORTFM_OPEN_NO_VALIDATION
 	{
 		if (!FMEMORY_INLINE_GMalloc)
 		{
diff --git a/Engine/Source/Runtime/Core/Public/Misc/MTTransactionallySafeAccessDetector.h b/Engine/Source/Runtime/Core/Public/Misc/MTTransactionallySafeAccessDetector.h
--- a/Engine/Source/Runtime/Core/Public/Misc/MTTransactionallySafeAccessDetector.h
+++ b/Engine/Source/Runtime/Core/Public/Misc/MTTransactionallySafeAccessDetector.h
@@ -85,7 +85,8 @@ private:
 	{
 		if (AutoRTFM::IsTransactional() || AutoRTFM::IsCommittingOrAborting())
 		{
-			const bool bResult = AutoRTFM::Open([&]
+			// Memory validation disabled as this is hot code.
+			const bool bResult = AutoRTFM::Open<AutoRTFM::EMemoryValidationLevel::Disabled>([&]
 				{
 					// The transactional system which can increment the lock counts
 					// is always single-threaded, thus this is safe to check without atomicity.
@@ -150,7 +151,8 @@ private:
 	{
 		if (AutoRTFM::IsTransactional() || AutoRTFM::IsCommittingOrAborting())
 		{
-			const bool bResult = AutoRTFM::Open([&]
+			// Memory validation disabled as this is hot code.
+			const bool bResult = AutoRTFM::Open<AutoRTFM::EMemoryValidationLevel::Disabled>([&]
 				{
 					if ((0 == State->ReadLockCount) && (0 == State->WriteLockCount))
 					{
diff --git a/Engine/Source/Runtime/Core/Public/Misc/TransactionallySafeCriticalSection.h b/Engine/Source/Runtime/Core/Public/Misc/TransactionallySafeCriticalSection.h
--- a/Engine/Source/Runtime/Core/Public/Misc/TransactionallySafeCriticalSection.h
+++ b/Engine/Source/Runtime/Core/Public/Misc/TransactionallySafeCriticalSection.h
@@ -72,7 +72,8 @@ struct FTransactionallySafeCriticalSectionDefinition final
 	{
 		if (AutoRTFM::IsTransactional() || AutoRTFM::IsCommittingOrAborting())
 		{
-			AutoRTFM::Open([this]
+			// Memory validation disabled as this is hot code.
+			UE_AUTORTFM_OPEN_NO_VALIDATION
 				{
 					// The transactional system which can increment TransactionalLockCount
 					// is always single-threaded, thus this is safe to check without atomicity.
@@ -82,7 +83,7 @@ struct FTransactionallySafeCriticalSectionDefinition final
 					}
 
 					State->TransactionalLockCount += 1;
-				});
+				};
 
 			// We explicitly copy the state here for the case that `this` was stack
 			// allocated and has already died before the on-abort is hit.
diff --git a/Engine/Source/Runtime/Core/Public/Misc/TransactionallySafeMutex.h b/Engine/Source/Runtime/Core/Public/Misc/TransactionallySafeMutex.h
--- a/Engine/Source/Runtime/Core/Public/Misc/TransactionallySafeMutex.h
+++ b/Engine/Source/Runtime/Core/Public/Misc/TransactionallySafeMutex.h
@@ -74,7 +74,8 @@ struct FTransactionallySafeMutexDefinition final
 	{
 		if (AutoRTFM::IsTransactional() || AutoRTFM::IsCommittingOrAborting())
 		{
-			AutoRTFM::Open([this]
+			// Memory validation disabled as this is hot code.
+			UE_AUTORTFM_OPEN_NO_VALIDATION
 				{
 					// The transactional system which can increment TransactionalLockCount
 					// is always single-threaded, thus this is safe to check without atomicity.
@@ -84,7 +85,7 @@ struct FTransactionallySafeMutexDefinition final
 					}
 
 					State->TransactionalLockCount += 1;
-				});
+				};
 
 			// We explicitly copy the state here for the case that `this` was stack
 			// allocated and has already died before the on-abort is hit.
diff --git a/Engine/Source/Runtime/Core/Public/Misc/TransactionallySafeRWLock.h b/Engine/Source/Runtime/Core/Public/Misc/TransactionallySafeRWLock.h
--- a/Engine/Source/Runtime/Core/Public/Misc/TransactionallySafeRWLock.h
+++ b/Engine/Source/Runtime/Core/Public/Misc/TransactionallySafeRWLock.h
@@ -76,7 +76,8 @@ struct FTransactionallySafeRWLockDefinition
 		}
 
 		// As the State was constructed in the open, it must be released in the open.
-		AutoRTFM::Open([&] { State = nullptr; });
+		// Memory validation disabled as this is hot code.
+		UE_AUTORTFM_OPEN_NO_VALIDATION{ State = nullptr; };
 	}
 
 	void ReadLock()
@@ -111,7 +112,8 @@ struct FTransactionallySafeRWLockDefinition
 	{
 		if (AutoRTFM::IsTransactional() || AutoRTFM::IsCommittingOrAborting())
 		{
-			AutoRTFM::Open([&]
+			// Memory validation disabled as this is hot code.
+			UE_AUTORTFM_OPEN_NO_VALIDATION
 				{
 					// The transactional system which can increment TransactionalLockCount
 					// is always single-threaded, thus this is safe to check without atomicity.
@@ -121,7 +123,7 @@ struct FTransactionallySafeRWLockDefinition
 					}
 
 					State->TransactionalLockCount += 1;
-				});
+				};
 
 			// We explicitly copy the state here for the case that `this` was stack
 			// allocated and has already died before the on-abort is hit.
diff --git a/Engine/Source/Runtime/Core/Public/ProfilingDebugging/CsvProfiler.h b/Engine/Source/Runtime/Core/Public/ProfilingDebugging/CsvProfiler.h
--- a/Engine/Source/Runtime/Core/Public/ProfilingDebugging/CsvProfiler.h
+++ b/Engine/Source/Runtime/Core/Public/ProfilingDebugging/CsvProfiler.h
@@ -571,14 +571,16 @@ public:
 		: StatName(InStatName)
 		, CategoryIndex(InCategoryIndex)
 	{
-		UE_AUTORTFM_OPEN{ FCsvProfiler::BeginStat(StatName, CategoryIndex, InNamedEventName); };
+		// Memory validation disabled as this is hot code.
+		UE_AUTORTFM_OPEN_NO_VALIDATION{ FCsvProfiler::BeginStat(StatName, CategoryIndex, InNamedEventName); };
 		AutoRTFM::PushOnAbortHandler(this, [SN = StatName, CI = CategoryIndex](){ FCsvProfiler::EndStat(SN, CI); });
 	}
 
 	~FScopedCsvStat()
 	{
 		AutoRTFM::PopOnAbortHandler(this);
-		UE_AUTORTFM_OPEN{ FCsvProfiler::EndStat(StatName, CategoryIndex); };
+		// Memory validation disabled as this is hot code.
+		UE_AUTORTFM_OPEN_NO_VALIDATION{ FCsvProfiler::EndStat(StatName, CategoryIndex); };
 	}
 	const char * StatName;
 	uint32 CategoryIndex;
diff --git a/Engine/Source/Runtime/Core/Public/Stats/LightweightStats.h b/Engine/Source/Runtime/Core/Public/Stats/LightweightStats.h
--- a/Engine/Source/Runtime/Core/Public/Stats/LightweightStats.h
+++ b/Engine/Source/Runtime/Core/Public/Stats/LightweightStats.h
@@ -59,11 +59,12 @@ public:
 	{
 		if (GCycleStatsShouldEmitNamedEvents && InStatId.IsValidStat())
 		{
-			AutoRTFM::Open([&]
-				{
-					bPop = true;
-					FPlatformMisc::BeginNamedEvent(FColor(0), InStatId.StatString);
-				});
+			// Memory validation disabled as this is hot code.
+			UE_AUTORTFM_OPEN_NO_VALIDATION
+			{
+				bPop = true;
+				FPlatformMisc::BeginNamedEvent(FColor(0), InStatId.StatString);
+			};
 			AutoRTFM::PushOnAbortHandler(this, []() { FPlatformMisc::EndNamedEvent(); });
 		}
 	}
@@ -77,13 +78,14 @@ public:
 	{
 		AutoRTFM::PopOnAbortHandler(this);
 
-		AutoRTFM::Open([&]
+		// Memory validation disabled as this is hot code.
+		UE_AUTORTFM_OPEN_NO_VALIDATION
+		{
+			if (bPop)
 			{
-				if (bPop)
-				{
-					FPlatformMisc::EndNamedEvent();
-				}
-			});
+				FPlatformMisc::EndNamedEvent();
+			}
+		};
 	}
 private:
 #if USE_LIGHTWEIGHT_STATS_FOR_HITCH_DETECTION && USE_HITCH_DETECTION
@@ -106,4 +108,4 @@ public:
 	{
 	}
 };
-#endif // UE_USE_LIGHTWEIGHT_STATS
\ No newline at end of file
+#endif // UE_USE_LIGHTWEIGHT_STATS
diff --git a/Engine/Source/Runtime/Core/Public/Stats/StatsSystemTypes.h b/Engine/Source/Runtime/Core/Public/Stats/StatsSystemTypes.h
--- a/Engine/Source/Runtime/Core/Public/Stats/StatsSystemTypes.h
+++ b/Engine/Source/Runtime/Core/Public/Stats/StatsSystemTypes.h
@@ -1620,7 +1620,11 @@ public:
 	 */
 	FORCEINLINE_STATS FScopeCycleCounter(TStatId StatId, EStatFlags StatFlags, bool bAlways = false, UE::FSourceLocation SourceLocation = UE::FSourceLocation::Current())
 	{
-		AutoRTFM::Open([&] { Start(StatId, StatFlags, bAlways, SourceLocation); });
+		// Memory validation disabled as this is hot code.
+		UE_AUTORTFM_OPEN_NO_VALIDATION
+		{
+			Start(StatId, StatFlags, bAlways, SourceLocation);
+		};
 		AutoRTFM::PushOnAbortHandler(this, [this]() { this->Stop(); });
 	}
 
@@ -1635,10 +1639,11 @@ public:
 	FORCEINLINE_STATS ~FScopeCycleCounter()
 	{
 		AutoRTFM::PopOnAbortHandler(this);
-		AutoRTFM::Open([&]
-			{
-				Stop();
-			});
+		// Memory validation disabled as this is hot code.
+		UE_AUTORTFM_OPEN_NO_VALIDATION
+		{
+			Stop();
+		};
 	}
 
 };
@@ -1817,4 +1822,4 @@ inline void StatsPrimaryEnableSubtract(int32 Value = 1)
 {
 }
 
-#endif // STATS
\ No newline at end of file
+#endif // STATS
diff --git a/Engine/Source/Runtime/CoreUObject/Private/UObject/Class.cpp b/Engine/Source/Runtime/CoreUObject/Private/UObject/Class.cpp
--- a/Engine/Source/Runtime/CoreUObject/Private/UObject/Class.cpp
+++ b/Engine/Source/Runtime/CoreUObject/Private/UObject/Class.cpp
@@ -6806,7 +6806,8 @@ UFunction* UClass::FindFunctionByName(FName InName, EIncludeSuperFlag::Type Incl
 
 	UFunction* Result = nullptr;
 
-	UE_AUTORTFM_OPEN
+	// Memory validation disabled as this is hot code.
+	UE_AUTORTFM_OPEN_NO_VALIDATION
 	{
 		UClass* SuperClass = GetSuperClass();
 		if (IncludeSuper == EIncludeSuperFlag::ExcludeSuper || ( Interfaces.Num() == 0 && SuperClass == nullptr ) )
diff --git a/Engine/Source/Runtime/CoreUObject/Private/UObject/CoreRedirects.cpp b/Engine/Source/Runtime/CoreUObject/Private/UObject/CoreRedirects.cpp
--- a/Engine/Source/Runtime/CoreUObject/Private/UObject/CoreRedirects.cpp
+++ b/Engine/Source/Runtime/CoreUObject/Private/UObject/CoreRedirects.cpp
@@ -1280,7 +1280,8 @@ public:
 	: InternalLock(InLock)
 	, NeedsUnlock(true)
 	{
-		UE_AUTORTFM_OPEN { InternalLock.ReadLock(); };
+		// Memory validation disabled as this is hot code.
+		UE_AUTORTFM_OPEN_NO_VALIDATION { InternalLock.ReadLock(); };
 		AutoRTFM::PushOnAbortHandler(this, [this] { this->InternalLock.ReadUnlock(); });
 	}
 
@@ -1289,7 +1290,8 @@ public:
 		if (NeedsUnlock)
 		{
 			AutoRTFM::PopOnAbortHandler(this);
-			UE_AUTORTFM_OPEN { InternalLock.ReadUnlock(); };
+			// Memory validation disabled as this is hot code.
+			UE_AUTORTFM_OPEN_NO_VALIDATION { InternalLock.ReadUnlock(); };
 		}
 	}
 protected:
diff --git a/Engine/Source/Runtime/CoreUObject/Private/UObject/ScriptCore.cpp b/Engine/Source/Runtime/CoreUObject/Private/UObject/ScriptCore.cpp
--- a/Engine/Source/Runtime/CoreUObject/Private/UObject/ScriptCore.cpp
+++ b/Engine/Source/Runtime/CoreUObject/Private/UObject/ScriptCore.cpp
@@ -193,7 +193,8 @@ FBlueprintContext* FBlueprintContextGetThreadSingletonImpl()
 FBlueprintContext* FBlueprintContext::GetThreadSingleton()
 {
 	FBlueprintContext* Result;
-	UE_AUTORTFM_OPEN{ Result = FBlueprintContextGetThreadSingletonImpl(); };
+	// Memory validation disabled as this is hot code.
+	UE_AUTORTFM_OPEN_NO_VALIDATION{ Result = FBlueprintContextGetThreadSingletonImpl(); };
 	return Result;
 }
 
diff --git a/Engine/Source/Runtime/CoreUObject/Private/UObject/UObjectHash.cpp b/Engine/Source/Runtime/CoreUObject/Private/UObject/UObjectHash.cpp
--- a/Engine/Source/Runtime/CoreUObject/Private/UObject/UObjectHash.cpp
+++ b/Engine/Source/Runtime/CoreUObject/Private/UObject/UObjectHash.cpp
@@ -526,7 +526,8 @@ public:
 		if (!(IsGarbageCollectingAndLockingUObjectHashTables() && IsInGameThread()))
 		{
 			Tables = &InTables;
-			UE_AUTORTFM_OPEN{ InTables.Lock(); };
+			// Memory validation disabled as this is hot code.
+			UE_AUTORTFM_OPEN_NO_VALIDATION{ InTables.Lock(); };
 			AutoRTFM::PushOnAbortHandler(this, [this](){ if (this->Tables){ this->Tables->Unlock(); }});
 		}
 		else
@@ -542,7 +543,8 @@ public:
 #if THREADSAFE_UOBJECTS
 		if (Tables)
 		{
-			UE_AUTORTFM_OPEN{ Tables->Unlock(); };
+			// Memory validation disabled as this is hot code.
+			UE_AUTORTFM_OPEN_NO_VALIDATION{ Tables->Unlock(); };
 			AutoRTFM::PopOnAbortHandler(this);
 		}
 #endif
@@ -872,7 +874,8 @@ UObject* StaticFindObjectFastInternal(const UClass* ObjectClass, const UObject*
 	UObject* Result = nullptr;
 
 	// Transactionally, a static find operation has to occur in the open as it touches shared state.
-	UE_AUTORTFM_OPEN
+	// Memory validation disabled as this is hot code.
+	UE_AUTORTFM_OPEN_NO_VALIDATION
 	{
 		INC_DWORD_STAT(STAT_FindObjectFast);
 
diff --git a/Engine/Source/Runtime/CoreUObject/Public/UObject/Stack.h b/Engine/Source/Runtime/CoreUObject/Public/UObject/Stack.h
--- a/Engine/Source/Runtime/CoreUObject/Public/UObject/Stack.h
+++ b/Engine/Source/Runtime/CoreUObject/Public/UObject/Stack.h
@@ -35,7 +35,8 @@ FORCEINLINE void* UeVstackAllocHelper(FVirtualStackAllocator* Allocator, size_t
 	{
 		// in transactional code we redirect the virtual stack allocator to FMemory::Malloc
 		// until we properly implement a transaction-friendly stack allocator
-		Result = AutoRTFM::Open([Size, Align]
+		// Memory validation disabled as this is hot code.
+		Result = AutoRTFM::Open<AutoRTFM::EMemoryValidationLevel::Disabled>([Size, Align]
 		{
 			return FMemory::Malloc(Size, Align);
 		});
diff --git a/Engine/Source/Runtime/CoreUObject/Public/UObject/UObjectBase.h b/Engine/Source/Runtime/CoreUObject/Public/UObject/UObjectBase.h
--- a/Engine/Source/Runtime/CoreUObject/Public/UObject/UObjectBase.h
+++ b/Engine/Source/Runtime/CoreUObject/Public/UObject/UObjectBase.h
@@ -320,7 +320,8 @@ private:
 
 		int32 Result = 0;
 
-		UE_AUTORTFM_OPEN
+		// Memory validation disabled as this is hot code.
+		UE_AUTORTFM_OPEN_NO_VALIDATION
 		{
 			Result = FPlatformAtomics::AtomicRead_Relaxed((int32*)&ObjectFlags);
 		};
diff --git a/Engine/Source/Runtime/Engine/Private/DataReplication.cpp b/Engine/Source/Runtime/Engine/Private/DataReplication.cpp
--- a/Engine/Source/Runtime/Engine/Private/DataReplication.cpp
+++ b/Engine/Source/Runtime/Engine/Private/DataReplication.cpp
@@ -5,6 +5,7 @@
 =============================================================================*/
 
 #include "Net/DataReplication.h"
+#include "AutoRTFM/AutoRTFM.h"
 #include "Containers/StaticBitArray.h"
 #include "EngineStats.h"
 #include "Engine/World.h"
@@ -400,7 +401,9 @@ void FObjectReplicator::InitRecentProperties(uint8* Source)
 	// If acting as a server and are IsInternalAck, that means we're recording.
 	// In that case, we don't need to create any receiving state, as no one will be sending data to us.
 	ECreateRepStateFlags Flags = (Connection->IsInternalAck() && bIsServer) ? ECreateRepStateFlags::SkipCreateReceivingState : ECreateRepStateFlags::None;
-	UE_AUTORTFM_OPEN
+
+	// Memory validation disabled as this is hot code.
+	UE_AUTORTFM_OPEN_NO_VALIDATION
 	{
 		RepState = LocalRepLayout.CreateRepState(Source, RepChangedPropertyTracker, Flags);
 	};
diff --git a/Engine/Source/Runtime/Experimental/Chaos/Private/Chaos/AABB.cpp b/Engine/Source/Runtime/Experimental/Chaos/Private/Chaos/AABB.cpp
--- a/Engine/Source/Runtime/Experimental/Chaos/Private/Chaos/AABB.cpp
+++ b/Engine/Source/Runtime/Experimental/Chaos/Private/Chaos/AABB.cpp
@@ -410,7 +410,8 @@ inline TAABB<FReal, 3> TransformedAABBHelperISPC(const TAABB<FReal, 3>& AABB, co
 	check(bRealTypeCompatibleWithISPC);
 #if INTEL_ISPC
 	TVector<Chaos::FReal, 3> NewMin, NewMax;
-	UE_AUTORTFM_OPEN
+	// Memory validation disabled as this is hot code.
+	UE_AUTORTFM_OPEN_NO_VALIDATION
 	{
 		ispc::TransformedAABB((const ispc::FTransform&)SpaceTransform, (const ispc::FVector&)AABB.Min(), (const ispc::FVector&)AABB.Max(), (ispc::FVector&)NewMin, (ispc::FVector&)NewMax);
 	};
@@ -436,7 +437,8 @@ inline TAABB<Chaos::FRealSingle, 3> TransformedAABBHelperISPC(const TAABB<Chaos:
 	static_assert(sizeof(ispc::FVector3f) == sizeof(Chaos::TVector<Chaos::FRealSingle, 3>), "sizeof(ispc::FVector3f) != sizeof(Chaos::TVector<Chaos::FRealSingle, 3>)");
 
 	TVector<Chaos::FRealSingle, 3> NewMin, NewMax;
-	UE_AUTORTFM_OPEN
+	// Memory validation disabled as this is hot code.
+	UE_AUTORTFM_OPEN_NO_VALIDATION
 	{
 		ispc::TransformedAABBMixed((const ispc::FTransform&)SpaceTransform, (const ispc::FVector3f&)AABB.Min(), (const ispc::FVector3f&)AABB.Max(), (ispc::FVector3f&)NewMin, (ispc::FVector3f&)NewMax);
 	};
@@ -460,7 +462,8 @@ inline TAABB<FReal, 3> TransformedAABBHelperISPC2(const TAABB<FReal, 3>& AABB, c
 	check(bRealTypeCompatibleWithISPC);
 #if INTEL_ISPC
 	TVector<Chaos::FReal, 3> NewMin, NewMax;
-	UE_AUTORTFM_OPEN
+	// Memory validation disabled as this is hot code.
+	UE_AUTORTFM_OPEN_NO_VALIDATION
 	{
 		ispc::TransformedAABB2((const ispc::FTransform&)SpaceTransform, (const ispc::FVector&)AABB.Min(), (const ispc::FVector&)AABB.Max(), (ispc::FVector&)NewMin, (ispc::FVector&)NewMax);
 	};
@@ -486,7 +489,8 @@ inline TAABB<Chaos::FRealSingle, 3> TransformedAABBHelperISPC2(const TAABB<Chaos
 	static_assert(sizeof(ispc::FVector3f) == sizeof(Chaos::TVector<Chaos::FRealSingle, 3>), "sizeof(ispc::FVector3f) != sizeof(Chaos::TVector<Chaos::FRealSingle, 3>)");
 
 	TVector<Chaos::FRealSingle, 3> NewMin, NewMax;
-	UE_AUTORTFM_OPEN
+	// Memory validation disabled as this is hot code.
+	UE_AUTORTFM_OPEN_NO_VALIDATION
 	{
 		ispc::TransformedAABBMixed((const ispc::FTransform&)SpaceTransform, (const ispc::FVector3f&)AABB.Min(), (const ispc::FVector3f&)AABB.Max(), (ispc::FVector3f&)NewMin, (ispc::FVector3f&)NewMax);
 	};
diff --git a/Engine/Source/Runtime/Experimental/Chaos/Public/Chaos/RefCountedObject.h b/Engine/Source/Runtime/Experimental/Chaos/Public/Chaos/RefCountedObject.h
--- a/Engine/Source/Runtime/Experimental/Chaos/Public/Chaos/RefCountedObject.h
+++ b/Engine/Source/Runtime/Experimental/Chaos/Public/Chaos/RefCountedObject.h
@@ -37,7 +37,8 @@ public:
 	{
 		bool bIsFirstReference;
 
-		UE_AUTORTFM_OPEN
+		// Memory validation disabled as this is hot code.
+		UE_AUTORTFM_OPEN_NO_VALIDATION
 		{
 			bIsFirstReference = (++NumRefs == 1);
 		};
@@ -54,7 +55,7 @@ public:
 				// After an object gains its initial reference, an AddRef call can be balanced
 				// out with a matching Release.
 				Release();
-			} 
+			}
 		};
 		// TRefCountPtr doesn't use the return value.
 		return 0;
@@ -76,7 +77,8 @@ public:
 	uint32 GetRefCount() const
 	{
 		uint32 Ret;
-		UE_AUTORTFM_OPEN
+		// Memory validation disabled as this is hot code.
+		UE_AUTORTFM_OPEN_NO_VALIDATION
 		{
 			Ret = uint32(NumRefs.load());
 		};
@@ -113,4 +115,4 @@ private:
 	mutable std::atomic<ERCM_RefCountMode> RefCountMode = ERCM_Transient;
 };
 
-}  // namespace Chaos
\ No newline at end of file
+}  // namespace Chaos
diff --git a/Engine/Source/Runtime/Experimental/Chaos/Public/Framework/Threading.h b/Engine/Source/Runtime/Experimental/Chaos/Public/Framework/Threading.h
--- a/Engine/Source/Runtime/Experimental/Chaos/Public/Framework/Threading.h
+++ b/Engine/Source/Runtime/Experimental/Chaos/Public/Framework/Threading.h
@@ -1047,7 +1047,8 @@ FORCEINLINE void EnsureIsInGameThreadContext()
 		{
 			if (AutoRTFM::IsTransactional() || AutoRTFM::IsCommittingOrAborting())
 			{
-				AutoRTFM::Open([this]
+				// Memory validation disabled as this is hot code.
+				UE_AUTORTFM_OPEN_NO_VALIDATION
 					{
 						// The transactional system which can increment TransactionalLockCount
 						// is always single-threaded, thus this is safe to check without atomicity.
@@ -1057,7 +1058,7 @@ FORCEINLINE void EnsureIsInGameThreadContext()
 						}
 
 						State->TransactionalLockCount += 1;
-					});
+					};
 
 				// We explicitly copy the state here for the case that `this` was stack
 				// allocated and has already died before the on-abort is hit.
