3-Open dataflow simulation asset + add attribute input to addscalarnode + more dataflow types + update simulation cache

diff --git a/Engine/Plugins/Experimental/Dataflow/Source/DataflowEditor/Private/Dataflow/AssetDefinition_DataflowAsset.cpp b/Engine/Plugins/Experimental/Dataflow/Source/DataflowEditor/Private/Dataflow/AssetDefinition_DataflowAsset.cpp
--- a/Engine/Plugins/Experimental/Dataflow/Source/DataflowEditor/Private/Dataflow/AssetDefinition_DataflowAsset.cpp
+++ b/Engine/Plugins/Experimental/Dataflow/Source/DataflowEditor/Private/Dataflow/AssetDefinition_DataflowAsset.cpp
@@ -181,14 +181,14 @@ FAssetOpenSupport UAssetDefinition_DataflowAsset::GetAssetOpenSupport(const FAss
 
 EAssetCommandResult UAssetDefinition_DataflowAsset::OpenAssets(const FAssetOpenArgs& OpenArgs) const
 {
-	if (bCanEditDataflow)
-	{
-		TArray<UDataflow*> DataflowObjects = OpenArgs.LoadObjects<UDataflow>();
+	TArray<UDataflow*> DataflowObjects = OpenArgs.LoadObjects<UDataflow>();
 
-		// For now the dataflow editor only works on one asset at a time
-		ensure(DataflowObjects.Num() == 0 || DataflowObjects.Num() == 1);
+	// For now the dataflow editor only works on one asset at a time
+	ensure(DataflowObjects.Num() == 0 || DataflowObjects.Num() == 1);
 
-		if (DataflowObjects.Num() == 1)
+	if (DataflowObjects.Num() == 1)
+	{
+		if (bCanEditDataflow || DataflowObjects[0]->Type == EDataflowType::Simulation)
 		{
 			UAssetEditorSubsystem* const AssetEditorSubsystem = GEditor->GetEditorSubsystem<UAssetEditorSubsystem>();
 			UDataflowEditor* const AssetEditor = NewObject<UDataflowEditor>(AssetEditorSubsystem, NAME_None, RF_Transient);
@@ -200,18 +200,18 @@ EAssetCommandResult UAssetDefinition_DataflowAsset::OpenAssets(const FAssetOpenA
 				return EAssetCommandResult::Handled;
 			}
 		}
-	}
-	else
-	{
-		TSharedRef<SMessageDialog> MessageDialog = SNew(SMessageDialog)
-			.Title(FText(LOCTEXT("Dataflow_OpenAssetDialog_Title", "Dataflow Asset")))
-			.Message(LOCTEXT("Dataflow_OpenAssetDialog_Text", "Dataflow assets can only be changed while editing assets using them (Cloth, Flesh, Geometry Collection, ...)"))
-			.Buttons({
-				SMessageDialog::FButton(LOCTEXT("Ok", "Ok"))
-				.SetPrimary(true)
-			});
-		MessageDialog->ShowModal();
-		return EAssetCommandResult::Handled;
+		else
+		{
+			TSharedRef<SMessageDialog> MessageDialog = SNew(SMessageDialog)
+				.Title(FText(LOCTEXT("Dataflow_OpenAssetDialog_Title", "Dataflow Asset")))
+				.Message(LOCTEXT("Dataflow_OpenAssetDialog_Text", "Dataflow assets can only be changed while editing assets using them (Cloth, Flesh, Geometry Collection, ...)"))
+				.Buttons({
+					SMessageDialog::FButton(LOCTEXT("Ok", "Ok"))
+					.SetPrimary(true)
+				});
+			MessageDialog->ShowModal();
+			return EAssetCommandResult::Handled;
+		}
 	}
 	return EAssetCommandResult::Unhandled;
 }
diff --git a/Engine/Plugins/Experimental/Dataflow/Source/DataflowEditor/Private/Dataflow/DataflowSimulationPanel.cpp b/Engine/Plugins/Experimental/Dataflow/Source/DataflowEditor/Private/Dataflow/DataflowSimulationPanel.cpp
--- a/Engine/Plugins/Experimental/Dataflow/Source/DataflowEditor/Private/Dataflow/DataflowSimulationPanel.cpp
+++ b/Engine/Plugins/Experimental/Dataflow/Source/DataflowEditor/Private/Dataflow/DataflowSimulationPanel.cpp
@@ -167,7 +167,7 @@ FReply SDataflowSimulationPanel::OnClick_Record()
 {
 	if (const TSharedPtr<FDataflowSimulationScene> PreviewScene = SimulationScene.Pin())
 	{
-		PreviewScene->UpdateSimulationCache();
+		PreviewScene->RecordSimulationCache();
 	}
 	return FReply::Handled();
 }
diff --git a/Engine/Plugins/Experimental/Dataflow/Source/DataflowEditor/Private/Dataflow/DataflowSimulationScene.cpp b/Engine/Plugins/Experimental/Dataflow/Source/DataflowEditor/Private/Dataflow/DataflowSimulationScene.cpp
--- a/Engine/Plugins/Experimental/Dataflow/Source/DataflowEditor/Private/Dataflow/DataflowSimulationScene.cpp
+++ b/Engine/Plugins/Experimental/Dataflow/Source/DataflowEditor/Private/Dataflow/DataflowSimulationScene.cpp
@@ -213,11 +213,56 @@ void FDataflowSimulationScene::CreateSimulationScene()
 	SetPreviewLOD(PreviewLOD);
 }
 
-void FDataflowSimulationScene::UpdateSimulationCache()
+void FDataflowSimulationScene::RecordSimulationCache()
 {
-	if(SimulationGenerator.IsValid())
+	if(SceneDescription->CacheParams.bAsyncCaching)
 	{
-		SimulationGenerator->RequestGeneratorAction(UE::Dataflow::EDataflowGeneratorActions::StartGenerate);
+		if(SimulationGenerator.IsValid())
+		{
+			SimulationGenerator->RequestGeneratorAction(UE::Dataflow::EDataflowGeneratorActions::StartGenerate);
+		}
+	}
+	else
+	{
+		bIsRecordingCache = true;
+        SimulationTime = TimeRange[0];
+		if(AChaosCacheManager* CacheManager = Cast<AChaosCacheManager>(RootSceneActor))
+		{
+			GetWorld()->GetSubsystem<UDataflowSimulationManager>()->SetSimulationEnabled(true);
+			CacheManager->CacheMode = ECacheMode::Record;
+			CacheManager->SetObservedComponentProperties(CacheManager->CacheMode);
+			CacheManager->BeginEvaluate();
+		}
+	}
+}
+
+void FDataflowSimulationScene::UpdateSimulationCache(float& DeltaSeconds)
+{
+	if(RootSceneActor)
+	{
+		AChaosCacheManager* CacheManager = Cast<AChaosCacheManager>(RootSceneActor);
+		if(!bIsRecordingCache)
+		{
+			// Update the cached simulation at some point in time
+			CacheManager->SetStartTime(SimulationTime);
+		}
+		else
+		{
+			DeltaSeconds = 1.0f / SceneDescription->CacheParams.FrameRate;
+			SimulationTime += DeltaSeconds;
+
+			if(SimulationTime >= TimeRange[1])
+			{
+				bIsRecordingCache = false;
+				CacheManager->EndEvaluate();
+
+				GetWorld()->GetSubsystem<UDataflowSimulationManager>()->SetSimulationEnabled(false);
+						
+				CacheManager->CacheMode = ECacheMode::None;
+				CacheManager->SetObservedComponentProperties(CacheManager->CacheMode);
+				CacheManager->BeginEvaluate();
+			}
+		}
 	}
 }
 
@@ -245,6 +290,7 @@ int32 FDataflowSimulationScene::GetPreviewLOD() const
 
 void FDataflowSimulationScene::TickDataflowScene(const float DeltaSeconds)
 {
+	float DeltaTime = DeltaSeconds;
 	if(const TObjectPtr<UDataflowBaseContent>& EditorContent = GetEditorContent())
 	{
 		if (const TObjectPtr<UDataflow> DataflowGraph = EditorContent->GetDataflowAsset())
@@ -265,16 +311,14 @@ void FDataflowSimulationScene::TickDataflowScene(const float DeltaSeconds)
 		// Load the cache at some point in time
 		if(SceneDescription->CacheAsset)
 		{
-			// Update the cached simulation at some point in time
-			if(RootSceneActor)
-			{
-				Cast<AChaosCacheManager>(RootSceneActor)->SetStartTime(SimulationTime);
-			}
+			// Update the simulation cache (loading/recording)
+			UpdateSimulationCache(DeltaTime);
+			
 			// Update all the skelmesh animations at the simulation time
 			UE::Dataflow::UpdateSkeletonAnimation(PreviewActor, SimulationTime);
 		}
 	}
-	GetWorld()->Tick(ELevelTick::LEVELTICK_All, DeltaSeconds);
+	GetWorld()->Tick(ELevelTick::LEVELTICK_All, DeltaTime);
 }
 
 void FDataflowSimulationScene::AddReferencedObjects(FReferenceCollector& Collector)
diff --git a/Engine/Plugins/Experimental/Dataflow/Source/DataflowEditor/Public/Dataflow/DataflowSimulationScene.h b/Engine/Plugins/Experimental/Dataflow/Source/DataflowEditor/Public/Dataflow/DataflowSimulationScene.h
--- a/Engine/Plugins/Experimental/Dataflow/Source/DataflowEditor/Public/Dataflow/DataflowSimulationScene.h
+++ b/Engine/Plugins/Experimental/Dataflow/Source/DataflowEditor/Public/Dataflow/DataflowSimulationScene.h
@@ -42,7 +42,7 @@ public:
 	FTransform BlueprintTransform = FTransform::Identity;
 
 	/** Caching asset to be used to record the simulation  */
-	UPROPERTY(EditAnywhere, Category="Caching", DisplayName="CacheCollection")
+	UPROPERTY(EditAnywhere, Category="Caching", DisplayName="Cache Collection")
 	TObjectPtr<UChaosCacheCollection> CacheAsset = nullptr;
 
 	/** Caching params used to record the simulation */
@@ -50,7 +50,7 @@ public:
 	FDataflowPreviewCacheParams CacheParams;
 
 	/** Geometry cache asset used to extract skeletal mesh results from simulation */
-	UPROPERTY(EditAnywhere, Category = "Geometry", DisplayName="GeometryCache", meta=(EditCondition = "CacheAsset != nullptr"))
+	UPROPERTY(EditAnywhere, Category = "Geometry", DisplayName="Geometry Cache", meta=(EditCondition = "CacheAsset != nullptr"))
 	TObjectPtr<UGeometryCache> GeometryCacheAsset = nullptr;
 
 	/** SkeletalMesh interpolated from simulation. This should match the SkeletalMesh used in GenerateSurfaceBindings node */
@@ -131,8 +131,8 @@ public:
 	/** Update Scene in response to the SceneDescription changing */
 	void SceneDescriptionPropertyChanged(const FName& PropertyName);
 
-	/** Update the simulation cache */
-	void UpdateSimulationCache();
+	/** Record the simulation cache */
+	void RecordSimulationCache();
 
 	/** Get the simulation time range */
 	const FVector2f& GetTimeRange() const {return TimeRange;}
@@ -153,6 +153,9 @@ public:
 
 private:
 
+	/** Update the simulation cache at some point in time (loading/recording)*/
+	void UpdateSimulationCache(float& DeltaSeconds);
+
 	/** Bind the scene selection to the components */
 	void BindSceneSelection();
 
@@ -180,7 +183,11 @@ private:
 	/** Handle for the delegate */
 	FDelegateHandle OnObjectsReinstancedHandle;
 
+	/** Preview LOD used in the simulation viewport */
 	int32 CurrentPreviewLOD = INDEX_NONE;
+
+	/** Boolean to check if we are recording the cache or not */
+	bool bIsRecordingCache = false;
 };
 
 
diff --git a/Engine/Plugins/Experimental/Dataflow/Source/DataflowNodes/Private/Dataflow/DataflowCollectionAddScalarVertexPropertyNode.cpp b/Engine/Plugins/Experimental/Dataflow/Source/DataflowNodes/Private/Dataflow/DataflowCollectionAddScalarVertexPropertyNode.cpp
--- a/Engine/Plugins/Experimental/Dataflow/Source/DataflowNodes/Private/Dataflow/DataflowCollectionAddScalarVertexPropertyNode.cpp
+++ b/Engine/Plugins/Experimental/Dataflow/Source/DataflowNodes/Private/Dataflow/DataflowCollectionAddScalarVertexPropertyNode.cpp
@@ -116,8 +116,9 @@ FDataflowCollectionAddScalarVertexPropertyNode::FDataflowCollectionAddScalarVert
 	: FDataflowNode(InParam, InGuid)
 {
 	RegisterInputConnection(&Collection);
+	RegisterInputConnection(&AttributeKey);
 	RegisterOutputConnection(&Collection, &Collection);
-	RegisterOutputConnection(&AttributeKey);
+	RegisterOutputConnection(&AttributeKey, &AttributeKey);
 }
 
 TArray<UE::Dataflow::FRenderingParameter> FDataflowCollectionAddScalarVertexPropertyNode::GetRenderParametersImpl() const
@@ -130,15 +131,25 @@ void FDataflowCollectionAddScalarVertexPropertyNode::Evaluate(UE::Dataflow::FCon
 {
 	using namespace UE::Dataflow;
 
+	// Get the pin value if plugged
+	FCollectionAttributeKey Key = GetValue(Context, &AttributeKey, AttributeKey);
+
+	// If nothing set used the local value
+	if(Key.Attribute.IsEmpty() && Key.Group.IsEmpty())
+	{
+		Key.Group = TargetGroup.Name.ToString();
+		Key.Attribute = Name;
+	}
+
 	if (Out->IsA<FManagedArrayCollection>(&Collection))
 	{
 		// Evaluate in collection
 		FManagedArrayCollection InCollection = GetValue<FManagedArrayCollection>(Context, &Collection);
 
-		if (!Name.IsEmpty())
+		if (!Key.Attribute.IsEmpty())
 		{
-			const FName InName(Name);
-			const FName InGroup = TargetGroup.Name;
+			const FName InName(Key.Attribute);
+			const FName InGroup(Key.Group);
 			TManagedArray<float>& Scalar = InCollection.AddAttribute<float>(InName, InGroup);
 
 			const int32 MaxWeightIndex = FMath::Min(VertexWeights.Num(), Scalar.Num());
@@ -162,7 +173,7 @@ void FDataflowCollectionAddScalarVertexPropertyNode::Evaluate(UE::Dataflow::FCon
 	}
 	else if (Out->IsA<FCollectionAttributeKey>(&AttributeKey))
 	{
-		SetValue(Context, FCollectionAttributeKey(Name,"Vertices"), &AttributeKey);
+		SetValue(Context, MoveTemp(Key), &AttributeKey);
 	}
 }
 
diff --git a/Engine/Plugins/Experimental/Dataflow/Source/DataflowNodes/Public/Dataflow/DataflowCollectionAddScalarVertexPropertyNode.h b/Engine/Plugins/Experimental/Dataflow/Source/DataflowNodes/Public/Dataflow/DataflowCollectionAddScalarVertexPropertyNode.h
--- a/Engine/Plugins/Experimental/Dataflow/Source/DataflowNodes/Public/Dataflow/DataflowCollectionAddScalarVertexPropertyNode.h
+++ b/Engine/Plugins/Experimental/Dataflow/Source/DataflowNodes/Public/Dataflow/DataflowCollectionAddScalarVertexPropertyNode.h
@@ -73,22 +73,23 @@ struct DATAFLOWNODES_API FDataflowCollectionAddScalarVertexPropertyNode : public
 
 	virtual TArray<UE::Dataflow::FRenderingParameter> GetRenderParametersImpl() const override;
 
+	
 public:
 
 	UPROPERTY(Meta = (DataflowInput, DataflowOutput, DataflowPassthrough = "Collection"))
 	FManagedArrayCollection Collection;
 
 	/** The name to be set as a weight map attribute. */
-	UPROPERTY(EditAnywhere, Category = "Vertex Attribute")
+	UPROPERTY(EditAnywhere, Category = "Vertex Attribute", meta = (DisplayName = "Attribute"))
 	FString Name;
 
-	UPROPERTY(meta = (DisplayName = "AttributeKey", DataflowOutput))
+	UPROPERTY(meta = (DisplayName = "AttributeKey", DataflowInput, DataflowOutput, DataflowPassthrough = "AttributeKey"))
 	FCollectionAttributeKey AttributeKey;
 
 	UPROPERTY()
 	TArray<float> VertexWeights;
 
-	UPROPERTY(EditAnywhere, Category = "Vertex Attribute")
+	UPROPERTY(EditAnywhere, Category = "Vertex Attribute", meta = (DisplayName = "Group"))
 	FScalarVertexPropertyGroup TargetGroup;
 
 	FDataflowCollectionAddScalarVertexPropertyNode(const UE::Dataflow::FNodeParameters& InParam, FGuid InGuid = FGuid::NewGuid());
diff --git a/Engine/Source/Runtime/Experimental/Dataflow/Core/Public/Dataflow/DataflowAnyType.h b/Engine/Source/Runtime/Experimental/Dataflow/Core/Public/Dataflow/DataflowAnyType.h
--- a/Engine/Source/Runtime/Experimental/Dataflow/Core/Public/Dataflow/DataflowAnyType.h
+++ b/Engine/Source/Runtime/Experimental/Dataflow/Core/Public/Dataflow/DataflowAnyType.h
@@ -14,7 +14,7 @@ struct FDataflowAnyType
 	using FPolicyType = FDataflowAllTypesPolicy;
 	using FStorageType = void;
 
-	GENERATED_USTRUCT_BODY()
+	GENERATED_BODY()
 	DATAFLOWCORE_API static const FName TypeName;
 };
 
@@ -24,7 +24,7 @@ struct FDataflowAllTypes : public FDataflowAnyType
 {
 	using FPolicyType = FDataflowAllTypesPolicy;
 	using FStorageType = void;
-	GENERATED_USTRUCT_BODY()
+	GENERATED_BODY()
 };
 
 /**
@@ -37,7 +37,7 @@ struct FDataflowNumericTypes : public FDataflowAnyType
 	using FPolicyType = FDataflowNumericTypePolicy;
 	using FStorageType = double;
 
-	GENERATED_USTRUCT_BODY()
+	GENERATED_BODY()
 
 	UPROPERTY(EditAnywhere, Category=Value)
 	double Value = 0.0;
@@ -53,7 +53,7 @@ struct FDataflowVectorTypes : public FDataflowAnyType
 	using FPolicyType = FDataflowVectorTypePolicy;
 	using FStorageType = FVector4;
 
-	GENERATED_USTRUCT_BODY()
+	GENERATED_BODY()
 
 	UPROPERTY(EditAnywhere, Category = Value)
 	FVector4 Value = FVector4(0, 0, 0, 0);
@@ -66,10 +66,36 @@ struct FDataflowStringTypes : public FDataflowAnyType
 	using FPolicyType = FDataflowStringTypePolicy;
 	using FStorageType = FString;
 
-	GENERATED_USTRUCT_BODY()
+	GENERATED_BODY()
 
 	UPROPERTY(EditAnywhere, Category = Value)
-	FString Value;
+	FString Value = "";
+};
+
+/** Bool types  */
+USTRUCT()
+struct FDataflowBoolTypes : public FDataflowAnyType
+{
+	using FPolicyType = TDataflowSingleTypePolicy<bool>;
+	using FStorageType = bool;
+
+	GENERATED_BODY()
+
+	UPROPERTY(EditAnywhere, Category = Value)
+	bool Value = true;
+};
+
+/** Transform types  */
+USTRUCT()
+struct FDataflowTransformTypes : public FDataflowAnyType
+{
+	using FPolicyType = TDataflowSingleTypePolicy<FTransform>;
+	using FStorageType = FTransform;
+
+	GENERATED_BODY()
+
+	UPROPERTY(EditAnywhere, Category = Value)
+	FTransform Value = FTransform::Identity;
 };
 
 /**
@@ -82,10 +108,10 @@ struct FDataflowStringConvertibleTypes : public FDataflowAnyType
 	using FPolicyType = FDataflowStringConvertibleTypePolicy;
 	using FStorageType = FString;
 
-	GENERATED_USTRUCT_BODY()
-
+	GENERATED_BODY()
+	
 	UPROPERTY(EditAnywhere, Category = Value)
-	FString Value;
+	FString Value = "";
 };
 
 /** UObject types */
@@ -95,13 +121,12 @@ struct FDataflowUObjectConvertibleTypes : public FDataflowAnyType
 	using FPolicyType = FDataflowUObjectConvertibleTypePolicy;
 	using FStorageType = TObjectPtr<UObject>;
 
-	GENERATED_USTRUCT_BODY()
+	GENERATED_BODY()
 
 	UPROPERTY(EditAnywhere, Category = Value)
-	TObjectPtr<UObject> Value;
+	TObjectPtr<UObject> Value = nullptr;
 };
 
-
 /** Selection types */
 USTRUCT()
 struct FDataflowSelectionTypes : public FDataflowAnyType
@@ -113,4 +138,70 @@ struct FDataflowSelectionTypes : public FDataflowAnyType
 
 	UPROPERTY()
 	FDataflowSelection Value;
-};
\ No newline at end of file
+};
+
+/** Vector array types */
+USTRUCT()
+struct FDataflowVectorArrayTypes : public FDataflowAnyType
+{
+	using FPolicyType = FDataflowVectorArrayPolicy;
+	using FStorageType = TArray<FVector4>;
+	
+	GENERATED_BODY()
+	
+	UPROPERTY(EditAnywhere, Category = Value)
+	TArray<FVector4> Value;
+};
+
+/** Numeric array types */
+USTRUCT()
+struct FDataflowNumericArrayTypes : public FDataflowAnyType
+{
+	using FPolicyType = FDataflowNumericArrayPolicy;
+	using FStorageType = TArray<double>;
+	
+	GENERATED_BODY()
+	
+	UPROPERTY(EditAnywhere, Category = Value)
+	TArray<double> Value;
+};
+
+/** String array types */
+USTRUCT()
+struct FDataflowStringArrayTypes : public FDataflowAnyType
+{
+	using FPolicyType = FDataflowStringArrayPolicy;
+	using FStorageType = TArray<FString>;
+	
+	GENERATED_BODY()
+	
+	UPROPERTY(EditAnywhere, Category = Value)
+	TArray<FString> Value;
+};
+
+/** Bool array types */
+USTRUCT()
+struct FDataflowBoolArrayTypes : public FDataflowAnyType
+{
+	using FPolicyType = TDataflowSingleTypePolicy<TArray<bool>>;
+	using FStorageType = TArray<bool>;
+	
+	GENERATED_BODY()
+	
+	UPROPERTY(EditAnywhere, Category = Value)
+	TArray<bool> Value;
+};
+
+/** Transform array types */
+USTRUCT()
+struct FDataflowTransformArrayTypes : public FDataflowAnyType
+{
+	using FPolicyType = TDataflowSingleTypePolicy<TArray<FTransform>>;
+	using FStorageType = TArray<FTransform>;
+	
+	GENERATED_BODY()
+	
+	UPROPERTY(EditAnywhere, Category = Value)
+	TArray<FTransform> Value;
+};
+
diff --git a/Engine/Source/Runtime/Experimental/Dataflow/Core/Public/Dataflow/DataflowTypePolicy.h b/Engine/Source/Runtime/Experimental/Dataflow/Core/Public/Dataflow/DataflowTypePolicy.h
--- a/Engine/Source/Runtime/Experimental/Dataflow/Core/Public/Dataflow/DataflowTypePolicy.h
+++ b/Engine/Source/Runtime/Experimental/Dataflow/Core/Public/Dataflow/DataflowTypePolicy.h
@@ -16,6 +16,7 @@ struct TDataflowPolicyTypeName
 {
 };
 
+#define UE_DATAFLOW_MAKE_STRING(x) #x
 #define UE_DATAFLOW_POLICY_DECLARE_TYPENAME(TType) \
 template<> \
 struct TDataflowPolicyTypeName<TType> \
@@ -24,7 +25,15 @@ struct TDataflowPolicyTypeName<TType> \
 	{ \
 		return TEXT(#TType); \
 	} \
-};
+}; \
+template<> \
+struct TDataflowPolicyTypeName<TArray<TType>> \
+{  \
+	FORCEINLINE static const TCHAR* GetName() \
+	{ \
+		return TEXT(UE_DATAFLOW_MAKE_STRING(TArray<TType>)); \
+	} \
+}; 
 
 UE_DATAFLOW_POLICY_DECLARE_TYPENAME(bool)
 UE_DATAFLOW_POLICY_DECLARE_TYPENAME(uint8)
@@ -47,6 +56,14 @@ UE_DATAFLOW_POLICY_DECLARE_TYPENAME(FVector4)
 UE_DATAFLOW_POLICY_DECLARE_TYPENAME(FVector2f)
 UE_DATAFLOW_POLICY_DECLARE_TYPENAME(FVector3f)
 UE_DATAFLOW_POLICY_DECLARE_TYPENAME(FVector4f)
+UE_DATAFLOW_POLICY_DECLARE_TYPENAME(FQuat)
+UE_DATAFLOW_POLICY_DECLARE_TYPENAME(FQuat4f)
+UE_DATAFLOW_POLICY_DECLARE_TYPENAME(FLinearColor)
+UE_DATAFLOW_POLICY_DECLARE_TYPENAME(FIntPoint)
+UE_DATAFLOW_POLICY_DECLARE_TYPENAME(FIntVector3)
+UE_DATAFLOW_POLICY_DECLARE_TYPENAME(FIntVector4)
+UE_DATAFLOW_POLICY_DECLARE_TYPENAME(FRotator)
+UE_DATAFLOW_POLICY_DECLARE_TYPENAME(FTransform)
 
 UE_DATAFLOW_POLICY_DECLARE_TYPENAME(FDataflowTransformSelection)
 UE_DATAFLOW_POLICY_DECLARE_TYPENAME(FDataflowVertexSelection)
@@ -173,11 +190,23 @@ struct FDataflowNumericTypePolicy :
 {
 };
 
+struct FDataflowNumericArrayPolicy: 
+	public TDataflowMultiTypePolicy<TArray<double>, TArray<float>, TArray<int64>,
+		TArray<uint64>, TArray<int32>, TArray<uint32>, TArray<int16>, TArray<uint16>, TArray<int8>, TArray<uint8>>
+{
+};
+
 struct FDataflowVectorTypePolicy:
-	public TDataflowMultiTypePolicy<FVector2D, FVector, FVector4, FVector2f, FVector3f, FVector4f>
+	public TDataflowMultiTypePolicy<FVector2D, FVector, FVector4, FVector2f, FVector3f, FVector4f, FQuat, FQuat4f, FLinearColor, FIntPoint, FIntVector, FIntVector4, FRotator>
 {
 };
 
+struct FDataflowVectorArrayPolicy:
+	public TDataflowMultiTypePolicy<TArray<FVector2D>, TArray<FVector>, TArray<FVector4>,
+		TArray<FVector2f>, TArray<FVector3f>, TArray<FVector4f>, TArray<FQuat>, TArray<FQuat4f>,
+		TArray<FLinearColor>, TArray<FIntPoint>, TArray<FIntVector>, TArray<FIntVector4>, TArray<FRotator>>
+{
+};
 
 struct FDataflowStringTypePolicy :
 	public TDataflowMultiTypePolicy<FString, FName>
@@ -189,6 +218,11 @@ struct FDataflowSelectionTypePolicy :
 {
 };
 
+struct FDataflowStringArrayPolicy :
+	public TDataflowMultiTypePolicy<TArray<FString>, TArray<FName>>
+{
+};
+
 /**
 * string comvertible types
 * - FString / Fname
@@ -209,6 +243,7 @@ struct FDataflowStringConvertibleTypePolicy : IDataflowTypePolicy
 			|| FDataflowNumericTypePolicy::SupportsTypeStatic(InType)
 			|| FDataflowVectorTypePolicy::SupportsTypeStatic(InType)
 			|| TDataflowSingleTypePolicy<bool>::SupportsTypeStatic(InType)
+			|| TDataflowSingleTypePolicy<FTransform>::SupportsTypeStatic(InType)
 			;
 	}
 
@@ -219,6 +254,7 @@ struct FDataflowStringConvertibleTypePolicy : IDataflowTypePolicy
 			|| FDataflowNumericTypePolicy::VisitPolicyByType(RequestedType, Visitor)
 			|| FDataflowVectorTypePolicy::VisitPolicyByType(RequestedType, Visitor)
 			|| TDataflowSingleTypePolicy<bool>::VisitPolicyByType(RequestedType, Visitor)
+			|| TDataflowSingleTypePolicy<FTransform>::VisitPolicyByType(RequestedType, Visitor)
 			;
 	}
 
@@ -384,6 +420,30 @@ struct FDataflowConverter<FVector4>
 		{
 			To = FVector4{ (double)From.X, (double)From.Y, (double)From.Z, (double)From.W };
 		}
+		else if constexpr (std::is_same_v<TFromType, FQuat4f> || std::is_same_v<TFromType, FQuat>)
+		{
+			To = FVector4{ (double)From.X, (double)From.Y, (double)From.Z, (double)From.W };
+		}
+		else if constexpr (std::is_same_v<TFromType, FLinearColor>)
+		{
+			To = FVector4 {(double)From.R, (double)From.G, (double)From.B, (double)From.A };
+		}
+		else if constexpr (std::is_same_v<TFromType, FIntVector3>)
+		{
+			To = FVector4 {(double)From.X, (double)From.Y, (double)From.Z, 0};
+		}
+		else if constexpr (std::is_same_v<TFromType, FIntVector4>)
+		{
+			To = FVector4 {(double)From.X, (double)From.Y, (double)From.Z, (double)From.W};
+		}
+		else if constexpr (std::is_same_v<TFromType, FIntPoint>)
+		{
+			To = FVector4 {(double)From.X, (double)From.Y, 0, 0};
+		}
+		else if constexpr (std::is_same_v<TFromType, FRotator>)
+		{
+			To = FVector4 {(double)From.Pitch, (double)From.Yaw, (double)From.Roll, 0};
+		}
 		else
 		{
 			To = From;
@@ -413,6 +473,34 @@ struct FDataflowConverter<FVector4>
 		{
 			To = FVector4f{ (float)From.X, (float)From.Y, (float)From.Z, (float)From.W };
 		}
+		else if constexpr (std::is_same_v<TToType, FQuat>)
+		{
+			To = FQuat{ From.X, From.Y, From.Z, From.W };
+		}
+		else if constexpr (std::is_same_v<TToType, FQuat4f>)
+		{
+			To = FQuat4f{ (float)From.X, (float)From.Y, (float)From.Z, (float)From.W };
+		}
+		else if constexpr (std::is_same_v<TToType, FLinearColor>)
+		{
+			To = FLinearColor{ (float)From.X, (float)From.Y, (float)From.Z, (float)From.W };
+		}
+		else if constexpr (std::is_same_v<TToType, FIntPoint>)
+		{
+			To = FIntPoint{ (int32)From.X, (int32)From.Y };
+		}
+		else if constexpr (std::is_same_v<TToType, FIntVector3>)
+		{
+			To = FIntVector3{ (int32)From.X, (int32)From.Y, (int32)From.Z };
+		}
+		else if constexpr (std::is_same_v<TToType, FIntVector4>)
+		{
+			To = FIntVector4{ (int32)From.X, (int32)From.Y, (int32)From.Z, (int32)From.W };
+		}
+		else if constexpr (std::is_same_v<TToType, FRotator>)
+		{
+			To = FRotator{ (double)From.X, (double)From.Y, (double)From.Z};
+		}
 		else
 		{
 			To = From;
@@ -420,7 +508,6 @@ struct FDataflowConverter<FVector4>
 	}
 };
 
-
 template <>
 struct FDataflowConverter<FDataflowSelection>
 {
@@ -435,4 +522,28 @@ struct FDataflowConverter<FDataflowSelection>
 	{
 		To.Initialize(From);
 	}
-};
\ No newline at end of file
+};
+
+template <typename ArrayType>
+struct FDataflowConverter<TArray<ArrayType>>
+{
+	template <typename TFromType>
+	static void From(const TFromType& From, TArray<ArrayType>& To)
+	{
+		To.SetNum(From.Num());
+		for (int32 ArrayIndex = 0; ArrayIndex < From.Num(); ++ArrayIndex)
+		{
+			FDataflowConverter<ArrayType>::From(From[ArrayIndex], To[ArrayIndex]);
+		}
+	}
+
+	template <typename TToType>
+	static void To(const TArray<ArrayType>& From, TToType& To)
+	{
+		To.SetNum(From.Num());
+		for (int32 ArrayIndex = 0; ArrayIndex < From.Num(); ++ArrayIndex)
+		{
+			FDataflowConverter<ArrayType>::To(From[ArrayIndex], To[ArrayIndex]);
+		}
+	}
+};
