Extend FIoChunkEncoding::Decode to accept additional flags specializing behavior.

#rnx

### FIoChunkEncoding
- Add a new flags enum EIoDecodeFlags which allows the caller to provide additional info when invoking FIoChunkEncoding::Decode
- Currently the only new behavior is EIoDecodeFlags::WriteOnly which indicates that we should not read from the provided output buffer contained in the 'OutRawData' parameter
-- This means when decompressing we will need to allocate a temp buffer to decompress to first, then memcpy to the final output buffer. This was already being done under certain conditions but this flag will enforce it in all cases unless the data is not compressed in which case a simple memcpy will suffice.

### FOnDemandIoBackend
- The memory buffer provided for a request can has the 'EIoReadOptionsFlags::HardwareTargetBuffer' flag associated with it, which indicates that the buffer is not suitable for reading from and attempting to do so will either fail or be very slow. If we detect this we  should invoke with the 'EIoDecodeFlags::WriteOnly' flag to make sure that this request is honored.
- Note that we only do this when the new cvar 'ias.EnableWriteOnlyDecoding' is set to true, but it is set to false by default. This will allow us to enable the behavior selectively at first.

diff --git a/Engine/Source/Runtime/Core/Private/IO/IoChunkEncoding.cpp b/Engine/Source/Runtime/Core/Private/IO/IoChunkEncoding.cpp
--- a/Engine/Source/Runtime/Core/Private/IO/IoChunkEncoding.cpp
+++ b/Engine/Source/Runtime/Core/Private/IO/IoChunkEncoding.cpp
@@ -230,7 +230,7 @@ bool FIoChunkEncoding::Decode(
 	return Decode(Params, EncodedBlocks, OutRawData);
 }
 
-bool FIoChunkEncoding::Decode(const FIoChunkDecodingParams& Params, FMemoryView EncodedBlocks, FMutableMemoryView OutRawData)
+bool FIoChunkEncoding::Decode(const FIoChunkDecodingParams& Params, FMemoryView EncodedBlocks, FMutableMemoryView OutRawData, EIoDecodeFlags Flags)
 {
 	TRACE_CPUPROFILER_EVENT_SCOPE(FIoChunkEncoding::Decode);
 
@@ -280,7 +280,9 @@ bool FIoChunkEncoding::Decode(const FIoChunkDecodingParams& Params, FMemoryView
 		DecryptedBlockBuffer = MakeUnique<uint8[]>(DecryptedScratchBufferSize);
 	}
 
+	const bool bForceTempBuffer = EnumHasAnyFlags(Flags, EIoDecodeFlags::WriteOnly);
 	const bool bVerifyBlocks = BlockHash.IsEmpty() == false;
+
 	for (uint32 BlockIndex = FirstBlockIndex; BlockIndex <= LastBlockIndex; BlockIndex++)
 	{
 		const uint64 RawBlockSize = BlockIndex == BlockCount - 1 ? LastRawBlockSize : BlockSize;
@@ -314,7 +316,7 @@ bool FIoChunkEncoding::Decode(const FIoChunkDecodingParams& Params, FMemoryView
 
 		if (CompressedBlockSize < RawBlockSize)
 		{
-			if (RawBlockReadSize == RawBlockSize)
+			if (RawBlockReadSize == RawBlockSize && !bForceTempBuffer)
 			{
 				if (!FCompression::UncompressMemory(Params.CompressionFormat, OutRawData.GetData(), int32(RawBlockReadSize), BlockView.GetData(), CompressedBlockSize))
 				{
diff --git a/Engine/Source/Runtime/Core/Public/IO/IoChunkEncoding.h b/Engine/Source/Runtime/Core/Public/IO/IoChunkEncoding.h
--- a/Engine/Source/Runtime/Core/Public/IO/IoChunkEncoding.h
+++ b/Engine/Source/Runtime/Core/Public/IO/IoChunkEncoding.h
@@ -5,6 +5,7 @@
 #include "IO/IoOffsetLength.h"
 #include "IO/IoStatus.h"
 #include "Memory/MemoryView.h"
+#include "Misc/EnumClassFlags.h"
 #include "UObject/NameTypes.h"
 
 class FIoBuffer;
@@ -38,6 +39,16 @@ struct FIoChunkDecodingParams
 	TConstArrayView<FIoBlockHash> BlockHash;
 };
 
+/** Additional options for FIoChunkEncoding::Decode calls */
+enum class EIoDecodeFlags : uint8
+{
+	None		= 0,
+	/** Avoid operations that read from the output buffer provided. This option may use more temp scratch memory */
+	WriteOnly	= 1 << 0
+};
+
+ENUM_CLASS_FLAGS(EIoDecodeFlags);
+
 /**
  * Encodes data into a set of encrypted and compressed blocks.
  * The chunk encoding information is enocoded into a 16 byte
@@ -74,7 +85,7 @@ public:
 
 	static CORE_API bool Encode(const FIoChunkEncodingParams& Params, FMemoryView RawData, FIoBuffer& OutEncodedData);
 	static CORE_API bool Encode(const FIoChunkEncodingParams& Params, FMemoryView RawData, FIoBuffer& OutHeader, FIoBuffer& OutEncodedBlocks);
-	static CORE_API bool Decode(const FIoChunkDecodingParams& Params, FMemoryView EncodedBlocks, FMutableMemoryView OutRawData);
+	static CORE_API bool Decode(const FIoChunkDecodingParams& Params, FMemoryView EncodedBlocks, FMutableMemoryView OutRawData, EIoDecodeFlags Flags = EIoDecodeFlags::None);
 	static CORE_API bool Decode(FMemoryView EncodedData, FName CompressionFormat, FMemoryView EncryptionKey, FMutableMemoryView OutRawData, uint64 Offset = 0);
 
 	static CORE_API TIoStatusOr<FIoOffsetAndLength> GetChunkRange(uint64 TotalRawSize, uint32 RawBlockSize, TConstArrayView<uint32> EncodedBlockSize, uint64 RawOffset, uint64 RawSize);
diff --git a/Engine/Source/Runtime/Experimental/IoStore/OnDemand/Private/OnDemandIoDispatcherBackend.cpp b/Engine/Source/Runtime/Experimental/IoStore/OnDemand/Private/OnDemandIoDispatcherBackend.cpp
--- a/Engine/Source/Runtime/Experimental/IoStore/OnDemand/Private/OnDemandIoDispatcherBackend.cpp
+++ b/Engine/Source/Runtime/Experimental/IoStore/OnDemand/Private/OnDemandIoDispatcherBackend.cpp
@@ -193,6 +193,13 @@ static FAutoConsoleVariableRef CVar_DistributedEndpointAttemptCount(
 	TEXT("Number of times we should try to resolve a distributed endpoint befor eusing the fallback url (if there is one)")
 );
 
+bool GIasEnableWriteOnlyDecoding = false;
+static FAutoConsoleVariableRef CVar_WriteOnlyDecoding(
+	TEXT("ias.EnableWriteOnlyDecoding"),
+	GIasEnableWriteOnlyDecoding,
+	TEXT("Enables the use of 'WriteOnly' flag when decoding to buffers with the 'HardwareTargetBuffer' flag")
+);
+
 
 // These priorities are indexed using the cvar below
 static UE::Tasks::ETaskPriority GCompleteMaterializeTaskPriorities[] =
@@ -1252,8 +1259,11 @@ void FOnDemandIoBackend::CompleteRequest(FChunkRequest* ChunkRequest)
 			TRACE_CPUPROFILER_EVENT_SCOPE(IasBackend::DecodeBlocks);
 			const uint64 RawSize = FMath::Min<uint64>(Request->Options.GetSize(), ChunkRequest->Params.ChunkInfo.RawSize());
 			Request->CreateBuffer(RawSize);
-			DecodingParams.RawOffset = Request->Options.GetOffset(); 
-			bDecoded = FIoChunkEncoding::Decode(DecodingParams, Chunk.GetView(), Request->GetBuffer().GetMutableView());
+			DecodingParams.RawOffset = Request->Options.GetOffset();
+
+			const EIoDecodeFlags Options = GIasEnableWriteOnlyDecoding && EnumHasAnyFlags(Request->Options.GetFlags(), EIoReadOptionsFlags::HardwareTargetBuffer) ? EIoDecodeFlags::WriteOnly : EIoDecodeFlags::None;
+
+			bDecoded = FIoChunkEncoding::Decode(DecodingParams, Chunk.GetView(), Request->GetBuffer().GetMutableView(), Options);
 
 			if (!bDecoded)
 			{
