Add FDataflowIMage type to dataflow core

diff --git a/Engine/Source/Runtime/Experimental/Dataflow/Core/DataflowCore.Build.cs b/Engine/Source/Runtime/Experimental/Dataflow/Core/DataflowCore.Build.cs
--- a/Engine/Source/Runtime/Experimental/Dataflow/Core/DataflowCore.Build.cs
+++ b/Engine/Source/Runtime/Experimental/Dataflow/Core/DataflowCore.Build.cs
@@ -14,6 +14,7 @@ namespace UnrealBuildTool.Rules
 					"Chaos",
 					"DeveloperSettings",
 					"InputCore",
+					"ImageCore",
 					"ApplicationCore"
 				}
 			);
diff --git a/Engine/Source/Runtime/Experimental/Dataflow/Core/Private/Dataflow/DataflowImage.cpp b/Engine/Source/Runtime/Experimental/Dataflow/Core/Private/Dataflow/DataflowImage.cpp
new file mode 100644
--- /dev/null
+++ b/Engine/Source/Runtime/Experimental/Dataflow/Core/Private/Dataflow/DataflowImage.cpp
@@ -0,0 +1,72 @@
+// Copyright Epic Games, Inc. All Rights Reserved.
+
+#include "Dataflow/DataflowImage.h"
+
+int32 FDataflowImage::GetWidth() const
+{
+	return Image.GetWidth();
+}
+
+int32 FDataflowImage::GetHeight() const
+{
+	return Image.GetHeight();
+}
+
+const FImage& FDataflowImage::GetImage() const
+{
+	return Image;
+}
+
+void FDataflowImage::CreateR32F(int32 Width, int32 Height)
+{
+	Image.Init(Width, Height, ERawImageFormat::R32F, EGammaSpace::Linear);
+}
+
+void FDataflowImage::CreateRGBA32F(int32 Width, int32 Height)
+{
+	Image.Init(Width, Height, ERawImageFormat::RGBA32F, EGammaSpace::Linear);
+}
+
+bool FDataflowImage::CopyRGBAPixels(TArrayView64<FVector4f> Pixels)
+{
+	if (Image.Format != ERawImageFormat::Type::RGBA32F)
+	{
+		return false;
+	}
+	if ((GetWidth() * GetHeight()) != Pixels.Num())
+	{
+		return false;
+	}
+	const FImageView SrcImage(Pixels.GetData(), GetWidth(), GetHeight(), /*NumSlices*/1, ERawImageFormat::Type::RGBA32F, EGammaSpace::Linear);
+	FImageCore::CopyImage(SrcImage, Image);
+	return true;
+}
+
+void FDataflowImage::ConvertToRGBA32F()
+{
+	Image.ChangeFormat(ERawImageFormat::RGBA32F, EGammaSpace::Linear);
+}
+
+bool FDataflowImage::Serialize(FArchive& Ar)
+{
+	FImageInfo Info(Image);
+	uint8 FormatAsInt = static_cast<uint8>(Info.Format);
+	uint8 GammaSpaceAsInt = static_cast<uint8>(Info.GammaSpace);
+
+	Ar << Info.SizeX;
+	Ar << Info.SizeY;
+	Ar << Info.NumSlices;
+	Ar << FormatAsInt;
+	Ar << GammaSpaceAsInt;
+
+	if (Ar.IsLoading())
+	{
+		Info.Format = static_cast<ERawImageFormat::Type>(FormatAsInt);
+		Info.GammaSpace = static_cast<EGammaSpace>(GammaSpaceAsInt);
+		Image.Init(Info);
+	}
+
+	Ar << Image.RawData;
+
+	return true;
+}
\ No newline at end of file
diff --git a/Engine/Source/Runtime/Experimental/Dataflow/Core/Public/Dataflow/DataflowImage.h b/Engine/Source/Runtime/Experimental/Dataflow/Core/Public/Dataflow/DataflowImage.h
new file mode 100644
--- /dev/null
+++ b/Engine/Source/Runtime/Experimental/Dataflow/Core/Public/Dataflow/DataflowImage.h
@@ -0,0 +1,55 @@
+// Copyright Epic Games, Inc. All Rights Reserved.
+
+#pragma once
+
+#include "ImageCore.h"
+
+#include "DataflowImage.generated.h"
+
+/** 
+* Represents image for dataflow
+* type is constrained to Float32 with 1 or 4 channels
+*/
+USTRUCT()
+struct FDataflowImage
+{
+	GENERATED_USTRUCT_BODY()
+
+	int32 GetWidth() const;
+	int32 GetHeight() const;
+
+	/** get the readonly underlying image object */
+	const FImage& GetImage() const;
+
+	/** Create a single channel float format image */
+	DATAFLOWCORE_API void CreateR32F(int32 Width, int32 Height);
+
+	/** Create a four channels float format image */
+	DATAFLOWCORE_API void CreateRGBA32F(int32 Width, int32 Height);
+
+	/**
+	* Copy RGBA32F pixels to the image
+	* Number of pixels must match and format must be RGBA32F
+	* return false if the copy could not be done
+	*/
+	DATAFLOWCORE_API bool CopyRGBAPixels(TArrayView64<FVector4f> Pixels);
+
+	/** 
+	* Convert the current image to a 4 channel float pixel format 
+	* previous data is kept 
+	*/
+	DATAFLOWCORE_API void ConvertToRGBA32F();
+
+	bool Serialize(FArchive& Ar);
+
+private:
+	FImage Image;
+};
+
+template<> struct TStructOpsTypeTraits<FDataflowImage> : public TStructOpsTypeTraitsBase2<FDataflowImage>
+{
+	enum
+	{
+		WithSerializer = true,
+	};
+};
\ No newline at end of file
