Geometry collection : Fix potential crash when getting mass and extents on a cluster that does not have geometry

diff --git a/Engine/Source/Runtime/Experimental/GeometryCollectionEngine/Private/GeometryCollection/GeometryCollectionComponent.cpp b/Engine/Source/Runtime/Experimental/GeometryCollectionEngine/Private/GeometryCollection/GeometryCollectionComponent.cpp
--- a/Engine/Source/Runtime/Experimental/GeometryCollectionEngine/Private/GeometryCollection/GeometryCollectionComponent.cpp
+++ b/Engine/Source/Runtime/Experimental/GeometryCollectionEngine/Private/GeometryCollection/GeometryCollectionComponent.cpp
@@ -6885,6 +6885,9 @@ void UGeometryCollectionComponent::GetMassAndExtents(int32 ItemIndex, float& Out
 {
 	using FGeometryCollectionPtr = const TSharedPtr<FGeometryCollection, ESPMode::ThreadSafe>;
 
+	static FName BoundingBoxAttributeName = "BoundingBox";
+	static FName MassAttributeName = "Mass";
+
 	OutMass = 0.0f;
 	OutExtents = FBox(EForceInit::ForceInitToZero);
 
@@ -6892,10 +6895,10 @@ void UGeometryCollectionComponent::GetMassAndExtents(int32 ItemIndex, float& Out
 	if (RestCollection && RestCollection->GetGeometryCollection())
 	{
 		const FGeometryCollection& Collection = *RestCollection->GetGeometryCollection();
-		if (const TManagedArray<float>* CollectionMass = Collection.FindAttribute<float>(TEXT("Mass"), FTransformCollection::TransformGroup))
+		if (const TManagedArray<float>* CollectionMass = Collection.FindAttribute<float>(MassAttributeName, FTransformCollection::TransformGroup))
 		{
-			const TManagedArray<FBox>* TransformBoundingBoxes = Collection.FindAttribute<FBox>(TEXT("BoundingBox"), FTransformCollection::TransformGroup);
-			const TManagedArray<FBox>* GeoBoundingBoxes = Collection.FindAttribute<FBox>(TEXT("BoundingBox"), FGeometryCollection::GeometryGroup);
+			const TManagedArray<FBox>* TransformBoundingBoxes = Collection.FindAttribute<FBox>(BoundingBoxAttributeName, FTransformCollection::TransformGroup);
+			const TManagedArray<FBox>* GeoBoundingBoxes = Collection.FindAttribute<FBox>(BoundingBoxAttributeName, FGeometryCollection::GeometryGroup);
 
 			int32 TransformIndex = INDEX_NONE;
 			FGeometryCollectionItemIndex GCItemIndex = FGeometryCollectionItemIndex::CreateFromExistingItemIndex(ItemIndex);
@@ -6914,7 +6917,11 @@ void UGeometryCollectionComponent::GetMassAndExtents(int32 ItemIndex, float& Out
 						}
 						else if (GeoBoundingBoxes)
 						{
-							OutExtents += (*GeoBoundingBoxes)[Collection.TransformToGeometryIndex[ChildTramsformIndex]];
+							const int32 GeometryIndex = Collection.TransformToGeometryIndex[TransformIndex];
+							if (GeoBoundingBoxes->IsValidIndex(GeometryIndex))
+							{
+								OutExtents += (*GeoBoundingBoxes)[GeometryIndex];
+							}
 						}
 					}
 				}
@@ -6927,9 +6934,13 @@ void UGeometryCollectionComponent::GetMassAndExtents(int32 ItemIndex, float& Out
 				{
 					OutExtents = (*TransformBoundingBoxes)[TransformIndex];
 				}
-				else
+				else if (GeoBoundingBoxes)
 				{
-					OutExtents = (*GeoBoundingBoxes)[Collection.TransformToGeometryIndex[TransformIndex]];
+					const int32 GeometryIndex = Collection.TransformToGeometryIndex[TransformIndex];
+					if (GeoBoundingBoxes->IsValidIndex(GeometryIndex))
+					{
+						OutExtents = (*GeoBoundingBoxes)[GeometryIndex];
+					}
 				}
 			}
 		}
