Fix glitchy GC rendering for Modular Vehicles

diff --git a/Engine/Source/Runtime/Experimental/GeometryCollectionEngine/Private/GeometryCollection/GeometryCollectionComponent.cpp b/Engine/Source/Runtime/Experimental/GeometryCollectionEngine/Private/GeometryCollection/GeometryCollectionComponent.cpp
--- a/Engine/Source/Runtime/Experimental/GeometryCollectionEngine/Private/GeometryCollection/GeometryCollectionComponent.cpp
+++ b/Engine/Source/Runtime/Experimental/GeometryCollectionEngine/Private/GeometryCollection/GeometryCollectionComponent.cpp
@@ -3253,6 +3253,11 @@ void UGeometryCollectionComponent::OnUpdateTransform(EUpdateTransformFlags Updat
 		MarkRenderTransformDirty();
 		MarkRenderDynamicDataDirty();
 	}
+	
+	if (bForceUpdateActiveTransforms)
+	{
+		RebaseDynamicCollectionTransformsOnNewWorldTransform();
+	}
 }
 
 bool UGeometryCollectionComponent::HasAnySockets() const
