Support for gyroscopic effects in Chaos. Adds Utilities::Solve for FMatrix33

diff --git a/Engine/Source/Runtime/Core/Public/UObject/FortniteMainBranchObjectVersion.h b/Engine/Source/Runtime/Core/Public/UObject/FortniteMainBranchObjectVersion.h
--- a/Engine/Source/Runtime/Core/Public/UObject/FortniteMainBranchObjectVersion.h
+++ b/Engine/Source/Runtime/Core/Public/UObject/FortniteMainBranchObjectVersion.h
@@ -535,6 +535,9 @@ struct FFortniteMainBranchObjectVersion
 		// Added support for display name in the Visual Logger
 		VisualLoggerSupportDisplayName,
 
+		// Added support for the GyroscopicTorque flag in CVD
+		GyroscopicTorquesSupportInChaosVisualDebugger,
+
 		// -----<new versions can be added above this line>-------------------------------------------------
 		VersionPlusOne,
 		LatestVersion = VersionPlusOne - 1
diff --git a/Engine/Source/Runtime/Engine/Classes/Components/PrimitiveComponent.h b/Engine/Source/Runtime/Engine/Classes/Components/PrimitiveComponent.h
--- a/Engine/Source/Runtime/Engine/Classes/Components/PrimitiveComponent.h
+++ b/Engine/Source/Runtime/Engine/Classes/Components/PrimitiveComponent.h
@@ -2718,6 +2718,14 @@ public:
 	UFUNCTION(BlueprintPure, Category="Physics")
 	ENGINE_API virtual bool GetUpdateKinematicFromSimulation() const;
 
+	/** Enabled/disables whether this body is affected by gyroscopic torque, mainly useful for long/thin objects that spin */
+	UFUNCTION(BlueprintCallable, Category = "Physics")
+	ENGINE_API void SetGyroscopicTorqueEnabled(bool bInGyroscopicTorqueEnabled);
+
+	/** Returns whether this component is affected by gyroscopic torque. */
+	UFUNCTION(BlueprintPure, Category = "Physics")
+	ENGINE_API virtual bool GetGyroscopicTorqueEnabled() const;
+
 	/** Sets the linear damping of this component. */
 	UFUNCTION(BlueprintCallable, Category="Physics")
 	ENGINE_API virtual void SetLinearDamping(float InDamping);
diff --git a/Engine/Source/Runtime/Engine/Classes/PhysicsEngine/BodyInstance.h b/Engine/Source/Runtime/Engine/Classes/PhysicsEngine/BodyInstance.h
--- a/Engine/Source/Runtime/Engine/Classes/PhysicsEngine/BodyInstance.h
+++ b/Engine/Source/Runtime/Engine/Classes/PhysicsEngine/BodyInstance.h
@@ -995,6 +995,8 @@ public:
 	ENGINE_API void SetGravityGroupIndex(int32 NewGravityGroupIndex);
 	/** Enables/disables whether this body, when kinematic, is updated from the simulation rather than when setting the kinematic target. */
 	ENGINE_API void SetUpdateKinematicFromSimulation(bool bUpdateKinematicFromSimulation);
+	/** Enabled/disables whether this body is affected by gyroscopic torque, mainly useful for long/thin objects that spin */
+	ENGINE_API void SetGyroscopicTorqueEnabled(bool bInGyroscopicTorqueEnabled);
 	/** Enables/disables contact modification */
 	ENGINE_API void SetContactModification(bool bNewContactModification);
 	/** Enables/disabled smoothed edge collisions */
diff --git a/Engine/Source/Runtime/Engine/Private/PhysicsEngine/BodyInstance.cpp b/Engine/Source/Runtime/Engine/Private/PhysicsEngine/BodyInstance.cpp
--- a/Engine/Source/Runtime/Engine/Private/PhysicsEngine/BodyInstance.cpp
+++ b/Engine/Source/Runtime/Engine/Private/PhysicsEngine/BodyInstance.cpp
@@ -1286,7 +1286,7 @@ void TInitBodiesHelperBase<InAllocatorType>::CreateActor_AssumesLocked(FBodyInst
 		FPhysicsInterface::SetMaxLinearVelocity_AssumesLocked(Instance->ActorHandle, TNumericLimits<float>::Max());
 		FPhysicsInterface::SetSmoothEdgeCollisionsEnabled_AssumesLocked(Instance->ActorHandle, Instance->bSmoothEdgeCollisions);
 		FPhysicsInterface::SetInertiaConditioningEnabled_AssumesLocked(Instance->ActorHandle, Instance->IsInertiaConditioningEnabled());
-
+		FPhysicsInterface::SetGyroscopicTorqueEnabled_AssumesLocked(Instance->ActorHandle, Instance->bGyroscopicTorqueEnabled);
 		FPhysicsInterface::SetPositionSolverIterationCount_AssumesLocked(Instance->ActorHandle, Instance->PositionSolverIterationCount);
 		FPhysicsInterface::SetVelocitySolverIterationCount_AssumesLocked(Instance->ActorHandle, Instance->VelocitySolverIterationCount);
 		FPhysicsInterface::SetProjectionSolverIterationCount_AssumesLocked(Instance->ActorHandle, Instance->ProjectionSolverIterationCount);
@@ -3808,6 +3808,24 @@ void FBodyInstance::SetUpdateKinematicFromSimulation(bool bInUpdateKinematicFrom
 	}
 }
 
+void FBodyInstance::SetGyroscopicTorqueEnabled(bool bInGyroscopicTorqueEnabled)
+{
+	if (bGyroscopicTorqueEnabled != bInGyroscopicTorqueEnabled)
+	{
+		bGyroscopicTorqueEnabled = bInGyroscopicTorqueEnabled;
+
+		{
+			FPhysicsCommand::ExecuteWrite(ActorHandle, [&](const FPhysicsActorHandle& Actor)
+				{
+					if (FPhysicsInterface::IsRigidBody(Actor))
+					{
+						FPhysicsInterface::SetGyroscopicTorqueEnabled_AssumesLocked(Actor, bGyroscopicTorqueEnabled);
+					}
+				});
+		}
+	}
+}
+
 
 void FBodyInstance::SetContactModification(bool bNewContactModification)
 {
diff --git a/Engine/Source/Runtime/Engine/Private/PhysicsEngine/ImmediatePhysics/ImmediatePhysicsChaos/ImmediatePhysicsActorHandle_Chaos.cpp b/Engine/Source/Runtime/Engine/Private/PhysicsEngine/ImmediatePhysics/ImmediatePhysicsChaos/ImmediatePhysicsActorHandle_Chaos.cpp
--- a/Engine/Source/Runtime/Engine/Private/PhysicsEngine/ImmediatePhysics/ImmediatePhysicsChaos/ImmediatePhysicsActorHandle_Chaos.cpp
+++ b/Engine/Source/Runtime/Engine/Private/PhysicsEngine/ImmediatePhysics/ImmediatePhysicsChaos/ImmediatePhysicsActorHandle_Chaos.cpp
@@ -261,6 +261,7 @@ namespace ImmediatePhysics_Chaos
 					Dynamic->SetGravityEnabled(BodyInstance->bEnableGravity);
 					Dynamic->SetGravityGroupIndex(BodyInstance->GravityGroupIndex);
 					Dynamic->SetUpdateKinematicFromSimulation(BodyInstance->bUpdateKinematicFromSimulation);
+					Dynamic->SetGyroscopicTorqueEnabled(BodyInstance->bGyroscopicTorqueEnabled);
 				}
 				Dynamic->SetDisabled(true);
 			}
diff --git a/Engine/Source/Runtime/Engine/Private/PrimitiveComponentPhysics.cpp b/Engine/Source/Runtime/Engine/Private/PrimitiveComponentPhysics.cpp
--- a/Engine/Source/Runtime/Engine/Private/PrimitiveComponentPhysics.cpp
+++ b/Engine/Source/Runtime/Engine/Private/PrimitiveComponentPhysics.cpp
@@ -572,6 +572,25 @@ bool UPrimitiveComponent::GetUpdateKinematicFromSimulation() const
 	return false;
 }
 
+void UPrimitiveComponent::SetGyroscopicTorqueEnabled(bool bEnable)
+{
+	FBodyInstance* BI = GetBodyInstance();
+	if (BI)
+	{
+		BI->SetGyroscopicTorqueEnabled(bEnable);
+	}
+}
+
+bool UPrimitiveComponent::GetGyroscopicTorqueEnabled() const
+{
+	FBodyInstance* BI = GetBodyInstance();
+	if (BI)
+	{
+		return BI->bGyroscopicTorqueEnabled;
+	}
+	return false;
+}
+
 void UPrimitiveComponent::SetLinearDamping(float InDamping)
 {
 	FBodyInstance* BI = GetBodyInstance();
diff --git a/Engine/Source/Runtime/Experimental/Chaos/Private/Chaos/Evolution/PBDMinEvolution.cpp b/Engine/Source/Runtime/Experimental/Chaos/Private/Chaos/Evolution/PBDMinEvolution.cpp
--- a/Engine/Source/Runtime/Experimental/Chaos/Private/Chaos/Evolution/PBDMinEvolution.cpp
+++ b/Engine/Source/Runtime/Experimental/Chaos/Private/Chaos/Evolution/PBDMinEvolution.cpp
@@ -263,7 +263,12 @@ namespace Chaos
 				const FReal LinearDrag = FMath::Min(FReal(1), Particle.LinearEtherDrag() * Dt);
 				const FReal AngularDrag = FMath::Min(FReal(1), Particle.AngularEtherDrag() * Dt);
 				const FVec3 V = FMath::Lerp(Particle.GetV() + DV, TargetV, LinearDrag);
-				const FVec3 W = FMath::Lerp(Particle.GetW() + DW, TargetW, AngularDrag);
+				FVec3 W = FMath::Lerp(Particle.GetW() + DW, TargetW, AngularDrag);
+
+				if (Particle.GyroscopicTorqueEnabled())
+				{
+					W = Utilities::GetAngularVelocityAdjustedForGyroscopicTorques(RCoM, Particle.I(), W, Dt);
+				}
 
 				// New position
 				const FVec3 PCoM = XCoM + V * Dt;
diff --git a/Engine/Source/Runtime/Experimental/Chaos/Private/Chaos/PBDRigidsEvolutionGBF.cpp b/Engine/Source/Runtime/Experimental/Chaos/Private/Chaos/PBDRigidsEvolutionGBF.cpp
--- a/Engine/Source/Runtime/Experimental/Chaos/Private/Chaos/PBDRigidsEvolutionGBF.cpp
+++ b/Engine/Source/Runtime/Experimental/Chaos/Private/Chaos/PBDRigidsEvolutionGBF.cpp
@@ -855,6 +855,14 @@ void FPBDRigidsEvolutionGBF::Integrate(FReal Dt)
 				const FReal AngularMultiplier = FMath::Max(FReal(0), FReal(1) - AngularDrag);
 				W *= AngularMultiplier;
 
+				// Apply gyroscopic as a velocity change. This could be applied before the clamps, or after,
+				// but the user might be surprised to see an angular velocity that exceeds the clamp
+				// so do it before.
+				if (Particle.GyroscopicTorqueEnabled())
+				{
+					W = Utilities::GetAngularVelocityAdjustedForGyroscopicTorques(Particle.RCom(), Particle.I(), W, Dt);
+				}
+
 				FReal LinearSpeedSq = V.SizeSquared();
 				if (LinearSpeedSq > Particle.MaxLinearSpeedSq())
 				{
diff --git a/Engine/Source/Runtime/Experimental/Chaos/Private/Chaos/Utilities.cpp b/Engine/Source/Runtime/Experimental/Chaos/Private/Chaos/Utilities.cpp
--- a/Engine/Source/Runtime/Experimental/Chaos/Private/Chaos/Utilities.cpp
+++ b/Engine/Source/Runtime/Experimental/Chaos/Private/Chaos/Utilities.cpp
@@ -19,5 +19,86 @@ namespace Chaos
 		{
 			return Solver->GetPhysicsResultsTime_External();
 		}
+
+		// Use this to enable multiple Newton-Raphson steps rather than just a single one. However, note
+		// that the benefit is tiny in normal cases - typically the first step achieves 95% of the convergence.
+#define USE_MULTI_STEP_GYROSCOPIC_CALCULATION 0
+
+		FVec3 GetAngularVelocityAdjustedForGyroscopicTorques(const FRotation3& Q, const FVec3& I, const FVec3& W, const FReal Dt)
+		{
+			if (Dt <= 0)
+			{
+				return W;
+			}
+			// See https://gdcvault.com/play/1022197
+			//
+			// The rotational motion is decomposed, and so we'll just apply the unforced part here by integrating:
+			//
+			// I WDot + W x I W = Torque
+			//
+			// Discretising WDot = (W2 - W1) / Dt
+			//
+			// and using W2 etc (i.e. solving implicitly) and only considering Torque = 0 
+			//
+			// I2 (W2 - W1) + Dt W2 x I2 W2 = 0
+			//
+			// We want to solve for W2 (noting that we don't have I2 yet).
+			//
+			// Newton-Raphson is used to solve problems in the form: f(x) = 0
+			//
+			// for a single var x, then iterate: x' = x - f(x) / f'(x)
+			//
+			// for multiple variables (X is a vector so call it x, y, z), the 1/f'(x) turns into multiply by
+			// the inverse of the Jacobean where J = [df/dx, df/dy, df/dz]
+			//
+			// Then we have X' = X - JInv(X) * F(X)
+
+			// Angular velocity in the space of the body, so that the inertia tensor is constant.
+			FVec3 WBody = Q.UnrotateVector(W);
+
+			const FMatrix33 IBody(I.X, 0, 0, 0, I.Y, 0, 0, 0, I.Z);
+
+#if USE_MULTI_STEP_GYROSCOPIC_CALCULATION
+			FVec3 WBodyOrig = WBody;
+#endif
+
+			// Note that FMatrix33::operator* applies from right to left, so the code here uses
+			// Utilities::Multiply, which operates from left to right, matching normal notation.
+
+#if USE_MULTI_STEP_GYROSCOPIC_CALCULATION
+	// Newton-Raphson iteration
+			for (int i = 0; i != 5; ++i)
+#endif
+			{
+				// Evaluate F(X) - i.e. the terms which if re-evaluated after the solve we want to be zero
+				// (because we are solving for the part where there is no external torque). Note that IBody
+				// * (WBody - WBodyOrig) is always zero on the first step of Newton-Raphson.
+				const FVec3 F =
+#if USE_MULTI_STEP_GYROSCOPIC_CALCULATION
+					Utilities::Multiply(IBody, (WBody - WBodyOrig)) +
+#endif
+					Dt * WBody.Cross(Utilities::Multiply(IBody, WBody));
+
+				// Jacobian. We essentially differentiate F with respect to WBody, noting that IBody is constant.
+				// The second term is differentiating a product, so D(uv) = D(u)v + uD(v)
+				const FMatrix33 J = IBody +
+					(Utilities::Multiply(Utilities::CrossProductMatrix(WBody), IBody) -
+						Utilities::CrossProductMatrix(Utilities::Multiply(IBody, WBody))) * Dt;
+
+				// It is more efficient to use a direct solve rather than calculating the inverse of J and
+				// then multiplying.
+				FVec3 JInvF;
+				if (Utilities::Solve(JInvF, J, F)) // Equivalent to Utilities::Multiply(J.Inverse(), F);
+				{
+					WBody = WBody - JInvF;
+				}
+			}
+
+			// Convert back to world coordinates
+			return Q.RotateVector(WBody);
+
+		}
+
+
 	}
 }
diff --git a/Engine/Source/Runtime/Experimental/Chaos/Public/Chaos/ParticleDirtyFlags.h b/Engine/Source/Runtime/Experimental/Chaos/Public/Chaos/ParticleDirtyFlags.h
--- a/Engine/Source/Runtime/Experimental/Chaos/Public/Chaos/ParticleDirtyFlags.h
+++ b/Engine/Source/Runtime/Experimental/Chaos/Public/Chaos/ParticleDirtyFlags.h
@@ -540,11 +540,14 @@ public:
 	bool UpdateKinematicFromSimulation() const { return MControlFlags.GetUpdateKinematicFromSimulation(); }
 	void SetUpdateKinematicFromSimulation(bool bUpdateKinematicFromSimulation) { MControlFlags.SetUpdateKinematicFromSimulation(bUpdateKinematicFromSimulation); }
 
+	bool GyroscopicTorqueEnabled() const { return MControlFlags.GetGyroscopicTorqueEnabled(); }
+	void SetGyroscopicTorqueEnabled(bool bInGyroscopicTorqueEnabled) { MControlFlags.SetGyroscopicTorqueEnabled(bInGyroscopicTorqueEnabled); }
+
 	bool CCDEnabled() const { return MControlFlags.GetCCDEnabled(); }
 	void SetCCDEnabled(bool bInCCDEnabled) { MControlFlags.SetCCDEnabled(bInCCDEnabled); }
 
 	bool MACDEnabled() const { return MControlFlags.GetMACDEnabled(); }
-	void SetMACDEnabled(bool bInCCDEnabled) { MControlFlags.SetMACDEnabled(bInCCDEnabled); }
+	void SetMACDEnabled(bool bInMACDEnabled) { MControlFlags.SetMACDEnabled(bInMACDEnabled); }
 
 	void SetPositionSolverIterationCount(const int32 PositionSolverIterationCountIn) { MIterationSettings.SetNumPositionIterations(PositionSolverIterationCountIn); }
 	void SetVelocitySolverIterationCount(const int32 VelocitySolverIterationCountIn) { MIterationSettings.SetNumVelocityIterations(VelocitySolverIterationCountIn); }
diff --git a/Engine/Source/Runtime/Experimental/Chaos/Public/Chaos/ParticleHandle.h b/Engine/Source/Runtime/Experimental/Chaos/Public/Chaos/ParticleHandle.h
--- a/Engine/Source/Runtime/Experimental/Chaos/Public/Chaos/ParticleHandle.h
+++ b/Engine/Source/Runtime/Experimental/Chaos/Public/Chaos/ParticleHandle.h
@@ -1288,6 +1288,16 @@ public:
 		PBDRigidParticles->ControlFlags(ParticleIdx).SetUpdateKinematicFromSimulation(bUpdateKinematicFromSimulation);
 	}
 
+	inline bool GyroscopicTorqueEnabled() const
+	{
+		return ControlFlags().GetGyroscopicTorqueEnabled();
+	}
+
+	inline void SetGyroscopicTorqueEnabled(bool bEnabled)
+	{
+		PBDRigidParticles->ControlFlags(ParticleIdx).SetGyroscopicTorqueEnabled(bEnabled);
+	}
+
 	inline bool CCDEnabled() const
 	{
 		return ControlFlags().GetCCDEnabled();
@@ -3399,6 +3409,12 @@ public:
 		MMiscData.Modify(true, MDirtyFlags, Proxy, [bUpdateKinematicFromSimulation](auto& Data) { Data.SetUpdateKinematicFromSimulation(bUpdateKinematicFromSimulation); });
 	}
 
+	bool GyroscopicTorqueEnabled() const { return MMiscData.Read().GyroscopicTorqueEnabled(); }
+	void SetGyroscopicTorqueEnabled(const bool bGyroscopicTorqueEnabled)
+	{
+		MMiscData.Modify(true, MDirtyFlags, Proxy, [bGyroscopicTorqueEnabled](auto& Data) { Data.SetGyroscopicTorqueEnabled(bGyroscopicTorqueEnabled); });
+	}
+
 	bool OneWayInteraction() const { return MMiscData.Read().OneWayInteraction(); }
 	void SetOneWayInteraction(const bool bInEnabled)
 	{
diff --git a/Engine/Source/Runtime/Experimental/Chaos/Public/Chaos/RigidParticleControlFlags.h b/Engine/Source/Runtime/Experimental/Chaos/Public/Chaos/RigidParticleControlFlags.h
--- a/Engine/Source/Runtime/Experimental/Chaos/Public/Chaos/RigidParticleControlFlags.h
+++ b/Engine/Source/Runtime/Experimental/Chaos/Public/Chaos/RigidParticleControlFlags.h
@@ -33,6 +33,9 @@ namespace Chaos
 		bool GetUpdateKinematicFromSimulation() const { return Flags.bUpdateKinematicFromSimulation; }
 		FRigidParticleControlFlags& SetUpdateKinematicFromSimulation(const bool bUpdateKinematicFromSimulation) { Flags.bUpdateKinematicFromSimulation = bUpdateKinematicFromSimulation; return *this; }
 
+		bool GetGyroscopicTorqueEnabled() const { return Flags.bGyroscopicTorqueEnabled; }
+		FRigidParticleControlFlags& SetGyroscopicTorqueEnabled(const bool bEnabled) { Flags.bGyroscopicTorqueEnabled = bEnabled; return *this; }
+
 		int32 GetGravityGroupIndex() const { return Flags.GravityGroupIndex; }
 		FRigidParticleControlFlags& SetGravityGroupIndex(const int32 GravityGroupIndex) 
 		{ 
@@ -81,6 +84,7 @@ namespace Chaos
 			FStorage GravityGroupIndex : 3;
 			FStorage bUpdateKinematicFromSimulation : 1;
 			FStorage bMACDEnabled : 1;
+			FStorage bGyroscopicTorqueEnabled : 1;
 			// Add new properties above this line
 			// Change FStorage typedef if we exceed the max bits
 		};
diff --git a/Engine/Source/Runtime/Experimental/Chaos/Public/Chaos/Utilities.h b/Engine/Source/Runtime/Experimental/Chaos/Public/Chaos/Utilities.h
--- a/Engine/Source/Runtime/Experimental/Chaos/Public/Chaos/Utilities.h
+++ b/Engine/Source/Runtime/Experimental/Chaos/Public/Chaos/Utilities.h
@@ -274,6 +274,45 @@ namespace Chaos
 			);
 		}
 
+		// Solves X * A = B (or A.Multiply(X) = B) using Gaussian elimination, writing the result
+		// into X. Returns true if there is a solution, false otherwise.
+		inline bool Solve(FVec3& X, const FMatrix33& A, const FVec3& B)
+		{
+			const FReal Determinant = A.Determinant();
+			if (FMath::IsNearlyZero(Determinant))
+			{
+				return false;
+			}
+
+			// Matrix33 stores in columns...
+			const FVec3 C0 = A.GetColumn(0);
+			const FVec3 C1 = A.GetColumn(1);
+			const FVec3 C2 = A.GetColumn(2);
+
+			// All reference material represents the problem in row form. 
+			// Form the 3x4 augmented matrix
+			FVec4 R0(C0.X, C1.X, C2.X, B.X);
+			FVec4 R1(C0.Y, C1.Y, C2.Y, B.Y);
+			FVec4 R2(C0.Z, C1.Z, C2.Z, B.Z);
+
+			// Now apply row operations to reduce it to echelon form: Ones along the diagonal, and
+			// zeros in the bottom left.
+			R0 *= 1 / R0[0];  // Make (0, 0) = 1
+			R1 -= R0 * R1[0]; // Make (1, 0) = 0
+			R2 -= R0 * R2[0]; // Make (2, 0) = 0
+
+			R1 *= 1 / R1[1];  // Make (1, 1) = 1
+			R2 -= R1 * R2[1]; // Make (2, 1) = 0
+
+			R2 *= 1 / R2[2]; // Make (2, 2) = 1
+
+			// R is now in echelon form, so we can work up from the bottom to calculate result
+			X[2] = R2[3];
+			X[1] = R1[3] - R1[2] * X[2];
+			X[0] = R0[3] - (R0[1] * X[1] + R0[2] * X[2]);
+			return true;
+		}
+
 		/**
 		 * Concatenate two transforms. This returns a transform that logically applies R then L.
 		 */
@@ -297,6 +336,12 @@ namespace Chaos
 			return ComputeWorldSpaceInertia(CoMRotation, FMatrix33(I.X, I.Y, I.Z));
 		}
 
+		// Given a body with world orientation Q, world angular velocity W, local moments of inertia
+		// I, this calculates the new angular velocity it should have after integrating forwards by
+		// Dt, considering the gyroscopic torques.
+		FVec3 GetAngularVelocityAdjustedForGyroscopicTorques(
+			const FRotation3& Q, const FVec3& I, const FVec3& W, const FReal Dt);
+
 		/**
 		 * Calculate the matrix that maps a constraint position error to constraint position and rotation corrections.
 		 */
diff --git a/Engine/Source/Runtime/Experimental/Chaos/Public/PhysicsProxy/SingleParticlePhysicsProxy.h b/Engine/Source/Runtime/Experimental/Chaos/Public/PhysicsProxy/SingleParticlePhysicsProxy.h
--- a/Engine/Source/Runtime/Experimental/Chaos/Public/PhysicsProxy/SingleParticlePhysicsProxy.h
+++ b/Engine/Source/Runtime/Experimental/Chaos/Public/PhysicsProxy/SingleParticlePhysicsProxy.h
@@ -459,6 +459,30 @@ public:
 		});
 	}
 
+	bool GyroscopicTorqueEnabled() const
+	{
+		return Read([](auto* Particle)
+			{
+				if (auto Rigid = Particle->CastToRigidParticle())
+				{
+					return Rigid->GyroscopicTorqueEnabled();
+				}
+
+				return false;
+			});
+	}
+
+	void SetGyroscopicTorqueEnabled(const bool InGyroscopicTorqueEnabled)
+	{
+		Write([InGyroscopicTorqueEnabled](auto* Particle)
+			{
+				if (auto Rigid = Particle->CastToRigidParticle())
+				{
+					return Rigid->SetGyroscopicTorqueEnabled(InGyroscopicTorqueEnabled);
+				}
+			});
+	}
+
 	bool CCDEnabled() const
 	{
 		return Read([](auto* Particle)
diff --git a/Engine/Source/Runtime/Experimental/ChaosVisualDebugger/Private/DataWrappers/ChaosVDParticleDataWrapper.cpp b/Engine/Source/Runtime/Experimental/ChaosVisualDebugger/Private/DataWrappers/ChaosVDParticleDataWrapper.cpp
--- a/Engine/Source/Runtime/Experimental/ChaosVisualDebugger/Private/DataWrappers/ChaosVDParticleDataWrapper.cpp
+++ b/Engine/Source/Runtime/Experimental/ChaosVisualDebugger/Private/DataWrappers/ChaosVDParticleDataWrapper.cpp
@@ -38,6 +38,11 @@ bool FChaosVDFRigidParticleControlFlags::Serialize(FArchive& Ar)
 		Ar << ProjectionSolverIterationCount;
 	}
 
+	if (Ar.CustomVer(FFortniteMainBranchObjectVersion::GUID) >= FFortniteMainBranchObjectVersion::GyroscopicTorquesSupportInChaosVisualDebugger)
+	{
+		Ar << bGyroscopicTorqueEnabled;
+	}
+
 	return !Ar.IsError();
 }
 
diff --git a/Engine/Source/Runtime/Experimental/ChaosVisualDebugger/Public/DataWrappers/ChaosVDParticleDataWrapper.h b/Engine/Source/Runtime/Experimental/ChaosVisualDebugger/Public/DataWrappers/ChaosVDParticleDataWrapper.h
--- a/Engine/Source/Runtime/Experimental/ChaosVisualDebugger/Public/DataWrappers/ChaosVDParticleDataWrapper.h
+++ b/Engine/Source/Runtime/Experimental/ChaosVisualDebugger/Public/DataWrappers/ChaosVDParticleDataWrapper.h
@@ -86,6 +86,7 @@ struct FChaosVDFRigidParticleControlFlags : public FChaosVDWrapperDataBase
 		  bInertiaConditioningEnabled(false), 
 		  GravityGroupIndex(0),
 		  bMACDEnabled(false), 
+		  bGyroscopicTorqueEnabled(false),
 		  PositionSolverIterationCount(8),
 		  VelocitySolverIterationCount(1),
 		  ProjectionSolverIterationCount(1)
@@ -103,6 +104,7 @@ struct FChaosVDFRigidParticleControlFlags : public FChaosVDWrapperDataBase
 		bInertiaConditioningEnabled = Other.GetInertiaConditioningEnabled();
 		GravityGroupIndex = Other.GetGravityGroupIndex();
 		bMACDEnabled = Other.GetMACDEnabled();
+		bGyroscopicTorqueEnabled = Other.GetGyroscopicTorqueEnabled();
 
 		bHasValidData = true;
 	}
@@ -120,6 +122,8 @@ struct FChaosVDFRigidParticleControlFlags : public FChaosVDWrapperDataBase
 	UPROPERTY(VisibleAnywhere, Category = "Particle Control Flags")
 	bool bMACDEnabled;
 	UPROPERTY(VisibleAnywhere, Category = "Particle Control Flags")
+	bool bGyroscopicTorqueEnabled;
+	UPROPERTY(VisibleAnywhere, Category = "Particle Control Flags")
 	uint8 PositionSolverIterationCount;
 	UPROPERTY(VisibleAnywhere, Category = "Particle Control Flags")
 	uint8 VelocitySolverIterationCount;
diff --git a/Engine/Source/Runtime/PhysicsCore/Private/BodyInstanceCore.cpp b/Engine/Source/Runtime/PhysicsCore/Private/BodyInstanceCore.cpp
--- a/Engine/Source/Runtime/PhysicsCore/Private/BodyInstanceCore.cpp
+++ b/Engine/Source/Runtime/PhysicsCore/Private/BodyInstanceCore.cpp
@@ -15,6 +15,7 @@ FBodyInstanceCore::FBodyInstanceCore()
 , bOverrideMass(false)
 , bEnableGravity(true)
 , bUpdateKinematicFromSimulation(false)
+, bGyroscopicTorqueEnabled(false)
 , bAutoWeld(false)
 , bStartAwake(true)
 , bGenerateWakeEvents(false)
diff --git a/Engine/Source/Runtime/PhysicsCore/Private/ChaosEngineInterface.cpp b/Engine/Source/Runtime/PhysicsCore/Private/ChaosEngineInterface.cpp
--- a/Engine/Source/Runtime/PhysicsCore/Private/ChaosEngineInterface.cpp
+++ b/Engine/Source/Runtime/PhysicsCore/Private/ChaosEngineInterface.cpp
@@ -1199,6 +1199,15 @@ void FChaosEngineInterface::SetUpdateKinematicFromSimulation_AssumesLocked(const
 	InActorReference->GetGameThreadAPI().SetUpdateKinematicFromSimulation(bUpdateKinematicFromSimulation);
 }
 
+bool FChaosEngineInterface::GetGyroscopicTorqueEnabled_AssumesLocked(const FPhysicsActorHandle& InActorReference)
+{
+	return InActorReference->GetGameThreadAPI().GyroscopicTorqueEnabled();
+}
+void FChaosEngineInterface::SetGyroscopicTorqueEnabled_AssumesLocked(const FPhysicsActorHandle& InActorReference, bool bGyroscopicTorqueEnabled)
+{
+	InActorReference->GetGameThreadAPI().SetGyroscopicTorqueEnabled(bGyroscopicTorqueEnabled);
+}
+
 void FChaosEngineInterface::SetOneWayInteraction_AssumesLocked(const FPhysicsActorHandle& InHandle, bool InOneWayInteraction)
 {
 	InHandle->GetGameThreadAPI().SetOneWayInteraction(InOneWayInteraction);
diff --git a/Engine/Source/Runtime/PhysicsCore/Public/BodyInstanceCore.h b/Engine/Source/Runtime/PhysicsCore/Public/BodyInstanceCore.h
--- a/Engine/Source/Runtime/PhysicsCore/Public/BodyInstanceCore.h
+++ b/Engine/Source/Runtime/PhysicsCore/Public/BodyInstanceCore.h
@@ -42,6 +42,10 @@ struct FBodyInstanceCore
 	UPROPERTY(EditAnywhere,BlueprintReadOnly,Category = Physics)
 	uint8 bUpdateKinematicFromSimulation : 1;
 
+	/** Enabled/disables whether this body is affected by gyroscopic torque, mainly useful for long/thin objects that spin */
+	UPROPERTY(EditAnywhere, BlueprintReadOnly, Category = Physics)
+	uint8 bGyroscopicTorqueEnabled : 1;
+
 	/** If true and is attached to a parent, the two bodies will be joined into a single rigid body. Physical settings like collision profile and body settings are determined by the root */
 	UPROPERTY(EditAnywhere,AdvancedDisplay,BlueprintReadWrite,Category = Physics,meta = (editcondition = "!bSimulatePhysics"))
 	uint8 bAutoWeld : 1;
diff --git a/Engine/Source/Runtime/PhysicsCore/Public/Chaos/ChaosEngineInterface.h b/Engine/Source/Runtime/PhysicsCore/Public/Chaos/ChaosEngineInterface.h
--- a/Engine/Source/Runtime/PhysicsCore/Public/Chaos/ChaosEngineInterface.h
+++ b/Engine/Source/Runtime/PhysicsCore/Public/Chaos/ChaosEngineInterface.h
@@ -468,6 +468,9 @@ public:
 	static PHYSICSCORE_API bool GetUpdateKinematicFromSimulation_AssumesLocked(const FPhysicsActorHandle& InActorReference);
 	static PHYSICSCORE_API void SetUpdateKinematicFromSimulation_AssumesLocked(const FPhysicsActorHandle& InActorReference, bool bUpdateKinematicFromSimulation);
 
+	static PHYSICSCORE_API bool GetGyroscopicTorqueEnabled_AssumesLocked(const FPhysicsActorHandle& InActorReference);
+	static PHYSICSCORE_API void SetGyroscopicTorqueEnabled_AssumesLocked(const FPhysicsActorHandle& InActorReference, bool bGyroscopicTorqueEnabled);
+
 	static PHYSICSCORE_API void SetOneWayInteraction_AssumesLocked(const FPhysicsActorHandle& InHandle, bool InOneWayInteraction);
 
 	static PHYSICSCORE_API float GetSleepEnergyThreshold_AssumesLocked(const FPhysicsActorHandle& InActorReference);
