Clang will optimize these lambdas and reuse registers, causing a crash if we don't keep them alive with TFunction, rather than TFunctionRef

#rnx

diff --git a/Engine/Source/Runtime/Experimental/JsonObjectGraph/Private/JsonStringifyImpl.cpp b/Engine/Source/Runtime/Experimental/JsonObjectGraph/Private/JsonStringifyImpl.cpp
--- a/Engine/Source/Runtime/Experimental/JsonObjectGraph/Private/JsonStringifyImpl.cpp
+++ b/Engine/Source/Runtime/Experimental/JsonObjectGraph/Private/JsonStringifyImpl.cpp
@@ -957,7 +957,7 @@ bool FJsonStringifyImpl::IsDeltaEncoding() const
 	return !EnumHasAnyFlags(WriteOptions.Flags, EJsonStringifyFlags::DisableDeltaEncoding);
 }
 
-FJsonStringifyImpl::FPendingScope::FPendingScope(FJsonStringifyImpl* To, TFunctionRef<void()> Prefix)
+FJsonStringifyImpl::FPendingScope::FPendingScope(FJsonStringifyImpl* To, const TFunction<void()>& Prefix)
 	: Owner(To)
 	, Outer(To->CurrentScope)
 	, PendingPrefix(Prefix)
@@ -965,7 +965,7 @@ FJsonStringifyImpl::FPendingScope::FPendingScope(FJsonStringifyImpl* To, TFuncti
 	To->CurrentScope = this;
 }
 
-FJsonStringifyImpl::FPendingScope::FPendingScope(FJsonStringifyImpl* To, TFunctionRef<void()> Prefix, TFunctionRef<void()> Postfix)
+FJsonStringifyImpl::FPendingScope::FPendingScope(FJsonStringifyImpl* To, const TFunction<void()>& Prefix, const TFunction<void()>& Postfix)
 	: Owner(To)
 	, Outer(To->CurrentScope)
 	, PendingPrefix(Prefix)
diff --git a/Engine/Source/Runtime/Experimental/JsonObjectGraph/Private/JsonStringifyImpl.h b/Engine/Source/Runtime/Experimental/JsonObjectGraph/Private/JsonStringifyImpl.h
--- a/Engine/Source/Runtime/Experimental/JsonObjectGraph/Private/JsonStringifyImpl.h
+++ b/Engine/Source/Runtime/Experimental/JsonObjectGraph/Private/JsonStringifyImpl.h
@@ -78,15 +78,15 @@ private:
 	// are values that match the default value we won't write out the enclosing scope:
 	struct FPendingScope
 	{
-		FPendingScope(FJsonStringifyImpl* To, TFunctionRef<void()> Prefix);
-		FPendingScope(FJsonStringifyImpl* To, TFunctionRef<void()> Prefix, TFunctionRef<void()> Postfix);
+		FPendingScope(FJsonStringifyImpl* To, const TFunction<void()>& Prefix);
+		FPendingScope(FJsonStringifyImpl* To, const TFunction<void()>& Prefix, const TFunction<void()>& Postfix);
 		~FPendingScope();
 		void Apply();
 	private:
 		FJsonStringifyImpl* Owner = nullptr;
 		FPendingScope* Outer = nullptr;
-		TFunctionRef<void()> PendingPrefix;
-		TOptional<TFunctionRef<void()>> PendingPostfix;
+		TFunction<void()> PendingPrefix;
+		TOptional<TFunction<void()>> PendingPostfix;
 
 		bool bHasBeenApplied = false;
 
