[Chaos] Allow to enable/disable debug draw on preview worlds

diff --git a/Engine/Source/Runtime/Experimental/ChaosSolverEngine/Private/Chaos/ChaosDebugDrawSubsystem.cpp b/Engine/Source/Runtime/Experimental/ChaosSolverEngine/Private/Chaos/ChaosDebugDrawSubsystem.cpp
--- a/Engine/Source/Runtime/Experimental/ChaosSolverEngine/Private/Chaos/ChaosDebugDrawSubsystem.cpp
+++ b/Engine/Source/Runtime/Experimental/ChaosSolverEngine/Private/Chaos/ChaosDebugDrawSubsystem.cpp
@@ -31,6 +31,10 @@ namespace Chaos::CVars
 {
 	extern int32 ChaosSolverDebugDrawShowServer;
 	extern int32 ChaosSolverDebugDrawShowClient;
+
+	// Temp cvar to allow debug draw in Asset Editors until we can implement a better per world (and per asset editor) solution - UE-229329
+	bool bChaosDebugDraw_PreviewWorldsEnabled = false;
+	FAutoConsoleVariableRef CVarbChaosDebugDraw_SupportPreviewWorlds(TEXT("p.Chaos.PreviewWorld.DebugDraw.Enabled"), bChaosDebugDraw_PreviewWorldsEnabled, TEXT("Enables/Disables Chaos debug Draw support in Preview worlds. Mostly used by Asset Editors."));
 }
 
 class FChaosDDRenderer : public ChaosDD::Private::IChaosDDRenderer
@@ -260,7 +264,8 @@ bool UChaosDebugDrawSubsystem::ShouldCreateSubsystem(UObject* Outer) const
 #if CHAOS_DEBUG_DRAW
 	if (UWorld* World = Cast<UWorld>(Outer))
 	{
-		const bool bCreateDebugDraw = !IsRunningCommandlet() && DoesSupportWorldType(World->WorldType);
+		const bool bIsWorldTypeSupported = World->IsPreviewWorld() ? Chaos::CVars::bChaosDebugDraw_PreviewWorldsEnabled : true;
+		const bool bCreateDebugDraw = !IsRunningCommandlet() && bIsWorldTypeSupported;
 
 		UE_LOG(LogChaosDD, Log, TEXT("%s Chaos Debug Draw Scene for world %s"), bCreateDebugDraw ? TEXT("Creating") : TEXT("Not creating"), *World->GetName());
 		
