Dataflow : Add EDataflowImageResolution
- also fix some UI limits for UV channels in the AutoUV nodes

diff --git a/Engine/Plugins/Experimental/GeometryCollectionPlugin/Source/GeometryCollectionNodes/Public/Dataflow/GeometryCollectionUVNodes.h b/Engine/Plugins/Experimental/GeometryCollectionPlugin/Source/GeometryCollectionNodes/Public/Dataflow/GeometryCollectionUVNodes.h
--- a/Engine/Plugins/Experimental/GeometryCollectionPlugin/Source/GeometryCollectionNodes/Public/Dataflow/GeometryCollectionUVNodes.h
+++ b/Engine/Plugins/Experimental/GeometryCollectionPlugin/Source/GeometryCollectionNodes/Public/Dataflow/GeometryCollectionUVNodes.h
@@ -56,11 +56,11 @@ private:
 	FDataflowFaceSelection FaceSelection;
 
 	/** UV channel to unwrap into ( 0 by default ) */
-	UPROPERTY(EditAnywhere, Category = Options, meta = (DataflowInput))
+	UPROPERTY(EditAnywhere, Category = Options, meta = (DataflowInput, ClampMin = 0, ClampMax = 7, UIMin = 0, UIMax = 7))
 	int32 UVChannel = 0;
 
 	/** Approximate space to leave between UV islands, measured in texels for 512x512 texture */
-	UPROPERTY(EditAnywhere, Category = Options)
+	UPROPERTY(EditAnywhere, Category = Options, meta = (ClampMin = 1, UIMin = 1))
 	int32 GutterSize = 1;
 
 	virtual void Evaluate(UE::Dataflow::FContext& Context, const FDataflowOutput* Out) const override;
diff --git a/Engine/Source/Runtime/Experimental/Dataflow/Core/Public/Dataflow/DataflowImage.h b/Engine/Source/Runtime/Experimental/Dataflow/Core/Public/Dataflow/DataflowImage.h
--- a/Engine/Source/Runtime/Experimental/Dataflow/Core/Public/Dataflow/DataflowImage.h
+++ b/Engine/Source/Runtime/Experimental/Dataflow/Core/Public/Dataflow/DataflowImage.h
@@ -6,6 +6,21 @@
 
 #include "DataflowImage.generated.h"
 
+UENUM()
+enum class EDataflowImageResolution : int32
+{
+	Resolution16 = 16 UMETA(DisplayName = "16 x 16"),
+	Resolution32 = 32 UMETA(DisplayName = "32 x 32"),
+	Resolution64 = 64 UMETA(DisplayName = "64 x 64"),
+	Resolution128 = 128 UMETA(DisplayName = "128 x 128"),
+	Resolution256 = 256 UMETA(DisplayName = "256 x 256"),
+	Resolution512 = 512 UMETA(DisplayName = "512 x 512"),
+	Resolution1024 = 1024 UMETA(DisplayName = "1024 x 1024"),
+	Resolution2048 = 2048 UMETA(DisplayName = "2048 x 2048"),
+	Resolution4096 = 4096 UMETA(DisplayName = "4096 x 4096"),
+	Resolution8192 = 8192 UMETA(DisplayName = "8192 x 8192")
+};
+
 /** 
 * Represents image for dataflow
 * type is constrained to Float32 with 1 or 4 channels
@@ -15,11 +30,11 @@ struct FDataflowImage
 {
 	GENERATED_USTRUCT_BODY()
 
-	int32 GetWidth() const;
-	int32 GetHeight() const;
+	DATAFLOWCORE_API int32 GetWidth() const;
+	DATAFLOWCORE_API int32 GetHeight() const;
 
 	/** get the readonly underlying image object */
-	const FImage& GetImage() const;
+	DATAFLOWCORE_API const FImage& GetImage() const;
 
 	/** Create a single channel float format image */
 	DATAFLOWCORE_API void CreateR32F(int32 Width, int32 Height);
@@ -40,7 +55,7 @@ struct FDataflowImage
 	*/
 	DATAFLOWCORE_API void ConvertToRGBA32F();
 
-	bool Serialize(FArchive& Ar);
+	DATAFLOWCORE_API bool Serialize(FArchive& Ar);
 
 private:
 	FImage Image;
