Make use of UTF8CHAR explicit, rather than relying on ANSICHAR, this fixes encoding of objects with z in their name

#rnx

diff --git a/Engine/Source/Runtime/Experimental/JsonObjectGraph/Private/PrettyJsonWriter.cpp b/Engine/Source/Runtime/Experimental/JsonObjectGraph/Private/PrettyJsonWriter.cpp
--- a/Engine/Source/Runtime/Experimental/JsonObjectGraph/Private/PrettyJsonWriter.cpp
+++ b/Engine/Source/Runtime/Experimental/JsonObjectGraph/Private/PrettyJsonWriter.cpp
@@ -7,7 +7,7 @@ namespace UE::Private
 {
 
 FPrettyJsonWriter::FPrettyJsonWriter(FArchive* const InStream, int32 InitialIndentLevel)
-	: TJsonWriter<ANSICHAR, TPrettyJsonPrintPolicySingleNewLine>(InStream, InitialIndentLevel)
+	: TJsonWriter<UTF8CHAR, TPrettyJsonPrintPolicySingleNewLine>(InStream, InitialIndentLevel)
 {
 }
 
diff --git a/Engine/Source/Runtime/Experimental/JsonObjectGraph/Private/PrettyJsonWriter.h b/Engine/Source/Runtime/Experimental/JsonObjectGraph/Private/PrettyJsonWriter.h
--- a/Engine/Source/Runtime/Experimental/JsonObjectGraph/Private/PrettyJsonWriter.h
+++ b/Engine/Source/Runtime/Experimental/JsonObjectGraph/Private/PrettyJsonWriter.h
@@ -8,20 +8,20 @@ namespace UE::Private
 {
 
 struct TPrettyJsonPrintPolicySingleNewLine
-	: public TPrettyJsonPrintPolicy<ANSICHAR>
+	: public TPrettyJsonPrintPolicy<UTF8CHAR>
 {
 	static inline void WriteLineTerminator(FArchive* Stream)
 	{
-		TJsonPrintPolicy<ANSICHAR>::WriteString(Stream, FStringView(TEXT("\n")));
+		TJsonPrintPolicy<UTF8CHAR>::WriteString(Stream, FStringView(TEXT("\n")));
 	}
 };
 
 // Extending TJsonWriter to get some control over whitespace and tabbing:
-struct FPrettyJsonWriter : public TJsonWriter<ANSICHAR, TPrettyJsonPrintPolicySingleNewLine>
+struct FPrettyJsonWriter : public TJsonWriter<UTF8CHAR, TPrettyJsonPrintPolicySingleNewLine>
 {
-	using Super = TJsonWriter<ANSICHAR, TPrettyJsonPrintPolicySingleNewLine>;
+	using Super = TJsonWriter<UTF8CHAR, TPrettyJsonPrintPolicySingleNewLine>;
 	using PrintPolicy = TPrettyJsonPrintPolicySingleNewLine;
-	using CharType = ANSICHAR;
+	using CharType = UTF8CHAR;
 
 	FPrettyJsonWriter(FArchive* const InStream, int32 InitialIndentLevel);
 
