Fix a rare crash in the OnDemand module.

#rnx

- When we have several outstanding requests to load the same chunk, instead of trying to load the chunk multiple times we will build a small linked list of the requests on the FChunkRequest. If we then need to cancel a request before it is completed we have to remove it from the FChunkRequest linked list. There was a bug that if we tried to remove the last element in the list and the list had more than one element then the tail pointer would not be correctly updated and would continue to point to the now removed request. Over time this could potentially corrupt either the list or memory in general.
- We now correctly update the tail pointer, I also added several additional asserts to make identifying such problems easier in the future.
- This crash was seen 1/50 attempts while testing how the module functioned under low core conditions. It requires that we build several requests for the same data over time and the requests take long enough that we start to cancel them and that we cancel in a LIFO order. Since this was so hard to reproduce I added a quick test case that will run on startup if the define 'UE_ENABLE_IAS_TESTING is set to 1. Otherwise there is no reasonable way to test this issue.

diff --git a/Engine/Source/Runtime/Experimental/IoStore/OnDemand/Private/OnDemandIoDispatcherBackend.cpp b/Engine/Source/Runtime/Experimental/IoStore/OnDemand/Private/OnDemandIoDispatcherBackend.cpp
--- a/Engine/Source/Runtime/Experimental/IoStore/OnDemand/Private/OnDemandIoDispatcherBackend.cpp
+++ b/Engine/Source/Runtime/Experimental/IoStore/OnDemand/Private/OnDemandIoDispatcherBackend.cpp
@@ -45,6 +45,9 @@
 /** When enabled the IAS system can add additional debug console commands for development use */
 #define UE_IAS_DEBUG_CONSOLE_CMDS (1 && !NO_CVARS && !UE_BUILD_SHIPPING)
 
+/** When enabled we will run some limited testing on start up for issues that are hard to reproduce with normal gameplay */
+#define UE_ENABLE_IAS_TESTING 0 && !UE_BUILD_SHIPPING
+
 namespace UE::IoStore
 {
 
@@ -612,6 +615,12 @@ struct FChunkRequest
 			}
 			check(It->NextRequest == Request);
 			It->NextRequest = It->NextRequest->NextRequest;
+
+			if (RequestTail == Request)
+			{
+				check(It->NextRequest == nullptr);
+				RequestTail = It;
+			}
 		}
 
 		Request->NextRequest = nullptr;
@@ -979,6 +988,8 @@ class FOnDemandIoBackend final
 
 				if (RemainingCount == 0)
 				{
+					check(ChunkRequest.RequestTail == nullptr);
+
 					ChunkRequest.bCancelled = true;
 					if (TheCache != nullptr)
 					{
@@ -1945,6 +1956,70 @@ TSharedPtr<IOnDemandIoDispatcherBackend> MakeOnDemandIoDispatcherBackend(
 		new FOnDemandIoBackend(Config, IoStore, MoveTemp(Cache)));
 }
 
+#if UE_ENABLE_IAS_TESTING
+
+namespace Tests
+{
+
+/***
+ * We should port these tests to the automation framework or lowlevel tests once we decide which way we are going with testing
+ * for this module.
+ */
+struct FTests
+{
+	void ChunkRequestLinkedListTest()
+	{
+		// We are not allowed to create FIoRequestImpl directly due to it's api, however we don't need the full functionality
+		// in order to test FChunkRequest so we can fake it with a malloced version of the struct, mem set to zero.
+		FIoRequestImpl* IoFirstRequest = (FIoRequestImpl*)FMemory::MallocZeroed(sizeof(FIoRequestImpl));
+		FIoRequestImpl* IoSecondRequest = (FIoRequestImpl*)FMemory::MallocZeroed(sizeof(FIoRequestImpl));
+		FIoRequestImpl* IoThirdRequest = (FIoRequestImpl*)FMemory::MallocZeroed(sizeof(FIoRequestImpl));
+
+		FChunkRequestParams Params;
+		FChunkRequest ChunkRequest(IoFirstRequest, Params);
+		ChunkRequest.AddDispatcherRequest(IoSecondRequest);
+		ChunkRequest.AddDispatcherRequest(IoThirdRequest);
+
+		check(ChunkRequest.RequestHead == IoFirstRequest);
+		check(ChunkRequest.RequestHead->NextRequest == IoSecondRequest);
+		check(ChunkRequest.RequestHead->NextRequest->NextRequest == IoThirdRequest);
+		check(ChunkRequest.RequestTail == IoThirdRequest);
+
+		check(ChunkRequest.RequestCount == 3);
+
+		ChunkRequest.RemoveDispatcherRequest(IoThirdRequest);
+		
+		check(ChunkRequest.RequestCount == 2);
+		check(ChunkRequest.RequestHead == IoFirstRequest);
+		check(ChunkRequest.RequestTail == IoSecondRequest);
+
+		ChunkRequest.RemoveDispatcherRequest(IoSecondRequest);
+		
+		check(ChunkRequest.RequestCount == 1);
+		check(ChunkRequest.RequestHead == IoFirstRequest);
+		check(ChunkRequest.RequestTail == IoFirstRequest);
+
+		ChunkRequest.RemoveDispatcherRequest(IoFirstRequest);
+		check(ChunkRequest.RequestCount == 0);
+
+		check(ChunkRequest.RequestHead == nullptr);
+		check(ChunkRequest.RequestTail == nullptr);
+
+		FMemory::Free(IoFirstRequest);
+		FMemory::Free(IoSecondRequest);
+		FMemory::Free(IoThirdRequest);
+	}
+
+	FTests()
+	{
+		ChunkRequestLinkedListTest();
+	}
+} RunTests;
+
+} // namespace Tests
+
+#endif // UE_ENABLE_IAS_TESTING
+
 } // namespace UE::IoStore
 
 #undef UE_IAS_DEBUG_CONSOLE_CMDS
\ No newline at end of file
