Remove localization from the 'ias.DisplayOnScreenStatistics' output.

#rnx

- This sort of debug output does not need localization and it is a waste of resources to have it updated each time we want to change the code.

diff --git a/Engine/Source/Runtime/Experimental/IoStore/OnDemand/Private/Statistics.cpp b/Engine/Source/Runtime/Experimental/IoStore/OnDemand/Private/Statistics.cpp
--- a/Engine/Source/Runtime/Experimental/IoStore/OnDemand/Private/Statistics.cpp
+++ b/Engine/Source/Runtime/Experimental/IoStore/OnDemand/Private/Statistics.cpp
@@ -411,7 +411,6 @@ FOnDemandIoBackendStats::FOnDemandIoBackendStats(EStatsFlags InFlags)
 		}
 	});
 
-#define LOCTEXT_NAMESPACE "IAS"
 	GStatisticsOnScreenDelegateHandle = FCoreDelegates::OnGetOnScreenMessages.AddLambda(
 		[this] (FCoreDelegates::FSeverityMessageMap& Out)
 		{
@@ -419,11 +418,11 @@ FOnDemandIoBackendStats::FOnDemandIoBackendStats(EStatsFlags InFlags)
 			{
 				const bool bIsConnected = GHttpConnectCount.Get() > GHttpDisconnectCount.Get();
 
-				FText CachingText;
+				TStringBuilder<256> CachingText;
 				if (!EnumHasAllFlags(Flags, EStatsFlags::CachingDisabled))
 				{
-					CachingText = FText::Format(
-						LOCTEXT("IASCacheEnabled", "Cached:{0} KiB | Read:{1} KiB ({2})"),
+					CachingText.Appendf(
+						TEXT("Cached:%d KiB | Read:%d KiB (%d)"),
 						GCacheCachedBytes.Get() >> 10,
 						GCacheReadBytes.Get() >> 10,
 						GCacheGetCount.Get()
@@ -431,13 +430,14 @@ FOnDemandIoBackendStats::FOnDemandIoBackendStats(EStatsFlags InFlags)
 				}
 				else
 				{
-					CachingText = LOCTEXT("IASCacheDisabled", "Caching Disabled");
+					CachingText = TEXTVIEW("Caching Disabled");
 				}
 
-				FText Message = FText::Format(
-					LOCTEXT("IAS", "IAS - {0}: {1} | Downloaded:{2} KiB ({3}) {4} ms | Retries:{5} | Pending:{6}"),
-					bIsConnected ? LOCTEXT("IASConnect", "Connected") : LOCTEXT("IASDisconnect", "Disconnected"),
-					CachingText,
+				TStringBuilder<256> Builder;
+				Builder.Appendf(
+					TEXT("IAS - %s: %s | Downloaded:%d KiB (%d) %d ms | Retries:%d | Pending:%d"),
+					bIsConnected ? TEXT("Connected") : TEXT("Disconnected"),
+					CachingText.ToString(),
 					GHttpDownloadedBytes.Get() >> 10,
 					GHttpGetCount.Get(),
 					GHttpDurationMsAvg.Get(),
@@ -445,20 +445,20 @@ FOnDemandIoBackendStats::FOnDemandIoBackendStats(EStatsFlags InFlags)
 					GHttpPendingCount.Get()
 				);
 
-				Out.Add(FCoreDelegates::EOnScreenMessageSeverity::Info, MoveTemp(Message));
+				Out.Add(FCoreDelegates::EOnScreenMessageSeverity::Info, FText::FromStringView(Builder.ToView()));
 
 				if (GIoDecodeErrorCount.Get() > 0 || GHttpErrorCount.Get() > 0)
 				{
-					FText Warning = FText::Format(LOCTEXT("IASWarning", "IAS - Decode Errors {0} | Http Errors {1}"),
+					TStringBuilder<256> Warning;
+					Warning.Appendf(TEXT("IAS - Decode Errors %d | Http Errors %d"),
 						GIoDecodeErrorCount.Get(),
 						GHttpErrorCount.Get()
 					);
 
-					Out.Add(FCoreDelegates::EOnScreenMessageSeverity::Error, MoveTemp(Warning));
+					Out.Add(FCoreDelegates::EOnScreenMessageSeverity::Error, FText::FromStringView(Builder.ToView()));
 				}
 			}
 		});
-#undef LOCTEXT_NAMESPACE
 }
 
 FOnDemandIoBackendStats::~FOnDemandIoBackendStats()
