Chaos Cloth complex collider refactor and add support for SkinnedTriangleMesh colliders.

diff --git a/Engine/Source/Runtime/Experimental/Chaos/Private/Chaos/PBDCollisionSpringConstraintsBase.cpp b/Engine/Source/Runtime/Experimental/Chaos/Private/Chaos/PBDCollisionSpringConstraintsBase.cpp
--- a/Engine/Source/Runtime/Experimental/Chaos/Private/Chaos/PBDCollisionSpringConstraintsBase.cpp
+++ b/Engine/Source/Runtime/Experimental/Chaos/Private/Chaos/PBDCollisionSpringConstraintsBase.cpp
@@ -7,38 +7,11 @@
 #include "Chaos/PBDSoftsSolverParticles.h"
 #include "Chaos/SoftsEvolutionLinearSystem.h"
 #include "Chaos/SoftsSolverParticlesRange.h"
-#if INTEL_ISPC
-#include "PBDCollisionSpringConstraints.ispc.generated.h"
-#endif
 
 #if !COMPILE_WITHOUT_UNREAL_SUPPORT
 #include "Chaos/Framework/Parallel.h"
 #include <atomic>
 
-#if INTEL_ISPC && !UE_BUILD_SHIPPING
-static_assert(sizeof(ispc::FVector3f) == sizeof(Chaos::Softs::FSolverVec3), "sizeof(ispc::FVector3f) != sizeof(Chaos::Softs::FSolverVec3)");
-static_assert(sizeof(ispc::FVector4f) == sizeof(Chaos::Softs::FPAndInvM), "sizeof(ispc::FVector4f) != sizeof(Chaos::Softs::FPAndInvM");
-static_assert(sizeof(ispc::FIntVector) == sizeof(Chaos::TVec3<int32>), "sizeof(ispc::FIntVector) != sizeof(Chaos::TVec3<int32>");
-
-bool bChaos_CollisionSpring_ISPC_Enabled = CHAOS_COLLISION_SPRING_ISPC_ENABLED_DEFAULT;
-FAutoConsoleVariableRef CVarChaosCollisionSpringsISPCEnabled(TEXT("p.Chaos.CollisionSpring.ISPC"), bChaos_CollisionSpring_ISPC_Enabled, TEXT("Whether to use ISPC optimizations in collision spring constraints"));
-#endif
-
-#if UE_BUILD_SHIPPING
-static constexpr Chaos::Softs::FSolverReal KinematicColliderMaxTimer = (Chaos::Softs::FSolverReal)0.1f;
-static constexpr Chaos::Softs::FSolverReal KinematicColliderFalloffMultiplier = (Chaos::Softs::FSolverReal)1.f;
-static constexpr Chaos::Softs::FSolverReal KinematicColliderMaxDepthMultiplier = (Chaos::Softs::FSolverReal)10.f;
-#else
-Chaos::Softs::FSolverReal KinematicColliderMaxTimer = (Chaos::Softs::FSolverReal)0.1f; 
-FAutoConsoleVariableRef CVarChaosCollisionSpringMaxTimer(TEXT("p.Chaos.CollisionSpring.MaxTimer"), KinematicColliderMaxTimer, TEXT("Amount of time (in seconds) to remember a kinematic collision connection after it has moved more than Thickness away. Increasing this can reduce jitter at the cost of more computation."));
-
-Chaos::Softs::FSolverReal KinematicColliderFalloffMultiplier = (Chaos::Softs::FSolverReal)1.f;
-FAutoConsoleVariableRef CVarChaosCollisionSpringFalloffMultiplier(TEXT("p.Chaos.CollisionSpring.FalloffMultiplier"), KinematicColliderFalloffMultiplier, TEXT("Tangential distance away from a triangle (scaled by thickness) beyond which a point isn't considered to be kinematically colliding"));
-
-Chaos::Softs::FSolverReal KinematicColliderMaxDepthMultiplier = (Chaos::Softs::FSolverReal)10.f;
-FAutoConsoleVariableRef CVarChaosCollisionSpringMaxDepthMultiplier(TEXT("p.Chaos.CollisionSpring.MaxDepthMultiplier"), KinematicColliderMaxDepthMultiplier, TEXT("Penetration depth beyond which we ignore the kinematic collision (so you don't push through the wrong side)"));
-#endif
-
 namespace Chaos::Softs {
 
 static TConstArrayView<FSolverVec3> GetConstArrayView(const FSolverParticles& Particles, int32 Offset, int32 NumParticles, const TArray<FSolverVec3>* const Data)
@@ -80,9 +53,8 @@ FPBDCollisionSpringConstraintsBase::FPBDCollisionSpringConstraintsBase(
 	, Stiffness(InStiffness)
 	, FrictionCoefficient(InFrictionCoefficient)
 	, bOnlyCollideKinematics(bInOnlyCollideKinematics)
-	, KinematicColliderThickness(InKinematicColliderThickness)
-	, KinematicColliderStiffness(InKinematicColliderStiffness)
 	, KinematicColliderFrictionCoefficient(InKinematicColliderFrictionCoefficient, InKinematicColliderFrictionMultipliers, InNumParticles)
+	, KinematicCollisions(InNumParticles, ThicknessWeighted, KinematicColliderFrictionCoefficient, InKinematicColliderStiffness, InKinematicColliderThickness)
 	, ProximityStiffness(InProximityStiffness)
 	, TriangleMesh(InTriangleMesh)
 	, ReferencePositions(InReferencePositions)
@@ -91,9 +63,6 @@ FPBDCollisionSpringConstraintsBase::FPBDCollisionSpringConstraintsBase(
 	, NumParticles(InNumParticles)
 	, bGlobalIntersectionAnalysis(false)
 {
-	PRAGMA_DISABLE_DEPRECATION_WARNINGS
-	Thickness = GetMaxThickness();
-	PRAGMA_ENABLE_DEPRECATION_WARNINGS
 	UpdateCollisionLayers(InSelfCollisionLayers);
 }
 
@@ -153,24 +122,29 @@ template<typename SpatialAccelerator, typename SolverParticlesOrRange>
 void FPBDCollisionSpringConstraintsBase::Init(const SolverParticlesOrRange& Particles, const FSolverReal Dt, const FPBDTriangleMeshCollisions::FTriangleSubMesh& CollidableSubMesh,
 	const SpatialAccelerator& DynamicSpatial, const SpatialAccelerator& KinematicColliderSpatial, const TConstArrayView<FPBDTriangleMeshCollisions::FGIAColor>& VertexGIAColors, const TArray<FPBDTriangleMeshCollisions::FGIAColor>& TriangleGIAColors)
 {
-	PRAGMA_DISABLE_DEPRECATION_WARNINGS
+
+	if constexpr (std::is_same<SpatialAccelerator, THierarchicalSpatialHash<int32, FSolverReal>>::value && std::is_same<FSolverParticlesRange, SolverParticlesOrRange>::value)
+	{
+		KinematicCollisions.SetGeometry(CollidableSubMesh.GetKinematicColliderSubMesh(), Particles.XArray(), Particles.GetV(), KinematicColliderSpatial);
+		KinematicCollisions.Init(Particles, Dt);
+	}
+	else
+	{
+		// Kinematic collider only supports spatial hash type
+		KinematicCollisions.Reset();
+	}
+
 	const int32 NumDynamicElements = CollidableSubMesh.GetDynamicSubMesh().GetNumElements();
 
-	if (!NumDynamicElements && !bOnlyCollideKinematics)
+	if (!NumDynamicElements || bOnlyCollideKinematics)
 	{
 		Constraints.Reset();
 		Barys.Reset();
 		FlipNormal.Reset();
 		ConstraintTypes.Reset();
-		KinematicCollidingParticles.Reset();
-		KinematicColliderElements.Reset();
 		return;
 	}
 	{
-		if (KinematicColliderTimers.Num() != NumParticles)
-		{
-			KinematicColliderTimers.Init(TMap<int32, FSolverReal>(), NumParticles);
-		}
 
 		bGlobalIntersectionAnalysis = VertexGIAColors.Num() == NumParticles + Offset && TriangleGIAColors.Num() == NumDynamicElements;
 		TRACE_CPUPROFILER_EVENT_SCOPE(ChaosPBDCollisionSpring_ProximityQuery);
@@ -184,16 +158,12 @@ void FPBDCollisionSpringConstraintsBase::Init(const SolverParticlesOrRange& Part
 		FlipNormal.SetNumUninitialized(NumCollidableParticles * MaxConnectionsPerPoint);
 		ConstraintTypes.SetNumUninitialized(NumCollidableParticles * MaxConnectionsPerPoint);
 
-		KinematicCollidingParticles.SetNumUninitialized(NumCollidableParticles);
-		KinematicColliderElements.SetNumUninitialized(NumCollidableParticles);
-
 		std::atomic<int32> ConstraintIndex(0);
-		std::atomic<int32> KinematicConstraintIndex(0);
 
 		const TConstArrayView<FSolverVec3> ReferencePositionsView = GetConstArrayView(Particles, Offset, NumParticles, ReferencePositions);
 
 		PhysicsParallelFor(NumCollidableParticles,
-			[this, Dt, &CollidableSubMesh, &DynamicSpatial, &KinematicColliderSpatial, &Particles, &ConstraintIndex, &KinematicConstraintIndex, MaxConnectionsPerPoint, &VertexGIAColors, &TriangleGIAColors, &ReferencePositionsView](int32 CollidableIndex)
+			[this, Dt, &CollidableSubMesh, &DynamicSpatial, &Particles, &ConstraintIndex, MaxConnectionsPerPoint, &VertexGIAColors, &TriangleGIAColors, &ReferencePositionsView](int32 CollidableIndex)
 			{
 				const int32 i = CollidableSubMesh.GetDynamicVertices().IsEmpty() ? CollidableIndex : CollidableSubMesh.GetDynamicVertices()[CollidableIndex] - Offset; // DynamicVertices already has offset applied.
 				const int32 Index = i + Offset;
@@ -257,24 +227,21 @@ void FPBDCollisionSpringConstraintsBase::Init(const SolverParticlesOrRange& Part
 					return true;
 				};
 
-				if (!bOnlyCollideKinematics)
+				if (ThicknessWeighted.HasWeightMap())
 				{
-					if (ThicknessWeighted.HasWeightMap())
+					if constexpr (std::is_same<SpatialAccelerator, THierarchicalSpatialHash<int32, FSolverReal>>::value)
 					{
-						if constexpr (std::is_same<SpatialAccelerator, THierarchicalSpatialHash<int32, FSolverReal>>::value)
-						{
-							CollidableSubMesh.GetDynamicSubMesh().PointProximityQuery(DynamicSpatial, static_cast<const TArrayView<const FSolverVec3>&>(Particles.XArray()), Index, Particles.X(Index), ParticleThickness * ExtraThicknessMult, ThicknessWeighted, ExtraThicknessMult, Offset, BroadphaseTest, DynamicResult);
-						}
-						else
-						{
-							check(false);
-						}
+						CollidableSubMesh.GetDynamicSubMesh().PointProximityQuery(DynamicSpatial, static_cast<const TArrayView<const FSolverVec3>&>(Particles.XArray()), Index, Particles.X(Index), ParticleThickness * ExtraThicknessMult, ThicknessWeighted, ExtraThicknessMult, Offset, BroadphaseTest, DynamicResult);
 					}
 					else
 					{
-						CollidableSubMesh.GetDynamicSubMesh().PointProximityQuery(DynamicSpatial, static_cast<const TArrayView<const FSolverVec3>&>(Particles.XArray()), Index, Particles.X(Index), ParticleThickness * ExtraThicknessMult, (FSolverReal)ThicknessWeighted * ExtraThicknessMult, BroadphaseTest, DynamicResult);
+						check(false);
 					}
 				}
+				else
+				{
+					CollidableSubMesh.GetDynamicSubMesh().PointProximityQuery(DynamicSpatial, static_cast<const TArrayView<const FSolverVec3>&>(Particles.XArray()), Index, Particles.X(Index), ParticleThickness * ExtraThicknessMult, (FSolverReal)ThicknessWeighted * ExtraThicknessMult, BroadphaseTest, DynamicResult);
+				}
 
 				if ( DynamicResult.Num() )
 				{
@@ -363,70 +330,6 @@ void FPBDCollisionSpringConstraintsBase::Init(const SolverParticlesOrRange& Part
 						ConstraintTypes[IndexToWrite] = ConstraintType;
 					}
 				}
-
-				// Kinematic colliders
-				// Increment existing timers and remove any elements that are too old.
-				for(TMap<int32, FSolverReal>::TIterator TimerIter = KinematicColliderTimers[i].CreateIterator(); TimerIter; ++TimerIter)
-				{
-					TimerIter.Value() += Dt;
-					if (TimerIter.Value() > KinematicColliderMaxTimer)
-					{
-						TimerIter.RemoveCurrent();
-					}
-				}
-
-				if(CollidableSubMesh.GetKinematicColliderSubMesh().GetNumElements() > 0)
-				{
-					TArray< TTriangleCollisionPoint<FSolverReal> > KinematicResult;
-					if (CollidableSubMesh.GetKinematicColliderSubMesh().PointProximityQuery(KinematicColliderSpatial, static_cast<const TArrayView<const FSolverVec3>&>(Particles.XArray()), Index, Particles.X(Index), ParticleThickness * ExtraThicknessMult, KinematicColliderThickness * ExtraThicknessMult, [](const int32 PointIndex, const int32 SubMeshTriangleIndex)->bool { return true; }, KinematicResult))
-					{
-						if (KinematicResult.Num() > MaxKinematicConnectionsPerPoint)
-						{
-							// TODO: once we have a PartialSort, use that instead here.
-							KinematicResult.Sort(
-								[](const TTriangleCollisionPoint<FSolverReal>& First, const TTriangleCollisionPoint<FSolverReal>& Second)->bool
-							{
-								return First.Phi < Second.Phi;
-							}
-							);
-							KinematicResult.SetNum(MaxKinematicConnectionsPerPoint, EAllowShrinking::No);
-						}
-						for (const TTriangleCollisionPoint<FSolverReal>& CollisionPoint : KinematicResult)
-						{
-							const int32 FullMeshTriangleIndex = CollidableSubMesh.GetFullMeshElementIndexFromKinematicElement(CollisionPoint.Indices[1]);
-							KinematicColliderTimers[i].FindOrAdd(FullMeshTriangleIndex, (FSolverReal)0.f);
-						}
-					}
-				}
-
-				// Get MaxKinematicConnectionsPerPoint most recent
-				if (KinematicColliderTimers[i].Num() > MaxKinematicConnectionsPerPoint)
-				{
-					KinematicColliderTimers[i].ValueSort([](const FSolverReal A, const FSolverReal B) { return A < B; });
-				}
-
-				if (KinematicColliderTimers[i].Num() > 0)
-				{
-					const int32 IndexToWrite = KinematicConstraintIndex.fetch_add(1);
-					KinematicCollidingParticles[IndexToWrite] = Index;
-
-					int32 LocalIndex = 0;
-					for (TMap<int32, FSolverReal>::TIterator TimerIter = KinematicColliderTimers[i].CreateIterator(); TimerIter; ++TimerIter)
-					{
-						if (LocalIndex < MaxKinematicConnectionsPerPoint)
-						{
-							KinematicColliderElements[IndexToWrite][LocalIndex++] = TimerIter.Key();
-						}
-						else
-						{
-							TimerIter.RemoveCurrent();
-						}
-					}
-					for (; LocalIndex < MaxKinematicConnectionsPerPoint; ++LocalIndex)
-					{
-						KinematicColliderElements[IndexToWrite][LocalIndex] = INDEX_NONE;
-					}
-				}
 			}
 		);
 
@@ -436,12 +339,7 @@ void FPBDCollisionSpringConstraintsBase::Init(const SolverParticlesOrRange& Part
 		Barys.SetNum(ConstraintNum, EAllowShrinking::No);
 		FlipNormal.SetNum(ConstraintNum, EAllowShrinking::No);
 		ConstraintTypes.SetNum(ConstraintNum, EAllowShrinking::No);
-
-		const int32 KinematicConstraintNum = KinematicConstraintIndex.load();
-		KinematicCollidingParticles.SetNum(KinematicConstraintNum, EAllowShrinking::No);
-		KinematicColliderElements.SetNum(KinematicConstraintNum, EAllowShrinking::No);
 	}
-	PRAGMA_ENABLE_DEPRECATION_WARNINGS
 }
 template void CHAOS_API FPBDCollisionSpringConstraintsBase::Init<FTriangleMesh::TBVHType<FSolverReal>>(const FSolverParticles& Particles, const FSolverReal Dt,
 	const FPBDTriangleMeshCollisions::FTriangleSubMesh& CollidableSubMesh, const FTriangleMesh::TBVHType<FSolverReal>& DynamicSpatial, const FTriangleMesh::TBVHType<FSolverReal>& KinematicColliderSpatial,
@@ -459,7 +357,6 @@ template void CHAOS_API FPBDCollisionSpringConstraintsBase::Init<FTriangleMesh::
 template<typename SolverParticlesOrRange>
 FSolverVec3 FPBDCollisionSpringConstraintsBase::GetDelta(const SolverParticlesOrRange& Particles, const int32 ConstraintIndex) const
 {
-	PRAGMA_DISABLE_DEPRECATION_WARNINGS
 	const TVec4<int32>& Constraint = Constraints[ConstraintIndex];
 	const int32 Index1 = Constraint[0];
 	const int32 Index2 = Constraint[1];
@@ -519,14 +416,12 @@ FSolverVec3 FPBDCollisionSpringConstraintsBase::GetDelta(const SolverParticlesOr
 	{
 		return RepulsionDelta;
 	}
-	PRAGMA_ENABLE_DEPRECATION_WARNINGS
 }
 template CHAOS_API FSolverVec3 FPBDCollisionSpringConstraintsBase::GetDelta(const FSolverParticles& Particles, const int32 i) const;
 template CHAOS_API FSolverVec3 FPBDCollisionSpringConstraintsBase::GetDelta(const FSolverParticlesRange& Particles, const int32 i) const;
 
 void FPBDCollisionSpringConstraintsBase::UpdateLinearSystem(const FSolverParticlesRange& Particles, const FSolverReal Dt, FEvolutionLinearSystem& LinearSystem) const
 {
-	PRAGMA_DISABLE_DEPRECATION_WARNINGS
 	LinearSystem.ReserveForParallelAdd(Constraints.Num() * 4, Constraints.Num() * 3);
 	for (int32 Index = 0; Index < Constraints.Num(); ++Index)
 	{
@@ -590,7 +485,6 @@ void FPBDCollisionSpringConstraintsBase::UpdateLinearSystem(const FSolverParticl
 			LinearSystem.AddSymmetricForceDerivative(Particles, &DfDxScaled, nullptr, Index4, Index4, Dt);
 		}
 	}
-	PRAGMA_ENABLE_DEPRECATION_WARNINGS
 }
 
 
@@ -598,7 +492,10 @@ template<typename SolverParticlesOrRange>
 void FPBDCollisionSpringConstraintsBase::Apply(SolverParticlesOrRange& InParticles, const FSolverReal Dt) const
 {
 	ApplyDynamicConstraints(InParticles, Dt);
-	ApplyKinematicConstraints(InParticles, Dt);
+	if constexpr (std::is_same<SolverParticlesOrRange, FSolverParticlesRange>::value)
+	{
+		KinematicCollisions.Apply(InParticles, Dt);
+	}
 }
 template CHAOS_API void FPBDCollisionSpringConstraintsBase::Apply(FSolverParticles& Particles, const FSolverReal Dt) const;
 template CHAOS_API void FPBDCollisionSpringConstraintsBase::Apply(FSolverParticlesRange& Particles, const FSolverReal Dt) const;
@@ -607,7 +504,6 @@ template<typename SolverParticlesOrRange>
 void FPBDCollisionSpringConstraintsBase::ApplyDynamicConstraints(SolverParticlesOrRange& Particles, const FSolverReal Dt) const
 {
 	TRACE_CPUPROFILER_EVENT_SCOPE(ChaosPBDCollisionSpring_ApplyDynamic);
-	PRAGMA_DISABLE_DEPRECATION_WARNINGS
 
 	check(Constraints.Num() == Barys.Num());
 	check(Constraints.Num() == FlipNormal.Num());
@@ -679,157 +575,6 @@ void FPBDCollisionSpringConstraintsBase::ApplyDynamicConstraints(SolverParticles
 		P3 -= PAndInvM[Index3].InvM * BarysData[ConstraintIndex][1] * Delta;
 		P4 -= PAndInvM[Index4].InvM * BarysData[ConstraintIndex][2] * Delta;
 	}
-	PRAGMA_ENABLE_DEPRECATION_WARNINGS
-}
-
-template<typename SolverParticlesOrRange>
-void FPBDCollisionSpringConstraintsBase::ApplyKinematicConstraints(SolverParticlesOrRange& Particles, const FSolverReal Dt) const
-{
-	if (!KinematicCollidingParticles.Num())
-	{
-		return;
-	}
-	TRACE_CPUPROFILER_EVENT_SCOPE(ChaosPBDCollisionSpring_ApplyKinematicConstraints);
-
-
-#if INTEL_ISPC
-	static_assert(sizeof(ispc::FIntVector) == sizeof(TVector<int32, MaxKinematicConnectionsPerPoint>), "sizeof(ispc::FIntVector) != sizeof(TVector<int32, MaxKinematicConnectionsPerPoint>)");
-	if (bRealTypeCompatibleWithISPC && bChaos_CollisionSpring_ISPC_Enabled)
-	{
-		if (ThicknessWeighted.HasWeightMap())
-		{
-			ispc::ApplyKinematicCollisionSpringConstraintsWithMaps(
-				(ispc::FVector4f*)Particles.GetPAndInvM().GetData(),
-				(const ispc::FVector3f*)Particles.XArray().GetData(),
-				KinematicCollidingParticles.GetData(),
-				(const ispc::FIntVector*)KinematicColliderElements.GetData(),
-				(const ispc::FIntVector*)TriangleMesh.GetElements().GetData(),
-				reinterpret_cast<const ispc::FVector2f&>(ThicknessWeighted.GetOffsetRange()),
-				ThicknessWeighted.GetMapValues().GetData(),
-				KinematicColliderThickness,
-				KinematicColliderFalloffMultiplier,
-				KinematicColliderMaxDepthMultiplier,
-				KinematicColliderStiffness,
-				KinematicColliderFrictionCoefficient.HasWeightMap(),
-				reinterpret_cast<const ispc::FVector2f&>(KinematicColliderFrictionCoefficient.GetOffsetRange()),
-				KinematicColliderFrictionCoefficient.GetMapValues().GetData(),
-				Offset,
-				KinematicCollidingParticles.Num()
-			);
-		}
-		else
-		{
-			const FSolverReal Height = (FSolverReal)ThicknessWeighted + KinematicColliderThickness;
-			const FSolverReal OneOverTangentialFalloffDist = (FSolverReal)1.f / FMath::Max(KinematicColliderFalloffMultiplier * Height, UE_KINDA_SMALL_NUMBER);
-			const FSolverReal MaxDepth = -Height * KinematicColliderMaxDepthMultiplier;
-			if (KinematicColliderFrictionCoefficient.HasWeightMap())
-			{
-				ispc::ApplyKinematicCollisionSpringConstraintsWithKinematicFrictionMap(
-					(ispc::FVector4f*)Particles.GetPAndInvM().GetData(),
-					(const ispc::FVector3f*)Particles.XArray().GetData(),
-					KinematicCollidingParticles.GetData(),
-					(const ispc::FIntVector*)KinematicColliderElements.GetData(),
-					(const ispc::FIntVector*)TriangleMesh.GetElements().GetData(),
-					Height,
-					OneOverTangentialFalloffDist,
-					MaxDepth,
-					KinematicColliderStiffness,
-					reinterpret_cast<const ispc::FVector2f&>(KinematicColliderFrictionCoefficient.GetOffsetRange()),
-					KinematicColliderFrictionCoefficient.GetMapValues().GetData(),
-					Offset,
-					KinematicCollidingParticles.Num()
-				);
-
-			}
-			else
-			{
-				ispc::ApplyKinematicCollisionSpringConstraints(
-					(ispc::FVector4f*)Particles.GetPAndInvM().GetData(),
-					(const ispc::FVector3f*)Particles.XArray().GetData(),
-					KinematicCollidingParticles.GetData(),
-					(const ispc::FIntVector*)KinematicColliderElements.GetData(),
-					(const ispc::FIntVector*)TriangleMesh.GetElements().GetData(),
-					Height,
-					OneOverTangentialFalloffDist,
-					MaxDepth,
-					KinematicColliderStiffness,
-					(FSolverReal)KinematicColliderFrictionCoefficient,
-					KinematicCollidingParticles.Num()
-				);
-			}
-		}
-	}
-	else
-#endif
-	{
-		for (int32 Index = 0; Index < KinematicCollidingParticles.Num(); ++Index)
-		{
-			const int32 Index1 = KinematicCollidingParticles[Index];
-			const FSolverReal ParticleThickness = ThicknessWeighted.GetValue(Index1 - Offset);
-
-			const FSolverReal Height = ParticleThickness + KinematicColliderThickness;
-			const FSolverReal OneOverTangentialFalloffDist = (FSolverReal)1.f / FMath::Max(KinematicColliderFalloffMultiplier * Height, UE_KINDA_SMALL_NUMBER);
-			const FSolverReal MaxDepth = -Height * KinematicColliderMaxDepthMultiplier;
-
-			for (int32 EIndex = 0; EIndex < MaxKinematicConnectionsPerPoint; ++EIndex)
-			{
-				const int32 ElemIndex = KinematicColliderElements[Index][EIndex];
-				if (ElemIndex == INDEX_NONE)
-				{
-					break;
-				}
-
-				const int32 Index2 = TriangleMesh.GetElements()[ElemIndex][0];
-				const int32 Index3 = TriangleMesh.GetElements()[ElemIndex][1];
-				const int32 Index4 = TriangleMesh.GetElements()[ElemIndex][2];
-
-				FSolverVec3& P1 = Particles.P(Index1);
-				const FSolverVec3& P2 = Particles.P(Index2);
-				const FSolverVec3& P3 = Particles.P(Index3);
-				const FSolverVec3& P4 = Particles.P(Index4);
-
-				const TTriangle<FSolverReal> Triangle(P2, P3, P4);
-				const FSolverVec3  Normal = -Triangle.GetNormal(); // normals are flipped from UE
-
-				FSolverVec3 Bary;
-				const FSolverVec3 P = FindClosestPointAndBaryOnTriangle(P2, P3, P4, P1, Bary);
-				const FSolverVec3 Difference = P1 - P;
-				const FSolverReal NormalDifference = Difference.Dot(Normal);
-
-				if (NormalDifference >= Height || NormalDifference < MaxDepth)
-				{
-					continue;
-				}
-
-				const FSolverReal TangentialDifference = (Difference - NormalDifference * Normal).Size();
-				const FSolverReal TangentialFalloff = (FSolverReal)1.f - TangentialDifference * OneOverTangentialFalloffDist;
-				if (TangentialFalloff <= 0.f)
-				{
-					continue;
-				}
-
-				const FSolverReal NormalDelta = Height - NormalDifference;
-				const FSolverVec3 RepulsionDelta = KinematicColliderStiffness * TangentialFalloff * NormalDelta * Normal;
-
-				P1 += RepulsionDelta;
-
-				const FSolverReal KinematicFrictionCoefficient = KinematicColliderFrictionCoefficient.GetValue(Index1 - Offset);
-				if (KinematicFrictionCoefficient > 0)
-				{
-					const FSolverVec3& X1 = Particles.X(Index1);
-					const FSolverVec3 X = Bary[0] * Particles.X(Index2) + Bary[1] * Particles.X(Index3) + Bary[2] * Particles.X(Index4);
-					const FSolverVec3 RelativeDisplacement = (P1 - X1) - (P - X);
-					const FSolverVec3 RelativeDisplacementTangent = RelativeDisplacement - RelativeDisplacement.Dot(Normal) * Normal;
-					const FSolverReal RelativeDisplacementTangentLength = RelativeDisplacementTangent.Length();
-					const FSolverReal PositionCorrection = FMath::Min(NormalDelta * KinematicFrictionCoefficient, RelativeDisplacementTangentLength);
-					const FSolverReal CorrectionRatio = RelativeDisplacementTangentLength < UE_SMALL_NUMBER ? 0.f : PositionCorrection / RelativeDisplacementTangentLength;
-					const FSolverVec3 FrictionDelta = -CorrectionRatio * RelativeDisplacementTangent;
-
-					P1 += FrictionDelta;
-				}
-			}
-		}
-	}
 }
 
 }  // End namespace Chaos::Softs
