Chaos Cloth complex collider refactor and add support for SkinnedTriangleMesh colliders.

diff --git a/Engine/Source/Runtime/Experimental/Chaos/Private/Chaos/PBDSoftBodyCollisionConstraint.cpp b/Engine/Source/Runtime/Experimental/Chaos/Private/Chaos/PBDSoftBodyCollisionConstraint.cpp
--- a/Engine/Source/Runtime/Experimental/Chaos/Private/Chaos/PBDSoftBodyCollisionConstraint.cpp
+++ b/Engine/Source/Runtime/Experimental/Chaos/Private/Chaos/PBDSoftBodyCollisionConstraint.cpp
@@ -29,10 +29,9 @@ static_assert(ispc::ImplicitObjectType::Capsule == Chaos::ImplicitObjectType::Ca
 static_assert(ispc::ImplicitObjectType::Union == Chaos::ImplicitObjectType::Union);
 static_assert(ispc::ImplicitObjectType::TaperedCapsule == Chaos::ImplicitObjectType::TaperedCapsule);
 static_assert(ispc::ImplicitObjectType::Convex == Chaos::ImplicitObjectType::Convex);
-static_assert(ispc::ImplicitObjectType::WeightedLatticeBone == Chaos::ImplicitObjectType::WeightedLatticeBone);
 static_assert(ispc::ImplicitObjectType::IsWeightedLattice == Chaos::ImplicitObjectType::IsWeightedLattice);
 static_assert(ispc::ImplicitObjectType::MLLevelSet == Chaos::ImplicitObjectType::MLLevelSet);
-static_assert(ispc::ImplicitObjectType::MLLevelSetBone == Chaos::ImplicitObjectType::MLLevelSetBone);
+static_assert(ispc::ImplicitObjectType::SkinnedTriangleMesh == Chaos::ImplicitObjectType::SkinnedTriangleMesh);
 static_assert(ispc::ImplicitObjectType::WeightedLatticeLevelSetType == (Chaos::ImplicitObjectType::IsWeightedLattice | Chaos::ImplicitObjectType::LevelSet));
 static_assert(sizeof(ispc::TArray) == sizeof(TArray<int>));
 // Sphere
@@ -119,15 +118,59 @@ static FAutoConsoleVariableRef CVarChaosSoftBodyCollisionISPCEnabled(TEXT("p.Cha
 #endif
 #endif
 
-static int MLLevelSetParallelize = 0;
 static int32 Chaos_SoftBodyCollision_ISPC_ParallelBatchSize = 128;
 #if !UE_BUILD_SHIPPING
 FAutoConsoleVariableRef CVarChaosSoftBodyCollisionISPCParallelBatchSize(TEXT("p.Chaos.SoftBodyCollision.ISPC.ParallelBatchSize"), Chaos_SoftBodyCollision_ISPC_ParallelBatchSize, TEXT("Parallel batch size for ISPC"));
-static FAutoConsoleVariableRef CVarChaosMLLevelSetParallelize(TEXT("p.Chaos.Cloth.ParallelizeMLLevelSet"), MLLevelSetParallelize, TEXT("Parallelize MLLevelSet Call: 0 = No, 1 = Yes"));
 #endif
 
 namespace Chaos::Softs {
 
+	namespace Private
+	{
+		static bool IsComplexBatchCollider(Chaos::EImplicitObjectType CollisionType)
+		{
+			return CollisionType == TWeightedLatticeImplicitObject<FLevelSet>::StaticType() || CollisionType == Chaos::ImplicitObjectType::MLLevelSet;
+		}
+
+		static bool IsSimpleCollider(Chaos::EImplicitObjectType CollisionType)
+		{
+			if (IsComplexBatchCollider(CollisionType))
+			{
+				return false;
+			}
+
+			if(CollisionType == Chaos::ImplicitObjectType::SkinnedTriangleMesh)
+			{
+				// Note: SkinnedTriangleMesh collisions are handled by FPBDSkinnedTriangleMeshCollisions
+				return false;
+			}
+
+			return true;
+		}
+
+		static void ApplyFriction(FSolverVec3& P, const FSolverVec3& X, const FSolverVec3& NormalWorld, const FSolverReal MaxFrictionCorrection, const FSolverReal Dt, const FSolverVec3& CollisionX, const FSolverVec3& CollisionV, const FSolverVec3& CollisionW, FSolverVec3& ColliderVelocityAtPoint)
+		{
+			const FSolverVec3 VectorToPoint = P - CollisionX;
+			ColliderVelocityAtPoint = CollisionV + FSolverVec3::CrossProduct(CollisionW, VectorToPoint);
+
+			const FSolverVec3 RelativeDisplacement = (P - X) - ColliderVelocityAtPoint * Dt; // This corresponds to the tangential velocity multiplied by dt (friction will drive this to zero if it is high enough)
+			const FSolverVec3 RelativeDisplacementTangent = RelativeDisplacement - FSolverVec3::DotProduct(RelativeDisplacement, NormalWorld) * NormalWorld; // Project displacement into the tangential plane
+			const FSolverReal RelativeDisplacementTangentLength = RelativeDisplacementTangent.Size();
+			if (RelativeDisplacementTangentLength >= UE_SMALL_NUMBER)
+			{
+				const FSolverReal PositionCorrection = FMath::Min<FSolverReal>(MaxFrictionCorrection, RelativeDisplacementTangentLength);
+				const FSolverReal CorrectionRatio = PositionCorrection / RelativeDisplacementTangentLength;
+				P -= CorrectionRatio * RelativeDisplacementTangent;
+			}
+		}
+
+		static void ApplyFriction(FSolverVec3& P, const FSolverVec3& X, const FSolverVec3& NormalWorld, const FSolverReal MaxFrictionCorrection, const FSolverReal Dt, const FSolverVec3& CollisionX, const FSolverVec3& CollisionV, const FSolverVec3& CollisionW)
+		{
+			FSolverVec3 ColliderVelocityAtPoint;
+			ApplyFriction(P, X, NormalWorld, MaxFrictionCorrection, Dt, CollisionX, CollisionV, CollisionW, ColliderVelocityAtPoint);
+		}
+	}
+
 // Helper function to call PhiWithNormal and return data to ISPC
 extern "C" void GetPhiWithNormalCollisionParticleRange(const uint8 * CollisionParticlesRange, const FSolverReal * InV, FSolverReal * Normal, FSolverReal * Phi, const int32 i, const int32 ProgramCount, const int32 Mask)
 {
@@ -155,65 +198,226 @@ extern "C" void GetPhiWithNormalCollisionParticleRange(const uint8 * CollisionPa
 	}
 }
 
-// Helper function to call PhiWithNormal and return data to ISPC
-extern "C" void GetPhiWithNormalAndVelocityBoneCollisionParticleRange(const uint8 * CollisionParticlesRange, const FSolverReal * InV, FSolverReal * Normal, FSolverReal * Phi, int32 * VelocityBone, const int32 i, const FSolverReal Thickness, const int32 ProgramCount, const int32 Mask)
+void FPBDSoftBodyCollisionConstraintBase::ApplyWithPlanarConstraints(FSolverParticlesRange& Particles, const FSolverReal Dt, const TArray<FSolverCollisionParticlesRange>& CollisionParticles, bool bGeneratePlanarConstraints)
 {
-	const FSolverCollisionParticlesRange& C = *(const FSolverCollisionParticlesRange*)CollisionParticlesRange;
+	TRACE_CPUPROFILER_EVENT_SCOPE(FPBDSoftBodyCollisionConstraint_ApplyWithPlanarConstraints);
 
-	for (int32 Index = 0; Index < ProgramCount; ++Index)
+	if (CollisionParticles.IsEmpty())
 	{
-		if (Mask & (1 << Index))
-		{
-			FSolverVec3 V;
-
-			// aos_to_soa3
-			V.X = InV[Index];
-			V.Y = InV[Index + ProgramCount];
-			V.Z = InV[Index + 2 * ProgramCount];
-
-			VelocityBone[Index] = i;
+		return;
+	}
 
-			FVec3 ImplicitNormal;
-			if (const TWeightedLatticeImplicitObject<FLevelSet>* LevelSet = C.GetGeometry(i)->GetObject< TWeightedLatticeImplicitObject<FLevelSet> >())
+	const bool bLockAndWriteContacts = bWriteDebugContacts && CollisionParticleCollided && Contacts && Normals && Phis;
+	const bool bWithFriction = FrictionCoefficient > (FSolverReal)UE_KINDA_SMALL_NUMBER;
+	if (bUseCCD)
+	{
+		if (bWithFriction)
+		{
+			if (bLockAndWriteContacts)
 			{
-				FWeightedLatticeImplicitObject::FEmbeddingCoordinate SurfaceCoord;
-				Phi[Index] = (FSolverReal)LevelSet->PhiWithNormalAndSurfacePoint(FVec3(V), ImplicitNormal, SurfaceCoord);
-				const FSolverReal Penetration = Thickness - Phi[Index]; // This is related to the Normal impulse
-				if (Penetration > (FSolverReal)0.)
+				ApplyInternalCCD<true, true>(Particles, Dt, CollisionParticles);
+			}
+			else
+			{
+				ApplyInternalCCD<true, false>(Particles, Dt, CollisionParticles);
+			}
+		}
+		else
+		{
+			if (bLockAndWriteContacts)
+			{
+				ApplyInternalCCD<false, true>(Particles, Dt, CollisionParticles);
+			}
+			else
+			{
+				ApplyInternalCCD<false, false>(Particles, Dt, CollisionParticles);
+			}
+		}
+	}
+	else
+	{
+		if (bGeneratePlanarConstraints)
+		{
+			InitPlanarConstraints(Particles, bWithFriction);
+#if INTEL_ISPC
+			if (bChaos_SoftBodyCollision_ISPC_Enabled && bRealTypeCompatibleWithISPC)
+			{
+				ApplySimpleInternalISPC(Particles, Dt, CollisionParticles, bUsePlanarConstraintForSimpleColliders);
+				ApplyComplexInternalISPC(Particles, Dt, CollisionParticles, bUsePlanarConstraintForComplexColliders);
+			}
+			else
+#endif
+			{
+				if (bLockAndWriteContacts)
 				{
-					const int32 StrongestBone = SurfaceCoord.GreatestInfluenceBone(LevelSet->GetBoneData());
-					if (StrongestBone != INDEX_NONE)
+					if (bWithFriction)
 					{
-						VelocityBone[Index] = LevelSet->GetSolverBoneIndices()[StrongestBone];
+						if (bUsePlanarConstraintForSimpleColliders)
+						{
+							ApplySimpleInternal<true, true, true>(Particles, Dt, CollisionParticles);
+						}
+						else
+						{
+							ApplySimpleInternal<true, true, false>(Particles, Dt, CollisionParticles);
+						}
+						if (bUsePlanarConstraintForComplexColliders)
+						{
+							ApplyComplexInternal<true, true, true>(Particles, Dt, CollisionParticles);
+						}
+						else
+						{
+							ApplyComplexInternal<true, true, false>(Particles, Dt, CollisionParticles);
+						}
+					}
+					else
+					{
+						if (bUsePlanarConstraintForSimpleColliders)
+						{
+							ApplySimpleInternal<true, false, true>(Particles, Dt, CollisionParticles);
+						}
+						else
+						{
+							ApplySimpleInternal<true, false, false>(Particles, Dt, CollisionParticles);
+						}
+						if (bUsePlanarConstraintForComplexColliders)
+						{
+							ApplyComplexInternal<true, false, true>(Particles, Dt, CollisionParticles);
+						}
+						else
+						{
+							ApplyComplexInternal<true, false, false>(Particles, Dt, CollisionParticles);
+						}
+					}
+				}
+				else
+				{
+					if (bWithFriction)
+					{
+						if (bUsePlanarConstraintForSimpleColliders)
+						{
+							ApplySimpleInternal<false, true, true>(Particles, Dt, CollisionParticles);
+						}
+						else
+						{
+							ApplySimpleInternal<false, true, false>(Particles, Dt, CollisionParticles);
+						}
+						if (bUsePlanarConstraintForComplexColliders)
+						{
+							ApplyComplexInternal<false, true, true>(Particles, Dt, CollisionParticles);
+						}
+						else
+						{
+							ApplyComplexInternal<false, true, false>(Particles, Dt, CollisionParticles);
+						}
+					}
+					else
+					{
+						if (bUsePlanarConstraintForSimpleColliders)
+						{
+							ApplySimpleInternal<false, false, true>(Particles, Dt, CollisionParticles);
+						}
+						else
+						{
+							ApplySimpleInternal<false, false, false>(Particles, Dt, CollisionParticles);
+						}
+						if (bUsePlanarConstraintForComplexColliders)
+						{
+							ApplyComplexInternal<false, false, true>(Particles, Dt, CollisionParticles);
+						}
+						else
+						{
+							ApplyComplexInternal<false, false, false>(Particles, Dt, CollisionParticles);
+						}
 					}
 				}
 			}
+			FinalizePlanarConstraints(Particles);
+		}
+		else
+		{
+#if INTEL_ISPC
+			if (bChaos_SoftBodyCollision_ISPC_Enabled && bRealTypeCompatibleWithISPC)
+			{
+				if (!bUsePlanarConstraintForSimpleColliders)
+				{
+					ApplySimpleInternalISPC(Particles, Dt, CollisionParticles, false);
+				}
+				if (!bUsePlanarConstraintForComplexColliders)
+				{
+					ApplyComplexInternalISPC(Particles, Dt, CollisionParticles, false);
+				}
+			}
 			else
+#endif
 			{
-				Phi[Index] = (FSolverReal)C.GetGeometry(i)->PhiWithNormal(FVec3(V), ImplicitNormal);
+				if (bLockAndWriteContacts)
+				{
+					if (bWithFriction)
+					{
+						if (!bUsePlanarConstraintForSimpleColliders)
+						{
+							ApplySimpleInternal<true, true, false>(Particles, Dt, CollisionParticles);
+						}
+						if (!bUsePlanarConstraintForComplexColliders)
+						{
+							ApplyComplexInternal<true, true, false>(Particles, Dt, CollisionParticles);
+						}
+					}
+					else
+					{
+						if (!bUsePlanarConstraintForSimpleColliders)
+						{
+							ApplySimpleInternal<true, false, false>(Particles, Dt, CollisionParticles);
+						}
+						if (!bUsePlanarConstraintForComplexColliders)
+						{
+							ApplyComplexInternal<true, false, false>(Particles, Dt, CollisionParticles);
+						}
+					}
+				}
+				else
+				{
+					if (bWithFriction)
+					{
+						if (!bUsePlanarConstraintForSimpleColliders)
+						{
+							ApplySimpleInternal<false, true, false>(Particles, Dt, CollisionParticles);
+						}
+						if (!bUsePlanarConstraintForComplexColliders)
+						{
+							ApplyComplexInternal<false, true, false>(Particles, Dt, CollisionParticles);
+						}
+					}
+					else
+					{
+						if (!bUsePlanarConstraintForSimpleColliders)
+						{
+							ApplySimpleInternal<false, false, false>(Particles, Dt, CollisionParticles);
+						}
+						if (!bUsePlanarConstraintForComplexColliders)
+						{
+							ApplyComplexInternal<false, false, false>(Particles, Dt, CollisionParticles);
+						}
+					}
+				}
 			}
-			FSolverVec3 Norm(ImplicitNormal);
-
-			// aos_to_soa3
-			Normal[Index] = Norm.X;
-			Normal[Index + ProgramCount] = Norm.Y;
-			Normal[Index + 2 * ProgramCount] = Norm.Z;
+			ApplyPlanarConstraints(Particles, Dt);
 		}
 	}
 }
 
 
-void FPBDSoftBodyCollisionConstraintBase::Apply(FSolverParticlesRange& Particles, const FSolverReal Dt, const TArray<FSolverCollisionParticlesRange>& CollisionParticles) const
+void FPBDSoftBodyCollisionConstraintBase::Apply(FSolverParticlesRange& Particles, const FSolverReal Dt, const TArray<FSolverCollisionParticlesRange>& CollisionParticles)
 {
 	TRACE_CPUPROFILER_EVENT_SCOPE(FPBDSoftBodyCollisionConstraint);
-	const bool bLockAndWriteContacts = bWriteDebugContacts && CollisionParticleCollided && Contacts && Normals && Phis;
-	const bool bWithFriction = FrictionCoefficient > (FSolverReal)UE_KINDA_SMALL_NUMBER;
 
 	if (CollisionParticles.IsEmpty())
 	{
 		return;
 	}
 
+	const bool bLockAndWriteContacts = bWriteDebugContacts && CollisionParticleCollided && Contacts && Normals && Phis;
+	const bool bWithFriction = FrictionCoefficient > (FSolverReal)UE_KINDA_SMALL_NUMBER;
 	if (bUseCCD)
 	{
 		if (bWithFriction)
@@ -241,193 +445,396 @@ void FPBDSoftBodyCollisionConstraintBase::Apply(FSolverParticlesRange& Particles
 	}
 	else
 	{
-		if (bWithFriction)
+#if INTEL_ISPC
+		if (bChaos_SoftBodyCollision_ISPC_Enabled && bRealTypeCompatibleWithISPC)
 		{
-			if (bDoCollisionPerTimestep)
-			{	
-				ApplyInternalISPCPerTimestep<true>(Particles, Dt, CollisionParticles);				
-			}
-			else if (bRealTypeCompatibleWithISPC && bChaos_SoftBodyCollision_ISPC_Enabled)
-			{
-				ApplyInternalISPC<true>(Particles, Dt, CollisionParticles);
-			}
-			else if (bLockAndWriteContacts)
-			{
-				ApplyInternal<true, true>(Particles, Dt, CollisionParticles);
-			}
-			else
-			{
-				ApplyInternal<true, false>(Particles, Dt, CollisionParticles);
-			}
+			ApplySimpleInternalISPC(Particles, Dt, CollisionParticles, false);
+			ApplyComplexInternalISPC(Particles, Dt, CollisionParticles, false);
 		}
 		else
+#endif
 		{
-			if (bDoCollisionPerTimestep)
-			{
-				ApplyInternalISPCPerTimestep<false>(Particles, Dt, CollisionParticles);				
-			}
-			else if (bRealTypeCompatibleWithISPC && bChaos_SoftBodyCollision_ISPC_Enabled)
-			{
-				ApplyInternalISPC<false>(Particles, Dt, CollisionParticles);
-			}
-			else if (bLockAndWriteContacts)
+			if (bLockAndWriteContacts)
 			{
-				ApplyInternal<false, true>(Particles, Dt, CollisionParticles);
+				if (bWithFriction)
+				{
+					ApplySimpleInternal<true, true, false>(Particles, Dt, CollisionParticles);
+					ApplyComplexInternal<true, true, false>(Particles, Dt, CollisionParticles);
+				}
+				else
+				{
+					ApplySimpleInternal<true, false, false>(Particles, Dt, CollisionParticles);
+					ApplyComplexInternal<true, false, false>(Particles, Dt, CollisionParticles);
+				}
 			}
 			else
 			{
-				ApplyInternal<false, false>(Particles, Dt, CollisionParticles);
+				if (bWithFriction)
+				{
+					ApplySimpleInternal<false, true, false>(Particles, Dt, CollisionParticles);
+					ApplyComplexInternal<false, true, false>(Particles, Dt, CollisionParticles);
+				}
+				else
+				{
+					ApplySimpleInternal<false, false, false>(Particles, Dt, CollisionParticles);
+					ApplyComplexInternal<false, false, false>(Particles, Dt, CollisionParticles);
+				}
 			}
 		}
 	}
 }
 
+void FPBDSoftBodyCollisionConstraintBase::InitPlanarConstraints(const FSolverParticlesRange& Particles, bool bWithFriction)
+{
+	// Initialize Planar Constraint.
+	PlanarConstraint.Reset();
+	HasPlanarData.Init(false, Particles.GetRangeSize());
+	PlanarDataPositions.SetNumUninitialized(Particles.GetRangeSize());
+	PlanarDataNormals.SetNumUninitialized(Particles.GetRangeSize());
+	if (bWithFriction)
+	{
+		PlanarDataVelocities.SetNumUninitialized(Particles.GetRangeSize());
+	}
+}
 
-template<bool bLockAndWriteContacts, bool bWithFriction>
-void FPBDSoftBodyCollisionConstraintBase::ApplyInternal(FSolverParticlesRange& Particles, const FSolverReal Dt, const TArray<FSolverCollisionParticlesRange>& CollisionParticles) const
+template<bool bLockAndWriteContacts, bool bWithFriction, bool bGeneratePlanarConstraints>
+void FPBDSoftBodyCollisionConstraintBase::ApplySimpleInternal(FSolverParticlesRange& Particles, const FSolverReal Dt, const TArray<FSolverCollisionParticlesRange>& CollisionParticles)
 {
+	TRACE_CPUPROFILER_EVENT_SCOPE(FPBDSoftBodyCollisionConstraint_ApplySimpleInternal);
+
+	if (!bEnableSimpleColliders)
+	{
+		return;
+	}
+
 	FPAndInvM* const PAndInvM = Particles.GetPAndInvM().GetData();
 	const FSolverVec3* const X = Particles.XArray().GetData();
-	PhysicsParallelFor(Particles.GetRangeSize(), [this, Dt, &PAndInvM, &X, &CollisionParticles](int32 Index)
-	{
-		if (PAndInvM[Index].InvM == (FSolverReal)0.)
-		{
-			return;
-		}
 
-		for (const FSolverCollisionParticlesRange& CollisionParticlesRange : CollisionParticles)
+	PhysicsParallelFor(Particles.GetRangeSize(), [this, &X, &PAndInvM, &CollisionParticles, Dt](int32 Index)
 		{
-			for(int32 CollisionIndex = 0; CollisionIndex < CollisionParticlesRange.GetRangeSize(); ++CollisionIndex)
+			if (PAndInvM[Index].InvM == (FSolverReal)0.)
 			{
-				const Chaos::EImplicitObjectType CollisionType = CollisionParticlesRange.GetGeometry(CollisionIndex)->GetType();
-				if ( CollisionType == Chaos::ImplicitObjectType::WeightedLatticeBone || CollisionType == TWeightedLatticeImplicitObject<FLevelSet>::StaticType())
-				{
-					// Skip skinned levelsets here. We will do a batch query for those in a second pass to match ISPC.
-					continue;
-				}
-
-				if (CollisionType == Chaos::ImplicitObjectType::MLLevelSetBone || CollisionType == Chaos::ImplicitObjectType::MLLevelSet)
-				{
-					// Skip  mllevelsets here. 
-					continue;
-				}
-
-				const FSolverRigidTransform3 Frame(CollisionParticlesRange.GetX(CollisionIndex), CollisionParticlesRange.R(CollisionIndex));
-				const FVec3 RigidSpacePosition(Frame.InverseTransformPosition(PAndInvM[Index].P));  // PhiWithNormal requires FReal based arguments
-				FVec3 ImplicitNormal;                                                                // since implicits don't use FSolverReal
-				const FSolverReal Phi = (FSolverReal)CollisionParticlesRange.GetGeometry(CollisionIndex)->PhiWithNormal(RigidSpacePosition, ImplicitNormal);
-				const int32 VelocityBone = CollisionIndex;
-				const FSolverReal Penetration = CollisionThickness - Phi; // This is related to the Normal impulse
-				const FSolverVec3 Normal(ImplicitNormal);
+				return;
+			}
 
-				if (Penetration > (FSolverReal)0.)
+			for (const FSolverCollisionParticlesRange& CollisionParticlesRange : CollisionParticles)
+			{
+				for (int32 CollisionIndex = 0; CollisionIndex < CollisionParticlesRange.GetRangeSize(); ++CollisionIndex)
 				{
-					const FSolverVec3 NormalWorld = Frame.TransformVector(Normal);
-
-					if constexpr (bLockAndWriteContacts)
+					const FImplicitObjectPtr& Geometry = CollisionParticlesRange.GetGeometry(CollisionIndex);
+					const Chaos::EImplicitObjectType CollisionType = Geometry->GetType();
+					if (Private::IsComplexBatchCollider(CollisionType))
 					{
-						checkSlow(Contacts);
-						checkSlow(Normals);
-						checkSlow(Phis);
-						FScopeLock Lock(&DebugMutex);
-						CollisionParticlesRange.GetArrayView(*CollisionParticleCollided)[CollisionIndex] = true;
-						Contacts->Emplace(PAndInvM[Index].P);
-						Normals->Emplace(NormalWorld);
-						Phis->Emplace(Phi);
+						continue;
 					}
 
-					PAndInvM[Index].P += Penetration * NormalWorld;
-
-					if constexpr (bWithFriction)
+					const FSolverRigidTransform3 Frame(CollisionParticlesRange.GetX(CollisionIndex), CollisionParticlesRange.R(CollisionIndex));
+					const FVec3 RigidSpacePosition(Frame.InverseTransformPosition(X[Index]));  // PhiWithNormal requires FReal based arguments
+					FVec3 ImplicitNormal;
+					const FSolverReal Phi = (FSolverReal)Geometry->PhiWithNormal(RigidSpacePosition, ImplicitNormal);
+					const FSolverReal Penetration = SoftBodyCollisionThickness + CollisionThickness - Phi;
+					if (Penetration > (FSolverReal)0.)
 					{
-						const FSolverVec3 VectorToPoint = PAndInvM[Index].P - CollisionParticlesRange.GetX(VelocityBone);
-						const FSolverVec3 RelativeDisplacement = (PAndInvM[Index].P - X[Index]) - (CollisionParticlesRange.V(VelocityBone) + FSolverVec3::CrossProduct(CollisionParticlesRange.W(VelocityBone), VectorToPoint)) * Dt; // This corresponds to the tangential velocity multiplied by dt (friction will drive this to zero if it is high enough)
-						const FSolverVec3 RelativeDisplacementTangent = RelativeDisplacement - FSolverVec3::DotProduct(RelativeDisplacement, NormalWorld) * NormalWorld; // Project displacement into the tangential plane
-						const FSolverReal RelativeDisplacementTangentLength = RelativeDisplacementTangent.Size();
-						if (RelativeDisplacementTangentLength >= UE_SMALL_NUMBER)
+						const FSolverVec3 NormalWorld = Frame.TransformVector(FSolverVec3(ImplicitNormal));
+						if constexpr (bLockAndWriteContacts)
 						{
-							const FSolverReal PositionCorrection = FMath::Min<FSolverReal>(Penetration * FrictionCoefficient, RelativeDisplacementTangentLength);
-							const FSolverReal CorrectionRatio = PositionCorrection / RelativeDisplacementTangentLength;
-							PAndInvM[Index].P -= CorrectionRatio * RelativeDisplacementTangent;
+							FScopeLock Lock(&DebugMutex);
+							CollisionParticlesRange.GetArrayView(*CollisionParticleCollided)[CollisionIndex] = true;
+							Contacts->Emplace(PAndInvM[Index].P);
+							Normals->Emplace(NormalWorld);
+							Phis->Emplace(Phi);
+						}
+
+						PAndInvM[Index].P += Penetration * NormalWorld;
+
+						if constexpr (bGeneratePlanarConstraints)
+						{
+							// Last collider per point wins.
+							HasPlanarData[Index] = true;
+							PlanarDataPositions[Index] = PAndInvM[Index].P;
+							PlanarDataNormals[Index] = NormalWorld;
+						}
+
+						if constexpr (bWithFriction)
+						{
+							FSolverVec3 ColliderVelocityAtPoint;
+							Private::ApplyFriction(PAndInvM[Index].P, X[Index], NormalWorld, Penetration * FrictionCoefficient, Dt, CollisionParticlesRange.X(CollisionIndex), CollisionParticlesRange.W(CollisionIndex), ColliderVelocityAtPoint);
+
+							if constexpr (bGeneratePlanarConstraints)
+							{
+								PlanarDataVelocities[Index] = ColliderVelocityAtPoint;
+							}
 						}
 					}
 				}
 			}
 		}
-	});
+	);
+}
+
+template<bool bLockAndWriteContacts, bool bWithFriction, bool bGeneratePlanarConstraints>
+void FPBDSoftBodyCollisionConstraintBase::ApplyComplexInternal(FSolverParticlesRange& Particles, const FSolverReal Dt, const TArray<FSolverCollisionParticlesRange>& CollisionParticles)
+{
+	TRACE_CPUPROFILER_EVENT_SCOPE(FPBDSoftBodyCollisionConstraint_ApplyComplexInternal);
+	if (!bEnableComplexColliders)
+	{
+		return;
+	}
 
-	// Skinned levelset passes
+	FPAndInvM* const PAndInvM = Particles.GetPAndInvM().GetData();
+	const FSolverVec3* const X = Particles.XArray().GetData();
 	for (const FSolverCollisionParticlesRange& CollisionParticlesRange : CollisionParticles)
 	{
 		for (int32 CollisionIndex = 0; CollisionIndex < CollisionParticlesRange.GetRangeSize(); ++CollisionIndex)
 		{
-			if (const TWeightedLatticeImplicitObject<FLevelSet>* const LevelSet = CollisionParticlesRange.GetGeometry(CollisionIndex)->GetObject< TWeightedLatticeImplicitObject<FLevelSet> >())
+			const FImplicitObjectPtr& Geometry = CollisionParticlesRange.GetGeometry(CollisionIndex);
+			const Chaos::EImplicitObjectType CollisionType = Geometry->GetType();
+			if (!Private::IsComplexBatchCollider(CollisionType))
 			{
-				const FSolverRigidTransform3 Frame(CollisionParticlesRange.X(CollisionIndex), CollisionParticlesRange.R(CollisionIndex));
-				if constexpr (bWithFriction)
+				continue;
+			}
+
+			const FSolverRotation3& CollisionR = CollisionParticlesRange.R(CollisionIndex);
+			const FSolverVec3& CollisionV = CollisionParticlesRange.V(CollisionIndex);
+			const FSolverVec3& CollisionX = CollisionParticlesRange.X(CollisionIndex);
+			const FSolverVec3& CollisionW = CollisionParticlesRange.W(CollisionIndex);
+
+			bool* const bCollided = bLockAndWriteContacts ? &CollisionParticlesRange.GetArrayView(*CollisionParticleCollided)[CollisionIndex] : nullptr;
+			const FSolverRigidTransform3 Frame(CollisionX, CollisionR);
+			if constexpr (bWithFriction)
+			{
+				if (const FPBDComplexColliderBoneData* const ColliderBoneData = ComplexBoneData.Find(FParticleRangeIndex(CollisionParticlesRange.GetRangeId(), CollisionIndex)))
 				{
-					LevelSet->BatchPhiWithNormalAndGreatestInfluenceBone(Particles.GetPAndInvM(), Frame.Inverse(), CollisionThickness, BatchPhis, BatchNormals, BatchVelocityBones);
-					check(BatchPhis.Num() == Particles.GetRangeSize());
-					check(BatchNormals.Num() == Particles.GetRangeSize());
-					check(BatchVelocityBones.Num() == Particles.GetRangeSize());
+					if (const TWeightedLatticeImplicitObject<FLevelSet>* const LevelSet = Geometry->GetObject<TWeightedLatticeImplicitObject<FLevelSet>>())
+					{
+						LevelSet->BatchPhiWithNormalAndGreatestInfluenceBone(Particles.GetPAndInvM(), Frame, SoftBodyCollisionThickness + CollisionThickness, BatchPhis, BatchNormals, BatchVelocityBones);
+					}
+					else
+					{
+						checkNoEntry();
+					}
+
+					// Apply
+					PhysicsParallelFor(Particles.GetRangeSize(), [this, &X, &PAndInvM, &Frame, &CollisionV, &CollisionX, &CollisionW, ColliderBoneData, Dt, &bCollided](int32 Index)
+						{
+							if (PAndInvM[Index].InvM == (FSolverReal)0.)
+							{
+								return;
+							}
+
+							const FSolverReal Penetration = SoftBodyCollisionThickness + CollisionThickness - BatchPhis[Index];
+							if (Penetration > (FSolverReal)0.)
+							{
+								const FSolverVec3 NormalWorld = Frame.TransformVector(BatchNormals[Index]);
+								if constexpr (bLockAndWriteContacts)
+								{
+									FScopeLock Lock(&DebugMutex);
+									*bCollided = true;
+									Contacts->Emplace(PAndInvM[Index].P);
+									Normals->Emplace(NormalWorld);
+									Phis->Emplace(BatchPhis[Index]);
+								}
+
+								PAndInvM[Index].P += Penetration * NormalWorld;
+
+								if constexpr (bGeneratePlanarConstraints)
+								{
+									// Last collider per point wins.
+									HasPlanarData[Index] = true;
+									PlanarDataPositions[Index] = PAndInvM[Index].P;
+									PlanarDataNormals[Index] = NormalWorld;
+								}
+
+								FSolverVec3 ColliderVelocityAtPoint;
+								const int32 StrongestBone = BatchVelocityBones[Index];
+								if (ColliderBoneData->MappedBoneIndices.IsValidIndex(StrongestBone))
+								{
+									const int32 MappedIndex = ColliderBoneData->MappedBoneIndices[StrongestBone];
+									Private::ApplyFriction(PAndInvM[Index].P, X[Index], NormalWorld, Penetration * FrictionCoefficient, Dt, ColliderBoneData->X[MappedIndex], ColliderBoneData->V[MappedIndex], ColliderBoneData->W[MappedIndex], ColliderVelocityAtPoint);
+								}
+								else
+								{
+									Private::ApplyFriction(PAndInvM[Index].P, X[Index], NormalWorld, Penetration * FrictionCoefficient, Dt, CollisionX, CollisionV, CollisionW, ColliderVelocityAtPoint);
+								}
+
+								if constexpr (bGeneratePlanarConstraints)
+								{
+									PlanarDataVelocities[Index] = ColliderVelocityAtPoint;
+								}
+							}
+						}
+					);
 				}
 				else
 				{
-					LevelSet->BatchPhiWithNormal(Particles.GetPAndInvM(), Frame.Inverse(), BatchPhis, BatchNormals);
-					check(BatchPhis.Num() == Particles.GetRangeSize());
-					check(BatchNormals.Num() == Particles.GetRangeSize());
-				}
-				const FSolverReal* const CollisionPhis = BatchPhis.GetData();
-				const FSolverVec3* const CollisionNormals = BatchNormals.GetData();
-				const int32* const VelocityBones = bWithFriction ? BatchVelocityBones.GetData() : nullptr;
-				PhysicsParallelFor(Particles.GetRangeSize(), [this, Dt, &CollisionParticlesRange, CollisionIndex, &PAndInvM, &X, &Frame, &CollisionPhis, &CollisionNormals, &VelocityBones](int32 Index)
-				{
-					if (PAndInvM[Index].InvM == (FSolverReal)0.)
+					if (const FMLLevelSet* const MLLevelSet = Geometry->GetObject<FMLLevelSet>())
+					{
+						BatchPhis.SetNumUninitialized(Particles.GetRangeSize());
+						BatchNormals.SetNumUninitialized(Particles.GetRangeSize());
+
+						// Batch Query
+						constexpr int32 MLLevelSetThreadNum = 0;
+						MLLevelSet->BatchPhiWithNormal(Particles.GetPAndInvM(), Frame, BatchPhis, BatchNormals, SoftBodyCollisionThickness + CollisionThickness, MLLevelSetThreadNum, 0, Particles.GetRangeSize());
+					}
+					else if (const TWeightedLatticeImplicitObject<FLevelSet>* const LevelSet = Geometry->GetObject<TWeightedLatticeImplicitObject<FLevelSet>>())
+					{
+						LevelSet->BatchPhiWithNormal(Particles.GetPAndInvM(), Frame, BatchPhis, BatchNormals);
+					}
+					else
 					{
-						return;
+						checkNoEntry();
 					}
 
-					const FSolverReal Penetration = CollisionThickness - CollisionPhis[Index];
-					if (Penetration > (FSolverReal)0.)
+					// Apply
+					PhysicsParallelFor(Particles.GetRangeSize(), [this, &X, &PAndInvM, &CollisionV, &CollisionX, &CollisionW, &Frame, Dt, &bCollided](int32 Index)
+						{
+							if (PAndInvM[Index].InvM == (FSolverReal)0.)
+							{
+								return;
+							}
+
+							const FSolverReal Penetration = SoftBodyCollisionThickness + CollisionThickness - BatchPhis[Index];
+							if (Penetration > (FSolverReal)0.)
+							{
+								const FSolverVec3 NormalWorld = Frame.TransformVector(BatchNormals[Index]);
+								if constexpr (bLockAndWriteContacts)
+								{
+									FScopeLock Lock(&DebugMutex);
+									*bCollided = true;
+									Contacts->Emplace(PAndInvM[Index].P);
+									Normals->Emplace(NormalWorld);
+									Phis->Emplace(BatchPhis[Index]);
+								}
+
+								PAndInvM[Index].P += Penetration * NormalWorld;
+
+								if constexpr (bGeneratePlanarConstraints)
+								{
+									// Last collider per point wins.
+									HasPlanarData[Index] = true;
+									PlanarDataPositions[Index] = PAndInvM[Index].P;
+									PlanarDataNormals[Index] = NormalWorld;
+								}
+
+								FSolverVec3 ColliderVelocityAtPoint;
+								Private::ApplyFriction(PAndInvM[Index].P, X[Index], NormalWorld, Penetration * FrictionCoefficient, Dt, CollisionX, CollisionV, CollisionW, ColliderVelocityAtPoint);
+
+								if (bGeneratePlanarConstraints)
+								{
+									PlanarDataVelocities[Index] = ColliderVelocityAtPoint;
+								}
+							}
+						}
+					);
+				}
+			}
+			else
+			{
+				if (const FMLLevelSet* const MLLevelSet = Geometry->GetObject<FMLLevelSet>())
+				{
+					BatchPhis.SetNumUninitialized(Particles.GetRangeSize());
+					BatchNormals.SetNumUninitialized(Particles.GetRangeSize());
+
+					// Batch Query
+					constexpr int32 MLLevelSetThreadNum = 0;
+					MLLevelSet->BatchPhiWithNormal(Particles.GetPAndInvM(), Frame, BatchPhis, BatchNormals, SoftBodyCollisionThickness + CollisionThickness, MLLevelSetThreadNum, 0, Particles.GetRangeSize());
+				}
+				else if (const TWeightedLatticeImplicitObject<FLevelSet>* const LevelSet = Geometry->GetObject<TWeightedLatticeImplicitObject<FLevelSet>>())
+				{
+					LevelSet->BatchPhiWithNormal(Particles.GetPAndInvM(), Frame, BatchPhis, BatchNormals);
+				}
+				else
+				{
+					checkNoEntry();
+				}
+
+				// Apply
+				PhysicsParallelFor(Particles.GetRangeSize(), [this, &X, &PAndInvM, &Frame, &bCollided](int32 Index)
 					{
-						const FSolverVec3 NormalWorld = Frame.TransformVector(CollisionNormals[Index]);
-						if constexpr (bLockAndWriteContacts)
+						if (PAndInvM[Index].InvM == (FSolverReal)0.)
 						{
-							checkSlow(Contacts);
-							checkSlow(Normals);
-							checkSlow(Phis);
-							FScopeLock Lock(&DebugMutex);
-							CollisionParticlesRange.GetArrayView(*CollisionParticleCollided)[CollisionIndex] = true;
-							Contacts->Emplace(PAndInvM[Index].P);
-							Normals->Emplace(NormalWorld);
-							Phis->Emplace(CollisionPhis[Index]);
+							return;
 						}
 
-						PAndInvM[Index].P = Penetration * NormalWorld;
-						if constexpr (bWithFriction)
+						const FSolverReal Penetration = SoftBodyCollisionThickness + CollisionThickness - BatchPhis[Index];
+						if (Penetration > (FSolverReal)0.)
 						{
-							const int32 VelocityBone = CollisionParticlesRange.IsValidIndex(VelocityBones[Index]) ? VelocityBones[Index] : CollisionIndex;
-							const FSolverVec3 VectorToPoint = PAndInvM[Index].P - CollisionParticlesRange.X(VelocityBone);
-							const FSolverVec3 RelativeDisplacement = (PAndInvM[Index].P - X[Index]) - (CollisionParticlesRange.V(VelocityBone) + FSolverVec3::CrossProduct(CollisionParticlesRange.W(VelocityBone), VectorToPoint)) * Dt; // This corresponds to the tangential velocity multiplied by dt (friction will drive this to zero if it is high enough)
-							const FSolverVec3 RelativeDisplacementTangent = RelativeDisplacement - FSolverVec3::DotProduct(RelativeDisplacement, NormalWorld) * NormalWorld; // Project displacement into the tangential plane
-							const FSolverReal RelativeDisplacementTangentLength = RelativeDisplacementTangent.Size();
-							if (RelativeDisplacementTangentLength >= UE_SMALL_NUMBER)
+							const FSolverVec3 NormalWorld = Frame.TransformVector(BatchNormals[Index]);
+							if constexpr (bLockAndWriteContacts)
 							{
-								const FSolverReal PositionCorrection = FMath::Min<FSolverReal>(Penetration * FrictionCoefficient, RelativeDisplacementTangentLength);
-								const FSolverReal CorrectionRatio = PositionCorrection / RelativeDisplacementTangentLength;
-								PAndInvM[Index].P -= CorrectionRatio * RelativeDisplacementTangent;
+								FScopeLock Lock(&DebugMutex);
+								*bCollided = true;
+								Contacts->Emplace(PAndInvM[Index].P);
+								Normals->Emplace(NormalWorld);
+								Phis->Emplace(BatchPhis[Index]);
+							}
+
+							PAndInvM[Index].P += Penetration * NormalWorld;
+
+							if constexpr (bGeneratePlanarConstraints)
+							{
+								// Last collider per point wins.
+								HasPlanarData[Index] = true;
+								PlanarDataPositions[Index] = PAndInvM[Index].P;
+								PlanarDataNormals[Index] = NormalWorld;
 							}
 						}
 					}
-				});
+				);
+			}
+		}
+	}
+}
+
+void FPBDSoftBodyCollisionConstraintBase::FinalizePlanarConstraints(const FSolverParticlesRange& Particles)
+{
+	TRACE_CPUPROFILER_EVENT_SCOPE(FPBDSoftBodyCollisionConstraint_FinalizePlanarConstraints);
+	// Count how many planar constraints we have.
+	int32 NumPlanarConstraints = 0;
+	for (const bool bHasPlanarData : HasPlanarData)
+	{
+		if (bHasPlanarData)
+		{
+			++NumPlanarConstraints;
+		}
+	}
+
+	PlanarConstraint.GetUniqueConstraintIndices().SetNumUninitialized(NumPlanarConstraints);
+	PlanarConstraint.GetTargetPositions().SetNumUninitialized(NumPlanarConstraints);
+	PlanarConstraint.GetTargetNormals().SetNumUninitialized(NumPlanarConstraints);
+	PlanarConstraint.GetTargetVelocities().SetNumUninitialized(NumPlanarConstraints);
+
+	if (NumPlanarConstraints > 0)
+	{
+		int32 ConstraintIndex = 0;
+		for (int32 VertexIndex = 0; VertexIndex < Particles.GetRangeSize(); ++VertexIndex)
+		{
+			if (HasPlanarData[VertexIndex])
+			{
+				PlanarConstraint.GetUniqueConstraintIndices()[ConstraintIndex] = VertexIndex;
+				PlanarConstraint.GetTargetPositions()[ConstraintIndex] = PlanarDataPositions[VertexIndex];
+				PlanarConstraint.GetTargetNormals()[ConstraintIndex] = PlanarDataNormals[VertexIndex];
+				PlanarConstraint.GetTargetVelocities()[ConstraintIndex] = PlanarDataVelocities.IsValidIndex(VertexIndex) ? PlanarDataVelocities[VertexIndex] : FSolverVec3(0.f);
+				++ConstraintIndex;
 			}
 		}
+		check(ConstraintIndex == NumPlanarConstraints);
 	}
 }
 
+void FPBDSoftBodyCollisionConstraintBase::ApplyPlanarConstraints(FSolverParticlesRange& Particles, const FSolverReal Dt)
+{
+	PlanarConstraint.Apply(Particles, Dt);
+}
+
 template<bool bLockAndWriteContacts, bool bWithFriction>
 void FPBDSoftBodyCollisionConstraintBase::ApplyInternalCCD(FSolverParticlesRange& Particles, const FSolverReal Dt, const TArray<FSolverCollisionParticlesRange>& CollisionParticles) const
 {
+	if (!bEnableSimpleColliders && !bEnableComplexColliders)
+	{
+		return;
+	}
+
 	FPAndInvM* const PAndInvM = Particles.GetPAndInvM().GetData();
 	const FSolverVec3* const X = Particles.XArray().GetData();
 	PhysicsParallelFor(Particles.GetRangeSize(), [this, Dt, &PAndInvM, &X, &CollisionParticles](int32 Index)
@@ -443,21 +850,31 @@ void FPBDSoftBodyCollisionConstraintBase::ApplyInternalCCD(FSolverParticlesRange
 
 			for (int32 CollisionIndex = 0; CollisionIndex < CollisionParticlesRange.GetRangeSize(); ++CollisionIndex)
 			{
-				if (CollisionParticlesRange.GetGeometry(CollisionIndex)->GetType() == Chaos::ImplicitObjectType::WeightedLatticeBone)
+				const FImplicitObjectPtr& Geometry = CollisionParticlesRange.GetGeometry(CollisionIndex);
+				const Chaos::EImplicitObjectType CollisionType = Geometry->GetType();
+				if (Private::IsComplexBatchCollider(CollisionType))
 				{
-					continue;
+					if (!bEnableComplexColliders)
+					{
+						continue;
+					}
 				}
-
-				if (CollisionParticlesRange.GetGeometry(CollisionIndex)->GetType() == Chaos::ImplicitObjectType::MLLevelSetBone)
+				else if (Private::IsSimpleCollider(CollisionType))
+				{
+					if (!bEnableSimpleColliders)
+					{
+						continue;
+					}
+				}
+				else
 				{
 					continue;
 				}
-
 				const FSolverRigidTransform3 Frame(CollisionParticlesRange.GetX(CollisionIndex), CollisionParticlesRange.R(CollisionIndex));
 
 				const Pair<FVec3, bool> PointPair = CollisionParticlesRange.GetGeometry(CollisionIndex)->FindClosestIntersection(  // Geometry operates in FReal
 					FVec3(CollisionTransforms[CollisionIndex].InverseTransformPositionNoScale(X[Index])),        // hence the back and forth
-					FVec3(Frame.InverseTransformPositionNoScale(PAndInvM[Index].P)), (FReal)CollisionThickness);                   // FVec3/FReal conversions
+					FVec3(Frame.InverseTransformPositionNoScale(PAndInvM[Index].P)), (FReal)(SoftBodyCollisionThickness + CollisionThickness));                   // FVec3/FReal conversions
 
 				if (PointPair.Second)
 				{
@@ -482,7 +899,9 @@ void FPBDSoftBodyCollisionConstraintBase::ApplyInternalCCD(FSolverParticlesRange
 					if constexpr (bWithFriction)
 					{
 						// Friction
-						int32 VelocityBone = CollisionIndex;
+						FSolverVec3 CollisionV = CollisionParticlesRange.V(CollisionIndex);
+						FSolverVec3 CollisionX = CollisionParticlesRange.X(CollisionIndex);
+						FSolverVec3 CollisionW = CollisionParticlesRange.W(CollisionIndex);
 						if (const TWeightedLatticeImplicitObject<FLevelSet>* LevelSet = CollisionParticlesRange.GetGeometry(CollisionIndex)->GetObject< TWeightedLatticeImplicitObject<FLevelSet> >())
 						{
 							TArray<FWeightedLatticeImplicitObject::FEmbeddingCoordinate> Coordinates;
@@ -501,24 +920,21 @@ void FPBDSoftBodyCollisionConstraintBase::ApplyInternalCCD(FSolverParticlesRange
 							}
 							if (ClosestCoordIndex != INDEX_NONE)
 							{
-								const int32 StrongestBone = Coordinates[ClosestCoordIndex].GreatestInfluenceBone(LevelSet->GetBoneData());
-								if (StrongestBone != INDEX_NONE)
+								if (const FPBDComplexColliderBoneData* const ColliderBoneData = ComplexBoneData.Find(FParticleRangeIndex(CollisionParticlesRange.GetRangeId(), CollisionIndex)))
 								{
-									VelocityBone = LevelSet->GetSolverBoneIndices()[StrongestBone];
+									const int32 StrongestBone = Coordinates[ClosestCoordIndex].GreatestInfluenceBone(LevelSet->GetBoneData());
+									if (ColliderBoneData->MappedBoneIndices.IsValidIndex(StrongestBone))
+									{
+										const int32 MappedIndex = ColliderBoneData->MappedBoneIndices[StrongestBone];
+										CollisionV = ColliderBoneData->V[MappedIndex];
+										CollisionX = ColliderBoneData->X[MappedIndex];
+										CollisionW = ColliderBoneData->W[MappedIndex];
+									}
 								}
 							}
 						}
 
-						const FSolverVec3 VectorToPoint = PAndInvM[Index].P - CollisionParticlesRange.GetX(VelocityBone);
-						const FSolverVec3 RelativeDisplacement = (PAndInvM[Index].P - X[Index]) - (CollisionParticlesRange.V(VelocityBone) + FSolverVec3::CrossProduct(CollisionParticlesRange.W(VelocityBone), VectorToPoint)) * Dt;  // This corresponds to the tangential velocity multiplied by dt (friction will drive this to zero if it is high enough)
-						const FSolverVec3 RelativeDisplacementTangent = RelativeDisplacement - FSolverVec3::DotProduct(RelativeDisplacement, NormalWorld) * NormalWorld;  // Project displacement into the tangential plane
-						const FSolverReal RelativeDisplacementTangentLength = RelativeDisplacementTangent.Size();
-						if (RelativeDisplacementTangentLength >= UE_SMALL_NUMBER)
-						{
-							const FSolverReal PositionCorrection = FMath::Min<FSolverReal>(Penetration * FrictionCoefficient, RelativeDisplacementTangentLength);
-							const FSolverReal CorrectionRatio = PositionCorrection / RelativeDisplacementTangentLength;
-							PAndInvM[Index].P -= CorrectionRatio * RelativeDisplacementTangent;
-						}
+						Private::ApplyFriction(PAndInvM[Index].P, X[Index], NormalWorld, Penetration * FrictionCoefficient, Dt, CollisionX, CollisionV, CollisionW);
 					}
 				}
 			}
@@ -526,171 +942,96 @@ void FPBDSoftBodyCollisionConstraintBase::ApplyInternalCCD(FSolverParticlesRange
 	});
 }
 
-template<bool bWithFriction>
-void FPBDSoftBodyCollisionConstraintBase::ApplyInternalISPC(FSolverParticlesRange& Particles, const FSolverReal Dt, const TArray<FSolverCollisionParticlesRange>& CollisionParticles) const
-{
 #if INTEL_ISPC
-	check(bRealTypeCompatibleWithISPC);
-	const int32 NumBatches = FMath::CeilToInt((FSolverReal)(Particles.GetRangeSize()) / (FSolverReal)Chaos_SoftBodyCollision_ISPC_ParallelBatchSize);
-	
-	PhysicsParallelFor(NumBatches, [this, &Particles, Dt, &CollisionParticles](int32 BatchNumber)
+void FPBDSoftBodyCollisionConstraintBase::ApplySimpleInternalISPC(FSolverParticlesRange& Particles, const FSolverReal Dt, const TArray<FSolverCollisionParticlesRange>& CollisionParticles, bool bGeneratePlanarConstraints)
+{
+	if (!bEnableSimpleColliders)
 	{
-		const int32 BatchBegin = (Chaos_SoftBodyCollision_ISPC_ParallelBatchSize * BatchNumber);
-		const int32 BatchEnd = FMath::Min(Particles.GetRangeSize(), BatchBegin + Chaos_SoftBodyCollision_ISPC_ParallelBatchSize);
-
-		for (const FSolverCollisionParticlesRange& CollisionParticlesRange : CollisionParticles)
-		{
-			if constexpr (bWithFriction)
-			{
-				//This function skips Skinned LevelSets, MLLevelSets, and their proxy bones.
-				ispc::ApplyPerParticleCollisionFastFrictionNoGroupCheck(
-					(ispc::FVector4f*)Particles.GetPAndInvM().GetData(),
-					(const ispc::FVector3f*)Particles.XArray().GetData(),
-					(const ispc::FVector3f*)CollisionParticlesRange.GetV().GetData(),
-					(const ispc::FVector3f*)CollisionParticlesRange.XArray().GetData(),
-					(const ispc::FVector3f*)CollisionParticlesRange.GetW().GetData(),
-					(const ispc::FVector4f*)CollisionParticlesRange.GetR().GetData(),
-					FrictionCoefficient,
-					CollisionThickness,
-					(const uint8*)&CollisionParticlesRange,
-					(const uint8*)CollisionParticlesRange.GetAllGeometry().GetData(),
-					sizeof(FImplicitObject),
-					FImplicitObject::GetOffsetOfType(),
-					FImplicitObject::GetOffsetOfMargin(),
-					Dt,
-					CollisionParticlesRange.GetRangeSize(),
-					BatchBegin,
-					BatchEnd);
-			}
-			else
-			{
-				//This function skips Skinned LevelSets, MLLevelSets, and their proxy bones.
-				ispc::ApplyPerParticleCollisionNoFrictionNoGroupCheck(
-					(ispc::FVector4f*)Particles.GetPAndInvM().GetData(),
-					(const ispc::FVector3f*)Particles.XArray().GetData(),
-					(const ispc::FVector3f*)CollisionParticlesRange.GetV().GetData(),
-					(const ispc::FVector3f*)CollisionParticlesRange.XArray().GetData(),
-					(const ispc::FVector3f*)CollisionParticlesRange.GetW().GetData(),
-					(const ispc::FVector4f*)CollisionParticlesRange.GetR().GetData(),
-					CollisionThickness,
-					(const uint8*)&CollisionParticlesRange,
-					(const uint8*)CollisionParticlesRange.GetAllGeometry().GetData(),
-					sizeof(FImplicitObject),
-					FImplicitObject::GetOffsetOfType(),
-					FImplicitObject::GetOffsetOfMargin(),
-					Dt,
-					CollisionParticlesRange.GetRangeSize(),
-					BatchBegin,
-					BatchEnd);
-			}
-		}
-	});
-	
+		return;
+	}
 
-	// Skinned levelset passes
-	for (const FSolverCollisionParticlesRange& CollisionParticlesRange : CollisionParticles)
-	{
-		for (int32 CollisionIndex = 0; CollisionIndex < CollisionParticlesRange.GetRangeSize(); ++CollisionIndex)
-		{
-			if (const TWeightedLatticeImplicitObject<FLevelSet>* const LevelSet = CollisionParticlesRange.GetGeometry(CollisionIndex)->GetObject< TWeightedLatticeImplicitObject<FLevelSet> >())
-			{
-				TRACE_CPUPROFILER_EVENT_SCOPE(FPBDSoftBodyCollisionConstraint_ApplyInternalISPC_LevelSet);
-				const FSolverRigidTransform3 Frame = FSolverRigidTransform3(CollisionParticlesRange.X(CollisionIndex), CollisionParticlesRange.R(CollisionIndex)).Inverse();
-				if constexpr (bWithFriction)
-				{
-					LevelSet->BatchPhiWithNormalAndGreatestInfluenceBone(Particles.GetPAndInvM(), Frame, CollisionThickness, BatchPhis, BatchNormals, BatchVelocityBones);
-					check(BatchPhis.Num() == Particles.GetRangeSize());
-					check(BatchNormals.Num() == Particles.GetRangeSize());
-					check(BatchVelocityBones.Num() == Particles.GetRangeSize());
+	TRACE_CPUPROFILER_EVENT_SCOPE(FPBDSoftBodyCollisionConstraint_ApplySimpleInternalISPC);
 
+	check(bRealTypeCompatibleWithISPC);
 
-					const FSolverRotation3 CPR_R = CollisionParticlesRange.R(CollisionIndex);
-					ispc::ApplyPerParticleCollisionFastFrictionBatchedPhiAndNormal(
-						(ispc::FVector4f*)Particles.GetPAndInvM().GetData(),
-						(const ispc::FVector3f*)Particles.XArray().GetData(),
-						CollisionIndex,
-						BatchPhis.GetData(),
-						(const ispc::FVector3f*)BatchNormals.GetData(),
-						BatchVelocityBones.GetData(),
-						(const ispc::FVector3f*)CollisionParticlesRange.GetV().GetData(),
-						(const ispc::FVector3f*)CollisionParticlesRange.XArray().GetData(),
-						(const ispc::FVector3f*)CollisionParticlesRange.GetW().GetData(),
-						reinterpret_cast<const ispc::FVector4f&>(CPR_R),
-						//(const ispc::FVector4f&)CollisionParticlesRange.R(CollisionIndex),
-						FrictionCoefficient,
-						CollisionThickness,
-						Dt,
-						CollisionParticlesRange.GetRangeSize(),
-						0,
-						Particles.GetRangeSize());
-				}
-				else
-				{
-					LevelSet->BatchPhiWithNormal(Particles.GetPAndInvM(), Frame, BatchPhis, BatchNormals);
-					check(BatchPhis.Num() == Particles.GetRangeSize());
-					check(BatchNormals.Num() == Particles.GetRangeSize());
-					const FSolverRotation3 CPR_R = CollisionParticlesRange.R(CollisionIndex);
-					ispc::ApplyPerParticleCollisionNoFrictionBatchedPhiAndNormal(
-						(ispc::FVector4f*)Particles.GetPAndInvM().GetData(),
-						BatchPhis.GetData(),
-						(const ispc::FVector3f*)BatchNormals.GetData(),
-						reinterpret_cast<const ispc::FVector4f&>(CPR_R),
-						//(const ispc::FVector4f&)CollisionParticlesRange.R(CollisionIndex),
-						CollisionThickness,
-						Dt,
-						0,
-						Particles.GetRangeSize());
-				}
-			}
-		}
-	}
+	const bool bWithFriction = FrictionCoefficient > (FSolverReal)UE_KINDA_SMALL_NUMBER;
 
-	// MLLevelset passes ApplyInternalISPC
-	for (const FSolverCollisionParticlesRange& CollisionParticlesRange : CollisionParticles)
-	{
-		const int32 NumberOfParticles = Particles.GetRangeSize();
-		FPAndInvM* const PAndInvM = Particles.GetPAndInvM().GetData();
+	const int32 NumBatches = FMath::CeilToInt((FSolverReal)(Particles.GetRangeSize()) / (FSolverReal)Chaos_SoftBodyCollision_ISPC_ParallelBatchSize);
 
-		if (MLLevelSetParallelize == 0)
+	// Simple colliders
+	PhysicsParallelFor(NumBatches, [this, &Particles, Dt, &CollisionParticles, bWithFriction, bGeneratePlanarConstraints](int32 BatchNumber)
 		{
-			for (int32 CollisionIndex = 0; CollisionIndex < CollisionParticlesRange.GetRangeSize(); ++CollisionIndex)
+			const int32 BatchBegin = (Chaos_SoftBodyCollision_ISPC_ParallelBatchSize * BatchNumber);
+			const int32 BatchEnd = FMath::Min(Particles.GetRangeSize(), BatchBegin + Chaos_SoftBodyCollision_ISPC_ParallelBatchSize);
+
+			for (const FSolverCollisionParticlesRange& CollisionParticlesRange : CollisionParticles)
 			{
-				const FSolverRigidTransform3 Frame = FSolverRigidTransform3(CollisionParticlesRange.X(CollisionIndex), CollisionParticlesRange.R(CollisionIndex)).Inverse();
-				if (FMLLevelSet* MLLevelSet = CollisionParticlesRange.GetGeometry(CollisionIndex)->GetObject<FMLLevelSet>())
+				if (bGeneratePlanarConstraints)
+				{
+					if (bWithFriction)
+					{
+						ispc::ApplyPerParticleSimpleCollisionFastFrictionAndGeneratePlanarConstraints(
+							(ispc::FVector4f*)Particles.GetPAndInvM().GetData(),
+							HasPlanarData.GetData(),
+							(ispc::FVector3f*)PlanarDataPositions.GetData(),
+							(ispc::FVector3f*)PlanarDataNormals.GetData(),
+							(ispc::FVector3f*)PlanarDataVelocities.GetData(),
+							(const ispc::FVector3f*)Particles.XArray().GetData(),
+							(const ispc::FVector3f*)CollisionParticlesRange.GetV().GetData(),
+							(const ispc::FVector3f*)CollisionParticlesRange.XArray().GetData(),
+							(const ispc::FVector3f*)CollisionParticlesRange.GetW().GetData(),
+							(const ispc::FVector4f*)CollisionParticlesRange.GetR().GetData(),
+							FrictionCoefficient,
+							SoftBodyCollisionThickness + CollisionThickness,
+							(const uint8*)&CollisionParticlesRange,
+							(const uint8*)CollisionParticlesRange.GetAllGeometry().GetData(),
+							sizeof(FImplicitObject),
+							FImplicitObject::GetOffsetOfType(),
+							FImplicitObject::GetOffsetOfMargin(),
+							Dt,
+							CollisionParticlesRange.GetRangeSize(),
+							BatchBegin,
+							BatchEnd);
+					}
+					else
+					{
+						ispc::ApplyPerParticleSimpleCollisionNoFrictionAndGeneratePlanarConstraints(
+							(ispc::FVector4f*)Particles.GetPAndInvM().GetData(),
+							HasPlanarData.GetData(),
+							(ispc::FVector3f*)PlanarDataPositions.GetData(),
+							(ispc::FVector3f*)PlanarDataNormals.GetData(),
+							(const ispc::FVector3f*)Particles.XArray().GetData(),
+							(const ispc::FVector3f*)CollisionParticlesRange.XArray().GetData(),
+							(const ispc::FVector4f*)CollisionParticlesRange.GetR().GetData(),
+							SoftBodyCollisionThickness + CollisionThickness,
+							(const uint8*)&CollisionParticlesRange,
+							(const uint8*)CollisionParticlesRange.GetAllGeometry().GetData(),
+							sizeof(FImplicitObject),
+							FImplicitObject::GetOffsetOfType(),
+							FImplicitObject::GetOffsetOfMargin(),
+							CollisionParticlesRange.GetRangeSize(),
+							BatchBegin,
+							BatchEnd);
+					}
+				}
+				else
 				{
-					if constexpr (bWithFriction)
+					if (bWithFriction)
 					{
-						BatchPhis.SetNum(Particles.GetRangeSize());
-						BatchNormals.SetNum(Particles.GetRangeSize());
-						BatchVelocityBones.SetNum(Particles.GetRangeSize());
-
-						const int32 MLLevelSetThreadNum = 0;
-						const int32 BatchBegin = 0;
-						const int32 BatchEnd = Particles.GetRangeSize();
-
-						MLLevelSet->BatchPhiWithNormalWithGreatestInfluenceBone(Particles.GetPAndInvM(), Frame,
-							BatchPhis, BatchNormals, BatchVelocityBones,
-							CollisionThickness, MLLevelSetThreadNum, BatchBegin, BatchEnd);
-
-						check(BatchPhis.Num() == Particles.GetRangeSize());
-						check(BatchNormals.Num() == Particles.GetRangeSize());
-						check(BatchVelocityBones.Num() == Particles.GetRangeSize());
-
-						const FSolverRotation3 CPR_R = CollisionParticlesRange.R(CollisionIndex);
-						ispc::ApplyPerParticleCollisionFastFrictionBatchedPhiAndNormal(
+						ispc::ApplyPerParticleSimpleCollisionFastFriction(
 							(ispc::FVector4f*)Particles.GetPAndInvM().GetData(),
 							(const ispc::FVector3f*)Particles.XArray().GetData(),
-							CollisionIndex,
-							BatchPhis.GetData(),
-							(const ispc::FVector3f*)BatchNormals.GetData(),
-							BatchVelocityBones.GetData(),
 							(const ispc::FVector3f*)CollisionParticlesRange.GetV().GetData(),
 							(const ispc::FVector3f*)CollisionParticlesRange.XArray().GetData(),
 							(const ispc::FVector3f*)CollisionParticlesRange.GetW().GetData(),
-							reinterpret_cast<const ispc::FVector4f&>(CPR_R),
-							//(const ispc::FVector4f&)CollisionParticlesRange.R(CollisionIndex),
+							(const ispc::FVector4f*)CollisionParticlesRange.GetR().GetData(),
 							FrictionCoefficient,
-							CollisionThickness,
+							SoftBodyCollisionThickness + CollisionThickness,
+							(const uint8*)&CollisionParticlesRange,
+							(const uint8*)CollisionParticlesRange.GetAllGeometry().GetData(),
+							sizeof(FImplicitObject),
+							FImplicitObject::GetOffsetOfType(),
+							FImplicitObject::GetOffsetOfMargin(),
 							Dt,
 							CollisionParticlesRange.GetRangeSize(),
 							BatchBegin,
@@ -698,567 +1039,262 @@ void FPBDSoftBodyCollisionConstraintBase::ApplyInternalISPC(FSolverParticlesRang
 					}
 					else
 					{
-						BatchPhis.SetNum(Particles.GetRangeSize());
-						BatchNormals.SetNum(Particles.GetRangeSize());
-
-						const int32 MLLevelSetThreadNum = 0;
-						const int32 BatchBegin = 0;
-						const int32 BatchEnd = Particles.GetRangeSize();
-
-						MLLevelSet->BatchPhiWithNormalWithGreatestInfluenceBone(Particles.GetPAndInvM(), Frame,
-							BatchPhis, BatchNormals, BatchVelocityBones,
-							CollisionThickness, MLLevelSetThreadNum, BatchBegin, BatchEnd);
-
-						check(BatchPhis.Num() == Particles.GetRangeSize());
-						check(BatchNormals.Num() == Particles.GetRangeSize());
-
-						const FSolverRotation3 CPR_R = CollisionParticlesRange.R(CollisionIndex);
-						ispc::ApplyPerParticleCollisionNoFrictionBatchedPhiAndNormal(
+						ispc::ApplyPerParticleSimpleCollisionNoFriction(
 							(ispc::FVector4f*)Particles.GetPAndInvM().GetData(),
-							BatchPhis.GetData(),
-							(const ispc::FVector3f*)BatchNormals.GetData(),
-							reinterpret_cast<const ispc::FVector4f&>(CPR_R),
-							//(const ispc::FVector4f&)(CollisionParticlesRange.R(CollisionIndex)),
-							CollisionThickness,
-							Dt,
+							(const ispc::FVector3f*)Particles.XArray().GetData(),
+							(const ispc::FVector3f*)CollisionParticlesRange.XArray().GetData(),
+							(const ispc::FVector4f*)CollisionParticlesRange.GetR().GetData(),
+							SoftBodyCollisionThickness + CollisionThickness,
+							(const uint8*)&CollisionParticlesRange,
+							(const uint8*)CollisionParticlesRange.GetAllGeometry().GetData(),
+							sizeof(FImplicitObject),
+							FImplicitObject::GetOffsetOfType(),
+							FImplicitObject::GetOffsetOfMargin(),
+							CollisionParticlesRange.GetRangeSize(),
 							BatchBegin,
 							BatchEnd);
 					}
 				}
 			}
-		}
-
-		if (MLLevelSetParallelize == 1)
-		{
-			//Find MLLevelSet CollisionIndices for Parallelization
-			TArray<int32> MLLevelSetIndexes;
-			for (int32 CollisionIndex = 0; CollisionIndex < CollisionParticlesRange.GetRangeSize(); ++CollisionIndex)
-			{
-				if (FMLLevelSet* MLLevelSet = CollisionParticlesRange.GetGeometry(CollisionIndex)->GetObject<FMLLevelSet>())
-				{
-					MLLevelSetIndexes.Add(CollisionIndex);
-				}
-			}
-
-			const int32 NumberOfMLLevelSets = MLLevelSetIndexes.Num();
-			if (NumberOfMLLevelSets > 0)
-			{
-				const int32 NumberOfParallelCalls = NumberOfMLLevelSets;
-				TArray<TArray<FSolverReal>> BatchPhisArr;
-				TArray<TArray<FSolverVec3>> BatchNormalsArr;
-				TArray<TArray<int32>> BatchVelocityBonesArr;
-
-				BatchPhisArr.SetNum(MLLevelSetIndexes.Num());
-				BatchNormalsArr.SetNum(MLLevelSetIndexes.Num());
-				BatchVelocityBonesArr.SetNum(MLLevelSetIndexes.Num());
-
-				// Allocate TArray's for Parallel MLLevelSet Calls.
-				for (int32 MLLevelSetIndex = 0; MLLevelSetIndex < NumberOfMLLevelSets; MLLevelSetIndex++)
-				{
-					BatchPhisArr[MLLevelSetIndex].SetNum(NumberOfParticles);
-					BatchNormalsArr[MLLevelSetIndex].SetNum(NumberOfParticles);
-					BatchVelocityBonesArr[MLLevelSetIndex].SetNum(NumberOfParticles);
-				}
-
-				PhysicsParallelFor(MLLevelSetIndexes.Num(), [this, &Particles, Dt, &MLLevelSetIndexes, &CollisionParticlesRange,
-					&PAndInvM, &BatchVelocityBonesArr, &BatchNormalsArr, &BatchPhisArr](int32 IndexML)
-					{
-						const int32 MLLevelSetIndex = IndexML;
-						const int32 CollisionIndex = MLLevelSetIndexes[MLLevelSetIndex];
-						const int32 MLLevelSetThread = 0;
-						const int32 OuterParticleBatchBegin = 0;
-						const int32 OuterParticleBatchEnd = Particles.GetRangeSize() - 1;						
-
-						FMLLevelSet* MLLevelSet = CollisionParticlesRange.GetGeometry(CollisionIndex)->GetObject<FMLLevelSet>();
-						FSolverRigidTransform3 Frame = FSolverRigidTransform3(CollisionParticlesRange.X(CollisionIndex), CollisionParticlesRange.R(CollisionIndex)).Inverse();
-
-						if constexpr (bWithFriction)
-						{
-							MLLevelSet->BatchPhiWithNormalWithGreatestInfluenceBone(Particles.GetPAndInvM(), Frame,
-								BatchPhisArr[MLLevelSetIndex], BatchNormalsArr[MLLevelSetIndex], BatchVelocityBonesArr[MLLevelSetIndex],
-								CollisionThickness, MLLevelSetThread, OuterParticleBatchBegin, OuterParticleBatchEnd);
-
-							check(BatchPhisArr[MLLevelSetIndex].Num() == Particles.GetRangeSize());
-							check(BatchNormalsArr[MLLevelSetIndex].Num() == Particles.GetRangeSize());
-							check(BatchVelocityBonesArr[MLLevelSetIndex].Num() == Particles.GetRangeSize());
-
-							const FSolverRotation3 CPR_R = CollisionParticlesRange.R(CollisionIndex);
-							ispc::ApplyPerParticleCollisionFastFrictionBatchedPhiAndNormal(
-								(ispc::FVector4f*)Particles.GetPAndInvM().GetData(),
-								(const ispc::FVector3f*)Particles.XArray().GetData(),
-								CollisionIndex,
-								BatchPhisArr[MLLevelSetIndex].GetData(),
-								(const ispc::FVector3f*)BatchNormalsArr[MLLevelSetIndex].GetData(),
-								BatchVelocityBonesArr[MLLevelSetIndex].GetData(),
-								(const ispc::FVector3f*)CollisionParticlesRange.GetV().GetData(),
-								(const ispc::FVector3f*)CollisionParticlesRange.XArray().GetData(),
-								(const ispc::FVector3f*)CollisionParticlesRange.GetW().GetData(),
-								reinterpret_cast<const ispc::FVector4f&>(CPR_R),
-								//(const ispc::FVector4f&)CollisionParticlesRange.R(CollisionIndex),
-								FrictionCoefficient,
-								CollisionThickness,
-								Dt,
-								CollisionParticlesRange.GetRangeSize(),
-								OuterParticleBatchBegin,
-								OuterParticleBatchEnd);
-						}
-						else
-						{
-							MLLevelSet->BatchPhiWithNormal(Particles.GetPAndInvM(), Frame,
-								BatchPhisArr[MLLevelSetIndex], BatchNormalsArr[MLLevelSetIndex],
-								CollisionThickness, MLLevelSetThread, OuterParticleBatchBegin, OuterParticleBatchEnd);
-
-							check(BatchPhisArr[MLLevelSetIndex].Num() == Particles.GetRangeSize());
-							check(BatchNormalsArr[MLLevelSetIndex].Num() == Particles.GetRangeSize());
-
-							const FSolverRotation3 CPR_R = CollisionParticlesRange.R(CollisionIndex);
-							ispc::ApplyPerParticleCollisionNoFrictionBatchedPhiAndNormal(
-								(ispc::FVector4f*)Particles.GetPAndInvM().GetData(),
-								BatchPhisArr[MLLevelSetIndex].GetData(),
-								(const ispc::FVector3f*)BatchNormalsArr[MLLevelSetIndex].GetData(),
-								reinterpret_cast<const ispc::FVector4f&>(CPR_R),
-								//(const ispc::FVector4f&)CollisionParticlesRange.R(CollisionIndex),
-								CollisionThickness,
-								Dt,
-								OuterParticleBatchBegin,
-								OuterParticleBatchEnd);
-						}
-					});
-			}		
-		}		
-	}
-#endif  // #if INTEL_ISPC
+		});
 }
 
-template<bool bWithFriction>
-void FPBDSoftBodyCollisionConstraintBase::ApplyInternalISPCPerTimestep(FSolverParticlesRange& Particles, const FSolverReal Dt, const TArray<FSolverCollisionParticlesRange>& CollisionParticles) const
+void FPBDSoftBodyCollisionConstraintBase::ApplyComplexInternalISPC(FSolverParticlesRange& Particles, const FSolverReal Dt, const TArray<FSolverCollisionParticlesRange>& CollisionParticles, bool bGeneratePlanarConstraints)
 {
-#if INTEL_ISPC
-	check(bRealTypeCompatibleWithISPC);
-	const int32 NumBatches = FMath::CeilToInt((FSolverReal)(Particles.GetRangeSize()) / (FSolverReal)Chaos_SoftBodyCollision_ISPC_ParallelBatchSize);
+	if (!bEnableComplexColliders)
 	{
-		PhysicsParallelFor(NumBatches, [this, &Particles, Dt, &CollisionParticles](int32 BatchNumber)
-		{
-			const int32 BatchBegin = (Chaos_SoftBodyCollision_ISPC_ParallelBatchSize * BatchNumber);
-			const int32 BatchEnd = FMath::Min(Particles.GetRangeSize(), BatchBegin + Chaos_SoftBodyCollision_ISPC_ParallelBatchSize);
-
-			for (const FSolverCollisionParticlesRange& CollisionParticlesRange : CollisionParticles)
-			{
-				if constexpr (bWithFriction)
-				{
-					ispc::ApplyPerParticleCollisionFastFrictionNoGroupCheck(
-						(ispc::FVector4f*)Particles.GetPAndInvM().GetData(),
-						(const ispc::FVector3f*)Particles.XArray().GetData(),
-						(const ispc::FVector3f*)CollisionParticlesRange.GetV().GetData(),
-						(const ispc::FVector3f*)CollisionParticlesRange.XArray().GetData(),
-						(const ispc::FVector3f*)CollisionParticlesRange.GetW().GetData(),
-						(const ispc::FVector4f*)CollisionParticlesRange.GetR().GetData(),
-						FrictionCoefficient,
-						CollisionThickness,
-						(const uint8*)&CollisionParticlesRange,
-						(const uint8*)CollisionParticlesRange.GetAllGeometry().GetData(),
-						sizeof(FImplicitObject),
-						FImplicitObject::GetOffsetOfType(),
-						FImplicitObject::GetOffsetOfMargin(),
-						Dt,
-						CollisionParticlesRange.GetRangeSize(),
-						BatchBegin,
-						BatchEnd);
-				}
-				else
-				{
-					ispc::ApplyPerParticleCollisionNoFrictionNoGroupCheck(
-						(ispc::FVector4f*)Particles.GetPAndInvM().GetData(),
-						(const ispc::FVector3f*)Particles.XArray().GetData(),
-						(const ispc::FVector3f*)CollisionParticlesRange.GetV().GetData(),
-						(const ispc::FVector3f*)CollisionParticlesRange.XArray().GetData(),
-						(const ispc::FVector3f*)CollisionParticlesRange.GetW().GetData(),
-						(const ispc::FVector4f*)CollisionParticlesRange.GetR().GetData(),
-						CollisionThickness,
-						(const uint8*)&CollisionParticlesRange,
-						(const uint8*)CollisionParticlesRange.GetAllGeometry().GetData(),
-						sizeof(FImplicitObject),
-						FImplicitObject::GetOffsetOfType(),
-						FImplicitObject::GetOffsetOfMargin(),
-						Dt,
-						CollisionParticlesRange.GetRangeSize(),
-						BatchBegin,
-						BatchEnd);
-				}
-			}
-		});
+		return;
 	}
 
-	FPAndInvM* const PAndInvM = Particles.GetPAndInvM().GetData();
-
-	CollisionIndices.Empty();
-	CollisionTargets.Empty();
-	CollisionNormalsOutput.Empty();
-	TSet<int32> CollisionIndicesSet;
+	TRACE_CPUPROFILER_EVENT_SCOPE(FPBDSoftBodyCollisionConstraint_ApplyComplexInternalISPC);
 
-	// Skinned levelset passes
+	const bool bWithFriction = FrictionCoefficient > (FSolverReal)UE_KINDA_SMALL_NUMBER;
 	for (const FSolverCollisionParticlesRange& CollisionParticlesRange : CollisionParticles)
 	{
 		for (int32 CollisionIndex = 0; CollisionIndex < CollisionParticlesRange.GetRangeSize(); ++CollisionIndex)
 		{
-			//This part is updated to visualize Normals. This part will be returned back.
-			if (const TWeightedLatticeImplicitObject<FLevelSet>* const LevelSet = CollisionParticlesRange.GetGeometry(CollisionIndex)->GetObject< TWeightedLatticeImplicitObject<FLevelSet> >())
+			const FImplicitObjectPtr& Geometry = CollisionParticlesRange.GetGeometry(CollisionIndex);
+			const Chaos::EImplicitObjectType CollisionType = Geometry->GetType();
+			if (!Private::IsComplexBatchCollider(CollisionType))
 			{
-				TRACE_CPUPROFILER_EVENT_SCOPE(FPBDSoftBodyCollisionConstraint_ApplyInternalISPC_LevelSet);
-				const FSolverRigidTransform3 Frame = FSolverRigidTransform3(CollisionParticlesRange.X(CollisionIndex), CollisionParticlesRange.R(CollisionIndex)).Inverse();
-				if constexpr (bWithFriction)
-				{
-					LevelSet->BatchPhiWithNormalAndGreatestInfluenceBone(Particles.GetPAndInvM(), Frame, CollisionThickness, BatchPhis, BatchNormals, BatchVelocityBones);
-					check(BatchPhis.Num() == Particles.GetRangeSize());
-					check(BatchNormals.Num() == Particles.GetRangeSize());
-					check(BatchVelocityBones.Num() == Particles.GetRangeSize());
-
-					const FSolverRotation3 CPR_R = CollisionParticlesRange.R(CollisionIndex);
-					ispc::ApplyPerParticleCollisionFastFrictionBatchedPhiAndNormal(
-						(ispc::FVector4f*)Particles.GetPAndInvM().GetData(),
-						(const ispc::FVector3f*)Particles.XArray().GetData(),
-						CollisionIndex,
-						BatchPhis.GetData(),
-						(const ispc::FVector3f*)BatchNormals.GetData(),
-						BatchVelocityBones.GetData(),
-						(const ispc::FVector3f*)CollisionParticlesRange.GetV().GetData(),
-						(const ispc::FVector3f*)CollisionParticlesRange.XArray().GetData(),
-						(const ispc::FVector3f*)CollisionParticlesRange.GetW().GetData(),
-						reinterpret_cast<const ispc::FVector4f&>(CPR_R),
-						//(const ispc::FVector4f*)CollisionParticlesRange.R(CollisionIndex),
-						FrictionCoefficient,
-						CollisionThickness,
-						Dt,
-						CollisionParticlesRange.GetRangeSize(),
-						0,
-						Particles.GetRangeSize());
-				}
-				else
-				{
-					LevelSet->BatchPhiWithNormal(Particles.GetPAndInvM(), Frame, BatchPhis, BatchNormals);
-					check(BatchPhis.Num() == Particles.GetRangeSize());
-					check(BatchNormals.Num() == Particles.GetRangeSize());
+				continue;
+			}
 
-					const FSolverRotation3 CPR_R = CollisionParticlesRange.R(CollisionIndex);
-					ispc::ApplyPerParticleCollisionNoFrictionBatchedPhiAndNormal(
-						(ispc::FVector4f*)Particles.GetPAndInvM().GetData(),
-						BatchPhis.GetData(),
-						(const ispc::FVector3f*)BatchNormals.GetData(),
-						reinterpret_cast<const ispc::FVector4f&>(CPR_R),
-						//(const ispc::FVector4f&)CollisionParticlesRange.R(CollisionIndex),
-						CollisionThickness,
-						Dt,
-						0,
-						Particles.GetRangeSize());
-				}
+			const FSolverVec3& CollisionX = CollisionParticlesRange.X(CollisionIndex);
+			const FSolverRotation3& CollisionR = CollisionParticlesRange.R(CollisionIndex);
+			const FSolverRigidTransform3 Frame(CollisionX, CollisionR);
+			if (bWithFriction)
+			{
+				const FSolverVec3& CollisionV = CollisionParticlesRange.V(CollisionIndex);
+				const FSolverVec3& CollisionW = CollisionParticlesRange.W(CollisionIndex);
 
-				//obtain data for planar constraints 
-				const FSolverReal* const CollisionPhis = BatchPhis.GetData();
-				const FSolverVec3* const CollisionNormals = BatchNormals.GetData();
-				for (int32 Index = 0; Index < BatchPhis.Num(); Index++)
+				if (const FMLLevelSet* const MLLevelSet = Geometry->GetObject<FMLLevelSet>())
 				{
-					if (PAndInvM[Index].InvM == (FSolverReal)0.)
-					{
-						continue;
-					}
-
-					const FSolverReal Penetration = CollisionThickness - CollisionPhis[Index];
-
-					if (Penetration > (FSolverReal)0.)
-					{
-						const FSolverVec3 NormalWorld = Frame.TransformVector(CollisionNormals[Index]);
-						if (!CollisionIndicesSet.Contains(Index))
-						{
-							CollisionIndicesSet.Add(Index);
-							CollisionIndices.Emplace(Index);
-							CollisionTargets.Emplace(PAndInvM[Index].P);
-							CollisionNormalsOutput.Emplace(NormalWorld);
-						}
-					}
-				}
-
-			}
-		}
-	}
+					BatchPhis.SetNumUninitialized(Particles.GetRangeSize());
+					BatchNormals.SetNumUninitialized(Particles.GetRangeSize());
 
-	// MLLevelset passes in ApplyInternalISPCPerTimestep
-	for (const FSolverCollisionParticlesRange& CollisionParticlesRange : CollisionParticles)
-	{
-		TRACE_CPUPROFILER_EVENT_SCOPE(PBDSoftBody_ApplyInternalISPCPerTimestep_FMLLevelset);		
-		const int32 NumberOfParticles = Particles.GetRangeSize();
+					// Batch Query
+					constexpr int32 MLLevelSetThreadNum = 0;
+					MLLevelSet->BatchPhiWithNormal(Particles.GetPAndInvM(), Frame, BatchPhis, BatchNormals, SoftBodyCollisionThickness + CollisionThickness, MLLevelSetThreadNum, 0, Particles.GetRangeSize());
 
-		if (MLLevelSetParallelize == 0)
-		{			
-			for (int32 CollisionIndex = 0; CollisionIndex < CollisionParticlesRange.GetRangeSize(); ++CollisionIndex)
-			{
-				const FSolverRigidTransform3 Frame = FSolverRigidTransform3(CollisionParticlesRange.X(CollisionIndex), CollisionParticlesRange.R(CollisionIndex)).Inverse();
-				
-				if (const FMLLevelSet* const MLLevelSet = CollisionParticlesRange.GetGeometry(CollisionIndex)->GetObject<FMLLevelSet>())
-				{
-					BatchPhis.Empty();
-					BatchNormals.Empty();
-					BatchVelocityBones.Empty();
-					if constexpr (bWithFriction)
+					// Apply
+					if (bGeneratePlanarConstraints)
 					{
-						BatchPhis.SetNum(Particles.GetRangeSize());
-						BatchNormals.SetNum(Particles.GetRangeSize());
-						BatchVelocityBones.SetNum(Particles.GetRangeSize());
-						
-						const int32 MLLevelSetThreadNum = 0;
-						const int32 BatchBegin = 0;
-						const int32 BatchEnd = Particles.GetRangeSize();
-
-						MLLevelSet->BatchPhiWithNormalWithGreatestInfluenceBone(Particles.GetPAndInvM(), Frame, 
-							BatchPhis, BatchNormals, BatchVelocityBones, 
-							CollisionThickness, MLLevelSetThreadNum, BatchBegin, BatchEnd);
-
-						check(BatchPhis.Num() == Particles.GetRangeSize());
-						check(BatchNormals.Num() == Particles.GetRangeSize());
-						check(BatchVelocityBones.Num() == Particles.GetRangeSize());
-
-						const FSolverRotation3 CPR_R = CollisionParticlesRange.R(CollisionIndex);
-						ispc::ApplyPerParticleCollisionFastFrictionBatchedPhiAndNormal(
+						ispc::ApplyPerParticleBatchCollisionFastFrictionAndGeneratePlanarConstraints(
 							(ispc::FVector4f*)Particles.GetPAndInvM().GetData(),
+							HasPlanarData.GetData(),
+							(ispc::FVector3f*)PlanarDataPositions.GetData(),
+							(ispc::FVector3f*)PlanarDataNormals.GetData(),
+							(ispc::FVector3f*)PlanarDataVelocities.GetData(),
 							(const ispc::FVector3f*)Particles.XArray().GetData(),
-							CollisionIndex,
 							BatchPhis.GetData(),
 							(const ispc::FVector3f*)BatchNormals.GetData(),
-							BatchVelocityBones.GetData(),
-							(const ispc::FVector3f*)CollisionParticlesRange.GetV().GetData(),
-							(const ispc::FVector3f*)CollisionParticlesRange.XArray().GetData(),
-							(const ispc::FVector3f*)CollisionParticlesRange.GetW().GetData(),
-							reinterpret_cast<const ispc::FVector4f&>(CPR_R),
-							//(const ispc::FVector4f&)CollisionParticlesRange.R(CollisionIndex),
+							reinterpret_cast<const ispc::FVector3f&>(CollisionV),
+							reinterpret_cast<const ispc::FVector3f&>(CollisionX),
+							reinterpret_cast<const ispc::FVector3f&>(CollisionW),
+							reinterpret_cast<const ispc::FVector4f&>(CollisionR),
 							FrictionCoefficient,
-							CollisionThickness,
+							SoftBodyCollisionThickness + CollisionThickness,
 							Dt,
-							CollisionParticlesRange.GetRangeSize(),
 							0,
 							Particles.GetRangeSize());
 					}
 					else
 					{
-						BatchPhis.SetNum(Particles.GetRangeSize());
-						BatchNormals.SetNum(Particles.GetRangeSize());
-						const int32 MLLevelSetThreadNum = 0;
-						const int32 BatchBegin = 0;
-						const int32 BatchEnd = Particles.GetRangeSize();
-
-						MLLevelSet->BatchPhiWithNormal(Particles.GetPAndInvM(), Frame, BatchPhis, BatchNormals, 
-							CollisionThickness, MLLevelSetThreadNum, BatchBegin, BatchEnd);
-
-						check(BatchPhis.Num() == Particles.GetRangeSize());
-						check(BatchNormals.Num() == Particles.GetRangeSize());
-
-						const FSolverRotation3 CPR_R = CollisionParticlesRange.R(CollisionIndex);
-						ispc::ApplyPerParticleCollisionNoFrictionBatchedPhiAndNormal(
+						ispc::ApplyPerParticleBatchCollisionFastFriction(
 							(ispc::FVector4f*)Particles.GetPAndInvM().GetData(),
+							(const ispc::FVector3f*)Particles.XArray().GetData(),
 							BatchPhis.GetData(),
 							(const ispc::FVector3f*)BatchNormals.GetData(),
-							reinterpret_cast<const ispc::FVector4f&>(CPR_R),
-							//(const ispc::FVector4f&)CollisionParticlesRange.R(CollisionIndex),
-							CollisionThickness,
+							reinterpret_cast<const ispc::FVector3f&>(CollisionV),
+							reinterpret_cast<const ispc::FVector3f&>(CollisionX),
+							reinterpret_cast<const ispc::FVector3f&>(CollisionW),
+							reinterpret_cast<const ispc::FVector4f&>(CollisionR),
+							FrictionCoefficient,
+							SoftBodyCollisionThickness + CollisionThickness,
 							Dt,
 							0,
 							Particles.GetRangeSize());
 					}
-
-
-					const FSolverReal* const CollisionPhis = BatchPhis.GetData();
-					const FSolverVec3* const CollisionNormals = BatchNormals.GetData();
-					
-					for (int32 Index = 0; Index < BatchPhis.Num(); Index++)
+				}
+				else if (const TWeightedLatticeImplicitObject<FLevelSet>* const LevelSet = Geometry->GetObject<TWeightedLatticeImplicitObject<FLevelSet>>())
+				{
+					if (const FPBDComplexColliderBoneData* const ColliderBoneData = ComplexBoneData.Find(FParticleRangeIndex(CollisionParticlesRange.GetRangeId(), CollisionIndex)))
 					{
-						if (PAndInvM[Index].InvM == (FSolverReal)0.)
+						LevelSet->BatchPhiWithNormalAndGreatestInfluenceBone(Particles.GetPAndInvM(), Frame, SoftBodyCollisionThickness + CollisionThickness, BatchPhis, BatchNormals, BatchVelocityBones);
+
+						if (bGeneratePlanarConstraints)
 						{
-							continue;
+							ispc::ApplyPerParticleBatchCollisionFastFrictionWithVelocityBonesAndGeneratePlanarConstraints(
+								(ispc::FVector4f*)Particles.GetPAndInvM().GetData(),
+								HasPlanarData.GetData(),
+								(ispc::FVector3f*)PlanarDataPositions.GetData(),
+								(ispc::FVector3f*)PlanarDataNormals.GetData(),
+								(ispc::FVector3f*)PlanarDataVelocities.GetData(),
+								(const ispc::FVector3f*)Particles.XArray().GetData(),
+								BatchPhis.GetData(),
+								(const ispc::FVector3f*)BatchNormals.GetData(),
+								BatchVelocityBones.GetData(),
+								reinterpret_cast<const ispc::FVector3f&>(CollisionV),
+								reinterpret_cast<const ispc::FVector3f&>(CollisionX),
+								reinterpret_cast<const ispc::FVector3f&>(CollisionW),
+								reinterpret_cast<const ispc::FVector4f&>(CollisionR),
+								ColliderBoneData->MappedBoneIndices.GetData(),
+								ColliderBoneData->MappedBoneIndices.Num(),
+								(const ispc::FVector3f*)ColliderBoneData->V.GetData(),
+								(const ispc::FVector3f*)ColliderBoneData->X.GetData(),
+								(const ispc::FVector3f*)ColliderBoneData->W.GetData(),
+								FrictionCoefficient,
+								SoftBodyCollisionThickness + CollisionThickness,
+								Dt,
+								0,
+								Particles.GetRangeSize());
 						}
-
-						const FSolverReal Penetration = CollisionThickness - CollisionPhis[Index];
-
-						if (Penetration > (FSolverReal)0.)
+						else
 						{
-							if (!CollisionIndicesSet.Contains(Index))
-							{
-								const FSolverVec3 NormalWorld = Frame.TransformVector(CollisionNormals[Index]);
-
-								CollisionIndicesSet.Add(Index);
-								CollisionIndices.Emplace(Index);
-								CollisionTargets.Emplace(PAndInvM[Index].P);
-								CollisionNormalsOutput.Emplace(NormalWorld);
-							}
-
+							ispc::ApplyPerParticleBatchCollisionFastFrictionWithVelocityBones(
+								(ispc::FVector4f*)Particles.GetPAndInvM().GetData(),
+								(const ispc::FVector3f*)Particles.XArray().GetData(),
+								BatchPhis.GetData(),
+								(const ispc::FVector3f*)BatchNormals.GetData(),
+								BatchVelocityBones.GetData(),
+								reinterpret_cast<const ispc::FVector3f&>(CollisionV),
+								reinterpret_cast<const ispc::FVector3f&>(CollisionX),
+								reinterpret_cast<const ispc::FVector3f&>(CollisionW),
+								reinterpret_cast<const ispc::FVector4f&>(CollisionR),
+								ColliderBoneData->MappedBoneIndices.GetData(),
+								ColliderBoneData->MappedBoneIndices.Num(),
+								(const ispc::FVector3f*)ColliderBoneData->V.GetData(),
+								(const ispc::FVector3f*)ColliderBoneData->X.GetData(),
+								(const ispc::FVector3f*)ColliderBoneData->W.GetData(),
+								FrictionCoefficient,
+								SoftBodyCollisionThickness + CollisionThickness,
+								Dt,
+								0,
+								Particles.GetRangeSize());
 						}
 					}
-					
-				}
-			}
-		}
-
-		if (MLLevelSetParallelize == 1)
-		{
-			//Find MLLevelSet CollisionIndices for Parallelization
-			TArray<int32> MLLevelSetIndexes;
-			for (int32 CollisionIndex = 0; CollisionIndex < CollisionParticlesRange.GetRangeSize(); ++CollisionIndex)
-			{
-				if (FMLLevelSet* MLLevelSet = CollisionParticlesRange.GetGeometry(CollisionIndex)->GetObject<FMLLevelSet>())
-				{
-					MLLevelSetIndexes.Add(CollisionIndex);
-				}
-			}			
-			
-			const int32 NumberOfMLLevelSets = MLLevelSetIndexes.Num();
-			if (NumberOfMLLevelSets > 0)
-			{
-				int32 NumberOfParallelThreadsPerMLLevelset = 1;
-
-				for (int32 MLCollisionIndex : MLLevelSetIndexes)
-				{
-					FMLLevelSet* MLLevelSet = CollisionParticlesRange.GetGeometry(MLCollisionIndex)->GetObject<FMLLevelSet>();
-					MLLevelSet->UpdateNeuralInferencesNumber(NumberOfParallelThreadsPerMLLevelset);
-				}
-
-				const int32 NumberOfParallelCalls = NumberOfParallelThreadsPerMLLevelset * NumberOfMLLevelSets;
-				TArray<TArray<FSolverReal>> BatchPhisArr;
-				TArray<TArray<FSolverVec3>> BatchNormalsArr;
-				TArray<TArray<int32>> BatchVelocityBonesArr;
-
-				BatchPhisArr.SetNum(MLLevelSetIndexes.Num());
-				BatchNormalsArr.SetNum(MLLevelSetIndexes.Num());
-				BatchVelocityBonesArr.SetNum(MLLevelSetIndexes.Num());
-
-				// Allocate TArray's for Parallel MLLevelSet Calls.
-				for (int32 MLLevelSetIndex = 0; MLLevelSetIndex < NumberOfMLLevelSets; MLLevelSetIndex++)
-				{
-					BatchPhisArr[MLLevelSetIndex].SetNum(NumberOfParticles);
-					BatchNormalsArr[MLLevelSetIndex].SetNum(NumberOfParticles);
-					BatchVelocityBonesArr[MLLevelSetIndex].SetNum(NumberOfParticles);
-				}
-
-				// Data for the Planar Constraints.
-				TArray<TArray<int32>> BatchCollisionIndices;
-				TArray<TArray<FSolverVec3>> BatchCollisionTargets;
-				TArray<TArray<FSolverVec3>> BatchCollisionNormalsOutput;
-
-				BatchCollisionIndices.SetNum(NumberOfParallelCalls);
-				BatchCollisionTargets.SetNum(NumberOfParallelCalls);
-				BatchCollisionNormalsOutput.SetNum(NumberOfParallelCalls);
-
-				PhysicsParallelFor(NumberOfParallelThreadsPerMLLevelset* MLLevelSetIndexes.Num(), [this, &Particles, Dt, &MLLevelSetIndexes, &CollisionParticlesRange,
-					&PAndInvM, &BatchVelocityBonesArr, &BatchNormalsArr, &BatchPhisArr, NumberOfParallelThreadsPerMLLevelset,
-					&BatchCollisionIndices, &BatchCollisionTargets, &BatchCollisionNormalsOutput](int32 IndexML)
+					else
 					{
-						const int32 MLLevelSetIndex = IndexML;
-						const int32 CollisionIndex = MLLevelSetIndexes[MLLevelSetIndex];
-
-						const int32 MLLevelSetThread = 0;
-						const int32 MLLevelSetThreadBlock = Particles.GetRangeSize() / NumberOfParallelThreadsPerMLLevelset;
-
-						int32 OuterParticleBatchBegin = (NumberOfParallelThreadsPerMLLevelset - 1) * MLLevelSetThreadBlock;
-						int32 OuterParticleBatchEnd = Particles.GetRangeSize() - 1;
-
-						FMLLevelSet* MLLevelSet = CollisionParticlesRange.GetGeometry(CollisionIndex)->GetObject<FMLLevelSet>();
-						FSolverRigidTransform3 Frame = FSolverRigidTransform3(CollisionParticlesRange.X(CollisionIndex), CollisionParticlesRange.R(CollisionIndex)).Inverse();
-
-						if constexpr (bWithFriction)
+						LevelSet->BatchPhiWithNormal(Particles.GetPAndInvM(), Frame, BatchPhis, BatchNormals);
+						if (bGeneratePlanarConstraints)
 						{
-							MLLevelSet->BatchPhiWithNormalWithGreatestInfluenceBone(Particles.GetPAndInvM(), Frame,
-								BatchPhisArr[MLLevelSetIndex], BatchNormalsArr[MLLevelSetIndex], BatchVelocityBonesArr[MLLevelSetIndex],
-								CollisionThickness, MLLevelSetThread, OuterParticleBatchBegin, OuterParticleBatchEnd);
-
-							check(BatchPhisArr[MLLevelSetIndex].Num() == Particles.GetRangeSize());
-							check(BatchNormalsArr[MLLevelSetIndex].Num() == Particles.GetRangeSize());
-							check(BatchVelocityBonesArr[MLLevelSetIndex].Num() == Particles.GetRangeSize());
-
-							const FSolverRotation3 CPR_R = CollisionParticlesRange.R(CollisionIndex);
-							ispc::ApplyPerParticleCollisionFastFrictionBatchedPhiAndNormal(
+							ispc::ApplyPerParticleBatchCollisionFastFrictionAndGeneratePlanarConstraints(
 								(ispc::FVector4f*)Particles.GetPAndInvM().GetData(),
+								HasPlanarData.GetData(),
+								(ispc::FVector3f*)PlanarDataPositions.GetData(),
+								(ispc::FVector3f*)PlanarDataNormals.GetData(),
+								(ispc::FVector3f*)PlanarDataVelocities.GetData(),
 								(const ispc::FVector3f*)Particles.XArray().GetData(),
-								CollisionIndex,
-								BatchPhisArr[MLLevelSetIndex].GetData(),
-								(const ispc::FVector3f*)BatchNormalsArr[MLLevelSetIndex].GetData(),
-								BatchVelocityBonesArr[MLLevelSetIndex].GetData(),
-								(const ispc::FVector3f*)CollisionParticlesRange.GetV().GetData(),
-								(const ispc::FVector3f*)CollisionParticlesRange.XArray().GetData(),
-								(const ispc::FVector3f*)CollisionParticlesRange.GetW().GetData(),
-								reinterpret_cast<const ispc::FVector4f&>(CPR_R),
-								//(const ispc::FVector4f&)CollisionParticlesRange.R(CollisionIndex),
+								BatchPhis.GetData(),
+								(const ispc::FVector3f*)BatchNormals.GetData(),
+								reinterpret_cast<const ispc::FVector3f&>(CollisionV),
+								reinterpret_cast<const ispc::FVector3f&>(CollisionX),
+								reinterpret_cast<const ispc::FVector3f&>(CollisionW),
+								reinterpret_cast<const ispc::FVector4f&>(CollisionR),
 								FrictionCoefficient,
-								CollisionThickness,
+								SoftBodyCollisionThickness + CollisionThickness,
 								Dt,
-								CollisionParticlesRange.GetRangeSize(),
-								OuterParticleBatchBegin,
-								OuterParticleBatchEnd);
+								0,
+								Particles.GetRangeSize());
 						}
 						else
 						{
-							MLLevelSet->BatchPhiWithNormal(Particles.GetPAndInvM(), Frame,
-								BatchPhisArr[MLLevelSetIndex], BatchNormalsArr[MLLevelSetIndex],
-								CollisionThickness, MLLevelSetThread, OuterParticleBatchBegin, OuterParticleBatchEnd);
-
-							check(BatchPhisArr[MLLevelSetIndex].Num() == Particles.GetRangeSize());
-							check(BatchNormalsArr[MLLevelSetIndex].Num() == Particles.GetRangeSize());
-
-							const FSolverRotation3 CPR_R = CollisionParticlesRange.R(CollisionIndex);
-							ispc::ApplyPerParticleCollisionNoFrictionBatchedPhiAndNormal(
+							ispc::ApplyPerParticleBatchCollisionFastFriction(
 								(ispc::FVector4f*)Particles.GetPAndInvM().GetData(),
-								BatchPhisArr[MLLevelSetIndex].GetData(),
-								(const ispc::FVector3f*)BatchNormalsArr[MLLevelSetIndex].GetData(),
-								reinterpret_cast<const ispc::FVector4f&>(CPR_R),
-								//(const ispc::FVector4f&)CollisionParticlesRange.R(CollisionIndex),
-								CollisionThickness,
+								(const ispc::FVector3f*)Particles.XArray().GetData(),
+								BatchPhis.GetData(),
+								(const ispc::FVector3f*)BatchNormals.GetData(),
+								reinterpret_cast<const ispc::FVector3f&>(CollisionV),
+								reinterpret_cast<const ispc::FVector3f&>(CollisionX),
+								reinterpret_cast<const ispc::FVector3f&>(CollisionW),
+								reinterpret_cast<const ispc::FVector4f&>(CollisionR),
+								FrictionCoefficient,
+								SoftBodyCollisionThickness + CollisionThickness,
 								Dt,
-								OuterParticleBatchBegin,
-								OuterParticleBatchEnd);
-						}
-
-						//obtain data for planar constraints 
-						const FSolverReal* const CollisionPhis = BatchPhisArr[MLLevelSetIndex].GetData();
-						const FSolverVec3* const CollisionNormals = BatchNormalsArr[MLLevelSetIndex].GetData();
-						for (int32 Index = OuterParticleBatchBegin; Index < OuterParticleBatchEnd; Index++)
-						{
-							if (PAndInvM[Index].InvM == (FSolverReal)0.)
-							{
-								continue;
-							}
-
-							const FSolverReal Penetration = CollisionThickness - CollisionPhis[Index];
-
-							if (Penetration > (FSolverReal)0.)
-							{
-								const FSolverVec3 NormalWorld = Frame.TransformVector(CollisionNormals[Index]);
-								BatchCollisionIndices[IndexML].Emplace(Index);
-								BatchCollisionTargets[IndexML].Emplace(PAndInvM[Index].P);
-								BatchCollisionNormalsOutput[IndexML].Emplace(NormalWorld);
-							}
+								0,
+								Particles.GetRangeSize());
 						}
-					});
+					}
+				}
+				else
+				{
+					checkNoEntry();
+				}
+			}
+			else
+			{
+				if (const FMLLevelSet* const MLLevelSet = Geometry->GetObject<FMLLevelSet>())
+				{
+					BatchPhis.SetNumUninitialized(Particles.GetRangeSize());
+					BatchNormals.SetNumUninitialized(Particles.GetRangeSize());
 
-				for (int32 MLLevelSetIndex = 0; MLLevelSetIndex < NumberOfMLLevelSets; MLLevelSetIndex++)
+					// Batch Query
+					constexpr int32 MLLevelSetThreadNum = 0;
+					MLLevelSet->BatchPhiWithNormal(Particles.GetPAndInvM(), Frame, BatchPhis, BatchNormals, SoftBodyCollisionThickness + CollisionThickness, MLLevelSetThreadNum, 0, Particles.GetRangeSize());
+				}
+				else if (const TWeightedLatticeImplicitObject<FLevelSet>* const LevelSet = Geometry->GetObject<TWeightedLatticeImplicitObject<FLevelSet>>())
 				{
-					for (int32 MLLevelSetThread = 0; MLLevelSetThread < NumberOfParallelThreadsPerMLLevelset; MLLevelSetThread++)
-					{
-						int32 IndexML = NumberOfParallelThreadsPerMLLevelset * MLLevelSetIndex + MLLevelSetThread;
+					LevelSet->BatchPhiWithNormal(Particles.GetPAndInvM(), Frame, BatchPhis, BatchNormals);
+				}
+				else
+				{
+					checkNoEntry();
+				}
 
-						for (int32 IndexIndex = 0; IndexIndex < BatchCollisionIndices[IndexML].Num(); IndexIndex++)
-						{
-							int32 Index = BatchCollisionIndices[IndexML][IndexIndex];
-							if (!CollisionIndicesSet.Contains(Index))
-							{
-								CollisionIndicesSet.Add(Index);
-								CollisionIndices.Emplace(Index);
-								CollisionTargets.Emplace(BatchCollisionTargets[IndexML][IndexIndex]);
-								CollisionNormalsOutput.Emplace(BatchCollisionNormalsOutput[IndexML][IndexIndex]);
-							}
-						}
-					}
-				}				
-			}			
+				// Apply
+				if (bGeneratePlanarConstraints)
+				{
+					ispc::ApplyPerParticleBatchCollisionNoFrictionAndGeneratePlanarConstraints(
+						(ispc::FVector4f*)Particles.GetPAndInvM().GetData(),
+						HasPlanarData.GetData(),
+						(ispc::FVector3f*)PlanarDataPositions.GetData(),
+						(ispc::FVector3f*)PlanarDataNormals.GetData(),
+						(const ispc::FVector3f*)Particles.XArray().GetData(),
+						BatchPhis.GetData(),
+						(const ispc::FVector3f*)BatchNormals.GetData(),
+						reinterpret_cast<const ispc::FVector4f&>(CollisionR),
+						SoftBodyCollisionThickness + CollisionThickness,
+						0,
+						Particles.GetRangeSize());
+				}
+				else
+				{
+					ispc::ApplyPerParticleBatchCollisionNoFriction(
+						(ispc::FVector4f*)Particles.GetPAndInvM().GetData(),
+						(const ispc::FVector3f*)Particles.XArray().GetData(),
+						BatchPhis.GetData(),
+						(const ispc::FVector3f*)BatchNormals.GetData(),
+						reinterpret_cast<const ispc::FVector4f&>(CollisionR),
+						SoftBodyCollisionThickness + CollisionThickness,
+						0,
+						Particles.GetRangeSize());
+				}
+			}
 		}
 	}
-#endif  // #if INTEL_ISPC
 }
-
+#endif  // #if INTEL_ISPC
 
 void FPBDSoftBodyCollisionConstraintBase::UpdateLinearSystem(const FSolverParticlesRange& Particles, const FSolverReal Dt, const TArray<FSolverCollisionParticlesRange>& CollisionParticles, FEvolutionLinearSystem& LinearSystem) const
 {
@@ -1289,21 +1325,11 @@ void FPBDSoftBodyCollisionConstraintBase::UpdateLinearSystem(const FSolverPartic
 		{
 			for (int32 CollisionIndex = 0; CollisionIndex < CollisionParticlesRange.GetRangeSize(); ++CollisionIndex)
 			{
-				if (CollisionParticlesRange.GetGeometry(CollisionIndex)->GetType() == Chaos::ImplicitObjectType::WeightedLatticeBone)
-				{
-					continue;
-				}
-
-				if (CollisionParticlesRange.GetGeometry(CollisionIndex)->GetType() == Chaos::ImplicitObjectType::MLLevelSetBone)
-				{
-					continue;
-				}
-
 				const FSolverRigidTransform3 Frame(CollisionParticlesRange.X(CollisionIndex), CollisionParticlesRange.R(CollisionIndex));
 				const FVec3 RigidSpacePosition(Frame.InverseTransformPosition(PAndInvM[Index].P));  // PhiWithNormal requires FReal based arguments
 				FVec3 ImplicitNormal;                                                                // since implicits don't use FSolverReal
 				const FSolverReal Phi = (FSolverReal)CollisionParticlesRange.GetGeometry(CollisionIndex)->PhiWithNormal(RigidSpacePosition, ImplicitNormal);
-				const FSolverReal Penetration = CollisionThickness - Phi; // This is related to the Normal impulse
+				const FSolverReal Penetration = SoftBodyCollisionThickness + CollisionThickness - Phi; // This is related to the Normal impulse
 				const FSolverVec3 Normal(ImplicitNormal);
 
 				if (Penetration > (FSolverReal)0.)
@@ -1338,9 +1364,14 @@ void FPBDSoftBodyCollisionConstraint::SetProperties(const FCollectionPropertyCon
 	{
 		CollisionThickness = MeshScale * GetCollisionThickness(PropertyCollection);
 	}
+	if (IsSoftBodyCollisionThicknessMutable(PropertyCollection))
+	{
+		SoftBodyCollisionThickness = GetSoftBodyCollisionThickness(PropertyCollection);
+	}
 	if (IsFrictionCoefficientMutable(PropertyCollection))
 	{
 		FrictionCoefficient = GetFrictionCoefficient(PropertyCollection);
+		PlanarConstraint.SetFrictionCoefficient(FrictionCoefficient);
 	}
 	if (IsUseCCDMutable(PropertyCollection))
 	{
@@ -1350,6 +1381,14 @@ void FPBDSoftBodyCollisionConstraint::SetProperties(const FCollectionPropertyCon
 	{
 		ProximityStiffness = GetProximityStiffness(PropertyCollection);
 	}
+	if (IsUsePlanarConstraintForSimpleCollidersMutable(PropertyCollection))
+	{
+		bUsePlanarConstraintForSimpleColliders = GetUsePlanarConstraintForSimpleColliders(PropertyCollection);
+	}
+	if (IsUsePlanarConstraintForComplexCollidersMutable(PropertyCollection))
+	{
+		bUsePlanarConstraintForComplexColliders = GetUsePlanarConstraintForComplexColliders(PropertyCollection);
+	}
 }
 
 }  // End namespace Chaos::Softs
diff --git a/Engine/Source/Runtime/Experimental/Chaos/Private/Chaos/PerParticlePBDCollisionConstraint.cpp b/Engine/Source/Runtime/Experimental/Chaos/Private/Chaos/PerParticlePBDCollisionConstraint.cpp
--- a/Engine/Source/Runtime/Experimental/Chaos/Private/Chaos/PerParticlePBDCollisionConstraint.cpp
+++ b/Engine/Source/Runtime/Experimental/Chaos/Private/Chaos/PerParticlePBDCollisionConstraint.cpp
@@ -23,8 +23,9 @@ static_assert(ispc::ImplicitObjectType::Capsule == Chaos::ImplicitObjectType::Ca
 static_assert(ispc::ImplicitObjectType::Union == Chaos::ImplicitObjectType::Union);
 static_assert(ispc::ImplicitObjectType::TaperedCapsule == Chaos::ImplicitObjectType::TaperedCapsule);
 static_assert(ispc::ImplicitObjectType::Convex == Chaos::ImplicitObjectType::Convex);
-static_assert(ispc::ImplicitObjectType::WeightedLatticeBone == Chaos::ImplicitObjectType::WeightedLatticeBone);
 static_assert(ispc::ImplicitObjectType::IsWeightedLattice == Chaos::ImplicitObjectType::IsWeightedLattice);
+static_assert(ispc::ImplicitObjectType::MLLevelSet == Chaos::ImplicitObjectType::MLLevelSet);
+static_assert(ispc::ImplicitObjectType::SkinnedTriangleMesh == Chaos::ImplicitObjectType::SkinnedTriangleMesh);
 static_assert(ispc::ImplicitObjectType::WeightedLatticeLevelSetType == (Chaos::ImplicitObjectType::IsWeightedLattice | Chaos::ImplicitObjectType::LevelSet));
 static_assert(sizeof(ispc::TArray) == sizeof(TArray<int>));
 // Sphere
@@ -146,53 +147,6 @@ extern "C" void GetPhiWithNormal(const uint8* CollisionParticles, const FSolverR
 	}
 }
 
-// Helper function to call PhiWithNormal and return data to ISPC
-extern "C" void GetPhiWithNormalAndVelocityBone(const uint8 * CollisionParticles, const FSolverReal * InV, FSolverReal * Normal, FSolverReal * Phi, int32 * VelocityBone, const int32 i, const FSolverReal Thickness, const int32 ProgramCount, const int32 Mask)
-{
-	const FSolverCollisionParticles& C = *(const FSolverCollisionParticles*)CollisionParticles;
-
-	for (int32 Index = 0; Index < ProgramCount; ++Index)
-	{
-		if (Mask & (1 << Index))
-		{
-			FSolverVec3 V;
-
-			// aos_to_soa3
-			V.X = InV[Index];
-			V.Y = InV[Index + ProgramCount];
-			V.Z = InV[Index + 2 * ProgramCount];
-
-			VelocityBone[Index] = i;
-
-			FVec3 ImplicitNormal;
-			if (const TWeightedLatticeImplicitObject<FLevelSet>* LevelSet = C.GetGeometry(i)->GetObject< TWeightedLatticeImplicitObject<FLevelSet> >())
-			{
-				FWeightedLatticeImplicitObject::FEmbeddingCoordinate SurfaceCoord;
-				Phi[Index] = (FSolverReal)LevelSet->PhiWithNormalAndSurfacePoint(FVec3(V), ImplicitNormal, SurfaceCoord);
-				const FSolverReal Penetration = Thickness - Phi[Index]; // This is related to the Normal impulse
-				if (Penetration > (FSolverReal)0.)
-				{
-					const int32 StrongestBone = SurfaceCoord.GreatestInfluenceBone(LevelSet->GetBoneData());
-					if (StrongestBone != INDEX_NONE)
-					{
-						VelocityBone[Index] = LevelSet->GetSolverBoneIndices()[StrongestBone];
-					}
-				}
-			}
-			else
-			{
-				Phi[Index] = (FSolverReal)C.GetGeometry(i)->PhiWithNormal(FVec3(V), ImplicitNormal);
-			}
-			FSolverVec3 Norm(ImplicitNormal);
-
-			// aos_to_soa3
-			Normal[Index] = Norm.X;
-			Normal[Index + ProgramCount] = Norm.Y;
-			Normal[Index + 2 * ProgramCount] = Norm.Z;
-		}
-	}
-}
-
 void FPerParticlePBDCollisionConstraint::ApplyHelperISPC(FSolverParticles& InParticles, const FSolverReal Dt, const int32 Offset, const int32 Range) const
 {
 	check(bRealTypeCompatibleWithISPC);
