Chaos Cloth complex collider refactor and add support for SkinnedTriangleMesh colliders.

diff --git a/Engine/Source/Runtime/Experimental/Chaos/Private/Chaos/PerParticlePBDCollisionConstraint.ispc b/Engine/Source/Runtime/Experimental/Chaos/Private/Chaos/PerParticlePBDCollisionConstraint.ispc
--- a/Engine/Source/Runtime/Experimental/Chaos/Private/Chaos/PerParticlePBDCollisionConstraint.ispc
+++ b/Engine/Source/Runtime/Experimental/Chaos/Private/Chaos/PerParticlePBDCollisionConstraint.ispc
@@ -28,8 +28,6 @@ enum ImplicitObjectType
 	TaperedCapsule,
 	MLLevelSet,
 	SkinnedTriangleMesh,
-	WeightedLatticeBone,
-	MLLevelSetBone,
 	WeightedLatticeLevelSetType = 38, //== LevelSet | IsWeightedLattice (cannot use expressions like this in ISPC)
 
 
@@ -208,6 +206,24 @@ export uniform int32 ExportHack_PerParticlePBDCollisionConstraintStructs(
 	return 0;
 }
 
+static inline uniform bool IsComplexBatchCollider(const uniform uint8 CollisionType)
+{
+	return CollisionType == WeightedLatticeLevelSetType || CollisionType == MLLevelSet;
+}
+
+static inline uniform bool IsSimpleCollider(const uniform uint8 CollisionType)
+{
+	if(IsComplexBatchCollider(CollisionType))
+	{
+		return false;
+	}
+	if(CollisionType == SkinnedTriangleMesh)
+	{
+		return false;
+	}
+	return true;
+}
+
 static inline FReal SafeNormalize(FVector &Direction)
 {
 	const FReal SizeSqr = VectorSizeSquared(Direction);
@@ -586,11 +602,92 @@ static inline void PhiWithNormalConvex(const uniform uint8 *uniform TypedObjectP
 	}
 }
 
+static inline bool GetSimplePhiWithNormal(const FVector3f& PhiWithNormalInputFloat, const uniform uint8 Type, const uniform uint8 *uniform Ptr, const uniform uint8 *uniform TypedPtr, const uniform int SizeofFImplicitObject, const uniform int OffsetofGeometryType, const uniform int OffsetOfMargin, float& Phi, FVector3f& Normal)
+{
+	switch(Type)
+	{
+	case Sphere:
+		{
+			FVector NormalReal;
+			FReal PhiReal;
+			const FVector PhiWithNormalInput = ConvertVector3fTo3Native(PhiWithNormalInputFloat);
+
+			const uniform FReal Margin = *((const uniform FReal *uniform)(Ptr + OffsetOfMargin));
+			PhiWithNormalSphere(TypedPtr, PhiWithNormalInput, NormalReal, PhiReal, Margin);
+			Phi = (varying float)PhiReal;
+			Normal = ConvertVector3NativeTo3f(NormalReal);
+			return true;
+		}
+	case Capsule:
+		{
+			FVector NormalReal;
+			FReal PhiReal;
+			const FVector PhiWithNormalInput = ConvertVector3fTo3Native(PhiWithNormalInputFloat);
+
+			const uniform FReal Margin = *((const uniform FReal *uniform)(Ptr + OffsetOfMargin));
+			PhiWithNormalCapsule(TypedPtr, PhiWithNormalInput, NormalReal, PhiReal, Margin);
+			Phi = (varying float)PhiReal;
+			Normal = ConvertVector3NativeTo3f(NormalReal);
+			return true;
+		}
+	case Union:
+		{
+			FVector NormalReal;
+			FReal PhiReal;
+			const FVector PhiWithNormalInput = ConvertVector3fTo3Native(PhiWithNormalInputFloat);
+
+			PhiWithNormalUnion(TypedPtr, PhiWithNormalInput, NormalReal, PhiReal, SizeofFImplicitObject, OffsetofGeometryType, OffsetOfMargin);
+			Phi = (varying float)PhiReal;
+			Normal = ConvertVector3NativeTo3f(NormalReal);
+			return true;
+		}
+	case TaperedCapsule:
+		{
+			FVector NormalReal;
+			FReal PhiReal;
+			const FVector PhiWithNormalInput = ConvertVector3fTo3Native(PhiWithNormalInputFloat);
+			PhiWithNormalTaperedCapsule(TypedPtr, PhiWithNormalInput, NormalReal, PhiReal);
+			Phi = (varying float)PhiReal;
+			Normal = ConvertVector3NativeTo3f(NormalReal);
+			return true;
+		}
+	case Convex:
+		{
+			PhiWithNormalConvex(TypedPtr, PhiWithNormalInputFloat, Normal, Phi);
+			return true;
+		}
+	default:
+		return false;
+	}
+	return false;
+}
+
+template<typename ColliderVecType>
+static inline void ApplyFriction(FVector3f& P, const FVector3f& X, const FVector3f& NormalWorld, const float MaxFrictionCorrection, const uniform float Dt, const ColliderVecType& CollisionX, const ColliderVecType& CollisionV, const ColliderVecType& CollisionW, FVector3f& ColliderVelocityAtPoint)
+{
+	const FVector3f VectorToPoint = P - CollisionX;
+	ColliderVelocityAtPoint = CollisionV + VectorCross(CollisionW, VectorToPoint);
+
+	const FVector3f RelativeDisplacement = (P - X) - ColliderVelocityAtPoint * Dt; // This corresponds to the tangential velocity multiplied by dt (friction will drive this to zero if it is high enough)
+	const FVector3f RelativeDisplacementTangent = RelativeDisplacement - VectorDot(RelativeDisplacement, NormalWorld) * NormalWorld; // Project displacement into the tangential plane
+	const float RelativeDisplacementTangentLength = VectorSize(RelativeDisplacementTangent);
+	const float PositionCorrection = min(MaxFrictionCorrection, RelativeDisplacementTangentLength);
+	const float CorrectionRatio = select(RelativeDisplacementTangentLength < FLOAT_SMALL_NUMBER, 0.0f, PositionCorrection / RelativeDisplacementTangentLength);
+	P = P - (CorrectionRatio * RelativeDisplacementTangent);
+}
+
+template<typename ColliderVecType>
+static inline void ApplyFriction(FVector3f& P, const FVector3f& X, const FVector3f& NormalWorld, const float MaxFrictionCorrection, const uniform float Dt, const ColliderVecType& CollisionX, const ColliderVecType& CollisionV, const ColliderVecType& CollisionW)
+{
+	FVector3f ColliderVelocityAtPoint;
+	ApplyFriction(P, X, NormalWorld, MaxFrictionCorrection, Dt, CollisionX, CollisionV, CollisionW, ColliderVelocityAtPoint);
+}
+
 extern "C" unmasked void GetPhiWithNormal(const uniform uint8 *uniform CollisionParticles, const uniform float *uniform V, uniform float *uniform Normal, uniform float *uniform Phi, const uniform int i, const uniform int ProgramCount, const uniform int Mask);
-extern "C" unmasked void GetPhiWithNormalAndVelocityBone(const uniform uint8 *uniform CollisionParticles, const uniform float *uniform V, uniform float *uniform Normal, uniform float *uniform Phi, uniform int *uniform VelocityBone, const uniform int i, const uniform float Thickness, const uniform int ProgramCount, const uniform int Mask);
 extern "C" unmasked void GetPhiWithNormalCollisionParticleRange(const uniform uint8 *uniform CollisionParticlesRange, const uniform float *uniform V, uniform float *uniform Normal, uniform float *uniform Phi, const uniform int i, const uniform int ProgramCount, const uniform int Mask);
-extern "C" unmasked void GetPhiWithNormalAndVelocityBoneCollisionParticleRange(const uniform uint8 *uniform CollisionParticlesRange, const uniform float *uniform V, uniform float *uniform Normal, uniform float *uniform Phi, uniform int *uniform VelocityBone, const uniform int i, const uniform float Thickness, const uniform int ProgramCount, const uniform int Mask);
 
+
+// Legacy cloth with friction
 export void ApplyPerParticleCollisionFastFriction(uniform FVector4f ParticlesPandInvM[],
 													const uniform FVector3f ParticlesX[],
 													const uniform FVector3f CollisionV[],
@@ -654,9 +751,6 @@ export void ApplyPerParticleCollisionFastFriction(uniform FVector4f ParticlesPan
 			const FVector3f TranslatedVec = P - CollisionX[i];
 			const FVector4f InverseR = QuatInverse(CollisionR[i]);
 			const FVector3f PhiWithNormalInputFloat = VectorQuaternionRotateVector(InverseR, TranslatedVec);
-			const FVector PhiWithNormalInput = ConvertVector3fTo3Native(PhiWithNormalInputFloat);
-			FVector NormalReal;
-			FReal PhiReal;
 			FVector3f Normal;
 			float Phi;
 
@@ -665,85 +759,115 @@ export void ApplyPerParticleCollisionFastFriction(uniform FVector4f ParticlesPan
 			const uniform uint8 *uniform TypedPtr = Ptr + SizeofFImplicitObject;
 			const uniform uint8 Type = Ptr[OffsetofGeometryType] & 0x3F;
 
-			int32 VelocityBone = i;
-			switch(Type)
+			if(!GetSimplePhiWithNormal(PhiWithNormalInputFloat, Type, Ptr, TypedPtr, SizeofFImplicitObject, OffsetofGeometryType, OffsetOfMargin, Phi, Normal))
 			{
-			case Sphere:
-				{
-					const uniform FReal Margin = *((const uniform FReal *uniform)(Ptr + OffsetOfMargin));
-					PhiWithNormalSphere(TypedPtr, PhiWithNormalInput, NormalReal, PhiReal, Margin);
-					Phi = (varying float)PhiReal;
-					Normal = ConvertVector3NativeTo3f(NormalReal);
-					break;
-				}
-			case Capsule:
-				{
-					const uniform FReal Margin = *((const uniform FReal *uniform)(Ptr + OffsetOfMargin));
-					PhiWithNormalCapsule(TypedPtr, PhiWithNormalInput, NormalReal, PhiReal, Margin);
-					Phi = (varying float)PhiReal;
-					Normal = ConvertVector3NativeTo3f(NormalReal);
-					break;
-				}
-			case Union:
-				{
-					PhiWithNormalUnion(TypedPtr, PhiWithNormalInput, NormalReal, PhiReal, SizeofFImplicitObject, OffsetofGeometryType, OffsetOfMargin);
-					Phi = (varying float)PhiReal;
-					Normal = ConvertVector3NativeTo3f(NormalReal);
-					break;
-				}
-			case TaperedCapsule:
-				{
-					PhiWithNormalTaperedCapsule(TypedPtr, PhiWithNormalInput, NormalReal, PhiReal);
-					Phi = (varying float)PhiReal;
-					Normal = ConvertVector3NativeTo3f(NormalReal);
-					break;
-				}
-			case Convex:
-				{
-					PhiWithNormalConvex(TypedPtr, PhiWithNormalInputFloat, Normal, Phi);
-					break;
-				}
-			case WeightedLatticeBone:
+				uniform int IMMask;
+				if(IM != 0)
 				{
-					Phi = PerGroupThickness + 10;
-					break;
+					IMMask = lanemask();
 				}
-			case WeightedLatticeLevelSetType:
+				GetPhiWithNormal(CollisionParticles, (const uniform float *uniform)&PhiWithNormalInputFloat, (uniform float *uniform)&Normal, (uniform float *uniform)&Phi, i, programCount, IMMask);
+			}
+
+			const float Penetration = PerGroupThickness - Phi; // This is related to the Normal impulse
+			if(Penetration > 0.0f)
+			{
+				// TransformVector. Scale is 1 so that part can be removed.
+				const FVector3f NormalWorld = VectorQuaternionRotateVector(CollisionR[i], Normal);
+				P = P + (Penetration * NormalWorld);
+				
+				const FVector3f X = VectorLoad(&ParticlesX[extract(Index, 0)]);
+				ApplyFriction(P, X, NormalWorld, Penetration * PerGroupFriction, Dt, CollisionX[i], CollisionV[i], CollisionW[i]);
+			}
+		}
+
+		if(IM != 0)
+		{
+			VectorStore(&ParticlesPandInvM[extract(Index, 0)], SetVector4(P, IM));
+		}
+	}
+}
+
+// Legacy cloth no friction
+export void ApplyPerParticleCollisionNoFriction(uniform FVector4f ParticlesPandInvM[],
+												const uniform FVector3f ParticlesX[],
+												const uniform FVector3f CollisionV[],
+												const uniform FVector3f CollisionX[],
+												const uniform FVector3f CollisionW[],
+												const uniform FVector4f CollisionR[],
+												const uniform uint32 DynamicGroupId,
+												const uniform uint32 KinematicGroupIds[],
+												const uniform float PerGroupThickness,
+												const uniform uint8 *uniform CollisionParticles,
+												const uniform uint8 *uniform Geometry,
+												const uniform int SizeofFImplicitObject,
+												const uniform int OffsetofGeometryType,
+												const uniform int OffsetOfMargin,
+												const uniform float Dt,
+												const uniform int InnerOffset,
+												const uniform int InnerRange,
+												const uniform int OuterOffset,
+												const uniform int OuterRange)
+{
+	uniform int BailAt = OuterRange;
+	uniform bool bBail = false;
+
+	// Pre-check for bail out.
+	foreach(Index = OuterOffset ... BailAt)
+	{
+		for(uniform int i = InnerOffset; i < InnerRange; i++)
+		{
+			const uniform uint32 KinematicGroupId = KinematicGroupIds[i];  // Collision group Id
+
+			if (KinematicGroupId != (uint32)INDEX_NONE && DynamicGroupId != KinematicGroupId)
+			{
+				foreach_active(j)
 				{
-					uniform int IMMask;
-					if(IM != 0)
+					if(!bBail)
 					{
-						IMMask = lanemask();
+						BailAt = extract(Index, j);
+						bBail = true;
 					}
-					GetPhiWithNormalAndVelocityBone(CollisionParticles, (const uniform float *uniform)&PhiWithNormalInputFloat, (uniform float *uniform)&Normal, (uniform float *uniform)&Phi, (uniform int *uniform)&VelocityBone, i, PerGroupThickness, programCount, IMMask);
-					break;
-				}
-			case MLLevelSetBone:
-				{
-					Phi = PerGroupThickness + 10;
-					break;
 				}
-			case MLLevelSet:
+			}
+		}
+	}
+
+	foreach(Index = OuterOffset ... BailAt)
+	{
+		const FVector4f PandInvM = VectorLoad(&ParticlesPandInvM[extract(Index, 0)]);
+		varying FVector3f P;
+		varying float IM;
+		UnzipPandInvM(PandInvM, P, IM);
+
+		if(all(IM == 0))
+		{
+			continue;
+		}
+
+		for(uniform int i = InnerOffset; i < InnerRange; i++) // SequentialFor
+		{
+			// InverseTransformPosition. Can skip Scale calculation because it is 1 for xyz.
+			const FVector3f TranslatedVec = P - CollisionX[i];
+			const FVector4f InverseR = QuatInverse(CollisionR[i]);
+			const FVector3f PhiWithNormalInputFloat = VectorQuaternionRotateVector(InverseR, TranslatedVec);
+			FVector3f Normal;
+			float Phi;
+
+			const uniform SerializedPtr *uniform GeometryPtr = (const uniform SerializedPtr *uniform)Geometry;
+			const uniform uint8 *uniform Ptr = GeometryPtr[i].Ptr;
+			const uniform uint8 *uniform TypedPtr = Ptr + SizeofFImplicitObject;
+			const uniform uint8 Type = Ptr[OffsetofGeometryType] & 0x3F;
+			
+			if(!GetSimplePhiWithNormal(PhiWithNormalInputFloat, Type, Ptr, TypedPtr, SizeofFImplicitObject, OffsetofGeometryType, OffsetOfMargin, Phi, Normal))
+			{
+				uniform int IMMask;
+				if(IM != 0)
 				{
-					uniform int IMMask;
-					if(IM != 0)
-					{
-						IMMask = lanemask();
-					}
-					GetPhiWithNormalAndVelocityBone(CollisionParticles, (const uniform float *uniform)&PhiWithNormalInputFloat, (uniform float *uniform)&Normal, (uniform float *uniform)&Phi, (uniform int *uniform)&VelocityBone, i, PerGroupThickness, programCount, IMMask);
-					break;
+					IMMask = lanemask();
 				}
-			default:
-				{
-					uniform int IMMask;
-					if(IM != 0)
-					{
-						IMMask = lanemask();
-					}
 
-					GetPhiWithNormal(CollisionParticles, (const uniform float *uniform)&PhiWithNormalInputFloat, (uniform float *uniform)&Normal, (uniform float *uniform)&Phi, i, programCount, IMMask);
-					break;
-				}
+				GetPhiWithNormal(CollisionParticles, (const uniform float *uniform)&PhiWithNormalInputFloat, (uniform float *uniform)&Normal, (uniform float *uniform)&Phi, i, programCount, IMMask);
 			}
 
 			const float Penetration = PerGroupThickness - Phi; // This is related to the Normal impulse
@@ -752,39 +876,6 @@ export void ApplyPerParticleCollisionFastFriction(uniform FVector4f ParticlesPan
 				// TransformVector. Scale is 1 so that part can be removed.
 				const FVector3f NormalWorld = VectorQuaternionRotateVector(CollisionR[i], Normal);
 				P = P + (Penetration * NormalWorld);
-				if(Type == WeightedLatticeLevelSetType)
-				{
-					#pragma ignore warning(perf)
-					const FVector3f CollisionXVal = CollisionX[VelocityBone];
-					const FVector3f VectorToPoint = P - CollisionXVal;
-					const FVector3f X = VectorLoad(&ParticlesX[extract(Index, 0)]);
-					
-					#pragma ignore warning(perf)
-					const FVector3f CollisionVVal = CollisionV[VelocityBone];
-					#pragma ignore warning(perf)
-					const FVector3f CollisionWVal = CollisionW[VelocityBone];
-					const FVector3f RelativeDisplacement = (P - X) - (CollisionVVal + VectorCross(CollisionWVal, VectorToPoint)) * Dt; // This corresponds to the tangential velocity multiplied by dt (friction will drive this to zero if it is high enough)
-					const FVector3f RelativeDisplacementTangent = RelativeDisplacement - VectorDot(RelativeDisplacement, NormalWorld) * NormalWorld; // Project displacement into the tangential plane
-					const float RelativeDisplacementTangentLength = VectorSize(RelativeDisplacementTangent);
-
-					const float PositionCorrection = min(Penetration * PerGroupFriction, RelativeDisplacementTangentLength);
-					const float CorrectionRatio = select(RelativeDisplacementTangentLength < FLOAT_SMALL_NUMBER, 0.0f, PositionCorrection / RelativeDisplacementTangentLength);
-					P = P - (CorrectionRatio * RelativeDisplacementTangent);
-				
-				}
-				else
-				{
-					const FVector3f VectorToPoint = P - CollisionX[i];
-					const FVector3f X = VectorLoad(&ParticlesX[extract(Index, 0)]);
-
-					const FVector3f RelativeDisplacement = (P - X) - (CollisionV[i] + VectorCross(CollisionW[i], VectorToPoint)) * Dt; // This corresponds to the tangential velocity multiplied by dt (friction will drive this to zero if it is high enough)
-					const FVector3f RelativeDisplacementTangent = RelativeDisplacement - VectorDot(RelativeDisplacement, NormalWorld) * NormalWorld; // Project displacement into the tangential plane
-					const float RelativeDisplacementTangentLength = VectorSize(RelativeDisplacementTangent);
-
-					const float PositionCorrection = min(Penetration * PerGroupFriction, RelativeDisplacementTangentLength);
-					const float CorrectionRatio = select(RelativeDisplacementTangentLength < FLOAT_SMALL_NUMBER, 0.0f, PositionCorrection / RelativeDisplacementTangentLength);
-					P = P - (CorrectionRatio * RelativeDisplacementTangent);
-				}
 			}
 		}
 
@@ -795,7 +886,7 @@ export void ApplyPerParticleCollisionFastFriction(uniform FVector4f ParticlesPan
 	}
 }
 
-export void ApplyPerParticleCollisionFastFrictionNoGroupCheck(uniform FVector4f ParticlesPandInvM[],
+export void ApplyPerParticleSimpleCollisionFastFriction(uniform FVector4f ParticlesPandInvM[],
 													const uniform FVector3f ParticlesX[],
 													const uniform FVector3f CollisionV[],
 													const uniform FVector3f CollisionX[],
@@ -827,13 +918,12 @@ export void ApplyPerParticleCollisionFastFrictionNoGroupCheck(uniform FVector4f
 
 		for(uniform int i = 0; i < InnerCollisionCount; i++) // SequentialFor
 		{
-
 			const uniform SerializedPtr *uniform GeometryPtr = (const uniform SerializedPtr *uniform)Geometry;
 			const uniform uint8 *uniform Ptr = GeometryPtr[i].Ptr;
 			const uniform uint8 *uniform TypedPtr = Ptr + SizeofFImplicitObject;
 			const uniform uint8 Type = Ptr[OffsetofGeometryType] & 0x3F;
 
-			if(Type == WeightedLatticeBone || Type == WeightedLatticeLevelSetType || Type == MLLevelSet || Type == MLLevelSetBone)
+			if(!IsSimpleCollider(Type))
 			{
 				continue;
 			}
@@ -842,72 +932,18 @@ export void ApplyPerParticleCollisionFastFrictionNoGroupCheck(uniform FVector4f
 			const FVector3f TranslatedVec = P - CollisionX[i];
 			const uniform FVector4f InverseR = QuatInverse(CollisionR[i]);
 			const FVector3f PhiWithNormalInputFloat = VectorQuaternionRotateVector(InverseR, TranslatedVec);
-			const FVector PhiWithNormalInput = ConvertVector3fTo3Native(PhiWithNormalInputFloat);
-			FVector NormalReal;
-			FReal PhiReal;
 			FVector3f Normal;
 			float Phi;
-
-			switch(Type)
+			
+			if(!GetSimplePhiWithNormal(PhiWithNormalInputFloat, Type, Ptr, TypedPtr, SizeofFImplicitObject, OffsetofGeometryType, OffsetOfMargin, Phi, Normal))
 			{
-			case Sphere:
-				{
-					const uniform FReal Margin = *((const uniform FReal *uniform)(Ptr + OffsetOfMargin));
-					PhiWithNormalSphere(TypedPtr, PhiWithNormalInput, NormalReal, PhiReal, Margin);
-					Phi = (varying float)PhiReal;
-					Normal = ConvertVector3NativeTo3f(NormalReal);
-					break;
-				}
-			case Capsule:
-				{
-					const uniform FReal Margin = *((const uniform FReal *uniform)(Ptr + OffsetOfMargin));
-					PhiWithNormalCapsule(TypedPtr, PhiWithNormalInput, NormalReal, PhiReal, Margin);
-					Phi = (varying float)PhiReal;
-					Normal = ConvertVector3NativeTo3f(NormalReal);
-					break;
-				}
-			case Union:
-				{
-					PhiWithNormalUnion(TypedPtr, PhiWithNormalInput, NormalReal, PhiReal, SizeofFImplicitObject, OffsetofGeometryType, OffsetOfMargin);
-					Phi = (varying float)PhiReal;
-					Normal = ConvertVector3NativeTo3f(NormalReal);
-					break;
-				}
-			case TaperedCapsule:
-				{
-					PhiWithNormalTaperedCapsule(TypedPtr, PhiWithNormalInput, NormalReal, PhiReal);
-					Phi = (varying float)PhiReal;
-					Normal = ConvertVector3NativeTo3f(NormalReal);
-					break;
-				}
-			case Convex:
-				{
-					PhiWithNormalConvex(TypedPtr, PhiWithNormalInputFloat, Normal, Phi);
-					break;
-				}
-			case WeightedLatticeBone:
-			case WeightedLatticeLevelSetType: // Skip these here. They're done in a separate pass with batched calls to PhiAndNormal
-				{
-					Phi = PerGroupThickness + 10;
-					break;
-				}
-			case MLLevelSet:
-			case MLLevelSetBone: // Skip these here. They're done in a separate pass with batched calls to PhiAndNormal
+				uniform int IMMask;
+				if(IM != 0)
 				{
-					Phi = PerGroupThickness + 10;
-					break;
+					IMMask = lanemask();
 				}
-			default:
-				{
-					uniform int IMMask;
-					if(IM != 0)
-					{
-						IMMask = lanemask();
-					}
 
-					GetPhiWithNormalCollisionParticleRange(CollisionParticlesRange, (const uniform float *uniform)&PhiWithNormalInputFloat, (uniform float *uniform)&Normal, (uniform float *uniform)&Phi, i, programCount, IMMask);
-					break;
-				}
+				GetPhiWithNormalCollisionParticleRange(CollisionParticlesRange, (const uniform float *uniform)&PhiWithNormalInputFloat, (uniform float *uniform)&Normal, (uniform float *uniform)&Phi, i, programCount, IMMask);
 			}
 
 			const float Penetration = PerGroupThickness - Phi; // This is related to the Normal impulse
@@ -916,16 +952,9 @@ export void ApplyPerParticleCollisionFastFrictionNoGroupCheck(uniform FVector4f
 				// TransformVector. Scale is 1 so that part can be removed.
 				const FVector3f NormalWorld = VectorQuaternionRotateVector(CollisionR[i], Normal);
 				P = P + (Penetration * NormalWorld);
-				const FVector3f VectorToPoint = P - CollisionX[i];
-				const FVector3f X = VectorLoad(&ParticlesX[extract(Index, 0)]);
 
-				const FVector3f RelativeDisplacement = (P - X) - (CollisionV[i] + VectorCross(CollisionW[i], VectorToPoint)) * Dt; // This corresponds to the tangential velocity multiplied by dt (friction will drive this to zero if it is high enough)
-				const FVector3f RelativeDisplacementTangent = RelativeDisplacement - VectorDot(RelativeDisplacement, NormalWorld) * NormalWorld; // Project displacement into the tangential plane
-				const float RelativeDisplacementTangentLength = VectorSize(RelativeDisplacementTangent);
-
-				const float PositionCorrection = min(Penetration * PerGroupFriction, RelativeDisplacementTangentLength);
-				const float CorrectionRatio = select(RelativeDisplacementTangentLength < FLOAT_SMALL_NUMBER, 0.0f, PositionCorrection / RelativeDisplacementTangentLength);
-				P = P - (CorrectionRatio * RelativeDisplacementTangent);
+				const FVector3f X = VectorLoad(&ParticlesX[extract(Index, 0)]);
+				ApplyFriction(P, X, NormalWorld, Penetration * PerGroupFriction, Dt, CollisionX[i], CollisionV[i], CollisionW[i]);
 			}
 		}
 
@@ -936,18 +965,23 @@ export void ApplyPerParticleCollisionFastFrictionNoGroupCheck(uniform FVector4f
 	}
 }
 
-export void ApplyPerParticleCollisionFastFrictionBatchedPhiAndNormal(uniform FVector4f ParticlesPandInvM[],
+export void ApplyPerParticleSimpleCollisionFastFrictionAndGeneratePlanarConstraints(uniform FVector4f ParticlesPandInvM[],
+													uniform bool HasPlanarData[],
+													uniform FVector3f PlanarDataPositions[],
+													uniform FVector3f PlanarDataNormals[],
+													uniform FVector3f PlanarDataVelocities[],
 													const uniform FVector3f ParticlesX[],
-													const uniform int CollisionIndex,
-													const uniform float BatchPhis[],
-													const uniform FVector3f BatchNormals[],
-													const uniform int BatchVelocityBones[],
 													const uniform FVector3f CollisionV[],
 													const uniform FVector3f CollisionX[],
 													const uniform FVector3f CollisionW[],
-													const uniform FVector4f& CollisionR,
+													const uniform FVector4f CollisionR[],
 													const uniform float PerGroupFriction,
 													const uniform float PerGroupThickness,
+													const uniform uint8 *uniform CollisionParticlesRange,
+													const uniform uint8 *uniform Geometry,
+													const uniform int SizeofFImplicitObject,
+													const uniform int OffsetofGeometryType,
+													const uniform int OffsetOfMargin,
 													const uniform float Dt,
 													const uniform int InnerCollisionCount,
 													const uniform int OuterParticleBatchBegin,
@@ -960,86 +994,205 @@ export void ApplyPerParticleCollisionFastFrictionBatchedPhiAndNormal(uniform FVe
 		varying float IM;
 		UnzipPandInvM(PandInvM, P, IM);
 
-		const float Phi = BatchPhis[Index];
+		varying bool bHasPlanarData = false;
+		varying FVector3f PlanarP;
+		varying FVector3f PlanarN;
+		varying FVector3f PlanarV;
 
-		const float Penetration = PerGroupThickness - Phi; // This is related to the Normal impulse
-		if(IM != 0 && Penetration > 0.0f)
+		if(all(IM == 0))
 		{
-			const FVector3f Normal = VectorLoad(&BatchNormals[extract(Index,0)]);
+			continue;
+		}
 
-			// TransformVector. Scale is 1 so that part can be removed.
-			const FVector3f NormalWorld = VectorQuaternionRotateVector(CollisionR, Normal);
-			P = P + (Penetration * NormalWorld);
+		for(uniform int i = 0; i < InnerCollisionCount; i++) // SequentialFor
+		{
+			const uniform SerializedPtr *uniform GeometryPtr = (const uniform SerializedPtr *uniform)Geometry;
+			const uniform uint8 *uniform Ptr = GeometryPtr[i].Ptr;
+			const uniform uint8 *uniform TypedPtr = Ptr + SizeofFImplicitObject;
+			const uniform uint8 Type = Ptr[OffsetofGeometryType] & 0x3F;
 
-			const int VelocityBone = BatchVelocityBones[Index] >= 0 && BatchVelocityBones[Index] < InnerCollisionCount ? BatchVelocityBones[Index] : CollisionIndex;
+			if(!IsSimpleCollider(Type))
+			{
+				continue;
+			}
 
-			#pragma ignore warning(perf)
-			const FVector3f CollisionXVal = CollisionX[VelocityBone];
-			const FVector3f VectorToPoint = P - CollisionXVal;
-			const FVector3f X = VectorLoad(&ParticlesX[extract(Index, 0)]);
-					
-			#pragma ignore warning(perf)
-			const FVector3f CollisionVVal = CollisionV[VelocityBone];
-			#pragma ignore warning(perf)
-			const FVector3f CollisionWVal = CollisionW[VelocityBone];
-			const FVector3f RelativeDisplacement = (P - X) - (CollisionVVal + VectorCross(CollisionWVal, VectorToPoint)) * Dt; // This corresponds to the tangential velocity multiplied by dt (friction will drive this to zero if it is high enough)
-			const FVector3f RelativeDisplacementTangent = RelativeDisplacement - VectorDot(RelativeDisplacement, NormalWorld) * NormalWorld; // Project displacement into the tangential plane
-			const float RelativeDisplacementTangentLength = VectorSize(RelativeDisplacementTangent);
-
-			const float PositionCorrection = min(Penetration * PerGroupFriction, RelativeDisplacementTangentLength);
-			const float CorrectionRatio = select(RelativeDisplacementTangentLength < FLOAT_SMALL_NUMBER, 0.0f, PositionCorrection / RelativeDisplacementTangentLength);
-			P = P - (CorrectionRatio * RelativeDisplacementTangent);
+			// InverseTransformPosition. Can skip Scale calculation because it is 1 for xyz.
+			const FVector3f TranslatedVec = P - CollisionX[i];
+			const uniform FVector4f InverseR = QuatInverse(CollisionR[i]);
+			const FVector3f PhiWithNormalInputFloat = VectorQuaternionRotateVector(InverseR, TranslatedVec);
+			FVector3f Normal;
+			float Phi;
 			
+			if(!GetSimplePhiWithNormal(PhiWithNormalInputFloat, Type, Ptr, TypedPtr, SizeofFImplicitObject, OffsetofGeometryType, OffsetOfMargin, Phi, Normal))
+			{
+				uniform int IMMask;
+				if(IM != 0)
+				{
+					IMMask = lanemask();
+				}
+
+				GetPhiWithNormalCollisionParticleRange(CollisionParticlesRange, (const uniform float *uniform)&PhiWithNormalInputFloat, (uniform float *uniform)&Normal, (uniform float *uniform)&Phi, i, programCount, IMMask);
+			}
+
+			const float Penetration = PerGroupThickness - Phi; // This is related to the Normal impulse
+			if(Penetration > 0.0f)
+			{
+				// TransformVector. Scale is 1 so that part can be removed.
+				PlanarN = VectorQuaternionRotateVector(CollisionR[i], Normal);
+				P = P + (Penetration * PlanarN);
+
+				bHasPlanarData = true;
+				PlanarP = P;
+				const FVector3f X = VectorLoad(&ParticlesX[extract(Index, 0)]);
+				ApplyFriction(P, X, PlanarN, Penetration * PerGroupFriction, Dt, CollisionX[i], CollisionV[i], CollisionW[i], PlanarV);
+			}
+		}
+
+		if(IM != 0)
+		{
 			VectorStore(&ParticlesPandInvM[extract(Index, 0)], SetVector4(P, IM));
+
+			if(bHasPlanarData)
+			{
+				HasPlanarData[Index] = true;
+				VectorStore(&PlanarDataPositions[extract(Index,0)], PlanarP);
+				VectorStore(&PlanarDataNormals[extract(Index,0)], PlanarN);
+				VectorStore(&PlanarDataVelocities[extract(Index,0)], PlanarV);
+			}
 		}
 	}
 }
 
-export void ApplyPerParticleCollisionNoFriction(uniform FVector4f ParticlesPandInvM[],
-												const uniform FVector3f ParticlesX[],
-												const uniform FVector3f CollisionV[],
-												const uniform FVector3f CollisionX[],
-												const uniform FVector3f CollisionW[],
-												const uniform FVector4f CollisionR[],
-												const uniform uint32 DynamicGroupId,
-												const uniform uint32 KinematicGroupIds[],
-												const uniform float PerGroupThickness,
-												const uniform uint8 *uniform CollisionParticles,
-												const uniform uint8 *uniform Geometry,
-												const uniform int SizeofFImplicitObject,
-												const uniform int OffsetofGeometryType,
-												const uniform int OffsetOfMargin,
-												const uniform float Dt,
-												const uniform int InnerOffset,
-												const uniform int InnerRange,
-												const uniform int OuterOffset,
-												const uniform int OuterRange)
+export void ApplyPerParticleBatchCollisionFastFriction(uniform FVector4f ParticlesPandInvM[],
+													const uniform FVector3f ParticlesX[],
+													const uniform float BatchPhis[],
+													const uniform FVector3f BatchNormals[],
+													const uniform FVector3f& CollisionV,
+													const uniform FVector3f& CollisionX,
+													const uniform FVector3f& CollisionW,
+													const uniform FVector4f& CollisionR,
+													const uniform float PerGroupFriction,
+													const uniform float PerGroupThickness,
+													const uniform float Dt,
+													const uniform int OuterParticleBatchBegin,
+													const uniform int OuterParticleBatchEnd)
 {
-	uniform int BailAt = OuterRange;
-	uniform bool bBail = false;
+	foreach(Index = OuterParticleBatchBegin ... OuterParticleBatchEnd)
+	{
+		const FVector4f PandInvM = VectorLoad(&ParticlesPandInvM[extract(Index, 0)]);
+		varying FVector3f P;
+		varying float IM;
+		UnzipPandInvM(PandInvM, P, IM);
 
-	// Pre-check for bail out.
-	foreach(Index = OuterOffset ... BailAt)
+		if(all(IM == 0))
+		{
+			continue;
+		}
+
+		const float Phi = BatchPhis[Index];
+
+		const float Penetration = PerGroupThickness - Phi; // This is related to the Normal impulse
+		if(IM != 0 && Penetration > 0.0f)
+		{
+			const FVector3f Normal = VectorLoad(&BatchNormals[extract(Index,0)]);
+
+			// TransformVector. Scale is 1 so that part can be removed.
+			const FVector3f NormalWorld = VectorQuaternionRotateVector(CollisionR, Normal);
+			P = P + (Penetration * NormalWorld);
+				
+			const FVector3f X = VectorLoad(&ParticlesX[extract(Index, 0)]);
+			ApplyFriction(P, X, NormalWorld, Penetration * PerGroupFriction, Dt, CollisionX, CollisionV, CollisionW);
+			
+			VectorStore(&ParticlesPandInvM[extract(Index, 0)], SetVector4(P, IM));
+		}
+	}
+}
+
+export void ApplyPerParticleBatchCollisionFastFrictionWithVelocityBones(uniform FVector4f ParticlesPandInvM[],
+													const uniform FVector3f ParticlesX[],
+													const uniform float BatchPhis[],
+													const uniform FVector3f BatchNormals[],
+													const uniform int BatchVelocityBones[],
+													const uniform FVector3f& CollisionV,
+													const uniform FVector3f& CollisionX,
+													const uniform FVector3f& CollisionW,
+													const uniform FVector4f& CollisionR,
+													const uniform int MappedBoneIndices[],
+													const uniform int NumMappedBones,
+													const uniform FVector3f ColliderBoneV[],
+													const uniform FVector3f ColliderBoneX[],
+													const uniform FVector3f ColliderBoneW[],
+													const uniform float PerGroupFriction,
+													const uniform float PerGroupThickness,
+													const uniform float Dt,
+													const uniform int OuterParticleBatchBegin,
+													const uniform int OuterParticleBatchEnd)
+{
+	foreach(Index = OuterParticleBatchBegin ... OuterParticleBatchEnd)
 	{
-		for(uniform int i = InnerOffset; i < InnerRange; i++)
+		const FVector4f PandInvM = VectorLoad(&ParticlesPandInvM[extract(Index, 0)]);
+		varying FVector3f P;
+		varying float IM;
+		UnzipPandInvM(PandInvM, P, IM);
+
+		if(all(IM == 0))
 		{
-			const uniform uint32 KinematicGroupId = KinematicGroupIds[i];  // Collision group Id
+			continue;
+		}
 
-			if (KinematicGroupId != (uint32)INDEX_NONE && DynamicGroupId != KinematicGroupId)
+		const float Phi = BatchPhis[Index];
+
+		const float Penetration = PerGroupThickness - Phi; // This is related to the Normal impulse
+		if(IM != 0 && Penetration > 0.0f)
+		{
+			const FVector3f Normal = VectorLoad(&BatchNormals[extract(Index,0)]);
+
+			// TransformVector. Scale is 1 so that part can be removed.
+			const FVector3f NormalWorld = VectorQuaternionRotateVector(CollisionR, Normal);
+			P = P + (Penetration * NormalWorld);
+				
+			const FVector3f X = VectorLoad(&ParticlesX[extract(Index, 0)]);
+
+			const int32 VelocityBone = BatchVelocityBones[Index];
+			if(VelocityBone >= 0 && VelocityBone < NumMappedBones)
+			{				
+				#pragma ignore warning(perf)
+				const int32 MappedIndex = MappedBoneIndices[VelocityBone];
+				
+				const FVector3f ColliderV = VectorGather(&ColliderBoneV[MappedIndex]);
+				const FVector3f ColliderX = VectorGather(&ColliderBoneX[MappedIndex]);
+				const FVector3f ColliderW = VectorGather(&ColliderBoneW[MappedIndex]);
+				ApplyFriction(P, X, NormalWorld, Penetration * PerGroupFriction, Dt, ColliderX, ColliderV, ColliderW);
+			}
+			else
 			{
-				foreach_active(j)
-				{
-					if(!bBail)
-					{
-						BailAt = extract(Index, j);
-						bBail = true;
-					}
-				}
+				ApplyFriction(P, X, NormalWorld, Penetration * PerGroupFriction, Dt, CollisionX, CollisionV, CollisionW);
 			}
+			
+			VectorStore(&ParticlesPandInvM[extract(Index, 0)], SetVector4(P, IM));
 		}
 	}
+}
 
-	foreach(Index = OuterOffset ... BailAt)
+export void ApplyPerParticleBatchCollisionFastFrictionAndGeneratePlanarConstraints(uniform FVector4f ParticlesPandInvM[],
+													uniform bool HasPlanarData[],
+													uniform FVector3f PlanarDataPositions[],
+													uniform FVector3f PlanarDataNormals[],
+													uniform FVector3f PlanarDataVelocities[],
+													const uniform FVector3f ParticlesX[],
+													const uniform float BatchPhis[],
+													const uniform FVector3f BatchNormals[],
+													const uniform FVector3f& CollisionV,
+													const uniform FVector3f& CollisionX,
+													const uniform FVector3f& CollisionW,
+													const uniform FVector4f& CollisionR,
+													const uniform float PerGroupFriction,
+													const uniform float PerGroupThickness,
+													const uniform float Dt,
+													const uniform int OuterParticleBatchBegin,
+													const uniform int OuterParticleBatchEnd)
+{
+	foreach(Index = OuterParticleBatchBegin ... OuterParticleBatchEnd)
 	{
 		const FVector4f PandInvM = VectorLoad(&ParticlesPandInvM[extract(Index, 0)]);
 		varying FVector3f P;
@@ -1051,76 +1204,163 @@ export void ApplyPerParticleCollisionNoFriction(uniform FVector4f ParticlesPandI
 			continue;
 		}
 
-		for(uniform int i = InnerOffset; i < InnerRange; i++) // SequentialFor
+		const float Phi = BatchPhis[Index];
+
+		const float Penetration = PerGroupThickness - Phi; // This is related to the Normal impulse
+		if(IM != 0 && Penetration > 0.0f)
 		{
-			// InverseTransformPosition. Can skip Scale calculation because it is 1 for xyz.
-			const FVector3f TranslatedVec = P - CollisionX[i];
-			const FVector4f InverseR = QuatInverse(CollisionR[i]);
-			const FVector3f PhiWithNormalInputFloat = VectorQuaternionRotateVector(InverseR, TranslatedVec);
-			const FVector PhiWithNormalInput = ConvertVector3fTo3Native(PhiWithNormalInputFloat);
-			FVector NormalReal;
-			FReal PhiReal;
-			FVector3f Normal;
-			float Phi;
+			const FVector3f Normal = VectorLoad(&BatchNormals[extract(Index,0)]);
+
+			// TransformVector. Scale is 1 so that part can be removed.
+			const FVector3f NormalWorld = VectorQuaternionRotateVector(CollisionR, Normal);
+			P = P + (Penetration * NormalWorld);
+			
+			HasPlanarData[Index] = true;
+			VectorStore(&PlanarDataPositions[extract(Index,0)], P);
+			VectorStore(&PlanarDataNormals[extract(Index,0)], NormalWorld);
+
+			const FVector3f X = VectorLoad(&ParticlesX[extract(Index, 0)]);
+			FVector3f SurfaceVelocity;
+			ApplyFriction(P, X, NormalWorld, Penetration * PerGroupFriction, Dt, CollisionX, CollisionV, CollisionW, SurfaceVelocity);
+
+			VectorStore(&PlanarDataVelocities[extract(Index,0)], SurfaceVelocity);
+			
+			VectorStore(&ParticlesPandInvM[extract(Index, 0)], SetVector4(P, IM));
+		}
+	}
+}
+
+export void ApplyPerParticleBatchCollisionFastFrictionWithVelocityBonesAndGeneratePlanarConstraints(uniform FVector4f ParticlesPandInvM[],
+													uniform bool HasPlanarData[],
+													uniform FVector3f PlanarDataPositions[],
+													uniform FVector3f PlanarDataNormals[],
+													uniform FVector3f PlanarDataVelocities[],
+													const uniform FVector3f ParticlesX[],
+													const uniform float BatchPhis[],
+													const uniform FVector3f BatchNormals[],
+													const uniform int BatchVelocityBones[],
+													const uniform FVector3f& CollisionV,
+													const uniform FVector3f& CollisionX,
+													const uniform FVector3f& CollisionW,
+													const uniform FVector4f& CollisionR,
+													const uniform int MappedBoneIndices[],
+													const uniform int NumMappedBones,
+													const uniform FVector3f ColliderBoneV[],
+													const uniform FVector3f ColliderBoneX[],
+													const uniform FVector3f ColliderBoneW[],
+													const uniform float PerGroupFriction,
+													const uniform float PerGroupThickness,
+													const uniform float Dt,
+													const uniform int OuterParticleBatchBegin,
+													const uniform int OuterParticleBatchEnd)
+{
+	foreach(Index = OuterParticleBatchBegin ... OuterParticleBatchEnd)
+	{
+		const FVector4f PandInvM = VectorLoad(&ParticlesPandInvM[extract(Index, 0)]);
+		varying FVector3f P;
+		varying float IM;
+		UnzipPandInvM(PandInvM, P, IM);
+
+		if(all(IM == 0))
+		{
+			continue;
+		}
+
+		const float Phi = BatchPhis[Index];
 
+		const float Penetration = PerGroupThickness - Phi; // This is related to the Normal impulse
+		if(IM != 0 && Penetration > 0.0f)
+		{
+			const FVector3f Normal = VectorLoad(&BatchNormals[extract(Index,0)]);
+
+			// TransformVector. Scale is 1 so that part can be removed.
+			const FVector3f NormalWorld = VectorQuaternionRotateVector(CollisionR, Normal);
+			P = P + (Penetration * NormalWorld);
+			
+			HasPlanarData[Index] = true;
+			VectorStore(&PlanarDataPositions[extract(Index,0)], P);
+			VectorStore(&PlanarDataNormals[extract(Index,0)], NormalWorld);
+
+			const FVector3f X = VectorLoad(&ParticlesX[extract(Index, 0)]);
+			FVector3f SurfaceVelocity;
+			const int32 VelocityBone = BatchVelocityBones[Index];
+			if(VelocityBone >= 0 && VelocityBone < NumMappedBones)
+			{				
+				#pragma ignore warning(perf)
+				const int32 MappedIndex = MappedBoneIndices[VelocityBone];
+				
+				const FVector3f ColliderV = VectorGather(&ColliderBoneV[MappedIndex]);
+				const FVector3f ColliderX = VectorGather(&ColliderBoneX[MappedIndex]);
+				const FVector3f ColliderW = VectorGather(&ColliderBoneW[MappedIndex]);
+				ApplyFriction(P, X, NormalWorld, Penetration * PerGroupFriction, Dt, ColliderX, ColliderV, ColliderW, SurfaceVelocity);
+			}
+			else
+			{
+				ApplyFriction(P, X, NormalWorld, Penetration * PerGroupFriction, Dt, CollisionX, CollisionV, CollisionW, SurfaceVelocity);
+			}
+
+			VectorStore(&PlanarDataVelocities[extract(Index,0)], SurfaceVelocity);
+			
+			VectorStore(&ParticlesPandInvM[extract(Index, 0)], SetVector4(P, IM));
+		}
+	}
+}
+
+export void ApplyPerParticleSimpleCollisionNoFriction(uniform FVector4f ParticlesPandInvM[],
+												const uniform FVector3f ParticlesX[],
+												const uniform FVector3f CollisionX[],
+												const uniform FVector4f CollisionR[],
+												const uniform float PerGroupThickness,
+												const uniform uint8 *uniform CollisionParticlesRange,
+												const uniform uint8 *uniform Geometry,
+												const uniform int SizeofFImplicitObject,
+												const uniform int OffsetofGeometryType,
+												const uniform int OffsetOfMargin,
+												const uniform int InnerCollisionCount,
+												const uniform int OuterParticleBatchBegin,
+												const uniform int OuterParticleBatchEnd)
+{
+
+	foreach(Index = OuterParticleBatchBegin ... OuterParticleBatchEnd)
+	{
+		const FVector4f PandInvM = VectorLoad(&ParticlesPandInvM[extract(Index, 0)]);
+		varying FVector3f P;
+		varying float IM;
+		UnzipPandInvM(PandInvM, P, IM);
+
+		if(all(IM == 0))
+		{
+			continue;
+		}
+
+		for(uniform int i = 0; i < InnerCollisionCount; i++) // SequentialFor
+		{
 			const uniform SerializedPtr *uniform GeometryPtr = (const uniform SerializedPtr *uniform)Geometry;
 			const uniform uint8 *uniform Ptr = GeometryPtr[i].Ptr;
 			const uniform uint8 *uniform TypedPtr = Ptr + SizeofFImplicitObject;
 			const uniform uint8 Type = Ptr[OffsetofGeometryType] & 0x3F;
 
-			switch(Type)
+			if(!IsSimpleCollider(Type))
 			{
-			case Sphere:
-				{
-					const uniform FReal Margin = *((const uniform FReal *uniform)(Ptr + OffsetOfMargin));
-					PhiWithNormalSphere(TypedPtr, PhiWithNormalInput, NormalReal, PhiReal, Margin);
-					Phi = (varying float)PhiReal;
-					Normal = ConvertVector3NativeTo3f(NormalReal);
-					break;
-				}
-			case Capsule:
-				{
-					const uniform FReal Margin = *((const uniform FReal *uniform)(Ptr + OffsetOfMargin));
-					PhiWithNormalCapsule(TypedPtr, PhiWithNormalInput, NormalReal, PhiReal, Margin);
-					Phi = (varying float)PhiReal;
-					Normal = ConvertVector3NativeTo3f(NormalReal);
-					break;
-				}
-			case Union:
-				{
-					PhiWithNormalUnion(TypedPtr, PhiWithNormalInput, NormalReal, PhiReal, SizeofFImplicitObject, OffsetofGeometryType, OffsetOfMargin);
-					Phi = (varying float)PhiReal;
-					Normal = ConvertVector3NativeTo3f(NormalReal);
-					break;
-				}
-			case TaperedCapsule:
-				{
-					PhiWithNormalTaperedCapsule(TypedPtr, PhiWithNormalInput, NormalReal, PhiReal);
-					Phi = (varying float)PhiReal;
-					Normal = ConvertVector3NativeTo3f(NormalReal);
-					break;
-				}
-			case Convex:
-				{
-					PhiWithNormalConvex(TypedPtr, PhiWithNormalInputFloat, Normal, Phi);
-					break;
-				}
-			case WeightedLatticeBone:
+				continue;
+			}
+
+			// InverseTransformPosition. Can skip Scale calculation because it is 1 for xyz.
+			const FVector3f TranslatedVec = P - CollisionX[i];
+			const uniform FVector4f InverseR = QuatInverse(CollisionR[i]);
+			const FVector3f PhiWithNormalInputFloat = VectorQuaternionRotateVector(InverseR, TranslatedVec);
+			FVector3f Normal;
+			float Phi;
+			
+			if(!GetSimplePhiWithNormal(PhiWithNormalInputFloat, Type, Ptr, TypedPtr, SizeofFImplicitObject, OffsetofGeometryType, OffsetOfMargin, Phi, Normal))
+			{
+				uniform int IMMask;
+				if(IM != 0)
 				{
-					Phi = PerGroupThickness + 10;
-					break;
+					IMMask = lanemask();
 				}
-			default:
-				{
-					uniform int IMMask;
-					if(IM != 0)
-					{
-						IMMask = lanemask();
-					}
 
-					GetPhiWithNormal(CollisionParticles, (const uniform float *uniform)&PhiWithNormalInputFloat, (uniform float *uniform)&Normal, (uniform float *uniform)&Phi, i, programCount, IMMask);
-					break;
-				}
+				GetPhiWithNormalCollisionParticleRange(CollisionParticlesRange, (const uniform float *uniform)&PhiWithNormalInputFloat, (uniform float *uniform)&Normal, (uniform float *uniform)&Phi, i, programCount, IMMask);
 			}
 
 			const float Penetration = PerGroupThickness - Phi; // This is related to the Normal impulse
@@ -1139,11 +1379,12 @@ export void ApplyPerParticleCollisionNoFriction(uniform FVector4f ParticlesPandI
 	}
 }
 
-export void ApplyPerParticleCollisionNoFrictionNoGroupCheck(uniform FVector4f ParticlesPandInvM[],
+export void ApplyPerParticleSimpleCollisionNoFrictionAndGeneratePlanarConstraints(uniform FVector4f ParticlesPandInvM[],
+												uniform bool HasPlanarData[],
+												uniform FVector3f PlanarDataPositions[],
+												uniform FVector3f PlanarDataNormals[],
 												const uniform FVector3f ParticlesX[],
-												const uniform FVector3f CollisionV[],
 												const uniform FVector3f CollisionX[],
-												const uniform FVector3f CollisionW[],
 												const uniform FVector4f CollisionR[],
 												const uniform float PerGroupThickness,
 												const uniform uint8 *uniform CollisionParticlesRange,
@@ -1151,7 +1392,6 @@ export void ApplyPerParticleCollisionNoFrictionNoGroupCheck(uniform FVector4f Pa
 												const uniform int SizeofFImplicitObject,
 												const uniform int OffsetofGeometryType,
 												const uniform int OffsetOfMargin,
-												const uniform float Dt,
 												const uniform int InnerCollisionCount,
 												const uniform int OuterParticleBatchBegin,
 												const uniform int OuterParticleBatchEnd)
@@ -1164,6 +1404,10 @@ export void ApplyPerParticleCollisionNoFrictionNoGroupCheck(uniform FVector4f Pa
 		varying float IM;
 		UnzipPandInvM(PandInvM, P, IM);
 
+		varying bool bHasPlanarData = false;
+		varying FVector3f PlanarP;
+		varying FVector3f PlanarN;
+
 		if(all(IM == 0))
 		{
 			continue;
@@ -1176,110 +1420,104 @@ export void ApplyPerParticleCollisionNoFrictionNoGroupCheck(uniform FVector4f Pa
 			const uniform uint8 *uniform TypedPtr = Ptr + SizeofFImplicitObject;
 			const uniform uint8 Type = Ptr[OffsetofGeometryType] & 0x3F;
 
-			if(Type == WeightedLatticeBone || Type == WeightedLatticeLevelSetType || Type == MLLevelSet || Type == MLLevelSetBone)
+			if(!IsSimpleCollider(Type))
 			{
 				continue;
 			}
 
 			// InverseTransformPosition. Can skip Scale calculation because it is 1 for xyz.
 			const FVector3f TranslatedVec = P - CollisionX[i];
-			const FVector4f InverseR = QuatInverse(CollisionR[i]);
+			const uniform FVector4f InverseR = QuatInverse(CollisionR[i]);
 			const FVector3f PhiWithNormalInputFloat = VectorQuaternionRotateVector(InverseR, TranslatedVec);
-			const FVector PhiWithNormalInput = ConvertVector3fTo3Native(PhiWithNormalInputFloat);
-			FVector NormalReal;
-			FReal PhiReal;
 			FVector3f Normal;
 			float Phi;
-
-
-			switch(Type)
+			
+			if(!GetSimplePhiWithNormal(PhiWithNormalInputFloat, Type, Ptr, TypedPtr, SizeofFImplicitObject, OffsetofGeometryType, OffsetOfMargin, Phi, Normal))
 			{
-			case Sphere:
+				uniform int IMMask;
+				if(IM != 0)
 				{
-					const uniform FReal Margin = *((const uniform FReal *uniform)(Ptr + OffsetOfMargin));
-					PhiWithNormalSphere(TypedPtr, PhiWithNormalInput, NormalReal, PhiReal, Margin);
-					Phi = (varying float)PhiReal;
-					Normal = ConvertVector3NativeTo3f(NormalReal);
-					break;
+					IMMask = lanemask();
 				}
-			case Capsule:
-				{
-					const uniform FReal Margin = *((const uniform FReal *uniform)(Ptr + OffsetOfMargin));
-					PhiWithNormalCapsule(TypedPtr, PhiWithNormalInput, NormalReal, PhiReal, Margin);
-					Phi = (varying float)PhiReal;
-					Normal = ConvertVector3NativeTo3f(NormalReal);
-					break;
-				}
-			case Union:
-				{
-					PhiWithNormalUnion(TypedPtr, PhiWithNormalInput, NormalReal, PhiReal, SizeofFImplicitObject, OffsetofGeometryType, OffsetOfMargin);
-					Phi = (varying float)PhiReal;
-					Normal = ConvertVector3NativeTo3f(NormalReal);
-					break;
-				}
-			case TaperedCapsule:
-				{
-					PhiWithNormalTaperedCapsule(TypedPtr, PhiWithNormalInput, NormalReal, PhiReal);
-					Phi = (varying float)PhiReal;
-					Normal = ConvertVector3NativeTo3f(NormalReal);
-					break;
-				}
-			case Convex:
-				{
-					PhiWithNormalConvex(TypedPtr, PhiWithNormalInputFloat, Normal, Phi);
-					break;
-				}
-			case WeightedLatticeBone:
-			case WeightedLatticeLevelSetType: // Skip these here. They're done in a separate pass with batched calls to PhiAndNormal
-				{
-					Phi = PerGroupThickness + 10;
-					break;
-				}
-			case MLLevelSet:
-			case MLLevelSetBone: // Skip these here. They're done in a separate pass with batched calls to PhiAndNormal
-				{
-					Phi = PerGroupThickness + 10;
-					break;
-				}
-			default:
-				{
-					uniform int IMMask;
-					if(IM != 0)
-					{
-						IMMask = lanemask();
-					}
 
-					GetPhiWithNormalCollisionParticleRange(CollisionParticlesRange, (const uniform float *uniform)&PhiWithNormalInputFloat, (uniform float *uniform)&Normal, (uniform float *uniform)&Phi, i, programCount, IMMask);
-					break;
-				}
+				GetPhiWithNormalCollisionParticleRange(CollisionParticlesRange, (const uniform float *uniform)&PhiWithNormalInputFloat, (uniform float *uniform)&Normal, (uniform float *uniform)&Phi, i, programCount, IMMask);
 			}
 
 			const float Penetration = PerGroupThickness - Phi; // This is related to the Normal impulse
 			if(Penetration > 0.0f)
 			{
 				// TransformVector. Scale is 1 so that part can be removed.
-				const FVector3f NormalWorld = VectorQuaternionRotateVector(CollisionR[i], Normal);
-				P = P + (Penetration * NormalWorld);
+				PlanarN = VectorQuaternionRotateVector(CollisionR[i], Normal);
+				P = P + (Penetration * PlanarN);
+
+				bHasPlanarData = true;
+				PlanarP = P;
 			}
 		}
 
 		if(IM != 0)
 		{
 			VectorStore(&ParticlesPandInvM[extract(Index, 0)], SetVector4(P, IM));
+
+			if(bHasPlanarData)
+			{
+				HasPlanarData[Index] = true;
+				VectorStore(&PlanarDataPositions[extract(Index,0)], PlanarP);
+				VectorStore(&PlanarDataNormals[extract(Index,0)], PlanarN);
+			}
 		}
 	}
 }
 
-export void ApplyPerParticleCollisionNoFrictionBatchedPhiAndNormal(uniform FVector4f ParticlesPandInvM[],
-												const uniform float BatchPhis[],
-												const uniform FVector3f BatchNormals[],
-												const uniform FVector4f& CollisionR,
-												const uniform float PerGroupThickness,
-												const uniform float Dt,
-												const uniform int OuterParticleBatchBegin,
-												const uniform int OuterParticleBatchEnd)
+export void ApplyPerParticleBatchCollisionNoFriction(uniform FVector4f ParticlesPandInvM[],
+													const uniform FVector3f ParticlesX[],
+													const uniform float BatchPhis[],
+													const uniform FVector3f BatchNormals[],
+													const uniform FVector4f& CollisionR,
+													const uniform float PerGroupThickness,
+													const uniform int OuterParticleBatchBegin,
+													const uniform int OuterParticleBatchEnd)
 {
+	foreach(Index = OuterParticleBatchBegin ... OuterParticleBatchEnd)
+	{
+		const FVector4f PandInvM = VectorLoad(&ParticlesPandInvM[extract(Index, 0)]);
+		varying FVector3f P;
+		varying float IM;
+		UnzipPandInvM(PandInvM, P, IM);
+
+		if(all(IM == 0))
+		{
+			continue;
+		}
+
+		const float Phi = BatchPhis[Index];
+
+		const float Penetration = PerGroupThickness - Phi; // This is related to the Normal impulse
+		if(IM != 0 && Penetration > 0.0f)
+		{
+			const FVector3f Normal = VectorLoad(&BatchNormals[extract(Index,0)]);
+
+			// TransformVector. Scale is 1 so that part can be removed.
+			const FVector3f NormalWorld = VectorQuaternionRotateVector(CollisionR, Normal);
+			P = P + (Penetration * NormalWorld);
+			
+			VectorStore(&ParticlesPandInvM[extract(Index, 0)], SetVector4(P, IM));
+		}
+	}
+}
 
+export void ApplyPerParticleBatchCollisionNoFrictionAndGeneratePlanarConstraints(uniform FVector4f ParticlesPandInvM[],
+													uniform bool HasPlanarData[],
+													uniform FVector3f PlanarDataPositions[],
+													uniform FVector3f PlanarDataNormals[],
+													const uniform FVector3f ParticlesX[],
+													const uniform float BatchPhis[],
+													const uniform FVector3f BatchNormals[],
+													const uniform FVector4f& CollisionR,
+													const uniform float PerGroupThickness,
+													const uniform int OuterParticleBatchBegin,
+													const uniform int OuterParticleBatchEnd)
+{
 	foreach(Index = OuterParticleBatchBegin ... OuterParticleBatchEnd)
 	{
 		const FVector4f PandInvM = VectorLoad(&ParticlesPandInvM[extract(Index, 0)]);
@@ -1287,6 +1525,11 @@ export void ApplyPerParticleCollisionNoFrictionBatchedPhiAndNormal(uniform FVect
 		varying float IM;
 		UnzipPandInvM(PandInvM, P, IM);
 
+		if(all(IM == 0))
+		{
+			continue;
+		}
+
 		const float Phi = BatchPhis[Index];
 
 		const float Penetration = PerGroupThickness - Phi; // This is related to the Normal impulse
@@ -1297,6 +1540,11 @@ export void ApplyPerParticleCollisionNoFrictionBatchedPhiAndNormal(uniform FVect
 			// TransformVector. Scale is 1 so that part can be removed.
 			const FVector3f NormalWorld = VectorQuaternionRotateVector(CollisionR, Normal);
 			P = P + (Penetration * NormalWorld);
+			
+			HasPlanarData[Index] = true;
+			VectorStore(&PlanarDataPositions[extract(Index,0)], P);
+			VectorStore(&PlanarDataNormals[extract(Index,0)], NormalWorld);
+			
 			VectorStore(&ParticlesPandInvM[extract(Index, 0)], SetVector4(P, IM));
 		}
 	}
diff --git a/Engine/Source/Runtime/Experimental/Chaos/Private/Chaos/SoftsEvolution.cpp b/Engine/Source/Runtime/Experimental/Chaos/Private/Chaos/SoftsEvolution.cpp
--- a/Engine/Source/Runtime/Experimental/Chaos/Private/Chaos/SoftsEvolution.cpp
+++ b/Engine/Source/Runtime/Experimental/Chaos/Private/Chaos/SoftsEvolution.cpp
@@ -6,12 +6,6 @@
 
 namespace Chaos::Softs {
 
-bool bPostTimestepCollisionUpdate = false;
-#if !UE_BUILD_SHIPPING
-static FAutoConsoleVariableRef CVarClothbPostTimestepCollisionUpdate(TEXT("p.Chaos.Cloth.PostTimestepCollisionUpdate"), bPostTimestepCollisionUpdate, TEXT("Do collision update at the end of the timestep instead of the beginning [def: false]"));
-
-#endif
-
 namespace Private
 {
 namespace EvolutionSolverDefault
@@ -225,7 +219,7 @@ int32 FEvolution::AddSoftBody(uint32 GroupId, int32 NumParticles, bool bEnable)
 	else
 	{
 		SoftBodyId = SoftBodies.AddSoftBody();
-		SoftBodies.ParticleRanges[SoftBodyId] = FSolverParticlesRange::AddParticleRange(Particles, NumParticles);
+		SoftBodies.ParticleRanges[SoftBodyId] = FSolverParticlesRange::AddParticleRange(Particles, NumParticles, SoftBodyId);
 	}
 
 	Groups.SoftBodies[GroupId].Emplace(SoftBodyId);
@@ -328,7 +322,7 @@ int32 FEvolution::AddCollisionParticleRange(uint32 GroupId, int32 NumParticles,
 	else
 	{
 		CollisionRangeId = CollisionRanges.AddRange();
-		CollisionRanges.ParticleRanges[CollisionRangeId] = FSolverCollisionParticlesRange::AddParticleRange(CollisionParticles, NumParticles);
+		CollisionRanges.ParticleRanges[CollisionRangeId] = FSolverCollisionParticlesRange::AddParticleRange(CollisionParticles, NumParticles, CollisionRangeId);
 	}
 
 	// Set properties
@@ -420,7 +414,20 @@ void FEvolution::AdvanceOneTimeStepInternal(const FSolverReal Dt, const int32 Ti
 {
 	TRACE_CPUPROFILER_EVENT_SCOPE(FSoftsEvolution_AdvanceOneTimeStepInternal);
 
+	FEvolutionGroupContext& Context = Groups.SolverContexts[GroupId];
+	Context.Reset();
+
 	TArray<FSolverCollisionParticlesRange> ActiveCollisionRanges = GetActiveCollisionParticles(GroupId);
+	{
+		// Collision kinematic update
+		TRACE_CPUPROFILER_EVENT_SCOPE(FSoftsEvolution_AdvanceOneTimeStepInternal_CollisionKinematicUpdate);
+		check(CollisionKinematicUpdate);
+		for (FSolverCollisionParticlesRange& ActiveRange : ActiveCollisionRanges)
+		{
+			CollisionKinematicUpdate(ActiveRange, Dt, Time);
+		}
+	}
+
 	// Right now, typically groups only have one active softbody in them, so don't parallelize softbodies
 	for (const int32 SoftBodyId : Groups.ActiveSoftBodies[GroupId])
 	{
@@ -436,6 +443,8 @@ void FEvolution::AdvanceOneTimeStepInternal(const FSolverReal Dt, const int32 Ti
 			(bDoNewtonUpdate ? ESolverMode::ForceBased : ESolverMode::None);
 		check(SolverMode != ESolverMode::None);
 
+		Context.Init(SolverMode, Dt, NumPBDIterations, bDoNewtonUpdate ? NumNewtonIterations : 0);
+
 		// PreSubstepParallelInits + LocalDamping
 		const bool bDoLocalDamping = SoftBodies.LocalDampings[SoftBodyId] > (FSolverReal)0.;
 		FPerParticleDampVelocity DampVelocityRule(SoftBodies.LocalDampings[SoftBodyId]);
@@ -466,15 +475,6 @@ void FEvolution::AdvanceOneTimeStepInternal(const FSolverReal Dt, const int32 Ti
 			TRACE_CPUPROFILER_EVENT_SCOPE(FSoftsEvolution_AdvanceOneTimeStepInternal_KinematicUpdates);
 			KinematicUpdate(SoftBodies.ParticleRanges[SoftBodyId], Dt, Time);
 		}
-		{
-			// Collision kinematic update
-			TRACE_CPUPROFILER_EVENT_SCOPE(FSoftsEvolution_AdvanceOneTimeStepInternal_CollisionKinematicUpdate);
-			check(CollisionKinematicUpdate);
-			for (FSolverCollisionParticlesRange& ActiveRange : ActiveCollisionRanges)
-			{
-				CollisionKinematicUpdate(ActiveRange, Dt, Time);
-			}
-		}
 
 		// PBD initial guess
 		if (NumPBDIterations > 0)
@@ -515,19 +515,6 @@ void FEvolution::AdvanceOneTimeStepInternal(const FSolverReal Dt, const int32 Ti
 			}
 		}
 
-		{
-			if (bDoCollisionPerTimestep && !bPostTimestepCollisionUpdate)		
-			{	
-				// Presubstep Collision Rules
-				TRACE_CPUPROFILER_EVENT_SCOPE(FSoftsEvolution_AdvanceOneTimeStepInternal_PreSubstepCollisionRules);
-				TConstArrayView<PBDCollisionConstraintRuleFunc> PerIterationCollisionConstraintRules(SoftBodies.PerIterationCollisionPBDConstraintRules[SoftBodyId]);
-				for (const PBDCollisionConstraintRuleFunc& ConstraintRule : PerIterationCollisionConstraintRules)
-				{
-					ConstraintRule(SoftBodies.ParticleRanges[SoftBodyId], Dt, ActiveCollisionRanges);
-				}		
-			}
-		}
-
 		{
 			// Parallel Inits
 			TRACE_CPUPROFILER_EVENT_SCOPE(FSoftsEvolution_AdvanceOneTimeStepInternal_PostInitialGuessParallelInits);
@@ -552,21 +539,19 @@ void FEvolution::AdvanceOneTimeStepInternal(const FSolverReal Dt, const int32 Ti
 			{
 				// Iteration loop
 				TRACE_CPUPROFILER_EVENT_SCOPE(FSoftsEvolution_AdvanceOneTimeStepInternal_IterationLoop);
-				for (int32 i = 0; i < NumPBDIterations; ++i)
+				for (Context.CurrentPBDIteration = 0; Context.CurrentPBDIteration < NumPBDIterations; ++Context.CurrentPBDIteration)
 				{
+
 					TConstArrayView<PBDConstraintRuleFunc> PerIterationConstraintRules(SoftBodies.PerIterationPBDConstraintRules[SoftBodyId]);
 					for (const PBDConstraintRuleFunc& ConstraintRule : PerIterationConstraintRules)
 					{
 						ConstraintRule(SoftBodies.ParticleRanges[SoftBodyId], Dt);
 					}
 
-					if (!bDoCollisionPerTimestep)
+					TConstArrayView<PBDCollisionConstraintRuleFunc> PerIterationCollisionConstraintRules(SoftBodies.PerIterationCollisionPBDConstraintRules[SoftBodyId]);
+					for (const PBDCollisionConstraintRuleFunc& ConstraintRule : PerIterationCollisionConstraintRules)
 					{
-						TConstArrayView<PBDCollisionConstraintRuleFunc> PerIterationCollisionConstraintRules(SoftBodies.PerIterationCollisionPBDConstraintRules[SoftBodyId]);
-						for (const PBDCollisionConstraintRuleFunc& ConstraintRule : PerIterationCollisionConstraintRules)
-						{
-							ConstraintRule(SoftBodies.ParticleRanges[SoftBodyId], Dt, ActiveCollisionRanges);
-						}
+						ConstraintRule(SoftBodies.ParticleRanges[SoftBodyId], Dt, ActiveCollisionRanges);
 					}
 
 					TConstArrayView<PBDConstraintRuleFunc> PerIterationPostCollisionsConstraintRules(SoftBodies.PerIterationPostCollisionsPBDConstraintRules[SoftBodyId]);
@@ -593,18 +578,6 @@ void FEvolution::AdvanceOneTimeStepInternal(const FSolverReal Dt, const int32 Ti
 		{
 			Private::CopyVtoVPrev(SoftBodies.ParticleRanges[SoftBodyId]);
 		}
-
-		{
-			if (bDoCollisionPerTimestep && bPostTimestepCollisionUpdate)		
-			{	
-				TRACE_CPUPROFILER_EVENT_SCOPE(FSoftsEvolution_AdvanceOneTimeStepInternal_PostSubstepCollisionRules);
-				TConstArrayView<PBDCollisionConstraintRuleFunc> PerIterationCollisionConstraintRules(SoftBodies.PerIterationCollisionPBDConstraintRules[SoftBodyId]);
-				for (const PBDCollisionConstraintRuleFunc& ConstraintRule : PerIterationCollisionConstraintRules)
-				{
-					ConstraintRule(SoftBodies.ParticleRanges[SoftBodyId], Dt, ActiveCollisionRanges);
-				}		
-			}
-		}
 		
 		// Linear system update... currently we don't have any cross-softbody constraints. 
 		// If we did, the linear system would need to live at the group level.
@@ -614,10 +587,10 @@ void FEvolution::AdvanceOneTimeStepInternal(const FSolverReal Dt, const int32 Ti
 			check(LinearSystemRules.Num());
 			TConstArrayView<UpdateLinearSystemCollisionsFunc> LinearSystemCollisionRules(SoftBodies.UpdateLinearSystemCollisionsRules[SoftBodyId]);
 			check(NumNewtonIterations > 0);
-			for (int32 NewtonIteration = 0; NewtonIteration < NumNewtonIterations; ++NewtonIteration)
+			for (Context.CurrentNewtonIteration = 0; Context.CurrentNewtonIteration < NumNewtonIterations; ++Context.CurrentNewtonIteration)
 			{
 				TRACE_CPUPROFILER_EVENT_SCOPE(FSoftsEvolution_AdvanceOneTimeStepInternal_NewtonIteration);
-				SoftBodies.LinearSystems[SoftBodyId].Init(SoftBodies.ParticleRanges[SoftBodyId], Dt, NewtonIteration == 0, LinearSystemParameters);
+				SoftBodies.LinearSystems[SoftBodyId].Init(SoftBodies.ParticleRanges[SoftBodyId], Dt, Context.CurrentNewtonIteration == 0, LinearSystemParameters);
 
 				for (const UpdateLinearSystemFunc& LinearSystemRule : LinearSystemRules)
 				{
diff --git a/Engine/Source/Runtime/Experimental/Chaos/Private/Chaos/WeightedLatticeImplicitObject.cpp b/Engine/Source/Runtime/Experimental/Chaos/Private/Chaos/WeightedLatticeImplicitObject.cpp
--- a/Engine/Source/Runtime/Experimental/Chaos/Private/Chaos/WeightedLatticeImplicitObject.cpp
+++ b/Engine/Source/Runtime/Experimental/Chaos/Private/Chaos/WeightedLatticeImplicitObject.cpp
@@ -199,7 +199,9 @@ FWeightedLatticeImplicitObject::FWeightedLatticeImplicitObject(const FWeightedLa
 	, Grid(Other.Grid)
 	, UsedBones(Other.UsedBones)
 	, ReferenceRelativeTransforms(Other.ReferenceRelativeTransforms)
+PRAGMA_DISABLE_DEPRECATION_WARNINGS
 	, SolverBoneIndices(Other.SolverBoneIndices)
+PRAGMA_ENABLE_DEPRECATION_WARNINGS
 	, LocalBoundingBox(Other.LocalBoundingBox)
 {
 	BoneData.Copy(Other.BoneData);
@@ -740,8 +742,7 @@ void TWeightedLatticeImplicitObject<TConcrete>::BatchPhiWithNormalAndGreatestInf
 			Phis[Index] = (Softs::FSolverReal)PhiWithNormalAndSurfacePoint(RigidSpacePosition, Normal, SurfaceCoord, bIncludeEmptyCellsFalse);
 			Normals[Index] = Softs::FSolverVec3(Normal);
 			const Softs::FSolverReal Penetration = Thickness - Phis[Index]; // This is related to the Normal impulse
-			const int32 StrongestBone = Penetration > (Softs::FSolverReal)0. ? SurfaceCoord.GreatestInfluenceBone(GetBoneData()) : INDEX_NONE;
-			GreatestInfluenceSolverBones[Index] = SolverBoneIndices.IsValidIndex(StrongestBone) ? SolverBoneIndices[StrongestBone] : INDEX_NONE;
+			GreatestInfluenceSolverBones[Index] = Penetration > (Softs::FSolverReal)0. ? SurfaceCoord.GreatestInfluenceBone(GetBoneData()) : INDEX_NONE;
 		}
 	});
 }
diff --git a/Engine/Source/Runtime/Experimental/Chaos/Public/Chaos/ImplicitObjectType.h b/Engine/Source/Runtime/Experimental/Chaos/Public/Chaos/ImplicitObjectType.h
--- a/Engine/Source/Runtime/Experimental/Chaos/Public/Chaos/ImplicitObjectType.h
+++ b/Engine/Source/Runtime/Experimental/Chaos/Public/Chaos/ImplicitObjectType.h
@@ -29,8 +29,6 @@ namespace Chaos
 			TaperedCapsule,
 			MLLevelSet,
 			SkinnedTriangleMesh,
-			WeightedLatticeBone,
-			MLLevelSetBone,
 			//Add entries above this line for serialization
 			ConcreteObjectCount, // Used to ensure bitflags do not overlap concrete type
 			IsWeightedLattice = 1 << 5,
@@ -98,8 +96,6 @@ namespace Chaos
 				TEXT("EImplicitObjectType:TaperedCapsule"),
 				TEXT("EImplicitObjectType:MLLevelSet"),
 				TEXT("EImplicitObjectType:SkinnedTriangleMesh"),
-				TEXT("EImplicitObjectType:WeightedLatticeBone"),
-				TEXT("EImplicitObjectType:MLLevelSetBone"),
 			};
 
 		static_assert(Chaos::ImplicitObjectType::ConcreteObjectCount == UE_ARRAY_COUNT(Names), "");
diff --git a/Engine/Source/Runtime/Experimental/Chaos/Public/Chaos/MLLevelset.h b/Engine/Source/Runtime/Experimental/Chaos/Public/Chaos/MLLevelset.h
--- a/Engine/Source/Runtime/Experimental/Chaos/Public/Chaos/MLLevelset.h
+++ b/Engine/Source/Runtime/Experimental/Chaos/Public/Chaos/MLLevelset.h
@@ -12,30 +12,6 @@ struct FKMLLevelSetElem;
 
 namespace Chaos
 {
-// Not actually a collision object. CollisionParticles that are used to update bones used by the FMLLevelSet have this as their geometry
-// It is used to get relative transform between consecutive bones (active joint vs parent of active joint)
-// Note that not necessarily parent  = parent of active joint
-class FMLLevelSetBoneProxy : public FImplicitObject
-{
-public:
-	using FImplicitObject::GetTypeName;
-
-	FMLLevelSetBoneProxy()
-		: FImplicitObject(EImplicitObject::DisableCollisions, ImplicitObjectType::MLLevelSetBone)
-	{}
-
-	static constexpr EImplicitObjectType StaticType() { return ImplicitObjectType::MLLevelSetBone; }
-
-	virtual FReal PhiWithNormal(const FVec3& x, FVec3& Normal) const override
-	{
-		return UE_BIG_NUMBER;
-	}
-
-	virtual uint32 GetTypeHash() const override
-	{
-		return 0;
-	}
-};
 
 struct FMLLevelSetNNEModelData
 {
@@ -74,19 +50,13 @@ class FMLLevelSet final : public FImplicitObject
 	CHAOS_API void BatchPhiWithNormal(const TConstArrayView<Softs::FPAndInvM> PAndInvMArray,
 		const Softs::FSolverRigidTransform3& SolverToThis, TArray<Softs::FSolverReal>& OutBatchPhis, TArray<Softs::FSolverVec3>& OutBatchNormals,
 		const Chaos::Softs::FSolverReal CollisionThickness, const int32 MLLevelsetThread, const int32 BatchBegin, const int32 BatchEnd) const;	
-	CHAOS_API void BatchPhiWithNormalWithGreatestInfluenceBone(const TConstArrayView<Softs::FPAndInvM> PAndInvMArray,
-		const Softs::FSolverRigidTransform3& SolverToThis, TArray<Softs::FSolverReal>& OutBatchPhis, TArray<Softs::FSolverVec3>& OutBatchNormals,
-		TArray<int32>& OutBatchGreatestInfluenceSolverBones, const Chaos::Softs::FSolverReal CollisionThickness, int32 MLLevelsetThread, int32 BatchBegin, int32 BatchEnd) const;
 
 	CHAOS_API FReal SignedDistance(const FVec3& x) const;
 	CHAOS_API void UpdateActiveBonesRelativeTransforms(TArray<FTransform>& InActiveBonesRelativeTransforms);
 	CHAOS_API void UpdateActiveBonesRelativeTransformsAndUpdateDebugPhi(TArray<FTransform>& InActiveBonesRelativeTransforms);
 	CHAOS_API void UpdateNeuralInferencesNumber(const int32 InNeuralInferencesNumber);
 
-	const TArray<int32>& GetSolverActiveBoneIndices() const { return SolverActiveBoneIndices; }
-	void SetSolverActiveBoneIndices(TArray<int32>&& InSolverActiveBoneIndices) { SolverActiveBoneIndices = MoveTemp(InSolverActiveBoneIndices); }
 	const int32 GetNumberOfActiveBones() const { return ActiveBoneNames.Num(); }
-	void SetSolverParentBoneIndex(int32 InSolverParentBoneIndex) { SolverParentBoneIndex = InSolverParentBoneIndex; }
 
 	const TArray<FName>& GetActiveBoneNames() const { return ActiveBoneNames; }
 	virtual const FAABB3 BoundingBox() const override { return LocalBoundingBox; } 
@@ -187,16 +157,12 @@ class FMLLevelSet final : public FImplicitObject
 	void BuildNNEModel(TObjectPtr<UNNEModelData> InNNEModelData, TSharedPtr<UE::NNE::IModelCPU>& NNEModel);
 	const bool IsFTransformArraysDifferent(TArray<FTransform>& FTransformArr1, TArray<FTransform>& FTransformArr2, FReal Tol = UE_KINDA_SMALL_NUMBER) const;
 
-	// We assume the Greatest Influence Bone is always the parent bone. 
-	const int32 GetGreatestInfluenceBoneSolverIndex() const { return SolverParentBoneIndex; }
 	// We assume the closest active bone joint is always the first active bone. 
 	const int32 GetClosestActiveBoneIndex(const FVec3 ParticlePositionMS) const { return 0; }  
 
   private:
 	//Relative Bone Transform Related Data
-	TArray<FName> ActiveBoneNames; 
-	TArray<int32> SolverActiveBoneIndices; //indices set by the solver for bone transform query
-	int32 SolverParentBoneIndex;
+	TArray<FName> ActiveBoneNames;
 	TArray<FTransform> ActiveBonesRelativeTransforms;
 
 	// NNE ML Model Elements
diff --git a/Engine/Source/Runtime/Experimental/Chaos/Public/Chaos/PBDCollisionSpringConstraints.h b/Engine/Source/Runtime/Experimental/Chaos/Public/Chaos/PBDCollisionSpringConstraints.h
--- a/Engine/Source/Runtime/Experimental/Chaos/Public/Chaos/PBDCollisionSpringConstraints.h
+++ b/Engine/Source/Runtime/Experimental/Chaos/Public/Chaos/PBDCollisionSpringConstraints.h
@@ -128,9 +128,6 @@ public:
 			{
 				ThicknessWeighted.SetWeightedValue(WeightedValue);
 			}
-PRAGMA_DISABLE_DEPRECATION_WARNINGS
-			Thickness = GetMaxThickness();
-PRAGMA_ENABLE_DEPRECATION_WARNINGS
 		}
 		if (IsSelfCollisionStiffnessMutable(PropertyCollection))
 		{
@@ -158,11 +155,11 @@ PRAGMA_ENABLE_DEPRECATION_WARNINGS
 		}
 		if (IsSelfCollisionKinematicColliderThicknessMutable(PropertyCollection))
 		{
-			KinematicColliderThickness = FMath::Max(GetSelfCollisionKinematicColliderThickness(PropertyCollection), 0.f);
+			KinematicCollisions.SetColliderThickness(FMath::Max(GetSelfCollisionKinematicColliderThickness(PropertyCollection), 0.f));
 		}
 		if (IsSelfCollisionKinematicColliderStiffnessMutable(PropertyCollection))
 		{
-			KinematicColliderStiffness = FMath::Clamp((FSolverReal)GetSelfCollisionKinematicColliderStiffness(PropertyCollection), 0.f, 1.f);
+			KinematicCollisions.SetStiffness(FMath::Clamp((FSolverReal)GetSelfCollisionKinematicColliderStiffness(PropertyCollection), 0.f, 1.f));
 		}
 		if (IsSelfCollisionKinematicColliderFrictionMutable(PropertyCollection))
 		{
diff --git a/Engine/Source/Runtime/Experimental/Chaos/Public/Chaos/PBDCollisionSpringConstraintsBase.h b/Engine/Source/Runtime/Experimental/Chaos/Public/Chaos/PBDCollisionSpringConstraintsBase.h
--- a/Engine/Source/Runtime/Experimental/Chaos/Public/Chaos/PBDCollisionSpringConstraintsBase.h
+++ b/Engine/Source/Runtime/Experimental/Chaos/Public/Chaos/PBDCollisionSpringConstraintsBase.h
@@ -6,6 +6,7 @@
 #include "Chaos/PBDFlatWeightMap.h"
 #include "Chaos/PBDSoftsEvolutionFwd.h"
 #include "Chaos/PBDSoftsSolverParticles.h"
+#include "Chaos/PBDKinematicTriangleMeshCollisions.h"
 #include "Chaos/PBDTriangleMeshCollisions.h"
 #include "Containers/Set.h"
 
@@ -48,23 +49,12 @@ public:
 		const FSolverVec2 InKinematicColliderFrictionCoefficient = FSolverVec2(DefaultKinematicColliderFrictionCoefficient),
 		const FSolverReal InProximityStiffness = DefaultProximityStiffness);
 
-	UE_DEPRECATED(5.4, "Use constructor with SelfCollisionLayers")
-	FPBDCollisionSpringConstraintsBase(
-		const int32 InOffset,
-		const int32 InNumParticles,
-		const FTriangleMesh& InTriangleMesh,
-		const TArray<FSolverVec3>* InReferencePositions,
-		TSet<TVec2<int32>>&& InDisabledCollisionElements,
-		const FSolverReal InThickness = BackCompatThickness,
-		const FSolverReal InStiffness = BackCompatStiffness,
-		const FSolverReal InFrictionCoefficient = BackCompatFrictionCoefficient)
-		: FPBDCollisionSpringConstraintsBase(InOffset, InNumParticles, InTriangleMesh, InReferencePositions, MoveTemp(InDisabledCollisionElements),
-			TConstArrayView<FRealSingle>(), TConstArrayView<FRealSingle>(), TConstArrayView<int32>(), FSolverVec2(InThickness), InStiffness, InFrictionCoefficient)
-	{}
-
-	PRAGMA_DISABLE_DEPRECATION_WARNINGS
 	virtual ~FPBDCollisionSpringConstraintsBase() {}
-	PRAGMA_ENABLE_DEPRECATION_WARNINGS
+
+	void ResetKinematicCollider()
+	{
+		KinematicCollisions.Reset();
+	}
 
 	template<typename SpatialAccelerator, typename SolverParticlesOrRange>
 	UE_DEPRECATED(5.4, "Use Init with CollidableSubMesh. This method is much less efficient as it recreates the CollidableSubMesh each call.")
@@ -79,15 +69,11 @@ public:
 
 	const TArray<TVec4<int32>>& GetConstraints() const 
 	{
-		PRAGMA_DISABLE_DEPRECATION_WARNINGS
 		return Constraints;  
-		PRAGMA_ENABLE_DEPRECATION_WARNINGS
 	}
 	const TArray<FSolverVec3>& GetBarys() const 
 	{
-		PRAGMA_DISABLE_DEPRECATION_WARNINGS
 		return Barys;
-		PRAGMA_ENABLE_DEPRECATION_WARNINGS
 	}
 
 	UE_DEPRECATED(5.4, "Thickness is now weighted.")
@@ -100,27 +86,19 @@ public:
 	bool GetGlobalIntersectionAnalysis() const { return bGlobalIntersectionAnalysis; }
 	const TArray<bool>& GetFlipNormals() const 
 	{
-		PRAGMA_DISABLE_DEPRECATION_WARNINGS
 		return FlipNormal;
-		PRAGMA_ENABLE_DEPRECATION_WARNINGS
 	}
 
-	const TArray<int32>& GetKinematicCollidingParticles() const { return KinematicCollidingParticles; }
-	const TArray<TMap<int32, FSolverReal>>& GetKinematicColliderTimers() const { return KinematicColliderTimers; }
+	const TArray<int32>& GetKinematicCollidingParticles() const { return KinematicCollisions.GetCollidingParticles(); }
+	const TArray<TMap<int32, FSolverReal>>& GetKinematicColliderTimers() const { return KinematicCollisions.GetTimers(); }
 	const FTriangleMesh& GetTriangleMesh() const { return TriangleMesh; }
 
-	UE_DEPRECATED(5.4, "Thickness is now weighted.")
-	void SetThickness(FSolverReal InThickness) 
-	{ 
-		SetThicknessWeighted(FSolverVec2(InThickness));
-	}
 	void SetThicknessWeighted(const FSolverVec2 InThickness) { ThicknessWeighted.SetWeightedValue(FSolverVec2::Max(InThickness, FSolverVec2(0.f))); }
 	void SetFrictionCoefficient(FSolverReal InFrictionCoefficient) { FrictionCoefficient = InFrictionCoefficient; }
 
 	template<typename SolverParticlesOrRange>
 	void Apply(SolverParticlesOrRange& Particles, const FSolverReal Dt, const int32 ConstraintIndex) const
 	{
-		PRAGMA_DISABLE_DEPRECATION_WARNINGS
 		const TVector<int32, 4>& Constraint = Constraints[ConstraintIndex];
 		const int32 Index1 = Constraint[0];
 		const int32 Index2 = Constraint[1];
@@ -143,7 +121,6 @@ public:
 		{
 			Particles.P(Index4) -= Particles.InvM(Index4) * Barys[ConstraintIndex][2] * Delta;
 		}
-		PRAGMA_ENABLE_DEPRECATION_WARNINGS
 	}
 
 	template<typename SolverParticlesOrRange>
@@ -197,32 +174,29 @@ protected:
 	FSolverReal Stiffness; // (0-1 compliance for PBD)
 	FSolverReal FrictionCoefficient; 
 	bool bOnlyCollideKinematics;
-	FSolverReal KinematicColliderThickness;
-	FSolverReal KinematicColliderStiffness;
+
 	FPBDFlatWeightMap KinematicColliderFrictionCoefficient;
-	FSolverReal ProximityStiffness; // (actual spring stiffness for force-based solver)
+	FPBDKinematicTriangleMeshCollisions KinematicCollisions;
 
-	UE_DEPRECATED(5.4, "Use ThicknessWeighted instead")
-	FSolverReal Thickness;
+	UE_DEPRECATED(5.6, "Use KinematicCollisions.Thickness instead")
+	FSolverReal KinematicColliderThickness;
+	UE_DEPRECATED(5.6, "Use KinematicCollisions.Stiffness instead")
+	FSolverReal KinematicColliderStiffness;
 
-	UE_DEPRECATED(5.4, "Constraints will be made private")
-	TArray<TVec4<int32>> Constraints;
-	UE_DEPRECATED(5.4, "Barys will be made private")
-	TArray<FSolverVec3> Barys;
-	UE_DEPRECATED(5.4, "FlipNormal will be made private")
-	TArray<bool> FlipNormal;
+	FSolverReal ProximityStiffness; // (actual spring stiffness for force-based solver)
 
 	CHAOS_API void UpdateCollisionLayers(const TConstArrayView<int32>& InFaceCollisionLayers);
 
 	int32 GetNumParticles() const { return NumParticles; }
 
 private:
-	template<typename SolverParticlesOrRange>
-	void ApplyDynamicConstraints(SolverParticlesOrRange& InParticles, const FSolverReal Dt) const;
 
-	template<typename SolverParticlesOrRange>
-	void ApplyKinematicConstraints(SolverParticlesOrRange& InParticles, const FSolverReal Dt) const;
+	TArray<TVec4<int32>> Constraints;
+	TArray<FSolverVec3> Barys;
+	TArray<bool> FlipNormal;
 
+	template<typename SolverParticlesOrRange>
+	void ApplyDynamicConstraints(SolverParticlesOrRange& InParticles, const FSolverReal Dt) const;
 
 	const FTriangleMesh& TriangleMesh;
 	const TArray<FSolverVec3>* ReferencePositions;
@@ -237,12 +211,6 @@ private:
 	};
 	TArray<EConstraintType> ConstraintTypes;
 
-	static constexpr int32 MaxKinematicConnectionsPerPoint = 3;
-	// Parallel arrays (for ISPC SOA)
-	TArray<int32> KinematicCollidingParticles;
-	TArray<TVector<int32, MaxKinematicConnectionsPerPoint>> KinematicColliderElements;
-
-	TArray<TMap<int32, FSolverReal>> KinematicColliderTimers; // Keep constraints for a cvar-defined time after it's moved out of proximity. Helps reduce jitter.
 
 	int32 Offset;
 	int32 NumParticles;
@@ -250,15 +218,13 @@ private:
 };
 }  // End namespace Chaos::Softs
 
+#if UE_ENABLE_INCLUDE_ORDER_DEPRECATED_IN_5_6
 #if !defined(CHAOS_COLLISION_SPRING_ISPC_ENABLED_DEFAULT)
 #define CHAOS_COLLISION_SPRING_ISPC_ENABLED_DEFAULT 1
 #endif
-
-// Support ISPC enable/disable in non-shipping builds
-#if !INTEL_ISPC || UE_BUILD_SHIPPING
-static constexpr bool bChaos_CollisionSpring_ISPC_Enabled = INTEL_ISPC && CHAOS_COLLISION_SPRING_ISPC_ENABLED_DEFAULT;
-#else
-extern CHAOS_API bool bChaos_CollisionSpring_ISPC_Enabled;
 #endif
 
+UE_DEPRECATED(5.6, "This variable has been renamed bChaos_KinematicTriangleMesh_ISPC_Enabled")
+static constexpr bool bChaos_CollisionSpring_ISPC_Enabled = false;
+
 #endif
diff --git a/Engine/Source/Runtime/Experimental/Chaos/Public/Chaos/PBDSoftBodyCollisionConstraint.h b/Engine/Source/Runtime/Experimental/Chaos/Public/Chaos/PBDSoftBodyCollisionConstraint.h
--- a/Engine/Source/Runtime/Experimental/Chaos/Public/Chaos/PBDSoftBodyCollisionConstraint.h
+++ b/Engine/Source/Runtime/Experimental/Chaos/Public/Chaos/PBDSoftBodyCollisionConstraint.h
@@ -2,8 +2,10 @@
 #pragma once
 
 #include "Chaos/Core.h"
+#include "Chaos/PBDKinematicTriangleMeshCollisions.h"
 #include "Chaos/PBDSoftsEvolutionFwd.h"
 #include "Chaos/PBDSoftsSolverParticles.h"
+#include "Chaos/PBDPlanarConstraints.h"
 #include "Chaos/SoftsSolverCollisionParticles.h"
 #include "Chaos/SoftsSolverCollisionParticlesRange.h"
 #include "Chaos/SoftsSolverParticlesRange.h"
@@ -22,10 +24,23 @@ extern CHAOS_API bool bChaos_SoftBodyCollision_ISPC_Enabled;
 
 namespace Chaos::Softs
 {
+	struct FPBDComplexColliderBoneData
+	{
+		// Use this array to go from BoneData indices to the X, R, V, etc. arrays below
+		TConstArrayView<int32> MappedBoneIndices;
+
+		TConstArrayView<FSolverVec3> X;
+		TConstArrayView<FSolverVec3> V;
+		TConstArrayView<FSolverRotation3> R;
+		TConstArrayView<FSolverVec3> W;
+	};
 
 class FPBDSoftBodyCollisionConstraintBase
 {
 public:
+	static constexpr bool bDefaultUsePlanarConstraintForSimpleColliders = false;
+	static constexpr bool bDefaultUsePlanarConstraintForComplexColliders = true;
+
 	FPBDSoftBodyCollisionConstraintBase(
 		const TArray<FSolverRigidTransform3>& InLastCollisionTransforms,
 		FSolverReal InCollisionThickness,
@@ -35,48 +50,82 @@ public:
 		TArray<bool>* InCollisionParticleCollided = nullptr,
 		TArray<FSolverVec3>* InContacts = nullptr,
 		TArray<FSolverVec3>* InNormals = nullptr,
-		TArray<FSolverReal>* InPhis = nullptr
+		TArray<FSolverReal>* InPhis = nullptr,
+		const FSolverReal InSoftBodyCollisionThickness = 0.f,
+		bool bInEnableSimpleColliders = true,
+		bool bInEnableComplexColliders = true,
+		bool bInUsePlanarConstraintForSimpleColliders = bDefaultUsePlanarConstraintForSimpleColliders,
+		bool bInUsePlanarConstraintForComplexColliders = bDefaultUsePlanarConstraintForComplexColliders,
+		const TMap<FParticleRangeIndex, FPBDComplexColliderBoneData>& InComplexBoneData = TMap<FParticleRangeIndex, FPBDComplexColliderBoneData>()
 		)
 		: LastCollisionTransforms(InLastCollisionTransforms)
 		, CollisionThickness(InCollisionThickness)
+		, SoftBodyCollisionThickness(InSoftBodyCollisionThickness)
 		, FrictionCoefficient(InFrictionCoefficient)
 		, bUseCCD(bInUseCCD)
 		, ProximityStiffness(InProximityStiffness)
+		, bEnableSimpleColliders(bInEnableSimpleColliders)
+		, bEnableComplexColliders(bInEnableComplexColliders)
+		, bUsePlanarConstraintForSimpleColliders(bInUsePlanarConstraintForSimpleColliders)
+		, bUsePlanarConstraintForComplexColliders(bInUsePlanarConstraintForComplexColliders)
+		, ComplexBoneData(InComplexBoneData)
 		, CollisionParticleCollided(InCollisionParticleCollided)
 		, Contacts(InContacts)
 		, Normals(InNormals)
 		, Phis(InPhis)
+		, PlanarConstraint(FrictionCoefficient)
 	{}
 
 	void SetWriteDebugContacts(bool bWrite) { bWriteDebugContacts = bWrite; }
 
-	// This method uses mutable data that is not thread-safe.
-	CHAOS_API void Apply(FSolverParticlesRange& Particles, const FSolverReal Dt, const TArray<FSolverCollisionParticlesRange>& CollisionParticles) const;
+	// Follow planar constraint settings and either generate or apply planar constraints if requested.
+	CHAOS_API void ApplyWithPlanarConstraints(FSolverParticlesRange& Particles, const FSolverReal Dt, const TArray<FSolverCollisionParticlesRange>& CollisionParticles, bool bGeneratePlanarConstraints);
+
+	// Ignore planar constraint settings and do a full collision detection and apply impulses
+	CHAOS_API void Apply(FSolverParticlesRange& Particles, const FSolverReal Dt, const TArray<FSolverCollisionParticlesRange>& CollisionParticles);
 	
 	CHAOS_API void UpdateLinearSystem(const FSolverParticlesRange& Particles, const FSolverReal Dt, const TArray<FSolverCollisionParticlesRange>& CollisionParticles, FEvolutionLinearSystem& LinearSystem) const;
 
-	CHAOS_API const bool DoCollisionPerTimestep() const { return bDoCollisionPerTimestep; }
-	CHAOS_API void SetDoCollisionPerTimestep(bool bInDoCollisionPerTimestep) { bDoCollisionPerTimestep = bInDoCollisionPerTimestep; }
-	CHAOS_API TArray<int32>& GetCollisionIndices() { return CollisionIndices; }
-	CHAOS_API TArray<FSolverVec3>& GetCollisionTargets() { return CollisionTargets; }
-	CHAOS_API TArray<FSolverVec3>& GetCollisionNormalsOutput() { return CollisionNormalsOutput; }
+	void OnCollisionRangeRemoved(int32 CollisionRangeId)
+	{
+		for (TMap<FParticleRangeIndex, FPBDComplexColliderBoneData>::TIterator Iter = ComplexBoneData.CreateIterator(); Iter; ++Iter)
+		{
+			if (Iter.Key().RangeId == CollisionRangeId)
+			{
+				Iter.RemoveCurrent();
+			}
+		}
+	}
 
 private:
-	template<bool bLockAndWriteContacts, bool bWithFriction>
-	void ApplyInternal(FSolverParticlesRange& Particles, const FSolverReal Dt, const TArray<FSolverCollisionParticlesRange>& CollisionParticles) const;
+
+	template<bool bLockAndWriteContacts, bool bWithFriction, bool bGeneratePlanarConstraints>
+	void ApplySimpleInternal(FSolverParticlesRange& Particles, const FSolverReal Dt, const TArray<FSolverCollisionParticlesRange>& CollisionParticles);
+	template<bool bLockAndWriteContacts, bool bWithFriction, bool bGeneratePlanarConstraints>
+	void ApplyComplexInternal(FSolverParticlesRange& Particles, const FSolverReal Dt, const TArray<FSolverCollisionParticlesRange>& CollisionParticles);
 	template<bool bLockAndWriteContacts, bool bWithFriction>
 	void ApplyInternalCCD(FSolverParticlesRange& Particles, const FSolverReal Dt, const TArray<FSolverCollisionParticlesRange>& CollisionParticles) const;
-	template<bool bWithFriction>
-	void ApplyInternalISPC(FSolverParticlesRange& Particles, const FSolverReal Dt, const TArray<FSolverCollisionParticlesRange>& CollisionParticles) const;
-	template<bool bWithFriction>
-	void ApplyInternalISPCPerTimestep(FSolverParticlesRange& Particles, const FSolverReal Dt, const TArray<FSolverCollisionParticlesRange>& CollisionParticles) const;
+#if INTEL_ISPC
+	void ApplySimpleInternalISPC(FSolverParticlesRange& Particles, const FSolverReal Dt, const TArray<FSolverCollisionParticlesRange>& CollisionParticles, bool bGeneratePlanarConstraints);
+	void ApplyComplexInternalISPC(FSolverParticlesRange& Particles, const FSolverReal Dt, const TArray<FSolverCollisionParticlesRange>& CollisionParticles, bool bGeneratePlanarConstraints);
+#endif
+	void InitPlanarConstraints(const FSolverParticlesRange& Particles, bool bWithFriction);
+	void FinalizePlanarConstraints(const FSolverParticlesRange& Particles);
+	void ApplyPlanarConstraints(FSolverParticlesRange& Particles, const FSolverReal Dt);
 
 protected:
 	const TArray<FSolverRigidTransform3>& LastCollisionTransforms; // Used by CCD
 	FSolverReal CollisionThickness;
+	FSolverReal SoftBodyCollisionThickness;
 	FSolverReal FrictionCoefficient;
 	bool bUseCCD;
 	FSolverReal ProximityStiffness; // Used by force-based solver
+	bool bEnableSimpleColliders = true;
+	bool bEnableComplexColliders = true;
+	bool bUsePlanarConstraintForSimpleColliders = bDefaultUsePlanarConstraintForSimpleColliders;
+	bool bUsePlanarConstraintForComplexColliders = bDefaultUsePlanarConstraintForComplexColliders;
+
+	TMap<FParticleRangeIndex, FPBDComplexColliderBoneData> ComplexBoneData; // Used to determine Velocity Bones for Skinned LevelSets.
 
 	/**  Used for writing debug contacts */
 	bool bWriteDebugContacts = false;
@@ -87,16 +136,31 @@ protected:
 	TArray<FSolverReal>* const Phis;
 	mutable FCriticalSection DebugMutex;
 
+	struct FPBDSoftBodyCollisionPlanarConstraint : public FPBDPlanarConstraints
+	{
+		FPBDSoftBodyCollisionPlanarConstraint() = default;
+		~FPBDSoftBodyCollisionPlanarConstraint() = default;
+		FPBDSoftBodyCollisionPlanarConstraint(const FSolverReal InFriction)
+			:FPBDPlanarConstraints(InFriction)
+		{}
+
+		TArray<int32>& GetUniqueConstraintIndices() { return UniqueConstraintIndices; }
+		TArray<FSolverVec3>& GetTargetPositions() { return TargetPositions; }
+		TArray<FSolverVec3>& GetTargetNormals() { return TargetNormals; }
+		TArray<FSolverVec3>& GetTargetVelocities() { return TargetVelocities; }
+	} PlanarConstraint;
+
 private:
-	// These arrays are used to do batch PhiAndNormal queries. They live here so they don't need to be reallocated each Apply call. 
-	mutable TArray<FSolverReal> BatchPhis;
-	mutable TArray<FSolverVec3> BatchNormals;
-	mutable TArray<int32> BatchVelocityBones;
-
-	mutable TArray<int32> CollisionIndices;
-	mutable TArray<FSolverVec3> CollisionTargets;
-	mutable TArray<FSolverVec3> CollisionNormalsOutput;
-	bool bDoCollisionPerTimestep = false;
+
+	// These arrays are used to do batch PhiAndNormal queries. They live here so they don't need to be reallocated each frame.
+	TArray<bool> HasPlanarData;
+	TArray<FSolverVec3> PlanarDataPositions;
+	TArray<FSolverVec3> PlanarDataNormals;
+	TArray<FSolverVec3> PlanarDataVelocities;
+
+	TArray<FSolverReal> BatchPhis;
+	TArray<FSolverVec3> BatchNormals;
+	TArray<int32> BatchVelocityBones;
 };
 
 
@@ -106,6 +170,7 @@ class FPBDSoftBodyCollisionConstraint : public FPBDSoftBodyCollisionConstraintBa
 
 public:
 	static constexpr FSolverReal DefaultCollisionThickness = (FSolverReal)1.;
+	static constexpr FSolverReal DefaultSoftBodyCollisionThickness = (FSolverReal)0.;
 	static constexpr FSolverReal DefaultFrictionCoefficient = (FSolverReal)0.8;
 	static constexpr FSolverReal DefaultProximityStiffness = (FSolverReal)100.;
 
@@ -116,28 +181,47 @@ public:
 		TArray<bool>* InCollisionParticleCollided = nullptr,
 		TArray<FSolverVec3>* InContacts = nullptr,
 		TArray<FSolverVec3>* InNormals = nullptr,
-		TArray<FSolverReal>* InPhis = nullptr)
+		TArray<FSolverReal>* InPhis = nullptr,
+		const TMap<FParticleRangeIndex, FPBDComplexColliderBoneData>& InComplexBoneData = TMap<FParticleRangeIndex, FPBDComplexColliderBoneData>())
 		: Base(InLastCollisionTransforms,
-			InMeshScale * GetCollisionThickness(PropertyCollection, DefaultCollisionThickness),
+			InMeshScale* GetCollisionThickness(PropertyCollection, DefaultCollisionThickness),
 			GetFrictionCoefficient(PropertyCollection, DefaultFrictionCoefficient),
 			GetUseCCD(PropertyCollection, false),
 			GetProximityStiffness(PropertyCollection, DefaultProximityStiffness),
-			InCollisionParticleCollided, InContacts, InNormals, InPhis)
+			InCollisionParticleCollided, InContacts, InNormals, InPhis,
+			GetSoftBodyCollisionThickness(PropertyCollection, DefaultSoftBodyCollisionThickness),
+			GetEnableSimpleColliders(PropertyCollection, true),
+			GetEnableComplexColliders(PropertyCollection, true),
+			GetUsePlanarConstraintForSimpleColliders(PropertyCollection, bDefaultUsePlanarConstraintForSimpleColliders),
+			GetUsePlanarConstraintForComplexColliders(PropertyCollection, bDefaultUsePlanarConstraintForComplexColliders),
+			InComplexBoneData
+		)
 		, MeshScale(InMeshScale)
 		, CollisionThicknessIndex(PropertyCollection)
+		, SoftBodyCollisionThicknessIndex(PropertyCollection)
 		, FrictionCoefficientIndex(PropertyCollection)
 		, UseCCDIndex(PropertyCollection)
 		, ProximityStiffnessIndex(PropertyCollection)
+		, EnableSimpleCollidersIndex(PropertyCollection)
+		, UsePlanarConstraintForSimpleCollidersIndex(PropertyCollection)
+		, EnableComplexCollidersIndex(PropertyCollection)
+		, UsePlanarConstraintForComplexCollidersIndex(PropertyCollection)
 	{}
 
 	CHAOS_API void SetProperties(const FCollectionPropertyConstFacade& PropertyCollection);
 
 private:
 	const FSolverReal MeshScale;
+
 	UE_CHAOS_DECLARE_PROPERTYCOLLECTION_NAME(CollisionThickness, float);
+	UE_CHAOS_DECLARE_PROPERTYCOLLECTION_NAME(SoftBodyCollisionThickness, float);
 	UE_CHAOS_DECLARE_PROPERTYCOLLECTION_NAME(FrictionCoefficient, float);
 	UE_CHAOS_DECLARE_PROPERTYCOLLECTION_NAME(UseCCD, bool);
 	UE_CHAOS_DECLARE_PROPERTYCOLLECTION_NAME(ProximityStiffness, float);
+	UE_CHAOS_DECLARE_PROPERTYCOLLECTION_NAME(EnableSimpleColliders, bool);
+	UE_CHAOS_DECLARE_PROPERTYCOLLECTION_NAME(UsePlanarConstraintForSimpleColliders, bool);
+	UE_CHAOS_DECLARE_PROPERTYCOLLECTION_NAME(EnableComplexColliders, bool);
+	UE_CHAOS_DECLARE_PROPERTYCOLLECTION_NAME(UsePlanarConstraintForComplexColliders, bool);
 };
 
 }  // End namespace Chaos::Softs
diff --git a/Engine/Source/Runtime/Experimental/Chaos/Public/Chaos/PBDSoftsEvolutionFwd.h b/Engine/Source/Runtime/Experimental/Chaos/Public/Chaos/PBDSoftsEvolutionFwd.h
--- a/Engine/Source/Runtime/Experimental/Chaos/Public/Chaos/PBDSoftsEvolutionFwd.h
+++ b/Engine/Source/Runtime/Experimental/Chaos/Public/Chaos/PBDSoftsEvolutionFwd.h
@@ -73,6 +73,7 @@ namespace Chaos
 		class FPBDAnimDriveConstraint;
 		class FPBDShapeConstraints;
 		class FPBDCollisionSpringConstraints;
+		class FPBDSkinnedTriangleMeshCollisionConstraints;
 		class FPBDTriangleMeshIntersections;
 		class FPBDTriangleMeshCollisions;
 		class FPBDSelfCollisionSphereConstraints;
@@ -85,5 +86,8 @@ namespace Chaos
 
 		// Linear system solver
 		class FEvolutionLinearSystem;
+
+		// Property facade
+		class FCollectionPropertyConstFacade;
 	}  // End namespace Softs
 }  // End namespace Chaos
diff --git a/Engine/Source/Runtime/Experimental/Chaos/Public/Chaos/PBDTriangleMeshCollisions.h b/Engine/Source/Runtime/Experimental/Chaos/Public/Chaos/PBDTriangleMeshCollisions.h
--- a/Engine/Source/Runtime/Experimental/Chaos/Public/Chaos/PBDTriangleMeshCollisions.h
+++ b/Engine/Source/Runtime/Experimental/Chaos/Public/Chaos/PBDTriangleMeshCollisions.h
@@ -386,6 +386,8 @@ public:
 		return NumContourMinimizationPostSteps;
 	}
 
+	bool GetCollidableSubMeshDirty() const { return bCollidableSubMeshDirty; }
+
 	const FTriangleSubMesh& GetCollidableSubMesh() const { return CollidableSubMesh; }
 	UE_DEPRECATED(5.4, "Use GetDynamicSpatialHash or GetKinematicColliderSpatialHash")
 	const FTriangleMesh::TSpatialHashType<FSolverReal>& GetSpatialHash() const { return DynamicSubMeshSpatialHash; }
diff --git a/Engine/Source/Runtime/Experimental/Chaos/Public/Chaos/ParticlesRange.h b/Engine/Source/Runtime/Experimental/Chaos/Public/Chaos/ParticlesRange.h
--- a/Engine/Source/Runtime/Experimental/Chaos/Public/Chaos/ParticlesRange.h
+++ b/Engine/Source/Runtime/Experimental/Chaos/Public/Chaos/ParticlesRange.h
@@ -7,6 +7,27 @@
 
 namespace Chaos::Softs
 {
+
+struct FParticleRangeIndex
+{
+	int32 RangeId;
+	int32 ParticleIndex;
+	bool operator<(const FParticleRangeIndex& Other) const
+	{
+		if (RangeId != Other.RangeId)
+		{
+			return RangeId < Other.RangeId;
+		}
+		return ParticleIndex < Other.ParticleIndex;
+	}
+	bool operator==(const FParticleRangeIndex&) const = default;
+};
+
+inline uint32 GetTypeHash(const FParticleRangeIndex& Index)
+{
+	return HashCombine(Index.RangeId, Index.ParticleIndex);
+}
+
 template<typename ParticlesType, typename = typename TEnableIf<TIsDerivedFrom<ParticlesType, TParticles<FSolverReal, 3>>::IsDerived>::Type>
 class TParticlesRange
 {
@@ -15,18 +36,19 @@ public:
 	TParticlesRange() = default;
 	virtual ~TParticlesRange() = default;
 
-	TParticlesRange(ParticlesType* InParticles, int32 InOffset, int32 InRangeSize)
+	TParticlesRange(ParticlesType* InParticles, int32 InOffset, int32 InRangeSize, int32 InRangeId = INDEX_NONE)
 		: Particles(InParticles)
 		, Offset(InOffset)
 		, RangeSize(InRangeSize)
+		, RangeId(InRangeId)
 	{
 	}
 
-	static TParticlesRange AddParticleRange(ParticlesType& InParticles, const int32 InRangeSize)
+	static TParticlesRange AddParticleRange(ParticlesType& InParticles, const int32 InRangeSize, const int32 InRangeId = INDEX_NONE)
 	{
 		const int32 Offset = (int32)InParticles.Size();
 		InParticles.AddParticles(InRangeSize);
-		return TParticlesRange(&InParticles, Offset, InRangeSize);
+		return TParticlesRange(&InParticles, Offset, InRangeSize, InRangeId);
 	}
 
 	bool IsValid() const
@@ -57,11 +79,13 @@ public:
 	{
 		return (Index >= 0) && (Index < RangeSize);
 	}
+	int32 GetRangeId() const { return RangeId; }
 
 protected:
 
 	ParticlesType* Particles = nullptr;
 	int32 Offset = INDEX_NONE;
 	int32 RangeSize = 0;
+	int32 RangeId = INDEX_NONE; // Can be used to identify a Range. 
 };
 }
diff --git a/Engine/Source/Runtime/Experimental/Chaos/Public/Chaos/PerParticlePBDCCDCollisionConstraint.h b/Engine/Source/Runtime/Experimental/Chaos/Public/Chaos/PerParticlePBDCCDCollisionConstraint.h
--- a/Engine/Source/Runtime/Experimental/Chaos/Public/Chaos/PerParticlePBDCCDCollisionConstraint.h
+++ b/Engine/Source/Runtime/Experimental/Chaos/Public/Chaos/PerParticlePBDCCDCollisionConstraint.h
@@ -82,7 +82,7 @@ private:
 			{
 				const uint32 KinematicGroupId = KinematicGroupIds[CollisionIndex];  // Collision group Id
 
-				if ((KinematicGroupId != (uint32)INDEX_NONE && DynamicGroupId != KinematicGroupId) || CollisionParticles.GetGeometry(CollisionIndex)->GetType() == Chaos::ImplicitObjectType::WeightedLatticeBone)
+				if (KinematicGroupId != (uint32)INDEX_NONE && DynamicGroupId != KinematicGroupId)
 				{
 					return; // Bail out if the collision groups doesn't match the particle group id, or use INDEX_NONE (= global collision that affects all particle)
 				}
@@ -114,35 +114,8 @@ private:
 					Particles.P(Index) += Penetration * NormalWorld;
 
 					// Friction
-					int32 VelocityBone = CollisionIndex;
-					if (const TWeightedLatticeImplicitObject<FLevelSet>* LevelSet = CollisionParticles.GetGeometry(CollisionIndex)->GetObject< TWeightedLatticeImplicitObject<FLevelSet> >())
-					{
-						TArray<FWeightedLatticeImplicitObject::FEmbeddingCoordinate> Coordinates;
-						LevelSet->GetEmbeddingCoordinates(PointPair.First, Coordinates, false);
-						int32 ClosestCoordIndex = INDEX_NONE;
-						double ClosestCoordPhi = UE_BIG_NUMBER;
-						for (int32 CoordIndex = 0; CoordIndex < Coordinates.Num(); ++CoordIndex)
-						{
-							FVec3 NormalUnused;
-							const double CoordPhi = FMath::Abs(LevelSet->GetEmbeddedObject()->PhiWithNormal(Coordinates[CoordIndex].UndeformedPosition(LevelSet->GetGrid()), NormalUnused));
-							if (CoordPhi < ClosestCoordPhi)
-							{
-								ClosestCoordIndex = CoordIndex;
-								ClosestCoordPhi = CoordPhi;
-							}
-						}
-						if (ClosestCoordIndex != INDEX_NONE)
-						{
-							const int32 StrongestBone = Coordinates[ClosestCoordIndex].GreatestInfluenceBone(LevelSet->GetBoneData());
-							if (StrongestBone != INDEX_NONE)
-							{
-								VelocityBone = LevelSet->GetSolverBoneIndices()[StrongestBone];
-							}
-						}
-					}
-
-					const FSolverVec3 VectorToPoint = Particles.P(Index) - CollisionParticles.GetX(VelocityBone);
-					const FSolverVec3 RelativeDisplacement = (Particles.P(Index) - Particles.GetX(Index)) - (CollisionParticles.V(VelocityBone) + FSolverVec3::CrossProduct(CollisionParticles.W(VelocityBone), VectorToPoint)) * Dt;  // This corresponds to the tangential velocity multiplied by dt (friction will drive this to zero if it is high enough)
+					const FSolverVec3 VectorToPoint = Particles.P(Index) - CollisionParticles.GetX(CollisionIndex);
+					const FSolverVec3 RelativeDisplacement = (Particles.P(Index) - Particles.GetX(Index)) - (CollisionParticles.V(CollisionIndex) + FSolverVec3::CrossProduct(CollisionParticles.W(CollisionIndex), VectorToPoint)) * Dt;  // This corresponds to the tangential velocity multiplied by dt (friction will drive this to zero if it is high enough)
 					const FSolverVec3 RelativeDisplacementTangent = RelativeDisplacement - FSolverVec3::DotProduct(RelativeDisplacement, NormalWorld) * NormalWorld;  // Project displacement into the tangential plane
 					const FSolverReal RelativeDisplacementTangentLength = RelativeDisplacementTangent.Size();
 					if (RelativeDisplacementTangentLength >= UE_SMALL_NUMBER)
