Chaos Cloth complex collider refactor and add support for SkinnedTriangleMesh colliders.

diff --git a/Engine/Source/Runtime/Experimental/Chaos/Public/Chaos/PerParticlePBDCollisionConstraint.h b/Engine/Source/Runtime/Experimental/Chaos/Public/Chaos/PerParticlePBDCollisionConstraint.h
--- a/Engine/Source/Runtime/Experimental/Chaos/Public/Chaos/PerParticlePBDCollisionConstraint.h
+++ b/Engine/Source/Runtime/Experimental/Chaos/Public/Chaos/PerParticlePBDCollisionConstraint.h
@@ -125,36 +125,16 @@ private:
 				{
 					const uint32 KinematicGroupId = MKinematicGroupIds[i];  // Collision group Id
 
-					if ((KinematicGroupId != (uint32)INDEX_NONE && DynamicGroupId != KinematicGroupId ) || CollisionParticles.GetGeometry(i)->GetType() == Chaos::ImplicitObjectType::WeightedLatticeBone)
+					if (KinematicGroupId != (uint32)INDEX_NONE && DynamicGroupId != KinematicGroupId )
 					{
 						return; // Bail out if the collision groups doesn't match the particle group id, or use INDEX_NONE (= global collision that affects all particle)
 					}
 					const FSolverRigidTransform3 Frame(CollisionParticles.GetX(i), CollisionParticles.GetR(i));
 					const FVec3 RigidSpacePosition(Frame.InverseTransformPosition(Particles.P(Index)));  // PhiWithNormal requires FReal based arguments
 					FVec3 ImplicitNormal;                                                                // since implicits don't use FSolverReal
-					FSolverReal Phi;
-					int32 VelocityBone = i;
-					FSolverReal Penetration;
-					if (const TWeightedLatticeImplicitObject<FLevelSet>* LevelSet = CollisionParticles.GetGeometry(i)->GetObject< TWeightedLatticeImplicitObject<FLevelSet> >())
-					{
-						FWeightedLatticeImplicitObject::FEmbeddingCoordinate SurfaceCoord;
-						Phi = (FSolverReal)LevelSet->PhiWithNormalAndSurfacePoint(RigidSpacePosition, ImplicitNormal, SurfaceCoord);
-						Penetration = PerGroupThickness - Phi; // This is related to the Normal impulse
-						if (Penetration > (FSolverReal)0.)
-						{
-							const int32 StrongestBone = SurfaceCoord.GreatestInfluenceBone(LevelSet->GetBoneData());
-							if (StrongestBone != INDEX_NONE)
-							{
-								VelocityBone = LevelSet->GetSolverBoneIndices()[StrongestBone];
-							}
-						}
-					}
-					else
-					{
-						Phi = (FSolverReal)CollisionParticles.GetGeometry(i)->PhiWithNormal(RigidSpacePosition, ImplicitNormal);
-						Penetration = PerGroupThickness - Phi; // This is related to the Normal impulse
-					}
+					const FSolverReal Phi = (FSolverReal)CollisionParticles.GetGeometry(i)->PhiWithNormal(RigidSpacePosition, ImplicitNormal);
 					const FSolverVec3 Normal(ImplicitNormal);
+					const FSolverReal Penetration = PerGroupThickness - Phi; // This is related to the Normal impulse
 
 					if (Penetration > (FSolverReal)0.)
 					{
@@ -176,8 +156,8 @@ private:
 
 						if (bFastPositionBasedFriction)
 						{
-							FSolverVec3 VectorToPoint = Particles.GetP(Index) - CollisionParticles.GetX(VelocityBone);
-							const FSolverVec3 RelativeDisplacement = (Particles.GetP(Index) - Particles.GetX(Index)) - (CollisionParticles.V(VelocityBone) + FSolverVec3::CrossProduct(CollisionParticles.W(VelocityBone), VectorToPoint)) * Dt; // This corresponds to the tangential velocity multiplied by dt (friction will drive this to zero if it is high enough)
+							FSolverVec3 VectorToPoint = Particles.GetP(Index) - CollisionParticles.GetX(i);
+							const FSolverVec3 RelativeDisplacement = (Particles.GetP(Index) - Particles.GetX(Index)) - (CollisionParticles.V(i) + FSolverVec3::CrossProduct(CollisionParticles.W(i), VectorToPoint)) * Dt; // This corresponds to the tangential velocity multiplied by dt (friction will drive this to zero if it is high enough)
 							const FSolverVec3 RelativeDisplacementTangent = RelativeDisplacement - FSolverVec3::DotProduct(RelativeDisplacement, NormalWorld) * NormalWorld; // Project displacement into the tangential plane
 							const FSolverReal RelativeDisplacementTangentLength = RelativeDisplacementTangent.Size();
 							if (RelativeDisplacementTangentLength >= UE_SMALL_NUMBER)
@@ -191,8 +171,8 @@ private:
 						{
 							// Note, to fix: Only use fast position based friction for now, since adding to TMaps here is not thread safe when calling Apply on multiple threads (will cause crash)
 							FVelocityConstraint Constraint;
-							FSolverVec3 VectorToPoint = Particles.GetP(Index) - CollisionParticles.GetX(VelocityBone);
-							Constraint.Velocity = CollisionParticles.V(VelocityBone) + FSolverVec3::CrossProduct(CollisionParticles.W(VelocityBone), VectorToPoint);
+							FSolverVec3 VectorToPoint = Particles.GetP(Index) - CollisionParticles.GetX(i);
+							Constraint.Velocity = CollisionParticles.V(i) + FSolverVec3::CrossProduct(CollisionParticles.W(i), VectorToPoint);
 							Constraint.Normal = Frame.TransformVector(Normal);
 						
 							MVelocityConstraints.Add(Index, Constraint);
@@ -216,7 +196,7 @@ private:
 				{
 					const uint32 KinematicGroupId = MKinematicGroupIds[i];  // Collision group Id
 
-					if ((KinematicGroupId != (uint32)INDEX_NONE && DynamicGroupId != KinematicGroupId) || CollisionParticles.GetGeometry(i)->GetType() == Chaos::ImplicitObjectType::WeightedLatticeBone)
+					if (KinematicGroupId != (uint32)INDEX_NONE && DynamicGroupId != KinematicGroupId)
 					{
 						return; // Bail out if the collision groups doesn't match the particle group id, or use INDEX_NONE (= global collision that affects all particle)
 					}
diff --git a/Engine/Source/Runtime/Experimental/Chaos/Public/Chaos/SoftsEvolution.h b/Engine/Source/Runtime/Experimental/Chaos/Public/Chaos/SoftsEvolution.h
--- a/Engine/Source/Runtime/Experimental/Chaos/Public/Chaos/SoftsEvolution.h
+++ b/Engine/Source/Runtime/Experimental/Chaos/Public/Chaos/SoftsEvolution.h
@@ -15,6 +15,8 @@
 namespace Chaos::Softs
 {
 
+
+
 enum struct ESolverMode : uint8
 {
 	None = 0,
@@ -23,6 +25,43 @@ enum struct ESolverMode : uint8
 };
 ENUM_CLASS_FLAGS(ESolverMode);
 
+/**
+ * Per Group context with information about the current solver configuration.
+ */
+struct FEvolutionGroupContext
+{
+	ESolverMode SolverMode = ESolverMode::None;
+	FSolverReal Dt = (FSolverReal)0.;
+
+	int32 NumPBDIterations = 0;
+	int32 CurrentPBDIteration = INDEX_NONE;
+
+	int32 NumNewtonIterations = 0;
+	int32 CurrentNewtonIteration = INDEX_NONE;
+
+	void Reset()
+	{
+		SolverMode = ESolverMode::None;
+		Dt = (FSolverReal)0.;
+
+		NumPBDIterations = 0;
+		CurrentPBDIteration = INDEX_NONE;
+
+		NumNewtonIterations = 0;
+		CurrentNewtonIteration = INDEX_NONE;
+	}
+
+	void Init(ESolverMode InSolverMode, FSolverReal InDt, int32 InNumPBDIterations, int32 InNumNewtonIterations)
+	{
+		SolverMode = InSolverMode;
+		Dt = InDt;
+		NumPBDIterations = InNumPBDIterations;
+		CurrentPBDIteration = INDEX_NONE;
+		NumNewtonIterations = InNumNewtonIterations;
+		CurrentNewtonIteration = INDEX_NONE;
+	}
+};
+
 /**
  * Solver can contain multiple "Groups". Groups do not interact with each other. 
  * They may be in different spaces. They may be solved in parallel, completely independently of each other. 
@@ -32,7 +71,7 @@ ENUM_CLASS_FLAGS(ESolverMode);
  * A Group can contain multiple "SoftBodies". SoftBodies can interact but have different 
  * constraint rules/forces.
  */
-class FEvolution 
+class FEvolution
 {
 public:
 
@@ -57,6 +96,17 @@ public:
 	const TSet<uint32>& GetActiveGroups() const { return ActiveGroups; }
 	CHAOS_API int32 NumActiveParticles() const;
 
+	const FEvolutionGroupContext& GetGroupContext(uint32 GroupId) const
+	{
+		return Groups.SolverContexts[GroupId];
+	}
+
+	// Convenience method to get by SoftBodyId
+	const FEvolutionGroupContext& GetGroupContextForSoftBody(int32 SoftBodyId) const
+	{
+		return Groups.SolverContexts[SoftBodies.GroupId[SoftBodyId]];
+	}
+
 	/** 
 	 * Add a SoftBody to a Group.
 	 * 
@@ -79,7 +129,7 @@ public:
 	const TSet<int32>& GetGroupActiveSoftBodies(uint32 GroupId) const { return Groups.ActiveSoftBodies[GroupId]; }
 	int32 GetLastLinearSolveIterations(int32 SoftBodyId) const { return SoftBodies.LinearSystems[SoftBodyId].GetLastSolveIterations(); }
 	FSolverReal GetLastLinearSolveError(int32 SoftBodyId) const { return SoftBodies.LinearSystems[SoftBodyId].GetLastSolveError(); }
-
+	
 	/**
 	 * Add Collision particle range to a group.
 	 * 
@@ -162,42 +212,52 @@ public:
 		return AllocateRules(SoftBodyId, NumRules, SoftBodies.PostSubstepConstraintRules);
 	}
 
+	// Presubstep init methods (always run at beginning of substep)
 	TArrayView<ParallelInitFunc> GetPreSubstepParallelInitRange(int32 SoftBodyId)
 	{
 		return GetRules(SoftBodyId, SoftBodies.PreSubstepParallelInits);
 	}
+	// PBD Rules that apply external forces (only run if doing PBD)
 	TArrayView<PBDConstraintRuleFunc> GetPBDExternalForceRulesRange(int32 SoftBodyId)
 	{
 		return GetRules(SoftBodyId, SoftBodies.PBDExternalForceRules);
 	}
+	// Post initial guess init methods (always run after kinematic and initial guess update, before any solving.) 
 	TArrayView<ParallelInitFunc> GetPostInitialGuessParallelInitRange(int32 SoftBodyId)
 	{
 		return GetRules(SoftBodyId, SoftBodies.PostInitialGuessParallelInits);
 	}
+	// Rules that run once per substep after all initial guess and initialization is done.
 	TArrayView<ConstraintRuleFunc> GetPreSubstepConstraintRulesRange(int32 SoftBodyId)
 	{
 		return GetRules(SoftBodyId, SoftBodies.PreSubstepConstraintRules);
 	}
+	// Normal per-iteration PBD rules (only run if doing PBD)
 	TArrayView<PBDConstraintRuleFunc> GetPerIterationPBDConstraintRulesRange(int32 SoftBodyId)
 	{
 		return GetRules(SoftBodyId, SoftBodies.PerIterationPBDConstraintRules);
 	}
+	// Collision per-iteration PBD rules (only run if doing PBD)
 	TArrayView<PBDCollisionConstraintRuleFunc> GetPerIterationCollisionPBDConstraintRulesRange(int32 SoftBodyId)
 	{
 		return GetRules(SoftBodyId, SoftBodies.PerIterationCollisionPBDConstraintRules);
 	}
+	// Normal per-iteration PBD rules that run after collisions (only run if doing PBD)
 	TArrayView<PBDConstraintRuleFunc> GetPerIterationPostCollisionsPBDConstraintRulesRange(int32 SoftBodyId)
 	{
 		return GetRules(SoftBodyId, SoftBodies.PerIterationPostCollisionsPBDConstraintRules);
 	}
+	// Linear system rules (only run if doing ForceBased)
 	TArrayView<UpdateLinearSystemFunc> GetUpdateLinearSystemRulesRange(int32 SoftBodyId)
 	{
 		return GetRules(SoftBodyId, SoftBodies.UpdateLinearSystemRules);
 	}
+	// Linear system collision rules (only run if doing ForceBased)
 	TArrayView<UpdateLinearSystemCollisionsFunc> GetUpdateLinearSystemCollisionsRulesRange(int32 SoftBodyId)
 	{
 		return GetRules(SoftBodyId, SoftBodies.UpdateLinearSystemCollisionsRules);
 	}
+	// Post substep rules (always run at end of substep)
 	TArrayView<ConstraintRuleFunc> GetPostSubstepConstraintRulesRange(int32 SoftBodyId)
 	{
 		return GetRules(SoftBodyId, SoftBodies.PostSubstepConstraintRules);
@@ -212,8 +272,6 @@ public:
 	bool GetDoQuasistatics() const { return bDoQuasistatics; }
 	void SetDisableTimeDependentNumIterations(bool bDisable) { bDisableTimeDependentNumIterations = bDisable; }
 	CHAOS_API void SetSolverProperties(const FCollectionPropertyConstFacade& PropertyCollection);
-	bool DoCollisionPerTimestep() const { return bDoCollisionPerTimestep; }
-	void SetDoCollisionPerTimestep(bool bInDoCollisionPerTimestep) { bDoCollisionPerTimestep = bInDoCollisionPerTimestep; }
 
 private:
 
@@ -360,6 +418,7 @@ private:
 			TArrayCollection::AddArray(&SoftBodies);
 			TArrayCollection::AddArray(&ActiveSoftBodies);
 			TArrayCollection::AddArray(&ActiveCollisionParticleRanges);
+			TArrayCollection::AddArray(&SolverContexts);
 		}
 		
 		void Reset()
@@ -378,6 +437,7 @@ private:
 		TArrayCollectionArray<TArray<int32>> SoftBodies;
 		TArrayCollectionArray<TSet<int32>> ActiveSoftBodies;
 		TArrayCollectionArray<TSet<int32>> ActiveCollisionParticleRanges;
+		TArrayCollectionArray<FEvolutionGroupContext> SolverContexts;
 	};
 
 	// Wrapper around FEvolutionLinearSystemSolverParameters that knows how to read a property collection
@@ -438,7 +498,6 @@ private:
 	bool bDoQuasistatics;
 	FSolverReal SolverFrequency; 
 	FLinearSystemParameters LinearSystemParameters; // Per-solver parameters that need to be passed to the linear system solver
-	bool bDoCollisionPerTimestep = false; // Use PlanarConstraints
 
 
 	// Per-Particle data
diff --git a/Engine/Source/Runtime/Experimental/Chaos/Public/Chaos/WeightedLatticeImplicitObject.h b/Engine/Source/Runtime/Experimental/Chaos/Public/Chaos/WeightedLatticeImplicitObject.h
--- a/Engine/Source/Runtime/Experimental/Chaos/Public/Chaos/WeightedLatticeImplicitObject.h
+++ b/Engine/Source/Runtime/Experimental/Chaos/Public/Chaos/WeightedLatticeImplicitObject.h
@@ -64,7 +64,9 @@ public:
 
 	CHAOS_API FWeightedLatticeImplicitObject(FWeightedLatticeImplicitObject&& Other);
 	
+	PRAGMA_DISABLE_DEPRECATION_WARNINGS
 	virtual ~FWeightedLatticeImplicitObject() override = default;
+	PRAGMA_ENABLE_DEPRECATION_WARNINGS
 
 	virtual const FAABB3 BoundingBox() const override { return LocalBoundingBox; }
 	const TArrayND<FWeightedLatticeInfluenceData, 3>& GetBoneData() const { return BoneData; }
@@ -72,8 +74,15 @@ public:
 	const TArrayND<bool, 3>& GetEmptyCells() const { return EmptyCells; }
 	const TUniformGrid<FReal, 3>& GetGrid() const { return Grid; }
 	const TArray<FName>& GetUsedBones() const { return UsedBones; }
-	const TArray<int32>& GetSolverBoneIndices() const { return SolverBoneIndices; }
-	void SetSolverBoneIndices(TArray<int32>&& InSolverBoneIndices) { SolverBoneIndices = MoveTemp(InSolverBoneIndices); }
+
+	UE_DEPRECATED(5.6, "SolverBoneIndices are no longer stored on WeightedLatticeImplicitObjects")
+		const TArray<int32>& GetSolverBoneIndices() const
+	{
+		static TArray<int32> Dummy;
+		return Dummy;
+	}
+	UE_DEPRECATED(5.6, "SolverBoneIndices are no longer stored on WeightedLatticeImplicitObjects")
+	void SetSolverBoneIndices(TArray<int32>&& InSolverBoneIndices) { }
 
 	CHAOS_API virtual void Serialize(FChaosArchive& Ar) override;
 
@@ -167,6 +176,7 @@ protected:
 	TArrayND<bool, 3> EmptyCells;
 
 	// Context-specific data (e.g., data a specific solver will need to do DeformPoints)
+	UE_DEPRECATED(5.6, "SolverBoneIndices are deprecated")
 	TArray<int32> SolverBoneIndices; // indices set by the solver for bone transform query
 
 	// Non-serialized deforming data
@@ -219,10 +229,11 @@ public:
 	// Batch calls used by cloth.
 	// Note that if x is in an empty cell / outside the deformed lattice, Phi will be UE_BIG_NUMBER and Normal will be unchanged (BatchNormals will be resized using SetNumUninitialized, so it may be left uninitialized)
 	// GreatestInfluenceBone will not be calculated (and set to INDEX_NONE) if Phi is >= Thickness.
+	// GreatestInfluenceBone is an index into UsedBones
 	// All outputs will be unchanged (and possibly uninitialized) for points with InvM == 0.
 	void BatchPhiWithNormalAndGreatestInfluenceBone(const TConstArrayView<Softs::FPAndInvM> PAndInvMArray,
 		const Softs::FSolverRigidTransform3& SolverToThis, const Softs::FSolverReal Thickness,
-		TArray<Softs::FSolverReal>& OutBatchPhis, TArray<Softs::FSolverVec3>& OutBatchNormals, TArray<int32>& OutBatchGreatestInfluenceSolverBones) const;
+		TArray<Softs::FSolverReal>& OutBatchPhis, TArray<Softs::FSolverVec3>& OutBatchNormals, TArray<int32>& OutBatchGreatestInfluenceBones) const;
 
 	// Batch calls used by cloth.
 	// Note that if x is in an empty cell / outside the deformed lattice, Phi will be UE_BIG_NUMBER and Normal will be unchanged (BatchNormals will be resized using SetNumUninitialized, so it may be left uninitialized)
@@ -245,6 +256,7 @@ public:
 	}
 
 private:
+
 	// Serialized data
 	ObjectType Object;
 
@@ -316,16 +328,16 @@ private:
 };
 
 // Not actually a collision object. CollisionParticles that are used to update bones used by the FWeightedLatticeImplicitObject have this as their geometry
-class FWeightedLatticeBoneProxy : public FImplicitObject
+class UE_DEPRECATED(5.6, "FWeightedLatticeBoneProxy have been deprecated. Do not use.") FWeightedLatticeBoneProxy : public FImplicitObject
 {
 public:
 	using FImplicitObject::GetTypeName;
 
 	FWeightedLatticeBoneProxy()
-		: FImplicitObject(EImplicitObject::DisableCollisions, ImplicitObjectType::WeightedLatticeBone)
+		: FImplicitObject(EImplicitObject::DisableCollisions, ImplicitObjectType::Unknown)
 	{}
 
-	static constexpr EImplicitObjectType StaticType() { return ImplicitObjectType::WeightedLatticeBone; }
+	static constexpr EImplicitObjectType StaticType() { return ImplicitObjectType::Unknown; }
 
 	virtual FReal PhiWithNormal(const FVec3& x, FVec3& Normal) const override
 	{
