Simplify object and field reference convention to just use GetPathName, include the PackageName in the __PackageSummary object. This matches the behavior of the Summary field in a uasset (holds version info and metadata)

#rnx

diff --git a/Engine/Source/Runtime/Experimental/JsonObjectGraph/Private/JsonObjectGraphConventions.h b/Engine/Source/Runtime/Experimental/JsonObjectGraph/Private/JsonObjectGraphConventions.h
--- a/Engine/Source/Runtime/Experimental/JsonObjectGraph/Private/JsonObjectGraphConventions.h
+++ b/Engine/Source/Runtime/Experimental/JsonObjectGraph/Private/JsonObjectGraphConventions.h
@@ -11,7 +11,10 @@
 // 'Package' level conventions, these are used at the root level:
 #define UE_JSON_ROOT_OBJECTS_KEY_TCHAR TEXT("__RootObjects")
 
+#define UE_JSON_PACKAGE_SUMMARY_KEY_TCHAR TEXT("__PackageSummary")
+
 #define UE_JSON_CUSTOM_VERSIONS_KEY_TCHAR TEXT("__CustomVersions")
+#define UE_JSON_PACKAGE_NAME_KEY_TCHAR TEXT("__PackageName")
 
 // UObject instance conventions:
 #define UE_JSON_OBJECT_INSTANCE_KEY_TCHAR TEXT("__UObject")
@@ -27,9 +30,6 @@
 #define UE_JSON_OBJECT_REF_PREFIX "uobject:"
 #define UE_JSON_FIELD_REF_PREFIX "ffield:"
 #define UE_JSON_REF_NONE "None"
-#define UE_JSON_DOMAIN_IDENTIFIER_NATIVE "__native__"
-#define UE_JSON_DOMAIN_IDENTIFIER_FILE "__file__"
-#define UE_JSON_REF_POSTFIX ".uasset"
 
 // Conventions related to encoding generic intrinsic types (TMap and TOptional):
 #define UE_JSON_TMAP_KEY_KEY_TCHAR TEXT("Key")
diff --git a/Engine/Source/Runtime/Experimental/JsonObjectGraph/Private/JsonStringifyImpl.cpp b/Engine/Source/Runtime/Experimental/JsonObjectGraph/Private/JsonStringifyImpl.cpp
--- a/Engine/Source/Runtime/Experimental/JsonObjectGraph/Private/JsonStringifyImpl.cpp
+++ b/Engine/Source/Runtime/Experimental/JsonObjectGraph/Private/JsonStringifyImpl.cpp
@@ -21,41 +21,6 @@
 namespace UE::Private
 {
 
-FUtf8String GetPackageModuleName(FUtf8StringView PackageName)
-{
-	if (PackageName.StartsWith("/Script/"))
-	{
-		// native package, module name will immediately follow script
-		// and object name will immediately follow that module name:
-		return FUtf8String(UE_JSON_DOMAIN_IDENTIFIER_NATIVE);
-	}
-	else if (PackageName.StartsWith("/Game/"))
-	{
-		return TCHAR_TO_UTF8(FApp::GetProjectName());
-	}
-	else
-	{
-		check(PackageName[0] == '/');
-		FUtf8StringView PackageNameNoSlash(PackageName.GetData() + 1, PackageName.Len() - 1);
-		// normal asset, it should be located within a plugin or
-		// project, that will serve as its 'module' name:
-		int32 FolderIndex = INDEX_NONE;
-		check(PackageNameNoSlash.FindChar(static_cast<UTF8CHAR>('/'), FolderIndex));
-
-		return FUtf8String(PackageName.Mid(1, FolderIndex));
-	}
-}
-
-FUtf8String GetPackageRelativePath(FUtf8StringView PackageName)
-{
-	check(PackageName[0] == '/');
-	FUtf8StringView PackageNameNoSlash(PackageName.GetData() + 1, PackageName.Len() - 1);
-	int32 FolderIndex = INDEX_NONE;
-	check(PackageNameNoSlash.FindChar(static_cast<UTF8CHAR>('/'), FolderIndex));
-	++FolderIndex;
-	return FUtf8String(PackageName.RightChop(FolderIndex));
-}
-
 // gathering utils
 struct FPackageReferenceFinder : public FArchiveUObject
 {
@@ -292,6 +257,11 @@ FJsonStringifyImpl::FJsonStringifyImpl(TConstArrayView<const UObject*> Roots, co
 	TArray<const UObject*> Exports;
 	GatherExports(Roots, {}, bFilterEditorOnly, RootObjects, Exports);
 
+	if(Exports.Num() > 0 )
+	{
+		PackageName = Exports[0]->GetPackage()->GetFName();
+	}
+
 	ObjectsToExport = TSet<const UObject*>(Exports);
 
 	MemoryWriter.SetIsPersistent(true);
@@ -364,29 +334,11 @@ void FJsonStringifyImpl::ToJsonBytes()
 	}
 	Writer->WriteArrayEnd();
 	
-#if WITH_TEXT_ARCHIVE_SUPPORT
-	if (Versions.Num() != 0)
+	if(ShouldWritePackageSummary())
 	{
-		// filter dupes, preserving order of first encounter - we need some stable order:
-		TSet<FGuid> Encountered;
-		Versions.SetNum(Algo::StableRemoveIf(Versions,
-			[&Encountered](const FCustomVersion& V)
-			{
-				if (Encountered.Contains(V.Key))
-				{
-					return true;
-				}
-				else
-				{
-					Encountered.Add(V.Key);
-					return false;
-				}
-			}));
-
-		Writer->WriteIdentifierPrefix(UE_JSON_CUSTOM_VERSIONS_KEY_TCHAR);
-		FJsonStringifyStructuredArchive::WriteCustomVersionValueInline(Versions, Writer->GetIndentLevel(), MemoryWriter);
+		WritePackageSummary();
 	}
-#endif // WITH_TEXT_ARCHIVE_SUPPORT
+
 	Writer->WriteObjectEnd();
 	// Result is now populated with Json representation
 }
@@ -517,53 +469,12 @@ FUtf8String FJsonStringifyImpl::WriteObjectReference(const UObject* ForObject) c
 		return UE_JSON_OBJECT_REF_PREFIX UE_JSON_REF_NONE;
 	}
 
-	// We want users to be able to move these files around the filesystem with
-	// minimal fixup, so intra artifact references are relative to the file itself,
-	// intra module references are relative to the module, but cross module references
-	// are unambiguous and fully qualified. Here are some examples:
-	//
-	// uobject:__file__/a/path/to/a/child - this is a reference to an object within our file
-	// uobject:__native__/native_module/native_object - this is a reference to a natively defined object in native_module
-	// uobject:real_module/a/path/to/a/child.uasset/a/path/to/a/child/object - this is a reference to an object in another file
-	//
-	// file extensions are included in the path to disambiguate between folders and assets
-	// that share a name. The special __file__ identifier exists so that assets internal
-	// reference remain valid when copied, the special __native__ identifier exists so that
-	// natively defined objects can be unambiguously referenced
-	//
-	// real_module in this example will refer to either a uplugin or uproject.
-
+	// Encode objects with the long standing pathname convention, we can
+	// encode a basis in the buffer to give users options when an asset
+	// is moved on the filesystem:
 	FUtf8StringBuilderBase Reference;
 	Reference << UE_JSON_OBJECT_REF_PREFIX;
-	if (ObjectsToExport.Contains(ForObject))
-	{
-		Reference << UE_JSON_DOMAIN_IDENTIFIER_FILE;
-		const UObject* OuterLimit = ForObject;
-		while (ObjectsToExport.Contains(OuterLimit))
-		{
-			OuterLimit = OuterLimit->GetOuter();
-		}
-		WriteObjectPath(Reference, ForObject, OuterLimit);
-	}
-	else if (const UPackage* Package = ForObject->GetPackage())
-	{
-		const FUtf8String PackageName = TCHAR_TO_UTF8(*Package->GetName());
-		const FUtf8String ModuleName = GetPackageModuleName(PackageName);
-
-		Reference << ModuleName;
-
-		FUtf8String PackagePath = GetPackageRelativePath(PackageName);
-		Reference << PackagePath;
-		if (ModuleName != UE_JSON_DOMAIN_IDENTIFIER_NATIVE)
-		{
-			Reference << UE_JSON_REF_POSTFIX;
-		}
-		WriteObjectPath(Reference, ForObject, Package);
-	}
-	else
-	{
-		checkf(false, TEXT("Found object with no package: %s"), *ForObject->GetPathName());
-	}
+	Reference << ForObject->GetPathName();
 
 	return Reference.ToString();
 }
@@ -578,42 +489,8 @@ FUtf8String FJsonStringifyImpl::WriteFieldReference(const FField* Value) const
 	// identical convention for FField as above, albeit with a fieldref prefix
 	// instead of a uobject:
 	FUtf8StringBuilderBase Reference;
-
 	Reference << UE_JSON_FIELD_REF_PREFIX;
-	const UObject* OwningObject = Value->GetOwnerUObject();
-	check(OwningObject);
-	if (ObjectsToExport.Contains(OwningObject))
-	{
-		// object is within this file:
-		Reference << UE_JSON_DOMAIN_IDENTIFIER_FILE;
-		const UObject* OuterLimit = OwningObject;
-		while (ObjectsToExport.Contains(OuterLimit))
-		{
-			OuterLimit = OuterLimit->GetOuter();
-		}
-
-		WriteFieldPath(Reference, Value, OuterLimit);
-	}
-	else if (const UPackage* Package = OwningObject->GetPackage())
-	{
-		// external object:
-		const FUtf8String PackageName = TCHAR_TO_UTF8(*Package->GetName());
-		const FUtf8String ModuleName = GetPackageModuleName(PackageName);
-
-		Reference << ModuleName;
-
-		FUtf8String PackagePath = GetPackageRelativePath(PackageName);
-		Reference << PackagePath;
-		if (ModuleName != UE_JSON_DOMAIN_IDENTIFIER_NATIVE)
-		{
-			Reference << UE_JSON_REF_POSTFIX;
-		}
-		WriteFieldPath(Reference, Value, Package);
-	}
-	else
-	{
-		checkf(false, TEXT("Found field with no package: %s"), *OwningObject->GetPathName());
-	}
+	Reference << Value->GetPathName();
 
 	return Reference.ToString();
 }
@@ -952,11 +829,63 @@ TArray<uint8> FJsonStringifyImpl::StructuredDataToJson(const UObject* Object, in
 }
 #endif // WITH_TEXT_ARCHIVE_SUPPORT
 
+void FJsonStringifyImpl::WritePackageSummary()
+{
+	Writer->WriteIdentifierPrefix(UE_JSON_PACKAGE_SUMMARY_KEY_TCHAR);
+	Writer->WriteObjectStartInline();
+
+#if WITH_TEXT_ARCHIVE_SUPPORT
+	if (Versions.Num() != 0)
+	{
+		// filter dupes, preserving order of first encounter - we need some stable order:
+		TSet<FGuid> Encountered;
+		Versions.SetNum(Algo::StableRemoveIf(Versions,
+			[&Encountered](const FCustomVersion& V)
+			{
+				if (Encountered.Contains(V.Key))
+				{
+					return true;
+				}
+				else
+				{
+					Encountered.Add(V.Key);
+					return false;
+				}
+			}));
+
+		Writer->WriteIdentifierPrefix(UE_JSON_CUSTOM_VERSIONS_KEY_TCHAR);
+		FJsonStringifyStructuredArchive::WriteCustomVersionValueInline(Versions, Writer->GetIndentLevel(), MemoryWriter);
+		Writer->HACK_SetPreviousTokenWrittenSquareClose(); // the above will write the custom version as an array of tuples, track that write
+	}
+#endif // WITH_TEXT_ARCHIVE_SUPPORT
+
+	// Write package name, this is useful for detecting file system level moves 
+	// and matches the behavior of FLinkerSave:
+	if(PackageName != FName())
+	{
+		Writer->WriteIdentifierPrefix(UE_JSON_PACKAGE_NAME_KEY_TCHAR);
+		Writer->WriteValueInline(PackageName.ToString());
+	}
+
+	// Note that there is no automatic compression here, if you add information to the package summary
+	// you should update ShouldWritePackageSummary
+	Writer->WriteObjectEnd();
+}
+
 bool FJsonStringifyImpl::IsDeltaEncoding() const
 {
 	return !EnumHasAnyFlags(WriteOptions.Flags, EJsonStringifyFlags::DisableDeltaEncoding);
 }
 
+bool FJsonStringifyImpl::ShouldWritePackageSummary() const
+{
+	return 
+#if WITH_TEXT_ARCHIVE_SUPPORT
+		Versions.Num() > 0 ||
+#endif
+		PackageName != FName();
+}
+
 FJsonStringifyImpl::FPendingScope::FPendingScope(FJsonStringifyImpl* To, const TFunction<void()>& Prefix)
 	: Owner(To)
 	, Outer(To->CurrentScope)
diff --git a/Engine/Source/Runtime/Experimental/JsonObjectGraph/Private/JsonStringifyImpl.h b/Engine/Source/Runtime/Experimental/JsonObjectGraph/Private/JsonStringifyImpl.h
--- a/Engine/Source/Runtime/Experimental/JsonObjectGraph/Private/JsonStringifyImpl.h
+++ b/Engine/Source/Runtime/Experimental/JsonObjectGraph/Private/JsonStringifyImpl.h
@@ -67,7 +67,10 @@ private:
 	TArray<uint8> StructuredDataToJson(const UObject* Object, int32 InitialIndentLevel);
 #endif// WITH_TEXT_ARCHIVE_SUPPORT
 
+	void WritePackageSummary();
+
 	bool IsDeltaEncoding() const;
+	bool ShouldWritePackageSummary() const;
 
 	const FJsonStringifyOptions& WriteOptions;
 
@@ -103,6 +106,7 @@ private:
 	TArray<const UObject*> RootObjects;
 	TSet<const UObject*> ObjectsToExport;
 	TSet<const UObject*> ObjectsExported;
+	FName PackageName;
 
 	TArray<FCustomVersion> Versions;
 };
diff --git a/Engine/Source/Runtime/Experimental/JsonObjectGraph/Private/PrettyJsonWriter.cpp b/Engine/Source/Runtime/Experimental/JsonObjectGraph/Private/PrettyJsonWriter.cpp
--- a/Engine/Source/Runtime/Experimental/JsonObjectGraph/Private/PrettyJsonWriter.cpp
+++ b/Engine/Source/Runtime/Experimental/JsonObjectGraph/Private/PrettyJsonWriter.cpp
@@ -149,6 +149,11 @@ void FPrettyJsonWriter::HACK_SetPreviousTokenWritten()
 	PreviousTokenWritten = EJsonToken::CurlyOpen; 
 }
 
+void FPrettyJsonWriter::HACK_SetPreviousTokenWrittenSquareClose()
+{
+	PreviousTokenWritten = EJsonToken::SquareClose;
+}
+
 void FPrettyJsonWriter::WriteCommaAndNewlineIfNeeded()
 {
 	if (PreviousTokenWritten != EJsonToken::CurlyOpen && 
diff --git a/Engine/Source/Runtime/Experimental/JsonObjectGraph/Private/PrettyJsonWriter.h b/Engine/Source/Runtime/Experimental/JsonObjectGraph/Private/PrettyJsonWriter.h
--- a/Engine/Source/Runtime/Experimental/JsonObjectGraph/Private/PrettyJsonWriter.h
+++ b/Engine/Source/Runtime/Experimental/JsonObjectGraph/Private/PrettyJsonWriter.h
@@ -51,6 +51,7 @@ struct FPrettyJsonWriter : public TJsonWriter<UTF8CHAR, TPrettyJsonPrintPolicySi
 	void WriteNewlineAndArrayEnd();
 	void WriteLineTerminator();
 	void HACK_SetPreviousTokenWritten();
+	void HACK_SetPreviousTokenWrittenSquareClose();
 
 	void WriteCommaAndNewlineIfNeeded();
 };
