Changed I/O store on-demand mount request to require a host group that holds one or more URLs.
- added a multi endpoint HTTP client to enable fetching content from multiple endpoints with multiple URLs.
- the IAS backend is unaffected by this change but will be updated in a later change.

#rnx

diff --git a/Engine/Source/Runtime/Experimental/IoStore/OnDemand/Private/IoStoreOnDemand.cpp b/Engine/Source/Runtime/Experimental/IoStore/OnDemand/Private/IoStoreOnDemand.cpp
--- a/Engine/Source/Runtime/Experimental/IoStore/OnDemand/Private/IoStoreOnDemand.cpp
+++ b/Engine/Source/Runtime/Experimental/IoStore/OnDemand/Private/IoStoreOnDemand.cpp
@@ -1323,27 +1323,38 @@ void FIoStoreOnDemandModule::InitializeInternal()
 	}
 	else if (!EndpointConfig.ServiceUrls.IsEmpty() && !EndpointConfig.TocPath.IsEmpty())
 	{
-		const FString TocUrl = EndpointConfig.ServiceUrls[0] / EndpointConfig.TocPath;
-		MountArgs.Emplace(FOnDemandMountArgs
+		TIoStatusOr<FOnDemandHostGroup> HostGroup = FOnDemandHostGroup::Create(EndpointConfig.ServiceUrls);
+		if (HostGroup.IsOk())
 		{
-			.MountId = TocUrl,
-			.Url = TocUrl,
-			.Options = EOnDemandMountOptions::StreamOnDemand
-		});
+			MountArgs.Emplace(FOnDemandMountArgs
+			{
+				.MountId = EndpointConfig.ServiceUrls[0] / EndpointConfig.TocPath,
+				.TocRelativeUrl = EndpointConfig.TocPath,
+				.HostGroup = HostGroup.ConsumeValueOrDie(),
+				.Options = EOnDemandMountOptions::StreamOnDemand
+			});
+		}
+		else
+		{
+			UE_LOG(LogIas, Error, TEXT("Failed to create host group, reason: '%s'"), *HostGroup.Status().ToString());
+		}
 	}
 
 #if !UE_BUILD_SHIPPING
 	TOptional<FOnDemandInstallArgs> InstallArgs;
 	if (FParse::Param(FCommandLine::Get(), TEXT("Iad")))
 	{
+		TIoStatusOr<FOnDemandHostGroup> HostGroup = FOnDemandHostGroup::Create(EndpointConfig.ServiceUrls);
+		check(HostGroup.IsOk());
+
 		if (MountArgs)
 		{
+			MountArgs.GetValue().HostGroup = HostGroup.ConsumeValueOrDie(); 
 			MountArgs.GetValue().Options = EOnDemandMountOptions::InstallOnDemand;
-			MountArgs.GetValue().Url = EndpointConfig.ServiceUrls[0] / EndpointConfig.TocPath;
+			MountArgs.GetValue().TocRelativeUrl = EndpointConfig.TocPath;
 
 			static FOnDemandContentHandle ContentHandle = FOnDemandContentHandle::Create(TEXT("AllContent"));
 			InstallArgs.Emplace();
-			InstallArgs->Url = EndpointConfig.ServiceUrls[0] / EndpointConfig.TocPath;
 			InstallArgs->MountId = MountArgs.GetValue().MountId;
 			InstallArgs->ContentHandle = ContentHandle;
 		}
diff --git a/Engine/Source/Runtime/Experimental/IoStore/OnDemand/Private/OnDemandHostGroup.cpp b/Engine/Source/Runtime/Experimental/IoStore/OnDemand/Private/OnDemandHostGroup.cpp
new file mode 100644
--- /dev/null
+++ b/Engine/Source/Runtime/Experimental/IoStore/OnDemand/Private/OnDemandHostGroup.cpp
@@ -0,0 +1,179 @@
+// Copyright Epic Games, Inc. All Rights Reserved.
+
+#include "IO/OnDemandHostGroup.h"
+
+#include "Containers/AnsiString.h"
+#include "Containers/Array.h"
+#include "Containers/StringView.h"
+#include "IO/Http/Client.h"
+
+namespace UE::IoStore
+{
+
+namespace Private
+{
+
+static bool ValidateUrl(FAnsiStringView Url, FString& Reason)
+{
+	HTTP::FConnectionPool::FParams Params;
+	if (Params.SetHostFromUrl(Url) >= 0)
+	{
+		return true;
+	}
+
+	Reason = FString::Printf(TEXT("Invalid URL '%s'"), *FString(Url));
+	return false;
+}
+
+} // namespace Private
+
+struct FOnDemandHostGroup::FImpl
+{
+	TArray<FAnsiString> HostUrls;
+	int32				PrimaryIndex = INDEX_NONE;
+};
+
+FOnDemandHostGroup::FOnDemandHostGroup()
+	: Impl(new FImpl())
+{
+}
+
+FOnDemandHostGroup::FOnDemandHostGroup(FOnDemandHostGroup::FSharedImpl&& InImpl)
+	: Impl(MoveTemp(InImpl))
+{
+}
+
+FOnDemandHostGroup::~FOnDemandHostGroup()
+{
+}
+
+TConstArrayView<FAnsiString> FOnDemandHostGroup::Hosts() const
+{
+	return Impl->HostUrls;
+}
+
+FAnsiStringView FOnDemandHostGroup::Host(int32 Index) const
+{
+	return Impl->HostUrls.IsEmpty() ? FAnsiStringView() : Impl->HostUrls[Index];
+}
+
+FAnsiStringView FOnDemandHostGroup::CycleHost(int32& InOutIndex)
+{
+	InOutIndex = (InOutIndex + 1) % Impl->HostUrls.Num();
+	return Host(InOutIndex);
+}
+
+void FOnDemandHostGroup::SetPrimaryHost(int32 Index)
+{
+	check(Index >= 0 && Index < Impl->HostUrls.Num());
+	Impl->PrimaryIndex = Index;
+}
+
+FAnsiStringView FOnDemandHostGroup::PrimaryHost() const
+{
+	if (Impl->PrimaryIndex != INDEX_NONE)
+	{
+		check(Impl->PrimaryIndex >= 0 && Impl->PrimaryIndex < Impl->HostUrls.Num());
+		return Impl->HostUrls[Impl->PrimaryIndex];
+	}
+
+	return FAnsiStringView();
+}
+
+int32 FOnDemandHostGroup::PrimaryHostIndex() const
+{
+	return Impl->PrimaryIndex;
+}
+
+bool FOnDemandHostGroup::IsEmpty() const
+{
+	return Impl->HostUrls.IsEmpty();
+}
+
+TIoStatusOr<FOnDemandHostGroup> FOnDemandHostGroup::Create(FAnsiStringView Url)
+{
+	if (Url.EndsWith('/'))
+	{
+		Url.RemoveSuffix(1);
+	}
+
+	FString Reason;
+	if (Private::ValidateUrl(Url, Reason) == false)
+	{
+		FIoStatus(EIoErrorCode::InvalidParameter, Reason);
+	}
+
+	FSharedImpl Impl(new FImpl());
+	Impl->HostUrls.Add(FAnsiString(Url));
+	Impl->PrimaryIndex = 0;
+
+	return FOnDemandHostGroup(MoveTemp(Impl)); 
+}
+
+TIoStatusOr<FOnDemandHostGroup> FOnDemandHostGroup::Create(FStringView Url)
+{
+	return Create(FAnsiString(StringCast<ANSICHAR>(Url.GetData(), Url.Len()))); 
+}
+
+TIoStatusOr<FOnDemandHostGroup> FOnDemandHostGroup::Create(TConstArrayView<FAnsiString> Urls)
+{
+	if (Urls.IsEmpty())
+	{
+		return FIoStatus(EIoErrorCode::InvalidParameter, TEXT("No URLs specified"));
+	}
+
+	FSharedImpl Impl(new FImpl());
+	FString		Reason;
+
+	for (const FAnsiString& Url : Urls)
+	{
+		FAnsiStringView UrlView = Url;
+
+		if (UrlView.EndsWith('/'))
+		{
+			UrlView.RemoveSuffix(1);
+		}
+
+		if (Private::ValidateUrl(UrlView, Reason) == false)
+		{
+			return FIoStatus(EIoErrorCode::InvalidParameter, Reason);
+		}
+		Impl->HostUrls.Emplace(UrlView);
+	}
+	Impl->PrimaryIndex = 0; 
+
+	return FOnDemandHostGroup(MoveTemp(Impl)); 
+}
+
+TIoStatusOr<FOnDemandHostGroup> FOnDemandHostGroup::Create(TConstArrayView<FString> Urls)
+{
+	if (Urls.IsEmpty())
+	{
+		return FIoStatus(EIoErrorCode::InvalidParameter, TEXT("No URLs specified"));
+	}
+
+	FSharedImpl Impl(new FImpl());
+	FString		Reason;
+
+	for (const FString& Url : Urls)
+	{
+		FAnsiString AnsiUrl(StringCast<ANSICHAR>(*Url));
+		FAnsiStringView UrlView = AnsiUrl;
+
+		if (UrlView.EndsWith('/'))
+		{
+			UrlView.RemoveSuffix(1);
+		}
+
+		if (Private::ValidateUrl(AnsiUrl, Reason) == false)
+		{
+			return FIoStatus(EIoErrorCode::InvalidParameter, Reason);
+		}
+		Impl->HostUrls.Emplace(AnsiUrl);
+	}
+	Impl->PrimaryIndex = 0; 
+
+	return FOnDemandHostGroup(MoveTemp(Impl)); 
+}
+
+} //namespace UE::IoStore
diff --git a/Engine/Source/Runtime/Experimental/IoStore/OnDemand/Private/OnDemandHttpClient.cpp b/Engine/Source/Runtime/Experimental/IoStore/OnDemand/Private/OnDemandHttpClient.cpp
--- a/Engine/Source/Runtime/Experimental/IoStore/OnDemand/Private/OnDemandHttpClient.cpp
+++ b/Engine/Source/Runtime/Experimental/IoStore/OnDemand/Private/OnDemandHttpClient.cpp
@@ -327,4 +327,203 @@ TIoStatusOr<FIoBuffer> FHttpClient::Get(FStringView Url, uint32 RetryCount, EHtt
 	return Get(AnsiUrl, RetryCount, Redirects);
 }
 
+
+FMultiEndpointHttpClient::FMultiEndpointHttpClient(const FMultiEndpointHttpClientConfig& InConfig)
+	: Config(InConfig)
+{
+}
+
+TUniquePtr<FMultiEndpointHttpClient> FMultiEndpointHttpClient::Create(const FMultiEndpointHttpClientConfig& Config)
+{
+	return TUniquePtr<FMultiEndpointHttpClient>(new FMultiEndpointHttpClient(Config));
+}
+
+void FMultiEndpointHttpClient::Get(const FOnDemandHostGroup& Endpoint, FAnsiStringView RelativeUrl, FOnHttpResponse&& OnResponse)
+{
+	FConnection& Connection = GetConnection(Endpoint);
+
+	IssueRequest(FRequest
+	{
+		.OnResponse		= MoveTemp(OnResponse),
+		.RelativeUrl	= FAnsiString(RelativeUrl),
+		.Connection		= Connection,
+		.StartTime		= FPlatformTime::Cycles64(),
+		.Host			= Connection.CurrentHost
+	});
+}
+
+bool FMultiEndpointHttpClient::Tick(int32 WaitTimeMs, uint32 MaxKiBPerSecond)
+{
+	EventLoop.Throttle(MaxKiBPerSecond);
+	const uint32 TicketCount = EventLoop.Tick(WaitTimeMs);
+
+	if (Retries.IsEmpty() == false)
+	{
+		const int32 RequestCount = FMath::Min(Retries.Num(), int32(HTTP::FEventLoop::MaxActiveTickets - TicketCount));
+		for (int32 Idx = 0; Idx < RequestCount; Idx++)
+		{
+			IssueRequest(MoveTemp(Retries[Idx]));
+		}
+		Retries.RemoveAtSwap(0, RequestCount);
+	}
+
+	const bool bIsIdle = EventLoop.IsIdle();
+	if (bIsIdle)
+	{
+		// Destroy all non active connection pool(s)
+		for (TUniquePtr<FConnection>& Connection : Connections)
+		{
+			for (int32 Idx = 0, Count = Connection->Pools.Num(); Idx < Count; ++Idx)
+			{
+				if (Idx != Connection->CurrentHost)
+				{
+					Connection->Pools[Idx].Reset();
+				}
+			}
+		}
+	}
+
+	return bIsIdle == false;
+}
+
+void FMultiEndpointHttpClient::IssueRequest(FRequest&& Request)
+{
+	using namespace UE::IoStore::HTTP;
+
+	check(Request.Connection.HostGroup.IsEmpty() == false);
+	check(Request.Connection.Pools[Request.Connection.CurrentHost].IsValid());
+
+	FAnsiStringView Url				= Request.RelativeUrl;
+	FConnectionPool& ConnectionPool = *Request.Connection.Pools[Request.Connection.CurrentHost];
+
+	auto HttpSink = [this, Request = MoveTemp(Request)](const FTicketStatus& TicketStatus) mutable
+	{
+		const uint64 DurationMs = uint64(FPlatformTime::ToMilliseconds64(FPlatformTime::Cycles64() - Request.StartTime));
+
+		switch (TicketStatus.GetId())
+		{
+			case FTicketStatus::EId::Response:
+			{
+				FResponse& HttpResponse = TicketStatus.GetResponse();
+				Request.StatusCode		= HttpResponse.GetStatusCode();
+				HttpResponse.SetDestination(&Request.Body);
+				break;
+			}
+			case FTicketStatus::EId::Content:
+			case FTicketStatus::EId::Error:
+			case FTicketStatus::EId::Cancelled:
+			{
+				Log(Request);
+
+				const bool bError		= TicketStatus.GetId() == FTicketStatus::EId::Error;
+				const bool bServerError = Request.StatusCode > 499 && Request.StatusCode < 600;
+
+				if ((bError || bServerError) && Request.RetryCount <= Config.MaxRetryCount)
+				{
+					RetryRequest(MoveTemp(Request));
+				}
+				else
+				{
+					FMultiEndpointHttpClientResponse Response
+					{
+						.Body					= MoveTemp(Request.Body),
+						.DurationMilliseconds	= DurationMs,
+						.StatusCode				= Request.StatusCode,
+						.RetryCount				= Request.RetryCount,
+						.HostIndex				= Request.Host
+					};
+
+					if (bError)
+					{
+						Response.Reason = FString(TicketStatus.GetErrorReason());
+					}
+					else if (TicketStatus.GetId() == FTicketStatus::EId::Cancelled)
+					{
+						Response.Reason = TEXT("Cancelled");
+					}
+
+					FOnHttpResponse OnResponse = MoveTemp(Request.OnResponse);
+					OnResponse(MoveTemp(Response));
+				}
+				break;
+			}
+		}
+	};
+
+	EventLoop.Send(
+		EventLoop.Get(Url, ConnectionPool),
+		MoveTemp(HttpSink));
+}
+
+void FMultiEndpointHttpClient::RetryRequest(FRequest&& Request)
+{
+	FConnection& Connection = Request.Connection;
+
+	// Try a different host URL after the first retry
+	if (Request.RetryCount > 1 && Request.Host == Request.Connection.CurrentHost)
+	{
+		FAnsiStringView HostUrl = Request.Connection.HostGroup.CycleHost(Connection.CurrentHost);
+		if (Connection.Pools[Connection.CurrentHost].IsValid() == false)
+		{
+			HTTP::FConnectionPool::FParams Params;
+			ensure(Params.SetHostFromUrl(HostUrl) >= 0);
+			if (Config.ReceiveBufferSize >= 0)
+			{
+				Params.RecvBufSize = Config.ReceiveBufferSize;
+			}
+			Params.ConnectionCount = uint16(Config.MaxConnectionCount);
+			Connection.Pools[Connection.CurrentHost] = MakeUnique<HTTP::FConnectionPool>(Params);
+		}
+	}
+
+	Request.StatusCode = 0;
+	Request.RetryCount++;
+	Request.Host = Connection.CurrentHost;
+	Retries.Emplace(MoveTemp(Request));
+}
+
+FMultiEndpointHttpClient::FConnection& FMultiEndpointHttpClient::GetConnection(const FOnDemandHostGroup& HostGroup)
+{
+	for (TUniquePtr<FConnection>& Conn : Connections)
+	{
+		check(Conn.IsValid());
+		if (Conn->HostGroup == HostGroup)
+		{
+			return *Conn;
+		}
+	}
+
+	FConnection& Conn = *Connections.Emplace_GetRef(new FConnection
+	{
+		.HostGroup	= HostGroup,
+	});
+
+	Conn.Pools.SetNum(HostGroup.Hosts().Num());
+	Conn.CurrentHost = HostGroup.PrimaryHostIndex();
+
+	HTTP::FConnectionPool::FParams Params;
+	ensure(Params.SetHostFromUrl(HostGroup.PrimaryHost()) >= 0);
+	if (Config.ReceiveBufferSize >= 0)
+	{
+		Params.RecvBufSize = Config.ReceiveBufferSize;
+	}
+	Params.ConnectionCount = uint16(Config.MaxConnectionCount);
+
+	Conn.Pools[Conn.CurrentHost] = MakeUnique<HTTP::FConnectionPool>(Params);
+	return Conn;
+}
+
+void FMultiEndpointHttpClient::Log(const FRequest& Request)
+{
+	const uint64 DurationMs = (uint64)FPlatformTime::ToMilliseconds64(FPlatformTime::Cycles64() - Request.StartTime);
+
+	TStringBuilder<512> Sb;
+	Sb << TEXT("GET ")
+		<< Request.Connection.HostGroup.Host(Request.Host)
+		<< Request.RelativeUrl;
+	Sb.Appendf(TEXT(" (Status: %u, %llums, %.2lf KiB)"), Request.StatusCode, DurationMs, double(Request.Body.GetSize()) / 1024.0);
+
+	UE_LOG(LogIoStoreOnDemand, Log, TEXT("%s"), Sb.ToString());
+}
+
 } // namespace UE::IoStore
diff --git a/Engine/Source/Runtime/Experimental/IoStore/OnDemand/Private/OnDemandHttpClient.h b/Engine/Source/Runtime/Experimental/IoStore/OnDemand/Private/OnDemandHttpClient.h
--- a/Engine/Source/Runtime/Experimental/IoStore/OnDemand/Private/OnDemandHttpClient.h
+++ b/Engine/Source/Runtime/Experimental/IoStore/OnDemand/Private/OnDemandHttpClient.h
@@ -2,12 +2,14 @@
 
 #pragma once
 
+#include "Containers/AnsiString.h"
 #include "Containers/StringView.h"
 #include "Containers/UnrealString.h"
 #include "IO/Http/Client.h"
 #include "IO/IoBuffer.h"
 #include "IO/IoOffsetLength.h"
 #include "IO/IoStatus.h"
+#include "IO/OnDemandHostGroup.h"
 #include "Templates/Function.h"
 #include "Templates/UniquePtr.h"
 
@@ -78,4 +80,72 @@ private:
 	int32 CurrentEndpoint = INDEX_NONE;
 };
 
+struct FMultiEndpointHttpClientConfig
+{
+	int32			MaxConnectionCount = 4;
+	int32			ReceiveBufferSize = -1;
+	uint32			MaxRetryCount = 0;
+	EHttpRedirects	Redirects = EHttpRedirects::Follow;
+};
+
+struct FMultiEndpointHttpClientResponse
+{
+	const bool	IsOk() const { return StatusCode > 199 && StatusCode < 299; }
+
+	FIoBuffer	Body;
+	FString		Reason;
+	uint64		DurationMilliseconds = 0;
+	uint32		StatusCode = 0;
+	uint32		RetryCount = 0;
+	int32		HostIndex = INDEX_NONE;
+};
+
+class FMultiEndpointHttpClient
+{
+public:
+	UE_NONCOPYABLE(FMultiEndpointHttpClient);
+
+	using FOnHttpResponse = TFunction<void(FMultiEndpointHttpClientResponse&&)>;
+
+	static TUniquePtr<FMultiEndpointHttpClient> Create(const FMultiEndpointHttpClientConfig& Config);
+
+	void Get(const FOnDemandHostGroup& Endpoint, FAnsiStringView RelativeUrl, FOnHttpResponse&& OnResponse);
+	bool Tick(int32 WaitTimeMs, uint32 MaxKiBPerSecond);
+	bool Tick() { return Tick(-1, 0); }
+
+private:
+	using FHttpConnectionPools = TArray<TUniquePtr<HTTP::FConnectionPool>, TInlineAllocator<4>>;
+
+	struct FConnection
+	{
+		FOnDemandHostGroup		HostGroup;
+		FHttpConnectionPools	Pools;
+		int32					CurrentHost = INDEX_NONE;
+	};
+
+	struct FRequest
+	{
+		FOnHttpResponse		OnResponse;
+		FAnsiString			RelativeUrl;
+		FIoOffsetAndLength	Range;
+		FConnection&		Connection;
+		FIoBuffer			Body;
+		uint64				StartTime = 0;
+		uint32				RetryCount = 0;
+		uint32				StatusCode = 0;
+		int32				Host = INDEX_NONE;
+	};
+
+									FMultiEndpointHttpClient(const FMultiEndpointHttpClientConfig& Config);
+	void							IssueRequest(FRequest&& Request);
+	void							RetryRequest(FRequest&& Request);
+	FConnection&					GetConnection(const FOnDemandHostGroup& Endpoint);
+	void							Log(const FRequest& Request);
+
+	FMultiEndpointHttpClientConfig	Config;
+	TArray<TUniquePtr<FConnection>>	Connections;
+	HTTP::FEventLoop				EventLoop;
+	TArray<FRequest>				Retries;
+};
+
 } // namespace UE::IoStore
diff --git a/Engine/Source/Runtime/Experimental/IoStore/OnDemand/Private/OnDemandIoStore.cpp b/Engine/Source/Runtime/Experimental/IoStore/OnDemand/Private/OnDemandIoStore.cpp
--- a/Engine/Source/Runtime/Experimental/IoStore/OnDemand/Private/OnDemandIoStore.cpp
+++ b/Engine/Source/Runtime/Experimental/IoStore/OnDemand/Private/OnDemandIoStore.cpp
@@ -169,6 +169,23 @@ static FAnsiStringBuilderBase& GetChunkUrl(
 	return OutUrl;
 }
 
+static FAnsiStringBuilderBase& GetChunkUrl(
+	const FOnDemandContainer& Container,
+	const FOnDemandChunkEntry& Entry,
+	FAnsiStringBuilderBase& OutUrl)
+{
+	OutUrl.Reset();
+	if (!Container.ChunksDirectory.IsEmpty())
+	{
+		OutUrl << "/" << Container.ChunksDirectory;
+	}
+
+	const FString HashString = LexToString(Entry.Hash);
+	OutUrl << "/" << HashString.Left(2) << "/" << HashString << ".iochunk";
+
+	return OutUrl;
+}
+
 ///////////////////////////////////////////////////////////////////////////////
 static TIoStatusOr<FIoBuffer> DecodeChunk(const FOnDemandChunkInfo& ChunkInfo, FMemoryView EncodedChunk)
 {
@@ -1102,14 +1119,11 @@ FIoStatus FOnDemandIoStore::TickMountRequest(FMountRequest& MountRequest)
 
 	// Containers haven't been created yet, do it now
 	TArray<FSharedOnDemandContainer> RequestedContainers;
-	
-	FStringView Host, TocRelUrl;
-	Private::SplitHostUrl(Args.Url, Host, TocRelUrl);
-	const FStringView TocPath = FPathViews::GetPath(TocRelUrl);
+	const FStringView TocPath = FPathViews::GetExtension(Args.TocRelativeUrl).IsEmpty() ? Args.TocRelativeUrl : FPathViews::GetPath(Args.TocRelativeUrl);
 
 	if (Args.Toc)
 	{
-		CreateContainersFromToc(Args.MountId, TocPath, *Args.Toc, RequestedContainers);
+		CreateContainersFromToc(Args.MountId, TocPath, *Args.Toc, Args.HostGroup, RequestedContainers);
 	}
 	else if (Args.FilePath.IsEmpty() == false)
 	{
@@ -1125,15 +1139,23 @@ FIoStatus FOnDemandIoStore::TickMountRequest(FMountRequest& MountRequest)
 
 		Args.Toc = MakeUnique<FOnDemandToc>(TocStatus.ConsumeValueOrDie());
 
-		CreateContainersFromToc(Args.MountId, TocPath, *Args.Toc, RequestedContainers);
+		CreateContainersFromToc(Args.MountId, TocPath, *Args.Toc, Args.HostGroup, RequestedContainers);
 	}
-	else if (Args.Url.IsEmpty() == false)
+	else if (Args.HostGroup.IsEmpty() == false && Args.TocRelativeUrl.IsEmpty() == false)
 	{
-		UE_LOG(LogIoStoreOnDemand, Log, TEXT("Loading TOC from URL '%s'"), *Args.Url);
+		TStringBuilder<512> Url;
+		Url.Append(Args.HostGroup.PrimaryHost());
+		if (Args.TocRelativeUrl[0] != TEXT('/'))
+		{
+			Url.AppendChar(TEXT('/'));
+		}
+		Url.Append(Args.TocRelativeUrl);
+
+		UE_LOG(LogIoStoreOnDemand, Log, TEXT("Loading TOC from URL '%s'"), Url.ToString());
 
 		const uint32 RetryCount = 2;
 		const bool bFollowRedirects = true;
-		TIoStatusOr<FOnDemandToc> TocStatus = FOnDemandToc::LoadFromUrl(Args.Url, 2, bFollowRedirects);
+		TIoStatusOr<FOnDemandToc> TocStatus = FOnDemandToc::LoadFromUrl(Url, 2, bFollowRedirects);
 
 		if (!TocStatus.IsOk())
 		{
@@ -1142,7 +1164,7 @@ FIoStatus FOnDemandIoStore::TickMountRequest(FMountRequest& MountRequest)
 
 		Args.Toc = MakeUnique<FOnDemandToc>(TocStatus.ConsumeValueOrDie());
 
-		CreateContainersFromToc(Args.MountId, TocPath, *Args.Toc, RequestedContainers);
+		CreateContainersFromToc(Args.MountId, TocPath, *Args.Toc, Args.HostGroup, RequestedContainers);
 	}
 	
 	TAnsiStringBuilder<512> ChunkUrlBuilder;
@@ -1177,6 +1199,11 @@ FIoStatus FOnDemandIoStore::TickMountRequest(FMountRequest& MountRequest)
 		// Try fetch and deserialize the container header if the container is used for installing content.
 		if (EnumHasAnyFlags(Container->Flags, EOnDemandContainerFlags::InstallOnDemand))
 		{
+			if (Container->HostGroup.IsEmpty())
+			{
+				return FIoStatusBuilder(EIoErrorCode::InvalidParameter) << TEXT("No valid endpoint(s)");
+			}
+
 			const FIoChunkId ChunkId = CreateContainerHeaderChunkId(Container->ContainerId);
 			if (const FOnDemandChunkInfo ChunkInfo = FOnDemandChunkInfo::Find(Container, ChunkId); ChunkInfo.IsValid())
 			{
@@ -1185,7 +1212,7 @@ FIoStatus FOnDemandIoStore::TickMountRequest(FMountRequest& MountRequest)
 
 				const uint32 RetryCount = 2;
 				TIoStatusOr<FIoBuffer> ContainerHeaderChunk = FHttpClient::Get(
-					Private::GetChunkUrl(Host, *Container, ChunkInfo.ChunkEntry(),
+					Private::GetChunkUrl(FString(Args.HostGroup.PrimaryHost()), *Container, ChunkInfo.ChunkEntry(),
 					ChunkUrlBuilder).ToView(),
 					RetryCount,
 					EHttpRedirects::Follow);
@@ -1301,9 +1328,6 @@ FOnDemandInstallResult FOnDemandIoStore::TickInstallRequest(const FInstallReques
 	}
 
 	TAnsiStringBuilder<512> ChunkUrl;
-	FStringView Host, TocRelUrl;
-	Private::SplitHostUrl(InstallRequest.Args.Url, Host, TocRelUrl);
-
 	TSet<FSharedOnDemandContainer> ContainersForInstallation;
 	TSet<FPackageId> PackageIdsToInstall;
 
@@ -1389,12 +1413,11 @@ FOnDemandInstallResult FOnDemandIoStore::TickInstallRequest(const FInstallReques
 	const int32 MaxConcurrentRequests = 16;
 	int32		ConcurrentRequests = 0;
 
-	FHttpClientConfig HttpConfig;
+	FMultiEndpointHttpClientConfig HttpConfig;
 	HttpConfig.MaxConnectionCount = 8;
 	HttpConfig.MaxRetryCount = 2;
-	HttpConfig.Endpoints.Add(FString(Host));
 
-	TUniquePtr<FHttpClient> HttpClient = FHttpClient::Create(MoveTemp(HttpConfig));
+	TUniquePtr<FMultiEndpointHttpClient> HttpClient = FMultiEndpointHttpClient::Create(MoveTemp(HttpConfig));
 	if (HttpClient.IsValid() == false)
 	{
 		Status = FIoStatusBuilder(EIoErrorCode::InvalidCode) << TEXT("Failed to initialize HTTP client");
@@ -1460,9 +1483,10 @@ FOnDemandInstallResult FOnDemandIoStore::TickInstallRequest(const FInstallReques
 
 			ConcurrentRequests++;
 			HttpClient->Get(
-				Private::GetChunkUrl(FStringView(), Container, ChunkEntry, ChunkUrl).ToView(),
+				Container.HostGroup,
+				Private::GetChunkUrl(Container, ChunkEntry, ChunkUrl).ToView(),
 				[this, &InstallRequest, &OutResult, ChunkId, ChunkEntry, &ConcurrentRequests]
-				(TIoStatusOr<FIoBuffer> ChunkStatus, uint64 DurationMs) mutable
+				(FMultiEndpointHttpClientResponse&& HttpResponse)
 				{
 					ConcurrentRequests--;
 
@@ -1475,14 +1499,14 @@ FOnDemandInstallResult FOnDemandIoStore::TickInstallRequest(const FInstallReques
 						return;
 					}
 					
-					if (!ChunkStatus.IsOk())
+					if (!HttpResponse.IsOk())
 					{
-						Status = ChunkStatus.Status();
+						Status = FIoStatus(EIoErrorCode::ReadError);
 						return;
 					}
 
-					FIoBuffer Chunk			= ChunkStatus.ConsumeValueOrDie();
-					const FIoHash ChunkHash = FIoHash::HashBuffer(Chunk.GetView());
+					FIoBuffer& Chunk		= HttpResponse.Body;
+					const FIoHash ChunkHash	= FIoHash::HashBuffer(Chunk.GetView());
 					if (ChunkHash != ChunkEntry.Hash)
 					{
 						Status = FIoStatus(EIoErrorCode::ReadError, TEXTVIEW("Hash mismatch"));
@@ -1703,6 +1727,7 @@ void FOnDemandIoStore::CreateContainersFromToc(
 	FStringView MountId,
 	FStringView TocPath,
 	FOnDemandToc& Toc,
+	FOnDemandHostGroup HostGroup,
 	TArray<FSharedOnDemandContainer>& Out)
 {
 	static FName AssetClassName(TEXT("OnDemandIoStore"));
@@ -1742,6 +1767,7 @@ void FOnDemandIoStore::CreateContainersFromToc(
 	for (FOnDemandTocContainerEntry& ContainerEntry : Toc.Containers)
 	{
 		FSharedOnDemandContainer Container = MakeShared<FOnDemandContainer>();
+		Container->HostGroup			= HostGroup;
 		Container->Name					= MoveTemp(ContainerEntry.ContainerName);
 		Container->MountId				= MountId;
 		Container->ChunksDirectory		= StringCast<ANSICHAR>(ChunksDirectory.GetData(), ChunksDirectory.Len());
diff --git a/Engine/Source/Runtime/Experimental/IoStore/OnDemand/Private/OnDemandIoStore.h b/Engine/Source/Runtime/Experimental/IoStore/OnDemand/Private/OnDemandIoStore.h
--- a/Engine/Source/Runtime/Experimental/IoStore/OnDemand/Private/OnDemandIoStore.h
+++ b/Engine/Source/Runtime/Experimental/IoStore/OnDemand/Private/OnDemandIoStore.h
@@ -108,6 +108,7 @@ struct FOnDemandContainer
 
 	FAES::FAESKey							EncryptionKey;
 	FSharedContainerHeader					Header;
+	FOnDemandHostGroup						HostGroup;
 	FString									EncryptionKeyGuid;
 	FString									Name;
 	FString									MountId;
@@ -331,6 +332,7 @@ private:
 								FStringView MountId,
 								FStringView TocPath,
 								FOnDemandToc& Toc,
+								FOnDemandHostGroup HostGroup,
 								TArray<FSharedOnDemandContainer>& Out);
 
 	FSharedInstallCache					InstallCache;
diff --git a/Engine/Source/Runtime/Experimental/IoStore/OnDemand/Public/IO/IoStoreOnDemand.h b/Engine/Source/Runtime/Experimental/IoStore/OnDemand/Public/IO/IoStoreOnDemand.h
--- a/Engine/Source/Runtime/Experimental/IoStore/OnDemand/Public/IO/IoStoreOnDemand.h
+++ b/Engine/Source/Runtime/Experimental/IoStore/OnDemand/Public/IO/IoStoreOnDemand.h
@@ -9,8 +9,8 @@
 #include "IO/IoChunkId.h"
 #include "IO/IoContainerId.h"
 #include "IO/IoHash.h"
+#include "IO/OnDemandHostGroup.h"
 #include "IO/PackageId.h"
-#include "IO/IoStatus.h"
 #include "Misc/EnumClassFlags.h"
 #include "Misc/Guid.h"
 #include "Modules/ModuleInterface.h"
@@ -273,17 +273,33 @@ enum class EOnDemandMountOptions
 };
 ENUM_CLASS_FLAGS(EOnDemandMountOptions);
 
-/** Arguments for mounting on-demand container TOC(s). */
+/** Arguments for mounting on-demand container TOC(s).
+ *
+ * On-demand content can be mounted by providing:
+ * 1. a serialized TOC
+ * 2. a filepath to a TOC on disk
+ * 3. a URL from where to fetch the TOC using HTTP
+ *
+ * The chunk URLs are derived from the provided TOC URL or from the
+ * serialized chunks directory property in the TOC, i.e. if the
+ * TocRelativeUrl is not specified the TOC ChunksDirectory property
+ * needs to form a qualified path from the host.
+ * Example:
+ * http(s)://<Host>/<TocRelativePath>/chunks/<1-Byte Hex>/<hash>.iochunk
+ * http(s)://<Host>/<TOC.ChunksDirectory>/chunks/<1-Byte Hex>/<hash>.iochunk
+ */
 struct FOnDemandMountArgs
 {
 	/** Mount an already serialized TOC. */
 	TUniquePtr<FOnDemandToc> Toc;
 	/** Mandatory ID to be used for unmounting all container file(s) included in the TOC. */
 	FString MountId;
-	/** Download the TOC from the specified URL. */
-	FString Url;
+	/** Relative URL from the primary endpoint from where to download the TOC. */
+	FString TocRelativeUrl;
 	/** Serialize the TOC from the specified file path. */
 	FString FilePath;
+	/** List of URLs from where to download the chunks. */
+	FOnDemandHostGroup HostGroup;
 	/** Mount options. */
 	EOnDemandMountOptions Options = EOnDemandMountOptions::StreamOnDemand;
 };
@@ -321,8 +337,6 @@ struct FOnDemandInstallArgs
 	TArray<FString> TagSets;
 	/** Package ID's to install. */
 	TArray<FPackageId> PackageIds;
-	/** URL from where to download the chunks. */
-	FString Url;
 	/** Content handle. */
 	FOnDemandContentHandle ContentHandle;
 	/** Install options. */
diff --git a/Engine/Source/Runtime/Experimental/IoStore/OnDemand/Public/IO/OnDemandHostGroup.h b/Engine/Source/Runtime/Experimental/IoStore/OnDemand/Public/IO/OnDemandHostGroup.h
new file mode 100644
--- /dev/null
+++ b/Engine/Source/Runtime/Experimental/IoStore/OnDemand/Public/IO/OnDemandHostGroup.h
@@ -0,0 +1,69 @@
+// Copyright Epic Games, Inc. All Rights Reserved.
+
+#pragma once
+
+#include "CoreTypes.h"
+#include "Containers/AnsiString.h"
+#include "Containers/StringFwd.h"
+#include "IO/IoStatus.h"
+#include "Templates/SharedPointer.h"
+
+#define UE_API IOSTOREONDEMAND_API
+
+namespace UE::IoStore
+{
+
+/** Holds a set of host URLs. */
+class FOnDemandHostGroup
+{
+	struct FImpl;
+	using FSharedImpl = TSharedPtr<FImpl, ESPMode::ThreadSafe>;
+
+public:
+	/** Creates a new empty host group. */
+	UE_API FOnDemandHostGroup();
+	/** Destructor. */ 
+	UE_API ~FOnDemandHostGroup();
+	/** Returns the list of available host URLs. */ 
+	UE_API TConstArrayView<FAnsiString> Hosts() const;
+	/** Get the URL at the specified index. */
+	UE_API FAnsiStringView Host(int32 Index) const;
+	/** Get the next available host starting from the specified index. */
+	UE_API FAnsiStringView CycleHost(int32& InOutIndex);
+	/** Set the primary host URL. */
+	UE_API void SetPrimaryHost(int32 Index);
+	/** Get the primary host URL. */
+	UE_API FAnsiStringView PrimaryHost() const;
+	/** Get the primary host index. */
+	UE_API int32 PrimaryHostIndex() const;
+	/** Returns whether the group is empty or not. */
+	UE_API bool IsEmpty() const;
+	/** Create a new host group with the specified URL. */
+	UE_API static TIoStatusOr<FOnDemandHostGroup> Create(FAnsiStringView Url);
+	/** Create a new host group with the specified URL. */
+	UE_API static TIoStatusOr<FOnDemandHostGroup> Create(FStringView Url);
+	/** Create a new host group with the specified URLs. */
+	UE_API static TIoStatusOr<FOnDemandHostGroup> Create(TConstArrayView<FAnsiString> Urls);
+	/** Create a new host group with the specified URLs. */
+	UE_API static TIoStatusOr<FOnDemandHostGroup> Create(TConstArrayView<FString> Urls);
+
+	/** Returns whether two endpoint instance are the same or not. */
+	friend inline bool operator==(const FOnDemandHostGroup& Lhs, const FOnDemandHostGroup& Rhs)
+	{
+		return Lhs.Impl.Get() == Rhs.Impl.Get();
+	}
+	/** Returns the type hash of the endpoint. */ 
+	friend inline uint32 GetTypeHash(const FOnDemandHostGroup& Endpoint)
+	{
+		return *reinterpret_cast<const uint32*>(Endpoint.Impl.Get());
+	}
+
+private:
+	FOnDemandHostGroup(FSharedImpl&& Impl);
+
+	FSharedImpl Impl;
+};
+
+} // namespace UE::IoStore
+
+#undef UE_API
