Iris - Ensure fully push based classes remain so. The classes known to be fully push based is manually configured in DefaultEngine.ini via the ReplicationStateDescriptorConfig. See example below.

[/Script/IrisCore.ReplicationStateDescriptorConfig]
bEnsureAllClassesAreFullyPushModel=false
+EnsureFullyPushModelClassNames=(ClassName=StaticMeshComponent)

diff --git a/Engine/Config/BaseEngine.ini b/Engine/Config/BaseEngine.ini
--- a/Engine/Config/BaseEngine.ini
+++ b/Engine/Config/BaseEngine.ini
@@ -1432,6 +1432,24 @@ DefaultSpatialFilterName=Spatial
 +CriticalClassConfigs=(ClassName=/Script/Engine.PlayerController)
 
 [/Script/IrisCore.ReplicationStateDescriptorConfig]
+bEnsureAllClassesAreFullyPushModel=false
+; Classes that are known to be fully push based and should remain so
++EnsureFullyPushModelClassNames=(ClassName=AudioComponent)
++EnsureFullyPushModelClassNames=(ClassName=BillboardComponent)
++EnsureFullyPushModelClassNames=(ClassName=BoxComponent)
++EnsureFullyPushModelClassNames=(ClassName=CapsuleComponent)
++EnsureFullyPushModelClassNames=(ClassName=NiagaraComponent)
++EnsureFullyPushModelClassNames=(ClassName=ParticleSystemComponent)
++EnsureFullyPushModelClassNames=(ClassName=PointLightComponent)
++EnsureFullyPushModelClassNames=(ClassName=PostProcessComponent)
++EnsureFullyPushModelClassNames=(ClassName=ProjectileMovementComponent)
++EnsureFullyPushModelClassNames=(ClassName=SceneComponent)
++EnsureFullyPushModelClassNames=(ClassName=SkeletalMeshComponent)
++EnsureFullyPushModelClassNames=(ClassName=SkyLightComponent)
++EnsureFullyPushModelClassNames=(ClassName=SplineComponent)
++EnsureFullyPushModelClassNames=(ClassName=StaticMeshComponent)
++EnsureFullyPushModelClassNames=(ClassName=VolumetricCloudComponent)
++EnsureFullyPushModelClassNames=(ClassName=WorldDataLayers)
 ; Declarate structs that are vetted to work using reflection based struct serialization even though there exists a custom NetSerialize function for the struct
 +SupportsStructNetSerializerList=(StructName=GameplayCueParameters)
 +SupportsStructNetSerializerList=(StructName=GameplayAbilityTargetData_LocationInfo)
diff --git a/Engine/Source/Runtime/Experimental/Iris/Core/Private/Iris/ReplicationState/ReplicationStateDescriptorBuilder.cpp b/Engine/Source/Runtime/Experimental/Iris/Core/Private/Iris/ReplicationState/ReplicationStateDescriptorBuilder.cpp
--- a/Engine/Source/Runtime/Experimental/Iris/Core/Private/Iris/ReplicationState/ReplicationStateDescriptorBuilder.cpp
+++ b/Engine/Source/Runtime/Experimental/Iris/Core/Private/Iris/ReplicationState/ReplicationStateDescriptorBuilder.cpp
@@ -2,6 +2,7 @@
 
 #include "Iris/ReplicationState/ReplicationStateDescriptorBuilder.h"
 #include "CoreTypes.h"
+#include "Containers/Set.h"
 #include "Iris/Core/BitTwiddling.h"
 #include "Iris/Core/IrisLog.h"
 #include "Iris/Core/IrisProfiler.h"
@@ -194,9 +195,12 @@ public:
 	static EMemberPropertyTraits GetHasObjectReferenceTraits(const FNetSerializer* NetSerializer);
 	static void GetIrisPropertyTraits(FMemberProperty& OutMemberProperty, const FProperty* Property, const TArray<FLifetimeProperty>* LifeTimeProperties, UClass* ObjectClass);
 
+	static bool ShouldValidateIsFullyPushModel(const UClass* Class);
+
 public:
 	enum class EDescriptorType : unsigned
 	{
+		Invalid,
 		Class,
 		Struct,
 		Function,
@@ -211,9 +215,10 @@ public:
 		const uint8* DefaultStateSourceData = nullptr;
 		// When building default state we need the ReplicationSystem in order to capture exports.
 		UReplicationSystem* ReplicationSystem = nullptr;
-		EDescriptorType DescriptorType;
-		uint32 bIsInitState : 1;
-		uint32 bAllMembersAreReplicated : 1;
+		EDescriptorType DescriptorType = EDescriptorType::Invalid;
+		bool bIsInitState = false;
+		bool bAllMembersAreReplicated = false;
+		bool bShouldValidateIsFullyPushModel = false;
 	};
 
 	void AddMemberProperty(const FMemberProperty& Info) { Members.Add(Info); }
@@ -365,6 +370,7 @@ private:
 	TRefCountPtr<const FReplicationStateDescriptor> CreateDescriptorForProperty(const FString& DescriptorName, const FProperty* Property, const FReplicationStateDescriptorBuilder::FParameters& Parameters) const;
 
 	static bool ValidateStructPropertyDescriptor(const FMemberProperty& MemberProperty, const FReplicationStateDescriptor* Descriptor);
+	void ValidateIsFullyPushModelClassDescriptor(const FBuilderContext& Context, const FReplicationStateDescriptor* Descriptor);
 
 	// Helpers
 	static void GetPropertyPathName(const FProperty* Property, FString& PathName);
@@ -373,8 +379,12 @@ private:
 	TArray<FMemberProperty> Members;
 	TArray<FMemberFunction> Functions;
 	FStructInfo StructInfo;
+
+	static TSet<FName> UniqueClassNames;
 };
 
+TSet<FName> FPropertyReplicationStateDescriptorBuilder::UniqueClassNames;
+
 TRefCountPtr<const FReplicationStateDescriptor> FPropertyReplicationStateDescriptorBuilder::GetDescriptorForStructProperty(const FMemberProperty& MemberProperty, FReplicationStateDescriptorRegistry* DescriptorRegistry) const
 {
 	const FStructProperty* StructProp = CastFieldChecked<const FStructProperty>(MemberProperty.Property);
@@ -452,6 +462,76 @@ bool FPropertyReplicationStateDescriptorBuilder::ValidateStructPropertyDescripto
 	return true;
 }
 
+bool FPropertyReplicationStateDescriptorBuilder::ShouldValidateIsFullyPushModel(const UClass* Class)
+{
+#if UE_BUILD_SHIPPING
+	return false;
+#else
+	// If we have processed this class before there's no need to verify it's push basedness again.
+	bool bAlreadyInSet = false;
+	UniqueClassNames.Emplace(Class->GetFName(), &bAlreadyInSet);
+	if (bAlreadyInSet)
+	{
+		return false;
+	}
+
+	const UReplicationStateDescriptorConfig* ReplicationStateDescriptorConfig = GetDefault<UReplicationStateDescriptorConfig>();
+	if (ReplicationStateDescriptorConfig->EnsureAllClassesAreFullyPushModel())
+	{
+		return true;
+	}
+
+	if (const FReplicationStateDescriptorClassPushModelConfig* ClassConfig = ReplicationStateDescriptorConfig->GetEnsureFullyPushModelClassNames().FindByPredicate([ClassName = Class->GetFName()](const FReplicationStateDescriptorClassPushModelConfig& ClassConfig) { return ClassConfig.ClassName == ClassName; }))
+	{
+		return true;
+	}
+
+ 	return false;
+#endif
+}
+
+void FPropertyReplicationStateDescriptorBuilder::ValidateIsFullyPushModelClassDescriptor(const FBuilderContext& Context, const FReplicationStateDescriptor* Descriptor)
+{
+	if (EnumHasAnyFlags(Descriptor->Traits, EReplicationStateTraits::HasPushBasedDirtiness))
+	{
+		return;
+	}
+
+	// A class and its descriptors is expected to be validated at most once. We want to report all of the failing descriptors because the user has explicitly requested us to.
+	ensureAlwaysMsgf(false, TEXT("ReplicationStateDescriptor %s for class %s does not have the HasPushBasedDirtiness trait. This prevents the class from being fully push based. See log for additional details."), ToCStr(Descriptor->DebugName), (Context.BuildParams.PathName ? Context.BuildParams.PathName->ToString() : TEXT("UNKNOWN")));
+
+	// Try to figure out the cause of the lack of push based trait
+	if (!Members.IsEmpty())
+	{
+		TStringBuilder<2048> Message;
+		Message.Appendf(TEXT("Members of ReplicationStateDescriptor %s lacking HasPushBasedDirtiness trait:\n"), ToCStr(Descriptor->DebugName));
+		for (const FMemberProperty& Member : Members)
+		{
+			if (!EnumHasAnyFlags(Member.Traits, EMemberPropertyTraits::HasPushBasedDirtiness))
+			{
+				if (Member.Property)
+				{
+					Message.Appendf(TEXT("Member property name: %s\n"), ToCStr(Member.Property->GetName()));
+				}
+				else if (const FReplicationStateMemberDebugDescriptor* DebugDescriptor = Descriptor->MemberDebugDescriptors)
+				{
+					const SIZE_T MemberIndex = &Member - Members.GetData();
+					Message.Appendf(TEXT("Member debug name: %s\n"), ToCStr(DebugDescriptor[MemberIndex].DebugName));
+				}
+				else
+				{
+					Message.Appendf(TEXT("Member index: %u\n"), unsigned(&Member - Members.GetData()));
+				}
+			}
+		}
+		UE_LOG_DESCRIPTORBUILDER_WARNING(TEXT("%s"), Message.ToString());
+	}
+	else
+	{
+		UE_LOG_DESCRIPTORBUILDER_WARNING(TEXT("ReplicationStateDescriptor %s for class %s does not have the HasPushBasedDirtiness trait. No additional info."), ToCStr(Descriptor->DebugName), Context.BuildParams.PathName ? Context.BuildParams.PathName->ToString() : TEXT(""));
+	}
+}
+
 
 TRefCountPtr<const FReplicationStateDescriptor> FPropertyReplicationStateDescriptorBuilder::GetDescriptorForArrayProperty(const FProperty* Property, FReplicationStateDescriptorRegistry* DescriptorRegistry) const
 {
@@ -1835,6 +1915,11 @@ FPropertyReplicationStateDescriptorBuilder::Build(const FString& StateName, FRep
 			}
 		}
 
+		if (Context.BuildParams.bShouldValidateIsFullyPushModel)
+		{
+			ValidateIsFullyPushModelClassDescriptor(Context, Descriptor);
+		}
+
 		return TRefCountPtr<const FReplicationStateDescriptor>(Descriptor);
 	}
 	else
@@ -2843,6 +2928,8 @@ SIZE_T FReplicationStateDescriptorBuilder::CreateDescriptorsForClass(FResult& Cr
 	// For-loop with CurrentClass initialized to InObjectClass causes static analysis warnings.
 	CA_ASSUME(InObjectClass != nullptr);
 
+	const bool bShouldValidateIsFullyPushModel = FPropertyReplicationStateDescriptorBuilder::ShouldValidateIsFullyPushModel(InObjectClass);
+
 	// Build state descriptors
 	FLazyGetPathNameHelper LazyPathName;
 
@@ -2859,6 +2946,8 @@ SIZE_T FReplicationStateDescriptorBuilder::CreateDescriptorsForClass(FResult& Cr
 			BuildParameters.bIsInitState = (BuilderTypeIndex == InitPropertyReplicationStateBuilderIndex);
 			// We set this to false for the time being for all replication states
 			BuildParameters.bAllMembersAreReplicated = false;
+			BuildParameters.bShouldValidateIsFullyPushModel = bShouldValidateIsFullyPushModel;
+
 			TRefCountPtr<const FReplicationStateDescriptor> Descriptor = Builders[BuilderTypeIndex].Build(InObjectClass->GetName() + InternalPropertyReplicationStateTypes[BuilderTypeIndex].NamePostFix, Parameters.DescriptorRegistry, BuildParameters);
 			CreatedDescriptors.Add(Descriptor);
 		}
diff --git a/Engine/Source/Runtime/Experimental/Iris/Core/Private/Iris/ReplicationState/ReplicationStateDescriptorConfig.cpp b/Engine/Source/Runtime/Experimental/Iris/Core/Private/Iris/ReplicationState/ReplicationStateDescriptorConfig.cpp
--- a/Engine/Source/Runtime/Experimental/Iris/Core/Private/Iris/ReplicationState/ReplicationStateDescriptorConfig.cpp
+++ b/Engine/Source/Runtime/Experimental/Iris/Core/Private/Iris/ReplicationState/ReplicationStateDescriptorConfig.cpp
@@ -14,3 +14,8 @@ TConstArrayView<FSupportsStructNetSerializerConfig> UReplicationStateDescriptorC
 {
 	return MakeArrayView(SupportsStructNetSerializerList);
 }
+
+TConstArrayView<FReplicationStateDescriptorClassPushModelConfig> UReplicationStateDescriptorConfig::GetEnsureFullyPushModelClassNames() const
+{
+	return MakeArrayView(EnsureFullyPushModelClassNames);
+}
diff --git a/Engine/Source/Runtime/Experimental/Iris/Core/Public/Iris/ReplicationState/ReplicationStateDescriptorConfig.h b/Engine/Source/Runtime/Experimental/Iris/Core/Public/Iris/ReplicationState/ReplicationStateDescriptorConfig.h
--- a/Engine/Source/Runtime/Experimental/Iris/Core/Public/Iris/ReplicationState/ReplicationStateDescriptorConfig.h
+++ b/Engine/Source/Runtime/Experimental/Iris/Core/Public/Iris/ReplicationState/ReplicationStateDescriptorConfig.h
@@ -6,6 +6,16 @@
 #include "UObject/ObjectMacros.h"
 #include "ReplicationStateDescriptorConfig.generated.h"
 
+USTRUCT()
+struct FReplicationStateDescriptorClassPushModelConfig
+{
+	GENERATED_BODY()
+	
+	/** Short class name, e.g. PlayerController. */
+	UPROPERTY()
+	FName ClassName;
+};
+
 USTRUCT()
 struct FSupportsStructNetSerializerConfig
 {
@@ -27,6 +37,10 @@ class UReplicationStateDescriptorConfig : public UObject
 
 public:
 	IRISCORE_API TConstArrayView<FSupportsStructNetSerializerConfig> GetSupportsStructNetSerializerList() const;
+	
+	IRISCORE_API TConstArrayView<FReplicationStateDescriptorClassPushModelConfig> GetEnsureFullyPushModelClassNames() const;
+
+	inline bool EnsureAllClassesAreFullyPushModel() const;
 
 protected:
 	UReplicationStateDescriptorConfig();
@@ -37,4 +51,17 @@ private:
 	 */
 	UPROPERTY(Config)
 	TArray<FSupportsStructNetSerializerConfig> SupportsStructNetSerializerList;
+
+	/** Which classes should ensure they are fully push model. */
+	UPROPERTY(Config)
+	TArray<FReplicationStateDescriptorClassPushModelConfig> EnsureFullyPushModelClassNames;
+
+	/** If you want to be alerted of all classes not being fully push model. */
+	UPROPERTY(Config)
+	bool bEnsureAllClassesAreFullyPushModel = false;
 };
+
+bool UReplicationStateDescriptorConfig::EnsureAllClassesAreFullyPushModel() const
+{
+	return bEnsureAllClassesAreFullyPushModel;
+}
