General IAD Support - IOStore On Demand Changes
- Allow mutiple mount requests with the same mount id for different containers. Each container can still only be mounted once.
- Fix bug with triming trailing slash

diff --git a/Engine/Source/Runtime/Experimental/IoStore/OnDemand/Private/OnDemandHostGroup.cpp b/Engine/Source/Runtime/Experimental/IoStore/OnDemand/Private/OnDemandHostGroup.cpp
--- a/Engine/Source/Runtime/Experimental/IoStore/OnDemand/Private/OnDemandHostGroup.cpp
+++ b/Engine/Source/Runtime/Experimental/IoStore/OnDemand/Private/OnDemandHostGroup.cpp
@@ -34,7 +34,7 @@ struct FOnDemandHostGroup::FImpl
 };
 
 FOnDemandHostGroup::FOnDemandHostGroup()
-	: Impl(new FImpl())
+	: Impl(MakeShared<FImpl>())
 {
 }
 
@@ -103,7 +103,7 @@ TIoStatusOr<FOnDemandHostGroup> FOnDemandHostGroup::Create(FAnsiStringView Url)
 		FIoStatus(EIoErrorCode::InvalidParameter, Reason);
 	}
 
-	FSharedImpl Impl(new FImpl());
+	FSharedImpl Impl = MakeShared<FImpl>();
 	Impl->HostUrls.Add(FAnsiString(Url));
 	Impl->PrimaryIndex = 0;
 
@@ -122,7 +122,7 @@ TIoStatusOr<FOnDemandHostGroup> FOnDemandHostGroup::Create(TConstArrayView<FAnsi
 		return FIoStatus(EIoErrorCode::InvalidParameter, TEXT("No URLs specified"));
 	}
 
-	FSharedImpl Impl(new FImpl());
+	FSharedImpl Impl = MakeShared<FImpl>();
 	FString		Reason;
 
 	for (const FAnsiString& Url : Urls)
@@ -152,28 +152,29 @@ TIoStatusOr<FOnDemandHostGroup> FOnDemandHostGroup::Create(TConstArrayView<FStri
 		return FIoStatus(EIoErrorCode::InvalidParameter, TEXT("No URLs specified"));
 	}
 
-	FSharedImpl Impl(new FImpl());
+	FSharedImpl Impl = MakeShared<FImpl>();
 	FString		Reason;
 
 	for (const FString& Url : Urls)
 	{
-		FAnsiString AnsiUrl(StringCast<ANSICHAR>(*Url));
-		FAnsiStringView UrlView = AnsiUrl;
+		FStringView UrlView = Url;
 
-		if (UrlView.EndsWith('/'))
+		if (UrlView.EndsWith(TEXT('/')))
 		{
 			UrlView.RemoveSuffix(1);
 		}
 
+		FAnsiString AnsiUrl(StringCast<ANSICHAR>(UrlView.GetData(), UrlView.Len()));
+
 		if (Private::ValidateUrl(AnsiUrl, Reason) == false)
 		{
 			return FIoStatus(EIoErrorCode::InvalidParameter, Reason);
 		}
-		Impl->HostUrls.Emplace(AnsiUrl);
+		Impl->HostUrls.Emplace(MoveTemp(AnsiUrl));
 	}
-	Impl->PrimaryIndex = 0; 
+	Impl->PrimaryIndex = 0;
 
-	return FOnDemandHostGroup(MoveTemp(Impl)); 
+	return FOnDemandHostGroup(MoveTemp(Impl));
 }
 
 } //namespace UE::IoStore
diff --git a/Engine/Source/Runtime/Experimental/IoStore/OnDemand/Private/OnDemandIoStore.cpp b/Engine/Source/Runtime/Experimental/IoStore/OnDemand/Private/OnDemandIoStore.cpp
--- a/Engine/Source/Runtime/Experimental/IoStore/OnDemand/Private/OnDemandIoStore.cpp
+++ b/Engine/Source/Runtime/Experimental/IoStore/OnDemand/Private/OnDemandIoStore.cpp
@@ -128,22 +128,6 @@ static bool HasInstallCacheDirectory()
 {
 	return FPaths::HasProjectPersistentDownloadDir();
 }
-///////////////////////////////////////////////////////////////////////////////
-static void SplitHostUrl(const FStringView& Url, FStringView& OutHost, FStringView& OutRemainder)
-{
-	OutHost = OutRemainder = FStringView();
-
-	if (Url.StartsWith(TEXTVIEW("http")))
-	{
-		int32 Delim = INDEX_NONE;
-		ensure(Url.FindChar(':', Delim));
-		const int32 ProtocolDelim = Delim + 3;
-		ensure(Url.RightChop(ProtocolDelim).FindChar('/', Delim));
-		OutHost = Url.Left(ProtocolDelim + Delim);
-	}
-
-	OutRemainder = Url.RightChop(OutHost.Len());
-}
 
 ///////////////////////////////////////////////////////////////////////////////
 static FAnsiStringBuilderBase& GetChunkUrl(
@@ -1087,19 +1071,19 @@ FIoStatus FOnDemandIoStore::TickMountRequest(FMountRequest& MountRequest)
 	}
 
 	bool bAnyPendingEncryptionKey = false;
-	bool bFoundContainers = false;
+	TArray<FString> ExistingContainers;
 
 	// Find containers matching the mount ID
 	{
 		UE::TUniqueLock Lock(ContainerMutex);
-		for (FSharedOnDemandContainer& Container : Containers)
+		for (const FSharedOnDemandContainer& Container : Containers)
 		{
 			if (Container->MountId != Args.MountId)
 			{
 				continue;
 			}
 			
-			bFoundContainers = true;
+			ExistingContainers.Add(Container->Name);
 			
 			if (EnumHasAnyFlags(Container->Flags, EOnDemandContainerFlags::PendingEncryptionKey))
 			{
@@ -1108,22 +1092,13 @@ FIoStatus FOnDemandIoStore::TickMountRequest(FMountRequest& MountRequest)
 		}
 	}
 
-	if (bFoundContainers)
-	{
-		if (bAnyPendingEncryptionKey)
-		{
-			return EIoErrorCode::PendingEncryptionKey;
-		}
-		return EIoErrorCode::Ok;
-	}
-
 	// Containers haven't been created yet, do it now
 	TArray<FSharedOnDemandContainer> RequestedContainers;
 	const FStringView TocPath = FPathViews::GetExtension(Args.TocRelativeUrl).IsEmpty() ? Args.TocRelativeUrl : FPathViews::GetPath(Args.TocRelativeUrl);
 
 	if (Args.Toc)
 	{
-		CreateContainersFromToc(Args.MountId, TocPath, *Args.Toc, Args.HostGroup, RequestedContainers);
+		CreateContainersFromToc(Args.MountId, TocPath, *Args.Toc, Args.HostGroup, ExistingContainers, RequestedContainers);
 	}
 	else if (Args.FilePath.IsEmpty() == false)
 	{
@@ -1139,7 +1114,7 @@ FIoStatus FOnDemandIoStore::TickMountRequest(FMountRequest& MountRequest)
 
 		Args.Toc = MakeUnique<FOnDemandToc>(TocStatus.ConsumeValueOrDie());
 
-		CreateContainersFromToc(Args.MountId, TocPath, *Args.Toc, Args.HostGroup, RequestedContainers);
+		CreateContainersFromToc(Args.MountId, TocPath, *Args.Toc, Args.HostGroup, ExistingContainers, RequestedContainers);
 	}
 	else if (Args.HostGroup.IsEmpty() == false && Args.TocRelativeUrl.IsEmpty() == false)
 	{
@@ -1164,7 +1139,7 @@ FIoStatus FOnDemandIoStore::TickMountRequest(FMountRequest& MountRequest)
 
 		Args.Toc = MakeUnique<FOnDemandToc>(TocStatus.ConsumeValueOrDie());
 
-		CreateContainersFromToc(Args.MountId, TocPath, *Args.Toc, Args.HostGroup, RequestedContainers);
+		CreateContainersFromToc(Args.MountId, TocPath, *Args.Toc, Args.HostGroup, ExistingContainers, RequestedContainers);
 	}
 	
 	TAnsiStringBuilder<512> ChunkUrlBuilder;
@@ -1501,7 +1476,9 @@ FOnDemandInstallResult FOnDemandIoStore::TickInstallRequest(const FInstallReques
 					
 					if (!HttpResponse.IsOk())
 					{
-						Status = FIoStatus(EIoErrorCode::ReadError);
+						Status = FIoStatus(EIoErrorCode::ReadError) << 
+							FString::Printf(TEXT("HTTP Status %u, "), HttpResponse.StatusCode) << 
+							HttpResponse.Reason;
 						return;
 					}
 
@@ -1728,6 +1705,7 @@ void FOnDemandIoStore::CreateContainersFromToc(
 	FStringView TocPath,
 	FOnDemandToc& Toc,
 	FOnDemandHostGroup HostGroup,
+	const TArray<FString>& ExistingContainers,
 	TArray<FSharedOnDemandContainer>& Out)
 {
 	static FName AssetClassName(TEXT("OnDemandIoStore"));
@@ -1766,6 +1744,11 @@ void FOnDemandIoStore::CreateContainersFromToc(
 
 	for (FOnDemandTocContainerEntry& ContainerEntry : Toc.Containers)
 	{
+		if (ExistingContainers.Contains(ContainerEntry.ContainerName))
+		{
+			continue;
+		}
+
 		FSharedOnDemandContainer Container = MakeShared<FOnDemandContainer>();
 		Container->HostGroup			= HostGroup;
 		Container->Name					= MoveTemp(ContainerEntry.ContainerName);
diff --git a/Engine/Source/Runtime/Experimental/IoStore/OnDemand/Private/OnDemandIoStore.h b/Engine/Source/Runtime/Experimental/IoStore/OnDemand/Private/OnDemandIoStore.h
--- a/Engine/Source/Runtime/Experimental/IoStore/OnDemand/Private/OnDemandIoStore.h
+++ b/Engine/Source/Runtime/Experimental/IoStore/OnDemand/Private/OnDemandIoStore.h
@@ -333,6 +333,7 @@ private:
 								FStringView TocPath,
 								FOnDemandToc& Toc,
 								FOnDemandHostGroup HostGroup,
+								const TArray<FString>& ExistingContainers,
 								TArray<FSharedOnDemandContainer>& Out);
 
 	FSharedInstallCache					InstallCache;
