Fix for GenerateCollisions resetting CCD Enabled Flag.

GenerateCollisionCCDImpl sets the CCD Enabled flag on the constraint. However, if CCD sweep is not enabled, GenerateCollision is called which will eventually reset this flag back to false even though CCD should be enabled. Removing this set to false fixes this issue.

diff --git a/Engine/Source/Runtime/Experimental/Chaos/Private/Chaos/Collision/ParticlePairMidPhase.cpp b/Engine/Source/Runtime/Experimental/Chaos/Private/Chaos/Collision/ParticlePairMidPhase.cpp
--- a/Engine/Source/Runtime/Experimental/Chaos/Private/Chaos/Collision/ParticlePairMidPhase.cpp
+++ b/Engine/Source/Runtime/Experimental/Chaos/Private/Chaos/Collision/ParticlePairMidPhase.cpp
@@ -477,10 +477,6 @@ namespace Chaos
 			Constraint->SetCullDistance(CullDistance);
 			Constraint->SetRelativeMovement(RelativeMovement);
 
-			// Constraint may have been previously used with CCD enabled (e.g., a midphase modifier)
-			// so we need to make sure that the CCD flag is disabled
-			Constraint->SetCCDEnabled(false);
-
 			// If the constraint was not used last frame, it needs to be reset, otherwise we will try to reuse
 			if (!bWasUpdatedLastTick || (Constraint->GetManifoldPoints().Num() == 0))
 			{
