Geometry collection : fix potential crash when using SetSimulatePhysics inside a blueprint construction script if the dynamic collection has not been created yet

diff --git a/Engine/Source/Runtime/Experimental/GeometryCollectionEngine/Private/GeometryCollection/GeometryCollectionComponent.cpp b/Engine/Source/Runtime/Experimental/GeometryCollectionEngine/Private/GeometryCollection/GeometryCollectionComponent.cpp
--- a/Engine/Source/Runtime/Experimental/GeometryCollectionEngine/Private/GeometryCollection/GeometryCollectionComponent.cpp
+++ b/Engine/Source/Runtime/Experimental/GeometryCollectionEngine/Private/GeometryCollection/GeometryCollectionComponent.cpp
@@ -1273,7 +1273,7 @@ void UGeometryCollectionComponent::SetSimulatePhysics(bool bEnabled)
 	if (bEnabled && !PhysicsProxy && RestCollection)
 	{
 		EnsurePhysicsStateCreated();
-		if (!PhysicsProxy && !RestCollection->IsEmpty())
+		if (!PhysicsProxy && !RestCollection->IsEmpty() && DynamicCollection)
 		{
 			RegisterAndInitializePhysicsProxy();
 		}
