[NetPhysics] RewindData to cache moving kinematic XR as it looked before solve, as done for dynamic particles.
Moving kinematic have their XR updated from kinematic target during intergration rather than at the end of the physics solve, so calculate how XR looked like before integrate by stepping them back once with V and W

diff --git a/Engine/Source/Runtime/Experimental/Chaos/Private/RewindData.cpp b/Engine/Source/Runtime/Experimental/Chaos/Private/RewindData.cpp
--- a/Engine/Source/Runtime/Experimental/Chaos/Private/RewindData.cpp
+++ b/Engine/Source/Runtime/Experimental/Chaos/Private/RewindData.cpp
@@ -1371,7 +1371,18 @@ void FRewindData::PushPTDirtyData(TPBDRigidParticleHandle<FReal, 3>& Handle, con
 	{
 		if (FParticlePositionRotation* PreXR = Latest.ParticlePositionRotation.WriteAccessNonDecreasing(FrameAndPhase, PropertiesPool))
 		{
-			PreXR->CopyFrom(Handle);
+			/** This is called post-solve but before PQ are applied on XR
+			* If this is a kinematic moving object XR have been updated already in the integrate step via KinematicTarget and the velocity has been updated based on the XR change
+			* Get the pre-solve state of moving kinematic particles by stepping their XR back one step via their velocities. */
+			if (Handle.IsMovingKinematic())
+			{
+				PreXR->SetX(Handle.GetX() - (Handle.GetV() * Solver->GetLastDt()));
+				PreXR->SetR(Chaos::FRotation3::IntegrateRotationWithAngularVelocity(Handle.GetR(), Handle.GetWf(), -Solver->GetLastDt()));
+			}
+			else
+			{
+				PreXR->CopyFrom(Handle);
+			}
 		}
 	}
 
