Delete IAS cache directory if the game hasn't been configured to use IAS.

#rnx

diff --git a/Engine/Source/Runtime/Experimental/IoStore/OnDemand/Private/IoStoreOnDemand.cpp b/Engine/Source/Runtime/Experimental/IoStore/OnDemand/Private/IoStoreOnDemand.cpp
--- a/Engine/Source/Runtime/Experimental/IoStore/OnDemand/Private/IoStoreOnDemand.cpp
+++ b/Engine/Source/Runtime/Experimental/IoStore/OnDemand/Private/IoStoreOnDemand.cpp
@@ -1238,6 +1238,19 @@ void FIoStoreOnDemandModule::InitializeInternal()
 	FOnDemandEndpointConfig EndpointConfig;
 	if (TryParseEndpointConfig(CommandLine, EndpointConfig) == false)
 	{
+		// Temporary solution: Delete the ias cache directory if the game hasn't been configured to use IAS
+		if (FPaths::HasProjectPersistentDownloadDir())
+		{
+			TStringBuilder<256> CachePath;
+			CachePath << FPaths::ProjectPersistentDownloadDir();
+			FPathViews::Append(CachePath, TEXTVIEW("ias"));
+
+			IFileManager& Ifm = IFileManager::Get();
+			if (Ifm.DirectoryExists(CachePath.ToString()))
+			{
+				Ifm.DeleteDirectory(CachePath.ToString(), false, true);
+			}
+		}
 		return;
 	}
 
