ChaosFlesh: supported gravity for Gauss-Seidel quasistatics

diff --git a/Engine/Plugins/ChaosCloth/Source/ChaosCloth/Private/ChaosCloth/ChaosClothConstraints.cpp b/Engine/Plugins/ChaosCloth/Source/ChaosCloth/Private/ChaosCloth/ChaosClothConstraints.cpp
--- a/Engine/Plugins/ChaosCloth/Source/ChaosCloth/Private/ChaosCloth/ChaosClothConstraints.cpp
+++ b/Engine/Plugins/ChaosCloth/Source/ChaosCloth/Private/ChaosCloth/ChaosClothConstraints.cpp
@@ -554,9 +554,9 @@ void FClothConstraints::CreateGSRules()
 	{
 		PBDEvolution->SetQuasistatics(bClothDoQuasistatics);
 		const Softs::FSolverVec3& GravityTerm = PBDEvolution->GetGravity(0);
-		GSMainConstraint->ExternalForce[0] = 0.f;
-		GSMainConstraint->ExternalForce[1] = 0.f;
-		GSMainConstraint->ExternalForce[2] = GravityTerm[2];
+		GSMainConstraint->ExternalAcceleration[0] = 0.f;
+		GSMainConstraint->ExternalAcceleration[1] = 0.f;
+		GSMainConstraint->ExternalAcceleration[2] = GravityTerm[2];
 	}
 
 	ConstraintInits[ConstraintInitIndex++] =
diff --git a/Engine/Source/Runtime/Experimental/Chaos/Private/Chaos/Deformable/ChaosDeformableSolver.cpp b/Engine/Source/Runtime/Experimental/Chaos/Private/Chaos/Deformable/ChaosDeformableSolver.cpp
--- a/Engine/Source/Runtime/Experimental/Chaos/Private/Chaos/Deformable/ChaosDeformableSolver.cpp
+++ b/Engine/Source/Runtime/Experimental/Chaos/Private/Chaos/Deformable/ChaosDeformableSolver.cpp
@@ -1806,8 +1806,6 @@ namespace Chaos::Softs
 					this->GSSphereRepulsionConstraints->UpdateSphereRepulsionConstraints(this->Evolution->Particles(), *SurfaceVertices, *ParticleComponentIndex);
 					TArray<TArray<int32>> SphereRepulsionConstraints, SphereRepulsionIncidentElements, SphereRepulsionIncidentElementsLocal;
 					this->GSSphereRepulsionConstraints->ReturnSphereRepulsionConstraints(SphereRepulsionConstraints, SphereRepulsionIncidentElements, SphereRepulsionIncidentElementsLocal);
-					//this->GSMainConstraint->AddDynamicConstraints(SphereRepulsionConstraints, SphereRepulsionIncidentElements, SphereRepulsionIncidentElementsLocal);
-					//this->GSMainConstraint->InitDynamicColor(InParticles);
 					this->GSMainConstraint->AddTransientConstraints(SphereRepulsionConstraints, SphereRepulsionIncidentElements, SphereRepulsionIncidentElementsLocal);
 					this->GSMainConstraint->InitTransientColor(InParticles);
 					this->GSSphereRepulsionConstraints->Init(InParticles, Dt);
@@ -1830,6 +1828,13 @@ namespace Chaos::Softs
 					};
 			}
 		}
+		if (Property.bEnableGravity && Property.bDoQuasistatics) // Quasistatic PBD evolution does not apply gravity, add here
+		{
+			// @todo(flesh): Each proxy may have different sim spaces and therefore different gravity directions
+			const uint32 GroupId = 0;
+			FSolverVec3 GravityDir = Evolution->GetGravity(GroupId);
+			GSMainConstraint->AddExternalAcceleration(GravityDir);
+		}
 	}
 
 	void FDeformableSolver::InitializeMuscleActivationVariables()
diff --git a/Engine/Source/Runtime/Experimental/Chaos/Public/Chaos/Deformable/GaussSeidelMainConstraint.h b/Engine/Source/Runtime/Experimental/Chaos/Public/Chaos/Deformable/GaussSeidelMainConstraint.h
--- a/Engine/Source/Runtime/Experimental/Chaos/Public/Chaos/Deformable/GaussSeidelMainConstraint.h
+++ b/Engine/Source/Runtime/Experimental/Chaos/Public/Chaos/Deformable/GaussSeidelMainConstraint.h
@@ -276,6 +276,11 @@ namespace Chaos::Softs
 		
 		CHAOS_API void ApplyCG(ParticleType& Particles, const T Dt);
 
+		/* Adds external acceleration, e.g. gravity (0,0,-980) cm/s^2 */
+		CHAOS_API void AddExternalAcceleration(const TVec3<T>& Acceleration) { ExternalAcceleration += Acceleration; };
+
+		CHAOS_API void ResetExternalAcceleration() { ExternalAcceleration = TVec3<T>((T)0.); };
+
 	private:
 
 		void ApplySingleParticle(const int32 p, const T Dt, ParticleType& Particles)
@@ -359,7 +364,7 @@ namespace Chaos::Softs
 				else
 				{
 					for (int32 alpha = 0; alpha < 3; alpha++) {
-						ParticleResidual[alpha] = -Dt * Dt * ExternalForce[alpha] * Particles.M(p);
+						ParticleResidual[alpha] = -Dt * Dt * ExternalAcceleration[alpha] * Particles.M(p);
 					}
 				}
 			};
@@ -438,7 +443,7 @@ namespace Chaos::Softs
 		bool IsFirstFrame = true;
 		int32 PassedIters = 0;
 
-		TVec3<T> ExternalForce = TVec3<T>((T)0.);
+		TVec3<T> ExternalAcceleration = TVec3<T>((T)0.);
 
 	};
 
