Extended I/O store on-demand cache usage information with streaming cache stats

#rnx

diff --git a/Engine/Source/Runtime/Experimental/IoStore/OnDemand/Internal/IO/IasCache.h b/Engine/Source/Runtime/Experimental/IoStore/OnDemand/Internal/IO/IasCache.h
--- a/Engine/Source/Runtime/Experimental/IoStore/OnDemand/Internal/IO/IasCache.h
+++ b/Engine/Source/Runtime/Experimental/IoStore/OnDemand/Internal/IO/IasCache.h
@@ -48,6 +48,9 @@ public:
 
 	/** Insert a new chunk into the cache. */
 	virtual FIoStatus Put(const FIoHash& Key, FIoBuffer& Data) = 0;
+
+	/** Get cache usage information. */
+	virtual void GetCacheUsage(uint64& OutDiskSize, uint64& OutMaxDiskSize) const = 0;
 };
 
 struct FIasCacheConfig
diff --git a/Engine/Source/Runtime/Experimental/IoStore/OnDemand/Private/IoStoreOnDemand.cpp b/Engine/Source/Runtime/Experimental/IoStore/OnDemand/Private/IoStoreOnDemand.cpp
--- a/Engine/Source/Runtime/Experimental/IoStore/OnDemand/Private/IoStoreOnDemand.cpp
+++ b/Engine/Source/Runtime/Experimental/IoStore/OnDemand/Private/IoStoreOnDemand.cpp
@@ -2,6 +2,7 @@
 
 #include "IO/IoStoreOnDemand.h"
 
+#include "Containers/StringView.h"
 #include "HAL/FileManager.h"
 #include "HAL/FileManagerGeneric.h"
 #include "HAL/LowLevelMemTracker.h"
@@ -92,11 +93,11 @@ static FAutoConsoleCommand OnDemandCacheUsageCommand(
 			return;
 		}
 
-		const FOnDemandCacheUsage& Usage = MaybeUsage.ValueOrDie();
+		const FOnDemandCacheUsage CacheUsage = MaybeUsage.ConsumeValueOrDie();
 		UE_LOG(LogIoStoreOnDemand, Display, TEXT("iostore.CacheUsage"));
-		UE_LOG(LogIoStoreOnDemand, Display, TEXT("\tMaxSize %" UINT64_FMT), Usage.MaxSize);
-		UE_LOG(LogIoStoreOnDemand, Display, TEXT("\tTotalSize %" UINT64_FMT), Usage.TotalSize);
-		UE_LOG(LogIoStoreOnDemand, Display, TEXT("\tReferencedBlockSize %" UINT64_FMT), Usage.ReferencedBlockSize);
+		TStringBuilder<512> Sb;
+		Sb << TEXTVIEW("InstallCache: ") << CacheUsage.InstallCache << TEXTVIEW(", StreamingCache: ") << CacheUsage.StreamingCache;
+		UE_LOG(LogIoStoreOnDemand, Display, TEXT("%s"), Sb.ToString());
 	}),
 	ECVF_Cheat
 );
@@ -1053,6 +1054,25 @@ TIoStatusOr<FOnDemandToc> FOnDemandToc::LoadFromUrl(FStringView Url, uint32 Retr
 	return LoadFromUrl(AnsiUrl, RetryCount, bFollowRedirects);
 }
 
+////////////////////////////////////////////////////////////////////////////////
+FStringBuilderBase& operator<<(FStringBuilderBase& Sb, const FOnDemandInstallCacheUsage& CacheUsage)
+{
+	Sb.Appendf(TEXT("MaxSize=%.2lf MiB, TotalSize=%.2lf MiB, ReferencedBlockSize=%.2lf MiB"),
+		double(CacheUsage.MaxSize) / 1024.0 / 1024.0,
+		double(CacheUsage.TotalSize) / 1024.0 / 1024.0,
+		double(CacheUsage.ReferencedBlockSize) / 1024.0 / 1024.0);
+
+	return Sb;
+}
+
+////////////////////////////////////////////////////////////////////////////////
+FStringBuilderBase& operator<<(FStringBuilderBase& Sb, const FOnDemandStreamingCacheUsage& CacheUsage)
+{
+	Sb.Appendf(TEXT("MaxSize=%.2lf MiB, TotalSize=%.2lf MiB"),
+		double(CacheUsage.MaxSize) / 1024.0 / 1024.0, double(CacheUsage.TotalSize) / 1024.0 / 1024.0);
+	return Sb;
+}
+
 ////////////////////////////////////////////////////////////////////////////////
 void FIoStoreOnDemandModule::SetBulkOptionalEnabled(bool bInEnabled)
 {
@@ -1187,14 +1207,20 @@ FIoStatus FIoStoreOnDemandModule::GetInstallSizesByMountId(const FOnDemandGetIns
 	return FIoStatus(EIoErrorCode::InvalidCode, NotInitializedError);
 }
 
-TIoStatusOr<FOnDemandCacheUsage> FIoStoreOnDemandModule::GetCacheUsage(ELogOnDemandCacheUsage LogMode) const
+TIoStatusOr<FOnDemandCacheUsage> FIoStoreOnDemandModule::GetCacheUsage() const
 {
-	if (IoStore)
+	if (IoStore.IsValid() == false)
 	{
-		return IoStore->GetCacheUsage(LogMode);
+		return FIoStatus(EIoErrorCode::InvalidCode, NotInitializedError);
 	}
 
-	return FIoStatus(EIoErrorCode::InvalidCode, NotInitializedError);
+	FOnDemandCacheUsage CacheUsage = IoStore->GetCacheUsage();
+	if (HttpIoDispatcherBackend.IsValid())
+	{
+		CacheUsage.StreamingCache = HttpIoDispatcherBackend->GetCacheUsage();
+	}
+
+	return CacheUsage;
 }
 
 void FIoStoreOnDemandModule::InitializeInternal()
diff --git a/Engine/Source/Runtime/Experimental/IoStore/OnDemand/Private/JournaledCache.cpp b/Engine/Source/Runtime/Experimental/IoStore/OnDemand/Private/JournaledCache.cpp
--- a/Engine/Source/Runtime/Experimental/IoStore/OnDemand/Private/JournaledCache.cpp
+++ b/Engine/Source/Runtime/Experimental/IoStore/OnDemand/Private/JournaledCache.cpp
@@ -570,6 +570,7 @@ public:
 	void					Drop();
 	uint64					RemainingUntilWrap();
 	uint32					DebugVisit(void* Param, FDebugCacheEntry::Callback* Callback);
+	void					Usage(uint64& OutMappedBytes, uint64& OutMaxSize);
 
 private:
 	struct FMapEntry
@@ -904,7 +905,14 @@ uint32 FDiskCache::DebugVisit(void* Param, FDebugCacheEntry::Callback* Callback)
 	return DataMap.Num();
 }
 
+////////////////////////////////////////////////////////////////////////////////
+void FDiskCache::Usage(uint64& OutMappedBytes, uint64& OutMaxSize)
+{
+	FReadScopeLock _(Lock);
 
+	OutMappedBytes	= MappedBytes;
+	OutMaxSize		= MaxDataSize;
+}
 
 // {{{1 loader .................................................................
 
@@ -1202,6 +1210,7 @@ public:
 	uint32			Flush();
 	uint32			WriteMemToDisk(int32 Allowance);
 	uint32			DebugVisit(void* Param, FDebugCacheEntry::Callback* Callback);
+	void			GetDiskUsage(uint64& OutDiskUsage, uint64& OutMaxDiskSize);
 
 private:
 	struct FPartial
@@ -1450,6 +1459,11 @@ uint32 FCache::DebugVisit(void* Param, FDebugCacheEntry::Callback* Callback)
 	return Count;
 }
 
+////////////////////////////////////////////////////////////////////////////////
+void FCache::GetDiskUsage(uint64& OutDiskUsage, uint64& OutMaxDiskSize)
+{
+	DiskCache.Usage(OutDiskUsage, OutMaxDiskSize);
+}
 
 
 // {{{1 governor ...............................................................
@@ -2116,6 +2130,7 @@ public:
 	virtual void				Materialize(const FIoHash& Key, FIoBuffer& Dest, EIoErrorCode& Status, UE::Tasks::FTaskEvent DoneEvent) override;
 	virtual void				Cancel(FIoBuffer& GivenDest) override;
 	virtual FIoStatus			Put(const FIoHash& Key, FIoBuffer& Data) override;
+	virtual void				GetCacheUsage(uint64& OutDiskSize, uint64& OutMaxDiskSize) const override;
 
 private:
 	struct FMaterialOp
@@ -2284,6 +2299,12 @@ FIoStatus FJournaledCache::Put(const FIoHash& Key, FIoBuffer& Data)
 	return Ok ? FIoStatus::Ok : FIoStatus(EIoErrorCode::Unknown);
 }
 
+////////////////////////////////////////////////////////////////////////////////
+void FJournaledCache::GetCacheUsage(uint64& OutDiskSize, uint64& OutMaxDiskSize) const
+{
+	Cache->GetDiskUsage(OutDiskSize, OutMaxDiskSize);
+}
+
 ////////////////////////////////////////////////////////////////////////////////
 uint64 FJournaledCache::ReduceKey(const FIoHash& Key)
 {
diff --git a/Engine/Source/Runtime/Experimental/IoStore/OnDemand/Private/OnDemandInstallCache.cpp b/Engine/Source/Runtime/Experimental/IoStore/OnDemand/Private/OnDemandInstallCache.cpp
--- a/Engine/Source/Runtime/Experimental/IoStore/OnDemand/Private/OnDemandInstallCache.cpp
+++ b/Engine/Source/Runtime/Experimental/IoStore/OnDemand/Private/OnDemandInstallCache.cpp
@@ -963,7 +963,7 @@ public:
 	virtual FIoStatus							Purge(TMap<FIoHash, uint64>&& ChunksToIntall) override;
 	virtual FIoStatus							PurgeAllUnreferenced() override;
 	virtual FIoStatus							Flush() override;
-	virtual FOnDemandInstallCacheStorageUsage	GetStorageUsage() override;
+	virtual FOnDemandInstallCacheUsage			GetCacheUsage() override;
 
 private:
 	void										AddReferencesToBlocks(
@@ -1557,7 +1557,7 @@ FIoStatus FOnDemandInstallCache::Flush()
 	return FIoStatus::Ok;
 }
 
-FOnDemandInstallCacheStorageUsage FOnDemandInstallCache::GetStorageUsage()
+FOnDemandInstallCacheUsage FOnDemandInstallCache::GetCacheUsage()
 {
 	// If this is called from a thread other than the OnDemandIoStore tick thread
 	// then its possible the block info and containers may not be in sync with each other
@@ -1574,7 +1574,7 @@ FOnDemandInstallCacheStorageUsage FOnDemandInstallCache::GetStorageUsage()
 
 	AddReferencesToBlocks(Containers, ChunkEntryIndices, BlockInfo);
 
-	return FOnDemandInstallCacheStorageUsage
+	return FOnDemandInstallCacheUsage
 	{
 		.MaxSize = MaxCacheSize,
 		.TotalSize = TotalCachedBytes,
diff --git a/Engine/Source/Runtime/Experimental/IoStore/OnDemand/Private/OnDemandInstallCache.h b/Engine/Source/Runtime/Experimental/IoStore/OnDemand/Private/OnDemandInstallCache.h
--- a/Engine/Source/Runtime/Experimental/IoStore/OnDemand/Private/OnDemandInstallCache.h
+++ b/Engine/Source/Runtime/Experimental/IoStore/OnDemand/Private/OnDemandInstallCache.h
@@ -13,6 +13,7 @@ namespace UE::IoStore
 
 class FOnDemandIoStore;
 struct FOnDemandContainer;
+struct FOnDemandInstallCacheUsage;
 
 struct FOnDemandInstallCacheStorageUsage
 {
@@ -31,7 +32,7 @@ public:
 	virtual FIoStatus							Purge(TMap<FIoHash, uint64>&& ChunksToInstall) = 0;
 	virtual FIoStatus							PurgeAllUnreferenced() = 0;
 	virtual FIoStatus							Flush() = 0;
-	virtual FOnDemandInstallCacheStorageUsage	GetStorageUsage() = 0;
+	virtual FOnDemandInstallCacheUsage			GetCacheUsage() = 0;
 };
 
 struct FOnDemandInstallCacheConfig
diff --git a/Engine/Source/Runtime/Experimental/IoStore/OnDemand/Private/OnDemandIoDispatcherBackend.cpp b/Engine/Source/Runtime/Experimental/IoStore/OnDemand/Private/OnDemandIoDispatcherBackend.cpp
--- a/Engine/Source/Runtime/Experimental/IoStore/OnDemand/Private/OnDemandIoDispatcherBackend.cpp
+++ b/Engine/Source/Runtime/Experimental/IoStore/OnDemand/Private/OnDemandIoDispatcherBackend.cpp
@@ -1066,6 +1066,7 @@ public:
 	virtual bool IsEnabled() const override;
 	virtual void AbandonCache() override;
 	virtual void ReportAnalytics(TArray<FAnalyticsEventAttribute>& OutAnalyticsArray) const override;
+	virtual FOnDemandStreamingCacheUsage GetCacheUsage() const override;
 
 	// Runnable
 	virtual bool Init() override { return true; }
@@ -1656,6 +1657,17 @@ void FOnDemandIoBackend::ReportAnalytics(TArray<FAnalyticsEventAttribute>& OutAn
 	}
 }
 
+FOnDemandStreamingCacheUsage FOnDemandIoBackend::GetCacheUsage() const
+{
+	FOnDemandStreamingCacheUsage Usage;
+	if (Cache.IsValid())
+	{
+		Cache->GetCacheUsage(Usage.TotalSize, Usage.MaxSize);
+	}
+
+	return Usage; 
+}
+
 void FOnDemandIoBackend::ProcessHttpRequests(FHttpClient& HttpClient, FBitWindow& HttpErrors, int32 MaxConcurrentRequests)
 {
 	int32 NumConcurrentRequests = 0;
@@ -2022,4 +2034,4 @@ struct FTests
 
 } // namespace UE::IoStore
 
-#undef UE_IAS_DEBUG_CONSOLE_CMDS
\ No newline at end of file
+#undef UE_IAS_DEBUG_CONSOLE_CMDS
diff --git a/Engine/Source/Runtime/Experimental/IoStore/OnDemand/Private/OnDemandIoDispatcherBackend.h b/Engine/Source/Runtime/Experimental/IoStore/OnDemand/Private/OnDemandIoDispatcherBackend.h
--- a/Engine/Source/Runtime/Experimental/IoStore/OnDemand/Private/OnDemandIoDispatcherBackend.h
+++ b/Engine/Source/Runtime/Experimental/IoStore/OnDemand/Private/OnDemandIoDispatcherBackend.h
@@ -12,6 +12,7 @@ namespace UE::IoStore
 {
 
 class FOnDemandIoStore;
+struct FOnDemandStreamingCacheUsage;
 
 struct FDistributedEndpointUrl
 {
@@ -61,6 +62,7 @@ public:
 	virtual bool IsEnabled() const = 0;
 	virtual void AbandonCache() = 0;
 	virtual void ReportAnalytics(TArray<FAnalyticsEventAttribute>& OutAnalyticsArray) const = 0;
+	virtual FOnDemandStreamingCacheUsage GetCacheUsage() const = 0;
 };
 
 TSharedPtr<IOnDemandIoDispatcherBackend> MakeOnDemandIoDispatcherBackend(
diff --git a/Engine/Source/Runtime/Experimental/IoStore/OnDemand/Private/OnDemandIoStore.cpp b/Engine/Source/Runtime/Experimental/IoStore/OnDemand/Private/OnDemandIoStore.cpp
--- a/Engine/Source/Runtime/Experimental/IoStore/OnDemand/Private/OnDemandIoStore.cpp
+++ b/Engine/Source/Runtime/Experimental/IoStore/OnDemand/Private/OnDemandIoStore.cpp
@@ -904,9 +904,6 @@ FIoStatus FOnDemandIoStore::InitializeOnDemandInstallCache()
 	}
 #endif
 
-	// Log initial cache state
-	GetCacheUsage(ELogOnDemandCacheUsage::Default);
-
 	return EIoErrorCode::Ok;
 }
 
@@ -1978,26 +1975,15 @@ void FOnDemandIoStore::GetReferencedContent(TArray<FSharedOnDemandContainer>& Ou
 	}
 }
 
-FOnDemandCacheUsage FOnDemandIoStore::GetCacheUsage(ELogOnDemandCacheUsage LogMode) const
+FOnDemandCacheUsage FOnDemandIoStore::GetCacheUsage() const
 {
-	FOnDemandInstallCacheStorageUsage Usage = InstallCache->GetStorageUsage();
-
-	if (LogMode == ELogOnDemandCacheUsage::Default)
+	FOnDemandCacheUsage CacheUsage;
+	if (InstallCache.IsValid())
 	{
-		UE_LOG(LogIoStoreOnDemand, Display, TEXT(""));
-		UE_LOG(LogIoStoreOnDemand, Display, TEXT("* IoStore On Demand Cache Usage"));
-		UE_LOG(LogIoStoreOnDemand, Display, TEXT("* \tMax Size:              %" UINT64_FMT), Usage.MaxSize);
-		UE_LOG(LogIoStoreOnDemand, Display, TEXT("* \tTotal Size:            %" UINT64_FMT), Usage.TotalSize);
-		UE_LOG(LogIoStoreOnDemand, Display, TEXT("* \tReferenced Block Size: %" UINT64_FMT), Usage.ReferencedBlockSize);
-		UE_LOG(LogIoStoreOnDemand, Display, TEXT(""));
+		CacheUsage.InstallCache = InstallCache->GetCacheUsage();
 	}
 
-	return FOnDemandCacheUsage
-	{
-		.MaxSize = Usage.MaxSize,
-		.TotalSize = Usage.TotalSize,
-		.ReferencedBlockSize = Usage.ReferencedBlockSize
-	};
+	return CacheUsage;
 }
 
 ///////////////////////////////////////////////////////////////////////////////
diff --git a/Engine/Source/Runtime/Experimental/IoStore/OnDemand/Private/OnDemandIoStore.h b/Engine/Source/Runtime/Experimental/IoStore/OnDemand/Private/OnDemandIoStore.h
--- a/Engine/Source/Runtime/Experimental/IoStore/OnDemand/Private/OnDemandIoStore.h
+++ b/Engine/Source/Runtime/Experimental/IoStore/OnDemand/Private/OnDemandIoStore.h
@@ -302,7 +302,7 @@ public:
 	FOnDemandChunkInfo		GetInstalledChunkInfo(const FIoChunkId& ChunkId);
 	void					ReleaseContent(FOnDemandInternalContentHandle& ContentHandle);
 	void					GetReferencedContent(TArray<FSharedOnDemandContainer>& OutContainers, TArray<TBitArray<>>& OutChunkEntryIndices);
-	FOnDemandCacheUsage		GetCacheUsage(ELogOnDemandCacheUsage LogMode = ELogOnDemandCacheUsage::None) const;
+	FOnDemandCacheUsage		GetCacheUsage() const;
 
 private:
 	FIoStatus				GetContainersForInstall(
diff --git a/Engine/Source/Runtime/Experimental/IoStore/OnDemand/Public/IO/IoStoreOnDemand.h b/Engine/Source/Runtime/Experimental/IoStore/OnDemand/Public/IO/IoStoreOnDemand.h
--- a/Engine/Source/Runtime/Experimental/IoStore/OnDemand/Public/IO/IoStoreOnDemand.h
+++ b/Engine/Source/Runtime/Experimental/IoStore/OnDemand/Public/IO/IoStoreOnDemand.h
@@ -427,18 +427,30 @@ struct FOnDemandGetInstallSizeArgs
 	TArray<FPackageId> PackageIds;
 };
 
-/** Holds information about cache usage */
-struct FOnDemandCacheUsage
+/** Holds information about the install cache usage */
+struct FOnDemandInstallCacheUsage
 {
 	uint64 MaxSize = 0;
 	uint64 TotalSize = 0;
 	uint64 ReferencedBlockSize = 0;
+
+	UE_API friend FStringBuilderBase& operator<<(FStringBuilderBase& Sb, const FOnDemandInstallCacheUsage& CacheUsage);
 };
 
-enum class ELogOnDemandCacheUsage : uint8
+/** Holds information about the stremaing cache usage */
+struct FOnDemandStreamingCacheUsage
+{
+	uint64 MaxSize = 0;
+	uint64 TotalSize = 0;
+
+	UE_API friend FStringBuilderBase& operator<<(FStringBuilderBase& Sb, const FOnDemandStreamingCacheUsage& CacheUsage);
+};
+
+/** Holds information about install and streaming cache usage */
+struct FOnDemandCacheUsage
 {
-	None = 0,
-	Default,
+	FOnDemandInstallCacheUsage InstallCache;
+	FOnDemandStreamingCacheUsage StreamingCache;
 };
 
 /** Token used for signalling an operation to be cancelled. */
@@ -507,7 +519,7 @@ public:
 	  * if there are requests currently being processed. 
 	  * This should only be used for purposes such as debugging telemetry. 
 	  */
-	UE_API TIoStatusOr<FOnDemandCacheUsage> GetCacheUsage(ELogOnDemandCacheUsage LogMode = ELogOnDemandCacheUsage::None) const;
+	UE_API TIoStatusOr<FOnDemandCacheUsage> GetCacheUsage() const;
 
 	virtual void StartupModule() override;
 	virtual void ShutdownModule() override;
