Iris - Document UE_ADD_READ_JOURNAL_ENTRY.

#rnx

diff --git a/Engine/Source/Runtime/Experimental/Iris/Core/Public/Iris/Serialization/NetJournal.h b/Engine/Source/Runtime/Experimental/Iris/Core/Public/Iris/Serialization/NetJournal.h
--- a/Engine/Source/Runtime/Experimental/Iris/Core/Public/Iris/Serialization/NetJournal.h
+++ b/Engine/Source/Runtime/Experimental/Iris/Core/Public/Iris/Serialization/NetJournal.h
@@ -17,7 +17,8 @@ class UReplicationSystem;
 #endif
 
 #if UE_NET_ENABLE_READ_JOURNAL
-	#define UE_ADD_READ_JOURNAL_ENTRY(SerializationContext, X) SerializationContext.AddReadJournalEntry(X);
+	/** Add entry to the read journal. The Name must be a static/permanently allocated string as the raw pointer will be stored. */
+	#define UE_ADD_READ_JOURNAL_ENTRY(SerializationContext, Name) SerializationContext.AddReadJournalEntry(Name);
 	#define UE_RESET_READ_JOURNAL(SerializationContext) SerializationContext.ResetReadJournal();
 #else
 	#define UE_ADD_READ_JOURNAL_ENTRY(...)
@@ -34,6 +35,7 @@ class FNetJournal
 
 public:
 	void Reset() { NumEntries = 0U; }
+	/** Add entry to the journal. Only the last 32 entries are stored. The Name must be a static/permanently allocated string as the raw pointer will be stored. */
 	void AddEntry(const TCHAR* Name, uint32 BitOffset, FNetRefHandle NetRefHandle);
 	FString Print(const UReplicationSystem* ReplicationSystem) const;
 
