Iris - Fix crash in FNetJournal::Print

#rnx

diff --git a/Engine/Source/Runtime/Experimental/Iris/Core/Private/Iris/ReplicationSystem/ReplicationReader.cpp b/Engine/Source/Runtime/Experimental/Iris/Core/Private/Iris/ReplicationSystem/ReplicationReader.cpp
--- a/Engine/Source/Runtime/Experimental/Iris/Core/Private/Iris/ReplicationSystem/ReplicationReader.cpp
+++ b/Engine/Source/Runtime/Experimental/Iris/Core/Private/Iris/ReplicationSystem/ReplicationReader.cpp
@@ -2487,8 +2487,9 @@ void FReplicationReader::ProcessQueuedBatches()
 
 			// Reset for batch
 			Context.ResetErrorContext();
+			Context.SetErrorHandleContext(PendingBatchData.Handle);
 			UE_RESET_READ_JOURNAL(Context);
-			UE_ADD_READ_JOURNAL_ENTRY(Context, ToCStr(FString::Printf(TEXT("ProcessQueuedBatches for %s"), ToCStr(PendingBatchData.Handle.ToString()))));
+			UE_ADD_READ_JOURNAL_ENTRY(Context, TEXT("ProcessQueuedBatches"));
 
 			// Process batched data & dispatch data
 			for (const FQueuedDataChunk& CurrentChunk : PendingBatchData.QueuedDataChunks)
