fix invalid vertex/face start being set on geometry collection if a resize creates a geometry with zero vertices/faces

diff --git a/Engine/Source/Runtime/Experimental/Chaos/Private/GeometryCollection/GeometryCollection.cpp b/Engine/Source/Runtime/Experimental/Chaos/Private/GeometryCollection/GeometryCollection.cpp
--- a/Engine/Source/Runtime/Experimental/Chaos/Private/GeometryCollection/GeometryCollection.cpp
+++ b/Engine/Source/Runtime/Experimental/Chaos/Private/GeometryCollection/GeometryCollection.cpp
@@ -902,17 +902,35 @@ void FGeometryCollection::UpdateVerticesGroupElements()
 	TManagedArray<int32>* VertexGeometryMap = FindAttribute<int32>("VertexToGeometryIndex", VerticesGroup);
 	if(ensure(VertexGeometryMap))
 	{ 
-		for (int32 GeometryIndex = NumElements(FGeometryCollection::GeometryGroup) - 1; GeometryIndex >= 0; GeometryIndex--)
+		const int32 NumGeometryElements = NumElements(FGeometryCollection::GeometryGroup);
+		for (int32 GeometryIndex = NumGeometryElements - 1; GeometryIndex >= 0; GeometryIndex--)
 		{
 			VertexStart[GeometryIndex] = INT_MAX;
 			VertexCount[GeometryIndex] = 0;
 		}
 
-		for (int32 VertexIndex = NumElements(VerticesGroup) - 1; VertexIndex >= 0; VertexIndex--)
+		const int32 NumVertexElements = NumElements(VerticesGroup);
+		for (int32 VertexIndex = NumVertexElements - 1; VertexIndex >= 0; VertexIndex--)
 		{
 			VertexStart[ (*VertexGeometryMap)[VertexIndex] ] = FMath::Min(VertexStart[ (*VertexGeometryMap)[VertexIndex] ], VertexIndex);
 			VertexCount[ (*VertexGeometryMap)[VertexIndex] ]++;
 		}
+
+		// final pass through to fix any groups with zero vertices, which will not have had their start/count updated by the above loop
+		// note: technically if count is zero, start does not matter, but other code assumes the starts are in monotonic order
+		int32 LastVertexStart = 0, LastVertexCount = 0;
+		for (int32 GeometryIndex = 0; GeometryIndex < NumGeometryElements; ++GeometryIndex)
+		{
+			if (VertexCount[GeometryIndex] == 0)
+			{
+				VertexStart[GeometryIndex] = LastVertexStart + LastVertexCount;
+			}
+			else
+			{
+				LastVertexStart = VertexStart[GeometryIndex];
+				LastVertexCount = VertexCount[GeometryIndex];
+			}
+		}
 	}
 
 	ensure(HasContiguousVertices());
@@ -953,17 +971,35 @@ void FGeometryCollection::UpdateFaceGroupElements()
 	TManagedArray<int32>* FaceGeometryMap = FindAttribute<int32>("FaceToGeometryIndex", FacesGroup);
 	if (ensure(FaceGeometryMap))
 	{
-		for (int32 GeometryIndex = NumElements(FGeometryCollection::GeometryGroup) - 1; GeometryIndex >= 0; GeometryIndex--)
+		const int32 NumGeometryElements = NumElements(FGeometryCollection::GeometryGroup);
+		for (int32 GeometryIndex = NumGeometryElements - 1; GeometryIndex >= 0; GeometryIndex--)
 		{
 			FaceStart[GeometryIndex] = INT_MAX;
 			FaceCount[GeometryIndex] = 0;
 		}
 
-		for (int32 FaceIndex = NumElements(FacesGroup) - 1; FaceIndex >= 0; FaceIndex--)
+		const int32 NumFaceElements = NumElements(FacesGroup);
+		for (int32 FaceIndex = NumFaceElements - 1; FaceIndex >= 0; FaceIndex--)
 		{
 			FaceStart[(*FaceGeometryMap)[FaceIndex]] = FMath::Min(FaceStart[(*FaceGeometryMap)[FaceIndex]], FaceIndex);
 			FaceCount[(*FaceGeometryMap)[FaceIndex]]++;
 		}
+
+		// final pass through to fix any groups with zero faces, which will not have had their start/count updated by the above loop
+		// note: technically if count is zero, start does not matter, but other code assumes the starts are in monotonic order
+		int32 LastFaceStart = 0, LastFaceCount = 0;
+		for (int32 GeometryIndex = 0; GeometryIndex < NumGeometryElements; ++GeometryIndex)
+		{
+			if (FaceCount[GeometryIndex] == 0)
+			{
+				FaceStart[GeometryIndex] = LastFaceStart + LastFaceCount;
+			}
+			else
+			{
+				LastFaceStart = FaceStart[GeometryIndex];
+				LastFaceCount = FaceCount[GeometryIndex];
+			}
+		}
 	}
 
 	ensure(HasContiguousVertices());
