Iris - Try hard passing an object pointer to ObjectReferenceCache::RemoveReference. Not doing so can cause a very slow path to be taken which iterates over every NetRefHandle.

diff --git a/Engine/Source/Runtime/Experimental/Iris/Core/Private/Iris/ReplicationSystem/ObjectReferenceCache.cpp b/Engine/Source/Runtime/Experimental/Iris/Core/Private/Iris/ReplicationSystem/ObjectReferenceCache.cpp
--- a/Engine/Source/Runtime/Experimental/Iris/Core/Private/Iris/ReplicationSystem/ObjectReferenceCache.cpp
+++ b/Engine/Source/Runtime/Experimental/Iris/Core/Private/Iris/ReplicationSystem/ObjectReferenceCache.cpp
@@ -51,6 +51,9 @@ namespace UE::Net::Private
 static bool bIrisAllowAsyncLoading = true;
 FAutoConsoleVariableRef CVarIrisAllowAsyncLoading(TEXT("net.iris.AllowAsyncLoading"), bIrisAllowAsyncLoading, TEXT("Flag to allow or disallow async loading when using iris replication. Note: net.allowAsyncLoading must also be enabled."), ECVF_Default);
 
+static bool bIrisEnableRemoveReferenceFallbackPath = true;
+FAutoConsoleVariableRef CVarIrisEnableRemoveReferenceFallbackPath(TEXT("net.iris.ObjectReferenceCache.EnableRemoveReferenceFallbackPath"), bIrisEnableRemoveReferenceFallbackPath, TEXT("Whether to enable the slow fallback path or not when ObjectReferenceCache::RemoveReference is passed a null pointer object."), ECVF_Default);
+
 FObjectReferenceCache::FPendingAsyncLoadRequest::FPendingAsyncLoadRequest(FNetRefHandle InNetRefHandle, double InRequestStartTime)
 : RequestStartTime(InRequestStartTime)
 {
@@ -538,7 +541,8 @@ void FObjectReferenceCache::RemoveReference(FNetRefHandle RefHandle, const UObje
 			}
 			else
 			{
-				ensureMsgf(Object == CachedItemToRemove.ObjectKey, TEXT("ObjectReferenceCache::RemoveReference: %s, Object 0x%p doesn't match cached reference 0x%p."), *RefHandle.ToString(), Object, CachedItemToRemove.ObjectKey);
+				// RemoveReference can be called multiple times for the same NetRefHandle. The first call will clear out the object references.
+				ensureMsgf(CachedItemToRemove.ObjectKey == nullptr || Object == CachedItemToRemove.ObjectKey, TEXT("ObjectReferenceCache::RemoveReference: %s, Object 0x%p doesn't match cached reference 0x%p."), *RefHandle.ToString(), Object, CachedItemToRemove.ObjectKey);
 			}
 
 #if UE_NET_CLEANUP_REFERENCES_WITH_DYNAMIC_OUTER
@@ -606,7 +610,7 @@ void FObjectReferenceCache::RemoveReference(FNetRefHandle RefHandle, const UObje
 					*RefHandle.ToString(), *GetNameSafe(Object), Object);
 			}
 		}
-		else
+		else if (bIrisEnableRemoveReferenceFallbackPath)
 		{
 			bool bSuccesfullyRemoved = false;
 			for (auto RefHandleIt = ObjectToNetReferenceHandle.CreateIterator(); RefHandleIt; ++RefHandleIt)
diff --git a/Engine/Source/Runtime/Experimental/Iris/Core/Private/Iris/ReplicationSystem/ObjectReplicationBridge.cpp b/Engine/Source/Runtime/Experimental/Iris/Core/Private/Iris/ReplicationSystem/ObjectReplicationBridge.cpp
--- a/Engine/Source/Runtime/Experimental/Iris/Core/Private/Iris/ReplicationSystem/ObjectReplicationBridge.cpp
+++ b/Engine/Source/Runtime/Experimental/Iris/Core/Private/Iris/ReplicationSystem/ObjectReplicationBridge.cpp
@@ -796,7 +796,12 @@ void UObjectReplicationBridge::UnregisterInstance(FNetRefHandle RefHandle)
 {
 	if (RefHandle.IsDynamic())
 	{
-		UObject* Instance = GetObjectFromReferenceHandle(RefHandle);
+		const UObject* Instance = GetObjectFromReferenceHandle(RefHandle);
+		// In particular for subobjects its likely to get duplicate calls to UnregisterInstance in which case GetObjectFromReferenceHandle will fail to retrieve the instance. We really want to pass a valid object pointer to avoid a slow path iterating over every NetRefHandle.
+		if (!Instance)
+		{
+			Instance = NetRefHandleManager->GetReplicatedObjectInstance(NetRefHandleManager->GetInternalIndex(RefHandle));
+		}
 		GetObjectReferenceCache()->RemoveReference(RefHandle, Instance);
 	}
 }
