Dataflow : Make action menu contextual on the pulled pin type and direction
- When no pin is pulled the list remain unfiltered

diff --git a/Engine/Plugins/Experimental/Dataflow/Source/DataflowEditor/Private/Dataflow/DataflowSchema.cpp b/Engine/Plugins/Experimental/Dataflow/Source/DataflowEditor/Private/Dataflow/DataflowSchema.cpp
--- a/Engine/Plugins/Experimental/Dataflow/Source/DataflowEditor/Private/Dataflow/DataflowSchema.cpp
+++ b/Engine/Plugins/Experimental/Dataflow/Source/DataflowEditor/Private/Dataflow/DataflowSchema.cpp
@@ -105,6 +105,42 @@ void UDataflowSchema::GetContextMenuActions(class UToolMenu* Menu, class UGraphN
 	}
 	Super::GetContextMenuActions(Menu, Context);
 }
+bool UDataflowSchema::CanPinBeConnectedToNode(const UEdGraphPin* Pin, const UE::Dataflow::FFactoryParameters& NodeParameters)
+{
+	if (Pin == nullptr)
+	{
+		// if there's no pulled pin, then all nodes are compatible by default 
+		return true; 
+	}
+	if (NodeParameters.DefaultNodeObject)
+	{
+		const FName PinType = Pin->PinType.PinCategory;
+		if (Pin->Direction == EEdGraphPinDirection::EGPD_Output)
+		{
+			const TArray<FDataflowInput*> Inputs = NodeParameters.DefaultNodeObject->GetInputs();
+			for (const FDataflowInput* Input : Inputs)
+			{
+				if (Input && Input->SupportsType(PinType))
+				{
+					return true;
+				}
+			}
+		}
+		else if (Pin->Direction == EEdGraphPinDirection::EGPD_Input)
+		{
+			const TArray<FDataflowOutput*> Outputs = NodeParameters.DefaultNodeObject->GetOutputs();
+			for (const FDataflowOutput* Output : Outputs)
+			{
+				if (Output && Output->SupportsType(PinType))
+				{
+					return true;
+				}
+			}
+		}
+	}
+	return false;
+}
+
 
 void UDataflowSchema::GetGraphContextActions(FGraphContextMenuBuilder& ContextMenuBuilder) const
 {
@@ -112,10 +148,14 @@ void UDataflowSchema::GetGraphContextActions(FGraphContextMenuBuilder& ContextMe
 	{
 		for (UE::Dataflow::FFactoryParameters NodeParameters : Factory->RegisteredParameters())
 		{
-			if (TSharedPtr<FAssetSchemaAction_Dataflow_CreateNode_DataflowEdNode> Action =
-				FAssetSchemaAction_Dataflow_CreateNode_DataflowEdNode::CreateAction(ContextMenuBuilder.CurrentGraph, NodeParameters.TypeName, NodeParameters.DisplayName))
+			// contextual filtering ( if pin is null always return true )
+			if (UDataflowSchema::CanPinBeConnectedToNode(ContextMenuBuilder.FromPin, NodeParameters))
 			{
-				ContextMenuBuilder.AddAction(Action);
+				if (TSharedPtr<FAssetSchemaAction_Dataflow_CreateNode_DataflowEdNode> Action =
+					FAssetSchemaAction_Dataflow_CreateNode_DataflowEdNode::CreateAction(ContextMenuBuilder.CurrentGraph, NodeParameters.TypeName, NodeParameters.DisplayName))
+				{
+					ContextMenuBuilder.AddAction(Action);
+				}
 			}
 		}
 	}
diff --git a/Engine/Plugins/Experimental/Dataflow/Source/DataflowEditor/Public/Dataflow/DataflowSchema.h b/Engine/Plugins/Experimental/Dataflow/Source/DataflowEditor/Public/Dataflow/DataflowSchema.h
--- a/Engine/Plugins/Experimental/Dataflow/Source/DataflowEditor/Public/Dataflow/DataflowSchema.h
+++ b/Engine/Plugins/Experimental/Dataflow/Source/DataflowEditor/Public/Dataflow/DataflowSchema.h
@@ -10,6 +10,11 @@
 
 class UDataflow; 
 
+namespace UE::Dataflow
+{
+	struct FFactoryParameters;
+}
+
 UCLASS()
 class DATAFLOWEDITOR_API UDataflowSchema : public UEdGraphSchema
 {
@@ -30,6 +35,8 @@ public:
 	//~ End UEdGraphSchema Interface
 
 	static FLinearColor GetTypeColor(const FName& Type);
+private:
+	static bool CanPinBeConnectedToNode(const UEdGraphPin* Pin, const UE::Dataflow::FFactoryParameters& NodeParameters);
 };
 
 class FDataflowConnectionDrawingPolicy : public FConnectionDrawingPolicy, public FGCObject
diff --git a/Engine/Source/Runtime/Experimental/Dataflow/Core/Private/Dataflow/DataflowNodeFactory.cpp b/Engine/Source/Runtime/Experimental/Dataflow/Core/Private/Dataflow/DataflowNodeFactory.cpp
--- a/Engine/Source/Runtime/Experimental/Dataflow/Core/Private/Dataflow/DataflowNodeFactory.cpp
+++ b/Engine/Source/Runtime/Experimental/Dataflow/Core/Private/Dataflow/DataflowNodeFactory.cpp
@@ -52,9 +52,8 @@ namespace UE::Dataflow
 		// gets modified to NewDisplayName = DisplayName + _v<VERSION>
 		// 
 		// Update DisplayName with version from TypeName
-		FString NewDisplayNameString = Parameters.DisplayName.ToString() + "_" + Parameters.NodeVersion.ToString();
-
-		FFactoryParameters NewParameters(Parameters.TypeName, FName(*NewDisplayNameString), Parameters.Category, Parameters.Tags, Parameters.ToolTip, Parameters.bIsDeprecated, Parameters.bIsExperimental, Parameters.NodeVersion);
+		FFactoryParameters NewParameters{ Parameters };
+		NewParameters.DisplayName = FName(Parameters.DisplayName.ToString() + "_" + Parameters.NodeVersion.ToString());
 
 		if (ClassMap.Contains(NewParameters.TypeName))
 		{
diff --git a/Engine/Source/Runtime/Experimental/Dataflow/Core/Public/Dataflow/DataflowNodeFactory.h b/Engine/Source/Runtime/Experimental/Dataflow/Core/Public/Dataflow/DataflowNodeFactory.h
--- a/Engine/Source/Runtime/Experimental/Dataflow/Core/Public/Dataflow/DataflowNodeFactory.h
+++ b/Engine/Source/Runtime/Experimental/Dataflow/Core/Public/Dataflow/DataflowNodeFactory.h
@@ -10,40 +10,43 @@ struct FDataflowConnection;
 
 namespace UE::Dataflow
 {
-	struct FNewNodeParameters {
+	struct FNewNodeParameters
+	{
 		FGuid Guid;
 		FName Type;
 		FName Name;
 		UObject* OwningObject = nullptr;
 	};
 
-	struct FFactoryParameters {
-		FFactoryParameters() {}
-		FFactoryParameters(FName InTypeName, FName InDisplayName, FName InCategory, FString InTags, FString InToolTip, bool bInIsDeprecated, bool bInIsExperimental, FName InNodeVersion)
-		: TypeName(InTypeName), DisplayName(InDisplayName), Category(InCategory), Tags(InTags), ToolTip(InToolTip), bIsDeprecated(bInIsDeprecated), bIsExperimental(bInIsExperimental), NodeVersion(InNodeVersion) {}
-
-		FName TypeName = FName("");
-		FName DisplayName = FName("");
-		FName Category = FName("");
-		FString Tags = FString("");
-		FString ToolTip = FString("");
+	struct FFactoryParameters
+	{
+		FName TypeName;
+		FName DisplayName;
+		FName Category;
+		FString Tags;
+		FString ToolTip;
 		bool bIsDeprecated = false;
 		bool bIsExperimental = false;
 		FName NodeVersion = FName(TEXT("v1"));
+		TSharedPtr<FDataflowNode> DefaultNodeObject;
 
-		bool IsValid() const {
+		bool IsValid() const
+		{
 			return !TypeName.ToString().IsEmpty() && !DisplayName.ToString().IsEmpty();
 		}
 
-		bool IsDeprecated() const {
+		bool IsDeprecated() const
+		{
 			return bIsDeprecated;
 		}
 
-		bool IsExperimental() const {
+		bool IsExperimental() const
+		{
 			return bIsExperimental;
 		}
 
-		FName GetVersion() const {
+		FName GetVersion() const
+		{
 			return NodeVersion;
 		}
 	};
@@ -79,16 +82,23 @@ namespace UE::Dataflow
 		template<typename T>
 		static void RegisterNodeFromType()
 		{
-			FFactoryParameters FactoryParameters =
+			UE::Dataflow::FNodeParameters CreateParams
+			{
+				.Name			= T::StaticType(),
+				.OwningObject	= GetTransientPackage(),
+			};
+
+			FFactoryParameters FactoryParameters
 			{
-				T::StaticType(),
-				T::StaticDisplay(),
-				T::StaticCategory(),
-				T::StaticTags(),
-				GetToolTipFromStruct(T::StaticStruct(), T::StaticType(), T::StaticDisplay()),
-				IsNodeDeprecated(T::StaticStruct()),
-				IsNodeExperimental(T::StaticStruct()),
-				GetVersionFromTypeName(T::StaticType())
+				.TypeName		= T::StaticType(),
+				.DisplayName	= T::StaticDisplay(),
+				.Category		= T::StaticCategory(),
+				.Tags		= T::StaticTags(),
+				.ToolTip		= GetToolTipFromStruct(T::StaticStruct(), T::StaticType(), T::StaticDisplay()),
+				.bIsDeprecated	= IsNodeDeprecated(T::StaticStruct()),
+				.bIsExperimental= IsNodeExperimental(T::StaticStruct()),
+				.NodeVersion	= GetVersionFromTypeName(T::StaticType()),
+				.DefaultNodeObject = MakeShared<T>(CreateParams),
 			};
 
 			auto CreationFunction = [](const UE::Dataflow::FNewNodeParameters& InParam)
