Iris - Add ensure to FObjectReferenceCache::RemoveReference to figure out if slow fallback path is needed or not.

#rnx

diff --git a/Engine/Source/Runtime/Experimental/Iris/Core/Private/Iris/ReplicationSystem/ObjectReferenceCache.cpp b/Engine/Source/Runtime/Experimental/Iris/Core/Private/Iris/ReplicationSystem/ObjectReferenceCache.cpp
--- a/Engine/Source/Runtime/Experimental/Iris/Core/Private/Iris/ReplicationSystem/ObjectReferenceCache.cpp
+++ b/Engine/Source/Runtime/Experimental/Iris/Core/Private/Iris/ReplicationSystem/ObjectReferenceCache.cpp
@@ -619,7 +619,9 @@ void FObjectReferenceCache::RemoveReference(FNetRefHandle RefHandle, const UObje
 				{
 					RefHandleIt.RemoveCurrent();
 					bSuccesfullyRemoved = true;
-					UE_LOG_REFERENCECACHE(Verbose, TEXT("ObjectReferenceCache::RemoveReference: %s. Successfully removed RefHandle from ObjectToNetReferenceHandle."), *RefHandle.ToString());
+					// It's actually bad if we succesfully remove a reference in this horribly slow fallback path. It means it's still need it and we need to figure out how to get rid of it.
+					UE_LOG_REFERENCECACHE(Error, TEXT("ObjectReferenceCache::RemoveReference: %s. Removed RefHandle from ObjectToNetReferenceHandle."), *RefHandle.ToString());
+					ensureMsgf(false, TEXT("FObjectReferenceCache::RemoveReference %s: Slow fallback path for null instance is still needed. It needs further investigation."), *RefHandle.ToString());
 					break;
 				}
 			}
