Demote an initialization logging error to a warning.

#rnx

- At the moment this module is now included in the standard code build so our assumption that the project should be configured for persistent storage no longer holds true for all projects on all platforms.
- For an immediate fix we are demoting the error to a warning as the system it reports for is not in production use at the moment.
- We will revist the initialization assumptions to prevent this from happening and restore the error log in a future submit.

diff --git a/Engine/Source/Runtime/Experimental/IoStore/OnDemand/Private/IoStoreOnDemand.cpp b/Engine/Source/Runtime/Experimental/IoStore/OnDemand/Private/IoStoreOnDemand.cpp
--- a/Engine/Source/Runtime/Experimental/IoStore/OnDemand/Private/IoStoreOnDemand.cpp
+++ b/Engine/Source/Runtime/Experimental/IoStore/OnDemand/Private/IoStoreOnDemand.cpp
@@ -1251,7 +1251,7 @@ void FIoStoreOnDemandModule::InitializeInternal()
 		IoStore = MakeShared<FOnDemandIoStore>();
 		if (FIoStatus Status = IoStore->Initialize(); !Status.IsOk())
 		{
-			UE_LOG(LogIas, Error, TEXT("Failed to initialize I/O store on demand, reason '%s'"), *Status.ToString());
+			UE_LOG(LogIas, Warning, TEXT("Failed to initialize I/O store on demand, reason '%s'"), *Status.ToString());
 			return;
 		}
 	}
