Small optimization for FHeightField::BuildGeomData and EditGeomData.  Convert indirect callback function into inlinable plain lambda template param.  Reduces BuildGeomData cost by 42%. Up to 1.5 ms per frame in worst case landscape sculpting.

diff --git a/Engine/Source/Runtime/Experimental/Chaos/Private/Chaos/HeightField.cpp b/Engine/Source/Runtime/Experimental/Chaos/Private/Chaos/HeightField.cpp
--- a/Engine/Source/Runtime/Experimental/Chaos/Private/Chaos/HeightField.cpp
+++ b/Engine/Source/Runtime/Experimental/Chaos/Private/Chaos/HeightField.cpp
@@ -597,8 +597,8 @@ namespace Chaos
 	};
 
 
-	template<typename BufferType>
-	void BuildGeomData(TArrayView<BufferType> BufferView, TArrayView<const uint8> MaterialIndexView, int32 NumRows, int32 NumCols, const FVec3& InScale, TUniqueFunction<FReal(const BufferType)> ToRealFunc, typename FHeightField::FDataType& OutData, FAABB3& OutBounds)
+	template<typename BufferType, typename ToRealFuncType>
+	void BuildGeomData(TArrayView<BufferType> BufferView, TArrayView<const uint8> MaterialIndexView, int32 NumRows, int32 NumCols, const FVec3& InScale, ToRealFuncType ToRealFunc, typename FHeightField::FDataType& OutData, FAABB3& OutBounds)
 	{
 		using FDataType = typename FHeightField::FDataType;
 
@@ -673,8 +673,8 @@ namespace Chaos
 		}
 	}
 
-	template<typename BufferType>
-	void EditGeomData(TArrayView<BufferType> BufferView, int32 InBeginRow, int32 InBeginCol, int32 NumRows, int32 NumCols, TUniqueFunction<FReal(const BufferType)> ToRealFunc, typename FHeightField::FDataType& OutData, FAABB3& OutBounds)
+	template<typename BufferType, typename ToRealFuncType>
+	void EditGeomData(TArrayView<BufferType> BufferView, int32 InBeginRow, int32 InBeginCol, int32 NumRows, int32 NumCols, ToRealFuncType ToRealFunc, typename FHeightField::FDataType& OutData, FAABB3& OutBounds)
 	{
 		using FDataType = typename FHeightField::FDataType;
 
@@ -763,12 +763,12 @@ namespace Chaos
 	FHeightField::FHeightField(TArray<FReal>&& Height, TArray<uint8>&& InMaterialIndices, int32 NumRows, int32 NumCols, const FVec3& InScale)
 		: FImplicitObject(EImplicitObject::HasBoundingBox, ImplicitObjectType::HeightField)
 	{
-		TUniqueFunction<FReal(FReal)> ConversionFunc = [](const FReal InVal) -> FReal
+		auto ConversionFunc = [](const FReal InVal) -> FReal
 		{
 			return InVal;
 		};
 
-		BuildGeomData<FReal>(MakeArrayView(Height), MakeArrayView(InMaterialIndices), NumRows, NumCols, FVec3(1), MoveTemp(ConversionFunc), GeomData, LocalBounds);
+		BuildGeomData<FReal>(MakeArrayView(Height), MakeArrayView(InMaterialIndices), NumRows, NumCols, FVec3(1), ConversionFunc, GeomData, LocalBounds);
 		CalcBounds();
 		SetScale(InScale);
 	}
@@ -776,12 +776,12 @@ namespace Chaos
 	Chaos::FHeightField::FHeightField(TArrayView<const uint16> InHeights, TArrayView<const uint8> InMaterialIndices, int32 InNumRows, int32 InNumCols, const FVec3& InScale)
 		: FImplicitObject(EImplicitObject::HasBoundingBox, ImplicitObjectType::HeightField)
 	{
-		TUniqueFunction<FReal(const uint16)> ConversionFunc = [](const uint16 InVal) -> FReal
+		auto ConversionFunc = [](const uint16 InVal) -> FReal
 		{
 			return (FReal)((int32)InVal - 32768);
 		};
 
-		BuildGeomData<const uint16>(InHeights, InMaterialIndices, InNumRows, InNumCols, FVec3(1), MoveTemp(ConversionFunc), GeomData, LocalBounds);
+		BuildGeomData<const uint16>(InHeights, InMaterialIndices, InNumRows, InNumCols, FVec3(1), ConversionFunc, GeomData, LocalBounds);
 		CalcBounds();
 		SetScale(InScale);
 	}
@@ -794,12 +794,12 @@ namespace Chaos
 
 		if(ensure(InHeights.Num() == NumExpectedValues && InBeginRow >= 0 && InBeginCol >= 0 && EndRow < GeomData.NumRows && EndCol < GeomData.NumCols))
 		{
-			TUniqueFunction<FReal(const uint16)> ConversionFunc = [](const uint16 InVal) -> FReal
+			auto ConversionFunc = [](const uint16 InVal) -> FReal
 			{
 				return (FReal)((int32)InVal - 32768);
 			};
 
-			EditGeomData<const uint16>(InHeights, InBeginRow, InBeginCol, InNumRows, InNumCols, MoveTemp(ConversionFunc), GeomData, LocalBounds);
+			EditGeomData<const uint16>(InHeights, InBeginRow, InBeginCol, InNumRows, InNumCols, ConversionFunc, GeomData, LocalBounds);
 
 			// Slow and dumb. TODO: Actually fix CellHeights inside EditGeomData.
 			CalcBounds();
@@ -814,12 +814,12 @@ namespace Chaos
 
 		if(ensure(InHeights.Num() == NumExpectedValues && InBeginRow >= 0 && InBeginCol >= 0 && EndRow < GeomData.NumRows && EndCol < GeomData.NumCols))
 		{
-			TUniqueFunction<FReal(FReal)> ConversionFunc = [](const FReal InVal) -> FReal
+			auto ConversionFunc = [](const FReal InVal) -> FReal
 			{
 				return InVal;
 			};
 
-			EditGeomData<FReal>(InHeights, InBeginRow, InBeginCol, InNumRows, InNumCols, MoveTemp(ConversionFunc), GeomData, LocalBounds);
+			EditGeomData<FReal>(InHeights, InBeginRow, InBeginCol, InNumRows, InNumCols, ConversionFunc, GeomData, LocalBounds);
 
 			// Slow and dumb. TODO: Actually fix CellHeights inside EditGeomData.
 			CalcBounds();
