[Chaos]
* Added private include dependency on PhysicsCore since module includes ChaosEngineInterface.hk

diff --git a/Engine/Source/Runtime/Experimental/Chaos/Chaos.Build.cs b/Engine/Source/Runtime/Experimental/Chaos/Chaos.Build.cs
--- a/Engine/Source/Runtime/Experimental/Chaos/Chaos.Build.cs
+++ b/Engine/Source/Runtime/Experimental/Chaos/Chaos.Build.cs
@@ -32,6 +32,8 @@ namespace UnrealBuildTool.Rules
 			PublicDefinitions.Add("COMPILE_WITHOUT_UNREAL_SUPPORT=0");
 			PrivateDefinitions.Add("CHAOS_INCLUDE_LEVEL_1=1");
 
+			PrivateIncludePathModuleNames.Add("PhysicsCore"); // Needed because we include "ChaosEngineInterface.h"
+
 			if (Target.bUseChaosMemoryTracking == true)
 			{
 				PublicDefinitions.Add("CHAOS_MEMORY_TRACKING=1");
