Fix Debug Draw Server Shape in PIE

diff --git a/Engine/Source/Runtime/Experimental/Chaos/Private/PBDRigidsSolver.cpp b/Engine/Source/Runtime/Experimental/Chaos/Private/PBDRigidsSolver.cpp
--- a/Engine/Source/Runtime/Experimental/Chaos/Private/PBDRigidsSolver.cpp
+++ b/Engine/Source/Runtime/Experimental/Chaos/Private/PBDRigidsSolver.cpp
@@ -2545,9 +2545,10 @@ TRACE_COUNTER_SET(ChaosTraceCounter_##Name, Value)
 		{
 			if (bChaosDebugDraw_UseNewQueue)
 			{
+				FColor Color = bIsServer ? FColor::Red : FColor::Blue;
 				for (FTransientGeometryParticleHandle& Particle : Particles.GetNonDisabledView())
 				{
-					FChaosDDParticle::DrawShapes(Particle.Handle());
+					FChaosDDParticle::DrawShapes(Particle.Handle(), Color);
 				}
 			}
 			else
