Chaos - fix debug draw on client/server. Was not honoring enabled or coloring options, and the debug draw renderer's IsServer flag was not set when rendering the server.

diff --git a/Engine/Source/Runtime/Experimental/Chaos/Private/Chaos/DebugDraw/DebugDrawParticle.cpp b/Engine/Source/Runtime/Experimental/Chaos/Private/Chaos/DebugDraw/DebugDrawParticle.cpp
--- a/Engine/Source/Runtime/Experimental/Chaos/Private/Chaos/DebugDraw/DebugDrawParticle.cpp
+++ b/Engine/Source/Runtime/Experimental/Chaos/Private/Chaos/DebugDraw/DebugDrawParticle.cpp
@@ -40,8 +40,14 @@ namespace Chaos::CVars
 	extern int32 ChaosSolverDrawShapesShowStatic;
 	extern int32 ChaosSolverDrawShapesShowKinematic;
 	extern int32 ChaosSolverDrawShapesShowDynamic;
+	extern int32 ChaosSolverDebugDrawShowServer;
+	extern int32 ChaosSolverDebugDrawShowClient;
+	extern int32 ChaosSolverDebugDrawColorShapeByClientServer;
 
 	extern Chaos::DebugDraw::FChaosDebugDrawSettings ChaosSolverDebugDebugDrawSettings;
+
+	extern DebugDraw::FChaosDebugDrawColorsByState GetSolverShapesColorsByState_Server();
+	extern DebugDraw::FChaosDebugDrawColorsByState GetSolverShapesColorsByState_Client();
 }
 
 namespace Chaos
@@ -160,6 +166,14 @@ namespace Chaos
 			{
 				return false;
 			}
+			if (!Chaos::CVars::ChaosSolverDebugDrawShowServer && bIsServer)
+			{
+				return false;
+			}
+			if (!Chaos::CVars::ChaosSolverDebugDrawShowClient && !bIsServer)
+			{
+				return false;
+			}
 
 			// Depending on the shape settings, we may not show the simple or complex shape
 			const EImplicitObjectType InnerType = GetInnerType(Implicit->GetType());
@@ -186,6 +200,17 @@ namespace Chaos
 			const Chaos::DebugDraw::FChaosDebugDrawSettings& Settings = Chaos::CVars::ChaosSolverDebugDebugDrawSettings;
 
 			const EImplicitObjectType InnerType = GetInnerType(Implicit->GetType());
+			if (Chaos::CVars::ChaosSolverDebugDrawColorShapeByClientServer)
+			{
+				if (bIsServer)
+				{
+					return Chaos::CVars::GetSolverShapesColorsByState_Server().GetColorFromState(ParticleData.ObjectState, ParticleData.bIsOneWay);
+				}
+				else
+				{
+					return Chaos::CVars::GetSolverShapesColorsByState_Client().GetColorFromState(ParticleData.ObjectState, ParticleData.bIsOneWay);
+				}
+			}
 			if (Chaos::DebugDraw::bChaosDebugDebugDrawColorShapesByShapeType)
 			{
 				return Settings.ShapesColorsPerShapeType.GetColorFromShapeType(InnerType);
diff --git a/Engine/Source/Runtime/Experimental/Chaos/Private/PBDRigidsSolver.cpp b/Engine/Source/Runtime/Experimental/Chaos/Private/PBDRigidsSolver.cpp
--- a/Engine/Source/Runtime/Experimental/Chaos/Private/PBDRigidsSolver.cpp
+++ b/Engine/Source/Runtime/Experimental/Chaos/Private/PBDRigidsSolver.cpp
@@ -233,7 +233,10 @@ namespace Chaos
 			/* InBoundsColorsPerShapeType=  */ DebugDraw::GetDefaultBoundsColorsByShapeType()
 		);
 
-		static DebugDraw::FChaosDebugDrawColorsByState GetSolverShapesColorsByState_Server()
+		extern DebugDraw::FChaosDebugDrawColorsByState GetSolverShapesColorsByState_Server();
+		extern DebugDraw::FChaosDebugDrawColorsByState GetSolverShapesColorsByState_Client();
+
+		DebugDraw::FChaosDebugDrawColorsByState GetSolverShapesColorsByState_Server()
 		{
 			static DebugDraw::FChaosDebugDrawColorsByState SolverShapesColorsByState_Server
 			{
@@ -246,7 +249,7 @@ namespace Chaos
 			return SolverShapesColorsByState_Server;
 		}
 
-		static DebugDraw::FChaosDebugDrawColorsByState GetSolverShapesColorsByState_Client()
+		DebugDraw::FChaosDebugDrawColorsByState GetSolverShapesColorsByState_Client()
 		{
 			static DebugDraw::FChaosDebugDrawColorsByState SolverShapesColorsByState_Client
 			{
@@ -2545,10 +2548,9 @@ TRACE_COUNTER_SET(ChaosTraceCounter_##Name, Value)
 		{
 			if (bChaosDebugDraw_UseNewQueue)
 			{
-				FColor Color = bIsServer ? FColor::Red : FColor::Blue;
 				for (FTransientGeometryParticleHandle& Particle : Particles.GetNonDisabledView())
 				{
-					FChaosDDParticle::DrawShapes(Particle.Handle(), Color);
+					FChaosDDParticle::DrawShapes(Particle.Handle());
 				}
 			}
 			else
diff --git a/Engine/Source/Runtime/Experimental/ChaosSolverEngine/Private/Chaos/ChaosDebugDrawSubsystem.cpp b/Engine/Source/Runtime/Experimental/ChaosSolverEngine/Private/Chaos/ChaosDebugDrawSubsystem.cpp
--- a/Engine/Source/Runtime/Experimental/ChaosSolverEngine/Private/Chaos/ChaosDebugDrawSubsystem.cpp
+++ b/Engine/Source/Runtime/Experimental/ChaosSolverEngine/Private/Chaos/ChaosDebugDrawSubsystem.cpp
@@ -47,12 +47,18 @@ public:
 		, RenderCost(0)
 		, SphereSegments(8)
 		, DepthPriority(10)
+		, bIsServer(false)
 	{
 	}
 
 	virtual bool IsServer() const override final
 	{
-		return World->GetNetMode() == ENetMode::NM_DedicatedServer;
+		return bIsServer;
+	}
+
+	void SetIsServer(const bool bInIsServer)
+	{
+		bIsServer = bInIsServer;
 	}
 
 	virtual FSphere3d GetDrawRegion() const override final
@@ -226,6 +232,7 @@ private:
 	// Draw Settings
 	int32 SphereSegments;
 	int8 DepthPriority;
+	bool bIsServer;
 };
 
 class FChaosDDWorldManager
@@ -405,12 +412,15 @@ void UChaosDebugDrawSubsystem::RenderScene()
 				// PIE clients, but the global scene is a hack anyway so won't worry about that
 				RenderFrame(Renderer, ChaosDD::Private::FChaosDDContext::ExtractGlobalFrame());
 
-				// Render the commands from the server on every client (in PIE)
-				RenderScene(Renderer, FChaosDDWorldManager::Get().GetServerDebugDrawScene());
-
 				// Render the commands from this world
 				RenderScene(Renderer, CDDScene);
 
+				// Render the commands from the server on every client (in PIE)
+				// Render the server last so that the client uses the command and render budgets first
+				Renderer.SetIsServer(true);
+				RenderScene(Renderer, FChaosDDWorldManager::Get().GetServerDebugDrawScene());
+				Renderer.SetIsServer(false);
+
 				if (Renderer.WasRenderBudgetExceeded())
 				{
 					constexpr int32 MsgId = 86421358;
