Iris - Minimize the number of broadcasts to NetSerializers to inform them that modules have been loaded. A broadcast will happen when at least one module has been loaded followed by a frame with no additional loaded modules.

diff --git a/Engine/Source/Runtime/Experimental/Iris/Core/Private/Iris/IrisCoreModule.cpp b/Engine/Source/Runtime/Experimental/Iris/Core/Private/Iris/IrisCoreModule.cpp
--- a/Engine/Source/Runtime/Experimental/Iris/Core/Private/Iris/IrisCoreModule.cpp
+++ b/Engine/Source/Runtime/Experimental/Iris/Core/Private/Iris/IrisCoreModule.cpp
@@ -1,6 +1,7 @@
 // Copyright Epic Games, Inc. All Rights Reserved.
 
 #include "CoreTypes.h"
+#include "Containers/Ticker.h"
 #include "Modules/ModuleInterface.h"
 #include "Modules/ModuleManager.h"
 #include "Misc/CoreDelegates.h"
@@ -8,9 +9,11 @@
 
 #include "Net/Core/Connection/NetEnums.h"
 
+#include "Iris/Core/IrisLog.h"
 #include "Iris/ReplicationState/DefaultPropertyNetSerializerInfos.h"
 #include "Iris/ReplicationState/PropertyNetSerializerInfoRegistry.h"
 #include "Iris/ReplicationSystem/LegacyPushModel.h"
+#include "Iris/ReplicationSystem/ReplicationSystem.h"
 #include "Iris/Serialization/InternalNetSerializerDelegates.h"
 #include "Iris/IrisConfigInternal.h"
 
@@ -53,6 +56,14 @@ private:
 		UE_NET_IRIS_INIT_LEGACY_PUSH_MODEL();
 		
 		ModulesChangedHandle = FModuleManager::Get().OnModulesChanged().AddRaw(this, &FIrisCoreModule::OnModulesChanged);
+		RepSysCreatedHandle = UE::Net::FReplicationSystemFactory::GetReplicationSystemCreatedDelegate().AddRaw(this, &FIrisCoreModule::OnRepSystemCreated);
+		RepSysDestroyedHandle = UE::Net::FReplicationSystemFactory::GetReplicationSystemDestroyedDelegate().AddRaw(this, &FIrisCoreModule::OnRepSystemDestroyed);
+
+		// Figure out how many ReplicationSystems there are so we start on a correct balance prior to getting callbacks.
+		for (TArrayView<UReplicationSystem*> RepSystems = UE::Net::FReplicationSystemFactory::GetAllReplicationSystems(); const UReplicationSystem* RepSystem : RepSystems)
+		{
+			RepSystemCount += (RepSystem != nullptr ? 1 : 0);
+		}
 	}
 
 	virtual void ShutdownModule() override
@@ -62,7 +73,13 @@ private:
 			FModuleManager::Get().OnModulesChanged().Remove(ModulesChangedHandle);
 			ModulesChangedHandle.Reset();
 		}
+
 	
+		UE::Net::FReplicationSystemFactory::GetReplicationSystemCreatedDelegate().Remove(RepSysCreatedHandle);
+		UE::Net::FReplicationSystemFactory::GetReplicationSystemDestroyedDelegate().Remove(RepSysDestroyedHandle);
+		RepSysCreatedHandle.Reset();
+		RepSysDestroyedHandle.Reset();
+
 		UE_NET_IRIS_SHUTDOWN_LEGACY_PUSH_MODEL();
 	}
 
@@ -73,10 +90,13 @@ private:
 
 	void OnAllModuleLoadingPhasesComplete()
 	{
-		bAllowLoadedModulesUpdatedCallback = true;
-		UE::Net::Private::FInternalNetSerializerDelegates::BroadcastLoadedModulesUpdated();
-
 		FCoreDelegates::OnAllModuleLoadingPhasesComplete.RemoveAll(this);
+
+		bAllowLoadedModulesUpdatedCallback = true;
+		if (RepSystemCount > 0 && ShouldBroadcastLoadedModulesUpdated())
+		{
+			ForceBroadcastLoadedModulesUpdated();
+		}
 	}
 
 	void OnModulesChanged(FName ModuleThatChanged, EModuleChangeReason ReasonForChange)
@@ -85,12 +105,18 @@ private:
 		{
 			case EModuleChangeReason::ModuleLoaded:
 			{
-				UE::Net::Private::FInternalNetSerializerDelegates::BroadcastPreFreezeNetSerializerRegistry();
-				UE::Net::Private::FInternalNetSerializerDelegates::BroadcastPostFreezeNetSerializerRegistry();
-				if (bAllowLoadedModulesUpdatedCallback)
+				++LoadedModulesCount;
+				if (bAllowLoadedModulesUpdatedCallback && (RepSystemCount > 0))
 				{
-					UE::Net::Private::FInternalNetSerializerDelegates::BroadcastLoadedModulesUpdated();
+					if (!BroadcastModulesUpdatedHandle.IsValid())
+					{
+						LoadedModulesCountAtTickerCreation = LoadedModulesCount;
+						BroadcastModulesUpdatedHandle = FTSTicker::GetCoreTicker().AddTicker(FTickerDelegate::CreateRaw(this, &FIrisCoreModule::BroadcastLoadedModulesUpdated));
+					}
 				}
+
+				UE::Net::Private::FInternalNetSerializerDelegates::BroadcastPreFreezeNetSerializerRegistry();
+				UE::Net::Private::FInternalNetSerializerDelegates::BroadcastPostFreezeNetSerializerRegistry();
 			}
 			break;
 
@@ -101,8 +127,81 @@ private:
 		}
 	}
 
+	bool BroadcastLoadedModulesUpdated(float /* DeltaTime */)
+	{
+		// If we're still loading modules check again next frame
+		if (LoadedModulesCountAtTickerCreation != LoadedModulesCount)
+		{
+			LoadedModulesCountAtTickerCreation = LoadedModulesCount;
+			return true;
+		}
+		else
+		{
+			if (RepSystemCount > 0)
+			{
+				UE_LOG(LogIris, Warning, TEXT("FInternalNetSerializerDelegates::BroadcastLoadedModulesUpdated() called while there are %d active ReplicationSystems. If polymorphic types are registered we may have corrupt data. A restart of the ReplicationSystem or NetDriver is recommended. Total loaded modules: %u."), RepSystemCount, LoadedModulesCount);
+			}
+			else
+			{
+				UE_LOG(LogIris, Display, TEXT("FInternalNetSerializerDelegates::BroadcastLoadedModulesUpdated() called while there are %d active ReplicationSystems. This is %s. Total loaded modules: %u."), RepSystemCount, (RepSystemCount == 0 ? TEXT("good") : TEXT("unexpected")), LoadedModulesCount);
+			}
+
+			BroadcastModulesUpdatedHandle.Reset();
+			UE::Net::Private::FInternalNetSerializerDelegates::BroadcastLoadedModulesUpdated();
+			return false;
+		}
+	}
+
+	void ForceBroadcastLoadedModulesUpdated()
+	{
+		ResetBroadcastLoadedModulesTicker();
+		LoadedModulesCountAtTickerCreation = LoadedModulesCount;
+		BroadcastLoadedModulesUpdated(0.0f);
+	}
+
+	bool ShouldBroadcastLoadedModulesUpdated() const
+	{
+		return BroadcastModulesUpdatedHandle.IsValid() || (LoadedModulesCountAtTickerCreation != LoadedModulesCount);
+	}
+
+	void OnRepSystemCreated(UReplicationSystem*)
+	{
+		if (ShouldBroadcastLoadedModulesUpdated())
+		{
+			ForceBroadcastLoadedModulesUpdated();
+		}
+		else if (RepSystemCount == 0)
+		{
+			UE_LOG(LogIris, Display, TEXT("FInternalNetSerializerDelegates::BroadcastLoadedModulesUpdated() not called when creating ReplicationSystem since no additional modules have been loaded since last broadcast. This is good. Total loaded modules: %u."), RepSystemCount, LoadedModulesCount);
+		}
+
+		// Update RepSystemCount after broadcasting such that we get logging ideally saying there weren't any active replication systems.
+		++RepSystemCount;
+	}
+
+	void OnRepSystemDestroyed(UReplicationSystem*)
+	{
+		--RepSystemCount;
+		ensure(RepSystemCount >= 0);
+	}
+
+	void ResetBroadcastLoadedModulesTicker()
+	{
+		if (BroadcastModulesUpdatedHandle.IsValid())
+		{
+			FTSTicker::GetCoreTicker().RemoveTicker(BroadcastModulesUpdatedHandle);
+			BroadcastModulesUpdatedHandle.Reset();
+		}
+	}
+
 private:
 	FDelegateHandle ModulesChangedHandle;
+	FDelegateHandle RepSysCreatedHandle;
+	FDelegateHandle RepSysDestroyedHandle;
+	FTSTicker::FDelegateHandle BroadcastModulesUpdatedHandle;
+	int32 RepSystemCount = 0;
+	uint32 LoadedModulesCount = 0;
+	uint32 LoadedModulesCountAtTickerCreation = 0;
 	bool bAllowLoadedModulesUpdatedCallback = false;
 };
 IMPLEMENT_MODULE(FIrisCoreModule, IrisCore);
diff --git a/Engine/Source/Runtime/Experimental/Iris/Core/Private/Iris/Serialization/PolymorphicNetSerializer.cpp b/Engine/Source/Runtime/Experimental/Iris/Core/Private/Iris/Serialization/PolymorphicNetSerializer.cpp
--- a/Engine/Source/Runtime/Experimental/Iris/Core/Private/Iris/Serialization/PolymorphicNetSerializer.cpp
+++ b/Engine/Source/Runtime/Experimental/Iris/Core/Private/Iris/Serialization/PolymorphicNetSerializer.cpp
@@ -28,12 +28,19 @@ void FPolymorphicNetSerializerScriptStructCache::InitForType(const UScriptStruct
 
 	// Find all script structs of this type and add them to the list and build descriptor
 	// (not sure of a better way to do this but it should only happen at startup)
-	for (TObjectIterator<UScriptStruct> It; It; ++It)
+	TArray<UObject*> Structs;
+	Structs.Reserve(EstimatedScriptStructCount);
+
+	constexpr bool bIncludeDerivedClasses = true;
+	GetObjectsOfClass(UScriptStruct::StaticClass(), Structs, bIncludeDerivedClasses, RF_ClassDefaultObject, GetObjectIteratorDefaultInternalExclusionFlags(EInternalObjectFlags::None));
+	EstimatedScriptStructCount = Structs.Max();
+	for (const UObject* Object : Structs)
 	{
-		if (It->IsChildOf(InScriptStruct))
+		const UScriptStruct* Struct = static_cast<const UScriptStruct*>(Object);
+		if (Struct->IsChildOf(InScriptStruct))
 		{
 			FTypeInfo Entry;
-			Entry.ScriptStruct = *It;
+			Entry.ScriptStruct = Struct;
 
 			// See if we already had a descriptor for the struct
 			const uint32 ExistingTypeIndex = GetTypeIndex(Entry.ScriptStruct);
@@ -45,7 +52,7 @@ void FPolymorphicNetSerializerScriptStructCache::InitForType(const UScriptStruct
 			{
 				// Get or create descriptor
 				FReplicationStateDescriptorBuilder::FParameters Params;
-				Entry.Descriptor = FReplicationStateDescriptorBuilder::CreateDescriptorForStruct(*It, Params);
+				Entry.Descriptor = FReplicationStateDescriptorBuilder::CreateDescriptorForStruct(Struct, Params);
 				bFoundNewScriptStructs = true;
 			}
 			
@@ -57,7 +64,7 @@ void FPolymorphicNetSerializerScriptStructCache::InitForType(const UScriptStruct
 			}
 			else
 			{
-				UE_LOG(LogIris, Error, TEXT("FPolymorphicNetSerializerScriptStructCache::InitForType Failed to create descriptor for type %s when building cache for base %s"), ToCStr((*It)->GetName()), ToCStr(InScriptStruct->GetName()));
+				UE_LOG(LogIris, Error, TEXT("FPolymorphicNetSerializerScriptStructCache::InitForType Failed to create descriptor for type %s when building cache for base %s"), ToCStr(Struct->GetName()), ToCStr(InScriptStruct->GetName()));
 			}
 		}
 	}
diff --git a/Engine/Source/Runtime/Experimental/Iris/Core/Public/Iris/Serialization/PolymorphicNetSerializer.h b/Engine/Source/Runtime/Experimental/Iris/Core/Public/Iris/Serialization/PolymorphicNetSerializer.h
--- a/Engine/Source/Runtime/Experimental/Iris/Core/Public/Iris/Serialization/PolymorphicNetSerializer.h
+++ b/Engine/Source/Runtime/Experimental/Iris/Core/Public/Iris/Serialization/PolymorphicNetSerializer.h
@@ -38,6 +38,7 @@ struct FPolymorphicNetSerializerScriptStructCache
 private:
 	TArray<FTypeInfo> RegisteredTypes;
 	EReplicationStateTraits CommonTraits = EReplicationStateTraits::None;
+	int32 EstimatedScriptStructCount = 4096;
 };
 
 }
