Store Implicit in single precision

diff --git a/Engine/Source/Programs/HeadlessChaos/Private/HeadlessChaosTestSim.cpp b/Engine/Source/Programs/HeadlessChaos/Private/HeadlessChaosTestSim.cpp
--- a/Engine/Source/Programs/HeadlessChaos/Private/HeadlessChaosTestSim.cpp
+++ b/Engine/Source/Programs/HeadlessChaos/Private/HeadlessChaosTestSim.cpp
@@ -22,7 +22,7 @@ namespace ChaosTest {
 		// The goal is to make sure that the bounds are updated correctly and the dynamic rests on top of the static
 		// in its final position.
 
-		auto Sphere = Chaos::FImplicitObjectPtr(new TSphere<FReal, 3>(FVec3(0), 10));
+		Chaos::FImplicitObjectPtr Sphere = Chaos::FImplicitObjectPtr(new Chaos::FSphere(FVec3(0), 10));
 
 		// Create solver #TODO make FFramework a little more general instead of mostly geometry collection focused
 		GeometryCollectionTest::FFramework Framework;
@@ -73,7 +73,7 @@ namespace ChaosTest {
 		TUniquePtr<FChaosPhysicsMaterial> PhysicsMaterial = MakeUnique<FChaosPhysicsMaterial>();
 		PhysicsMaterial->SleepCounterThreshold = 2;
 
-		Chaos::FImplicitObjectPtr Sphere(new TSphere<FReal, 3>(FVec3(0, 0, 0), 50));
+		Chaos::FImplicitObjectPtr Sphere(new Chaos::FSphere(FVec3(0, 0, 0), 50));
 		Static->SetGeometry(Sphere);
 		Dynamic->SetGeometry(Sphere);
 
@@ -264,7 +264,7 @@ namespace ChaosTest {
 
 	GTEST_TEST(AllTraits, DISABLED_SimTests_MidSubstepSleep)
 	{
-		Chaos::FImplicitObjectPtr Sphere{new TSphere<FReal, 3>(FVec3(0), 10)};
+		Chaos::FImplicitObjectPtr Sphere{new Chaos::FSphere(FVec3(0), 10)};
 
 		FChaosSolversModule* Module = FChaosSolversModule::GetModule();
 		auto Solver = Module->CreateSolver(nullptr, /*AsyncDt=*/-1);
diff --git a/Engine/Source/Programs/HeadlessChaos/Private/HeadlessChaosTestSolverProxies.cpp b/Engine/Source/Programs/HeadlessChaos/Private/HeadlessChaosTestSolverProxies.cpp
--- a/Engine/Source/Programs/HeadlessChaos/Private/HeadlessChaosTestSolverProxies.cpp
+++ b/Engine/Source/Programs/HeadlessChaos/Private/HeadlessChaosTestSolverProxies.cpp
@@ -22,7 +22,7 @@ namespace ChaosTest {
 
 	void SingleParticleProxySingleThreadTest()
 	{
-		auto Sphere = Chaos::FImplicitObjectPtr(new TSphere<FReal, 3>(FVec3(0), 10));
+		Chaos::FImplicitObjectPtr Sphere = Chaos::FImplicitObjectPtr(new Chaos::FSphere(FVec3(0), 10));
 
 		FChaosSolversModule* Module = FChaosSolversModule::GetModule();
 
@@ -64,7 +64,7 @@ namespace ChaosTest {
 	void SingleParticleProxyWakeEventPropagationTest()
 	{
 		using namespace Chaos;
-		auto Sphere = Chaos::FImplicitObjectPtr(new TSphere<FReal, 3>(FVec3(0), 10));
+		auto Sphere = Chaos::FImplicitObjectPtr(new Chaos::FSphere(FVec3(0), 10));
 
 		FChaosSolversModule* Module = FChaosSolversModule::GetModule();
 
@@ -122,7 +122,7 @@ namespace ChaosTest {
 	{
 		// this test make sure that we are not leaking uniqueIdx when creating and destroying particle in one frame (without PT running)
 
-		auto Sphere = Chaos::FImplicitObjectPtr(new TSphere<FReal, 3>(FVec3(0), 10));
+		auto Sphere = Chaos::FImplicitObjectPtr(new Chaos::FSphere(FVec3(0), 10));
 
 		FChaosSolversModule* Module = FChaosSolversModule::GetModule();
 
@@ -142,7 +142,7 @@ namespace ChaosTest {
 		// keep track of the actual unique Idx
 		FUniqueIdx FirstIdx = Particle.UniqueIdx();
 
-		// unregister the objkect in the same GT frame 
+		// unregister the object in the same GT frame 
 		Solver->UnregisterObject(Proxy);
 
 		// run PT to make sure the callbacks are running and updating the internal pending lists 
@@ -170,7 +170,7 @@ namespace ChaosTest {
 		// This array is physics thread only, it should only be added after ticking solver.
 
 		using namespace Chaos;
-		auto Sphere = Chaos::FImplicitObjectPtr(new TSphere<FReal, 3>(FVec3(0), 10));
+		auto Sphere = Chaos::FImplicitObjectPtr(new Chaos::FSphere(FVec3(0), 10));
 		FChaosSolversModule* Module = FChaosSolversModule::GetModule();
 
 		// Make a solver
diff --git a/Engine/Source/Programs/HeadlessChaos/Private/HeadlessChaosTestSweep.cpp b/Engine/Source/Programs/HeadlessChaos/Private/HeadlessChaosTestSweep.cpp
--- a/Engine/Source/Programs/HeadlessChaos/Private/HeadlessChaosTestSweep.cpp
+++ b/Engine/Source/Programs/HeadlessChaos/Private/HeadlessChaosTestSweep.cpp
@@ -949,7 +949,7 @@ namespace ChaosTest
 		FSweepResultData SweepSphere(const FTriangleMeshImplicitObject& TriangleMesh, const FVec3& SweepStart, const FVec3& SweepDirection, const FReal SweepLength = 100)
 		{
 			const FReal SphereRadius = 25;
-			const TSphere<FReal, 3> Sphere(FVec3::ZeroVector, SphereRadius);
+			const Chaos::FSphere Sphere(FVec3::ZeroVector, SphereRadius);
 			const FRigidTransform3 StartTM(SweepStart, TRotation<FReal, 3>::Identity);
 			const FVec3 NormalizedSweepDirection = SweepDirection.GetSafeNormal();
 
diff --git a/Engine/Source/Programs/HeadlessChaos/Private/HeadlessChaosTestUtility.cpp b/Engine/Source/Programs/HeadlessChaos/Private/HeadlessChaosTestUtility.cpp
--- a/Engine/Source/Programs/HeadlessChaos/Private/HeadlessChaosTestUtility.cpp
+++ b/Engine/Source/Programs/HeadlessChaos/Private/HeadlessChaosTestUtility.cpp
@@ -36,7 +36,7 @@ namespace ChaosTest {
 		InParticles.InvM(RigidBodyIndex) = 1.0;
 		InParticles.I(RigidBodyIndex) = TVec3<FRealSingle>(1, 1, 1);
 		InParticles.InvI(RigidBodyIndex) = TVec3<FRealSingle>(1, 1, 1);
-		InParticles.SetGeometry(RigidBodyIndex, MakeImplicitObjectPtr<TSphere<FReal, 3>>(FVec3(0), Scale));
+		InParticles.SetGeometry(RigidBodyIndex, MakeImplicitObjectPtr<Chaos::FSphere>(FVec3(0), Scale));
 		InParticles.SetObjectState(RigidBodyIndex, EObjectStateType::Dynamic);
 
 		return RigidBodyIndex;
@@ -242,7 +242,7 @@ namespace ChaosTest {
 		Particle->I() = TVec3<FRealSingle>(ScaleSq / 6.0);
 		Particle->InvI() = TVec3<FRealSingle>(6.0 / ScaleSq);
 
-		Particle->SetGeometry(MakeImplicitObjectPtr<TSphere<FReal, 3>>(FVec3(0), Scale.X / 2.0));
+		Particle->SetGeometry(MakeImplicitObjectPtr<Chaos::FSphere>(FVec3(0), Scale.X / 2.0));
 
 		int32 CollisionIndex = 0;
 		Particle->CollisionParticlesInitIfNeeded();
diff --git a/Engine/Source/Programs/HeadlessChaos/Private/HeadlessChaosTestVehicles.cpp b/Engine/Source/Programs/HeadlessChaos/Private/HeadlessChaosTestVehicles.cpp
--- a/Engine/Source/Programs/HeadlessChaos/Private/HeadlessChaosTestVehicles.cpp
+++ b/Engine/Source/Programs/HeadlessChaos/Private/HeadlessChaosTestVehicles.cpp
@@ -1516,7 +1516,7 @@ namespace ChaosTest
 		TUniquePtr<FChaosPhysicsMaterial> PhysicsMaterial = MakeUnique<FChaosPhysicsMaterial>();
 		PhysicsMaterial->SleepCounterThreshold = 2;
 
-		Chaos::FImplicitObjectPtr Box(new TSphere<FReal, 3>(FVec3(0, 0, 0), 50));
+		Chaos::FImplicitObjectPtr Box(new Chaos::FSphere(FVec3(0, 0, 0), 50));
 		Dynamic->SetGeometry(Box);
 
 		Evolution.SetPhysicsMaterial(Dynamic, MakeSerializable(PhysicsMaterial));
diff --git a/Engine/Source/Runtime/Engine/Private/Collision/CollisionDebugDrawing.cpp b/Engine/Source/Runtime/Engine/Private/Collision/CollisionDebugDrawing.cpp
--- a/Engine/Source/Runtime/Engine/Private/Collision/CollisionDebugDrawing.cpp
+++ b/Engine/Source/Runtime/Engine/Private/Collision/CollisionDebugDrawing.cpp
@@ -209,14 +209,14 @@ void DrawGeomOverlaps(const UWorld* InWorld, const Chaos::FImplicitObject& Geom,
 	{
 		DrawBoxOverlap(InWorld, Pos, Box->Extents() * 0.5f, Rot, Overlaps, Lifetime);
 	}
-	else if (const auto Sphere = Geom.template GetObject<TSphere<FReal, 3>>())
+	else if (const auto Sphere = Geom.template GetObject<Chaos::FSphere>())
 	{
-		DrawSphereOverlap(InWorld, Pos, Sphere->GetRadius(), Overlaps, Lifetime);
+		DrawSphereOverlap(InWorld, Pos, Sphere->GetRadiusf(), Overlaps, Lifetime);
 	}
 	else if (const auto Capsule = Geom.template GetObject<FCapsule>())
 	{
 		// Convert here from Chaos to unreal definition of capsule height
-		DrawCapsuleOverlap(InWorld, Pos, Capsule->GetHeight() * 0.5f + Capsule->GetRadius(), Capsule->GetRadius(), Rot, Overlaps, Lifetime);
+		DrawCapsuleOverlap(InWorld, Pos, Capsule->GetHeightf() * 0.5f + Capsule->GetRadiusf(), Capsule->GetRadiusf(), Rot, Overlaps, Lifetime);
 	}
 }
 
@@ -227,14 +227,14 @@ void DrawGeomSweeps(const UWorld* InWorld, const FVector& Start, const FVector&
 	{
 		DrawBoxSweeps(InWorld, Start, End, Box->Extents() * 0.5f, Rotation, Hits, Lifetime);
 	}
-	else if (const auto Sphere = Geom.template GetObject<TSphere<FReal, 3>>())
+	else if (const Chaos::FSphere* Sphere = Geom.template GetObject<Chaos::FSphere>())
 	{
-		DrawSphereSweeps(InWorld, Start, End, Sphere->GetRadius(), Hits, Lifetime);
+		DrawSphereSweeps(InWorld, Start, End, Sphere->GetRadiusf(), Hits, Lifetime);
 	}
 	else if (const auto Capsule = Geom.template GetObject<FCapsule>())
 	{
 		// Convert here from Chaos to unreal definition of capsule height
-		DrawCapsuleSweeps(InWorld, Start, End, Capsule->GetHeight() * 0.5f + Capsule->GetRadius(), Capsule->GetRadius(), Rotation, Hits, Lifetime);
+		DrawCapsuleSweeps(InWorld, Start, End, Capsule->GetHeightf() * 0.5f + Capsule->GetRadiusf(), Capsule->GetRadiusf(), Rotation, Hits, Lifetime);
 	}
 }
 
diff --git a/Engine/Source/Runtime/Engine/Private/Collision/SceneQuery.cpp b/Engine/Source/Runtime/Engine/Private/Collision/SceneQuery.cpp
--- a/Engine/Source/Runtime/Engine/Private/Collision/SceneQuery.cpp
+++ b/Engine/Source/Runtime/Engine/Private/Collision/SceneQuery.cpp
@@ -187,14 +187,14 @@ bool BuildQueryShape(const FPhysicsGeometryCollection& GeometryCollection, Chaos
 		}
 		case ECollisionShapeType::Sphere:
 		{
-			const Chaos::TSphere<Chaos::FReal, 3>& Sphere = GeometryCollection.GetSphereGeometry();
-			OutQueryShape.CollisionShape = FCollisionShape::MakeSphere(Sphere.GetRadius());
+			const Chaos::FSphere& Sphere = GeometryCollection.GetSphereGeometry();
+			OutQueryShape.CollisionShape = FCollisionShape::MakeSphere(Sphere.GetRadiusf());
 			return true;
 		}
 		case ECollisionShapeType::Capsule:
 		{
 			const Chaos::FCapsule& Capsule = GeometryCollection.GetCapsuleGeometry();
-			OutQueryShape.CollisionShape = FCollisionShape::MakeCapsule(Capsule.GetRadius(), Capsule.GetHeight() * 0.5f);
+			OutQueryShape.CollisionShape = FCollisionShape::MakeCapsule(Capsule.GetRadiusf(), Capsule.GetHeightf() * 0.5f);
 			return true;
 		}
 		default:
diff --git a/Engine/Source/Runtime/Engine/Private/PhysicsEngine/BodyInstance.cpp b/Engine/Source/Runtime/Engine/Private/PhysicsEngine/BodyInstance.cpp
--- a/Engine/Source/Runtime/Engine/Private/PhysicsEngine/BodyInstance.cpp
+++ b/Engine/Source/Runtime/Engine/Private/PhysicsEngine/BodyInstance.cpp
@@ -2218,7 +2218,7 @@ bool FBodyInstance::UpdateBodyScale(const FVector& InScale3D, bool bForceUpdate)
 
 					FReal Radius = FMath::Max<FReal>(SphereElem->Radius * AdjustedScale3DAbs.X, FCollisionShape::MinSphereRadius());
 					FVec3 Center = RelativeTM.TransformPosition(SphereElem->Center) * InScale3D;
-					Chaos::FImplicitObjectPtr NewSphere = MakeImplicitObjectPtr<TSphere<FReal, 3>>(Center, Radius);
+					Chaos::FImplicitObjectPtr NewSphere = MakeImplicitObjectPtr<Chaos::FSphere>(Center, Radius);
 
 					NewGeometry.Emplace(MoveTemp(NewSphere));
 					bSuccess = true;
diff --git a/Engine/Source/Runtime/Engine/Private/PhysicsEngine/CollisionAnalyzerCapture.h b/Engine/Source/Runtime/Engine/Private/PhysicsEngine/CollisionAnalyzerCapture.h
--- a/Engine/Source/Runtime/Engine/Private/PhysicsEngine/CollisionAnalyzerCapture.h
+++ b/Engine/Source/Runtime/Engine/Private/PhysicsEngine/CollisionAnalyzerCapture.h
@@ -30,7 +30,7 @@ inline static void ConvertGeometryCollection(const FPhysicsGeometryCollection& G
 		{
 			OutQueryShape = ECAQueryShape::Capsule;
 			const auto& CapsuleGeom = GeomCollection.GetCapsuleGeometry();
-			OutDims = FVector(CapsuleGeom.GetRadius(), CapsuleGeom.GetRadius(), CapsuleGeom.GetHeight() * 0.5f + CapsuleGeom.GetRadius());
+			OutDims = FVector(CapsuleGeom.GetRadiusf(), CapsuleGeom.GetRadiusf(), CapsuleGeom.GetHeightf() * 0.5f + CapsuleGeom.GetRadiusf());
 			break;
 		}
 
@@ -38,7 +38,7 @@ inline static void ConvertGeometryCollection(const FPhysicsGeometryCollection& G
 		{
 			OutQueryShape = ECAQueryShape::Sphere;
 			const auto& SphereGeom = GeomCollection.GetSphereGeometry();
-			OutDims = FVector(SphereGeom.GetRadius());
+			OutDims = FVector(SphereGeom.GetRadiusf());
 			break;
 		}
 
diff --git a/Engine/Source/Runtime/Engine/Private/PhysicsEngine/ImmediatePhysics/ImmediatePhysicsChaos/ImmediatePhysicsActorHandle_Chaos.cpp b/Engine/Source/Runtime/Engine/Private/PhysicsEngine/ImmediatePhysics/ImmediatePhysicsChaos/ImmediatePhysicsActorHandle_Chaos.cpp
--- a/Engine/Source/Runtime/Engine/Private/PhysicsEngine/ImmediatePhysics/ImmediatePhysicsChaos/ImmediatePhysicsActorHandle_Chaos.cpp
+++ b/Engine/Source/Runtime/Engine/Private/PhysicsEngine/ImmediatePhysics/ImmediatePhysicsChaos/ImmediatePhysicsActorHandle_Chaos.cpp
@@ -39,7 +39,7 @@ namespace ImmediatePhysics_Chaos
 		NewShape->SetSimEnabled(false);
 
 		OutMass = Mass;
-		OutInertia = TSphere<FReal, 3>::GetInertiaTensor(Mass, Radius).GetDiagonal();
+		OutInertia = Chaos::FSphere::GetInertiaTensor(Mass, Radius).GetDiagonal();
 		OutCoMTransform = FTransform::Identity;
 		OutShapes.Emplace(MoveTemp(NewShape));
 		OutGeom = MoveTemp(ImplicitSphere);
diff --git a/Engine/Source/Runtime/Engine/Private/SkeletalMeshComponentPhysics.cpp b/Engine/Source/Runtime/Engine/Private/SkeletalMeshComponentPhysics.cpp
--- a/Engine/Source/Runtime/Engine/Private/SkeletalMeshComponentPhysics.cpp
+++ b/Engine/Source/Runtime/Engine/Private/SkeletalMeshComponentPhysics.cpp
@@ -3167,8 +3167,8 @@ void USkeletalMeshComponent::ProcessClothCollisionWithEnvironment()
 					auto AddSphere = [&NewCollisionData](const FImplicitSphere3& ImplicitSphere, const FTransform& ComponentToClothTransform, const FVec3& Scale = FVec3::OneVector)
 					{
 						NewCollisionData.Spheres.Emplace(
-							ImplicitSphere.GetRadius() * Scale.X,  // Assumes uniform scale
-							ComponentToClothTransform.TransformPosition(ImplicitSphere.GetCenter()));
+							ImplicitSphere.GetRadiusf() * Scale.X,  // Assumes uniform scale
+							ComponentToClothTransform.TransformPosition(FVec3(ImplicitSphere.GetCenterf())));
 					};
 
 					auto AddBox = [&NewCollisionData, Thickness](const FImplicitBox3& ImplicitBox, const FTransform& ComponentToClothTransform, const FVec3& Scale = FVec3::OneVector)
@@ -3183,10 +3183,10 @@ void USkeletalMeshComponent::ProcessClothCollisionWithEnvironment()
 						{
 							const int32 BaseSphereIndex = NewCollisionData.Spheres.Num();
 
-							const FReal Radius = ImplicitCapsule.GetRadius() * Scale.X;  // Assumes uniform scale
-							const FReal HalfHeight = ImplicitCapsule.GetHeight() * (FReal)0.5;
+							const FReal Radius = ImplicitCapsule.GetRadiusf() * Scale.X;  // Assumes uniform scale
+							const FReal HalfHeight = ImplicitCapsule.GetHeightf() * (FReal)0.5;
 							const FVector HalfSegment = ComponentToClothTransform.GetUnitAxis(EAxis::X) * HalfHeight * Scale.X;  // Assumes uniform scale
-							const FVector TransformedCenter = ComponentToClothTransform.TransformPosition(ImplicitCapsule.GetCenter());
+							const FVector TransformedCenter = ComponentToClothTransform.TransformPosition(FVec3(ImplicitCapsule.GetCenterf()));
 
 							NewCollisionData.Spheres.Emplace(Radius, TransformedCenter + HalfSegment);
 							NewCollisionData.Spheres.Emplace(Radius, TransformedCenter - HalfSegment);
diff --git a/Engine/Source/Runtime/Experimental/Chaos/Private/Chaos/BoundingVolumeHierarchy.cpp b/Engine/Source/Runtime/Experimental/Chaos/Private/Chaos/BoundingVolumeHierarchy.cpp
--- a/Engine/Source/Runtime/Experimental/Chaos/Private/Chaos/BoundingVolumeHierarchy.cpp
+++ b/Engine/Source/Runtime/Experimental/Chaos/Private/Chaos/BoundingVolumeHierarchy.cpp
@@ -706,7 +706,7 @@ namespace Chaos
 }
 
 //template class CHAOS_API Chaos::TBoundingVolume<int32, FReal, 3>;
-template class Chaos::TBoundingVolumeHierarchy<TArray<Chaos::TSphere<Chaos::FReal, 3>*>, TArray<int32>, Chaos::FReal, 3>;
+template class Chaos::TBoundingVolumeHierarchy<TArray<Chaos::FSphere*>, TArray<int32>, Chaos::FReal, 3>;
 template class Chaos::TBoundingVolumeHierarchy<TArray<Chaos::TTetrahedron<Chaos::FReal>*>, TArray<int32>, Chaos::FReal, 3>;
 template class Chaos::TBoundingVolumeHierarchy<Chaos::TPBDRigidParticles<Chaos::FReal, 3>, TArray<int32>, Chaos::FReal, 3>;
 template class Chaos::TBoundingVolumeHierarchy<Chaos::FParticles, TArray<int32>, Chaos::FReal, 3>;
diff --git a/Engine/Source/Runtime/Experimental/Chaos/Private/Chaos/ChaosDebugDraw.cpp b/Engine/Source/Runtime/Experimental/Chaos/Private/Chaos/ChaosDebugDraw.cpp
--- a/Engine/Source/Runtime/Experimental/Chaos/Private/Chaos/ChaosDebugDraw.cpp
+++ b/Engine/Source/Runtime/Experimental/Chaos/Private/Chaos/ChaosDebugDraw.cpp
@@ -809,7 +809,7 @@ namespace Chaos
 				const TImplicitObjectScaled<FLevelSet, bInstanced>* Scaled = Implicit->template GetObject<TImplicitObjectScaled<FLevelSet, bInstanced>>();
 				// even though thhe levelset is scaled, the debugdraw uses the collisionParticles  that are pre-scaled
 				// so no need to pass the scaled transform and just extract the wrapped LevelSet
-				DrawShapesImpl(Particle, ShapeTransform, Scaled->GetUnscaledObject(), Shape, Scaled->GetMargin(), Color, Duration, Settings);
+				DrawShapesImpl(Particle, ShapeTransform, Scaled->GetUnscaledObject(), Shape, Scaled->GetMarginf(), Color, Duration, Settings);
 				break;
 			}
 			case ImplicitObjectType::Unknown:
@@ -818,7 +818,7 @@ namespace Chaos
 			{
 				const TImplicitObjectScaled<FConvex, bInstanced>* Scaled = Implicit->template GetObject<TImplicitObjectScaled<FConvex, bInstanced>>();
 				ScaleTM.SetScale3D(Scaled->GetScale());
-				DrawShapesImpl(Particle, ScaleTM * ShapeTransform, Scaled->GetUnscaledObject(), Shape, Scaled->GetMargin(), Color, Duration, Settings);
+				DrawShapesImpl(Particle, ScaleTM * ShapeTransform, Scaled->GetUnscaledObject(), Shape, Scaled->GetMarginf(), Color, Duration, Settings);
 				break;
 			}
 			case ImplicitObjectType::TaperedCylinder:
@@ -872,7 +872,7 @@ namespace Chaos
 			case ImplicitObjectType::Convex:
 			{
 				const TImplicitObjectInstanced<FConvex>* Instanced = Implicit->template GetObject<TImplicitObjectInstanced<FConvex>>();
-				DrawShapesImpl(Particle, ShapeTransform, Instanced->GetInstancedObject(), Shape, Instanced->GetMargin(), Color, Duration, Settings);
+				DrawShapesImpl(Particle, ShapeTransform, Instanced->GetInstancedObject(), Shape, Instanced->GetMarginf(), Color, Duration, Settings);
 				break;
 			}
 			case ImplicitObjectType::TaperedCylinder:
@@ -1066,10 +1066,10 @@ namespace Chaos
 			{
 			case ImplicitObjectType::Sphere:
 			{
-				const TSphere<FReal, 3>* Sphere = Implicit->template GetObject<TSphere<FReal, 3>>();
-				const FVec3 P = ShapeTransform.TransformPosition(Sphere->GetCenter());
+				const FSphere* Sphere = Implicit->template GetObject<FSphere>();
+				const FVec3 P = ShapeTransform.TransformPosition(FVec3(Sphere->GetCenterf()));
 				const FVec3 Scale = ShapeTransform.GetScale3D();
-				DrawDebugSphere(P, Sphere->GetRadius() * Scale.Z, 8, ShapeColor, false, Duration, uint8(Settings.DrawPriority), Settings.ShapeThicknesScale * Settings.LineThickness);
+				DrawDebugSphere(P, Sphere->GetRadiusf() * Scale.Z, 8, ShapeColor, false, Duration, uint8(Settings.DrawPriority), Settings.ShapeThicknesScale * Settings.LineThickness);
 				break;
 			}
 			case ImplicitObjectType::Box:
@@ -1085,10 +1085,10 @@ namespace Chaos
 			case ImplicitObjectType::Capsule:
 			{
 				const FCapsule* Capsule = Implicit->template GetObject<FCapsule>();
-				const FVec3 P = ShapeTransform.TransformPosition(Capsule->GetCenter());
+				const FVec3 P = ShapeTransform.TransformPosition(FVec3(Capsule->GetCenterf()));
 				const FRotation3 Q = ShapeTransform.GetRotation() * FRotationMatrix::MakeFromZ(Capsule->GetAxis());
 				const FVec3 Scale = ShapeTransform.GetScale3D();
-				DrawDebugCapsule(P, ((FReal)0.5 * Capsule->GetHeight() + Capsule->GetRadius()) * Scale.Z, Capsule->GetRadius() * Scale.X, Q, ShapeColor, false, Duration, uint8(Settings.DrawPriority), Settings.ShapeThicknesScale * Settings.LineThickness);
+				DrawDebugCapsule(P, ((FReal)0.5 * Capsule->GetHeightf() + Capsule->GetRadiusf()) * Scale.Z, Capsule->GetRadiusf() * Scale.X, Q, ShapeColor, false, Duration, uint8(Settings.DrawPriority), Settings.ShapeThicknesScale * Settings.LineThickness);
 				break;
 			}
 			case ImplicitObjectType::LevelSet:
diff --git a/Engine/Source/Runtime/Experimental/Chaos/Private/Chaos/Collision/CapsuleTriangleContactPoint.cpp b/Engine/Source/Runtime/Experimental/Chaos/Private/Chaos/Collision/CapsuleTriangleContactPoint.cpp
--- a/Engine/Source/Runtime/Experimental/Chaos/Private/Chaos/Collision/CapsuleTriangleContactPoint.cpp
+++ b/Engine/Source/Runtime/Experimental/Chaos/Private/Chaos/Collision/CapsuleTriangleContactPoint.cpp
@@ -51,11 +51,11 @@ namespace Chaos
 		//
 
 		// Capsule segment
-		const FVec3 P0 = Capsule.GetX1();
-		const FVec3 P1 = Capsule.GetX2();
-		const FReal R = Capsule.GetRadius();
-		const FVec3 Axis = Capsule.GetAxis();
-		const FReal L = Capsule.GetHeight();
+		const FVec3 P0 = Capsule.GetX1f();
+		const FVec3 P1 = Capsule.GetX2f();
+		const FReal R = Capsule.GetRadiusf();
+		const FVec3 Axis = Capsule.GetAxisf();
+		const FReal L = Capsule.GetHeightf();
 		const FReal RejectDistance = R + CullDistance;
 		const FReal RejectDistanceSq = RejectDistance * RejectDistance;
 
@@ -68,7 +68,7 @@ namespace Chaos
 		// Tolerances
 		// @todo(chaos): Everything below needs to be made size agnostic (i.e., all checks with a distance tolerance
 		// need to cope with very larger or small segments and triangles). This will probably not be good enough...(add some tests)
-		const FReal DistanceTolerance = FReal(1.e-5) * Capsule.GetHeight();
+		const FReal DistanceTolerance = FReal(1.e-5f * Capsule.GetHeightf());
 		const FReal NormalTolerance = FReal(1.e-5);
 		const FReal NormalToleranceSq = NormalTolerance * NormalTolerance;
 		const FReal FaceContactSinAngleThreshold = FReal(0.34);	// ~Sin(20deg)
diff --git a/Engine/Source/Runtime/Experimental/Chaos/Private/Chaos/Collision/ContactPointsMiscShapes.cpp b/Engine/Source/Runtime/Experimental/Chaos/Private/Chaos/Collision/ContactPointsMiscShapes.cpp
--- a/Engine/Source/Runtime/Experimental/Chaos/Private/Chaos/Collision/ContactPointsMiscShapes.cpp
+++ b/Engine/Source/Runtime/Experimental/Chaos/Private/Chaos/Collision/ContactPointsMiscShapes.cpp
@@ -269,16 +269,16 @@ namespace Chaos
 	}
 
 
-	FContactPoint SphereSphereContactPoint(const TSphere<FReal, 3>& Sphere1, const FRigidTransform3& Sphere1Transform, const TSphere<FReal, 3>& Sphere2, const FRigidTransform3& Sphere2Transform, const FReal CullDistance)
+	FContactPoint SphereSphereContactPoint(const FSphere& Sphere1, const FRigidTransform3& Sphere1Transform, const FSphere& Sphere2, const FRigidTransform3& Sphere2Transform, const FReal CullDistance)
 	{
 		FContactPoint Result;
 
-		const FReal R1 = Sphere1.GetRadius();
-		const FReal R2 = Sphere2.GetRadius();
+		const FReal R1 = Sphere1.GetRadiusf();
+		const FReal R2 = Sphere2.GetRadiusf();
 
 		// World-space contact
-		const FVec3 Center1 = Sphere1Transform.TransformPosition(Sphere1.GetCenter());
-		const FVec3 Center2 = Sphere2Transform.TransformPosition(Sphere2.GetCenter());
+		const FVec3 Center1 = Sphere1Transform.TransformPosition(FVec3(Sphere1.GetCenterf()));
+		const FVec3 Center2 = Sphere2Transform.TransformPosition(FVec3(Sphere2.GetCenterf()));
 		const FVec3 Direction = Center1 - Center2;
 		const FReal SizeSq = Direction.SizeSquared();
 		const FReal CullDistanceSq = FMath::Square(CullDistance + R1 + R2);
@@ -288,21 +288,21 @@ namespace Chaos
 			const FVec3 Normal = Size > UE_SMALL_NUMBER ? Direction / Size : FVec3(0, 0, 1);
 			const FReal NewPhi = Size - (R1 + R2);
 
-			Result.ShapeContactPoints[0] = Sphere1.GetCenter() - Sphere1Transform.InverseTransformVector(R1 * Normal);
-			Result.ShapeContactPoints[1] = Sphere2.GetCenter() + Sphere2Transform.InverseTransformVector(R2 * Normal);
+			Result.ShapeContactPoints[0] = FVec3(Sphere1.GetCenterf()) - Sphere1Transform.InverseTransformVector(R1 * Normal);
+			Result.ShapeContactPoints[1] = FVec3(Sphere2.GetCenterf()) + Sphere2Transform.InverseTransformVector(R2 * Normal);
 			Result.ShapeContactNormal = Sphere2Transform.InverseTransformVector(Normal);
 			Result.Phi = NewPhi;
 		}
 		return Result;
 	}
 
-	FContactPoint SpherePlaneContactPoint(const TSphere<FReal, 3>& Sphere, const FRigidTransform3& SphereTransform, const TPlane<FReal, 3>& Plane, const FRigidTransform3& PlaneTransform)
+	FContactPoint SpherePlaneContactPoint(const FSphere& Sphere, const FRigidTransform3& SphereTransform, const TPlane<FReal, 3>& Plane, const FRigidTransform3& PlaneTransform)
 	{
 		FContactPoint Result;
 
-		FReal SphereRadius = Sphere.GetRadius();
+		FReal SphereRadius = Sphere.GetRadiusf();
 
-		FVec3 SpherePosWorld = SphereTransform.TransformPosition(Sphere.GetCenter());
+		FVec3 SpherePosWorld = SphereTransform.TransformPosition(FVec3(Sphere.GetCenterf()));
 		FVec3 SpherePosPlane = PlaneTransform.InverseTransformPosition(SpherePosWorld);
 
 		FVec3 NormalPlane;
@@ -318,19 +318,20 @@ namespace Chaos
 		return Result;
 	}
 
-	FContactPoint SphereBoxContactPoint(const TSphere<FReal, 3>& Sphere, const FRigidTransform3& SphereTransform, const FImplicitBox3& Box, const FRigidTransform3& BoxTransform)
+	FContactPoint SphereBoxContactPoint(const FSphere& Sphere, const FRigidTransform3& SphereTransform, const FImplicitBox3& Box, const FRigidTransform3& BoxTransform)
 	{
 		FContactPoint Result;
 
-		const FVec3 SphereWorld = SphereTransform.TransformPosition(Sphere.GetCenter());	// World-space sphere pos
+
+		const FVec3 SphereWorld = SphereTransform.TransformPosition(FVec3(Sphere.GetCenterf()));	// World-space sphere pos
 		const FVec3 SphereBox = BoxTransform.InverseTransformPosition(SphereWorld);			// Box-space sphere pos
 
 		FVec3 NormalBox;																	// Box-space normal
 		FReal PhiToSphereCenter = Box.PhiWithNormal(SphereBox, NormalBox);
-		FReal Phi = PhiToSphereCenter - Sphere.GetRadius();
+		FReal Phi = PhiToSphereCenter - Sphere.GetRadiusf();
 
 		FVec3 NormalWorld = BoxTransform.TransformVectorNoScale(NormalBox);
-		FVec3 LocationWorld = SphereWorld - (Sphere.GetRadius()) * NormalWorld;
+		FVec3 LocationWorld = SphereWorld - (Sphere.GetRadiusf()) * NormalWorld;
 
 		Result.ShapeContactPoints[0] = SphereTransform.InverseTransformPosition(LocationWorld);
 		Result.ShapeContactPoints[1] = BoxTransform.InverseTransformPosition(LocationWorld - Phi * NormalWorld);
@@ -339,23 +340,23 @@ namespace Chaos
 		return Result;
 	}
 
-	FContactPoint SphereCapsuleContactPoint(const TSphere<FReal, 3>& A, const FRigidTransform3& ATransform, const FCapsule& B, const FRigidTransform3& BTransform)
+	FContactPoint SphereCapsuleContactPoint(const FSphere& A, const FRigidTransform3& ATransform, const FCapsule& B, const FRigidTransform3& BTransform)
 	{
 		FContactPoint Result;
 
-		FVector A1 = ATransform.TransformPosition(A.GetCenter());
-		FVector B1 = BTransform.TransformPosition(B.GetX1());
-		FVector B2 = BTransform.TransformPosition(B.GetX2());
+		FVector A1 = ATransform.TransformPosition(FVector(A.GetCenterf()));
+		FVector B1 = BTransform.TransformPosition(FVector(B.GetX1f()));
+		FVector B2 = BTransform.TransformPosition(FVector(B.GetX2f()));
 		FVector P2 = FMath::ClosestPointOnSegment(A1, B1, B2);
 
 		FVec3 Delta = P2 - A1;
 		FReal DeltaLen = Delta.Size();
 		if (DeltaLen > UE_KINDA_SMALL_NUMBER)
 		{
-			FReal NewPhi = DeltaLen - (A.GetRadius() + B.GetRadius());
+			FReal NewPhi = DeltaLen - (A.GetRadiusf() + B.GetRadiusf());
 			FVec3 Dir = Delta / DeltaLen;
-			FVec3 LocationA = A1 + Dir * A.GetRadius();
-			FVec3 LocationB = P2 - Dir * B.GetRadius();
+			FVec3 LocationA = A1 + Dir * A.GetRadiusf();
+			FVec3 LocationB = P2 - Dir * B.GetRadiusf();
 			FVec3 Normal = -Dir;
 
 			Result.ShapeContactPoints[0] = ATransform.InverseTransformPosition(LocationA);
@@ -368,21 +369,21 @@ namespace Chaos
 	}
 
 	template <typename TriMeshType>
-	FContactPoint SphereTriangleMeshContactPoint(const TSphere<FReal, 3>& A, const FRigidTransform3& ATransform, const TriMeshType& B, const FRigidTransform3& BTransform, const FReal CullDistance)
+	FContactPoint SphereTriangleMeshContactPoint(const FSphere& A, const FRigidTransform3& ATransform, const TriMeshType& B, const FRigidTransform3& BTransform, const FReal CullDistance)
 	{
-		return GJKImplicitContactPoint<TSphere<FReal, 3>>(TSphere<FReal, 3>(A), ATransform, B, BTransform, CullDistance);
+		return GJKImplicitContactPoint<FSphere>(FSphere(A), ATransform, B, BTransform, CullDistance);
 	}
 
 	template<typename TriMeshType>
-	FContactPoint SphereTriangleMeshSweptContactPoint(const TSphere<FReal, 3>& A, const FRigidTransform3& ATransform, const TriMeshType& B, const FRigidTransform3& BStartTransform, const FVec3& Dir, const FReal Length, const FReal IgnorePenetration, const FReal TargetPenetration, FReal& TOI)
+	FContactPoint SphereTriangleMeshSweptContactPoint(const FSphere& A, const FRigidTransform3& ATransform, const TriMeshType& B, const FRigidTransform3& BStartTransform, const FVec3& Dir, const FReal Length, const FReal IgnorePenetration, const FReal TargetPenetration, FReal& TOI)
 	{
 		if (const TImplicitObjectScaled<FTriangleMeshImplicitObject>* ScaledTriangleMesh = B.template GetObject<const TImplicitObjectScaled<FTriangleMeshImplicitObject>>())
 		{
-			return GJKImplicitScaledTriMeshSweptContactPoint<TSphere<FReal, 3>>(A, ATransform, *ScaledTriangleMesh, BStartTransform, Dir, Length, IgnorePenetration, TargetPenetration, TOI);
+			return GJKImplicitScaledTriMeshSweptContactPoint<FSphere>(A, ATransform, *ScaledTriangleMesh, BStartTransform, Dir, Length, IgnorePenetration, TargetPenetration, TOI);
 		}
 		else if (const FTriangleMeshImplicitObject* TriangleMesh = B.template GetObject<const FTriangleMeshImplicitObject>())
 		{
-			return GJKImplicitSweptContactPoint(TSphere<FReal, 3>(A), ATransform, *TriangleMesh, BStartTransform, Dir, Length, IgnorePenetration, TargetPenetration, TOI);
+			return GJKImplicitSweptContactPoint(FSphere(A), ATransform, *TriangleMesh, BStartTransform, Dir, Length, IgnorePenetration, TargetPenetration, TOI);
 		}
 
 		ensure(false);
@@ -400,9 +401,9 @@ namespace Chaos
 		return GJKImplicitContactPoint<TBox<FReal, 3>>(A, ATransform, B, BTransform, CullDistance);
 	}
 
-	FContactPoint SphereHeightFieldContactPoint(const TSphere<FReal, 3>& A, const FRigidTransform3& ATransform, const FHeightField& B, const FRigidTransform3& BTransform, const FReal CullDistance)
+	FContactPoint SphereHeightFieldContactPoint(const FSphere& A, const FRigidTransform3& ATransform, const FHeightField& B, const FRigidTransform3& BTransform, const FReal CullDistance)
 	{
-		return GJKImplicitContactPoint<TSphere<FReal, 3>>(TSphere<FReal, 3>(A), ATransform, B, BTransform, CullDistance);
+		return GJKImplicitContactPoint<FSphere>(FSphere(A), ATransform, B, BTransform, CullDistance);
 	}
 
 	FContactPoint CapsuleHeightFieldContactPoint(const FCapsule& A, const FRigidTransform3& ATransform, const FHeightField& B, const FRigidTransform3& BTransform, const FReal CullDistance)
@@ -475,10 +476,10 @@ namespace Chaos
 	{
 		FContactPoint Result;
 
-		FVector A1 = ATransform.TransformPosition(A.GetX1());
-		FVector A2 = ATransform.TransformPosition(A.GetX2());
-		FVector B1 = BTransform.TransformPosition(B.GetX1());
-		FVector B2 = BTransform.TransformPosition(B.GetX2());
+		FVector A1 = ATransform.TransformPosition(FVector(A.GetX1f()));
+		FVector A2 = ATransform.TransformPosition(FVector(A.GetX2f()));
+		FVector B1 = BTransform.TransformPosition(FVector(B.GetX1f()));
+		FVector B2 = BTransform.TransformPosition(FVector(B.GetX2f()));
 		FVector P1, P2;
 		FMath::SegmentDistToSegmentSafe(A1, A2, B1, B2, P1, P2);
 
@@ -486,11 +487,11 @@ namespace Chaos
 		FReal DeltaLen = Delta.Size();
 		if (DeltaLen > UE_KINDA_SMALL_NUMBER)
 		{
-			FReal NewPhi = DeltaLen - (A.GetRadius() + B.GetRadius());
+			FReal NewPhi = DeltaLen - (A.GetRadiusf() + B.GetRadiusf());
 			FVec3 Dir = Delta / DeltaLen;
 			FVec3 Normal = -Dir;
-			FVec3 LocationA = P1 + Dir * A.GetRadius();
-			FVec3 LocationB = P2 - Dir * B.GetRadius();
+			FVec3 LocationA = P1 + Dir * A.GetRadiusf();
+			FVec3 LocationB = P2 - Dir * B.GetRadiusf();
 
 			Result.ShapeContactPoints[0] = ATransform.InverseTransformPosition(LocationA);
 			Result.ShapeContactPoints[1] = BTransform.InverseTransformPosition(LocationB);
@@ -510,16 +511,16 @@ namespace Chaos
 	// Template  Instantiations
 	template FContactPoint GJKImplicitSweptContactPoint<FHeightField>(const FImplicitObject& A, const FRigidTransform3& AStartTransform, const FHeightField& B, const FRigidTransform3& BTransform, const FVec3& Dir, const FReal Length, const FReal IgnorePenetration, const FReal TargetPenetration, FReal& TOI);
 
-	template FContactPoint SphereTriangleMeshSweptContactPoint<TImplicitObjectScaled<class Chaos::FTriangleMeshImplicitObject, 1>>(const TSphere<FReal, 3>& A, const FRigidTransform3& ATransform, const TImplicitObjectScaled<class Chaos::FTriangleMeshImplicitObject, 1>& B, const FRigidTransform3& BStartTransform, const FVec3& Dir, const FReal Length, const FReal IgnorePenetration, const FReal TargetPenetration, FReal& TOI);
-	template FContactPoint SphereTriangleMeshSweptContactPoint<FTriangleMeshImplicitObject>(const TSphere<FReal, 3>& A, const FRigidTransform3& ATransform, const FTriangleMeshImplicitObject& B, const FRigidTransform3& BStartTransform, const FVec3& Dir, const FReal Length, const FReal IgnorePenetration, const FReal TargetPenetration, FReal& TOI);
+	template FContactPoint SphereTriangleMeshSweptContactPoint<TImplicitObjectScaled<class Chaos::FTriangleMeshImplicitObject, 1>>(const FSphere& A, const FRigidTransform3& ATransform, const TImplicitObjectScaled<class Chaos::FTriangleMeshImplicitObject, 1>& B, const FRigidTransform3& BStartTransform, const FVec3& Dir, const FReal Length, const FReal IgnorePenetration, const FReal TargetPenetration, FReal& TOI);
+	template FContactPoint SphereTriangleMeshSweptContactPoint<FTriangleMeshImplicitObject>(const FSphere& A, const FRigidTransform3& ATransform, const FTriangleMeshImplicitObject& B, const FRigidTransform3& BStartTransform, const FVec3& Dir, const FReal Length, const FReal IgnorePenetration, const FReal TargetPenetration, FReal& TOI);
 	template FContactPoint CapsuleTriangleMeshSweptContactPoint<TImplicitObjectScaled<class Chaos::FTriangleMeshImplicitObject, 1>>(const FCapsule& A, const FRigidTransform3& ATransform, const TImplicitObjectScaled<class Chaos::FTriangleMeshImplicitObject, 1>& B, const FRigidTransform3& BStartTransform, const FVec3& Dir, const FReal Length, const FReal IgnorePenetration, const FReal TargetPenetration, FReal& TOI);
 	template FContactPoint CapsuleTriangleMeshSweptContactPoint<FTriangleMeshImplicitObject>(const FCapsule& A, const FRigidTransform3& ATransform, const FTriangleMeshImplicitObject& B, const FRigidTransform3& BStartTransform, const FVec3& Dir, const FReal Length, const FReal IgnorePenetration, const FReal TargetPenetration, FReal& TOI);
 	template FContactPoint ConvexTriangleMeshSweptContactPoint<TImplicitObjectScaled<class Chaos::FTriangleMeshImplicitObject, 1>>(const FImplicitObject& A, const FRigidTransform3& ATransform, const TImplicitObjectScaled<class Chaos::FTriangleMeshImplicitObject, 1>& B, const FRigidTransform3& BStartTransform, const FVec3& Dir, const FReal Length, const FReal IgnorePenetration, const FReal TargetPenetration, FReal& TOI);
 	template FContactPoint ConvexTriangleMeshSweptContactPoint<FTriangleMeshImplicitObject>(const FImplicitObject& A, const FRigidTransform3& ATransform, const FTriangleMeshImplicitObject& B, const FRigidTransform3& BStartTransform, const FVec3& Dir, const FReal Length, const FReal IgnorePenetration, const FReal TargetPenetration, FReal& TOI);
 	template FContactPoint BoxTriangleMeshContactPoint<TImplicitObjectScaled<class Chaos::FTriangleMeshImplicitObject, 1>>(const FImplicitBox3& A, const FRigidTransform3& ATransform, const TImplicitObjectScaled<class Chaos::FTriangleMeshImplicitObject, 1>& B, const FRigidTransform3& BTransform, const FReal CullDistance);
 	template FContactPoint BoxTriangleMeshContactPoint<FTriangleMeshImplicitObject>(const FImplicitBox3& A, const FRigidTransform3& ATransform, const FTriangleMeshImplicitObject& B, const FRigidTransform3& BTransform, const FReal CullDistance);
-	template FContactPoint SphereTriangleMeshContactPoint<TImplicitObjectScaled<class Chaos::FTriangleMeshImplicitObject, 1>>(const TSphere<FReal, 3>& A, const FRigidTransform3& ATransform, const TImplicitObjectScaled<class Chaos::FTriangleMeshImplicitObject, 1>& B, const FRigidTransform3& BTransform, const FReal CullDistance);
-	template FContactPoint SphereTriangleMeshContactPoint<FTriangleMeshImplicitObject>(const TSphere<FReal, 3>& A, const FRigidTransform3& ATransform, const FTriangleMeshImplicitObject& B, const FRigidTransform3& BTransform, const FReal CullDistance);
+	template FContactPoint SphereTriangleMeshContactPoint<TImplicitObjectScaled<class Chaos::FTriangleMeshImplicitObject, 1>>(const FSphere& A, const FRigidTransform3& ATransform, const TImplicitObjectScaled<class Chaos::FTriangleMeshImplicitObject, 1>& B, const FRigidTransform3& BTransform, const FReal CullDistance);
+	template FContactPoint SphereTriangleMeshContactPoint<FTriangleMeshImplicitObject>(const FSphere& A, const FRigidTransform3& ATransform, const FTriangleMeshImplicitObject& B, const FRigidTransform3& BTransform, const FReal CullDistance);
 	template FContactPoint CapsuleTriangleMeshContactPoint<TImplicitObjectScaled<class Chaos::FTriangleMeshImplicitObject, 1>>(const FCapsule& A, const FRigidTransform3& ATransform, const TImplicitObjectScaled<class Chaos::FTriangleMeshImplicitObject, 1>& B, const FRigidTransform3& BTransform, const FReal CullDistance);
 	template FContactPoint CapsuleTriangleMeshContactPoint<FTriangleMeshImplicitObject>(const FCapsule& A, const FRigidTransform3& ATransform, const FTriangleMeshImplicitObject& B, const FRigidTransform3& BTransform, const FReal CullDistance);
 }
diff --git a/Engine/Source/Runtime/Experimental/Chaos/Private/Chaos/Collision/ConvexTriangleContactPoint.cpp b/Engine/Source/Runtime/Experimental/Chaos/Private/Chaos/Collision/ConvexTriangleContactPoint.cpp
--- a/Engine/Source/Runtime/Experimental/Chaos/Private/Chaos/Collision/ConvexTriangleContactPoint.cpp
+++ b/Engine/Source/Runtime/Experimental/Chaos/Private/Chaos/Collision/ConvexTriangleContactPoint.cpp
@@ -742,7 +742,7 @@ namespace Chaos::Private
 	{
 		const FReal GJKEpsilon = Chaos_Collision_GJKEpsilon;
 		const FReal EPAEpsilon = Chaos_Collision_EPAEpsilon;
-		const TGJKCoreShape<ConvexType> GJKConvex(Convex, Convex.GetMargin());
+		const TGJKCoreShape<ConvexType> GJKConvex(Convex, Convex.GetMarginf());
 		const TGJKShape<FTriangle> GJKTriangle(Triangle);
 
 		// The GJKPenetrationSameSpace version uses EPA which has issues with almost exactly touching shapes. 
diff --git a/Engine/Source/Runtime/Experimental/Chaos/Private/Chaos/Collision/PBDCollisionConstraint.cpp b/Engine/Source/Runtime/Experimental/Chaos/Private/Chaos/Collision/PBDCollisionConstraint.cpp
--- a/Engine/Source/Runtime/Experimental/Chaos/Private/Chaos/Collision/PBDCollisionConstraint.cpp
+++ b/Engine/Source/Runtime/Experimental/Chaos/Private/Chaos/Collision/PBDCollisionConstraint.cpp
@@ -161,7 +161,7 @@ namespace Chaos
 	FPBDCollisionConstraint FPBDCollisionConstraint::MakeTriangle(const FImplicitObject* Implicit0)
 	{
 		FPBDCollisionConstraint Constraint;
-		Constraint.InitMarginsAndTolerances(Implicit0->GetCollisionType(), ImplicitObjectType::Triangle, FRealSingle(Implicit0->GetMargin()), FRealSingle(0));
+		Constraint.InitMarginsAndTolerances(Implicit0->GetCollisionType(), ImplicitObjectType::Triangle, Implicit0->GetMarginf(), FRealSingle(0));
 		return Constraint;
 	}
 
@@ -325,8 +325,8 @@ namespace Chaos
 		Flags.bCCDEnabled = (InCCDType == ECollisionCCDType::Enabled);
 
 		// Initialize tolerances that depend on shape type etc, also the bIsQuadratic flags
-		const FRealSingle Margin0 = FRealSingle(GetImplicit0()->GetMargin());
-		const FRealSingle Margin1 = FRealSingle(GetImplicit1()->GetMargin());
+		const FRealSingle Margin0 = GetImplicit0()->GetMarginf();
+		const FRealSingle Margin1 = GetImplicit1()->GetMarginf();
 		const EImplicitObjectType ImplicitType0 = Private::GetImplicitCollisionType(GetParticle0(), GetImplicit0());
 		const EImplicitObjectType ImplicitType1 = Private::GetImplicitCollisionType(GetParticle1(), GetImplicit1());
 		InitMarginsAndTolerances(ImplicitType0, ImplicitType1, Margin0, Margin1);
diff --git a/Engine/Source/Runtime/Experimental/Chaos/Private/Chaos/Collision/ParticlePairMidPhase.cpp b/Engine/Source/Runtime/Experimental/Chaos/Private/Chaos/Collision/ParticlePairMidPhase.cpp
--- a/Engine/Source/Runtime/Experimental/Chaos/Private/Chaos/Collision/ParticlePairMidPhase.cpp
+++ b/Engine/Source/Runtime/Experimental/Chaos/Private/Chaos/Collision/ParticlePairMidPhase.cpp
@@ -2164,11 +2164,11 @@ namespace Chaos
 		const FRigidTransform3 ShapeWorldTransform1 = FConstGenericParticleHandle(GetParticle1())->GetTransformPQ();
 
 		// Sphere distance culling
-		const FVec3 SphereWorldPos0 = ShapeWorldTransform0.TransformPositionNoScale(Sphere0.GetCenter());
-		const FVec3 SphereWorldPos1 = ShapeWorldTransform1.TransformPositionNoScale(Sphere1.GetCenter());
+		const FVec3 SphereWorldPos0 = ShapeWorldTransform0.TransformPositionNoScale(FVec3(Sphere0.GetCenterf()));
+		const FVec3 SphereWorldPos1 = ShapeWorldTransform1.TransformPositionNoScale(FVec3(Sphere1.GetCenterf()));
 		const FVec3 DR = SphereWorldPos0 - SphereWorldPos1;
 		const FReal DRLenSq = DR.SizeSquared();
-		const FReal CullSeparation = Sphere0.GetRadius() + Sphere1.GetRadius() + CullDistance;
+		const FReal CullSeparation = Sphere0.GetRadiusf() + Sphere1.GetRadiusf() + CullDistance;
 		if (DRLenSq > FMath::Square(CullSeparation))
 		{
 			return 0;
diff --git a/Engine/Source/Runtime/Experimental/Chaos/Private/Chaos/Collision/SphereConvexContactPoint.cpp b/Engine/Source/Runtime/Experimental/Chaos/Private/Chaos/Collision/SphereConvexContactPoint.cpp
--- a/Engine/Source/Runtime/Experimental/Chaos/Private/Chaos/Collision/SphereConvexContactPoint.cpp
+++ b/Engine/Source/Runtime/Experimental/Chaos/Private/Chaos/Collision/SphereConvexContactPoint.cpp
@@ -55,8 +55,8 @@ namespace Chaos
 			// Run GJK to find separating distance if available
 			// NOTE: Sphere is treated as a point (its core shape), Convex margin is ignored so we are using the outer non-shrunken hull.
 			// Sphere is also created in the space of the convex to eliminate the transform in the Support call (per GJK iteration).
-			const FVec3 SpherePosConvex = SphereToConvexTransform.TransformPositionNoScale(Sphere.GetCenter());
-			const FReal SphereRadius = Sphere.GetRadius();
+			const FVec3 SpherePosConvex = SphereToConvexTransform.TransformPositionNoScale(FVec3(Sphere.GetCenterf()));
+			const FReal SphereRadius = Sphere.GetRadiusf();
 
 			// A point in the Minkowski Sum (A-B) in a likely direction
 			// @todo(chaos): we could do better and return the correct direction separation vector most of the time by selecting the face 
@@ -83,7 +83,7 @@ namespace Chaos
 			{
 				// Results so far are all in convex-space and for a point rather than a sphere. Convert them.
 				PosConvex = SpherePosConvex - Phi * NormalConvex;
-				PosSphere = Sphere.GetCenter() + SphereToConvexTransform.InverseTransformVectorNoScale(-SphereRadius * NormalConvex);
+				PosSphere = Sphere.GetCenterf() + SphereToConvexTransform.InverseTransformVectorNoScale(-SphereRadius * NormalConvex);
 				Phi = Phi - SphereRadius;
 
 				FContactPoint ContactPoint;
diff --git a/Engine/Source/Runtime/Experimental/Chaos/Private/Chaos/Collision/SphereTriangleContactPoint.cpp b/Engine/Source/Runtime/Experimental/Chaos/Private/Chaos/Collision/SphereTriangleContactPoint.cpp
--- a/Engine/Source/Runtime/Experimental/Chaos/Private/Chaos/Collision/SphereTriangleContactPoint.cpp
+++ b/Engine/Source/Runtime/Experimental/Chaos/Private/Chaos/Collision/SphereTriangleContactPoint.cpp
@@ -18,18 +18,18 @@ namespace Chaos
 		FReal Phi;
 
 		// A point in the Minkowski Sum (A-B) in a likely direction
-		const FVec3 InitialV = Triangle.GetCentroid() - Sphere.GetCenter();
+		const FVec3 InitialV = Triangle.GetCentroid() - FVec3(Sphere.GetCenterf());
 
 		// Run GJK to find separating distance if available
 		// NOTE: Sphere is treated as a point (its core shape) for GJK to reduce iterations
 		// Also, Sphere and Triangle are in the same space
 		const EGJKDistanceResult GjkResult = GJKDistance(
 			MakeGJKShape(Triangle),
-			FGJKSphere(Sphere.GetCenter(), 0),
+			FGJKSphere(FVec3(Sphere.GetCenterf()), 0),
 			InitialV,
 			Phi, PosTri, PosSphere, NormalTri);
 
-		const FReal SphereRadius = Sphere.GetRadius();
+		const FReal SphereRadius = Sphere.GetRadiusf();
 
 		// Build the contact point
 		if ((Phi - SphereRadius) < CullDistance)
diff --git a/Engine/Source/Runtime/Experimental/Chaos/Private/Chaos/CollisionOneShotManifolds.cpp b/Engine/Source/Runtime/Experimental/Chaos/Private/Chaos/CollisionOneShotManifolds.cpp
--- a/Engine/Source/Runtime/Experimental/Chaos/Private/Chaos/CollisionOneShotManifolds.cpp
+++ b/Engine/Source/Runtime/Experimental/Chaos/Private/Chaos/CollisionOneShotManifolds.cpp
@@ -1086,11 +1086,11 @@ namespace Chaos
 			// Transform Capsule into Convex space
 			// We only keep the segment at this point - we'll add the radius back later if we create contacts
 			const FRigidTransform3 CapsuleToConvexTransform = CapsuleTransform.GetRelativeTransformNoScale(ConvexTransform);
-			const FReal CapsuleRadius = Capsule.GetRadius();
-			const FReal CapsuleHeight = Capsule.GetHeight();
+			const FReal CapsuleRadius = Capsule.GetRadiusf();
+			const FReal CapsuleHeight = Capsule.GetHeightf();
 			const FLineSegment3 CapsuleAxisSegment = FLineSegment3(
-				CapsuleToConvexTransform.TransformPositionNoScale(Capsule.GetX1()),
-				CapsuleToConvexTransform.TransformVectorNoScale(Capsule.GetAxis()),
+				CapsuleToConvexTransform.TransformPositionNoScale(FVec3(Capsule.GetX1f())),
+				CapsuleToConvexTransform.TransformVectorNoScale(FVec3(Capsule.GetAxisf())),
 				CapsuleHeight);
 
 			// No margins for the convex or the capsule for this GJK. We'll correct for the radius later...
@@ -1289,7 +1289,7 @@ namespace Chaos
 			// Find the deepest penetration. This is used to determine the planes and points to use for the manifold
 			// MaxMarginDelta is an upper bound on the distance from the contact on the rounded core shape to the actual shape surface. 
 			const FReal MarginMultiplier = Chaos_Collision_Manifold_TriangleConvexMarginMultiplier;
-			const FReal ConvexMargin = MarginMultiplier * Convex.GetMargin();
+			const FReal ConvexMargin = MarginMultiplier * FReal(Convex.GetMarginf());
 			FReal MaxMarginDelta = FReal(0);
 			int32 VertexIndexA = INDEX_NONE, VertexIndexB = INDEX_NONE;
 			const TGJKCoreShape<ConvexType> GJKConvex(Convex, ConvexMargin);
diff --git a/Engine/Source/Runtime/Experimental/Chaos/Private/Chaos/CollisionOneShotManifoldsMeshShapes.cpp b/Engine/Source/Runtime/Experimental/Chaos/Private/Chaos/CollisionOneShotManifoldsMeshShapes.cpp
--- a/Engine/Source/Runtime/Experimental/Chaos/Private/Chaos/CollisionOneShotManifoldsMeshShapes.cpp
+++ b/Engine/Source/Runtime/Experimental/Chaos/Private/Chaos/CollisionOneShotManifoldsMeshShapes.cpp
@@ -62,7 +62,7 @@ namespace Chaos
 		}
 
 		template <typename TriMeshType>
-		void ConstructSphereTriangleMeshOneShotManifold(const TSphere<FReal, 3>& Sphere, const FRigidTransform3& SphereWorldTransform, const TriMeshType& TriangleMesh, const FRigidTransform3& TriMeshWorldTransform, const FReal Dt, FPBDCollisionConstraint& Constraint)
+		void ConstructSphereTriangleMeshOneShotManifold(const FSphere& Sphere, const FRigidTransform3& SphereWorldTransform, const TriMeshType& TriangleMesh, const FRigidTransform3& TriMeshWorldTransform, const FReal Dt, FPBDCollisionConstraint& Constraint)
 		{
 			// We only build one shot manifolds once
 			// All convexes are pre-scaled, or wrapped in TImplicitObjectScaled
@@ -79,7 +79,7 @@ namespace Chaos
 			}
 		}
 
-		void ConstructSphereHeightFieldOneShotManifold(const TSphere<FReal, 3>& Sphere, const FRigidTransform3& SphereTransform, const FHeightField& Heightfield, const FRigidTransform3& HeightfieldTransform, const FReal Dt, FPBDCollisionConstraint& Constraint)
+		void ConstructSphereHeightFieldOneShotManifold(const FSphere& Sphere, const FRigidTransform3& SphereTransform, const FHeightField& Heightfield, const FRigidTransform3& HeightfieldTransform, const FReal Dt, FPBDCollisionConstraint& Constraint)
 		{
 			// We only build one shot manifolds once
 			// All convexes are pre-scaled, or wrapped in TImplicitObjectScaled
diff --git a/Engine/Source/Runtime/Experimental/Chaos/Private/Chaos/CollisionOneShotManifoldsMiscShapes.cpp b/Engine/Source/Runtime/Experimental/Chaos/Private/Chaos/CollisionOneShotManifoldsMiscShapes.cpp
--- a/Engine/Source/Runtime/Experimental/Chaos/Private/Chaos/CollisionOneShotManifoldsMiscShapes.cpp
+++ b/Engine/Source/Runtime/Experimental/Chaos/Private/Chaos/CollisionOneShotManifoldsMiscShapes.cpp
@@ -30,9 +30,9 @@ namespace Chaos
 	{
 
 		void ConstructSphereSphereOneShotManifold(
-			const TSphere<FReal, 3>& SphereA,
+			const FSphere& SphereA,
 			const FRigidTransform3& SphereATransform, //world
-			const TSphere<FReal, 3>& SphereB,
+			const FSphere& SphereB,
 			const FRigidTransform3& SphereBTransform, //world
 			const FReal Dt,
 			FPBDCollisionConstraint& Constraint)
@@ -55,7 +55,7 @@ namespace Chaos
 		}
 
 		void ConstructSpherePlaneOneShotManifold(
-			const TSphere<FReal, 3>& Sphere, 
+			const FSphere& Sphere, 
 			const FRigidTransform3& SphereTransform, 
 			const TPlane<FReal, 3>& Plane, 
 			const FRigidTransform3& PlaneTransform, 
@@ -77,7 +77,7 @@ namespace Chaos
 			}
 		}	
 
-		void ConstructSphereBoxOneShotManifold(const TSphere<FReal, 3>& Sphere, const FRigidTransform3& SphereTransform, const FImplicitBox3& Box, const FRigidTransform3& BoxTransform, const FReal Dt, FPBDCollisionConstraint& Constraint)
+		void ConstructSphereBoxOneShotManifold(const FSphere& Sphere, const FRigidTransform3& SphereTransform, const FImplicitBox3& Box, const FRigidTransform3& BoxTransform, const FReal Dt, FPBDCollisionConstraint& Constraint)
 		{
 			// We only build one shot manifolds once
 			// All convexes are pre-scaled, or wrapped in TImplicitObjectScaled
@@ -100,7 +100,7 @@ namespace Chaos
 		// we don't end up jittering between collisions on each end cap. E.g., consider a small capsule
 		// lying horizontally on a very large sphere (almost flat) - we need at least 2 contact points to 
 		// make this stable.
-		void ConstructSphereCapsuleOneShotManifold(const TSphere<FReal, 3>& Sphere, const FRigidTransform3& SphereTransform, const FCapsule& Capsule, const FRigidTransform3& CapsuleTransform, const FReal Dt, FPBDCollisionConstraint& Constraint)
+		void ConstructSphereCapsuleOneShotManifold(const FSphere& Sphere, const FRigidTransform3& SphereTransform, const FCapsule& Capsule, const FRigidTransform3& CapsuleTransform, const FReal Dt, FPBDCollisionConstraint& Constraint)
 		{
 			// We only build one shot manifolds once
 			// All convexes are pre-scaled, or wrapped in TImplicitObjectScaled
@@ -110,17 +110,17 @@ namespace Chaos
 			Constraint.ResetActiveManifoldContacts();
 
 			// Build a multi-point manifold
-			const FReal NetCullDistance = Sphere.GetRadius() + Capsule.GetRadius() + Constraint.GetCullDistance();
+			const FReal NetCullDistance = Sphere.GetRadiusf() + Capsule.GetRadiusf() + Constraint.GetCullDistance();
 			const FReal NetCullDistanceSq = FMath::Square(NetCullDistance);
 
 			// Transform the sphere into capsule space and find the closest point on the capsule line segment
 			// @todo(chaos) this would be much simpler if the sphere's were always at the origin and capsules were at the origin and axis aligned
 			const FRigidTransform3 SphereToCapsuleTransform = SphereTransform.GetRelativeTransformNoScale(CapsuleTransform);
-			const FVec3 SpherePos = SphereToCapsuleTransform.TransformPositionNoScale(Sphere.GetCenter());
-			const FReal NearPosT = Utilities::ClosestTimeOnLineSegment(SpherePos, Capsule.GetX1(), Capsule.GetX2());
+			const FVec3 SpherePos = SphereToCapsuleTransform.TransformPositionNoScale(FVec3(Sphere.GetCenterf()));
+			const FReal NearPosT = Utilities::ClosestTimeOnLineSegment(FVec3f(SpherePos), Capsule.GetX1f(), Capsule.GetX2f());
 
 			// Add the closest contact point to the manifold
-			const FVec3 NearPos = FMath::Lerp(Capsule.GetX1(), Capsule.GetX2(), NearPosT);
+			const FVec3 NearPos = FMath::Lerp(Capsule.GetX1f(), Capsule.GetX2f(), NearPosT);
 			const FVec3 NearPosDelta = SpherePos - NearPos;
 			const FReal NearPosDistanceSq = NearPosDelta.SizeSquared();
 			if (NearPosDistanceSq > UE_SMALL_NUMBER)
@@ -129,11 +129,11 @@ namespace Chaos
 				{
 					const FReal NearPosDistance = FMath::Sqrt(NearPosDistanceSq);
 					const FVec3 NearPosDir = NearPosDelta / NearPosDistance;
-					const FReal NearPhi = NearPosDistance - Sphere.GetRadius() - Capsule.GetRadius();
+					const FReal NearPhi = NearPosDistance - Sphere.GetRadiusf() - Capsule.GetRadiusf();
 
 					FContactPoint NearContactPoint;
-					NearContactPoint.ShapeContactPoints[0] = SphereToCapsuleTransform.InverseTransformPositionNoScale(SpherePos - Sphere.GetRadius() * NearPosDir);
-					NearContactPoint.ShapeContactPoints[1] = NearPos + Capsule.GetRadius() * NearPosDir;
+					NearContactPoint.ShapeContactPoints[0] = SphereToCapsuleTransform.InverseTransformPositionNoScale(SpherePos - Sphere.GetRadiusf() * NearPosDir);
+					NearContactPoint.ShapeContactPoints[1] = NearPos + Capsule.GetRadiusf() * NearPosDir;
 					NearContactPoint.ShapeContactNormal = NearPosDir;
 					NearContactPoint.Phi = NearPhi;
 					NearContactPoint.FaceIndex = INDEX_NONE;
@@ -142,7 +142,7 @@ namespace Chaos
 
 					// If we have a small sphere, just stick with the 1-point manifold
 					const FReal SphereCapsuleSizeThreshold = FReal(Chaos_Collision_Manifold_SphereCapsuleSizeThreshold);
-					if (Sphere.GetRadius() < SphereCapsuleSizeThreshold * (Capsule.GetHeight() + Capsule.GetRadius()))
+					if (Sphere.GetRadiusf() < SphereCapsuleSizeThreshold * (Capsule.GetHeightf() + Capsule.GetRadiusf()))
 					{
 						return;
 					}
@@ -157,7 +157,7 @@ namespace Chaos
 					// Note: we only consider the radius of the dynamic object(s) when deciding what "deep" means because the extra contacts are only
 					// to prevent excessive rotation from the single contact we have so far, and only the dynamic objects will rotate.
 					const FReal DeepRadiusFraction = Chaos_Collision_Manifold_CapsuleDeepPenetrationFraction;
-					const bool bIsDeep = NearPhi < -DeepRadiusFraction * Capsule.GetRadius();
+					const bool bIsDeep = NearPhi < -DeepRadiusFraction * Capsule.GetRadiusf();
 					if (!bIsDeep)
 					{
 						return;
@@ -175,21 +175,21 @@ namespace Chaos
 					{
 						// Orthogonal must point towards the sphere, but currently depends on the relative axis orientation
 						CapsuleOrthogonal = CapsuleOrthogonal * FMath::InvSqrt(CapsuleOrthogonalLenSq);
-						if (FVec3::DotProduct(CapsuleOrthogonal, SpherePos - Capsule.GetCenter()) < FReal(0))
+						if (FVec3::DotProduct(CapsuleOrthogonal, SpherePos - Capsule.GetCenterf()) < FReal(0))
 						{
 							CapsuleOrthogonal = -CapsuleOrthogonal;
 						}
 
 						if (NearPosT > EndCapDistanceThreshold)
 						{
-							const FVec3 EndCapPos0 = Capsule.GetX1() + CapsuleOrthogonal * Capsule.GetRadius();
+							const FVec3 EndCapPos0 = Capsule.GetX1f() + CapsuleOrthogonal * Capsule.GetRadiusf();
 							const FReal EndCapDistance0 = (SpherePos - EndCapPos0).Size();
-							const FReal EndCapPhi0 = EndCapDistance0 - Sphere.GetRadius();
+							const FReal EndCapPhi0 = EndCapDistance0 - Sphere.GetRadiusf();
 							
 							if (EndCapPhi0 < Constraint.GetCullDistance())
 							{
 								const FVec3 EndCapPosDir0 = (SpherePos - EndCapPos0) / EndCapDistance0;
-								const FVec3 SpherePos0 = SpherePos - EndCapPosDir0 * Sphere.GetRadius();
+								const FVec3 SpherePos0 = SpherePos - EndCapPosDir0 * Sphere.GetRadiusf();
 						
 								FContactPoint EndCapContactPoint0;
 								EndCapContactPoint0.ShapeContactPoints[0] = SphereToCapsuleTransform.InverseTransformPositionNoScale(SpherePos0);
@@ -204,14 +204,14 @@ namespace Chaos
 
 						if (NearPosT < FReal(1) - EndCapDistanceThreshold)
 						{
-							const FVec3 EndCapPos1 = Capsule.GetX2() + CapsuleOrthogonal * Capsule.GetRadius();
+							const FVec3 EndCapPos1 = Capsule.GetX2f() + CapsuleOrthogonal * Capsule.GetRadiusf();
 							const FReal EndCapDistance1 = (SpherePos - EndCapPos1).Size();
-							const FReal EndCapPhi1 = EndCapDistance1 - Sphere.GetRadius();
+							const FReal EndCapPhi1 = EndCapDistance1 - Sphere.GetRadiusf();
 
 							if (EndCapPhi1 < Constraint.GetCullDistance())
 							{
 								const FVec3 EndCapPosDir1 = (SpherePos - EndCapPos1) / EndCapDistance1;
-								const FVec3 SpherePos1 = SpherePos - EndCapPosDir1 * Sphere.GetRadius();
+								const FVec3 SpherePos1 = SpherePos - EndCapPosDir1 * Sphere.GetRadiusf();
 
 								FContactPoint EndCapContactPoint0;
 								EndCapContactPoint0.ShapeContactPoints[0] = SphereToCapsuleTransform.InverseTransformPositionNoScale(SpherePos1);
@@ -243,8 +243,8 @@ namespace Chaos
 			ContactPoints.Add(ClosestContactPoint);
 
 			// If the sphere is "large" compared to the convex add more points
-			const FVec3 SpherePos = SphereToConvexTransform.TransformPositionNoScale(Sphere.GetCenter());
-			const FReal SphereRadius = Sphere.GetRadius();
+			const FVec3 SpherePos = SphereToConvexTransform.TransformPositionNoScale(FVec3(Sphere.GetCenterf()));
+			const FReal SphereRadius = Sphere.GetRadiusf();
 			const FReal SpheerConvexManifoldSizeThreshold = FReal(1);
 			const FReal ConvexSize = Convex.BoundingBox().Extents().GetAbsMax();
 			if (SphereRadius > SpheerConvexManifoldSizeThreshold * ConvexSize)
@@ -286,7 +286,7 @@ namespace Chaos
 			}
 		}
 
-		void ConstructSphereConvexManifold(const TSphere<FReal, 3>& Sphere, const FRigidTransform3& SphereTransform, const FImplicitObject3& Convex, const FRigidTransform3& ConvexTransform, const FReal Dt, FPBDCollisionConstraint& Constraint)
+		void ConstructSphereConvexManifold(const FSphere& Sphere, const FRigidTransform3& SphereTransform, const FImplicitObject3& Convex, const FRigidTransform3& ConvexTransform, const FReal Dt, FPBDCollisionConstraint& Constraint)
 		{
 			// We only build one shot manifolds once
 			// All convexes are pre-scaled, or wrapped in TImplicitObjectScaled
@@ -339,12 +339,12 @@ namespace Chaos
 			FVec3 AAxis(CapsuleATransform.TransformVector(CapsuleA.GetSegment().GetAxis()));
 			const FVec3 BAxis(CapsuleBTransform.TransformVector(CapsuleB.GetSegment().GetAxis()));
 
-			const FReal AHalfLen = CapsuleA.GetHeight() / 2.0f;
-			const FReal BHalfLen = CapsuleB.GetHeight() / 2.0f;
+			const FReal AHalfLen = CapsuleA.GetHeightf() / 2.0f;
+			const FReal BHalfLen = CapsuleB.GetHeightf() / 2.0f;
 
 			// Used in a few places below where we need to use the smaller/larger capsule, but always a dynamic one
-			const FReal ADynamicRadius = FConstGenericParticleHandle(Constraint.GetParticle0())->IsDynamic() ? CapsuleA.GetRadius() : TNumericLimits<FReal>::Max();
-			const FReal BDynamicRadius = FConstGenericParticleHandle(Constraint.GetParticle1())->IsDynamic() ? CapsuleB.GetRadius() : TNumericLimits<FReal>::Max();
+			const FReal ADynamicRadius = FConstGenericParticleHandle(Constraint.GetParticle0())->IsDynamic() ? CapsuleA.GetRadiusf() : TNumericLimits<FReal>::Max();
+			const FReal BDynamicRadius = FConstGenericParticleHandle(Constraint.GetParticle1())->IsDynamic() ? CapsuleB.GetRadiusf() : TNumericLimits<FReal>::Max();
 
 			// Make both capsules point in the same general direction
 			FReal ADotB = FVec3::DotProduct(AAxis, BAxis);
@@ -357,8 +357,8 @@ namespace Chaos
 			// Get the closest points on the two line segments. This is used to generate the closest contact point
 			// which is always added to the manifold (if within CullDistance). We may also add other points.
 			FVector AClosest, BClosest;
-			const FVector ACenter = CapsuleATransform.TransformPosition(CapsuleA.GetCenter());
-			const FVector BCenter = CapsuleBTransform.TransformPosition(CapsuleB.GetCenter());
+			const FVector ACenter = CapsuleATransform.TransformPosition(FVec3(CapsuleA.GetCenterf()));
+			const FVector BCenter = CapsuleBTransform.TransformPosition(FVec3(CapsuleB.GetCenterf()));
 			FMath::SegmentDistToSegmentSafe(
 				ACenter + AHalfLen * AAxis,
 				ACenter - AHalfLen * AAxis,
@@ -371,7 +371,7 @@ namespace Chaos
 			FReal ClosestDeltaLen = ClosestDelta.Size();
 
 			// Stop now if we are beyond the cull distance
-			const FReal ClosestPhi = ClosestDeltaLen - (CapsuleA.GetRadius() + CapsuleB.GetRadius());
+			const FReal ClosestPhi = ClosestDeltaLen - (CapsuleA.GetRadiusf() + CapsuleB.GetRadiusf());
 			if (ClosestPhi > Constraint.GetCullDistance())
 			{
 				return;
@@ -388,8 +388,8 @@ namespace Chaos
 				// Center axes exactly intersect. We'll fake a result that pops the capsules out along the Z axis, with the smaller capsule going up
 				ClosestNormal = (ADynamicRadius <= BDynamicRadius) ? FVec3(0, 0, 1) : FVec3(0, 0, -1);
 			}
-			const FVec3 ClosestLocationA = AClosest - ClosestNormal * CapsuleA.GetRadius();
-			const FVec3 ClosestLocationB = BClosest + ClosestNormal * CapsuleB.GetRadius();
+			const FVec3 ClosestLocationA = AClosest - ClosestNormal * CapsuleA.GetRadiusf();
+			const FVec3 ClosestLocationB = BClosest + ClosestNormal * CapsuleB.GetRadiusf();
 
 			// We always add the closest point to the manifold
 			// We may also add 2 more points generated from the end cap positions of the smaller capsule
@@ -557,11 +557,11 @@ namespace Chaos
 			// Generate the extra manifold points
 			if (ADynamicRadius <= BDynamicRadius)
 			{
-				MakeCapsuleEndPointContacts(TA, ACenter, AAxis, AHalfLen, CapsuleA.GetRadius(), CapsuleATransform, BCenter, BAxis, BHalfLen, CapsuleB.GetRadius(), CapsuleBTransform, ClosestNormal, ADotB, false);
+				MakeCapsuleEndPointContacts(TA, ACenter, AAxis, AHalfLen, CapsuleA.GetRadiusf(), CapsuleATransform, BCenter, BAxis, BHalfLen, CapsuleB.GetRadiusf(), CapsuleBTransform, ClosestNormal, ADotB, false);
 			}
 			else
 			{
-				MakeCapsuleEndPointContacts(TB, BCenter, BAxis, BHalfLen, CapsuleB.GetRadius(), CapsuleBTransform, ACenter, AAxis, AHalfLen, CapsuleA.GetRadius(), CapsuleATransform, ClosestNormal, ADotB, true);
+				MakeCapsuleEndPointContacts(TB, BCenter, BAxis, BHalfLen, CapsuleB.GetRadiusf(), CapsuleBTransform, ACenter, AAxis, AHalfLen, CapsuleA.GetRadiusf(), CapsuleATransform, ClosestNormal, ADotB, true);
 			}
 		}
 
diff --git a/Engine/Source/Runtime/Experimental/Chaos/Private/Chaos/CollisionOneShotManifoldsMiscShapes.h b/Engine/Source/Runtime/Experimental/Chaos/Private/Chaos/CollisionOneShotManifoldsMiscShapes.h
--- a/Engine/Source/Runtime/Experimental/Chaos/Private/Chaos/CollisionOneShotManifoldsMiscShapes.h
+++ b/Engine/Source/Runtime/Experimental/Chaos/Private/Chaos/CollisionOneShotManifoldsMiscShapes.h
@@ -14,15 +14,15 @@ namespace Chaos
 	namespace Collisions
 	{
 		void ConstructSphereSphereOneShotManifold(
-			const TSphere<FReal, 3>& Sphere1,
+			const FSphere& Sphere1,
 			const FRigidTransform3& Convex1Transform,
-			const TSphere<FReal, 3>& Sphere2,
+			const FSphere& Sphere2,
 			const FRigidTransform3& Convex2Transform,
 			const FReal Dt,
 			FPBDCollisionConstraint& Constraint);
 
 		void ConstructSpherePlaneOneShotManifold(
-			const TSphere<FReal, 3>& Sphere,
+			const FSphere& Sphere,
 			const FRigidTransform3& SphereTransform,
 			const TPlane<FReal, 3>& Plane,
 			const FRigidTransform3& PlaneTransform,
@@ -30,7 +30,7 @@ namespace Chaos
 			FPBDCollisionConstraint& Constraint);
 
 		void ConstructSphereBoxOneShotManifold(
-			const TSphere<FReal, 3>& Sphere,
+			const FSphere& Sphere,
 			const FRigidTransform3& SphereTransform,
 			const FImplicitBox3& Box,
 			const FRigidTransform3& BoxTransform,
@@ -38,7 +38,7 @@ namespace Chaos
 			FPBDCollisionConstraint& Constraint);
 
 		void ConstructSphereCapsuleOneShotManifold(
-			const TSphere<FReal, 3>& Sphere,
+			const FSphere& Sphere,
 			const FRigidTransform3& SphereTransform,
 			const FCapsule& Capsule,
 			const FRigidTransform3& CapsuleTransform,
@@ -46,7 +46,7 @@ namespace Chaos
 			FPBDCollisionConstraint& Constraint);
 
 		void ConstructSphereConvexManifold(
-			const TSphere<FReal, 3>& Sphere,
+			const FSphere& Sphere,
 			const FRigidTransform3& SphereTransform,
 			const FImplicitObject3& Convex,
 			const FRigidTransform3& ConvexTransform,
diff --git a/Engine/Source/Runtime/Experimental/Chaos/Private/Chaos/CollisionResolution.cpp b/Engine/Source/Runtime/Experimental/Chaos/Private/Chaos/CollisionResolution.cpp
--- a/Engine/Source/Runtime/Experimental/Chaos/Private/Chaos/CollisionResolution.cpp
+++ b/Engine/Source/Runtime/Experimental/Chaos/Private/Chaos/CollisionResolution.cpp
@@ -214,9 +214,9 @@ namespace Chaos
 			{
 				return Implicit0->template GetObject<const TImplicitObjectInstanced<FCapsule>>()->GetInstancedObject();
 			}
-			else if (Implicit0OuterType == TImplicitObjectInstanced<TSphere<FReal, 3>>::StaticType())
+			else if (Implicit0OuterType == TImplicitObjectInstanced<FSphere>::StaticType())
 			{
-				return Implicit0->template GetObject<const TImplicitObjectInstanced<TSphere<FReal, 3>>>()->GetInstancedObject();
+				return Implicit0->template GetObject<const TImplicitObjectInstanced<FSphere>>()->GetInstancedObject();
 			}
 			else if (Implicit0OuterType == TImplicitObjectInstanced<FConvex>::StaticType())
 			{
@@ -499,7 +499,7 @@ namespace Chaos
 		// Sphere - Sphere
 		//
 
-		void UpdateSphereSphereConstraint(const TSphere<FReal, 3>& Sphere1, const FRigidTransform3& Sphere1Transform, const TSphere<FReal, 3>& Sphere2, const FRigidTransform3& Sphere2Transform, const FReal Dt, FPBDCollisionConstraint& Constraint)
+		void UpdateSphereSphereConstraint(const FSphere& Sphere1, const FRigidTransform3& Sphere1Transform, const FSphere& Sphere2, const FRigidTransform3& Sphere2Transform, const FReal Dt, FPBDCollisionConstraint& Constraint)
 		{
 			if (Constraint.GetUseManifold())
 			{
@@ -515,7 +515,7 @@ namespace Chaos
 		// Sphere - HeightField
 		//
 
-		void UpdateSphereHeightFieldConstraint(const TSphere<FReal, 3>& A, const FRigidTransform3& ATransform, const FHeightField& B, const FRigidTransform3& BTransform, const FReal Dt, FPBDCollisionConstraint& Constraint)
+		void UpdateSphereHeightFieldConstraint(const FSphere& A, const FRigidTransform3& ATransform, const FHeightField& B, const FRigidTransform3& BTransform, const FReal Dt, FPBDCollisionConstraint& Constraint)
 		{
 			if (Constraint.GetUseManifold())
 			{
@@ -527,7 +527,7 @@ namespace Chaos
 			}
 		}
 
-		void UpdateSphereHeightFieldConstraintSwept(TGeometryParticleHandle<FReal, 3>* Particle0, const TSphere<FReal, 3>& A, const FRigidTransform3& ATransform, const FHeightField& B, const FRigidTransform3& BTransform, const FVec3& Dir, const FReal Length, const FReal Dt, FPBDCollisionConstraint& Constraint)
+		void UpdateSphereHeightFieldConstraintSwept(TGeometryParticleHandle<FReal, 3>* Particle0, const FSphere& A, const FRigidTransform3& ATransform, const FHeightField& B, const FRigidTransform3& BTransform, const FVec3& Dir, const FReal Length, const FReal Dt, FPBDCollisionConstraint& Constraint)
 		{
 			FReal TOI = FPBDCollisionConstraint::MaxTOI;
 			UpdateContactPointNoCull(Constraint, GJKImplicitSweptContactPoint(A, ATransform, B, BTransform, Dir, Length, Constraint.GetCCDEnablePenetration(), Constraint.GetCCDTargetPenetration(), TOI), Dt, CCDNoCullAllShapePairs);
@@ -538,7 +538,7 @@ namespace Chaos
 		//  Sphere-Plane
 		//
 
-		void UpdateSpherePlaneConstraint(const TSphere<FReal, 3>& Sphere, const FRigidTransform3& SphereTransform, const TPlane<FReal, 3>& Plane, const FRigidTransform3& PlaneTransform, const FReal Dt, FPBDCollisionConstraint& Constraint)
+		void UpdateSpherePlaneConstraint(const FSphere& Sphere, const FRigidTransform3& SphereTransform, const TPlane<FReal, 3>& Plane, const FRigidTransform3& PlaneTransform, const FReal Dt, FPBDCollisionConstraint& Constraint)
 		{
 			if (Constraint.GetUseManifold())
 			{
@@ -555,7 +555,7 @@ namespace Chaos
 		//
 
 
-		void UpdateSphereBoxConstraint(const TSphere<FReal, 3>& Sphere, const FRigidTransform3& SphereTransform, const FImplicitBox3& Box, const FRigidTransform3& BoxTransform, const FReal Dt, FPBDCollisionConstraint& Constraint)
+		void UpdateSphereBoxConstraint(const FSphere& Sphere, const FRigidTransform3& SphereTransform, const FImplicitBox3& Box, const FRigidTransform3& BoxTransform, const FReal Dt, FPBDCollisionConstraint& Constraint)
 		{
 			if (Constraint.GetUseManifold())
 			{
@@ -572,7 +572,7 @@ namespace Chaos
 		//
 
 
-		void UpdateSphereCapsuleConstraint(const TSphere<FReal, 3>& A, const FRigidTransform3& ATransform, const FCapsule& B, const FRigidTransform3& BTransform, const FReal Dt, FPBDCollisionConstraint& Constraint)
+		void UpdateSphereCapsuleConstraint(const FSphere& A, const FRigidTransform3& ATransform, const FCapsule& B, const FRigidTransform3& BTransform, const FReal Dt, FPBDCollisionConstraint& Constraint)
 		{
 			if (Constraint.GetUseManifold())
 			{
@@ -588,7 +588,7 @@ namespace Chaos
 		// Sphere-Convex
 		//
 
-		void UpdateSphereConvexConstraint(const TSphere<FReal, 3>& A, const FRigidTransform3& ATransform, const FImplicitObject3& B, const FRigidTransform3& BTransform, const FReal Dt, FPBDCollisionConstraint& Constraint)
+		void UpdateSphereConvexConstraint(const FSphere& A, const FRigidTransform3& ATransform, const FImplicitObject3& B, const FRigidTransform3& BTransform, const FReal Dt, FPBDCollisionConstraint& Constraint)
 		{
 			if (Constraint.GetUseManifold())
 			{
@@ -605,7 +605,7 @@ namespace Chaos
 		//
 
 		template <typename TriMeshType>
-		void UpdateSphereTriangleMeshConstraint(const TSphere<FReal, 3>& Sphere0, const FRigidTransform3& WorldTransform0, const TriMeshType& TriangleMesh1, const FRigidTransform3& WorldTransform1, const FReal Dt, FPBDCollisionConstraint& Constraint)
+		void UpdateSphereTriangleMeshConstraint(const FSphere& Sphere0, const FRigidTransform3& WorldTransform0, const TriMeshType& TriangleMesh1, const FRigidTransform3& WorldTransform1, const FReal Dt, FPBDCollisionConstraint& Constraint)
 		{
 			if (Constraint.GetUseManifold())
 			{
@@ -619,7 +619,7 @@ namespace Chaos
 		}
 
 		template<typename TriMeshType>
-		void UpdateSphereTriangleMeshConstraintSwept(TGeometryParticleHandle<FReal, 3>* Particle0, const TSphere<FReal, 3>& Sphere0, const FRigidTransform3& WorldTransform0, const TriMeshType& TriangleMesh1, const FRigidTransform3& WorldTransform1, const FVec3& Dir, const FReal Length, const FReal Dt, FPBDCollisionConstraint& Constraint)
+		void UpdateSphereTriangleMeshConstraintSwept(TGeometryParticleHandle<FReal, 3>* Particle0, const FSphere& Sphere0, const FRigidTransform3& WorldTransform0, const TriMeshType& TriangleMesh1, const FRigidTransform3& WorldTransform1, const FVec3& Dir, const FReal Length, const FReal Dt, FPBDCollisionConstraint& Constraint)
 		{
 			FReal TOI = FPBDCollisionConstraint::MaxTOI;
 			UpdateContactPointNoCull(Constraint, SphereTriangleMeshSweptContactPoint(Sphere0, WorldTransform0, TriangleMesh1, WorldTransform1, Dir, Length, Constraint.GetCCDEnablePenetration(), Constraint.GetCCDTargetPenetration(), TOI), Dt, CCDNoCullAllShapePairs);
@@ -929,35 +929,35 @@ namespace Chaos
 			switch (Constraint.GetShapesType())
 			{
 			case EContactShapesType::SphereSphere:
-				UpdateSphereSphereConstraint(*Implicit0.template GetObject<TSphere<FReal, 3>>(), WorldTransform0, *Implicit1.template GetObject<TSphere<FReal, 3>>(), WorldTransform1, Dt, Constraint);
+				UpdateSphereSphereConstraint(*Implicit0.template GetObject<FSphere>(), WorldTransform0, *Implicit1.template GetObject<FSphere>(), WorldTransform1, Dt, Constraint);
 				break;
 			case EContactShapesType::SphereCapsule:
-				UpdateSphereCapsuleConstraint(*Implicit0.template GetObject<TSphere<FReal, 3>>(), WorldTransform0, *Implicit1.template GetObject<FCapsule>(), WorldTransform1, Dt, Constraint);
+				UpdateSphereCapsuleConstraint(*Implicit0.template GetObject<FSphere>(), WorldTransform0, *Implicit1.template GetObject<FCapsule>(), WorldTransform1, Dt, Constraint);
 				break;
 			case EContactShapesType::SphereBox:
-				UpdateSphereBoxConstraint(*Implicit0.template GetObject<TSphere<FReal, 3>>(), WorldTransform0, *Implicit1.template GetObject<TBox<FReal, 3>>(), WorldTransform1, Dt, Constraint);
+				UpdateSphereBoxConstraint(*Implicit0.template GetObject<FSphere>(), WorldTransform0, *Implicit1.template GetObject<TBox<FReal, 3>>(), WorldTransform1, Dt, Constraint);
 				break;
 			case EContactShapesType::SphereConvex:
-				UpdateSphereConvexConstraint(*Implicit0.template GetObject<TSphere<FReal, 3>>(), WorldTransform0, Implicit1, WorldTransform1, Dt, Constraint);
+				UpdateSphereConvexConstraint(*Implicit0.template GetObject<FSphere>(), WorldTransform0, Implicit1, WorldTransform1, Dt, Constraint);
 				break;
 			case EContactShapesType::SphereTriMesh:
 				if (const TImplicitObjectScaled<FTriangleMeshImplicitObject>* ScaledTriMesh = Implicit1.template GetObject<const TImplicitObjectScaled<FTriangleMeshImplicitObject> >())
 				{
-					UpdateSphereTriangleMeshConstraint(*Implicit0.template GetObject<TSphere<FReal, 3>>(), WorldTransform0, *ScaledTriMesh, WorldTransform1, Dt, Constraint);
+					UpdateSphereTriangleMeshConstraint(*Implicit0.template GetObject<FSphere>(), WorldTransform0, *ScaledTriMesh, WorldTransform1, Dt, Constraint);
 					break;
 				}
 				else if (const FTriangleMeshImplicitObject* TriangleMeshImplicit = Implicit1.template GetObject<const FTriangleMeshImplicitObject>())
 				{
-					UpdateSphereTriangleMeshConstraint(*Implicit0.template GetObject<TSphere<FReal, 3>>(), WorldTransform0, *TriangleMeshImplicit, WorldTransform1, Dt, Constraint);
+					UpdateSphereTriangleMeshConstraint(*Implicit0.template GetObject<FSphere>(), WorldTransform0, *TriangleMeshImplicit, WorldTransform1, Dt, Constraint);
 					break;
 				}
 				ensure(false);
 				break;
 			case EContactShapesType::SphereHeightField:
-				UpdateSphereHeightFieldConstraint(*Implicit0.template GetObject<TSphere<FReal, 3>>(), WorldTransform0, *Implicit1.template GetObject< FHeightField >(), WorldTransform1, Dt, Constraint);
+				UpdateSphereHeightFieldConstraint(*Implicit0.template GetObject<FSphere>(), WorldTransform0, *Implicit1.template GetObject< FHeightField >(), WorldTransform1, Dt, Constraint);
 				break;
 			case EContactShapesType::SpherePlane:
-				UpdateSpherePlaneConstraint(*Implicit0.template GetObject<TSphere<FReal, 3>>(), WorldTransform0, *Implicit1.template GetObject<TPlane<FReal, 3>>(), WorldTransform1, Dt, Constraint);
+				UpdateSpherePlaneConstraint(*Implicit0.template GetObject<FSphere>(), WorldTransform0, *Implicit1.template GetObject<TPlane<FReal, 3>>(), WorldTransform1, Dt, Constraint);
 				break;
 			case EContactShapesType::CapsuleCapsule:
 				UpdateCapsuleCapsuleConstraint(*Implicit0.template GetObject<FCapsule>(), WorldTransform0, *Implicit1.template GetObject<FCapsule>(), WorldTransform1, Dt, Constraint);
@@ -1069,7 +1069,7 @@ namespace Chaos
 					return true;
 				case EContactShapesType::SphereHeightField:
 				{
-					const TSphere<FReal, 3>* Object0 = Implicit0.template GetObject<const TSphere<FReal, 3>>();
+					const FSphere* Object0 = Implicit0.template GetObject<const FSphere>();
 					const FHeightField* Object1 = Implicit1.template GetObject<const FHeightField>();
 					UpdateSphereHeightFieldConstraintSwept(Particle0, *Object0, WorldTransform0, *Object1, WorldTransform1, DirCCD, LengthCCD, Dt, Constraint);
 					return true;
@@ -1083,7 +1083,7 @@ namespace Chaos
 				}
 				case EContactShapesType::SphereTriMesh:
 				{
-					const TSphere<FReal, 3>* Object0 = Implicit0.template GetObject<const TSphere<FReal, 3>>();
+					const FSphere* Object0 = Implicit0.template GetObject<const FSphere>();
 					if (const TImplicitObjectScaled<FTriangleMeshImplicitObject>* ScaledTriangleMesh = Implicit1.template GetObject<const TImplicitObjectScaled<FTriangleMeshImplicitObject>>())
 					{
 						UpdateSphereTriangleMeshConstraintSwept(Particle0, *Object0, WorldTransform0, *ScaledTriangleMesh, WorldTransform1, DirCCD, LengthCCD, Dt, Constraint);
@@ -1193,51 +1193,51 @@ namespace Chaos
 				bOutSwap = true;
 				return EContactShapesType::BoxPlane;
 			}
-			else if (Implicit0Type == TSphere<FReal, 3>::StaticType() && Implicit1Type == TSphere<FReal, 3>::StaticType())
+			else if (Implicit0Type == FSphere::StaticType() && Implicit1Type == FSphere::StaticType())
 			{
 				return EContactShapesType::SphereSphere;
 			}
-			else if (Implicit0Type == TSphere<FReal, 3>::StaticType() && Implicit1Type == FHeightField::StaticType())
+			else if (Implicit0Type == FSphere::StaticType() && Implicit1Type == FHeightField::StaticType())
 			{
 				return EContactShapesType::SphereHeightField;
 			}
-			else if (Implicit0Type == FHeightField::StaticType() && Implicit1Type == TSphere<FReal, 3>::StaticType())
+			else if (Implicit0Type == FHeightField::StaticType() && Implicit1Type == FSphere::StaticType())
 			{
 				bOutSwap = true;
 				return EContactShapesType::SphereHeightField;
 			}
-			else if (Implicit0Type == TSphere<FReal, 3>::StaticType() && Implicit1Type == TPlane<FReal, 3>::StaticType())
+			else if (Implicit0Type == FSphere::StaticType() && Implicit1Type == TPlane<FReal, 3>::StaticType())
 			{
 				return EContactShapesType::SpherePlane;
 			}
-			else if (Implicit0Type == TPlane<FReal, 3>::StaticType() && Implicit1Type == TSphere<FReal, 3>::StaticType())
+			else if (Implicit0Type == TPlane<FReal, 3>::StaticType() && Implicit1Type == FSphere::StaticType())
 			{
 				bOutSwap = true;
 				return EContactShapesType::SpherePlane;
 			}
-			else if (Implicit0Type == TSphere<FReal, 3>::StaticType() && Implicit1Type == TBox<FReal, 3>::StaticType())
+			else if (Implicit0Type == FSphere::StaticType() && Implicit1Type == TBox<FReal, 3>::StaticType())
 			{
 				return EContactShapesType::SphereBox;
 			}
-			else if (Implicit0Type == TBox<FReal, 3>::StaticType() && Implicit1Type == TSphere<FReal, 3>::StaticType())
+			else if (Implicit0Type == TBox<FReal, 3>::StaticType() && Implicit1Type == FSphere::StaticType())
 			{
 				bOutSwap = true;
 				return EContactShapesType::SphereBox;
 			}
-			else if (Implicit0Type == TSphere<FReal, 3>::StaticType() && Implicit1Type == FCapsule::StaticType())
+			else if (Implicit0Type == FSphere::StaticType() && Implicit1Type == FCapsule::StaticType())
 			{
 				return EContactShapesType::SphereCapsule;
 			}
-			else if (Implicit0Type == FCapsule::StaticType() && Implicit1Type == TSphere<FReal, 3>::StaticType())
+			else if (Implicit0Type == FCapsule::StaticType() && Implicit1Type == FSphere::StaticType())
 			{
 				bOutSwap = true;
 				return EContactShapesType::SphereCapsule;
 			}
-			else if (Implicit0Type == TSphere<FReal, 3>::StaticType() && Implicit1Type == FImplicitConvex3::StaticType())
+			else if (Implicit0Type == FSphere::StaticType() && Implicit1Type == FImplicitConvex3::StaticType())
 			{
 				return EContactShapesType::SphereConvex;
 			}
-			else if (Implicit0Type == FImplicitConvex3::StaticType() && Implicit1Type == TSphere<FReal, 3>::StaticType())
+			else if (Implicit0Type == FImplicitConvex3::StaticType() && Implicit1Type == FSphere::StaticType())
 			{
 				bOutSwap = true;
 				return EContactShapesType::SphereConvex;
@@ -1282,11 +1282,11 @@ namespace Chaos
 				bOutSwap = true;
 				return EContactShapesType::BoxTriMesh;
 			}
-			else if (Implicit0Type == TSphere<FReal, 3>::StaticType() && Implicit1Type == FTriangleMeshImplicitObject::StaticType())
+			else if (Implicit0Type == FSphere::StaticType() && Implicit1Type == FTriangleMeshImplicitObject::StaticType())
 			{
 				return EContactShapesType::SphereTriMesh;
 			}
-			else if (Implicit0Type == FTriangleMeshImplicitObject::StaticType() && Implicit1Type == TSphere<FReal, 3>::StaticType())
+			else if (Implicit0Type == FTriangleMeshImplicitObject::StaticType() && Implicit1Type == FSphere::StaticType())
 			{
 				bOutSwap = true;
 				return EContactShapesType::SphereTriMesh;
diff --git a/Engine/Source/Runtime/Experimental/Chaos/Private/Chaos/CollisionResolutionUtil.cpp b/Engine/Source/Runtime/Experimental/Chaos/Private/Chaos/CollisionResolutionUtil.cpp
--- a/Engine/Source/Runtime/Experimental/Chaos/Private/Chaos/CollisionResolutionUtil.cpp
+++ b/Engine/Source/Runtime/Experimental/Chaos/Private/Chaos/CollisionResolutionUtil.cpp
@@ -391,7 +391,7 @@ namespace Chaos
 				}
 				else if (Object.GetType() == ImplicitObjectType::Sphere && NumParticles > 0)
 				{
-					const TSphere<FReal, 3>* Sphere = Object.GetObject<Chaos::TSphere<FReal, 3>>();
+					const FSphere* Sphere = Object.GetObject<Chaos::FSphere>();
 
 					if (NormalAveraging && UpdateType != ECollisionUpdateType::Any)
 					{
diff --git a/Engine/Source/Runtime/Experimental/Chaos/Private/Chaos/Convex.cpp b/Engine/Source/Runtime/Experimental/Chaos/Private/Chaos/Convex.cpp
--- a/Engine/Source/Runtime/Experimental/Chaos/Private/Chaos/Convex.cpp
+++ b/Engine/Source/Runtime/Experimental/Chaos/Private/Chaos/Convex.cpp
@@ -43,7 +43,7 @@ namespace Chaos
 
 		OutFaceIndex = INDEX_NONE;	//finding face is expensive, should be called directly by user
 		const FRigidTransform3 StartTM(StartPoint, FRotation3::FromIdentity());
-		const TSphere<FReal, 3> Sphere(FVec3(0), Thickness);
+		const FSphere Sphere(FVec3(0), Thickness);
 		return GJKRaycast(*this, Sphere, StartTM, Dir, Length, OutTime, OutPosition, OutNormal);
 	}
 
diff --git a/Engine/Source/Runtime/Experimental/Chaos/Private/Chaos/DebugDraw/DebugDrawImplicitObject.cpp b/Engine/Source/Runtime/Experimental/Chaos/Private/Chaos/DebugDraw/DebugDrawImplicitObject.cpp
--- a/Engine/Source/Runtime/Experimental/Chaos/Private/Chaos/DebugDraw/DebugDrawImplicitObject.cpp
+++ b/Engine/Source/Runtime/Experimental/Chaos/Private/Chaos/DebugDraw/DebugDrawImplicitObject.cpp
@@ -45,7 +45,7 @@ namespace Chaos::Private
 		float Duration)
 	{
 		const FVec3 Scale = Transform.GetScale3D();
-		Renderer.RenderSphere(Transform.TransformPosition(Sphere.GetCenter()), float(Scale.Z * Sphere.GetRadius()), Color, LineThickness, Duration);
+		Renderer.RenderSphere(Transform.TransformPosition(FVec3(Sphere.GetCenterf())), float(Scale.Z) * Sphere.GetRadiusf(), Color, LineThickness, Duration);
 	}
 
 	void ChaosDDRenderCapsule(
@@ -57,9 +57,9 @@ namespace Chaos::Private
 		float Duration)
 	{
 		const FVec3 Scale = Transform.GetScale3D();
-		const FVec3 P = Transform.TransformPosition(Capsule.GetCenter());
+		const FVec3 P = Transform.TransformPosition(FVec3(Capsule.GetCenterf()));
 		const FRotation3 Q = Transform.GetRotation() * FRotationMatrix::MakeFromZ(Capsule.GetAxis());
-		Renderer.RenderCapsule(P, Q, float(Scale.Z * (0.5 * Capsule.GetHeight() + Capsule.GetRadius())), float(Scale.X * Capsule.GetRadius()), Color, LineThickness, Duration);
+		Renderer.RenderCapsule(P, Q, float(Scale.Z * (0.5 * Capsule.GetHeightf() + Capsule.GetRadiusf())), float(Scale.X * Capsule.GetRadiusf()), Color, LineThickness, Duration);
 	}
 
 	void ChaosDDRenderBox(
diff --git a/Engine/Source/Runtime/Experimental/Chaos/Private/Chaos/Deformable/ChaosDeformableSolver.cpp b/Engine/Source/Runtime/Experimental/Chaos/Private/Chaos/Deformable/ChaosDeformableSolver.cpp
--- a/Engine/Source/Runtime/Experimental/Chaos/Private/Chaos/Deformable/ChaosDeformableSolver.cpp
+++ b/Engine/Source/Runtime/Experimental/Chaos/Private/Chaos/Deformable/ChaosDeformableSolver.cpp
@@ -2242,8 +2242,9 @@ namespace Chaos::Softs
 								if (GeomType == ImplicitObjectType::Sphere)
 								{
 									const FSphere& SphereGeometry = Geometry->GetObjectChecked<FSphere>();
-									FVector Center = ToFVector(CollisionParticles.GetX(Index)) + SphereGeometry.GetCenter();
-									FReal Radius = SphereGeometry.GetRadius();
+									FVector Center = ToFVector(CollisionParticles.GetX(Index)) + FVector(SphereGeometry.GetCenterf());
+
+									FReal Radius = SphereGeometry.GetRadiusf();
 									Chaos::FDebugDrawQueue::GetInstance().DrawDebugSphere(Center, Radius, 12, FColor::Red, false, -1.0f, 0, 1.f);
 								}
 								else if (GeomType == ImplicitObjectType::Box)
diff --git a/Engine/Source/Runtime/Experimental/Chaos/Private/Chaos/HeightField.cpp b/Engine/Source/Runtime/Experimental/Chaos/Private/Chaos/HeightField.cpp
--- a/Engine/Source/Runtime/Experimental/Chaos/Private/Chaos/HeightField.cpp
+++ b/Engine/Source/Runtime/Experimental/Chaos/Private/Chaos/HeightField.cpp
@@ -505,7 +505,7 @@ namespace Chaos
 				// Checks the query geometry support point along the normal at the sweep end point
 				// NOTE: We are using SupportCore, so we need to add Radius to the early-out distance check. We could just use Support
 				// but this way avoids an extra sqrt. We don't modify IgnorePenetration in the ctor because it gets used below.
-				VectorRegister4Float RadiusV = MakeVectorRegisterFloatFromDouble(VectorSetFloat1(OtherGeom.GetRadius()));
+				VectorRegister4Float RadiusV = VectorSetFloat1(OtherGeom.GetRadiusf());
 				VectorRegister4Float IgnorePenetrationSimd = VectorSubtract(MakeVectorRegisterFloatFromDouble(VectorSetFloat1(IgnorePenetration)), RadiusV);
 				VectorRegister4Float TriangleNormalSimd = VectorNormalize(VectorCross(VectorSubtract(B, A), VectorSubtract(C, A)));
 				VectorRegister4Float OtherTriangleNormalSimd = VectorQuaternionInverseRotateVector(RotationSimd, TriangleNormalSimd);
@@ -1881,7 +1881,7 @@ namespace Chaos
 		return GJKContactPointImp(QueryGeom, QueryTM, Thickness, ContactLocation, ContactNormal, ContactPhi, ContactFaceIndex);
 	}
 
-	bool FHeightField::GJKContactPoint(const TSphere<FReal, 3>& QueryGeom, const FRigidTransform3& QueryTM, const FReal Thickness, FVec3& ContactLocation, FVec3& ContactNormal, FReal& ContactPhi, int32& ContactFaceIndex) const
+	bool FHeightField::GJKContactPoint(const FSphere& QueryGeom, const FRigidTransform3& QueryTM, const FReal Thickness, FVec3& ContactLocation, FVec3& ContactNormal, FReal& ContactPhi, int32& ContactFaceIndex) const
 	{
 		return GJKContactPointImp(QueryGeom, QueryTM, Thickness, ContactLocation, ContactNormal, ContactPhi, ContactFaceIndex);
 	}
@@ -1901,7 +1901,7 @@ namespace Chaos
 		return GJKContactPointImp(QueryGeom, QueryTM, Thickness, ContactLocation, ContactNormal, ContactPhi, ContactFaceIndex);
 	}
 
-	bool FHeightField::GJKContactPoint(const TImplicitObjectScaled < TSphere<FReal, 3>>& QueryGeom, const FRigidTransform3& QueryTM, const FReal Thickness, FVec3& ContactLocation, FVec3& ContactNormal, FReal& ContactPhi, int32& ContactFaceIndex) const
+	bool FHeightField::GJKContactPoint(const TImplicitObjectScaled < FSphere>& QueryGeom, const FRigidTransform3& QueryTM, const FReal Thickness, FVec3& ContactLocation, FVec3& ContactNormal, FReal& ContactPhi, int32& ContactFaceIndex) const
 	{
 		return GJKContactPointImp(QueryGeom, QueryTM, Thickness, ContactLocation, ContactNormal, ContactPhi, ContactFaceIndex);
 	}
@@ -1960,7 +1960,7 @@ namespace Chaos
 			const VectorRegister4Float Offset = VectorCross(AB, AC);
 
 			FTriangleRegister TriangleConvex(A, B, C);
-			return GJKIntersectionSameSpaceSimd(TriangleConvex, QueryGeom, Thickness, VectorNormalizeSafe(Offset, GlobalVectorConstants::Float1000));
+			return GJKIntersectionSameSpaceSimd(TriangleConvex, QueryGeom, FRealSingle(Thickness), VectorNormalizeSafe(Offset, GlobalVectorConstants::Float1000));
 		};
 
 		bool bResult = false;
@@ -2006,30 +2006,30 @@ namespace Chaos
 				[[maybe_unused]] FAABBSimd AABBSimd;
 				if constexpr (std::is_same<QueryGeomType, FCapsule>::value)
 				{
-					Radius = FRealSingle(QueryGeom.GetRadius());
-					const FVec3f X1f = QueryGeom.GetX1();
-					const FVec3f X2f = QueryGeom.GetX2();
+					Radius = QueryGeom.GetRadiusf();
+					const FVec3f X1f = QueryGeom.GetX1f();
+					const FVec3f X2f = QueryGeom.GetX2f();
 					X1 = VectorLoadFloat3(&X1f.X);
 					X2 = VectorLoadFloat3(&X2f.X);
 				}
 				else if constexpr (std::is_same<QueryGeomType, TImplicitObjectScaled < FCapsule>>::value)
 				{
-					Radius = FRealSingle(QueryGeom.GetRadius());
-					const FVec3f X1f = QueryGeom.GetUnscaledObject()->GetX1() * QueryGeom.GetScale();
-					const FVec3f X2f = QueryGeom.GetUnscaledObject()->GetX2() * QueryGeom.GetScale();
+					Radius = QueryGeom.GetRadiusf();
+					const FVec3f X1f = QueryGeom.GetUnscaledObject()->GetX1f() * QueryGeom.GetScale();
+					const FVec3f X2f = QueryGeom.GetUnscaledObject()->GetX2f() * QueryGeom.GetScale();
 					X1 = VectorLoadFloat3(&X1f.X);
 					X2 = VectorLoadFloat3(&X2f.X);
 				}
 				else if constexpr (std::is_same<QueryGeomType, Chaos::FSphere>::value)
 				{
-					Radius = FRealSingle(QueryGeom.GetRadius());
-					const FVec3f X1f = QueryGeom.GetCenter();
+					Radius = QueryGeom.GetRadiusf();
+					const FVec3f X1f = QueryGeom.GetCenterf();
 					X1 = VectorLoadFloat3(&X1f.X);
 				}
 				else if constexpr (std::is_same<QueryGeomType, TImplicitObjectScaled < Chaos::FSphere>>::value)
 				{
 					Radius = FRealSingle(QueryGeom.GetRadius());
-					const FVec3f X1f = QueryGeom.GetUnscaledObject()->GetCenter() * QueryGeom.GetScale();
+					const FVec3f X1f = QueryGeom.GetUnscaledObject()->GetCenterf() * QueryGeom.GetScale();
 					X1 = VectorLoadFloat3(&X1f.X);
 				}
 				else if constexpr (std::is_same<QueryGeomType, TBox<FReal, 3>>::value || std::is_same<QueryGeomType, TImplicitObjectScaled<TBox<FReal, 3>>>::value)
@@ -2112,7 +2112,7 @@ namespace Chaos
 		return false;
 	}
 
-	bool FHeightField::OverlapGeom(const TSphere<FReal, 3>& QueryGeom, const FRigidTransform3& QueryTM, const FReal Thickness, FMTDInfo* OutMTD) const
+	bool FHeightField::OverlapGeom(const FSphere& QueryGeom, const FRigidTransform3& QueryTM, const FReal Thickness, FMTDInfo* OutMTD) const
 	{
 		return OverlapGeomImp(QueryGeom, QueryTM, Thickness, OutMTD);
 	}
@@ -2132,7 +2132,7 @@ namespace Chaos
 		return OverlapGeomImp(QueryGeom, QueryTM, Thickness, OutMTD);
 	}
 
-	bool FHeightField::OverlapGeom(const TImplicitObjectScaled<TSphere<FReal, 3>>& QueryGeom, const FRigidTransform3& QueryTM, const FReal Thickness, FMTDInfo* OutMTD) const
+	bool FHeightField::OverlapGeom(const TImplicitObjectScaled<FSphere>& QueryGeom, const FRigidTransform3& QueryTM, const FReal Thickness, FMTDInfo* OutMTD) const
 	{
 		return OverlapGeomImp(QueryGeom, QueryTM, Thickness, OutMTD);
 	}
@@ -2180,7 +2180,7 @@ namespace Chaos
 		return bHit;
 	}
 
-	bool FHeightField::SweepGeom(const TSphere<FReal, 3>& QueryGeom, const FRigidTransform3& StartTM, const FVec3& Dir, const FReal Length, FReal& OutTime, FVec3& OutPosition, FVec3& OutNormal, int32& OutFaceIndex, FVec3& OutFaceNormal, const FReal Thickness, bool bComputeMTD) const
+	bool FHeightField::SweepGeom(const FSphere& QueryGeom, const FRigidTransform3& StartTM, const FVec3& Dir, const FReal Length, FReal& OutTime, FVec3& OutPosition, FVec3& OutNormal, int32& OutFaceIndex, FVec3& OutFaceNormal, const FReal Thickness, bool bComputeMTD) const
 	{
 		return SweepGeomImp(QueryGeom, StartTM, Dir, Length, OutTime, OutPosition, OutNormal, OutFaceIndex, Thickness, bComputeMTD);
 	}
@@ -2200,7 +2200,7 @@ namespace Chaos
 		return SweepGeomImp(QueryGeom, StartTM, Dir, Length, OutTime, OutPosition, OutNormal, OutFaceIndex, Thickness, bComputeMTD);
 	}
 
-	bool FHeightField::SweepGeom(const TImplicitObjectScaled<TSphere<FReal, 3>>& QueryGeom, const FRigidTransform3& StartTM, const FVec3& Dir, const FReal Length, FReal& OutTime, FVec3& OutPosition, FVec3& OutNormal, int32& OutFaceIndex, FVec3& OutFaceNormal, const FReal Thickness, bool bComputeMTD) const
+	bool FHeightField::SweepGeom(const TImplicitObjectScaled<FSphere>& QueryGeom, const FRigidTransform3& StartTM, const FVec3& Dir, const FReal Length, FReal& OutTime, FVec3& OutPosition, FVec3& OutNormal, int32& OutFaceIndex, FVec3& OutFaceNormal, const FReal Thickness, bool bComputeMTD) const
 	{
 		return SweepGeomImp(QueryGeom, StartTM, Dir, Length, OutTime, OutPosition, OutNormal, OutFaceIndex, Thickness, bComputeMTD);
 	}
@@ -2258,7 +2258,7 @@ namespace Chaos
 		return bHit;
 	}
 
-	bool FHeightField::SweepGeomCCD(const TSphere<FReal, 3>& QueryGeom, const FRigidTransform3& StartTM, const FVec3& Dir, const FReal Length, const FReal IgnorePenetration, const FReal TargetPenetration, FReal& OutTOI, FReal& OutPhi, FVec3& OutPosition, FVec3& OutNormal, int32& OutFaceIndex, FVec3& OutFaceNormal) const
+	bool FHeightField::SweepGeomCCD(const FSphere& QueryGeom, const FRigidTransform3& StartTM, const FVec3& Dir, const FReal Length, const FReal IgnorePenetration, const FReal TargetPenetration, FReal& OutTOI, FReal& OutPhi, FVec3& OutPosition, FVec3& OutNormal, int32& OutFaceIndex, FVec3& OutFaceNormal) const
 	{
 		return SweepGeomCCDImp(QueryGeom, StartTM, Dir, Length, IgnorePenetration, TargetPenetration, OutTOI, OutPhi, OutPosition, OutNormal, OutFaceIndex);
 	}
@@ -2278,7 +2278,7 @@ namespace Chaos
 		return SweepGeomCCDImp(QueryGeom, StartTM, Dir, Length, IgnorePenetration, TargetPenetration, OutTOI, OutPhi, OutPosition, OutNormal, OutFaceIndex);
 	}
 
-	bool FHeightField::SweepGeomCCD(const TImplicitObjectScaled<TSphere<FReal, 3>>& QueryGeom, const FRigidTransform3& StartTM, const FVec3& Dir, const FReal Length, const FReal IgnorePenetration, const FReal TargetPenetration, FReal& OutTOI, FReal& OutPhi, FVec3& OutPosition, FVec3& OutNormal, int32& OutFaceIndex, FVec3& OutFaceNormal) const
+	bool FHeightField::SweepGeomCCD(const TImplicitObjectScaled<FSphere>& QueryGeom, const FRigidTransform3& StartTM, const FVec3& Dir, const FReal Length, const FReal IgnorePenetration, const FReal TargetPenetration, FReal& OutTOI, FReal& OutPhi, FVec3& OutPosition, FVec3& OutNormal, int32& OutFaceIndex, FVec3& OutFaceNormal) const
 	{
 		return SweepGeomCCDImp(QueryGeom, StartTM, Dir, Length, IgnorePenetration, TargetPenetration, OutTOI, OutPhi, OutPosition, OutNormal, OutFaceIndex);
 	}
diff --git a/Engine/Source/Runtime/Experimental/Chaos/Private/Chaos/ImplicitObject.cpp b/Engine/Source/Runtime/Experimental/Chaos/Private/Chaos/ImplicitObject.cpp
--- a/Engine/Source/Runtime/Experimental/Chaos/Private/Chaos/ImplicitObject.cpp
+++ b/Engine/Source/Runtime/Experimental/Chaos/Private/Chaos/ImplicitObject.cpp
@@ -311,7 +311,7 @@ namespace Chaos
 				Ar << ObjectType;
 				switch((EImplicitObjectType)ObjectType)
 				{
-				case ImplicitObjectType::Sphere: { Value = TUniquePtr<TSphere<FReal, 3>>(new TSphere<FReal, 3>()); break; }
+				case ImplicitObjectType::Sphere: { Value = TUniquePtr<FSphere>(new FSphere()); break; }
 				case ImplicitObjectType::Box: { Value = TUniquePtr<TBox<FReal, 3>>(new TBox<FReal, 3>()); break; }
 				case ImplicitObjectType::Plane: { Value = TUniquePtr<TPlane<FReal, 3>>(new TPlane<FReal, 3>()); break; }
 				case ImplicitObjectType::LevelSet: { Value = TUniquePtr<FLevelSet>(new FLevelSet()); break; }
@@ -443,7 +443,7 @@ namespace Chaos
 
 		switch((EImplicitObjectType)ObjectType)
 		{
-		case ImplicitObjectType::Sphere: if(Ar.IsLoading()) { return new TSphere<FReal, 3>(); } break;
+		case ImplicitObjectType::Sphere: if(Ar.IsLoading()) { return new FSphere(); } break;
 		case ImplicitObjectType::Box: if(Ar.IsLoading()) { return new TBox<FReal, 3>(); } break;
 		case ImplicitObjectType::Plane: if(Ar.IsLoading()) { return new TPlane<FReal, 3>(); } break;
 		case ImplicitObjectType::Capsule: if(Ar.IsLoading()) { return new FCapsule(); } break;
