Store Implicit in single precision

diff --git a/Engine/Source/Runtime/Experimental/Chaos/Private/Chaos/Levelset.cpp b/Engine/Source/Runtime/Experimental/Chaos/Private/Chaos/Levelset.cpp
--- a/Engine/Source/Runtime/Experimental/Chaos/Private/Chaos/Levelset.cpp
+++ b/Engine/Source/Runtime/Experimental/Chaos/Private/Chaos/Levelset.cpp
@@ -1542,12 +1542,12 @@ FReal FLevelSet::PhiWithNormal(const FVec3& x, FVec3& Normal) const
 	return SizeSquared > 0 ? (sqrt(SizeSquared) + Phi) : Phi;
 }
 
-void GetGeomSurfaceSamples(const TSphere<FReal, 3>& InGeom, TArray<FVec3>& OutSamples)
+void GetGeomSurfaceSamples(const FSphere& InGeom, TArray<FVec3>& OutSamples)
 {
 	OutSamples.Reset();
 	OutSamples.AddUninitialized(6);
 
-	const FReal Radius = InGeom.GetRadius();
+	const FReal Radius = InGeom.GetRadiusf();
 
 	OutSamples[0] = FVec3(Radius, 0, 0);
 	OutSamples[1] = FVec3(-Radius, 0, 0);
@@ -1580,8 +1580,8 @@ void GetGeomSurfaceSamples(const FCapsule& InGeom, TArray<FVec3>& OutSamples)
 	OutSamples.Reset();
 	OutSamples.AddUninitialized(14);
 
-	FReal Radius = InGeom.GetRadius();
-	FReal HalfHeight = InGeom.GetHeight() * 0.5f;
+	FReal Radius = InGeom.GetRadiusf();
+	FReal HalfHeight = InGeom.GetHeightf() * 0.5f;
 
 	OutSamples[0]	= FVec3(HalfHeight + Radius, 0, 0);
 	OutSamples[1]	= FVec3(-HalfHeight - Radius, 0, 0);
@@ -1660,7 +1660,7 @@ bool FLevelSet::SweepGeomImp(const QueryGeomType& QueryGeom, const FRigidTransfo
 	return bHit;
 }
 
-bool FLevelSet::SweepGeom(const TSphere<FReal, 3>& QueryGeom, const FRigidTransform3& StartTM, const FVec3& Dir, const FReal Length, FReal& OutTime, FVec3& OutPosition, FVec3& OutNormal, int32& OutFaceIndex, const FReal Thickness, const bool bComputeMTD) const
+bool FLevelSet::SweepGeom(const FSphere& QueryGeom, const FRigidTransform3& StartTM, const FVec3& Dir, const FReal Length, FReal& OutTime, FVec3& OutPosition, FVec3& OutNormal, int32& OutFaceIndex, const FReal Thickness, const bool bComputeMTD) const
 {
 	return SweepGeomImp(QueryGeom, StartTM, Dir, Length, OutTime, OutPosition, OutNormal, OutFaceIndex, Thickness, bComputeMTD);
 }
@@ -1680,7 +1680,7 @@ bool FLevelSet::SweepGeom(const FConvex& QueryGeom, const FRigidTransform3& Star
 	return SweepGeomImp(QueryGeom, StartTM, Dir, Length, OutTime, OutPosition, OutNormal, OutFaceIndex, Thickness, bComputeMTD);
 }
 
-bool FLevelSet::SweepGeom(const TImplicitObjectScaled<TSphere<FReal, 3>>& QueryGeom, const FRigidTransform3& StartTM, const FVec3& Dir, const FReal Length, FReal& OutTime, FVec3& OutPosition, FVec3& OutNormal, int32& OutFaceIndex, const FReal Thickness, const bool bComputeMTD) const
+bool FLevelSet::SweepGeom(const TImplicitObjectScaled<FSphere>& QueryGeom, const FRigidTransform3& StartTM, const FVec3& Dir, const FReal Length, FReal& OutTime, FVec3& OutPosition, FVec3& OutNormal, int32& OutFaceIndex, const FReal Thickness, const bool bComputeMTD) const
 {
 	return SweepGeomImp(QueryGeom, StartTM, Dir, Length, OutTime, OutPosition, OutNormal, OutFaceIndex, Thickness, bComputeMTD);
 }
@@ -1700,7 +1700,7 @@ bool FLevelSet::SweepGeom(const TImplicitObjectScaled<FConvex>& QueryGeom, const
 	return SweepGeomImp(QueryGeom, StartTM, Dir, Length, OutTime, OutPosition, OutNormal, OutFaceIndex, Thickness, bComputeMTD);
 }
 
-void GetGeomSurfaceSamplesExtended(const TSphere<FReal, 3>& InGeom, TArray<FVec3>& OutSamples)
+void GetGeomSurfaceSamplesExtended(const FSphere& InGeom, TArray<FVec3>& OutSamples)
 {
 	OutSamples = InGeom.ComputeLocalSamplePoints(NumOverlapSphereSamples);
 }
@@ -1785,7 +1785,7 @@ bool FLevelSet::OverlapGeomImp(const QueryGeomType& QueryGeom, const FRigidTrans
 	return bResult;
 }
 
-bool FLevelSet::OverlapGeom(const TSphere<FReal, 3>& QueryGeom, const FRigidTransform3& QueryTM, const FReal Thickness, FMTDInfo* OutMTD) const
+bool FLevelSet::OverlapGeom(const FSphere& QueryGeom, const FRigidTransform3& QueryTM, const FReal Thickness, FMTDInfo* OutMTD) const
 {
 	return OverlapGeomImp(QueryGeom, QueryTM, Thickness, OutMTD);
 }
@@ -1805,7 +1805,7 @@ bool FLevelSet::OverlapGeom(const FConvex& QueryGeom, const FRigidTransform3& Qu
 	return OverlapGeomImp(QueryGeom, QueryTM, Thickness, OutMTD);
 }
 
-bool FLevelSet::OverlapGeom(const TImplicitObjectScaled<TSphere<FReal, 3>>& QueryGeom, const FRigidTransform3& QueryTM, const FReal Thickness, FMTDInfo* OutMTD) const
+bool FLevelSet::OverlapGeom(const TImplicitObjectScaled<FSphere>& QueryGeom, const FRigidTransform3& QueryTM, const FReal Thickness, FMTDInfo* OutMTD) const
 {
 	return OverlapGeomImp(QueryGeom, QueryTM, Thickness, OutMTD);
 }
diff --git a/Engine/Source/Runtime/Experimental/Chaos/Private/Chaos/PBDRigidClusteringAlgo.cpp b/Engine/Source/Runtime/Experimental/Chaos/Private/Chaos/PBDRigidClusteringAlgo.cpp
--- a/Engine/Source/Runtime/Experimental/Chaos/Private/Chaos/PBDRigidClusteringAlgo.cpp
+++ b/Engine/Source/Runtime/Experimental/Chaos/Private/Chaos/PBDRigidClusteringAlgo.cpp
@@ -415,7 +415,7 @@ namespace Chaos
 				}
 				else
 				{
-					Parent->SetGeometry(MakeImplicitObjectPtr<TSphere<FReal, 3>>(FVec3(0), BoundsExtents.Size() * 0.5f));
+					Parent->SetGeometry(MakeImplicitObjectPtr<FSphere>(FVec3(0), BoundsExtents.Size() * 0.5f));
 				}
 			}
 			else // !UseLevelsetCollision
diff --git a/Engine/Source/Runtime/Experimental/Chaos/Private/Chaos/PBDSoftBodyCollisionConstraint.cpp b/Engine/Source/Runtime/Experimental/Chaos/Private/Chaos/PBDSoftBodyCollisionConstraint.cpp
--- a/Engine/Source/Runtime/Experimental/Chaos/Private/Chaos/PBDSoftBodyCollisionConstraint.cpp
+++ b/Engine/Source/Runtime/Experimental/Chaos/Private/Chaos/PBDSoftBodyCollisionConstraint.cpp
@@ -11,106 +11,8 @@
 #include "Misc/ScopeLock.h"
 #if INTEL_ISPC
 #include "PerParticlePBDCollisionConstraint.ispc.generated.h"
+#include "PerParticlePBDCollisionConstraintISPCDataVerification.h"
 
-// These includes are used for the static_asserts below
-#include "Chaos/Capsule.h"
-#include "Chaos/Convex.h"
-#include "Chaos/ImplicitObjectUnion.h"
-#include "Chaos/Sphere.h"
-#include "Chaos/TaperedCapsule.h"
-#include "Chaos/TaperedCylinder.h"
-#endif
-
-#if INTEL_ISPC
-static_assert(sizeof(ispc::FVector3f) == sizeof(Chaos::Softs::FSolverVec3), "sizeof(ispc::FVector3f) != sizeof(Chaos::Softs::FSolverVec3)");
-static_assert(sizeof(ispc::FVector4f) == sizeof(Chaos::Softs::FSolverRotation3), "sizeof(ispc::FVector4f) != sizeof(Chaos::Softs::FSolverRotation3)");
-static_assert(ispc::ImplicitObjectType::Sphere == Chaos::ImplicitObjectType::Sphere);
-static_assert(ispc::ImplicitObjectType::Capsule == Chaos::ImplicitObjectType::Capsule);
-static_assert(ispc::ImplicitObjectType::Union == Chaos::ImplicitObjectType::Union);
-static_assert(ispc::ImplicitObjectType::TaperedCapsule == Chaos::ImplicitObjectType::TaperedCapsule);
-static_assert(ispc::ImplicitObjectType::Convex == Chaos::ImplicitObjectType::Convex);
-static_assert(ispc::ImplicitObjectType::IsWeightedLattice == Chaos::ImplicitObjectType::IsWeightedLattice);
-static_assert(ispc::ImplicitObjectType::MLLevelSet == Chaos::ImplicitObjectType::MLLevelSet);
-static_assert(ispc::ImplicitObjectType::SkinnedTriangleMesh == Chaos::ImplicitObjectType::SkinnedTriangleMesh);
-static_assert(ispc::ImplicitObjectType::WeightedLatticeLevelSetType == (Chaos::ImplicitObjectType::IsWeightedLattice | Chaos::ImplicitObjectType::LevelSet));
-static_assert(sizeof(ispc::TArray) == sizeof(TArray<int>));
-// Sphere
-static_assert(sizeof(Chaos::FImplicitObject) == Chaos::FSphere::FISPCDataVerifier::OffsetOfCenter());
-static_assert(sizeof(ispc::FVector) == Chaos::FSphere::FISPCDataVerifier::SizeOfCenter());
-// Capsule
-static_assert(sizeof(Chaos::FImplicitObject) == Chaos::FCapsule::FISPCDataVerifier::OffsetOfMSegment());
-static_assert(sizeof(ispc::Segment) == Chaos::FCapsule::FISPCDataVerifier::SizeOfMSegment());
-// Union (only specific case of MObjects = [FTaperedCylinder, Sphere, Sphere] is used here.
-static_assert(sizeof(Chaos::FImplicitObject) == Chaos::FImplicitObjectUnion::FISPCDataVerifier::OffsetOfMObjects());
-static_assert(sizeof(ispc::TArray) == Chaos::FImplicitObjectUnion::FISPCDataVerifier::SizeOfMObjects());
-// TaperedCylinder
-static_assert(sizeof(Chaos::FImplicitObject) + offsetof(ispc::FTaperedCylinder, MPlane1) == Chaos::FTaperedCylinder::FISPCDataVerifier::OffsetOfMPlane1());
-static_assert(sizeof(ispc::FTaperedCylinder::MPlane1) == Chaos::FTaperedCylinder::FISPCDataVerifier::SizeOfMPlane1());
-static_assert(sizeof(Chaos::FImplicitObject) + offsetof(ispc::FTaperedCylinder, MPlane2) == Chaos::FTaperedCylinder::FISPCDataVerifier::OffsetOfMPlane2());
-static_assert(sizeof(ispc::FTaperedCylinder::MPlane2) == Chaos::FTaperedCylinder::FISPCDataVerifier::SizeOfMPlane2());
-static_assert(sizeof(Chaos::FImplicitObject) + offsetof(ispc::FTaperedCylinder, Height) == Chaos::FTaperedCylinder::FISPCDataVerifier::OffsetOfMHeight());
-static_assert(sizeof(ispc::FTaperedCylinder::Height) == Chaos::FTaperedCylinder::FISPCDataVerifier::SizeOfMHeight());
-static_assert(sizeof(Chaos::FImplicitObject) + offsetof(ispc::FTaperedCylinder, Radius1) == Chaos::FTaperedCylinder::FISPCDataVerifier::OffsetOfMRadius1());
-static_assert(sizeof(ispc::FTaperedCylinder::Radius1) == Chaos::FTaperedCylinder::FISPCDataVerifier::SizeOfMRadius1());
-static_assert(sizeof(Chaos::FImplicitObject) + offsetof(ispc::FTaperedCylinder, Radius2) == Chaos::FTaperedCylinder::FISPCDataVerifier::OffsetOfMRadius2());
-static_assert(sizeof(ispc::FTaperedCylinder::Radius2) == Chaos::FTaperedCylinder::FISPCDataVerifier::SizeOfMRadius2());
-// TaperedCapsule
-static_assert(sizeof(Chaos::FImplicitObject) + offsetof(ispc::FTaperedCapsule, Origin) == Chaos::FTaperedCapsule::FISPCDataVerifier::OffsetOfOrigin());
-static_assert(sizeof(ispc::FTaperedCapsule::Origin) == Chaos::FTaperedCapsule::FISPCDataVerifier::SizeOfOrigin());
-static_assert(sizeof(Chaos::FImplicitObject) + offsetof(ispc::FTaperedCapsule, Axis) == Chaos::FTaperedCapsule::FISPCDataVerifier::OffsetOfAxis());
-static_assert(sizeof(ispc::FTaperedCapsule::Axis) == Chaos::FTaperedCapsule::FISPCDataVerifier::SizeOfAxis());
-static_assert(sizeof(Chaos::FImplicitObject) + offsetof(ispc::FTaperedCapsule, Height) == Chaos::FTaperedCapsule::FISPCDataVerifier::OffsetOfHeight());
-static_assert(sizeof(ispc::FTaperedCapsule::Height) == Chaos::FTaperedCapsule::FISPCDataVerifier::SizeOfHeight());
-static_assert(sizeof(Chaos::FImplicitObject) + offsetof(ispc::FTaperedCapsule, Radius1) == Chaos::FTaperedCapsule::FISPCDataVerifier::OffsetOfRadius1());
-static_assert(sizeof(ispc::FTaperedCapsule::Radius1) == Chaos::FTaperedCapsule::FISPCDataVerifier::SizeOfRadius1());
-static_assert(sizeof(Chaos::FImplicitObject) + offsetof(ispc::FTaperedCapsule, Radius2) == Chaos::FTaperedCapsule::FISPCDataVerifier::OffsetOfRadius2());
-static_assert(sizeof(ispc::FTaperedCapsule::Radius2) == Chaos::FTaperedCapsule::FISPCDataVerifier::SizeOfRadius2());
-// Convex
-static_assert(sizeof(Chaos::FImplicitObject) + offsetof(ispc::FConvex, Planes) == Chaos::FConvex::FISPCDataVerifier::OffsetOfPlanes());
-static_assert(sizeof(ispc::FConvex::Planes) == Chaos::FConvex::FISPCDataVerifier::SizeOfPlanes());
-static_assert(sizeof(Chaos::FImplicitObject) + offsetof(ispc::FConvex, Vertices) == Chaos::FConvex::FISPCDataVerifier::OffsetOfVertices());
-static_assert(sizeof(ispc::FConvex::Vertices) == Chaos::FConvex::FISPCDataVerifier::SizeOfVertices());
-static_assert(sizeof(Chaos::FImplicitObject) + offsetof(ispc::FConvex, StructureData) == Chaos::FConvex::FISPCDataVerifier::OffsetOfStructureData());
-static_assert(sizeof(ispc::FConvex::StructureData) == Chaos::FConvex::FISPCDataVerifier::SizeOfStructureData());
-
-static_assert(sizeof(ispc::FPlaneConcrete3f) == sizeof(Chaos::FConvex::FPlaneType));
-static_assert(sizeof(ispc::FVector3f) == sizeof(Chaos::FConvex::FVec3Type)); // Vertices array
-
-static_assert(offsetof(ispc::FConvexStructureData, Data) == Chaos::FConvexStructureData::FISPCDataVerifier::OffsetOfData());
-static_assert(sizeof(ispc::FConvexStructureData::Data) == Chaos::FConvexStructureData::FISPCDataVerifier::SizeOfData());
-static_assert(offsetof(ispc::FConvexStructureData, IndexType) == Chaos::FConvexStructureData::FISPCDataVerifier::OffsetOfIndexType());
-static_assert(sizeof(ispc::FConvexStructureData::IndexType) == Chaos::FConvexStructureData::FISPCDataVerifier::SizeOfIndexType());
-
-static_assert(offsetof(ispc::FConvexStructureDataImp, Planes) == Chaos::FConvexStructureData::FConvexStructureDataSmall::FISPCDataVerifier::OffsetOfPlanes());
-static_assert(offsetof(ispc::FConvexStructureDataImp, Planes) == Chaos::FConvexStructureData::FConvexStructureDataMedium::FISPCDataVerifier::OffsetOfPlanes());
-static_assert(offsetof(ispc::FConvexStructureDataImp, Planes) == Chaos::FConvexStructureData::FConvexStructureDataLarge::FISPCDataVerifier::OffsetOfPlanes());
-static_assert(sizeof(ispc::FConvexStructureDataImp::Planes) == Chaos::FConvexStructureData::FConvexStructureDataSmall::FISPCDataVerifier::SizeOfPlanes());
-static_assert(sizeof(ispc::FConvexStructureDataImp::Planes) == Chaos::FConvexStructureData::FConvexStructureDataMedium::FISPCDataVerifier::SizeOfPlanes());
-static_assert(sizeof(ispc::FConvexStructureDataImp::Planes) == Chaos::FConvexStructureData::FConvexStructureDataLarge::FISPCDataVerifier::SizeOfPlanes());
-
-static_assert(offsetof(ispc::FConvexStructureDataImp, HalfEdges) == Chaos::FConvexStructureData::FConvexStructureDataSmall::FISPCDataVerifier::OffsetOfHalfEdges());
-static_assert(offsetof(ispc::FConvexStructureDataImp, HalfEdges) == Chaos::FConvexStructureData::FConvexStructureDataMedium::FISPCDataVerifier::OffsetOfHalfEdges());
-static_assert(offsetof(ispc::FConvexStructureDataImp, HalfEdges) == Chaos::FConvexStructureData::FConvexStructureDataLarge::FISPCDataVerifier::OffsetOfHalfEdges());
-static_assert(sizeof(ispc::FConvexStructureDataImp::HalfEdges) == Chaos::FConvexStructureData::FConvexStructureDataSmall::FISPCDataVerifier::SizeOfHalfEdges());
-static_assert(sizeof(ispc::FConvexStructureDataImp::HalfEdges) == Chaos::FConvexStructureData::FConvexStructureDataMedium::FISPCDataVerifier::SizeOfHalfEdges());
-static_assert(sizeof(ispc::FConvexStructureDataImp::HalfEdges) == Chaos::FConvexStructureData::FConvexStructureDataLarge::FISPCDataVerifier::SizeOfHalfEdges());
-
-static_assert(offsetof(ispc::FConvexStructureDataImp, Vertices) == Chaos::FConvexStructureData::FConvexStructureDataSmall::FISPCDataVerifier::OffsetOfVertices());
-static_assert(offsetof(ispc::FConvexStructureDataImp, Vertices) == Chaos::FConvexStructureData::FConvexStructureDataMedium::FISPCDataVerifier::OffsetOfVertices());
-static_assert(offsetof(ispc::FConvexStructureDataImp, Vertices) == Chaos::FConvexStructureData::FConvexStructureDataLarge::FISPCDataVerifier::OffsetOfVertices());
-static_assert(sizeof(ispc::FConvexStructureDataImp::Vertices) == Chaos::FConvexStructureData::FConvexStructureDataSmall::FISPCDataVerifier::SizeOfVertices());
-static_assert(sizeof(ispc::FConvexStructureDataImp::Vertices) == Chaos::FConvexStructureData::FConvexStructureDataMedium::FISPCDataVerifier::SizeOfVertices());
-static_assert(sizeof(ispc::FConvexStructureDataImp::Vertices) == Chaos::FConvexStructureData::FConvexStructureDataLarge::FISPCDataVerifier::SizeOfVertices());
-
-static_assert(sizeof(ispc::PlanesS) == sizeof(Chaos::FConvexStructureData::FConvexStructureDataSmall::FPlaneData));
-static_assert(sizeof(ispc::PlanesM) == sizeof(Chaos::FConvexStructureData::FConvexStructureDataMedium::FPlaneData));
-static_assert(sizeof(ispc::PlanesL) == sizeof(Chaos::FConvexStructureData::FConvexStructureDataLarge::FPlaneData));
-static_assert(sizeof(ispc::HalfEdgesS) == sizeof(Chaos::FConvexStructureData::FConvexStructureDataSmall::FHalfEdgeData));
-static_assert(sizeof(ispc::HalfEdgesM) == sizeof(Chaos::FConvexStructureData::FConvexStructureDataMedium::FHalfEdgeData));
-static_assert(sizeof(ispc::HalfEdgesL) == sizeof(Chaos::FConvexStructureData::FConvexStructureDataLarge::FHalfEdgeData));
-static_assert(sizeof(ispc::VerticesS) == sizeof(Chaos::FConvexStructureData::FConvexStructureDataSmall::FVertexData));
-static_assert(sizeof(ispc::VerticesM) == sizeof(Chaos::FConvexStructureData::FConvexStructureDataMedium::FVertexData));
-static_assert(sizeof(ispc::VerticesL) == sizeof(Chaos::FConvexStructureData::FConvexStructureDataLarge::FVertexData));
 
 #if !UE_BUILD_SHIPPING
 bool bChaos_SoftBodyCollision_ISPC_Enabled = CHAOS_PER_PARTICLE_COLLISION_ISPC_ENABLED_DEFAULT;
diff --git a/Engine/Source/Runtime/Experimental/Chaos/Private/Chaos/ParticleHandle.cpp b/Engine/Source/Runtime/Experimental/Chaos/Private/Chaos/ParticleHandle.cpp
--- a/Engine/Source/Runtime/Experimental/Chaos/Private/Chaos/ParticleHandle.cpp
+++ b/Engine/Source/Runtime/Experimental/Chaos/Private/Chaos/ParticleHandle.cpp
@@ -57,9 +57,9 @@ namespace Chaos
 		{
 			return const_cast<FCapsule*>(Implicit0->template GetObject<TImplicitObjectInstanced<FCapsule>>()->GetInstancedObject());
 		}
-		else if (Implicit0OuterType == TImplicitObjectInstanced<TSphere<FReal, 3>>::StaticType())
+		else if (Implicit0OuterType == TImplicitObjectInstanced<FSphere>::StaticType())
 		{
-			return const_cast<TSphere<FReal, 3>*>(Implicit0->template GetObject<TImplicitObjectInstanced<TSphere<FReal, 3>>>()->GetInstancedObject());
+			return const_cast<FSphere*>(Implicit0->template GetObject<TImplicitObjectInstanced<FSphere>>()->GetInstancedObject());
 		}
 		else if (Implicit0OuterType == TImplicitObjectInstanced<FTriangleMeshImplicitObject>::StaticType())
 		{
diff --git a/Engine/Source/Runtime/Experimental/Chaos/Private/Chaos/PerParticlePBDCollisionConstraint.cpp b/Engine/Source/Runtime/Experimental/Chaos/Private/Chaos/PerParticlePBDCollisionConstraint.cpp
--- a/Engine/Source/Runtime/Experimental/Chaos/Private/Chaos/PerParticlePBDCollisionConstraint.cpp
+++ b/Engine/Source/Runtime/Experimental/Chaos/Private/Chaos/PerParticlePBDCollisionConstraint.cpp
@@ -5,106 +5,7 @@
 #include "HAL/IConsoleManager.h"
 #if INTEL_ISPC
 #include "PerParticlePBDCollisionConstraint.ispc.generated.h"
-
-// These includes are used for the static_asserts below
-#include "Chaos/Capsule.h"
-#include "Chaos/Convex.h"
-#include "Chaos/ImplicitObjectUnion.h"
-#include "Chaos/Sphere.h"
-#include "Chaos/TaperedCapsule.h"
-#include "Chaos/TaperedCylinder.h"
-#endif
-
-#if INTEL_ISPC
-static_assert(sizeof(ispc::FVector3f) == sizeof(Chaos::Softs::FSolverVec3), "sizeof(ispc::FVector3f) != sizeof(Chaos::Softs::FSolverVec3)");
-static_assert(sizeof(ispc::FVector4f) == sizeof(Chaos::Softs::FSolverRotation3), "sizeof(ispc::FVector4f) != sizeof(Chaos::Softs::FSolverRotation3)");
-static_assert(ispc::ImplicitObjectType::Sphere == Chaos::ImplicitObjectType::Sphere);
-static_assert(ispc::ImplicitObjectType::Capsule == Chaos::ImplicitObjectType::Capsule);
-static_assert(ispc::ImplicitObjectType::Union == Chaos::ImplicitObjectType::Union);
-static_assert(ispc::ImplicitObjectType::TaperedCapsule == Chaos::ImplicitObjectType::TaperedCapsule);
-static_assert(ispc::ImplicitObjectType::Convex == Chaos::ImplicitObjectType::Convex);
-static_assert(ispc::ImplicitObjectType::IsWeightedLattice == Chaos::ImplicitObjectType::IsWeightedLattice);
-static_assert(ispc::ImplicitObjectType::MLLevelSet == Chaos::ImplicitObjectType::MLLevelSet);
-static_assert(ispc::ImplicitObjectType::SkinnedTriangleMesh == Chaos::ImplicitObjectType::SkinnedTriangleMesh);
-static_assert(ispc::ImplicitObjectType::WeightedLatticeLevelSetType == (Chaos::ImplicitObjectType::IsWeightedLattice | Chaos::ImplicitObjectType::LevelSet));
-static_assert(sizeof(ispc::TArray) == sizeof(TArray<int>));
-// Sphere
-static_assert(sizeof(Chaos::FImplicitObject) == Chaos::FSphere::FISPCDataVerifier::OffsetOfCenter());
-static_assert(sizeof(ispc::FVector) == Chaos::FSphere::FISPCDataVerifier::SizeOfCenter());
-// Capsule
-static_assert(sizeof(Chaos::FImplicitObject) == Chaos::FCapsule::FISPCDataVerifier::OffsetOfMSegment());
-static_assert(sizeof(ispc::Segment) == Chaos::FCapsule::FISPCDataVerifier::SizeOfMSegment());
-// Union (only specific case of MObjects = [FTaperedCylinder, Sphere, Sphere] is used here.
-static_assert(sizeof(Chaos::FImplicitObject) == Chaos::FImplicitObjectUnion::FISPCDataVerifier::OffsetOfMObjects());
-static_assert(sizeof(ispc::TArray) == Chaos::FImplicitObjectUnion::FISPCDataVerifier::SizeOfMObjects());
-// TaperedCylinder
-static_assert(sizeof(Chaos::FImplicitObject) + offsetof(ispc::FTaperedCylinder, MPlane1) == Chaos::FTaperedCylinder::FISPCDataVerifier::OffsetOfMPlane1());
-static_assert(sizeof(ispc::FTaperedCylinder::MPlane1) == Chaos::FTaperedCylinder::FISPCDataVerifier::SizeOfMPlane1());
-static_assert(sizeof(Chaos::FImplicitObject) + offsetof(ispc::FTaperedCylinder, MPlane2) == Chaos::FTaperedCylinder::FISPCDataVerifier::OffsetOfMPlane2());
-static_assert(sizeof(ispc::FTaperedCylinder::MPlane2) == Chaos::FTaperedCylinder::FISPCDataVerifier::SizeOfMPlane2());
-static_assert(sizeof(Chaos::FImplicitObject) + offsetof(ispc::FTaperedCylinder, Height) == Chaos::FTaperedCylinder::FISPCDataVerifier::OffsetOfMHeight());
-static_assert(sizeof(ispc::FTaperedCylinder::Height) == Chaos::FTaperedCylinder::FISPCDataVerifier::SizeOfMHeight());
-static_assert(sizeof(Chaos::FImplicitObject) + offsetof(ispc::FTaperedCylinder, Radius1) == Chaos::FTaperedCylinder::FISPCDataVerifier::OffsetOfMRadius1());
-static_assert(sizeof(ispc::FTaperedCylinder::Radius1) == Chaos::FTaperedCylinder::FISPCDataVerifier::SizeOfMRadius1());
-static_assert(sizeof(Chaos::FImplicitObject) + offsetof(ispc::FTaperedCylinder, Radius2) == Chaos::FTaperedCylinder::FISPCDataVerifier::OffsetOfMRadius2());
-static_assert(sizeof(ispc::FTaperedCylinder::Radius2) == Chaos::FTaperedCylinder::FISPCDataVerifier::SizeOfMRadius2());
-// TaperedCapsule
-static_assert(sizeof(Chaos::FImplicitObject) + offsetof(ispc::FTaperedCapsule, Origin) == Chaos::FTaperedCapsule::FISPCDataVerifier::OffsetOfOrigin());
-static_assert(sizeof(ispc::FTaperedCapsule::Origin) == Chaos::FTaperedCapsule::FISPCDataVerifier::SizeOfOrigin());
-static_assert(sizeof(Chaos::FImplicitObject) + offsetof(ispc::FTaperedCapsule, Axis) == Chaos::FTaperedCapsule::FISPCDataVerifier::OffsetOfAxis());
-static_assert(sizeof(ispc::FTaperedCapsule::Axis) == Chaos::FTaperedCapsule::FISPCDataVerifier::SizeOfAxis());
-static_assert(sizeof(Chaos::FImplicitObject) + offsetof(ispc::FTaperedCapsule, Height) == Chaos::FTaperedCapsule::FISPCDataVerifier::OffsetOfHeight());
-static_assert(sizeof(ispc::FTaperedCapsule::Height) == Chaos::FTaperedCapsule::FISPCDataVerifier::SizeOfHeight());
-static_assert(sizeof(Chaos::FImplicitObject) + offsetof(ispc::FTaperedCapsule, Radius1) == Chaos::FTaperedCapsule::FISPCDataVerifier::OffsetOfRadius1());
-static_assert(sizeof(ispc::FTaperedCapsule::Radius1) == Chaos::FTaperedCapsule::FISPCDataVerifier::SizeOfRadius1());
-static_assert(sizeof(Chaos::FImplicitObject) + offsetof(ispc::FTaperedCapsule, Radius2) == Chaos::FTaperedCapsule::FISPCDataVerifier::OffsetOfRadius2());
-static_assert(sizeof(ispc::FTaperedCapsule::Radius2) == Chaos::FTaperedCapsule::FISPCDataVerifier::SizeOfRadius2());
-// Convex
-static_assert(sizeof(Chaos::FImplicitObject) + offsetof(ispc::FConvex, Planes) == Chaos::FConvex::FISPCDataVerifier::OffsetOfPlanes());
-static_assert(sizeof(ispc::FConvex::Planes) == Chaos::FConvex::FISPCDataVerifier::SizeOfPlanes());
-static_assert(sizeof(Chaos::FImplicitObject) + offsetof(ispc::FConvex, Vertices) == Chaos::FConvex::FISPCDataVerifier::OffsetOfVertices());
-static_assert(sizeof(ispc::FConvex::Vertices) == Chaos::FConvex::FISPCDataVerifier::SizeOfVertices());
-static_assert(sizeof(Chaos::FImplicitObject) + offsetof(ispc::FConvex, StructureData) == Chaos::FConvex::FISPCDataVerifier::OffsetOfStructureData());
-static_assert(sizeof(ispc::FConvex::StructureData) == Chaos::FConvex::FISPCDataVerifier::SizeOfStructureData());
-
-static_assert(sizeof(ispc::FPlaneConcrete3f) == sizeof(Chaos::FConvex::FPlaneType));
-static_assert(sizeof(ispc::FVector3f) == sizeof(Chaos::FConvex::FVec3Type)); // Vertices array
-
-static_assert(offsetof(ispc::FConvexStructureData, Data) == Chaos::FConvexStructureData::FISPCDataVerifier::OffsetOfData());
-static_assert(sizeof(ispc::FConvexStructureData::Data) == Chaos::FConvexStructureData::FISPCDataVerifier::SizeOfData());
-static_assert(offsetof(ispc::FConvexStructureData, IndexType) == Chaos::FConvexStructureData::FISPCDataVerifier::OffsetOfIndexType());
-static_assert(sizeof(ispc::FConvexStructureData::IndexType) == Chaos::FConvexStructureData::FISPCDataVerifier::SizeOfIndexType());
-
-static_assert(offsetof(ispc::FConvexStructureDataImp, Planes) == Chaos::FConvexStructureData::FConvexStructureDataSmall::FISPCDataVerifier::OffsetOfPlanes());
-static_assert(offsetof(ispc::FConvexStructureDataImp, Planes) == Chaos::FConvexStructureData::FConvexStructureDataMedium::FISPCDataVerifier::OffsetOfPlanes());
-static_assert(offsetof(ispc::FConvexStructureDataImp, Planes) == Chaos::FConvexStructureData::FConvexStructureDataLarge::FISPCDataVerifier::OffsetOfPlanes());
-static_assert(sizeof(ispc::FConvexStructureDataImp::Planes) == Chaos::FConvexStructureData::FConvexStructureDataSmall::FISPCDataVerifier::SizeOfPlanes());
-static_assert(sizeof(ispc::FConvexStructureDataImp::Planes) == Chaos::FConvexStructureData::FConvexStructureDataMedium::FISPCDataVerifier::SizeOfPlanes());
-static_assert(sizeof(ispc::FConvexStructureDataImp::Planes) == Chaos::FConvexStructureData::FConvexStructureDataLarge::FISPCDataVerifier::SizeOfPlanes());
-
-static_assert(offsetof(ispc::FConvexStructureDataImp, HalfEdges) == Chaos::FConvexStructureData::FConvexStructureDataSmall::FISPCDataVerifier::OffsetOfHalfEdges());
-static_assert(offsetof(ispc::FConvexStructureDataImp, HalfEdges) == Chaos::FConvexStructureData::FConvexStructureDataMedium::FISPCDataVerifier::OffsetOfHalfEdges());
-static_assert(offsetof(ispc::FConvexStructureDataImp, HalfEdges) == Chaos::FConvexStructureData::FConvexStructureDataLarge::FISPCDataVerifier::OffsetOfHalfEdges());
-static_assert(sizeof(ispc::FConvexStructureDataImp::HalfEdges) == Chaos::FConvexStructureData::FConvexStructureDataSmall::FISPCDataVerifier::SizeOfHalfEdges());
-static_assert(sizeof(ispc::FConvexStructureDataImp::HalfEdges) == Chaos::FConvexStructureData::FConvexStructureDataMedium::FISPCDataVerifier::SizeOfHalfEdges());
-static_assert(sizeof(ispc::FConvexStructureDataImp::HalfEdges) == Chaos::FConvexStructureData::FConvexStructureDataLarge::FISPCDataVerifier::SizeOfHalfEdges());
-
-static_assert(offsetof(ispc::FConvexStructureDataImp, Vertices) == Chaos::FConvexStructureData::FConvexStructureDataSmall::FISPCDataVerifier::OffsetOfVertices());
-static_assert(offsetof(ispc::FConvexStructureDataImp, Vertices) == Chaos::FConvexStructureData::FConvexStructureDataMedium::FISPCDataVerifier::OffsetOfVertices());
-static_assert(offsetof(ispc::FConvexStructureDataImp, Vertices) == Chaos::FConvexStructureData::FConvexStructureDataLarge::FISPCDataVerifier::OffsetOfVertices());
-static_assert(sizeof(ispc::FConvexStructureDataImp::Vertices) == Chaos::FConvexStructureData::FConvexStructureDataSmall::FISPCDataVerifier::SizeOfVertices());
-static_assert(sizeof(ispc::FConvexStructureDataImp::Vertices) == Chaos::FConvexStructureData::FConvexStructureDataMedium::FISPCDataVerifier::SizeOfVertices());
-static_assert(sizeof(ispc::FConvexStructureDataImp::Vertices) == Chaos::FConvexStructureData::FConvexStructureDataLarge::FISPCDataVerifier::SizeOfVertices());
-
-static_assert(sizeof(ispc::PlanesS) == sizeof(Chaos::FConvexStructureData::FConvexStructureDataSmall::FPlaneData));
-static_assert(sizeof(ispc::PlanesM) == sizeof(Chaos::FConvexStructureData::FConvexStructureDataMedium::FPlaneData));
-static_assert(sizeof(ispc::PlanesL) == sizeof(Chaos::FConvexStructureData::FConvexStructureDataLarge::FPlaneData));
-static_assert(sizeof(ispc::HalfEdgesS) == sizeof(Chaos::FConvexStructureData::FConvexStructureDataSmall::FHalfEdgeData));
-static_assert(sizeof(ispc::HalfEdgesM) == sizeof(Chaos::FConvexStructureData::FConvexStructureDataMedium::FHalfEdgeData));
-static_assert(sizeof(ispc::HalfEdgesL) == sizeof(Chaos::FConvexStructureData::FConvexStructureDataLarge::FHalfEdgeData));
-static_assert(sizeof(ispc::VerticesS) == sizeof(Chaos::FConvexStructureData::FConvexStructureDataSmall::FVertexData));
-static_assert(sizeof(ispc::VerticesM) == sizeof(Chaos::FConvexStructureData::FConvexStructureDataMedium::FVertexData));
-static_assert(sizeof(ispc::VerticesL) == sizeof(Chaos::FConvexStructureData::FConvexStructureDataLarge::FVertexData));
+#include "PerParticlePBDCollisionConstraintISPCDataVerification.h"
 
 #if !UE_BUILD_SHIPPING || USE_ISPC_KERNEL_CONSOLE_VARIABLES_IN_SHIPPING 
 bool bChaos_PerParticleCollision_ISPC_Enabled = CHAOS_PER_PARTICLE_COLLISION_ISPC_ENABLED_DEFAULT;
diff --git a/Engine/Source/Runtime/Experimental/Chaos/Private/Chaos/PerParticlePBDCollisionConstraint.ispc b/Engine/Source/Runtime/Experimental/Chaos/Private/Chaos/PerParticlePBDCollisionConstraint.ispc
--- a/Engine/Source/Runtime/Experimental/Chaos/Private/Chaos/PerParticlePBDCollisionConstraint.ispc
+++ b/Engine/Source/Runtime/Experimental/Chaos/Private/Chaos/PerParticlePBDCollisionConstraint.ispc
@@ -110,15 +110,9 @@ struct VerticesL
 
 struct Segment
 {
-	FVector Point;
-	FVector Axis;
-	FReal Length;
-};
-
-struct FPlaneConcrete3d
-{
-	FVector3d MX;
-	FVector3d MNormal;
+	FVector3f Point;
+	FVector3f Axis;
+	float Length;
 };
 
 struct FPlaneConcrete3f
@@ -157,27 +151,27 @@ struct FConvex
 	FVector3f LocalBoundingBoxMin;
 	FVector3f LocalBoundingBoxMax;
 	FConvexStructureData StructureData;
-	// FReal Volume; // this code never uses this, and we never assume to know the sizeof(FConvex)
+	// float Volume; // this code never uses this, and we never assume to know the sizeof(FConvex)
 	// FVector CenterOfMass; // this code never uses this, and we never assume to know the sizeof(FConvex)
 };
 
 struct FTaperedCapsule
 {
-	FVector Origin;
-	FVector Axis;
-	FReal Height;
-	FReal Radius1;
-	FReal Radius2;
+	FVector3f Origin;
+	FVector3f Axis;
+	float Height;
+	float Radius1;
+	float Radius2;
 	// FAABB3 LocalBoundingBox; // this code never uses this, and we never assume to know the sizeof(FTaperedCapsule)
 };
 
 struct FTaperedCylinder
 {
-	FPlaneConcrete3d MPlane1;
-	FPlaneConcrete3d MPlane2;
-	FReal Height;
-	FReal Radius1;
-	FReal Radius2;
+	FPlaneConcrete3f MPlane1;
+	FPlaneConcrete3f MPlane2;
+	float Height;
+	float Radius1;
+	float Radius2;
 	// FAABB3 MLocalBoundingBox; // this code never uses this, and we never assume to know the sizeof(FTaperedCylinder)
 };
 
@@ -195,7 +189,6 @@ export uniform int32 ExportHack_PerParticlePBDCollisionConstraintStructs(
 	uniform VerticesM* uniform _VerticesM,
 	uniform VerticesL* uniform _VerticesL,
 	uniform Segment* uniform _Segment,
-	uniform FPlaneConcrete3d* uniform _Plane3d,
 	uniform FPlaneConcrete3f* uniform _Plane3f,
 	uniform FConvexStructureDataIndexType IndexType,
 	uniform FConvex* uniform _FConvex,
@@ -224,29 +217,29 @@ static inline uniform bool IsSimpleCollider(const uniform uint8 CollisionType)
 	return true;
 }
 
-static inline FReal SafeNormalize(FVector &Direction)
+static inline float SafeNormalize(FVector3f &Direction)
 {
-	const FReal SizeSqr = VectorSizeSquared(Direction);
-	const FReal Size = sqrt(SizeSqr);
-	Direction = VectorSelect((SizeSqr < KINDA_SMALL_NUMBER), ForwardVector, Direction / Size);
+	const float SizeSqr = VectorSizeSquared(Direction);
+	const float Size = sqrt(SizeSqr);
+	Direction = VectorSelect((SizeSqr < KINDA_SMALL_NUMBER), FloatForwardVector, Direction / Size);
 	return (Size < KINDA_SMALL_NUMBER) ? ZERO : Size;
 }
 
-static inline bool VectorNormalize(FVector &V)
+static inline bool VectorNormalize(FVector3f &V)
 {
-	const FReal SquareSum = VectorSizeSquared(V);
+	const float SquareSum = VectorSizeSquared(V);
 	if(SquareSum > SMALL_NUMBER)
 	{
-		const FReal Scale = InvSqrt(SquareSum);
+		const float Scale = InvSqrt(SquareSum);
 		V = V * Scale;
 		return true;
 	}
 	return false;
 }
 
-static inline void PhiWithNormalSphere(const uniform uint8 *uniform TypedObjectPtr, const FVector &PhiWithNormalInput, FVector &Normal, FReal &Phi, const uniform FReal Radius)
+static inline void PhiWithNormalSphere(const uniform uint8 *uniform TypedObjectPtr, const FVector3f &PhiWithNormalInput, FVector3f &Normal, float &Phi, const uniform float Radius)
 {
-	const uniform FVector Center = *((const uniform FVector *uniform)&TypedObjectPtr[0]);
+	const uniform FVector3f Center = *((const uniform FVector3f *uniform)&TypedObjectPtr[0]);
 	Normal = PhiWithNormalInput - Center;
 	Phi = SafeNormalize(Normal) - Radius;
 }
@@ -266,76 +259,76 @@ static inline float SignedDistancePlane(const uniform FPlaneConcrete3f *uniform
 	return VectorDot(x - P.MX, P.MNormal);
 }
 
-static inline void PhiWithNormalCapsule(const uniform uint8 *uniform TypedObjectPtr, const FVector &PhiWithNormalInput, FVector &Normal, FReal &Phi, const uniform FReal Radius)
+static inline void PhiWithNormalCapsule(const uniform uint8 *uniform TypedObjectPtr, const FVector3f &PhiWithNormalInput, FVector3f &Normal, float &Phi, const uniform float Radius)
 {
 	const uniform Segment MSegment = *((const uniform Segment *uniform)&TypedObjectPtr[0]);
 
-	const FReal Dot = clamp(VectorDot(PhiWithNormalInput - MSegment.Point, MSegment.Axis), (uniform FReal)(0), MSegment.Length);
-	const FVector ProjectedPoint = SetVector(Dot) * MSegment.Axis + MSegment.Point;
+	const float Dot = clamp(VectorDot(PhiWithNormalInput - MSegment.Point, MSegment.Axis), (uniform float)(0), MSegment.Length);
+	const FVector3f ProjectedPoint = MSegment.Axis * Dot + MSegment.Point;
 	Normal = PhiWithNormalInput - ProjectedPoint;
 	Phi = SafeNormalize(Normal) - Radius;
 }
 
-static inline void PhiWithNormalTaperedCapsule(const uniform uint8* uniform TypedObjectPtr, const FVector& PhiWithNormalInput, FVector& Normal, FReal& Phi)
+static inline void PhiWithNormalTaperedCapsule(const uniform uint8* uniform TypedObjectPtr, const FVector3f& PhiWithNormalInput, FVector3f& Normal, float& Phi)
 {
 	const uniform FTaperedCapsule CapsuleObj = *((const uniform FTaperedCapsule *uniform)&TypedObjectPtr[0]);
 
-	const FReal DistanceAlongAxis = clamp(VectorDot(PhiWithNormalInput - CapsuleObj.Origin, CapsuleObj.Axis), (uniform FReal)(0), CapsuleObj.Height);
-	const FVector ProjectedPoint = SetVector(DistanceAlongAxis) * CapsuleObj.Axis + CapsuleObj.Origin;
-	const FReal Alpha = DistanceAlongAxis / CapsuleObj.Height;
-	const FReal Radius = CapsuleObj.Radius1 * ((uniform FReal)(1.0) - Alpha) + CapsuleObj.Radius2 * Alpha;
+	const float DistanceAlongAxis = clamp(VectorDot(PhiWithNormalInput - CapsuleObj.Origin, CapsuleObj.Axis), (uniform float)(0), CapsuleObj.Height);
+	const FVector3f ProjectedPoint = CapsuleObj.Axis * DistanceAlongAxis + CapsuleObj.Origin;
+	const float Alpha = DistanceAlongAxis / CapsuleObj.Height;
+	const float Radius = CapsuleObj.Radius1 * ((uniform float)(1.0) - Alpha) + CapsuleObj.Radius2 * Alpha;
 	Normal = PhiWithNormalInput - ProjectedPoint;
 	Phi = SafeNormalize(Normal) - Radius;
 }
 
-static inline void PhiWithNormalTaperedCylinder(const uniform uint8 *uniform TypedObjectPtr, const FVector &PhiWithNormalInput, FVector &Normal, FReal &Phi, const uniform int SizeofFImplicitObject)
+static inline void PhiWithNormalTaperedCylinder(const uniform uint8 *uniform TypedObjectPtr, const FVector3f &PhiWithNormalInput, FVector3f &Normal, float &Phi, const uniform int SizeofFImplicitObject)
 {
 	const uniform FTaperedCylinder CylinderObj = *((const uniform FTaperedCylinder *uniform)&TypedObjectPtr[0]);
 
-	const FReal Distance1 = VectorDot(PhiWithNormalInput - CylinderObj.MPlane1.MX, CylinderObj.MPlane1.MNormal);
+	const float Distance1 = VectorDot(PhiWithNormalInput - CylinderObj.MPlane1.MX, CylinderObj.MPlane1.MNormal);
 	// Used to be Distance2 = CylinderObj.MPlane2.PhiWithNormal(x, Normal2); but that would trigger 
-	const FReal Distance2 = CylinderObj.Height - Distance1;          // the ensure on Distance2 being slightly larger than MHeight in some border cases
-	const FVector SideVector = (PhiWithNormalInput - (CylinderObj.MPlane1.MNormal * Distance1 + CylinderObj.MPlane1.MX));
-	const FReal Alpha = Distance1 / CylinderObj.Height;
-	const FReal Radius = CylinderObj.Radius1 * ((uniform FReal)(1.0) - Alpha) + CylinderObj.Radius2 * Alpha;
-	const FReal SideDistance = VectorSize(SideVector) - Radius;
-	const FReal TopDistance = Distance1 < Distance2 ? Distance1 : Distance2;
+	const float Distance2 = CylinderObj.Height - Distance1;          // the ensure on Distance2 being slightly larger than MHeight in some border cases
+	const FVector3f SideVector = (PhiWithNormalInput - (CylinderObj.MPlane1.MNormal * Distance1 + CylinderObj.MPlane1.MX));
+	const float Alpha = Distance1 / CylinderObj.Height;
+	const float Radius = CylinderObj.Radius1 * ((uniform float)(1.0) - Alpha) + CylinderObj.Radius2 * Alpha;
+	const float SideDistance = VectorSize(SideVector) - Radius;
+	const float TopDistance = Distance1 < Distance2 ? Distance1 : Distance2;
 
 	if (Distance1 < SMALL_NUMBER)
 	{
-		const FVector v = PhiWithNormalInput - (CylinderObj.MPlane1.MNormal * Distance1 + CylinderObj.MPlane1.MX);
+		const FVector3f v = PhiWithNormalInput - (CylinderObj.MPlane1.MNormal * Distance1 + CylinderObj.MPlane1.MX);
 		if (VectorSize(v) > CylinderObj.Radius1)
 		{
-			const FVector Corner = VectorGetSafeNormal(v) * CylinderObj.Radius1 + CylinderObj.MPlane1.MX;
-			const FVector CornerVector = PhiWithNormalInput - Corner;
+			const FVector3f Corner = VectorGetSafeNormal(v) * CylinderObj.Radius1 + CylinderObj.MPlane1.MX;
+			const FVector3f CornerVector = PhiWithNormalInput - Corner;
 			Normal = VectorGetSafeNormal(CornerVector);
 			Phi = VectorSize(CornerVector);
 		}
 		else
 		{
-			Normal = CylinderObj.MPlane1.MNormal * (uniform FReal)(-1.0);
+			Normal = CylinderObj.MPlane1.MNormal * (uniform float)(-1.0);
 			Phi = -Distance1;
 		}
 	}
 	else if (Distance2 < SMALL_NUMBER)
 	{
-		const FVector v = PhiWithNormalInput - (CylinderObj.MPlane2.MNormal * Distance2 + CylinderObj.MPlane2.MX);
+		const FVector3f v = PhiWithNormalInput - (CylinderObj.MPlane2.MNormal * Distance2 + CylinderObj.MPlane2.MX);
 		if (VectorSize(v) > CylinderObj.Radius2)
 		{
-			const FVector Corner = VectorGetSafeNormal(v) * CylinderObj.Radius2 + CylinderObj.MPlane2.MX;
-			const FVector CornerVector = PhiWithNormalInput - Corner;
+			const FVector3f Corner = VectorGetSafeNormal(v) * CylinderObj.Radius2 + CylinderObj.MPlane2.MX;
+			const FVector3f CornerVector = PhiWithNormalInput - Corner;
 			Normal = VectorGetSafeNormal(CornerVector);
 			Phi = VectorSize(CornerVector);
 		}
 		else
 		{
-			Normal = CylinderObj.MPlane2.MNormal * (uniform FReal)(-1.0);
+			Normal = CylinderObj.MPlane2.MNormal * (uniform float)(-1.0);
 			Phi = -Distance2;
 		}
 	}
-	else if (SideDistance < (uniform FReal)(0.0) && TopDistance < -SideDistance)
+	else if (SideDistance < (uniform float)(0.0) && TopDistance < -SideDistance)
 	{
-		Normal = Distance1 < Distance2 ? CylinderObj.MPlane1.MNormal * (uniform FReal)(-1.0) : CylinderObj.MPlane2.MNormal * (uniform FReal)(-1.0);
+		Normal = Distance1 < Distance2 ? CylinderObj.MPlane1.MNormal * (uniform float)(-1.0) : CylinderObj.MPlane2.MNormal * (uniform float)(-1.0);
 		Phi = -TopDistance;
 	}
 	else
@@ -345,7 +338,7 @@ static inline void PhiWithNormalTaperedCylinder(const uniform uint8 *uniform Typ
 	}
 }
 
-static inline void PhiWithNormalUnion(const uniform uint8 *uniform TypedPtr, const FVector &PhiWithNormalInput, FVector &Normal, FReal &Phi, const uniform int SizeofFImplicitObject, const uniform int OffsetofGeometryType, const uniform int OffsetOfMargin)
+static inline void PhiWithNormalUnion(const uniform uint8 *uniform TypedPtr, const FVector3f &PhiWithNormalInput, FVector3f &Normal, float &Phi, const uniform int SizeofFImplicitObject, const uniform int OffsetofGeometryType, const uniform int OffsetOfMargin)
 {
 	// This is expecting a MObjects = [TaperedCylinder, Sphere, Sphere]
 	assert( ((const uniform TArray *uniform)TypedPtr)->ArrayNum == 3 );
@@ -357,8 +350,8 @@ static inline void PhiWithNormalUnion(const uniform uint8 *uniform TypedPtr, con
 	uniform uint8 *uniform Object = MObjects->Data[0].Ptr;
 	uniform uint8 *uniform TypedObjectPtr = Object + SizeofFImplicitObject;
 
-	FVector NextNormal;
-	FReal NextPhi;
+	FVector3f NextNormal;
+	float NextPhi;
 
 	assert( (Object[OffsetofGeometryType] & 0x3F) == TaperedCylinder );
 	PhiWithNormalTaperedCylinder(TypedObjectPtr, PhiWithNormalInput, NextNormal, NextPhi, SizeofFImplicitObject);
@@ -376,7 +369,7 @@ static inline void PhiWithNormalUnion(const uniform uint8 *uniform TypedPtr, con
 	}
 
 	Object = MObjects->Data[1].Ptr;
-	uniform FReal Margin = *((const uniform FReal *uniform)(Object + OffsetOfMargin));
+	uniform float Margin = *((const uniform float *uniform)(Object + OffsetOfMargin));
 	TypedObjectPtr = Object + SizeofFImplicitObject;
 
 	assert( (Object[OffsetofGeometryType] & 0x3F) == Sphere );
@@ -395,7 +388,7 @@ static inline void PhiWithNormalUnion(const uniform uint8 *uniform TypedPtr, con
 	}
 
 	Object = MObjects->Data[2].Ptr;
-	Margin = *((const uniform FReal *uniform)(Object + OffsetOfMargin));
+	Margin = *((const uniform float *uniform)(Object + OffsetOfMargin));
 	TypedObjectPtr = Object + SizeofFImplicitObject;
 
 	assert( (Object[OffsetofGeometryType] & 0x3F) == Sphere );
@@ -608,47 +601,24 @@ static inline bool GetSimplePhiWithNormal(const FVector3f& PhiWithNormalInputFlo
 	{
 	case Sphere:
 		{
-			FVector NormalReal;
-			FReal PhiReal;
-			const FVector PhiWithNormalInput = ConvertVector3fTo3Native(PhiWithNormalInputFloat);
-
-			const uniform FReal Margin = *((const uniform FReal *uniform)(Ptr + OffsetOfMargin));
-			PhiWithNormalSphere(TypedPtr, PhiWithNormalInput, NormalReal, PhiReal, Margin);
-			Phi = (varying float)PhiReal;
-			Normal = ConvertVector3NativeTo3f(NormalReal);
+			const uniform float Margin = *((const uniform float *uniform)(Ptr + OffsetOfMargin));
+			PhiWithNormalSphere(TypedPtr, PhiWithNormalInputFloat, Normal, Phi, Margin);
 			return true;
 		}
 	case Capsule:
 		{
-			FVector NormalReal;
-			FReal PhiReal;
-			const FVector PhiWithNormalInput = ConvertVector3fTo3Native(PhiWithNormalInputFloat);
-
-			const uniform FReal Margin = *((const uniform FReal *uniform)(Ptr + OffsetOfMargin));
-			PhiWithNormalCapsule(TypedPtr, PhiWithNormalInput, NormalReal, PhiReal, Margin);
-			Phi = (varying float)PhiReal;
-			Normal = ConvertVector3NativeTo3f(NormalReal);
+			const uniform float Margin = *((const uniform float *uniform)(Ptr + OffsetOfMargin));
+			PhiWithNormalCapsule(TypedPtr, PhiWithNormalInputFloat, Normal, Phi, Margin);
 			return true;
 		}
 	case Union:
 		{
-			FVector NormalReal;
-			FReal PhiReal;
-			const FVector PhiWithNormalInput = ConvertVector3fTo3Native(PhiWithNormalInputFloat);
-
-			PhiWithNormalUnion(TypedPtr, PhiWithNormalInput, NormalReal, PhiReal, SizeofFImplicitObject, OffsetofGeometryType, OffsetOfMargin);
-			Phi = (varying float)PhiReal;
-			Normal = ConvertVector3NativeTo3f(NormalReal);
+			PhiWithNormalUnion(TypedPtr, PhiWithNormalInputFloat, Normal, Phi, SizeofFImplicitObject, OffsetofGeometryType, OffsetOfMargin);
 			return true;
 		}
 	case TaperedCapsule:
 		{
-			FVector NormalReal;
-			FReal PhiReal;
-			const FVector PhiWithNormalInput = ConvertVector3fTo3Native(PhiWithNormalInputFloat);
-			PhiWithNormalTaperedCapsule(TypedPtr, PhiWithNormalInput, NormalReal, PhiReal);
-			Phi = (varying float)PhiReal;
-			Normal = ConvertVector3NativeTo3f(NormalReal);
+			PhiWithNormalTaperedCapsule(TypedPtr, PhiWithNormalInputFloat, Normal, Phi);
 			return true;
 		}
 	case Convex:
diff --git a/Engine/Source/Runtime/Experimental/Chaos/Private/Chaos/TriangleMeshImplicitObject.cpp b/Engine/Source/Runtime/Experimental/Chaos/Private/Chaos/TriangleMeshImplicitObject.cpp
--- a/Engine/Source/Runtime/Experimental/Chaos/Private/Chaos/TriangleMeshImplicitObject.cpp
+++ b/Engine/Source/Runtime/Experimental/Chaos/Private/Chaos/TriangleMeshImplicitObject.cpp
@@ -366,7 +366,7 @@ struct FTriangleMeshOverlapVisitorNoMTD
 	FTriangleMeshOverlapVisitorNoMTD(const TRigidTransform<FReal, 3>& WorldScaleQueryTM, const QueryGeomType& InQueryGeom, FReal InThickness, const FVec3& InTriMeshScale, const FTriangleMeshImplicitObject* InTriMesh)
 		: QueryGeom(InQueryGeom)
 		, TriMeshScale(InTriMeshScale)
-		, Thickness(InThickness)
+		, Thickness(FRealSingle(InThickness))
 		, TriMesh(InTriMesh)
 		, bFoundIntersection(false)
 	{
@@ -427,7 +427,7 @@ struct FTriangleMeshOverlapVisitorNoMTD
 	VectorRegister4Float Translation;
 	const QueryGeomType& QueryGeom;
 	const FVec3 TriMeshScale;
-	FReal Thickness;
+	FRealSingle Thickness;
 	const FTriangleMeshImplicitObject* TriMesh;
 	bool bFoundIntersection;
 };
@@ -450,12 +450,12 @@ template <>
 struct FTriangleMeshOverlapVisitorNoMTD<FCapsule>
 {
 	FTriangleMeshOverlapVisitorNoMTD(const TRigidTransform<FReal, 3>& WorldScaleQueryTM, const FCapsule& InQueryGeom, FReal InThickness, const FVec3f& InTriMeshScale, const FTriangleMeshImplicitObject* InTriMesh)
-		: Radius(FRealSingle(InQueryGeom.GetRadius() + InThickness))
+		: Radius(InQueryGeom.GetRadiusf() + FRealSingle(InThickness))
 		, TriMesh(InTriMesh)
 		, bFoundIntersection(false)
 	{
-		const FVec3f X1f = WorldScaleQueryTM.TransformPosition(InQueryGeom.GetX1());
-		const FVec3f X2f = WorldScaleQueryTM.TransformPosition(InQueryGeom.GetX2());
+		const FVec3f X1f = WorldScaleQueryTM.TransformPosition(FVec3(InQueryGeom.GetX1f()));
+		const FVec3f X2f = WorldScaleQueryTM.TransformPosition(FVec3(InQueryGeom.GetX2f()));
 		X1 = VectorLoadFloat3(&X1f.X);
 		X2 = VectorLoadFloat3(&X2f.X);
 		TriMeshScale = VectorLoadFloat3(&InTriMeshScale.X);
@@ -505,8 +505,8 @@ struct FTriangleMeshOverlapVisitorNoMTD<TImplicitObjectScaled<FCapsule>>
 		, TriMesh(InTriMesh)
 		, bFoundIntersection(false)
 	{
-		const FVec3f X1f = WorldScaleQueryTM.TransformPosition(InQueryGeom.GetUnscaledObject()->GetX1() * InQueryGeom.GetScale());
-		const FVec3f X2f = WorldScaleQueryTM.TransformPosition(InQueryGeom.GetUnscaledObject()->GetX2() * InQueryGeom.GetScale());
+		const FVec3f X1f = WorldScaleQueryTM.TransformPosition(FVec3(InQueryGeom.GetUnscaledObject()->GetX1f() * InQueryGeom.GetScale()));
+		const FVec3f X2f = WorldScaleQueryTM.TransformPosition(FVec3(InQueryGeom.GetUnscaledObject()->GetX2f() * InQueryGeom.GetScale()));
 		X1 = VectorLoadFloat3(&X1f.X);
 		X2 = VectorLoadFloat3(&X2f.X);
 		TriMeshScale = VectorLoadFloat3(&InTriMeshScale.X);
@@ -551,11 +551,11 @@ template <>
 struct FTriangleMeshOverlapVisitorNoMTD<Chaos::FSphere>
 {
 	FTriangleMeshOverlapVisitorNoMTD(const TRigidTransform<FReal, 3>& WorldScaleQueryTM, const Chaos::FSphere& InQueryGeom, FReal InThickness, const FVec3f& InTriMeshScale, const FTriangleMeshImplicitObject* InTriMesh)
-		: Radius(FRealSingle(InQueryGeom.GetRadius() + InThickness))
+		: Radius(InQueryGeom.GetRadiusf() + FRealSingle(InThickness))
 		, TriMesh(InTriMesh)
 		, bFoundIntersection(false)
 	{
-		const FVec3f Xf = WorldScaleQueryTM.TransformPosition(InQueryGeom.GetCenter());
+		const FVec3f Xf = WorldScaleQueryTM.TransformPosition(FVec3(InQueryGeom.GetCenterf()));
 		X = VectorLoadFloat3(&Xf.X);
 		TriMeshScale = VectorLoadFloat3(&InTriMeshScale.X);
 	}
@@ -603,9 +603,12 @@ struct FTriangleMeshOverlapVisitorNoMTD<TImplicitObjectScaled<Chaos::FSphere>>
 		, TriMesh(InTriMesh)
 		, bFoundIntersection(false)
 	{
-		const FVec3f Xf = WorldScaleQueryTM.TransformPosition(InQueryGeom.GetUnscaledObject()->GetCenter() * InQueryGeom.GetScale());
+		const FVec3f Xf = WorldScaleQueryTM.TransformPosition(FVec3(InQueryGeom.GetUnscaledObject()->GetCenterf()) * InQueryGeom.GetScale());
 		X = VectorLoadFloat3(&Xf.X);
 		TriMeshScale = VectorLoadFloat3(&InTriMeshScale.X);
+		FReal R = InQueryGeom.GetRadius();
+		FRealSingle Rf = InQueryGeom.GetRadiusf();
+		ensure(Rf == R);
 	}
 
 	bool VisitOverlap(int32 TriIdx)
@@ -802,7 +805,7 @@ TArray<int32> FTrimeshBVH::FindAllIntersections (const FAABB3& Intersection) con
 
 FReal FTriangleMeshImplicitObject::PhiWithNormal(const FVec3& x, FVec3& Normal) const
 {
-	TSphere<FReal, 3> TestSphere(x, 0.0f);
+	FSphere TestSphere(x, 0.0f);
 	FRigidTransform3 TestXf(FVec3(0.0), FRotation3::FromIdentity());
 	FVec3 TestLocation = x;
 	FReal Depth = TNumericLimits<FReal>::Max();
@@ -967,7 +970,7 @@ bool FTriangleMeshImplicitObject::GJKContactPointImp(const QueryGeomType& QueryG
 	return LambdaHelper(MElements.GetSmallIndexBuffer());
 }
 
-bool FTriangleMeshImplicitObject::GJKContactPoint(const TSphere<FReal, 3>& QueryGeom, const FRigidTransform3& QueryTM, const FReal WorldThickness, FVec3& Location, FVec3& Normal, FReal& ContactPhi, int32& ContactFaceIndex) const
+bool FTriangleMeshImplicitObject::GJKContactPoint(const FSphere& QueryGeom, const FRigidTransform3& QueryTM, const FReal WorldThickness, FVec3& Location, FVec3& Normal, FReal& ContactPhi, int32& ContactFaceIndex) const
 {
 	return GJKContactPointImp(QueryGeom, QueryTM, WorldThickness, Location, Normal, ContactPhi, ContactFaceIndex);
 }
@@ -987,7 +990,7 @@ bool FTriangleMeshImplicitObject::GJKContactPoint(const FConvex& QueryGeom, cons
 	return GJKContactPointImp(QueryGeom, QueryTM, WorldThickness, Location, Normal, ContactPhi, ContactFaceIndex);
 }
 
-bool FTriangleMeshImplicitObject::GJKContactPoint(const TImplicitObjectScaled< TSphere<FReal, 3> >& QueryGeom, const FRigidTransform3& QueryTM, const FReal WorldThickness, FVec3& Location, FVec3& Normal, FReal& ContactPhi, int32& ContactFaceIndex, FVec3 TriMeshScale) const
+bool FTriangleMeshImplicitObject::GJKContactPoint(const TImplicitObjectScaled< FSphere >& QueryGeom, const FRigidTransform3& QueryTM, const FReal WorldThickness, FVec3& Location, FVec3& Normal, FReal& ContactPhi, int32& ContactFaceIndex, FVec3 TriMeshScale) const
 {
 	return GJKContactPointImp(QueryGeom, QueryTM, WorldThickness, Location, Normal, ContactPhi, ContactFaceIndex, TriMeshScale);
 }
@@ -1177,7 +1180,7 @@ bool FTriangleMeshImplicitObject::OverlapGeomImp(const QueryGeomType& QueryGeom,
 	}
 }
 
-bool FTriangleMeshImplicitObject::OverlapGeom(const TSphere<FReal, 3>& QueryGeom, const FRigidTransform3& QueryTM, const FReal Thickness, FMTDInfo* OutMTD) const
+bool FTriangleMeshImplicitObject::OverlapGeom(const FSphere& QueryGeom, const FRigidTransform3& QueryTM, const FReal Thickness, FMTDInfo* OutMTD) const
 {
 	return OverlapGeomImp<true>(QueryGeom, QueryTM, Thickness, OutMTD);
 }
@@ -1197,7 +1200,7 @@ bool FTriangleMeshImplicitObject::OverlapGeom(const FConvex& QueryGeom, const FR
 	return OverlapGeomImp<false>(QueryGeom, QueryTM, Thickness, OutMTD);
 }
 
-bool FTriangleMeshImplicitObject::OverlapGeom(const TImplicitObjectScaled<TSphere<FReal, 3>>& QueryGeom, const FRigidTransform3& QueryTM, const FReal Thickness, FMTDInfo* OutMTD, FVec3 TriMeshScale) const
+bool FTriangleMeshImplicitObject::OverlapGeom(const TImplicitObjectScaled<FSphere>& QueryGeom, const FRigidTransform3& QueryTM, const FReal Thickness, FMTDInfo* OutMTD, FVec3 TriMeshScale) const
 {
 	return OverlapGeomImp<true>(QueryGeom, QueryTM, Thickness, OutMTD, TriMeshScale);
 }
@@ -1513,7 +1516,7 @@ struct FTriangleMeshSweepVisitorCCD
 		// Checks the query geometry support point along the normal at the sweep end point
 		// NOTE: We are using SupportCore, so we need to add Radius to the early-out distance check. We could just use Support
 		// but this way avoids an extra sqrt. We don't modify IgnorePenetration in the ctor because it gets used below.
-		VectorRegister4Float RadiusV = MakeVectorRegisterFloatFromDouble(VectorSetFloat1(QueryGeom.GetRadius()));
+		VectorRegister4Float RadiusV = VectorSetFloat1(QueryGeom.GetRadiusf());
 		VectorRegister4Float IgnorePenetrationSimd = VectorSubtract(MakeVectorRegisterFloatFromDouble(VectorSetFloat1(IgnorePenetration)), RadiusV);
 		VectorRegister4Float LengthSimd = MakeVectorRegisterFloatFromDouble(VectorSetFloat1(LengthScale * Length));
 		VectorRegister4Float EndPointSimd = VectorAdd(TranslationSimd, VectorMultiply(RayDirSimd, LengthSimd));
@@ -1692,7 +1695,7 @@ bool FTriangleMeshImplicitObject::SweepGeom(const FConvex& QueryGeom, const FRig
 	return SweepGeomImp(QueryGeom, StartTM, Dir, Length, OutTime, OutPosition, OutNormal, OutFaceIndex, OutFaceNormal, Thickness, bComputeMTD, TriMeshScale);
 }
 
-bool FTriangleMeshImplicitObject::SweepGeom(const TImplicitObjectScaled<TSphere<FReal, 3>>& QueryGeom, const FRigidTransform3& StartTM, const FVec3& Dir, const FReal Length, FReal& OutTime, FVec3& OutPosition, FVec3& OutNormal, int32& OutFaceIndex, FVec3& OutFaceNormal, const FReal Thickness, const bool bComputeMTD, FVec3 TriMeshScale) const
+bool FTriangleMeshImplicitObject::SweepGeom(const TImplicitObjectScaled<FSphere>& QueryGeom, const FRigidTransform3& StartTM, const FVec3& Dir, const FReal Length, FReal& OutTime, FVec3& OutPosition, FVec3& OutNormal, int32& OutFaceIndex, FVec3& OutFaceNormal, const FReal Thickness, const bool bComputeMTD, FVec3 TriMeshScale) const
 {
 	return SweepGeomImp(QueryGeom, StartTM, Dir, Length, OutTime, OutPosition, OutNormal, OutFaceIndex, OutFaceNormal, Thickness, bComputeMTD, TriMeshScale);
 }
@@ -1771,7 +1774,7 @@ bool FTriangleMeshImplicitObject::SweepGeomCCDImp(const QueryGeomType& QueryGeom
 	return bHit;
 }
 
-bool FTriangleMeshImplicitObject::SweepGeomCCD(const TSphere<FReal, 3>& QueryGeom, const FRigidTransform3& StartTM, const FVec3& Dir, const FReal Length, const FReal IgnorePenetration, const FReal TargetPenetration, FReal& OutTOI, FReal& OutPhi, FVec3& OutPosition, FVec3& OutNormal, int32& OutFaceIndex, FVec3& OutFaceNormal, const FVec3& TriMeshScale) const
+bool FTriangleMeshImplicitObject::SweepGeomCCD(const FSphere& QueryGeom, const FRigidTransform3& StartTM, const FVec3& Dir, const FReal Length, const FReal IgnorePenetration, const FReal TargetPenetration, FReal& OutTOI, FReal& OutPhi, FVec3& OutPosition, FVec3& OutNormal, int32& OutFaceIndex, FVec3& OutFaceNormal, const FVec3& TriMeshScale) const
 {
 	return SweepGeomCCDImp(QueryGeom, StartTM, Dir, Length, IgnorePenetration, TargetPenetration, OutTOI, OutPhi, OutPosition, OutNormal, OutFaceIndex, OutFaceNormal, TriMeshScale);
 }
@@ -1791,7 +1794,7 @@ bool FTriangleMeshImplicitObject::SweepGeomCCD(const FConvex& QueryGeom, const F
 	return SweepGeomCCDImp(QueryGeom, StartTM, Dir, Length, IgnorePenetration, TargetPenetration, OutTOI, OutPhi, OutPosition, OutNormal, OutFaceIndex, OutFaceNormal, TriMeshScale);
 }
 
-bool FTriangleMeshImplicitObject::SweepGeomCCD(const TImplicitObjectScaled<TSphere<FReal, 3>>& QueryGeom, const FRigidTransform3& StartTM, const FVec3& Dir, const FReal Length, const FReal IgnorePenetration, const FReal TargetPenetration, FReal& OutTOI, FReal& OutPhi, FVec3& OutPosition, FVec3& OutNormal, int32& OutFaceIndex, FVec3& OutFaceNormal, const FVec3& TriMeshScale) const
+bool FTriangleMeshImplicitObject::SweepGeomCCD(const TImplicitObjectScaled<FSphere>& QueryGeom, const FRigidTransform3& StartTM, const FVec3& Dir, const FReal Length, const FReal IgnorePenetration, const FReal TargetPenetration, FReal& OutTOI, FReal& OutPhi, FVec3& OutPosition, FVec3& OutNormal, int32& OutFaceIndex, FVec3& OutFaceNormal, const FVec3& TriMeshScale) const
 {
 	return SweepGeomCCDImp(QueryGeom, StartTM, Dir, Length, IgnorePenetration, TargetPenetration, OutTOI, OutPhi, OutPosition, OutNormal, OutFaceIndex, OutFaceNormal, TriMeshScale);
 }
diff --git a/Engine/Source/Runtime/Experimental/Chaos/Private/GeometryCollection/GeometryCollectionCollisionStructureManager.cpp b/Engine/Source/Runtime/Experimental/Chaos/Private/GeometryCollection/GeometryCollectionCollisionStructureManager.cpp
--- a/Engine/Source/Runtime/Experimental/Chaos/Private/GeometryCollection/GeometryCollectionCollisionStructureManager.cpp
+++ b/Engine/Source/Runtime/Experimental/Chaos/Private/GeometryCollection/GeometryCollectionCollisionStructureManager.cpp
@@ -241,7 +241,7 @@ FCollisionStructureManager::NewImplicitSphere(
 	const float CollisionObjectReduction,
 	const ECollisionTypeEnum CollisionType)
 {
-	Chaos::FImplicitObjectRef Implicit = new Chaos::TSphere<Chaos::FReal, 3>(Chaos::FVec3(0), Radius * (1 - CollisionObjectReduction / 100.f));
+	Chaos::FImplicitObjectRef Implicit = new Chaos::FSphere(Chaos::FVec3(0), Radius * (1 - CollisionObjectReduction / 100.f));
 	UpdateImplicitFlags(Implicit, CollisionType);
 	return Implicit;
 }
@@ -494,7 +494,7 @@ FCollisionStructureManager::CalculateUnitMassInertiaTensor(
 	}
 	else if (ImplicitType == EImplicitTypeEnum::Chaos_Implicit_Sphere)
 	{
-		Tensor = FVector(Chaos::TSphere<Chaos::FReal, 3>::GetInertiaTensor(1.0, Radius, false).M[0][0]);
+		Tensor = FVector(Chaos::FSphere::GetInertiaTensor(1.0, Radius, false).M[0][0]);
 	}
 	ensureMsgf(Tensor.X != 0.f && Tensor.Y != 0.f && Tensor.Z != 0.f, TEXT("Rigid bounds check failure."));
 	return Tensor;
@@ -514,7 +514,7 @@ FCollisionStructureManager::CalculateVolume(
 	}
 	else if (ImplicitType == EImplicitTypeEnum::Chaos_Implicit_Sphere)
 	{
-		Volume = Chaos::TSphere<Chaos::FReal, 3>::GetVolume(Radius);
+		Volume = Chaos::FSphere::GetVolume(Radius);
 	}
 	ensureMsgf(Volume != 0.f, TEXT("Rigid volume check failure."));
 	return Volume;
diff --git a/Engine/Source/Runtime/Experimental/Chaos/Private/GeometryCollection/GeometryCollectionConvexUtility.cpp b/Engine/Source/Runtime/Experimental/Chaos/Private/GeometryCollection/GeometryCollectionConvexUtility.cpp
--- a/Engine/Source/Runtime/Experimental/Chaos/Private/GeometryCollection/GeometryCollectionConvexUtility.cpp
+++ b/Engine/Source/Runtime/Experimental/Chaos/Private/GeometryCollection/GeometryCollectionConvexUtility.cpp
@@ -2766,9 +2766,9 @@ void FGeometryCollectionConvexUtility::ConvertImplicitToConvexArray(const Chaos:
 	{
 	case Chaos::ImplicitObjectType::Sphere:
 	{
-		if (const Chaos::TSphere<Chaos::FReal, 3>*Sphere = InImplicit.template GetObject<Chaos::TSphere<Chaos::FReal, 3>>())
+		if (const Chaos::FSphere*Sphere = InImplicit.template GetObject<Chaos::FSphere>())
 		{
-			const FTransform SphereTransform(FQuat::Identity, Sphere->GetCenter(), FVector(Sphere->GetRadius()));
+			const FTransform SphereTransform(FQuat::Identity, FVector(Sphere->GetCenterf()), FVector(Sphere->GetRadiusf()));
 			const TArray<Chaos::FConvex::FVec3Type> Vertices(IcoSphere_Subdiv1, IcoSphere_Subdiv1_Num);
 
 			FTransformedConvex& TransformedConvex = InOutConvex.Emplace_GetRef();
@@ -2803,14 +2803,14 @@ void FGeometryCollectionConvexUtility::ConvertImplicitToConvexArray(const Chaos:
 	{
 		if (const Chaos::FCapsule* Capsule = InImplicit.template GetObject<Chaos::FCapsule>())
 		{
-			const FTransform CapsuleTransform(FRotationMatrix::MakeFromZ(Capsule->GetAxis()).ToQuat(), Capsule->GetCenter());
+			const FTransform CapsuleTransform(FRotationMatrix::MakeFromZ(FVector3d(Capsule->GetAxisf())).ToQuat(), FVector3d(Capsule->GetCenterf()));
 
-			const FVector3f HalfHeightOffset(0, 0, static_cast<float>(Capsule->GetHeight()) * 0.5f);
+			const FVector3f HalfHeightOffset(0, 0, Capsule->GetHeightf() * 0.5f);
 
 			TArray<Chaos::FConvex::FVec3Type> Vertices;
 			for (int32 VtxIndex = 0; VtxIndex < IcoHemisphere_Subdiv1_Num; VtxIndex++)
 			{
-				const float CapsuleRadius = static_cast<float>(Capsule->GetRadius());
+				const float CapsuleRadius = Capsule->GetRadiusf();
 				Vertices.Add(IcoHemisphere_Subdiv1[VtxIndex] * CapsuleRadius + HalfHeightOffset); // top hemisphere
 				Vertices.Add(IcoHemisphere_Subdiv1[VtxIndex] * -CapsuleRadius - HalfHeightOffset); // bottom hemisphere
 			}
diff --git a/Engine/Source/Runtime/Experimental/Chaos/Public/Chaos/Box.h b/Engine/Source/Runtime/Experimental/Chaos/Public/Chaos/Box.h
--- a/Engine/Source/Runtime/Experimental/Chaos/Public/Chaos/Box.h
+++ b/Engine/Source/Runtime/Experimental/Chaos/Public/Chaos/Box.h
@@ -55,14 +55,14 @@ namespace Chaos
 		    : FImplicitObject(EImplicitObject::FiniteConvex, ImplicitObjectType::Box)
 			, AABB(Other.AABB)
 		{
-			SetMargin(Other.GetMargin());
+			SetMargin(Other.GetMarginf());
 		}
 
 		FORCEINLINE TBox(TBox<T, d>&& Other)
 		    : FImplicitObject(EImplicitObject::FiniteConvex, ImplicitObjectType::Box)
 		    , AABB(MoveTemp(Other.AABB))
 		{
-			SetMargin(Other.GetMargin());
+			SetMargin(Other.GetMarginf());
 		}
 
 		FORCEINLINE TBox<T, d>& operator=(const TBox<T, d>& Other)
@@ -73,7 +73,7 @@ namespace Chaos
 			this->bHasBoundingBox = Other.bHasBoundingBox;
 
 			AABB = Other.AABB;
-			SetMargin(Other.GetMargin());
+			SetMargin(Other.GetMarginf());
 			return *this;
 		}
 
@@ -85,7 +85,7 @@ namespace Chaos
 			this->bHasBoundingBox = Other.bHasBoundingBox;
 
 			AABB = MoveTemp(Other.AABB);
-			SetMargin(Other.GetMargin());
+			SetMargin(Other.GetMarginf());
 			return *this;
 		}
 
@@ -445,7 +445,7 @@ namespace Chaos
 
 		virtual FString ToString() const override
 		{ 
-			return FString::Printf(TEXT("Box: Min: [%s], Max: [%s], Margin: %f"), *Min().ToString(), *Max().ToString(), GetMargin());
+			return FString::Printf(TEXT("Box: Min: [%s], Max: [%s], Margin: %f"), *Min().ToString(), *Max().ToString(), GetMarginf());
 		}
 
 		FORCEINLINE void SerializeImp(FArchive& Ar)
@@ -456,10 +456,9 @@ namespace Chaos
 			Ar.UsingCustomVersion(FReleaseObjectVersion::GUID);
 			if (Ar.CustomVer(FReleaseObjectVersion::GUID) >= FReleaseObjectVersion::MarginAddedToConvexAndBox)
 			{
-				// LWC_TODO : potential precision loss, to be changed when we can serialize FReal as double
 				FRealSingle MarginFloat = (FRealSingle)FImplicitObject::Margin;
 				Ar << MarginFloat;
-				FImplicitObject::Margin = (T)MarginFloat;
+				FImplicitObject::Margin = (FRealSingle)MarginFloat;
 			}
 		}
 
diff --git a/Engine/Source/Runtime/Experimental/Chaos/Public/Chaos/Capsule.h b/Engine/Source/Runtime/Experimental/Chaos/Public/Chaos/Capsule.h
--- a/Engine/Source/Runtime/Experimental/Chaos/Public/Chaos/Capsule.h
+++ b/Engine/Source/Runtime/Experimental/Chaos/Public/Chaos/Capsule.h
@@ -31,21 +31,21 @@ namespace Chaos
 		    : FImplicitObject(EImplicitObject::FiniteConvex, ImplicitObjectType::Capsule)
 			, MSegment(x1, x2)
 		{
-			SetRadius(Radius);
+			SetRadius(static_cast<FRealSingle>(Radius));
 		}
 
 		FCapsule(const FCapsule& Other)
 		    : FImplicitObject(EImplicitObject::FiniteConvex, ImplicitObjectType::Capsule)
 			, MSegment(Other.MSegment)
 		{
-			SetRadius(Other.GetRadius());
+			SetRadius(Other.GetRadiusf());
 		}
 
 		FCapsule(FCapsule&& Other)
 		    : FImplicitObject(EImplicitObject::FiniteConvex, ImplicitObjectType::Capsule)
 			, MSegment(MoveTemp(Other.MSegment))
 		{
-			SetRadius(Other.GetRadius());
+			SetRadius(Other.GetRadiusf());
 		}
 
 		FCapsule& operator=(FCapsule&& InSteal)
@@ -56,7 +56,7 @@ namespace Chaos
 			this->bHasBoundingBox = InSteal.bHasBoundingBox;
 
 			MSegment = MoveTemp(InSteal.MSegment);
-			SetRadius(InSteal.GetRadius());
+			SetRadius(InSteal.GetRadiusf());
 
 			return *this;
 		}
@@ -72,7 +72,13 @@ namespace Chaos
 			return FCapsule(X1, X2, Radius);
 		}
 
+		UE_DEPRECATED(5.6, "Please Use GetRadiusf instead.")
 		virtual FReal GetRadius() const override
+		{
+			return static_cast<FReal>(Margin);
+		}
+
+		virtual FRealSingle GetRadiusf() const override
 		{
 			return Margin;
 		}
@@ -109,27 +115,28 @@ namespace Chaos
 
 		virtual FReal PhiWithNormal(const FVec3& x, FVec3& Normal) const override
 		{
-			auto Dot = FMath::Clamp(FVec3::DotProduct(x - GetX1(), GetAxis()), (FReal)0., GetHeight());
-			FVec3 ProjectedPoint = Dot * GetAxis() + GetX1();
-			Normal = x - ProjectedPoint;
-			return Normal.SafeNormalize() - GetRadius();
+			FVec3 xf = x;
+			FRealSingle Dot = FMath::Clamp(FVec3f::DotProduct(xf - GetX1f(), GetAxisf()), 0.0f, GetHeightf());
+			FVec3f ProjectedPoint = Dot * GetAxisf() + GetX1f();
+			Normal = xf - ProjectedPoint;
+			return Normal.SafeNormalize() - GetRadiusf();
 		}
 
 		virtual const FAABB3 BoundingBox() const override
 		{
-			FAABB3 Box = MSegment.BoundingBox();
-			Box.Thicken(GetRadius());
+			FAABB3 Box = FAABB3(MSegment.BoundingBox());
+			Box.Thicken(GetRadiusf());
 			return Box;
 		}
 
 		virtual FAABB3 CalculateTransformedBounds(const FRigidTransform3& Transform) const override
 		{
-			const FVec3 X1 = Transform.TransformPositionNoScale(MSegment.GetX1());
-			const FVec3 X2 = Transform.TransformPositionNoScale(MSegment.GetX2());
+			const FVec3 X1 = Transform.TransformPositionNoScale(FVec3(MSegment.GetX1()));
+			const FVec3 X2 = Transform.TransformPositionNoScale(FVec3(MSegment.GetX2()));
 			const FVec3 MinSegment = X1.ComponentwiseMin(X2);
 			const FVec3 MaxSegment = X1.ComponentwiseMax(X2);
 
-			const FVec3 RadiusV = FVec3(GetRadius());
+			const FVec3 RadiusV = FVec3(GetRadiusf());
 			return FAABB3(MinSegment - RadiusV, MaxSegment + RadiusV);
 		}
 
@@ -307,8 +314,8 @@ namespace Chaos
 			if (bCheckCaps)
 			{
 				//can avoid some work here, but good enough for now
-				TSphere<FReal, 3> X1Sphere(X1, MRadius);
-				TSphere<FReal, 3> X2Sphere(X2, MRadius);
+				FSphere X1Sphere(X1, MRadius);
+				FSphere X2Sphere(X2, MRadius);
 
 				FReal Time1, Time2;
 				FVec3 Position1, Position2;
@@ -354,15 +361,26 @@ namespace Chaos
 
 		virtual bool Raycast(const FVec3& StartPoint, const FVec3& Dir, const FReal Length, const FReal Thickness, FReal& OutTime, FVec3& OutPosition, FVec3& OutNormal, int32& OutFaceIndex) const override
 		{
-			return RaycastFast(GetRadius(), GetHeight(), GetAxis(), GetX1(), GetX2(), StartPoint, Dir, Length, Thickness, OutTime, OutPosition, OutNormal, OutFaceIndex);
+			return RaycastFast(GetRadiusf(), GetHeightf(), GetAxis(), GetX1f(), GetX2f(), StartPoint, Dir, Length, Thickness, OutTime, OutPosition, OutNormal, OutFaceIndex);
 		}
 
 		FORCEINLINE FVec3 Support(const FVec3& Direction, const FReal Thickness, int32& VertexIndex) const
 		{
-			return MSegment.Support(Direction, GetRadius() + Thickness, VertexIndex);
+			return MSegment.Support(FVec3f(Direction), GetRadiusf() + FRealSingle(Thickness), VertexIndex);
+		}
+
+		FORCEINLINE FVec3f Supportf(const FVec3f& Direction, const FRealSingle Thickness, int32& VertexIndex) const
+		{
+			return MSegment.Support(Direction, GetRadiusf() + Thickness, VertexIndex);
 		}
 
 		FORCEINLINE FVec3 SupportCore(const FVec3& Direction, const FReal InMargin, FReal* OutSupportDelta, int32& VertexIndex) const
+		{
+			// NOTE: Ignores InMargin, assumes Radius
+			return MSegment.SupportCore(FVec3f(Direction), VertexIndex);
+		}
+
+		FORCEINLINE FVec3f SupportCore(const FVec3f& Direction, const FRealSingle InMargin, FRealSingle* OutSupportDelta, int32& VertexIndex) const
 		{
 			// NOTE: Ignores InMargin, assumes Radius
 			return MSegment.SupportCore(Direction, VertexIndex);
@@ -383,7 +401,13 @@ namespace Chaos
 		{
 			// NOTE: Ignores InMargin, assumes Radius
 			// Note: Scaling the direction vector like this, might not seem quite right, but works due to the commutativity of the single dot product that follows
-			return SupportCore(Scale * Direction, GetMargin(), OutSupportDelta, VertexIndex) * Scale;
+			FRealSingle SupportDeltaSingle;
+			const FVec3 Result = SupportCore(Scale * Direction, GetMarginf(), &SupportDeltaSingle, VertexIndex) * Scale;
+			if (OutSupportDelta)
+			{
+				*OutSupportDelta = SupportDeltaSingle;
+			}
+			return Result;
 		}
 
 		FORCEINLINE void SerializeImp(FArchive& Ar)
@@ -393,7 +417,7 @@ namespace Chaos
 			MSegment.Serialize(Ar);
 
 			// Radius is now stored in the base class Margin
-			FRealSingle ArRadius = (FRealSingle)GetRadius(); // LWC_TODO : potential precision loss, to be changed when we can serialize FReal as double
+			FRealSingle ArRadius = GetRadiusf();
 			Ar << ArRadius;
 			SetRadius(ArRadius);
 			
@@ -419,7 +443,7 @@ namespace Chaos
 
 		virtual FString ToString() const override
 		{
-			return FString::Printf(TEXT("Capsule: Height: %f Radius: %f"), GetHeight(), GetRadius());
+			return FString::Printf(TEXT("Capsule: Height: %f Radius: %f"), GetHeightf(), GetRadiusf());
 		}
 		
 		virtual Chaos::FImplicitObjectPtr CopyGeometry() const override
@@ -429,33 +453,45 @@ namespace Chaos
 
 		virtual Chaos::FImplicitObjectPtr CopyGeometryWithScale(const FVec3& Scale) const override
 		{
-			return  Chaos::FImplicitObjectPtr(new FCapsule(GetX1() * Scale, GetX2() * Scale, GetRadius() * Scale.Min()));
+			return  Chaos::FImplicitObjectPtr(new FCapsule(GetX1f() * Scale, GetX2f() * Scale, GetRadiusf() * Scale.Min()));
 		}
 
-		FReal GetHeight() const { return MSegment.GetLength(); }
+		UE_DEPRECATED(5.6, "Use GetHeightf instead")
+		FReal GetHeight() const { return FReal(MSegment.GetLength()); }
+		FRealSingle GetHeightf() const { return MSegment.GetLength(); }
 		/** Returns the bottommost point on the capsule. */
-		const FVec3 GetOrigin() const { return GetX1() + GetAxis() * -GetRadius(); }
+		UE_DEPRECATED(5.6, "Use GetOriginf instead")
+		const FVec3 GetOrigin() const { return GetX1f() + GetAxisf() * -GetRadiusf(); }
+		const FVec3f GetOriginf() const { return GetX1f() + GetAxisf() * -GetRadiusf(); }
 		/** Returns the topmost point on the capsule. */
-		const FVec3 GetInsertion() const { return GetX1() + GetAxis() * (GetHeight() + GetRadius()); }
+		const FVec3f GetInsertion() const { return GetX1f() + GetAxisf() * (GetHeightf() + GetRadiusf()); }
+		UE_DEPRECATED(5.6, "Use GetCenterf instead")
 		FVec3 GetCenter() const { return MSegment.GetCenter(); }
+		FVec3f GetCenterf() const { return MSegment.GetCenter(); }
 		/** Returns the centroid (center of mass). */
-		FVec3 GetCenterOfMass() const { return GetCenter(); }
-		const FVec3& GetAxis() const { return MSegment.GetAxis(); }
-		const FVec3& GetX1() const { return MSegment.GetX1(); }
+		FVec3 GetCenterOfMass() const { return GetCenterf(); }
+		FVec3f GetCenterOfMassf() const { return GetCenterf(); }
+		const FVec3 GetAxis() const { return MSegment.GetAxis(); }
+		UE_DEPRECATED(5.6, "Use GetX1f instead")
+		const FVec3 GetX1() const { return MSegment.GetX1(); }
+		UE_DEPRECATED(5.6, "Use GetX2f instead")
 		FVec3 GetX2() const { return MSegment.GetX2(); }
-		TSegment<FReal> GetSegment() const { return TSegment<FReal>(GetX1(), GetX2()); }
+		const FVec3f GetAxisf() const { return MSegment.GetAxis(); }
+		const FVec3f GetX1f() const { return MSegment.GetX1(); }
+		FVec3f GetX2f() const { return MSegment.GetX2(); }
+		TSegment<FReal> GetSegment() const { return TSegment<FReal>(GetX1f(), GetX2f()); }
 
-		FReal GetArea() const { return GetArea(GetHeight(), GetRadius()); }
+		FReal GetArea() const { return GetArea(GetHeightf(), GetRadiusf()); }
 		static FReal GetArea(const FReal Height, const FReal Radius)
 		{
 			static const FReal PI2 = 2.f * UE_PI;
 			return PI2 * Radius * (Height + 2.f * Radius); 
 		}
 
-		FReal GetVolume() const { return GetVolume(GetHeight(), GetRadius()); }
+		FReal GetVolume() const { return GetVolume(GetHeightf(), GetRadiusf()); }
 		static FReal GetVolume(const FReal Height, const FReal Radius) { static const FReal FourThirds = 4.0f / 3.0f; return UE_PI * Radius * Radius * (Height + FourThirds * Radius); }
 
-		FMatrix33 GetInertiaTensor(const FReal Mass) const { return GetInertiaTensor(Mass, GetHeight(), GetRadius()); }
+		FMatrix33 GetInertiaTensor(const FReal Mass) const { return GetInertiaTensor(Mass, GetHeightf(), GetRadiusf()); }
 		static FMatrix33 GetInertiaTensor(const FReal Mass, const FReal Height, const FReal Radius)
 		{
 			// https://www.wolframalpha.com/input/?i=capsule&assumption=%7B%22C%22,+%22capsule%22%7D+-%3E+%7B%22Solid%22%7D
@@ -482,13 +518,13 @@ namespace Chaos
 
 		virtual uint32 GetTypeHash() const override
 		{
-			return HashCombine(UE::Math::GetTypeHash(GetX1()), UE::Math::GetTypeHash(GetAxis()));
+			return HashCombine(UE::Math::GetTypeHash(GetX1f()), UE::Math::GetTypeHash(GetAxis()));
 		}
 
 		FVec3 GetClosestEdgePosition(int32 PlaneIndexHint, const FVec3& Position) const
 		{
-			FVec3 P0 = GetX1();
-			FVec3 P1 = GetX2();
+			FVec3 P0 = GetX1f();
+			FVec3 P1 = GetX2f();
 			const FVec3 EdgePosition = FMath::ClosestPointOnLine(P0, P1, Position);
 			return EdgePosition;
 		}
@@ -520,9 +556,9 @@ namespace Chaos
 			switch (VertexIndex)
 			{
 			case 0:
-				Result = GetX1(); break;
+				Result = GetX1f(); break;
 			case 1:
-				Result = GetX2(); break;
+				Result = GetX2f(); break;
 			}
 
 			return Result;
@@ -590,28 +626,28 @@ namespace Chaos
 #endif // #if INTEL_ISPC
 
 	private:
-		void SetRadius(FReal InRadius) { SetMargin(InRadius); }
+		void SetRadius(FRealSingle InRadius) { SetMargin(InRadius); }
 
-		TSegment<FReal> MSegment;
+		TSegment<FRealSingle> MSegment;
 	};
 
 	struct FCapsuleSpecializeSamplingHelper
 	{
 		static FORCEINLINE void ComputeSamplePoints(TArray<FVec3>& Points, const FCapsule& Capsule, const int32 NumPoints)
 		{
-			if (NumPoints <= 1 || Capsule.GetRadius() <= UE_SMALL_NUMBER)
+			if (NumPoints <= 1 || Capsule.GetRadiusf() <= UE_SMALL_NUMBER)
 			{
 				const int32 Offset = Points.Num();
-				if (Capsule.GetHeight() <= UE_SMALL_NUMBER)
+				if (Capsule.GetHeightf() <= UE_SMALL_NUMBER)
 				{
 					Points.SetNumUninitialized(Offset + 1);
-					Points[Offset] = Capsule.GetCenter();
+					Points[Offset] = Capsule.GetCenterf();
 				}
 				else
 				{
 					Points.SetNumUninitialized(Offset + 3);
-					Points[0] = Capsule.GetOrigin();
-					Points[1] = Capsule.GetCenter();
+					Points[0] = Capsule.GetOriginf();
+					Points[1] = Capsule.GetCenterf();
 					Points[2] = Capsule.GetInsertion();
 				}
 				return;
@@ -620,7 +656,7 @@ namespace Chaos
 		}
 
 		static FORCEINLINE void ComputeGoldenSpiralPoints(TArray<FVec3>& Points, const FCapsule& Capsule, const int32 NumPoints)
-		{ ComputeGoldenSpiralPoints(Points, Capsule.GetOrigin(), Capsule.GetAxis(), Capsule.GetHeight(), Capsule.GetRadius(), NumPoints); }
+		{ ComputeGoldenSpiralPoints(Points, Capsule.GetOriginf(), Capsule.GetAxisf(), Capsule.GetHeightf(), Capsule.GetRadiusf(), NumPoints); }
 
 		static FORCEINLINE void ComputeGoldenSpiralPoints(
 		    TArray<FVec3>& Points,
@@ -658,10 +694,10 @@ namespace Chaos
 			{
 				// Points vary in Z: [-Radius-HalfHeight, -HalfHeight]
 				TSphereSpecializeSamplingHelper<FReal, 3>::ComputeBottomHalfSemiSphere(
-				    Points, TSphere<FReal, 3>(FVec3(0, 0, -HalfHeight), Radius), NumPointsEndCap, Points.Num());
+				    Points, FSphere(FVec3(0, 0, -HalfHeight), Radius), NumPointsEndCap, Points.Num());
 #if 0
 				{
-					TSphere<FReal, 3> Sphere(FVec3(0, 0, -HalfHeight), Radius);
+					FSphere Sphere(FVec3(0, 0, -HalfHeight), Radius);
 					for(int32 i=Offset; i < Points.Num(); i++)
 					{
 						const FVec3& Pt = Points[i];
@@ -688,10 +724,10 @@ namespace Chaos
 #endif
 				// Points vary in Z: [HalfHeight, HalfHeight+Radius]
 				TSphereSpecializeSamplingHelper<FReal, 3>::ComputeTopHalfSemiSphere(
-				    Points, TSphere<FReal, 3>(FVec3(0, 0, HalfHeight), Radius), NumPointsEndCap, Points.Num());
+				    Points, FSphere(FVec3(0, 0, HalfHeight), Radius), NumPointsEndCap, Points.Num());
 #if 0
 				{
-					TSphere<FReal, 3> Sphere(FVec3(0, 0, HalfHeight), Radius);
+					FSphere Sphere(FVec3(0, 0, HalfHeight), Radius);
 					for(int32 i=TmpOffset; i < Points.Num(); i++)
 					{
 						const FVec3& Pt = Points[i];
@@ -729,8 +765,8 @@ namespace Chaos
 	FORCEINLINE TArray<FVec3> FCapsule::ComputeLocalSamplePoints(const int32 NumPoints) const
 	{
 		TArray<FVec3> Points;
-		const FVec3 Mid = GetCenter();
-		const FCapsule Capsule(GetX1() - Mid, GetX1() + (GetAxis() * GetHeight()) - Mid, GetRadius());
+		const FVec3f Mid = GetCenterf();
+		const FCapsule Capsule(GetX1f() - Mid, GetX1f() + (GetAxisf() * GetHeightf()) - Mid, GetRadiusf());
 		FCapsuleSpecializeSamplingHelper::ComputeSamplePoints(Points, Capsule, NumPoints);
 		return Points;
 	}
diff --git a/Engine/Source/Runtime/Experimental/Chaos/Public/Chaos/CastingUtilities.h b/Engine/Source/Runtime/Experimental/Chaos/Public/Chaos/CastingUtilities.h
--- a/Engine/Source/Runtime/Experimental/Chaos/Public/Chaos/CastingUtilities.h
+++ b/Engine/Source/Runtime/Experimental/Chaos/Public/Chaos/CastingUtilities.h
@@ -25,7 +25,7 @@ namespace Chaos
 			{
 			case ImplicitObjectType::Sphere:
 			{
-				return Func(Geom.template GetObjectChecked<TSphere<FReal, 3>>());
+				return Func(Geom.template GetObjectChecked<FSphere>());
 			}
 			case ImplicitObjectType::Box:
 			{
@@ -41,7 +41,7 @@ namespace Chaos
 			}
 			case ImplicitObjectType::IsScaled | ImplicitObjectType::Sphere:
 			{
-				return Func(Geom.template GetObjectChecked<TImplicitObjectScaled<TSphere<FReal, 3>>>());
+				return Func(Geom.template GetObjectChecked<TImplicitObjectScaled<FSphere>>());
 			}
 			case ImplicitObjectType::IsScaled | ImplicitObjectType::Box:
 			{
@@ -57,7 +57,7 @@ namespace Chaos
 			}
 			case ImplicitObjectType::IsInstanced | ImplicitObjectType::Sphere:
 			{
-				return Func(Geom.template GetObjectChecked< TImplicitObjectInstanced<TSphere<FReal, 3>>>().GetInstancedObject()->template GetObjectChecked<TSphere<FReal, 3>>());
+				return Func(Geom.template GetObjectChecked< TImplicitObjectInstanced<FSphere>>().GetInstancedObject()->template GetObjectChecked<FSphere>());
 			}
 			case ImplicitObjectType::IsInstanced | ImplicitObjectType::Box:
 			{
@@ -79,7 +79,7 @@ namespace Chaos
 			default: 
 				check(false);
 			}
-			return Func(Geom.template GetObjectChecked<TSphere<FReal, 3>>());	//needed for return type
+			return Func(Geom.template GetObjectChecked<FSphere>());	//needed for return type
 		}
 
 		// Call the lambda with concrete shape type. Unwraps shapes contained in Instanced (e.g., Instanced-Sphere will be called with Sphere. Note that this will effectively discard any instance properties like the margin)
@@ -89,17 +89,17 @@ namespace Chaos
 			const EImplicitObjectType Type = Geom.GetType();
 			switch (Type)
 			{
-			case ImplicitObjectType::Sphere: return Func(Geom.template GetObjectChecked<TSphere<FReal, 3>>(), TM);
+			case ImplicitObjectType::Sphere: return Func(Geom.template GetObjectChecked<FSphere>(), TM);
 			case ImplicitObjectType::Box: return Func(Geom.template GetObjectChecked<TBox<FReal, 3>>(), TM);
 			case ImplicitObjectType::Capsule: return Func(Geom.template GetObjectChecked<FCapsule>(), TM);
 			case ImplicitObjectType::Convex: return Func(Geom.template GetObjectChecked<FConvex>(), TM);
-			case ImplicitObjectType::IsScaled | ImplicitObjectType::Sphere: return Func(Geom.template GetObjectChecked<TImplicitObjectScaled<TSphere<FReal, 3>>>(), TM);
+			case ImplicitObjectType::IsScaled | ImplicitObjectType::Sphere: return Func(Geom.template GetObjectChecked<TImplicitObjectScaled<FSphere>>(), TM);
 			case ImplicitObjectType::IsScaled | ImplicitObjectType::Box: return Func(Geom.template GetObjectChecked< TImplicitObjectScaled<TBox<FReal, 3>>>(), TM);
 			case ImplicitObjectType::IsScaled | ImplicitObjectType::Capsule: return Func(Geom.template GetObjectChecked< TImplicitObjectScaled<FCapsule>>(), TM);
 			case ImplicitObjectType::IsScaled | ImplicitObjectType::Convex: return Func(Geom.template GetObjectChecked< TImplicitObjectScaled<FConvex>>(), TM);
 			case ImplicitObjectType::IsInstanced | ImplicitObjectType::Sphere:
 			{
-				return Func(Geom.template GetObjectChecked< TImplicitObjectInstanced<TSphere<FReal, 3>>>().GetInstancedObject()->template GetObjectChecked<TSphere<FReal, 3>>(), TM);
+				return Func(Geom.template GetObjectChecked< TImplicitObjectInstanced<FSphere>>().GetInstancedObject()->template GetObjectChecked<FSphere>(), TM);
 			}
 			case ImplicitObjectType::IsInstanced | ImplicitObjectType::Box:
 			{
@@ -121,7 +121,7 @@ namespace Chaos
 
 			default: check(false);
 			}
-			return Func(Geom.template GetObjectChecked<TSphere<FReal, 3>>(), TM);	//needed for return type
+			return Func(Geom.template GetObjectChecked<FSphere>(), TM);	//needed for return type
 		}
 
 		// Call the lambda with concrete shape type. This version does NOT unwrap shapes contained in Instanced or Scaled.
@@ -131,15 +131,15 @@ namespace Chaos
 			const EImplicitObjectType Type = Geom.GetType();
 			switch (Type)
 			{
-			case ImplicitObjectType::Sphere: return Func(Geom.template GetObjectChecked<TSphere<FReal, 3>>(), TM);
+			case ImplicitObjectType::Sphere: return Func(Geom.template GetObjectChecked<FSphere>(), TM);
 			case ImplicitObjectType::Box: return Func(Geom.template GetObjectChecked<TBox<FReal, 3>>(), TM);
 			case ImplicitObjectType::Capsule: return Func(Geom.template GetObjectChecked<FCapsule>(), TM);
 			case ImplicitObjectType::Convex: return Func(Geom.template GetObjectChecked<FConvex>(), TM);
-			case ImplicitObjectType::IsScaled | ImplicitObjectType::Sphere: return Func(Geom.template GetObjectChecked<TImplicitObjectScaled<TSphere<FReal, 3>>>(), TM);
+			case ImplicitObjectType::IsScaled | ImplicitObjectType::Sphere: return Func(Geom.template GetObjectChecked<TImplicitObjectScaled<FSphere>>(), TM);
 			case ImplicitObjectType::IsScaled | ImplicitObjectType::Box: return Func(Geom.template GetObjectChecked< TImplicitObjectScaled<TBox<FReal, 3>>>(), TM);
 			case ImplicitObjectType::IsScaled | ImplicitObjectType::Capsule: return Func(Geom.template GetObjectChecked< TImplicitObjectScaled<FCapsule>>(), TM);
 			case ImplicitObjectType::IsScaled | ImplicitObjectType::Convex: return Func(Geom.template GetObjectChecked< TImplicitObjectScaled<FConvex>>(), TM);
-			case ImplicitObjectType::IsInstanced | ImplicitObjectType::Sphere: return Func(Geom.template GetObjectChecked< TImplicitObjectInstanced<TSphere<FReal, 3>>>(), TM);
+			case ImplicitObjectType::IsInstanced | ImplicitObjectType::Sphere: return Func(Geom.template GetObjectChecked< TImplicitObjectInstanced<FSphere>>(), TM);
 			case ImplicitObjectType::IsInstanced | ImplicitObjectType::Box: return Func(Geom.template GetObjectChecked< TImplicitObjectInstanced<TBox<FReal, 3>>>(), TM);
 			case ImplicitObjectType::IsInstanced | ImplicitObjectType::Capsule: return Func(Geom.template GetObjectChecked< TImplicitObjectInstanced<FCapsule>>(), TM);
 			case ImplicitObjectType::IsInstanced | ImplicitObjectType::Convex: return Func(Geom.template GetObjectChecked< TImplicitObjectInstanced<FConvex>>(), TM);
@@ -151,7 +151,7 @@ namespace Chaos
 
 			default: check(false);
 			}
-			return Func(Geom.template GetObjectChecked<TSphere<FReal, 3>>(), TM);	//needed for return type
+			return Func(Geom.template GetObjectChecked<FSphere>(), TM);	//needed for return type
 		}
 
 		template <typename Lambda>
@@ -160,15 +160,15 @@ namespace Chaos
 			const EImplicitObjectType Type = Geom.GetType();
 			switch (Type)
 			{
-			case ImplicitObjectType::Sphere: return Func(Geom.template GetObjectChecked<TSphere<FReal, 3>>(), TM0, TM1);
+			case ImplicitObjectType::Sphere: return Func(Geom.template GetObjectChecked<FSphere>(), TM0, TM1);
 			case ImplicitObjectType::Box: return Func(Geom.template GetObjectChecked<TBox<FReal, 3>>(), TM0, TM1);
 			case ImplicitObjectType::Capsule: return Func(Geom.template GetObjectChecked<FCapsule>(), TM0, TM1);
 			case ImplicitObjectType::Convex: return Func(Geom.template GetObjectChecked<FConvex>(), TM0, TM1);
-			case ImplicitObjectType::IsScaled | ImplicitObjectType::Sphere: return Func(Geom.template GetObjectChecked<TImplicitObjectScaled<TSphere<FReal, 3>>>(), TM0, TM1);
+			case ImplicitObjectType::IsScaled | ImplicitObjectType::Sphere: return Func(Geom.template GetObjectChecked<TImplicitObjectScaled<FSphere>>(), TM0, TM1);
 			case ImplicitObjectType::IsScaled | ImplicitObjectType::Box: return Func(Geom.template GetObjectChecked< TImplicitObjectScaled<TBox<FReal, 3>>>(), TM0, TM1);
 			case ImplicitObjectType::IsScaled | ImplicitObjectType::Capsule: return Func(Geom.template GetObjectChecked< TImplicitObjectScaled<FCapsule>>(), TM0, TM1);
 			case ImplicitObjectType::IsScaled | ImplicitObjectType::Convex: return Func(Geom.template GetObjectChecked< TImplicitObjectScaled<FConvex>>(), TM0, TM1);
-			case ImplicitObjectType::IsInstanced | ImplicitObjectType::Sphere: return Func(Geom.template GetObjectChecked< TImplicitObjectInstanced<TSphere<FReal, 3>>>(), TM0, TM1);
+			case ImplicitObjectType::IsInstanced | ImplicitObjectType::Sphere: return Func(Geom.template GetObjectChecked< TImplicitObjectInstanced<FSphere>>(), TM0, TM1);
 			case ImplicitObjectType::IsInstanced | ImplicitObjectType::Box: return Func(Geom.template GetObjectChecked< TImplicitObjectInstanced<TBox<FReal, 3>>>(), TM0, TM1);
 			case ImplicitObjectType::IsInstanced | ImplicitObjectType::Capsule: return Func(Geom.template GetObjectChecked< TImplicitObjectInstanced<FCapsule>>(), TM0, TM1);
 			case ImplicitObjectType::IsInstanced | ImplicitObjectType::Convex: return Func(Geom.template GetObjectChecked< TImplicitObjectInstanced<FConvex>>(), TM0, TM1);
@@ -180,7 +180,7 @@ namespace Chaos
 
 			default: check(false);
 			}
-			return Func(Geom.template GetObjectChecked<TSphere<FReal, 3>>(), TM0, TM1);	//needed for return type
+			return Func(Geom.template GetObjectChecked<FSphere>(), TM0, TM1);	//needed for return type
 		}
 
 		inline
@@ -204,9 +204,9 @@ namespace Chaos
 			{
 				return ImplicitObject->template GetObject<const TImplicitObjectScaled<FCapsule>>()->GetUnscaledObject();
 			}
-			else if (ImplicitType == TImplicitObjectScaled<TSphere<FReal, 3>>::StaticType())
+			else if (ImplicitType == TImplicitObjectScaled<FSphere>::StaticType())
 			{
-				return ImplicitObject->template GetObject<const TImplicitObjectScaled<TSphere<FReal, 3>>>()->GetUnscaledObject();
+				return ImplicitObject->template GetObject<const TImplicitObjectScaled<FSphere>>()->GetUnscaledObject();
 			}
 			else if (ImplicitType == TImplicitObjectScaled<FTriangleMeshImplicitObject>::StaticType())
 			{
@@ -225,9 +225,9 @@ namespace Chaos
 			{
 				return ImplicitObject->template GetObject<const TImplicitObjectInstanced<FCapsule>>()->GetInstancedObject();
 			}
-			else if (ImplicitType == TImplicitObjectInstanced<TSphere<FReal, 3>>::StaticType())
+			else if (ImplicitType == TImplicitObjectInstanced<FSphere>::StaticType())
 			{
-				return ImplicitObject->template GetObject<const TImplicitObjectInstanced<TSphere<FReal, 3>>>()->GetInstancedObject();
+				return ImplicitObject->template GetObject<const TImplicitObjectInstanced<FSphere>>()->GetInstancedObject();
 			}
 			else if (ImplicitType == TImplicitObjectInstanced<FTriangleMeshImplicitObject>::StaticType())
 			{
diff --git a/Engine/Source/Runtime/Experimental/Chaos/Public/Chaos/Collision/CollisionUtil.h b/Engine/Source/Runtime/Experimental/Chaos/Public/Chaos/Collision/CollisionUtil.h
--- a/Engine/Source/Runtime/Experimental/Chaos/Public/Chaos/Collision/CollisionUtil.h
+++ b/Engine/Source/Runtime/Experimental/Chaos/Public/Chaos/Collision/CollisionUtil.h
@@ -64,7 +64,7 @@ namespace Chaos
 			Flags.bEnableOBBCheck1 = bAllowBoundsCheck && !bIsSphere1;					// No OBB test for spheres
 			Flags.bEnableDistanceCheck = bAllowBoundsCheck && bIsSphere0 && bIsSphere1;	// Simple distance test for sphere pairs
 
-			OutDistanceCheckSize = (Flags.bEnableDistanceCheck) ? FRealSingle(Implicit0->GetMargin() + Implicit1->GetMargin()) : 0;	// Sphere-Sphere bounds test
+			OutDistanceCheckSize = (Flags.bEnableDistanceCheck) ? FRealSingle(Implicit0->GetMarginf() + Implicit1->GetMarginf()) : 0;	// Sphere-Sphere bounds test
 
 			// Do not try to reuse manifold points for capsules or spheres (against anything)
 			// NOTE: This can also be disabled for all shape types by the solver (see GenerateCollisionImpl and the Context)
diff --git a/Engine/Source/Runtime/Experimental/Chaos/Public/Chaos/Collision/ContactPointsMiscShapes.h b/Engine/Source/Runtime/Experimental/Chaos/Public/Chaos/Collision/ContactPointsMiscShapes.h
--- a/Engine/Source/Runtime/Experimental/Chaos/Public/Chaos/Collision/ContactPointsMiscShapes.h
+++ b/Engine/Source/Runtime/Experimental/Chaos/Public/Chaos/Collision/ContactPointsMiscShapes.h
@@ -38,23 +38,23 @@ namespace Chaos
 	template <typename GeometryA, typename GeometryB>
 	FContactPoint GJKImplicitContactPoint(const FImplicitObject& A, const FRigidTransform3& ATransform, const GeometryB& B, const FRigidTransform3& BTransform, const FReal CullDistance);
 
-	FContactPoint SphereSphereContactPoint(const TSphere<FReal, 3>& Sphere1, const FRigidTransform3& Sphere1Transform, const TSphere<FReal, 3>& Sphere2, const FRigidTransform3& Sphere2Transform, const FReal CullDistance);
-	FContactPoint SpherePlaneContactPoint(const TSphere<FReal, 3>& Sphere, const FRigidTransform3& SphereTransform, const TPlane<FReal, 3>& Plane, const FRigidTransform3& PlaneTransform);
-	FContactPoint SphereBoxContactPoint(const TSphere<FReal, 3>& Sphere, const FRigidTransform3& SphereTransform, const FImplicitBox3& Box, const FRigidTransform3& BoxTransform);
-	FContactPoint SphereCapsuleContactPoint(const TSphere<FReal, 3>& A, const FRigidTransform3& ATransform, const FCapsule& B, const FRigidTransform3& BTransform);
+	FContactPoint SphereSphereContactPoint(const FSphere& Sphere1, const FRigidTransform3& Sphere1Transform, const FSphere& Sphere2, const FRigidTransform3& Sphere2Transform, const FReal CullDistance);
+	FContactPoint SpherePlaneContactPoint(const FSphere& Sphere, const FRigidTransform3& SphereTransform, const TPlane<FReal, 3>& Plane, const FRigidTransform3& PlaneTransform);
+	FContactPoint SphereBoxContactPoint(const FSphere& Sphere, const FRigidTransform3& SphereTransform, const FImplicitBox3& Box, const FRigidTransform3& BoxTransform);
+	FContactPoint SphereCapsuleContactPoint(const FSphere& A, const FRigidTransform3& ATransform, const FCapsule& B, const FRigidTransform3& BTransform);
 
 	template <typename TriMeshType>
-	FContactPoint SphereTriangleMeshContactPoint(const TSphere<FReal, 3>& A, const FRigidTransform3& ATransform, const TriMeshType& B, const FRigidTransform3& BTransform, const FReal CullDistance);
+	FContactPoint SphereTriangleMeshContactPoint(const FSphere& A, const FRigidTransform3& ATransform, const TriMeshType& B, const FRigidTransform3& BTransform, const FReal CullDistance);
 
 	template<typename TriMeshType>
-	FContactPoint SphereTriangleMeshSweptContactPoint(const TSphere<FReal, 3>& A, const FRigidTransform3& ATransform, const TriMeshType& B, const FRigidTransform3& BStartTransform, const FVec3& Dir, const FReal Length, const FReal IgnorePenetration, const FReal TargetPenetration, FReal& TOI);
+	FContactPoint SphereTriangleMeshSweptContactPoint(const FSphere& A, const FRigidTransform3& ATransform, const TriMeshType& B, const FRigidTransform3& BStartTransform, const FVec3& Dir, const FReal Length, const FReal IgnorePenetration, const FReal TargetPenetration, FReal& TOI);
 
 	FContactPoint BoxHeightFieldContactPoint(const FImplicitBox3& A, const FRigidTransform3& ATransform, const FHeightField& B, const FRigidTransform3& BTransform, const FReal CullDistance);
 
 	template <typename TriMeshType>
 	FContactPoint BoxTriangleMeshContactPoint(const FImplicitBox3& A, const FRigidTransform3& ATransform, const TriMeshType& B, const FRigidTransform3& BTransform, const FReal CullDistance);
 
-	FContactPoint SphereHeightFieldContactPoint(const TSphere<FReal, 3>& A, const FRigidTransform3& ATransform, const FHeightField& B, const FRigidTransform3& BTransform, const FReal CullDistance);
+	FContactPoint SphereHeightFieldContactPoint(const FSphere& A, const FRigidTransform3& ATransform, const FHeightField& B, const FRigidTransform3& BTransform, const FReal CullDistance);
 	FContactPoint CapsuleHeightFieldContactPoint(const FCapsule& A, const FRigidTransform3& ATransform, const FHeightField& B, const FRigidTransform3& BTransform, const FReal CullDistance);
 	template <typename TriMeshType>
 	FContactPoint CapsuleTriangleMeshContactPoint(const FCapsule& A, const FRigidTransform3& ATransform, const TriMeshType& B, const FRigidTransform3& BTransform, const FReal CullDistance);
diff --git a/Engine/Source/Runtime/Experimental/Chaos/Public/Chaos/Collision/ConvexContactPointUtilities.h b/Engine/Source/Runtime/Experimental/Chaos/Public/Chaos/Collision/ConvexContactPointUtilities.h
--- a/Engine/Source/Runtime/Experimental/Chaos/Public/Chaos/Collision/ConvexContactPointUtilities.h
+++ b/Engine/Source/Runtime/Experimental/Chaos/Public/Chaos/Collision/ConvexContactPointUtilities.h
@@ -58,8 +58,8 @@ namespace Chaos::Private
 	// Project a capsule segment onto an axis and return the projected range as well as the vertex indices that bound the range
 	inline void ProjectOntoAxis(const FCapsule& Capsule, const FVec3& AxisN, const FVec3& AxisX, FReal& PMin, FReal& PMax, int32& MinVertexIndex, int32& MaxVertexIndex)
 	{
-		const FVec3 V0 = Capsule.GetX1();
-		const FVec3 V1 = Capsule.GetX2();
+		const FVec3 V0 = Capsule.GetX1f();
+		const FVec3 V1 = Capsule.GetX2f();
 		const FReal D0 = FVec3::DotProduct(V0 - AxisX, AxisN);
 		const FReal D1 = FVec3::DotProduct(V1 - AxisX, AxisN);
 		if (D0 < D1)
diff --git a/Engine/Source/Runtime/Experimental/Chaos/Public/Chaos/Convex.h b/Engine/Source/Runtime/Experimental/Chaos/Public/Chaos/Convex.h
--- a/Engine/Source/Runtime/Experimental/Chaos/Public/Chaos/Convex.h
+++ b/Engine/Source/Runtime/Experimental/Chaos/Public/Chaos/Convex.h
@@ -995,9 +995,9 @@ namespace Chaos
 			Ar.UsingCustomVersion(FReleaseObjectVersion::GUID);
 			if (Ar.CustomVer(FReleaseObjectVersion::GUID) >= FReleaseObjectVersion::MarginAddedToConvexAndBox)
 			{
-				FRealSingle MarginFloat = (FRealSingle)FImplicitObject::Margin; // LWC_TODO : potential precision loss, to be changed when we can serialize FReal as double
+				FRealSingle MarginFloat = (FRealSingle)FImplicitObject::Margin;
 				Ar << MarginFloat;
-				FImplicitObject::Margin = (FReal)MarginFloat;
+				FImplicitObject::Margin = MarginFloat;
 			}
 
 			if (Ar.CustomVer(FReleaseObjectVersion::GUID) >= FReleaseObjectVersion::StructureDataAddedToConvex)
diff --git a/Engine/Source/Runtime/Experimental/Chaos/Public/Chaos/Cylinder.h b/Engine/Source/Runtime/Experimental/Chaos/Public/Chaos/Cylinder.h
--- a/Engine/Source/Runtime/Experimental/Chaos/Public/Chaos/Cylinder.h
+++ b/Engine/Source/Runtime/Experimental/Chaos/Public/Chaos/Cylinder.h
@@ -222,7 +222,7 @@ namespace Chaos
 			// Project End and Center into Plane
 			FVec3 ProjectedEnd = EndPoint - FVec3::DotProduct(EndPoint - StartPoint, MPlane1.Normal()) * MPlane1.Normal();
 			FVec3 ProjectedCenter = MPlane1.X() - FVec3::DotProduct(MPlane1.X() - StartPoint, MPlane1.Normal()) * MPlane1.Normal();
-			auto ProjectedSphere = TSphere<FReal, 3>(ProjectedCenter, MRadius);
+			auto ProjectedSphere = FSphere(ProjectedCenter, MRadius);
 			auto InfiniteCylinderIntersection = ProjectedSphere.FindClosestIntersection(StartPoint, ProjectedEnd, Thickness);
 			if (InfiniteCylinderIntersection.Second)
 			{
diff --git a/Engine/Source/Runtime/Experimental/Chaos/Public/Chaos/Deformable/GaussSeidelWeakConstraints.h b/Engine/Source/Runtime/Experimental/Chaos/Public/Chaos/Deformable/GaussSeidelWeakConstraints.h
--- a/Engine/Source/Runtime/Experimental/Chaos/Public/Chaos/Deformable/GaussSeidelWeakConstraints.h
+++ b/Engine/Source/Runtime/Experimental/Chaos/Public/Chaos/Deformable/GaussSeidelWeakConstraints.h
@@ -604,21 +604,21 @@ namespace Chaos::Softs
 				}
 			}
 
-			TArray<Chaos::TSphere<Chaos::FReal, 3>*> VertexSpherePtrs;
-			TArray<Chaos::TSphere<Chaos::FReal, 3>> VertexSpheres;
+			TArray<Chaos::FSphere*> VertexSpherePtrs;
+			TArray<Chaos::FSphere> VertexSpheres;
 
-			VertexSpheres.Init(Chaos::TSphere<Chaos::FReal, 3>(Chaos::TVec3<Chaos::FReal>(0), DetectRadius), SurfaceElementsPositions.Num());
+			VertexSpheres.Init(Chaos::FSphere(Chaos::TVec3<Chaos::FReal>(0), DetectRadius), SurfaceElementsPositions.Num());
 			VertexSpherePtrs.SetNum(SurfaceElementsPositions.Num());
 
 			for (int32 i = 0; i < SurfaceElementsPositions.Num(); i++)
 			{
 				Chaos::TVec3<Chaos::FReal> SphereCenter(SurfaceElementsPositions[i]);
-				Chaos::TSphere<Chaos::FReal, 3> VertexSphere(SphereCenter, DetectRadius);
-				VertexSpheres[i] = Chaos::TSphere<Chaos::FReal, 3>(SphereCenter, DetectRadius);
+				Chaos::FSphere VertexSphere(SphereCenter, DetectRadius);
+				VertexSpheres[i] = Chaos::FSphere(SphereCenter, DetectRadius);
 				VertexSpherePtrs[i] = &VertexSpheres[i];
 			}
 			Chaos::TBoundingVolumeHierarchy<
-				TArray<Chaos::TSphere<Chaos::FReal, 3>*>,
+				TArray<Chaos::FSphere*>,
 				TArray<int32>,
 				Chaos::FReal,
 				3> VertexBVH(VertexSpherePtrs);
