Store Implicit in single precision

diff --git a/Engine/Source/Runtime/Experimental/Chaos/Public/Chaos/GJK.h b/Engine/Source/Runtime/Experimental/Chaos/Public/Chaos/GJK.h
--- a/Engine/Source/Runtime/Experimental/Chaos/Public/Chaos/GJK.h
+++ b/Engine/Source/Runtime/Experimental/Chaos/Public/Chaos/GJK.h
@@ -50,11 +50,13 @@ namespace Chaos
 
 		template<class T>
 		FGeomGJKHelperSIMD(const T& Geom)
-			: Geometry(&Geom), Margin((FRealSingle)Geom.GetMargin()), Radius((FRealSingle)Geom.GetRadius()), Func(&SupportCoreSimd<T>)
+			: Geometry(&Geom), Margin(Geom.GetMarginf()), Radius(Geom.GetRadiusf()), Func(&SupportCoreSimd<T>)
 		{}
 
 		FRealSingle GetRadius() const { return Radius; }
 		FRealSingle GetMargin() const { return Margin; }
+		FRealSingle GetRadiusf() const { return Radius; }
+		FRealSingle GetMarginf() const { return Margin; }
 
 		VectorRegister4Float operator()(const VectorRegister4Float V) const { return SupportFunction(V); }
 		VectorRegister4Float SupportFunction(const VectorRegister4Float V) const { return Func(Geometry, Margin, V); }
@@ -100,16 +102,16 @@ namespace Chaos
 		// Spheres and Capsules: always use the core shape and full "margin" because it represents the radius
 		// Sphere/Capsule versus OtherShape: no margin on other
 		// OtherShape versus OtherShape: use margin of the smaller shape, zero margin on the other
-		const T RadiusA = A.GetRadius();
-		const T RadiusB = B.GetRadius();
+		const T RadiusA = A.GetRadiusf();
+		const T RadiusB = B.GetRadiusf();
 		const bool bHasRadiusA = RadiusA > 0;
 		const bool bHasRadiusB = RadiusB > 0;
 
 		// The sweep margins if required. Only one can be non-zero (we keep the smaller one)
 		const T SweepMarginScale = 0.05f;
-		const bool bAIsSmallest = A.GetMargin() < B.GetMargin();
-		const T SweepMarginA = (bHasRadiusA || bHasRadiusB) ? 0.0f : (bAIsSmallest ? SweepMarginScale * A.GetMargin() : 0.0f);
-		const T SweepMarginB = (bHasRadiusA || bHasRadiusB) ? 0.0f : (bAIsSmallest ? 0.0f : SweepMarginScale * B.GetMargin());
+		const bool bAIsSmallest = A.GetMarginf() < B.GetMarginf();
+		const T SweepMarginA = (bHasRadiusA || bHasRadiusB) ? 0.0f : (bAIsSmallest ? SweepMarginScale * A.GetMarginf() : 0.0f);
+		const T SweepMarginB = (bHasRadiusA || bHasRadiusB) ? 0.0f : (bAIsSmallest ? 0.0f : SweepMarginScale * B.GetMarginf());
 
 		// Net margin (note: both SweepMargins are zero if either Radius is non-zero, and only one SweepMargin can be non-zero)
 		outMarginA = RadiusA + SweepMarginA;
@@ -133,30 +135,16 @@ namespace Chaos
 		return GJKIntersection(A, B, EpsilonScale, BToATM, InThicknessA, InitialDir);
 	}
 
-	template <typename T>
-	bool GJKIntersection(const FGeomGJKHelperSIMD& RESTRICT A, const FGeomGJKHelperSIMD& RESTRICT B, FReal EpsilonScale, const TRigidTransform<T, 3>& BToATM, const T InThicknessA = 0, const TVector<T, 3>& InitialDir = TVector<T, 3>(1, 0, 0))
-	{
-		const UE::Math::TQuat<T>& RotationDouble = BToATM.GetRotation();
-		VectorRegister4Float RotationSimd = MakeVectorRegisterFloatFromDouble(MakeVectorRegister(RotationDouble.X, RotationDouble.Y, RotationDouble.Z, RotationDouble.W));
-
-		const UE::Math::TVector<T>& TranslationDouble = BToATM.GetTranslation();
-		const VectorRegister4Float TranslationSimd = MakeVectorRegisterFloatFromDouble(MakeVectorRegister(TranslationDouble.X, TranslationDouble.Y, TranslationDouble.Z, 0.0));
-		// Normalize rotation
-		RotationSimd = VectorNormalizeSafe(RotationSimd, GlobalVectorConstants::Float0001);
-
-		const VectorRegister4Float InitialDirSimd = MakeVectorRegisterFloatFromDouble(MakeVectorRegister(InitialDir[0], InitialDir[1], InitialDir[2], 0.0));
-		
-		return GJKIntersectionSimd(A, B, EpsilonScale, TranslationSimd, RotationSimd, InThicknessA, InitialDirSimd);
-	}
+	
 
 	template <typename TGeometryA, typename TGeometryB>
-	bool GJKIntersectionSimd(const TGeometryA& RESTRICT A, const TGeometryB& RESTRICT B, const VectorRegister4Float& Translation, const VectorRegister4Float& Rotation, FReal InThicknessA, const VectorRegister4Float& InitialDir)
+	bool GJKIntersectionSimd(const TGeometryA& RESTRICT A, const TGeometryB& RESTRICT B, const VectorRegister4Float& Translation, const VectorRegister4Float& Rotation, FRealSingle InThicknessA, const VectorRegister4Float& InitialDir)
 	{
 		const FReal EpsilonScale = FMath::Max<FReal>(A.TGeometryA::BoundingBox().Extents().Max(), B.TGeometryB::BoundingBox().Extents().Max());
-		return GJKIntersectionSimd(A, B, EpsilonScale, Translation, Rotation, InThicknessA, InitialDir);
+		return GJKIntersectionSimd(A, B, FRealSingle(EpsilonScale), Translation, Rotation, InThicknessA, InitialDir);
 	}
 
-	inline bool GJKIntersectionSimd(const FGeomGJKHelperSIMD& RESTRICT A, const FGeomGJKHelperSIMD& RESTRICT B, FReal EpsilonScale, const VectorRegister4Float& Translation, const VectorRegister4Float& Rotation, FReal InThicknessA, const VectorRegister4Float& InitialDir)
+	inline bool GJKIntersectionSimd(const FGeomGJKHelperSIMD& RESTRICT A, const FGeomGJKHelperSIMD& RESTRICT B, FRealSingle EpsilonScale, const VectorRegister4Float& Translation, const VectorRegister4Float& Rotation, FRealSingle InThicknessA, const VectorRegister4Float& InitialDir)
 	{
 		VectorRegister4Float V = VectorNegate(InitialDir);
 		V = VectorNormalizeSafe(V, MakeVectorRegisterFloatConstant(-1.f, 0.f, 0.f, 0.f));
@@ -171,15 +159,15 @@ namespace Chaos
 		VectorRegister4Float BarycentricSimd;
 		int32 NumVerts = 0;
 
-		FReal ThicknessA;
-		FReal ThicknessB;
+		FRealSingle ThicknessA;
+		FRealSingle ThicknessB;
 		CalculateQueryMargins(A, B, ThicknessA, ThicknessB);
 		ThicknessA += InThicknessA;
 
-		EpsilonScale = FMath::Min<FReal>(EpsilonScale, 1e5);
+		EpsilonScale = FMath::Min<FRealSingle>(EpsilonScale, 1e5f);
 		
-		const FReal InflationReal = ThicknessA + ThicknessB + 1e-6 * EpsilonScale;
-		const VectorRegister4Float Inflation = VectorSet1(static_cast<FRealSingle>(InflationReal));
+		const FRealSingle InflationReal = ThicknessA + ThicknessB + 1e-6f * EpsilonScale;
+		const VectorRegister4Float Inflation = VectorSet1(InflationReal);
 		const VectorRegister4Float Inflation2 = VectorMultiply(Inflation, Inflation);
 
 		int32 VertexIndexA = INDEX_NONE, VertexIndexB = INDEX_NONE;
@@ -190,9 +178,9 @@ namespace Chaos
 				break;	//if taking too long just stop. This should never happen
 			}
 			const VectorRegister4Float NegVSimd = VectorNegate(V);
-			const VectorRegister4Float SupportASimd = A.SupportFunction(NegVSimd, (FRealSingle)ThicknessA);
+			const VectorRegister4Float SupportASimd = A.SupportFunction(NegVSimd, ThicknessA);
 			const VectorRegister4Float VInBSimd = VectorQuaternionRotateVector(AToBRotation, V);
-			const VectorRegister4Float SupportBLocalSimd = B.SupportFunction(VInBSimd, (FRealSingle)ThicknessB);
+			const VectorRegister4Float SupportBLocalSimd = B.SupportFunction(VInBSimd, ThicknessB);
 			const VectorRegister4Float SupportBSimd = VectorAdd(VectorQuaternionRotateVector(Rotation, SupportBLocalSimd), Translation);
 			const VectorRegister4Float WSimd = VectorSubtract(SupportASimd, SupportBSimd);
 
@@ -224,6 +212,22 @@ namespace Chaos
 		return bNearZero;
 	}
 
+	template <typename T>
+	bool GJKIntersection(const FGeomGJKHelperSIMD& RESTRICT A, const FGeomGJKHelperSIMD& RESTRICT B, FReal EpsilonScale, const TRigidTransform<T, 3>& BToATM, const T InThicknessA = 0, const TVector<T, 3>& InitialDir = TVector<T, 3>(1, 0, 0))
+	{
+		const UE::Math::TQuat<T>& RotationDouble = BToATM.GetRotation();
+		VectorRegister4Float RotationSimd = MakeVectorRegisterFloatFromDouble(MakeVectorRegister(RotationDouble.X, RotationDouble.Y, RotationDouble.Z, RotationDouble.W));
+
+		const UE::Math::TVector<T>& TranslationDouble = BToATM.GetTranslation();
+		const VectorRegister4Float TranslationSimd = MakeVectorRegisterFloatFromDouble(MakeVectorRegister(TranslationDouble.X, TranslationDouble.Y, TranslationDouble.Z, 0.0));
+		// Normalize rotation
+		RotationSimd = VectorNormalizeSafe(RotationSimd, GlobalVectorConstants::Float0001);
+
+		const VectorRegister4Float InitialDirSimd = MakeVectorRegisterFloatFromDouble(MakeVectorRegister(InitialDir[0], InitialDir[1], InitialDir[2], 0.0));
+
+		return GJKIntersectionSimd(A, B, FRealSingle(EpsilonScale), TranslationSimd, RotationSimd, FRealSingle(InThicknessA), InitialDirSimd);
+	}
+
 	/** 
 		Determines if two convex geometries in the same space overlap
 		IMPORTANT: the two convex geometries must be in the same space!
@@ -237,7 +241,7 @@ namespace Chaos
 	*/
 
 	template <typename TGeometryA, typename TGeometryB>
-	bool GJKIntersectionSameSpaceSimd(const TGeometryA& A, const TGeometryB& B, FReal InThicknessA, const VectorRegister4Float& InitialDir)
+	bool GJKIntersectionSameSpaceSimd(const TGeometryA& A, const TGeometryB& B, FRealSingle InThicknessA, const VectorRegister4Float& InitialDir)
 	{
 		VectorRegister4Float V = VectorNegate(InitialDir);
 
@@ -252,15 +256,15 @@ namespace Chaos
 		bool bNearZero = false;
 		int NumIterations = 0;
 		VectorRegister4Float PrevDist2 = GlobalVectorConstants::BigNumber;
-		FReal ThicknessA;
-		FReal ThicknessB;
+		FRealSingle ThicknessA;
+		FRealSingle ThicknessB;
 		CalculateQueryMargins(A, B, ThicknessA, ThicknessB);
 		ThicknessA += InThicknessA;
 
-		const FReal Inflation = ThicknessA + ThicknessB + 1e-3;
-		const FReal Inflation2 = Inflation * Inflation;
+		const FRealSingle Inflation = ThicknessA + ThicknessB + 1e-3f;
+		const FRealSingle Inflation2 = Inflation * Inflation;
 
-		VectorRegister4Float InflationSimd = VectorSet1(static_cast<FRealSingle>(Inflation));
+		VectorRegister4Float InflationSimd = VectorSet1(Inflation);
 		VectorRegister4Float Inflation2Simd = VectorMultiply(InflationSimd, InflationSimd);
 		do
 		{
@@ -422,10 +426,10 @@ namespace Chaos
 
 		const void* Geometry;
 		SupportFunc Func;
-		FReal Margin;
+		FRealSingle Margin;
 
 		template<class T>
-		FGeomGJKHelper(const T& Geom) : Geometry(&Geom), Func(&SupportCore<T>), Margin(Geom.GetMargin()) {}
+		FGeomGJKHelper(const T& Geom) : Geometry(&Geom), Func(&SupportCore<T>), Margin(Geom.GetMarginf()) {}
 
 		FVector SupportFunction(const FVec3& V, int32& VertexIndex) const { return Func(Geometry, V, nullptr, VertexIndex); }
 		FVector SupportFunction(const FVec3& V, FReal* OutSupportDelta, int32& VertexIndex) const { return Func(Geometry, V, OutSupportDelta, VertexIndex); }
@@ -451,13 +455,14 @@ namespace Chaos
 		}
 
 		FReal GetMargin() const { return Margin; }
+		FReal GetMarginf() const { return Margin; }
 
 	private:
 		template<class T>
 		static FVector SupportCore(const void* Geom, const FVec3& Direction, FReal* OutSupportDelta, int32& VertexIndex)
 		{
 			const T* Geometry = (const T*)Geom;
-			return Geometry->SupportCore(Direction, Geometry->GetMargin(), OutSupportDelta, VertexIndex);
+			return Geometry->SupportCore(Direction, Geometry->GetMarginf(), OutSupportDelta, VertexIndex);
 		}
 	};
 
@@ -504,7 +509,7 @@ namespace Chaos
 		// Return the support vertex in A-local space for a vector V in A-local space
 		auto SupportAFunc = [&A, &SupportDeltaA, &VertexIndexA](const TVec3<T>& V)
 		{
-			return A.SupportCore(V, A.GetMargin(), &SupportDeltaA, VertexIndexA);
+			return A.SupportCore(V, A.GetMarginf(), &SupportDeltaA, VertexIndexA);
 		};
 
 		const TRotation<T, 3> AToBRotation = BToATM.GetRotation().Inverse();
@@ -513,7 +518,7 @@ namespace Chaos
 		auto SupportBFunc = [&B, &AToBRotation, &SupportDeltaB, &VertexIndexB](const TVec3<T>& V)
 		{
 			const TVec3<T> VInB = AToBRotation * V;
-			return B.SupportCore(VInB, B.GetMargin(), &SupportDeltaB, VertexIndexB);
+			return B.SupportCore(VInB, B.GetMarginf(), &SupportDeltaB, VertexIndexB);
 		};
 
 		// V and Simplex are in A-local space
@@ -529,8 +534,8 @@ namespace Chaos
 		bool bIsDegenerate = false;				// True if GJK cannot make any more progress
 		bool bIsContact = false;				// True if shapes are within Epsilon or overlapping - GJK cannot provide a solution
 		int32 NumIterations = 0;
-		const T ThicknessA = A.GetMargin();
-		const T ThicknessB = B.GetMargin();
+		const T ThicknessA = A.GetMarginf();
+		const T ThicknessB = B.GetMarginf();
 		const T SeparatedDistance = ThicknessA + ThicknessB + Epsilon;
 		while (!bIsContact && !bIsDegenerate)
 		{
@@ -596,7 +601,7 @@ namespace Chaos
 			auto SupportBInAFunc = [&B, &BToATM, &AToBRotation, &SupportDeltaB, &VertexIndexB](const TVec3<T>& V)
 			{
 				const TVec3<T> VInB = AToBRotation * V;
-				const TVec3<T> SupportBLocal = B.SupportCore(VInB, B.GetMargin(), &SupportDeltaB, VertexIndexB);
+				const TVec3<T> SupportBLocal = B.SupportCore(VInB, B.GetMarginf(), &SupportDeltaB, VertexIndexB);
 				return BToATM.TransformPositionNoScale(SupportBLocal);
 			};
 
@@ -1929,8 +1934,8 @@ namespace Chaos
 		const TVector<T, 3> DirectionInB = BToATM.GetRotation().Inverse() * Direction;
 
 		int32 UnusedVertexIndex = INDEX_NONE;
-		const TVector<T, 3> SupportA = A.SupportCore(Direction, A.GetMargin(), nullptr, UnusedVertexIndex);
-		const TVector<T, 3> SupportBLocal = B.SupportCore(-DirectionInB, B.GetMargin(), nullptr, UnusedVertexIndex);
+		const TVector<T, 3> SupportA = A.SupportCore(Direction, A.GetMarginf(), (T*)nullptr, UnusedVertexIndex);
+		const TVector<T, 3> SupportBLocal = B.SupportCore(-DirectionInB, B.GetMarginf(), (T*)nullptr, UnusedVertexIndex);
 		const TVector<T, 3> SupportB = BToATM.TransformPositionNoScale(SupportBLocal);
 		return SupportA - SupportB;
 	}
diff --git a/Engine/Source/Runtime/Experimental/Chaos/Public/Chaos/GJKShape.h b/Engine/Source/Runtime/Experimental/Chaos/Public/Chaos/GJKShape.h
--- a/Engine/Source/Runtime/Experimental/Chaos/Public/Chaos/GJKShape.h
+++ b/Engine/Source/Runtime/Experimental/Chaos/Public/Chaos/GJKShape.h
@@ -36,6 +36,11 @@ namespace Chaos
 		}
 
 		FReal GetMargin() const
+		{
+			return 0.0;
+		}
+
+		FRealSingle GetMarginf() const
 		{
 			return 0.0f;
 		}
@@ -128,10 +133,15 @@ namespace Chaos
 
 		TGJKCoreShape(const FShapeType& InShape)
 			: Shape(InShape)
-			, Margin(InShape.GetMargin())
+			, Margin(InShape.GetMarginf())
 		{}
 
 		TGJKCoreShape(const FShapeType& InShape, const FReal InMargin)
+			: Shape(InShape)
+			, Margin(FRealSingle(InMargin))
+		{}
+
+		TGJKCoreShape(const FShapeType& InShape, const FRealSingle InMargin)
 			: Shape(InShape)
 			, Margin(InMargin)
 		{}
@@ -146,6 +156,11 @@ namespace Chaos
 			return Margin;
 		}
 
+		FRealSingle GetMarginf() const
+		{
+			return Margin;
+		}
+
 		FVec3 SupportCore(const FVec3 Dir, const FReal InMargin, FReal* OutSupportDelta, int32& VertexIndex) const
 		{
 			return Shape.SupportCore(Dir, InMargin, OutSupportDelta, VertexIndex);
@@ -162,7 +177,7 @@ namespace Chaos
 		}
 
 		const FShapeType& Shape;
-		const FReal Margin;
+		const FRealSingle Margin;
 	};
 
 	// Like TGJKCoreShape but for the second shape if it has a transform relative to the first.
@@ -179,13 +194,13 @@ namespace Chaos
 		TGJKCoreShapeTransformed(const FShapeType& InShape, const FRigidTransform3& InTransform)
 			: Transform(InTransform)
 			, Shape(InShape)
-			, Margin(InShape.GetMargin())
+			, Margin(InShape.GetMarginf())
 		{}
 
 		TGJKCoreShapeTransformed(const FShapeType& InShape, const FReal InMargin, const FRigidTransform3& InTransform)
 			: Transform(InTransform)
 			, Shape(InShape)
-			, Margin(InMargin)
+			, Margin(FRealSingle(InMargin))
 		{}
 
 		const FRigidTransform3& GetTransform() const
@@ -199,6 +214,11 @@ namespace Chaos
 		}
 
 		FReal GetMargin() const
+		{
+			return FReal(Margin);
+		}
+
+		FRealSingle GetMarginf() const
 		{
 			return Margin;
 		}
@@ -222,7 +242,7 @@ namespace Chaos
 
 		FRigidTransform3 Transform;
 		const FShapeType& Shape;
-		const FReal Margin;
+		const FRealSingle Margin;
 	};
 
 	/**
diff --git a/Engine/Source/Runtime/Experimental/Chaos/Public/Chaos/GeometryQueries.h b/Engine/Source/Runtime/Experimental/Chaos/Public/Chaos/GeometryQueries.h
--- a/Engine/Source/Runtime/Experimental/Chaos/Public/Chaos/GeometryQueries.h
+++ b/Engine/Source/Runtime/Experimental/Chaos/Public/Chaos/GeometryQueries.h
@@ -50,9 +50,9 @@ namespace Chaos
 		//if(BType == ImplicitObjectType::Sphere)
 		//{
 		//	const FImplicitObject& BBase = static_cast<const FImplicitObject&>(B);
-		//	const TSphere<FReal, 3>& BSphere = static_cast<const TSphere<FReal, 3>&>(BBase);
+		//	const FSphere& BSphere = static_cast<const FSphere&>(BBase);
 		//	const FVec3 PtInA = BToATM.TransformPositionNoScale(BSphere.GetCenter());
-		//	return A.Overlap(PtInA, Thickness + BSphere.GetRadius());
+		//	return A.Overlap(PtInA, Thickness + BSphere.GetRadiusf());
 		//}
 		////todo: A is a sphere
 		//else 
@@ -189,9 +189,9 @@ namespace Chaos
 		if (bSweepAsRaycast)
 		{
 			const FImplicitObject& BBase = B;
-			const TSphere<FReal, 3>& BSphere = BBase.GetObjectChecked<TSphere<FReal, 3>>();
-			const FVec3 Start = BToATM.TransformPositionNoScale(BSphere.GetCenter());
-			bResult = A.Raycast(Start, LocalDir, Length, Thickness + BSphere.GetRadius(), OutTime, LocalPosition, LocalNormal, OutFaceIndex);
+			const FSphere& BSphere = BBase.GetObjectChecked<FSphere>();
+			const FVec3 Start = BToATM.TransformPositionNoScale(FVec3(BSphere.GetCenterf()));
+			bResult = A.Raycast(Start, LocalDir, Length, Thickness + BSphere.GetRadiusf(), OutTime, LocalPosition, LocalNormal, OutFaceIndex);
 		}
 		//todo: handle case where A is a sphere
 		else if (A.IsConvex())
diff --git a/Engine/Source/Runtime/Experimental/Chaos/Public/Chaos/HeightField.h b/Engine/Source/Runtime/Experimental/Chaos/Public/Chaos/HeightField.h
--- a/Engine/Source/Runtime/Experimental/Chaos/Public/Chaos/HeightField.h
+++ b/Engine/Source/Runtime/Experimental/Chaos/Public/Chaos/HeightField.h
@@ -64,40 +64,40 @@ namespace Chaos
 		CHAOS_API virtual bool Raycast(const FVec3& StartPoint, const FVec3& Dir, const FReal Length, const FReal Thickness, FReal& OutTime, FVec3& OutPosition, FVec3& OutNormal, int32& OutFaceIndex) const override;
 		CHAOS_API virtual bool Overlap(const FVec3& Point, const FReal Thickness) const override;
 		
-		CHAOS_API bool OverlapGeom(const TSphere<FReal, 3>& QueryGeom, const FRigidTransform3& QueryTM, const FReal Thickness, FMTDInfo* OutMTD = nullptr) const;
+		CHAOS_API bool OverlapGeom(const FSphere& QueryGeom, const FRigidTransform3& QueryTM, const FReal Thickness, FMTDInfo* OutMTD = nullptr) const;
 		CHAOS_API bool OverlapGeom(const TBox<FReal, 3>& QueryGeom, const FRigidTransform3& QueryTM, const FReal Thickness, FMTDInfo* OutMTD = nullptr) const;
 		CHAOS_API bool OverlapGeom(const FCapsule& QueryGeom, const FRigidTransform3& QueryTM, const FReal Thickness, FMTDInfo* OutMTD = nullptr) const;
 		CHAOS_API bool OverlapGeom(const FConvex& QueryGeom, const FRigidTransform3& QueryTM, const FReal Thickness, FMTDInfo* OutMTD = nullptr) const;
-		CHAOS_API bool OverlapGeom(const TImplicitObjectScaled<TSphere<FReal, 3>>& QueryGeom, const FRigidTransform3& QueryTM, const FReal Thickness, FMTDInfo* OutMTD = nullptr) const;
+		CHAOS_API bool OverlapGeom(const TImplicitObjectScaled<FSphere>& QueryGeom, const FRigidTransform3& QueryTM, const FReal Thickness, FMTDInfo* OutMTD = nullptr) const;
 		CHAOS_API bool OverlapGeom(const TImplicitObjectScaled<TBox<FReal, 3>>& QueryGeom, const FRigidTransform3& QueryTM, const FReal Thickness, FMTDInfo* OutMTD = nullptr) const;
 		CHAOS_API bool OverlapGeom(const TImplicitObjectScaled<FCapsule>& QueryGeom, const FRigidTransform3& QueryTM, const FReal Thickness, FMTDInfo* OutMTD = nullptr) const;
 		CHAOS_API bool OverlapGeom(const TImplicitObjectScaled<FConvex>& QueryGeom, const FRigidTransform3& QueryTM, const FReal Thickness, FMTDInfo* OutMTD = nullptr) const;
 
-		CHAOS_API bool SweepGeom(const TSphere<FReal, 3>& QueryGeom, const FRigidTransform3& StartTM, const FVec3& Dir, const FReal Length, FReal& OutTime, FVec3& OutPosition, FVec3& OutNormal, int32& OutFaceIndex, FVec3& OutFaceNormal, const FReal Thickness = 0, bool bComputeMTD = false) const;
+		CHAOS_API bool SweepGeom(const FSphere& QueryGeom, const FRigidTransform3& StartTM, const FVec3& Dir, const FReal Length, FReal& OutTime, FVec3& OutPosition, FVec3& OutNormal, int32& OutFaceIndex, FVec3& OutFaceNormal, const FReal Thickness = 0, bool bComputeMTD = false) const;
 		CHAOS_API bool SweepGeom(const TBox<FReal, 3>& QueryGeom, const FRigidTransform3& StartTM, const FVec3& Dir, const FReal Length, FReal& OutTime, FVec3& OutPosition, FVec3& OutNormal, int32& OutFaceIndex, FVec3& OutFaceNormal, const FReal Thickness = 0, bool bComputeMTD = false) const;
 		CHAOS_API bool SweepGeom(const FCapsule& QueryGeom, const FRigidTransform3& StartTM, const FVec3& Dir, const FReal Length, FReal& OutTime, FVec3& OutPosition, FVec3& OutNormal, int32& OutFaceIndex, FVec3& OutFaceNormal, const FReal Thickness = 0, bool bComputeMTD = false) const;
 		CHAOS_API bool SweepGeom(const FConvex& QueryGeom, const FRigidTransform3& StartTM, const FVec3& Dir, const FReal Length, FReal& OutTime, FVec3& OutPosition, FVec3& OutNormal, int32& OutFaceIndex, FVec3& OutFaceNormal, const FReal Thickness = 0, bool bComputeMTD = false) const;
-		CHAOS_API bool SweepGeom(const TImplicitObjectScaled<TSphere<FReal, 3>>& QueryGeom, const FRigidTransform3& StartTM, const FVec3& Dir, const FReal Length, FReal& OutTime, FVec3& OutPosition, FVec3& OutNormal, int32& OutFaceIndex, FVec3& OutFaceNormal, const FReal Thickness = 0, bool bComputeMTD = false) const;
+		CHAOS_API bool SweepGeom(const TImplicitObjectScaled<FSphere>& QueryGeom, const FRigidTransform3& StartTM, const FVec3& Dir, const FReal Length, FReal& OutTime, FVec3& OutPosition, FVec3& OutNormal, int32& OutFaceIndex, FVec3& OutFaceNormal, const FReal Thickness = 0, bool bComputeMTD = false) const;
 		CHAOS_API bool SweepGeom(const TImplicitObjectScaled<TBox<FReal, 3>>& QueryGeom, const FRigidTransform3& StartTM, const FVec3& Dir, const FReal Length, FReal& OutTime, FVec3& OutPosition, FVec3& OutNormal, int32& OutFaceIndex, FVec3& OutFaceNormal, const FReal Thickness = 0, bool bComputeMTD = false) const;
 		CHAOS_API bool SweepGeom(const TImplicitObjectScaled<FCapsule>& QueryGeom, const FRigidTransform3& StartTM, const FVec3& Dir, const FReal Length, FReal& OutTime, FVec3& OutPosition, FVec3& OutNormal, int32& OutFaceIndex, FVec3& OutFaceNormal, const FReal Thickness = 0, bool bComputeMTD = false) const;
 		CHAOS_API bool SweepGeom(const TImplicitObjectScaled<FConvex>& QueryGeom, const FRigidTransform3& StartTM, const FVec3& Dir, const FReal Length, FReal& OutTime, FVec3& OutPosition, FVec3& OutNormal, int32& OutFaceIndex, FVec3& OutFaceNormal, const FReal Thickness = 0, bool bComputeMTD = false) const;
 
 		// Sweep used for CCD. Ignores triangles we penetrate by less than IgnorePenetration, and calculate the TOI for a depth of TargetPenetration. If both are zero, this is equivalent to SweepGeom
-		CHAOS_API bool SweepGeomCCD(const TSphere<FReal, 3>& QueryGeom, const FRigidTransform3& StartTM, const FVec3& Dir, const FReal Length, const FReal IgnorePenetration, const FReal TargetPenetration, FReal& OutTOI, FReal& OutPhi, FVec3& OutPosition, FVec3& OutNormal, int32& OutFaceIndex, FVec3& OutFaceNormal) const;
+		CHAOS_API bool SweepGeomCCD(const FSphere& QueryGeom, const FRigidTransform3& StartTM, const FVec3& Dir, const FReal Length, const FReal IgnorePenetration, const FReal TargetPenetration, FReal& OutTOI, FReal& OutPhi, FVec3& OutPosition, FVec3& OutNormal, int32& OutFaceIndex, FVec3& OutFaceNormal) const;
 		CHAOS_API bool SweepGeomCCD(const TBox<FReal, 3>& QueryGeom, const FRigidTransform3& StartTM, const FVec3& Dir, const FReal Length, const FReal IgnorePenetration, const FReal TargetPenetration, FReal& OutTOI, FReal& OutPhi, FVec3& OutPosition, FVec3& OutNormal, int32& OutFaceIndex, FVec3& OutFaceNormal) const;
 		CHAOS_API bool SweepGeomCCD(const FCapsule& QueryGeom, const FRigidTransform3& StartTM, const FVec3& Dir, const FReal Length, const FReal IgnorePenetration, const FReal TargetPenetration, FReal& OutTOI, FReal& OutPhi, FVec3& OutPosition, FVec3& OutNormal, int32& OutFaceIndex, FVec3& OutFaceNormal) const;
 		CHAOS_API bool SweepGeomCCD(const FConvex& QueryGeom, const FRigidTransform3& StartTM, const FVec3& Dir, const FReal Length, const FReal IgnorePenetration, const FReal TargetPenetration, FReal& OutTOI, FReal& OutPhi, FVec3& OutPosition, FVec3& OutNormal, int32& OutFaceIndex, FVec3& OutFaceNormal) const;
-		CHAOS_API bool SweepGeomCCD(const TImplicitObjectScaled<TSphere<FReal, 3>>& QueryGeom, const FRigidTransform3& StartTM, const FVec3& Dir, const FReal Length, const FReal IgnorePenetration, const FReal TargetPenetration, FReal& OutTOI, FReal& OutPhi, FVec3& OutPosition, FVec3& OutNormal, int32& OutFaceIndex, FVec3& OutFaceNormal) const;
+		CHAOS_API bool SweepGeomCCD(const TImplicitObjectScaled<FSphere>& QueryGeom, const FRigidTransform3& StartTM, const FVec3& Dir, const FReal Length, const FReal IgnorePenetration, const FReal TargetPenetration, FReal& OutTOI, FReal& OutPhi, FVec3& OutPosition, FVec3& OutNormal, int32& OutFaceIndex, FVec3& OutFaceNormal) const;
 		CHAOS_API bool SweepGeomCCD(const TImplicitObjectScaled<TBox<FReal, 3>>& QueryGeom, const FRigidTransform3& StartTM, const FVec3& Dir, const FReal Length, const FReal IgnorePenetration, const FReal TargetPenetration, FReal& OutTOI, FReal& OutPhi, FVec3& OutPosition, FVec3& OutNormal, int32& OutFaceIndex, FVec3& OutFaceNormal) const;
 		CHAOS_API bool SweepGeomCCD(const TImplicitObjectScaled<FCapsule>& QueryGeom, const FRigidTransform3& StartTM, const FVec3& Dir, const FReal Length, const FReal IgnorePenetration, const FReal TargetPenetration, FReal& OutTOI, FReal& OutPhi, FVec3& OutPosition, FVec3& OutNormal, int32& OutFaceIndex, FVec3& OutFaceNormal) const;
 		CHAOS_API bool SweepGeomCCD(const TImplicitObjectScaled<FConvex>& QueryGeom, const FRigidTransform3& StartTM, const FVec3& Dir, const FReal Length, const FReal IgnorePenetration, const FReal TargetPenetration, FReal& OutTOI, FReal& OutPhi, FVec3& OutPosition, FVec3& OutNormal, int32& OutFaceIndex, FVec3& OutFaceNormal) const;
 
 		CHAOS_API bool GJKContactPoint(const TBox<FReal, 3>& QueryGeom, const FRigidTransform3& QueryTM, const FReal Thickness, FVec3& ContactLocation, FVec3& ContactNormal, FReal& ContactPhi, int32& ContactFaceIndex) const;
-		CHAOS_API bool GJKContactPoint(const TSphere<FReal, 3>& QueryGeom, const FRigidTransform3& QueryTM, const FReal Thickness, FVec3& ContactLocation, FVec3& ContactNormal, FReal& ContactPhi, int32& ContactFaceIndex) const;
+		CHAOS_API bool GJKContactPoint(const FSphere& QueryGeom, const FRigidTransform3& QueryTM, const FReal Thickness, FVec3& ContactLocation, FVec3& ContactNormal, FReal& ContactPhi, int32& ContactFaceIndex) const;
 		CHAOS_API bool GJKContactPoint(const FCapsule& QueryGeom, const FRigidTransform3& QueryTM, const FReal Thickness, FVec3& ContactLocation, FVec3& ContactNormal, FReal& ContactPhi, int32& ContactFaceIndex) const;
 		CHAOS_API bool GJKContactPoint(const FConvex& QueryGeom, const FRigidTransform3& QueryTM, const FReal Thickness, FVec3& ContactLocation, FVec3& ContactNormal, FReal& ContactPhi, int32& ContactFaceIndex) const;
 		CHAOS_API bool GJKContactPoint(const TImplicitObjectScaled<TBox<FReal, 3>>& QueryGeom, const FRigidTransform3& QueryTM, const FReal Thickness, FVec3& ContactLocation, FVec3& ContactNormal, FReal& ContactPhi, int32& ContactFaceIndex) const;
-		CHAOS_API bool GJKContactPoint(const TImplicitObjectScaled<TSphere<FReal, 3>>& QueryGeom, const FRigidTransform3& QueryTM, const FReal Thickness, FVec3& ContactLocation, FVec3& ContactNormal, FReal& ContactPhi, int32& ContactFaceIndex) const;
+		CHAOS_API bool GJKContactPoint(const TImplicitObjectScaled<FSphere>& QueryGeom, const FRigidTransform3& QueryTM, const FReal Thickness, FVec3& ContactLocation, FVec3& ContactNormal, FReal& ContactPhi, int32& ContactFaceIndex) const;
 		CHAOS_API bool GJKContactPoint(const TImplicitObjectScaled<FCapsule>& QueryGeom, const FRigidTransform3& QueryTM, const FReal Thickness, FVec3& ContactLocation, FVec3& ContactNormal, FReal& ContactPhi, int32& ContactFaceIndex) const;
 		CHAOS_API bool GJKContactPoint(const TImplicitObjectScaled<FConvex>& QueryGeom, const FRigidTransform3& QueryTM, const FReal Thickness, FVec3& ContactLocation, FVec3& ContactNormal, FReal& ContactPhi, int32& ContactFaceIndex) const;
 
diff --git a/Engine/Source/Runtime/Experimental/Chaos/Public/Chaos/ImplicitObject.h b/Engine/Source/Runtime/Experimental/Chaos/Public/Chaos/ImplicitObject.h
--- a/Engine/Source/Runtime/Experimental/Chaos/Public/Chaos/ImplicitObject.h
+++ b/Engine/Source/Runtime/Experimental/Chaos/Public/Chaos/ImplicitObject.h
@@ -34,6 +34,7 @@ class FImplicitObject;
 using FAABB3 = TAABB<FReal, 3>;
 using FParticles = TParticles<FReal, 3>;
 using FSphere = TSphere<FReal, 3>; // warning: code assumes that FImplicitObjects with type ImplicitObjectType::Sphere are FSpheres, but all TSpheres will think they have ImplicitObjectType::Sphere.
+using FSpheref = TSphere<FRealSingle, 3>;
 
 
 template<class T, int d, bool bSerializable>
@@ -202,8 +203,13 @@ public:
 	
 	void SetCollisionType(EImplicitObjectType InCollisionType) { CollisionType = InCollisionType; }
 	
-	virtual FReal GetRadius() const { return 0.0f; }
+
+	UE_DEPRECATED(5.6, "Please use GetRadiusf instead.")
+	virtual FReal GetRadius() const { return 0.0; }
+	virtual FRealSingle GetRadiusf() const { return 0.0f; }
+	 UE_DEPRECATED(5.6, "Please use GetMarginf instead.")
 	virtual FReal GetMargin() const { return Margin; }
+	virtual FRealSingle GetMarginf() const { return Margin; }
 	
 	static int32 GetOffsetOfMargin() { return offsetof(FImplicitObject, Margin); }
 
@@ -558,10 +564,11 @@ protected:
 	// Not all derived types support a margin, and for some it represents some other
 	// property (the radius of a sphere for example), so the setter should only be exposed 
 	// in a derived class if at all (and it may want to change the size of the core shape as well)
-	void SetMargin(FReal InMargin) { Margin = InMargin; }
+	void SetMargin(FReal InMargin) { Margin = static_cast<FRealSingle>(InMargin); }
+	void SetMargin(FRealSingle InMargin) { Margin = InMargin; }
 
 	// Note: if you change the size of FImplicitObject or add fields, you will likely need to update the calculation of PadBytes below.
-	FReal Margin;
+	FRealSingle Margin;
 	bool bIsConvex;
 	bool bDoCollide;
 	bool bHasBoundingBox;
@@ -580,12 +587,13 @@ protected:
 
 #if DISALLOW_FIMPLICIT_OBJECT_TAIL_PADDING
 	// the following assumptions are made when calculating PadBytes and inserting any explicit padding between fields:
-	static_assert(alignof(bool) <= sizeof(FReal)); // Otherwise, padding would be added between Margin the first bool field
+	static_assert(alignof(FRealSingle) <= alignof(FChaosRefCountedObject)); // Otherwise, padding will be added before Margin
+	static_assert(alignof(bool) <= sizeof(FRealSingle)); // Otherwise, padding would be added between Margin the first bool field
 	static_assert(alignof(EImplicitObjectType) <= 4 * sizeof(bool));  // Otherwise, padding would be added between last bool field and first EImplicitObjectType field
 
-	static constexpr int AlignOfFImplicitObject = (int)FMath::Max(alignof(FReal), FMath::Max(alignof(bool), alignof(EImplicitObjectType)));
+	static constexpr int AlignOfFImplicitObject = (int)FMath::Max(alignof(FChaosRefCountedObject), FMath::Max(alignof(FRealSingle), FMath::Max(alignof(bool), alignof(EImplicitObjectType))));
 	static constexpr int PadBytes =
-		AlignOfFImplicitObject - (sizeof(FReal) + 4 * sizeof(bool) + 2 * sizeof(EImplicitObjectType)) % AlignOfFImplicitObject;
+		AlignOfFImplicitObject - (sizeof(FChaosRefCountedObject) + sizeof(FRealSingle) + 4 * sizeof(bool) + 2 * sizeof(EImplicitObjectType)) % AlignOfFImplicitObject;
 	static_assert(PadBytes > 0);
 	char Pad[PadBytes];
 #endif // DISALLOW_FIMPLICIT_OBJECT_TAIL_PADDING
diff --git a/Engine/Source/Runtime/Experimental/Chaos/Public/Chaos/ImplicitObjectScaled.h b/Engine/Source/Runtime/Experimental/Chaos/Public/Chaos/ImplicitObjectScaled.h
--- a/Engine/Source/Runtime/Experimental/Chaos/Public/Chaos/ImplicitObjectScaled.h
+++ b/Engine/Source/Runtime/Experimental/Chaos/Public/Chaos/ImplicitObjectScaled.h
@@ -39,7 +39,7 @@ public:
 
 
 protected:
-	FReal OuterMargin;
+	FRealSingle OuterMargin;
 };
 
 template <typename TConcrete>
@@ -68,7 +68,7 @@ public:
 		ensure(IsInstanced(MObject->GetType()) == false);	//cannot have an instance of an instance
 		this->bIsConvex = MObject->IsConvex();
 		this->bDoCollide = MObject->GetDoCollide();
-		this->OuterMargin = InMargin;
+		this->OuterMargin = FRealSingle(InMargin);
 		SetMargin(OuterMargin + MObject->GetMargin());
 	}
 
@@ -79,8 +79,8 @@ public:
 		ensure(IsInstanced(MObject->GetType()) == false);	//cannot have an instance of an instance
 		this->bIsConvex = MObject->IsConvex();
 		this->bDoCollide = MObject->GetDoCollide();
-		this->OuterMargin = InMargin;
-		SetMargin(OuterMargin + MObject->GetMargin());
+		this->OuterMargin = FRealSingle(InMargin);
+		SetMargin(OuterMargin + MObject->GetMarginf());
 	}
 
 	TImplicitObjectInstanced(TImplicitObjectInstanced<TConcrete>&& Other)
@@ -110,9 +110,15 @@ public:
 		return MObject.GetReference();
 	}
 
+	UE_DEPRECATED(5.6, "Use GetRadiusf instead")
 	virtual FReal GetRadius() const override
 	{
-		return MObject->GetRadius();
+		return MObject->GetRadiusf();
+	}
+
+	virtual FRealSingle GetRadiusf() const override
+	{
+		return MObject->GetRadiusf();
 	}
 
 	bool GetDoCollide() const
@@ -205,7 +211,7 @@ public:
 
 	virtual FString ToString() const override
 	{
-		return FString::Printf(TEXT("Instanced %s, Margin %f"), *MObject->ToString(), GetMargin());
+		return FString::Printf(TEXT("Instanced %s, Margin %f"), *MObject->ToString(), GetMarginf());
 	}
 
 	static const TImplicitObjectInstanced<TConcrete>& AsInstancedChecked(const FImplicitObject& Obj)
@@ -412,14 +418,16 @@ public:
 		return SignVector.X * SignVector.Y * SignVector.Z;
 	}
 
-	const FVec3& GetScale() const
+	// UE_DEPRECATED(5.6, "Use GetScalef instead")
+	const FVec3 GetScale() const
 	{
-		return MScale;
+		return FVec3(MScale);
 	}
 
-	const FVec3& GetInvScale() const
+	// UE_DEPRECATED(5.6, "Use GetInvScalef instead")
+	const FVec3 GetInvScale() const
 	{
-		return MInvScale;
+		return FVec3(MInvScale);
 	}
 
 	virtual const FAABB3 BoundingBox() const override
@@ -428,9 +436,9 @@ public:
 	}
 
 protected:
-	FVec3 MScale;
-	FVec3 MInvScale;
-	FReal OuterMargin;	//Allows us to inflate the instance before the scale is applied. This is useful when sweeps need to apply a non scale on a geometry with uniform thickness
+	FVec3f MScale;
+	FVec3f MInvScale;
+	FRealSingle OuterMargin;	//Allows us to inflate the instance before the scale is applied. This is useful when sweeps need to apply a non scale on a geometry with uniform thickness
 	FAABB3 MLocalBoundingBox;
 };
 
@@ -473,7 +481,7 @@ public:
 		: FImplicitObjectScaled(EImplicitObject::HasBoundingBox, Object->GetType())
 		, MObject(Object)
 	{
-		InitScaledImplicit(Scale, InMargin);
+		InitScaledImplicit(Scale, FRealSingle(InMargin));
 	}
 	
 	TImplicitObjectScaled(TConcrete* Object, const FVec3& Scale, FReal InMargin = 0)
@@ -485,7 +493,7 @@ public:
 		{
 			Object->MakePersistent();
 		}
-		InitScaledImplicit(Scale, InMargin);
+		InitScaledImplicit(Scale, FRealSingle(InMargin));
 	}
 
 	TImplicitObjectScaled(const TImplicitObjectScaled<TConcrete, bInstanced>& Other) = delete;
@@ -502,7 +510,7 @@ public:
 		this->MInvScale = Other.MInvScale;
 		this->OuterMargin = Other.OuterMargin;
 		this->MLocalBoundingBox = Other.MLocalBoundingBox;
-		SetMargin(Other.GetMargin());
+		SetMargin(Other.GetMarginf());
 	}
 	~TImplicitObjectScaled(){}
 
@@ -582,7 +590,12 @@ public:
 
 	virtual FReal GetRadius() const override
 	{
-		return (MObject->GetRadius() > 0.0f) ? Margin : 0.0f;
+		return (MObject->GetRadiusf() > 0.0f) ? Margin : 0.0f;
+	}
+
+	virtual FRealSingle GetRadiusf() const override
+	{
+		return (MObject->GetRadiusf() > 0.0f) ? Margin : 0.0f;
 	}
 
 	virtual FReal PhiWithNormal(const FVec3& X, FVec3& Normal) const override
@@ -623,7 +636,7 @@ public:
 				else if (NewTime < Length) 
 				{
 					OutPosition = MScale * UnscaledPosition;
-					OutNormal = (MInvScale * UnscaledNormal).GetSafeNormal(TNumericLimits<FReal>::Min());
+					OutNormal = (MInvScale * UnscaledNormal).GetSafeNormal(TNumericLimits<FRealSingle>::Min());
 					OutTime = NewTime;
 					return true;
 				}
@@ -653,7 +666,7 @@ public:
 			TVector<T, d> UnscaledNormal;
 			T UnscaledTime;
 
-			TRigidTransform<T, d> BToATMNoScale(BToATM.GetLocation() * MInvScale, BToATM.GetRotation());
+			TRigidTransform<T, d> BToATMNoScale(BToATM.GetLocation() * FVec3(MInvScale), BToATM.GetRotation());
 		
 			if (MObject->SweepGeom(B, BToATMNoScale, UnscaledDir, UnscaledLength, UnscaledTime, UnscaledPosition, UnscaledNormal, OutFaceIndex, OutFaceNormal, Thickness, bComputeMTD, MScale))
 			{
@@ -692,7 +705,7 @@ public:
 			T TOI;
 			T Phi;
 
-			TRigidTransform<T, d> BToATMNoScale(BToATM.GetLocation() * MInvScale, BToATM.GetRotation());
+			TRigidTransform<T, d> BToATMNoScale(BToATM.GetLocation() * FVec3(MInvScale), BToATM.GetRotation());
 
 			if (MObject->SweepGeomCCD(B, BToATMNoScale, UnscaledDir, UnscaledLength, IgnorePenetration, TargetPenetration, TOI, Phi, UnscaledPosition, UnscaledNormal, OutFaceIndex, OutFaceNormal, MScale))
 			{
@@ -713,7 +726,7 @@ public:
 	template <typename QueryGeomType>
 	bool GJKContactPoint(const QueryGeomType& A, const FRigidTransform3& AToBTM, const FReal Thickness, FVec3& Location, FVec3& Normal, FReal& Penetration, int32& FaceIndex) const
 	{
-		TRigidTransform<T, d> AToBTMNoScale(AToBTM.GetLocation() * MInvScale, AToBTM.GetRotation());
+		TRigidTransform<T, d> AToBTMNoScale(AToBTM.GetLocation() * FVec3(MInvScale), AToBTM.GetRotation());
 
 		// Thickness is a culling distance which cannot be non-uniformly scaled. This gets passed into the ImplicitObject API unscaled so that
 		// if can do the right thing internally if possible, but it makes the API a little confusing. (This is only exists used TriMesh and Heightfield.)
@@ -731,7 +744,7 @@ public:
 		ensure(Thickness == 0 || (FMath::IsNearlyEqual(MScale[0], MScale[1]) && FMath::IsNearlyEqual(MScale[0], MScale[2])));
 
 		auto ScaledB = MakeScaledHelper(B, MInvScale);
-		TRigidTransform<T, d> BToATMNoScale(BToATM.GetLocation() * MInvScale, BToATM.GetRotation());
+		TRigidTransform<T, d> BToATMNoScale(BToATM.GetLocation() * FVec3(MInvScale), BToATM.GetRotation());
 		return MObject->OverlapGeom(ScaledB, BToATMNoScale, Thickness, OutMTD, MScale);
 	}
 
@@ -932,9 +945,15 @@ public:
 		return MakeVectorRegisterFloatFromDouble(MakeVectorRegister(SupportVert.X, SupportVert.Y, SupportVert.Z, 0.0));
 	}
 
+	// UE_DEPRECATED(5.6, "Use SetScale(FVec3f) in single precision instead")
 	void SetScale(const FVec3& Scale)
 	{
-		constexpr FReal MinMagnitude = 1e-6f;
+		SetScale(FVec3f(Scale));
+	}
+
+	void SetScale(const FVec3f& Scale)
+	{
+		constexpr FRealSingle MinMagnitude = 1e-6f;
 		for (int Axis = 0; Axis < 3; ++Axis)
 		{
 			if (!CHAOS_ENSURE(FMath::Abs(Scale[Axis]) >= MinMagnitude))
@@ -948,7 +967,7 @@ public:
 
 			MInvScale[Axis] = 1 / MScale[Axis];
 		}
-		SetMargin(OuterMargin + MScale[0] * MObject->GetMargin());
+		SetMargin(OuterMargin + MScale[0] * MObject->GetMarginf());
 		UpdateBounds();
 	}
 
@@ -1015,7 +1034,7 @@ public:
 
 	virtual FString ToString() const override
 	{
-		return FString::Printf(TEXT("Scaled %s, Scale: [%f, %f, %f], Margin: %f"), *MObject->ToString(), MScale.X, MScale.Y, MScale.Z, GetMargin());
+		return FString::Printf(TEXT("Scaled %s, Scale: [%f, %f, %f], Margin: %f"), *MObject->ToString(), MScale.X, MScale.Y, MScale.Z, GetMarginf());
 	}
 
 private:
@@ -1040,7 +1059,7 @@ private:
 		return new TImplicitObjectScaled<TConcrete, false>(reinterpret_cast<TRefCountPtr<TConcrete>&&>(DuplicatedShape), Obj->MScale, Obj->OuterMargin);
 	}
 
-	void InitScaledImplicit(const FVec3& Scale, FReal InMargin = 0)
+	void InitScaledImplicit(const FVec3& Scale, FRealSingle InMargin = 0)
 	{
 		ensureMsgf((IsScaled(MObject->GetType()) == false), TEXT("Scaled objects should not contain each other."));
 		ensureMsgf((IsInstanced(MObject->GetType()) == false), TEXT("Scaled objects should not contain instances."));
@@ -1129,17 +1148,17 @@ const T* UnwrapImplicit(const FImplicitObject& Implicit, FVec3& OutScale, FReal
 	if (const TImplicitObjectScaled<T>* ScaledImplicit = Implicit.template GetObject<TImplicitObjectScaled<T>>())
 	{
 		OutScale = ScaledImplicit->GetScale();
-		OutMargin = ScaledImplicit->GetMargin();
+		OutMargin = FReal(ScaledImplicit->GetMarginf());
 		return ScaledImplicit->GetUnscaledObject();
 	}
 	else if (const TImplicitObjectInstanced<T>* InstancedImplicit = Implicit.template GetObject<TImplicitObjectInstanced<T>>())
 	{
-		OutMargin = InstancedImplicit->GetMargin();
+		OutMargin = FReal(InstancedImplicit->GetMarginf());
 		return InstancedImplicit->GetInstancedObject();
 	}
 	else if (const T* RawImplicit = Implicit.template GetObject<T>())
 	{
-		OutMargin = RawImplicit->GetMargin();
+		OutMargin = FReal(RawImplicit->GetMarginf());
 		return RawImplicit;
 	}
 	else
diff --git a/Engine/Source/Runtime/Experimental/Chaos/Public/Chaos/ImplicitObjectTransformed.h b/Engine/Source/Runtime/Experimental/Chaos/Public/Chaos/ImplicitObjectTransformed.h
--- a/Engine/Source/Runtime/Experimental/Chaos/Public/Chaos/ImplicitObjectTransformed.h
+++ b/Engine/Source/Runtime/Experimental/Chaos/Public/Chaos/ImplicitObjectTransformed.h
@@ -169,13 +169,13 @@ public:
 	virtual FReal GetMargin() const override
 	{
 		// If the inner shape is quadratic, we have no margin
-		return (MObject->GetRadius() > 0.0f) ? 0.0f : Margin;
+		return (MObject->GetRadiusf() > 0.0f) ? 0.0f : Margin;
 	}
 
 	virtual FReal GetRadius() const override
 	{
 		// If the inner shape is quadratic, so are we
-		return (MObject->GetRadius() > 0.0f) ? Margin : 0.0f;
+		return (MObject->GetRadiusf() > 0.0f) ? Margin : 0.0f;
 	}
 
 	bool GetDoCollide() const
diff --git a/Engine/Source/Runtime/Experimental/Chaos/Public/Chaos/Levelset.h b/Engine/Source/Runtime/Experimental/Chaos/Public/Chaos/Levelset.h
--- a/Engine/Source/Runtime/Experimental/Chaos/Public/Chaos/Levelset.h
+++ b/Engine/Source/Runtime/Experimental/Chaos/Public/Chaos/Levelset.h
@@ -155,12 +155,12 @@ class FLevelSet final : public FImplicitObject
 	}
 
 
-	CHAOS_API bool SweepGeom(const TSphere<FReal, 3>& QueryGeom, const FRigidTransform3& StartTM, const FVec3& Dir, const FReal Length, FReal& OutTime, FVec3& OutPosition, FVec3& OutNormal, int32& OutFaceIndex, const FReal Thickness = 0, const bool bComputeMTD = false) const;
+	CHAOS_API bool SweepGeom(const FSphere& QueryGeom, const FRigidTransform3& StartTM, const FVec3& Dir, const FReal Length, FReal& OutTime, FVec3& OutPosition, FVec3& OutNormal, int32& OutFaceIndex, const FReal Thickness = 0, const bool bComputeMTD = false) const;
 	CHAOS_API bool SweepGeom(const TBox<FReal, 3>& QueryGeom, const FRigidTransform3& StartTM, const FVec3& Dir, const FReal Length, FReal& OutTime, FVec3& OutPosition, FVec3& OutNormal, int32& OutFaceIndex, const FReal Thickness = 0, const bool bComputeMTD = false) const;
 	CHAOS_API bool SweepGeom(const FCapsule& QueryGeom, const FRigidTransform3& StartTM, const FVec3& Dir, const FReal Length, FReal& OutTime, FVec3& OutPosition, FVec3& OutNormal, int32& OutFaceIndex, const FReal Thickness = 0, const bool bComputeMTD = false) const;
 	CHAOS_API bool SweepGeom(const FConvex& QueryGeom, const FRigidTransform3& StartTM, const FVec3& Dir, const FReal Length, FReal& OutTime, FVec3& OutPosition, FVec3& OutNormal, int32& OutFaceIndex, const FReal Thickness = 0, const bool bComputeMTD = false) const;
 
-	CHAOS_API bool SweepGeom(const TImplicitObjectScaled<TSphere<FReal, 3>>& QueryGeom, const FRigidTransform3& StartTM, const FVec3& Dir, const FReal Length, FReal& OutTime, FVec3& OutPosition, FVec3& OutNormal, int32& OutFaceIndex, const FReal Thickness = 0, const bool bComputeMTD = false) const;
+	CHAOS_API bool SweepGeom(const TImplicitObjectScaled<FSphere>& QueryGeom, const FRigidTransform3& StartTM, const FVec3& Dir, const FReal Length, FReal& OutTime, FVec3& OutPosition, FVec3& OutNormal, int32& OutFaceIndex, const FReal Thickness = 0, const bool bComputeMTD = false) const;
 	CHAOS_API bool SweepGeom(const TImplicitObjectScaled<TBox<FReal, 3>>& QueryGeom, const FRigidTransform3& StartTM, const FVec3& Dir, const FReal Length, FReal& OutTime, FVec3& OutPosition, FVec3& OutNormal, int32& OutFaceIndex, const FReal Thickness = 0, const bool bComputeMTD = false) const;
 	CHAOS_API bool SweepGeom(const TImplicitObjectScaled<FCapsule>& QueryGeom, const FRigidTransform3& StartTM, const FVec3& Dir, const FReal Length, FReal& OutTime, FVec3& OutPosition, FVec3& OutNormal, int32& OutFaceIndex, const FReal Thickness = 0, const bool bComputeMTD = false) const;
 	CHAOS_API bool SweepGeom(const TImplicitObjectScaled<FConvex>& QueryGeom, const FRigidTransform3& StartTM, const FVec3& Dir, const FReal Length, FReal& OutTime, FVec3& OutPosition, FVec3& OutNormal, int32& OutFaceIndex, const FReal Thickness = 0, const bool bComputeMTD = false) const;
@@ -168,12 +168,12 @@ class FLevelSet final : public FImplicitObject
 	template<typename QueryGeomType>
 	bool SweepGeomImp(const QueryGeomType& QueryGeom, const FRigidTransform3& StartTM, const FVec3& Dir, const FReal Length, FReal& OutTime, FVec3& OutPosition, FVec3& OutNormal, int32& OutFaceIndex, const FReal Thickness, const bool bComputeMTD) const;
 
-	CHAOS_API bool OverlapGeom(const TSphere<FReal, 3>& QueryGeom, const FRigidTransform3& QueryTM, const FReal Thickness, FMTDInfo* OutMTD = nullptr) const;
+	CHAOS_API bool OverlapGeom(const FSphere& QueryGeom, const FRigidTransform3& QueryTM, const FReal Thickness, FMTDInfo* OutMTD = nullptr) const;
 	CHAOS_API bool OverlapGeom(const TBox<FReal, 3>& QueryGeom, const FRigidTransform3& QueryTM, const FReal Thickness, FMTDInfo* OutMTD = nullptr) const;
 	CHAOS_API bool OverlapGeom(const FCapsule& QueryGeom, const FRigidTransform3& QueryTM, const FReal Thickness, FMTDInfo* OutMTD = nullptr) const;
 	CHAOS_API bool OverlapGeom(const FConvex& QueryGeom, const FRigidTransform3& QueryTM, const FReal Thickness, FMTDInfo* OutMTD = nullptr) const;
 
-	CHAOS_API bool OverlapGeom(const TImplicitObjectScaled<TSphere<FReal, 3>>& QueryGeom, const FRigidTransform3& QueryTM, const FReal Thickness, FMTDInfo* OutMTD = nullptr) const;
+	CHAOS_API bool OverlapGeom(const TImplicitObjectScaled<FSphere>& QueryGeom, const FRigidTransform3& QueryTM, const FReal Thickness, FMTDInfo* OutMTD = nullptr) const;
 	CHAOS_API bool OverlapGeom(const TImplicitObjectScaled<TBox<FReal, 3>>& QueryGeom, const FRigidTransform3& QueryTM, const FReal Thickness, FMTDInfo* OutMTD = nullptr) const;
 	CHAOS_API bool OverlapGeom(const TImplicitObjectScaled<FCapsule>& QueryGeom, const FRigidTransform3& QueryTM, const FReal Thickness, FMTDInfo* OutMTD = nullptr) const;
 	CHAOS_API bool OverlapGeom(const TImplicitObjectScaled<FConvex>& QueryGeom, const FRigidTransform3& QueryTM, const FReal Thickness, FMTDInfo* OutMTD = nullptr) const;
diff --git a/Engine/Source/Runtime/Experimental/Chaos/Public/Chaos/PBDCollisionSphereConstraints.h b/Engine/Source/Runtime/Experimental/Chaos/Public/Chaos/PBDCollisionSphereConstraints.h
--- a/Engine/Source/Runtime/Experimental/Chaos/Public/Chaos/PBDCollisionSphereConstraints.h
+++ b/Engine/Source/Runtime/Experimental/Chaos/Public/Chaos/PBDCollisionSphereConstraints.h
@@ -23,7 +23,7 @@ class FPBDCollisionSphereConstraints : public FPerParticleRule
 	{
 		for (int32 i = 0; i < (int32)InParticles.Size(); ++i)
 		{
-			MObjects.Add(Chaos::FImplicitObjectPtr(new TSphere<FReal, 3>(InParticles.P(i), Height)));
+			MObjects.Add(Chaos::FImplicitObjectPtr(new FSphere(InParticles.P(i), Height)));
 		}
 		TBoundingVolumeHierarchy<TArray<Chaos::FImplicitObjectPtr>, TArray<int32>> Hierarchy(MObjects);
 		FCriticalSection CriticalSection;
diff --git a/Engine/Source/Runtime/Experimental/Chaos/Public/Chaos/Segment.h b/Engine/Source/Runtime/Experimental/Chaos/Public/Chaos/Segment.h
--- a/Engine/Source/Runtime/Experimental/Chaos/Public/Chaos/Segment.h
+++ b/Engine/Source/Runtime/Experimental/Chaos/Public/Chaos/Segment.h
@@ -31,7 +31,7 @@ namespace Chaos
 
 		FORCEINLINE const TVec3<T>& GetAxis() const { return MAxis; }
 
-		FORCEINLINE FReal GetLength() const { return MLength; }
+		FORCEINLINE T GetLength() const { return MLength; }
 
 		TVec3<T> Support(const TVec3<T>& Direction, const T Thickness, int32& VertexIndex) const
 		{
@@ -63,7 +63,7 @@ namespace Chaos
 		{
 			Ar << MPoint << MAxis;
 			
-			FRealSingle LengthFloat = (FRealSingle)MLength; // LWC_TODO : potential precision loss, to be changed when we can serialize FReal as double
+			FRealSingle LengthFloat = (FRealSingle)MLength;
 			Ar << LengthFloat;
 			MLength = (T)LengthFloat;
 		}
diff --git a/Engine/Source/Runtime/Experimental/Chaos/Public/Chaos/Sphere.h b/Engine/Source/Runtime/Experimental/Chaos/Public/Chaos/Sphere.h
--- a/Engine/Source/Runtime/Experimental/Chaos/Public/Chaos/Sphere.h
+++ b/Engine/Source/Runtime/Experimental/Chaos/Public/Chaos/Sphere.h
@@ -33,14 +33,14 @@ namespace Chaos
 		    : FImplicitObject(EImplicitObject::IsConvex | EImplicitObject::HasBoundingBox, ImplicitObjectType::Sphere)
 		    , Center(InCenter)
 		{
-			SetRadius(InRadius);
+			SetRadius(FRealSingle(InRadius));
 		}
 
 		TSphere(const TSphere<T, d>& Other)
 		    : FImplicitObject(EImplicitObject::IsConvex | EImplicitObject::HasBoundingBox, ImplicitObjectType::Sphere)
 		    , Center(Other.Center)
 		{
-			SetRadius(Other.GetRadius());
+			SetRadius(Other.GetRadiusf());
 		}
 
 		TSphere(TSphere<T, d>&& Other)
@@ -57,7 +57,7 @@ namespace Chaos
 			this->bDoCollide = InSteal.bDoCollide;
 			this->bHasBoundingBox = InSteal.bHasBoundingBox;
 			Center = MoveTemp(InSteal.Center);
-			SetRadius(InSteal.GetRadius());
+			SetRadius(InSteal.GetRadiusf());
 
 			return *this;
 		}
@@ -69,27 +69,33 @@ namespace Chaos
 			return ImplicitObjectType::Sphere; 
 		}
 
+		UE_DEPRECATED(5.6, "Use GetRadiusf instead.")
 		virtual FReal GetRadius() const override
 		{
 			return Margin;
 		}
 
+		virtual FRealSingle GetRadiusf() const override
+		{
+			return Margin;
+		}
+
 		virtual T PhiWithNormal(const TVector<T, d>& InSamplePoint, TVector<T, d>& OutNormal) const override
 		{
 			OutNormal = InSamplePoint - Center;
-			return OutNormal.SafeNormalize() - GetRadius();
+			return OutNormal.SafeNormalize() - GetRadiusf();
 		}
 
 		bool Intersects(const TSphere<T, d>& Other) const
 		{
-			T CenterDistance = FVector::DistSquared(Other.GetCenter(), GetCenter());
-			T RadialSum = Other.GetRadius() + GetRadius();
+			FRealSingle CenterDistance = FVec3f::DistSquared(Other.GetCenterf(), GetCenterf());
+			FRealSingle RadialSum = Other.GetRadiusf() + GetRadiusf();
 			return RadialSum >= CenterDistance;
 		}
 
 		TVector<T, d> FindClosestPoint(const TVector<T, d>& StartPoint, const T Thickness = (T)0) const
 		{
-			TVector<T, d> Result = Center + (StartPoint - Center).GetSafeNormal() * (GetRadius() + Thickness);
+			TVector<T, d> Result = Center + (StartPoint - Center).GetSafeNormal() * (GetRadiusf() + Thickness);
 			return Result;
 		}
 
@@ -99,7 +105,7 @@ namespace Chaos
 			ensure(Length > 0);
 			OutFaceIndex = INDEX_NONE;
 
-			const T EffectiveRadius = Thickness + GetRadius();
+			const T EffectiveRadius = Thickness + GetRadiusf();
 			const T EffectiveRadius2 = EffectiveRadius * EffectiveRadius;
 			const TVector<T, d> Offset = Center - StartPoint;
 			const T OffsetSize2 = Offset.SizeSquared();
@@ -149,7 +155,7 @@ namespace Chaos
 			Direction = Direction.GetSafeNormal();
 			TVector<T, d> SphereToStart = StartPoint - Center;
 			T DistanceProjected = TVector<T, d>::DotProduct(Direction, SphereToStart);
-			T EffectiveRadius = GetRadius() + Thickness;
+			T EffectiveRadius = GetRadiusf() + Thickness;
 			T UnderRoot = DistanceProjected * DistanceProjected - SphereToStart.SizeSquared() + EffectiveRadius * EffectiveRadius;
 			if (UnderRoot < 0)
 			{
@@ -196,10 +202,10 @@ namespace Chaos
 			}
 			const TVector<T,d> Normalized = Direction / sqrt(SizeSqr);
 
-			return Center + Normalized * (GetRadius() + Thickness);
+			return Center + Normalized * (GetRadiusf() + Thickness);
 		}
 
-		FORCEINLINE const TVector<T, d>& SupportCore(const TVector<T, d>& Direction, const FReal InMargin, FReal* OutSupportDelta, int32& VertexIndex) const
+		FORCEINLINE const TVector<T, d> SupportCore(const TVector<T, d>& Direction, const FReal InMargin, FReal* OutSupportDelta, int32& VertexIndex) const
 		{
 			VertexIndex = 0;
 			// Note: ignores InMargin, assumed Radius
@@ -208,7 +214,7 @@ namespace Chaos
 
 		FORCEINLINE VectorRegister4Float SupportCoreSimd(const VectorRegister4Float& Direction, const FReal InMargin) const
 		{
-			return MakeVectorRegisterFloatFromDouble(MakeVectorRegister(Center[0], Center[1], Center[2], 0.0));
+			return MakeVectorRegisterFloat(Center[0], Center[1], Center[2], 0.0f);
 		}
 		FORCEINLINE TVector<T, d> SupportCoreScaled(const TVector<T, d>& Direction, const FReal InMargin, const TVector<T, d>& Scale, FReal* OutSupportDelta, int32& VertexIndex) const
 		{
@@ -219,19 +225,19 @@ namespace Chaos
 
 		virtual const TAABB<T, d> BoundingBox() const override
 		{
-			return TAABB<T,d>(Center - TVector<T,d>(GetRadius()),Center + TVector<T,d>(GetRadius()));
+			return TAABB<T,d>(Center - FVec3f(GetRadiusf()),Center + FVec3f(GetRadiusf()));
 		}
 
 		virtual FAABB3 CalculateTransformedBounds(const FRigidTransform3& Transform) const override
 		{
-			const FVec3 TransformedCenter = Transform.TransformPosition(Center);
-			const FVec3 Extents = FVec3(GetRadius());
+			const FVec3 TransformedCenter = Transform.TransformPosition(FVec3(Center));
+			const FVec3 Extents = FVec3(GetRadiusf());
 			return FAABB3(TransformedCenter - Extents, TransformedCenter + Extents);
 		}
 
 		T GetArea() const 
 		{ 
-			return GetArea(GetRadius());
+			return GetArea(GetRadiusf());
 		}
 		
 		static T GetArea(const T InRadius)
@@ -243,7 +249,7 @@ namespace Chaos
 
 		T GetVolume() const 
 		{
-			return GetVolume(GetRadius());
+			return GetVolume(GetRadiusf());
 		}
 
 		static T GetVolume(const T InRadius)
@@ -253,19 +259,30 @@ namespace Chaos
 			return FourThirdsPI * InRadius * InRadius * InRadius;
 		}
 
-		const TVector<T, d>& GetCenter() const 
+		UE_DEPRECATED(5.6, "Use GetCenterf instead")
+		const TVector<T, d> GetCenter() const 
 		{ 
 			return Center; 
 		}
 
-		const TVector<T, d>& GetCenterOfMass() const 
+		const FVec3f& GetCenterf() const
+		{
+			return Center;
+		}
+
+		const TVector<T, d> GetCenterOfMass() const 
 		{ 
 			return Center; 
 		}
+		
+		const FVec3f& GetCenterOfMassf() const
+		{
+			return Center;
+		}
 
 		virtual FString ToString() const override
 		{
-			return FString::Printf(TEXT("Sphere: Center:%s, Radius:%f"), *Center.ToString(), GetRadius());
+			return FString::Printf(TEXT("Sphere: Center:%s, Radius:%f"), *Center.ToString(), GetRadiusf());
 		}
 
 		FORCEINLINE void SerializeImp(FArchive& Ar)
@@ -274,7 +291,7 @@ namespace Chaos
 			Ar << Center;
 
 			// Radius is now stored in the base class Margin
-			FRealSingle ArRadius = (FRealSingle)GetRadius(); // LWC_TODO : potential precision loss, to be changed when we can serialize FReal as double
+			FRealSingle ArRadius = GetRadiusf();
 			Ar << ArRadius;
 			SetRadius(ArRadius);
 		}
@@ -297,7 +314,7 @@ namespace Chaos
 		{
 			TArray<TVector<T, d>> Points;
 			
-			if(GetRadius() <= UE_KINDA_SMALL_NUMBER)
+			if(GetRadiusf() <= UE_KINDA_SMALL_NUMBER)
 			{
 				// If we're too small (and will create NaNs) then just take the centre
 				Points.Add(TVector<T, d>(0.0));
@@ -305,7 +322,7 @@ namespace Chaos
 			else
 			{
 				Points.Reserve(NumPoints);
-				TSphere<T, d> LocalSphere(TVector<T, d>(0.0), GetRadius());
+				TSphere<T, d> LocalSphere(TVector<T, d>(0.0), GetRadiusf());
 				TSphereSpecializeSamplingHelper<T, d>::ComputeSamplePoints(Points, LocalSphere, NumPoints);
 			}
 
@@ -334,7 +351,7 @@ namespace Chaos
 
 		PMatrix<T, d, d> GetInertiaTensor(const T InMass, const bool bInThinShell = false) const 
 		{ 
-			return GetInertiaTensor(InMass, GetRadius(), bInThinShell);
+			return GetInertiaTensor(InMass, GetRadiusf(), bInThinShell);
 		}
 
 		static PMatrix<T, d, d> GetInertiaTensor(const T InMass, const T InRadius, const bool bInThinShell = false)
@@ -353,18 +370,18 @@ namespace Chaos
 		virtual uint32 GetTypeHash() const override
 		{
 			const uint32 CenterHash = UE::Math::GetTypeHash(Center);
-			const uint32 RadiusHash = ::GetTypeHash(GetRadius());
+			const uint32 RadiusHash = ::GetTypeHash(GetRadiusf());
 			return HashCombine(CenterHash, RadiusHash);
 		}
 
 		virtual Chaos::FImplicitObjectPtr CopyGeometry() const override
 		{
-			return Chaos::FImplicitObjectPtr(new TSphere<T,d>(Center, GetRadius()));
+			return Chaos::FImplicitObjectPtr(new FSphere(Center, GetRadiusf()));
 		}
 
 		virtual Chaos::FImplicitObjectPtr CopyGeometryWithScale(const FVec3& Scale) const override
 		{
-			return  Chaos::FImplicitObjectPtr(new TSphere<T, d>(Center * Scale, GetRadius() * Scale.Min()));
+			return  Chaos::FImplicitObjectPtr(new FSphere(Center * Scale, GetRadiusf() * Scale.Min()));
 		}
 
 #if INTEL_ISPC
@@ -380,9 +397,9 @@ namespace Chaos
 #endif // #if INTEL_ISPC
 
 	private:
-		void SetRadius(FReal InRadius) { SetMargin(InRadius); }
+		void SetRadius(FRealSingle InRadius) { SetMargin(InRadius); }
 
-		TVector<T, d> Center;
+		FVec3f Center;
 
 	private:
 
@@ -399,7 +416,7 @@ namespace Chaos
 	{
 		static FORCEINLINE void ComputeSamplePoints(TArray<TVec2<T>>& Points, const TSphere<T, 2>& Sphere, const int32 NumPoints)
 		{
-			if (NumPoints <= 1 || Sphere.GetRadius() < UE_KINDA_SMALL_NUMBER)
+			if (NumPoints <= 1 || Sphere.GetRadiusf() < UE_KINDA_SMALL_NUMBER)
 			{
 				const int32 Offset = Points.AddUninitialized(1);
 				Points[Offset] = Sphere.Center();
@@ -457,10 +474,10 @@ namespace Chaos
 	{
 		static FORCEINLINE void ComputeSamplePoints(TArray<TVec3<T>>& Points, const TSphere<T, 3>& Sphere, const int32 NumPoints)
 		{
-			if (NumPoints <= 1 || Sphere.GetRadius() < UE_KINDA_SMALL_NUMBER)
+			if (NumPoints <= 1 || Sphere.GetRadiusf() < UE_KINDA_SMALL_NUMBER)
 			{
 				const int32 Offset = Points.AddUninitialized(1);
-				Points[Offset] = Sphere.GetCenter();
+				Points[Offset] = Sphere.GetCenterf();
 				return;
 			}
 			ComputeGoldenSpiralPoints(Points, Sphere, NumPoints);
@@ -473,7 +490,7 @@ namespace Chaos
 		    TArray<TVec3<T>>& Points, const TSphere<T, 3>& Sphere, const int32 NumPoints,
 		    const bool FirstHalf = true, const bool SecondHalf = true, const int32 SpiralSeed = 0)
 		{
-			ComputeGoldenSpiralPoints(Points, Sphere.GetCenter(), Sphere.GetRadius(), NumPoints, FirstHalf, SecondHalf, SpiralSeed);
+			ComputeGoldenSpiralPoints(Points, Sphere.GetCenterf(), Sphere.GetRadiusf(), NumPoints, FirstHalf, SecondHalf, SpiralSeed);
 		}
 
 		/**
diff --git a/Engine/Source/Runtime/Experimental/Chaos/Public/Chaos/TaperedCapsule.h b/Engine/Source/Runtime/Experimental/Chaos/Public/Chaos/TaperedCapsule.h
--- a/Engine/Source/Runtime/Experimental/Chaos/Public/Chaos/TaperedCapsule.h
+++ b/Engine/Source/Runtime/Experimental/Chaos/Public/Chaos/TaperedCapsule.h
@@ -17,18 +17,18 @@ namespace Chaos
 		{
 			this->bIsConvex = true;
 		}
-		FTaperedCapsule(const FVec3& X1, const FVec3& X2, const FReal Radius1, const FReal Radius2)
+		FTaperedCapsule(const FVec3& X1, const FVec3& X2, const FReal InRadius1, const FReal InRadius2)
 		    : FImplicitObject(EImplicitObject::FiniteConvex, ImplicitObjectType::TaperedCapsule)
 		    , Origin(X1)
 		    , Axis((X2 - X1).GetSafeNormal())
-		    , Height((X2 - X1).Size())
-		    , Radius1(Radius1)
-		    , Radius2(Radius2)
+		    , Height((FRealSingle)((X2 - X1).Size()))
+		    , Radius1((FRealSingle)InRadius1)
+		    , Radius2((FRealSingle)InRadius2)
 		    , LocalBoundingBox(X1, X1)
 		{
 			this->bIsConvex = true;
 			LocalBoundingBox.GrowToInclude(X2);
-			FReal MaxRadius = FMath::Max(Radius1, Radius2);
+			const FRealSingle MaxRadius = FMath::Max(Radius1, Radius2);
 			LocalBoundingBox = FAABB3(LocalBoundingBox.Min() - FVec3(MaxRadius), LocalBoundingBox.Max() + FVec3(MaxRadius));
 		}
 		FTaperedCapsule(const FTaperedCapsule& Other)
@@ -106,13 +106,12 @@ namespace Chaos
 
 		FReal PhiWithNormal(const FVec3& x, FVec3& OutNormal) const
 		{
-			const FVec3 OriginToX = x - Origin;
-			const FReal DistanceAlongAxis = FMath::Clamp(FVec3::DotProduct(OriginToX, Axis), (FReal)0.0, Height);
-			const FVec3 ClosestPoint = Origin + Axis * DistanceAlongAxis;
-			const FReal Radius = (Height > UE_SMALL_NUMBER) ? FMath::Lerp(Radius1, Radius2, DistanceAlongAxis / Height) : FMath::Max(Radius1, Radius2);
+			const FVec3f OriginToX = x - Origin;
+			const FRealSingle DistanceAlongAxis = FMath::Clamp(FVec3f::DotProduct(OriginToX, Axis), (FRealSingle)0.0, Height);
+			const FVec3f ClosestPoint = Origin + Axis * DistanceAlongAxis;
+			const FRealSingle Radius = (Height > UE_SMALL_NUMBER) ? FMath::Lerp(Radius1, Radius2, DistanceAlongAxis / Height) : FMath::Max(Radius1, Radius2);
 			OutNormal = (x - ClosestPoint);
-			const FReal NormalSize = OutNormal.SafeNormalize();
-			return NormalSize - Radius;
+			return OutNormal.SafeNormalize() - Radius;
 		}
 
 		FReal GetRadius1() const { return Radius1; }
@@ -125,7 +124,7 @@ namespace Chaos
 		FVec3 GetOrigin() const { return GetX1(); }
 		/** Returns the topmost hemisphere center of capsule . */
 		FVec3 GetInsertion() const { return GetX2(); }
-		FVec3 GetCenter() const { return Origin + Axis * (Height * (FReal)0.5); }
+		FVec3 GetCenter() const { return Origin + Axis * (Height * (FRealSingle)0.5); }
 		/** Returns the centroid (center of mass). */
 		FVec3 GetCenterOfMass() const // centroid
 		{
@@ -257,8 +256,8 @@ namespace Chaos
 			return PI_OVER_3 * Height * (Radius1 * Radius1 + Radius1 * Radius2 + Radius2 * Radius2);
 		}
 
-		FVec3 Origin, Axis;
-		FReal Height, Radius1, Radius2;
+		FVec3f Origin, Axis;
+		FRealSingle Height, Radius1, Radius2;
 		FAABB3 LocalBoundingBox;
 	};
 
@@ -387,8 +386,8 @@ namespace Chaos
 			int32 NumPointsEndCap2;
 			int32 NumPointsTaperedSection;
 
-			const FReal Cap1Area = TSphere<FReal, 3>::GetArea(Radius1) / (FReal)2.;
-			const FReal Cap2Area = TSphere<FReal, 3>::GetArea(Radius2) / (FReal)2.;
+			const FReal Cap1Area = FSphere::GetArea(Radius1) / (FReal)2.;
+			const FReal Cap2Area = FSphere::GetArea(Radius2) / (FReal)2.;
 			const FReal TaperedSectionArea = FTaperedCapsule::GetArea(Height, Radius1, Radius2, /*IncludeEndCaps*/ false);
 			const FReal AllArea = TaperedSectionArea + Cap1Area + Cap2Area;
 			if (AllArea > UE_KINDA_SMALL_NUMBER)
@@ -411,7 +410,7 @@ namespace Chaos
 			{
 				// Points vary in Z: [-Radius1-HalfHeight, -HalfHeight]
 				TSphereSpecializeSamplingHelper<FReal, 3>::ComputeBottomHalfSemiSphere(
-					Points, TSphere<FReal, 3>(FVec3(0, 0, -HalfHeight), Radius1), NumPointsEndCap1, SpiralSeed);
+					Points, FSphere(FVec3(0, 0, -HalfHeight), Radius1), NumPointsEndCap1, SpiralSeed);
 				SpiralSeed += Points.Num();
 
 				// Points vary in Z: [-HalfHeight, HalfHeight], about the Z axis.
@@ -421,7 +420,7 @@ namespace Chaos
 
 				// Points vary in Z: [HalfHeight, HalfHeight+Radius2]
 				TSphereSpecializeSamplingHelper<FReal, 3>::ComputeTopHalfSemiSphere(
-					Points, TSphere<FReal, 3>(FVec3(0, 0, HalfHeight), Radius2), NumPointsEndCap2, SpiralSeed);
+					Points, FSphere(FVec3(0, 0, HalfHeight), Radius2), NumPointsEndCap2, SpiralSeed);
 				SpiralSeed += Points.Num();
 			}
 		}
diff --git a/Engine/Source/Runtime/Experimental/Chaos/Public/Chaos/TaperedCylinder.h b/Engine/Source/Runtime/Experimental/Chaos/Public/Chaos/TaperedCylinder.h
--- a/Engine/Source/Runtime/Experimental/Chaos/Public/Chaos/TaperedCylinder.h
+++ b/Engine/Source/Runtime/Experimental/Chaos/Public/Chaos/TaperedCylinder.h
@@ -21,14 +21,14 @@ namespace Chaos
 		    : FImplicitObject(EImplicitObject::FiniteConvex, ImplicitObjectType::TaperedCylinder)
 		    , MPlane1(x1, (x2 - x1).GetSafeNormal())
 		    , MPlane2(x2, -MPlane1.Normal())
-		    , MHeight((x2 - x1).Size())
-		    , MRadius1(Radius1)
-		    , MRadius2(Radius2)
+		    , MHeight((FRealSingle)((x2 - x1).Size()))
+		    , MRadius1((FRealSingle)Radius1)
+		    , MRadius2((FRealSingle)Radius2)
 		    , MLocalBoundingBox(x1, x1)
 		{
 			this->bIsConvex = true;
 			MLocalBoundingBox.GrowToInclude(x2);
-			FReal MaxRadius = MRadius1;
+			FRealSingle MaxRadius = MRadius1;
 			if (MaxRadius < MRadius2)
 				MaxRadius = MRadius2;
 			MLocalBoundingBox = FAABB3(MLocalBoundingBox.Min() - FVec3(MaxRadius), MLocalBoundingBox.Max() + FVec3(MaxRadius));
@@ -103,16 +103,17 @@ namespace Chaos
 
 		FReal PhiWithNormal(const FVec3& x, FVec3& Normal) const
 		{
-			const FVec3& Normal1 = MPlane1.Normal();
-			const FReal Distance1 = MPlane1.SignedDistance(x);
+			const FVec3f Xf(x);
+			const FVec3f& Normal1 = MPlane1.Normal();
+			const FRealSingle Distance1 = MPlane1.SignedDistance(Xf);
 			if (Distance1 < UE_SMALL_NUMBER)
 			{
-				ensure(MPlane2.SignedDistance(x) > (FReal)0.);
-				const FVec3 v = x - FVec3(Normal1 * Distance1 + MPlane1.X());
+				ensure(MPlane2.SignedDistance(Xf) > (FRealSingle)0.);
+				const FVec3f v = x - FVec3f(Normal1 * Distance1 + MPlane1.X());
 				if (v.Size() > MRadius1)
 				{
-					const FVec3 Corner = v.GetSafeNormal() * MRadius1 + MPlane1.X();
-					const FVec3 CornerVector = x - Corner;
+					const FVec3f Corner = FVec3f(v.GetSafeNormal()) * MRadius1 + MPlane1.X();
+					const FVec3f CornerVector = x - Corner;
 					Normal = CornerVector.GetSafeNormal();
 					return CornerVector.Size();
 				}
@@ -122,15 +123,15 @@ namespace Chaos
 					return -Distance1;
 				}
 			}
-			const FVec3& Normal2 = MPlane2.Normal();  // Used to be Distance2 = MPlane2.PhiWithNormal(x, Normal2); but that would trigger 
-			const FReal Distance2 = MHeight - Distance1;          // the ensure on Distance2 being slightly larger than MHeight in some border cases
+			const FVec3f& Normal2 = MPlane2.Normal();  // Used to be Distance2 = MPlane2.PhiWithNormal(x, Normal2); but that would trigger 
+			const FRealSingle Distance2 = MHeight - Distance1;          // the ensure on Distance2 being slightly larger than MHeight in some border cases
 			if (Distance2 < UE_SMALL_NUMBER)
 			{
-				const FVec3 v = x - FVec3(Normal2 * Distance2 + MPlane2.X());
+				const FVec3f v = Xf - FVec3f(Normal2 * Distance2 + MPlane2.X());
 				if (v.Size() > MRadius2)
 				{
-					const FVec3 Corner = v.GetSafeNormal() * MRadius2 + MPlane2.X();
-					const FVec3 CornerVector = x - Corner;
+					const FVec3f Corner = FVec3f(v.GetSafeNormal()) * MRadius2 + MPlane2.X();
+					const FVec3f CornerVector = Xf - Corner;
 					Normal = CornerVector.GetSafeNormal();
 					return CornerVector.Size();
 				}
@@ -141,11 +142,11 @@ namespace Chaos
 				}
 			}
 			ensure(Distance1 <= MHeight && Distance2 <= MHeight);
-			const FVec3 SideVector = (x - FVec3(Normal1 * Distance1 + MPlane1.X()));
-			const FReal SideDistance = SideVector.Size() - GetRadiusAtDistance(Distance1);
+			const FVec3f SideVector = (Xf - FVec3f(Normal1 * Distance1 + MPlane1.X()));
+			const FRealSingle SideDistance = SideVector.Size() - GetRadiusAtDistance(Distance1);
 			if (SideDistance < 0.)
 			{
-				const FReal TopDistance = Distance1 < Distance2 ? Distance1 : Distance2;
+				const FRealSingle TopDistance = Distance1 < Distance2 ? Distance1 : Distance2;
 				if (TopDistance < -SideDistance)
 				{
 					Normal = Distance1 < Distance2 ? -Normal1 : -Normal2;
@@ -248,12 +249,16 @@ namespace Chaos
 		FReal GetRadius2() const { return MRadius2; }
 		FReal GetHeight() const { return MHeight; }
 		FReal GetSlantHeight() const { const FReal R1mR2 = MRadius1-MRadius2; return FMath::Sqrt(R1mR2*R1mR2 + MHeight*MHeight); }
-		const FVec3& GetX1() const { return MPlane1.X(); }
-		const FVec3& GetX2() const { return MPlane2.X(); }
+		FVec3 GetX1() const { return MPlane1.X(); }
+		const FVec3f& GetX1f() const { return MPlane1.X(); }
+		FVec3 GetX2() const { return MPlane2.X(); }
+		const FVec3f& GetX2f() const { return MPlane2.X(); }
 		/** Returns the bottommost point on the cylinder. */
-		const FVec3& GetOrigin() const { return MPlane1.X(); }
+		FVec3 GetOrigin() const { return MPlane1.X(); }
+		const FVec3f& GetOriginf() const { return MPlane1.X(); }
 		/** Returns the topmost point on the cylinder. */
-		const FVec3& GetInsertion() const { return MPlane2.X(); }
+		FVec3 GetInsertion() const { return MPlane2.X(); }
+		const FVec3f& GetInsertionf() const { return MPlane2.X(); }
 		FVec3 GetCenter() const { return (MPlane1.X() + MPlane2.X()) * (FReal)0.5; }
 		/** Returns the centroid (center of mass). */
 		FVec3 GetCenterOfMass() const // centroid
@@ -352,14 +357,14 @@ namespace Chaos
 
 	private:
 		//Phi is distance from closest point on plane1
-		FReal GetRadiusAtDistance(const FReal& Phi) const
+		FRealSingle GetRadiusAtDistance(const FRealSingle& Phi) const
 		{
-			const FReal Alpha = Phi / MHeight;
-			return MRadius1 * (static_cast<FReal>(1.) - Alpha) + MRadius2 * Alpha;
+			const FRealSingle Alpha = Phi / MHeight;
+			return MRadius1 * (static_cast<FRealSingle>(1.) - Alpha) + MRadius2 * Alpha;
 		}
 
-		TPlaneConcrete<FReal, 3> MPlane1, MPlane2;
-		FReal MHeight, MRadius1, MRadius2;
+		TPlaneConcrete<FRealSingle, 3> MPlane1, MPlane2;
+		FRealSingle MHeight, MRadius1, MRadius2;
 		FAABB3 MLocalBoundingBox;
 	};
 
diff --git a/Engine/Source/Runtime/Experimental/Chaos/Public/Chaos/Triangle.h b/Engine/Source/Runtime/Experimental/Chaos/Public/Chaos/Triangle.h
--- a/Engine/Source/Runtime/Experimental/Chaos/Public/Chaos/Triangle.h
+++ b/Engine/Source/Runtime/Experimental/Chaos/Public/Chaos/Triangle.h
@@ -310,6 +310,8 @@ namespace Chaos
 
 		FORCEINLINE T GetMargin() const { return 0; }
 		FORCEINLINE T GetRadius() const { return 0; }
+		FORCEINLINE FRealSingle GetMarginf() const { return 0.0f; }
+		FORCEINLINE FRealSingle GetRadiusf() const { return 0.0f; }
 
 		FORCEINLINE bool Raycast(const TVec3<T>& StartPoint, const TVec3<T>& Dir, const T Length, const T Thickness, T& OutTime, TVec3<T>& OutPosition, TVec3<T>& OutNormal, int32& OutFaceIndex) const
 		{
diff --git a/Engine/Source/Runtime/Experimental/Chaos/Public/Chaos/TriangleMeshImplicitObject.h b/Engine/Source/Runtime/Experimental/Chaos/Public/Chaos/TriangleMeshImplicitObject.h
--- a/Engine/Source/Runtime/Experimental/Chaos/Public/Chaos/TriangleMeshImplicitObject.h
+++ b/Engine/Source/Runtime/Experimental/Chaos/Public/Chaos/TriangleMeshImplicitObject.h
@@ -538,47 +538,47 @@ namespace Chaos
 			return 0.0f;
 		}
 
-		CHAOS_API virtual FReal PhiWithNormal(const FVec3& x, FVec3& Normal) const;
+		CHAOS_API virtual FReal PhiWithNormal(const FVec3& x, FVec3& Normal) const override;
 
 		CHAOS_API virtual bool Raycast(const FVec3& StartPoint, const FVec3& Dir, const FReal Length, const FReal Thickness, FReal& OutTime, FVec3& OutPosition, FVec3& OutNormal, int32& OutFaceIndex) const override;
 		CHAOS_API virtual bool Overlap(const FVec3& Point, const FReal Thickness) const override;
 
-		CHAOS_API bool OverlapGeom(const TSphere<FReal, 3>& QueryGeom, const FRigidTransform3& QueryTM, const FReal Thickness, FMTDInfo* OutMTD = nullptr) const;
+		CHAOS_API bool OverlapGeom(const FSphere& QueryGeom, const FRigidTransform3& QueryTM, const FReal Thickness, FMTDInfo* OutMTD = nullptr) const;
 		CHAOS_API bool OverlapGeom(const TBox<FReal, 3>& QueryGeom, const FRigidTransform3& QueryTM, const FReal Thickness, FMTDInfo* OutMTD = nullptr) const;
 		CHAOS_API bool OverlapGeom(const FCapsule& QueryGeom, const FRigidTransform3& QueryTM, const FReal Thickness, FMTDInfo* OutMTD = nullptr) const;
 		CHAOS_API bool OverlapGeom(const FConvex& QueryGeom, const FRigidTransform3& QueryTM, const FReal Thickness, FMTDInfo* OutMTD = nullptr) const;
 
-		CHAOS_API bool OverlapGeom(const TImplicitObjectScaled<TSphere<FReal, 3>>& QueryGeom, const FRigidTransform3& QueryTM, const FReal Thickness, FMTDInfo* OutMTD = nullptr, FVec3 TriMeshScale = FVec3(1.0f)) const;
+		CHAOS_API bool OverlapGeom(const TImplicitObjectScaled<FSphere>& QueryGeom, const FRigidTransform3& QueryTM, const FReal Thickness, FMTDInfo* OutMTD = nullptr, FVec3 TriMeshScale = FVec3(1.0f)) const;
 		CHAOS_API bool OverlapGeom(const TImplicitObjectScaled<TBox<FReal, 3>>& QueryGeom, const FRigidTransform3& QueryTM, const FReal Thickness, FMTDInfo* OutMTD = nullptr, FVec3 TriMeshScale = FVec3(1.0f)) const;
 		CHAOS_API bool OverlapGeom(const TImplicitObjectScaled<FCapsule>& QueryGeom, const FRigidTransform3& QueryTM, const FReal Thickness, FMTDInfo* OutMTD = nullptr, FVec3 TriMeshScale = FVec3(1.0f)) const;
 		CHAOS_API bool OverlapGeom(const TImplicitObjectScaled<FConvex>& QueryGeom, const FRigidTransform3& QueryTM, const FReal Thickness, FMTDInfo* OutMTD = nullptr, FVec3 TriMeshScale = FVec3(1.0f)) const;
 
-		CHAOS_API bool SweepGeom(const TSphere<FReal, 3>& QueryGeom, const FRigidTransform3& StartTM, const FVec3& Dir, const FReal Length, FReal& OutTime, FVec3& OutPosition, FVec3& OutNormal, int32& OutFaceIndex, FVec3& OutFaceNormal, const FReal Thickness = 0, const bool bComputeMTD = false, FVec3 TriMeshScale = FVec3(1.0f)) const;
+		CHAOS_API bool SweepGeom(const FSphere& QueryGeom, const FRigidTransform3& StartTM, const FVec3& Dir, const FReal Length, FReal& OutTime, FVec3& OutPosition, FVec3& OutNormal, int32& OutFaceIndex, FVec3& OutFaceNormal, const FReal Thickness = 0, const bool bComputeMTD = false, FVec3 TriMeshScale = FVec3(1.0f)) const;
 		CHAOS_API bool SweepGeom(const TBox<FReal, 3>& QueryGeom, const FRigidTransform3& StartTM, const FVec3& Dir, const FReal Length, FReal& OutTime, FVec3& OutPosition, FVec3& OutNormal, int32& OutFaceIndex, FVec3& OutFaceNormal, const FReal Thickness = 0, const bool bComputeMTD = false, FVec3 TriMeshScale = FVec3(1.0f)) const;
 		CHAOS_API bool SweepGeom(const FCapsule& QueryGeom, const FRigidTransform3& StartTM, const FVec3& Dir, const FReal Length, FReal& OutTime, FVec3& OutPosition, FVec3& OutNormal, int32& OutFaceIndex, FVec3& OutFaceNormal, const FReal Thickness = 0, const bool bComputeMTD = false, FVec3 TriMeshScale = FVec3(1.0f)) const;
 		CHAOS_API bool SweepGeom(const FConvex& QueryGeom, const FRigidTransform3& StartTM, const FVec3& Dir, const FReal Length, FReal& OutTime, FVec3& OutPosition, FVec3& OutNormal, int32& OutFaceIndex, FVec3& OutFaceNormal, const FReal Thickness = 0, const bool bComputeMTD = false, FVec3 TriMeshScale = FVec3(1.0f)) const;
 
-		CHAOS_API bool SweepGeom(const TImplicitObjectScaled<TSphere<FReal, 3>>& QueryGeom, const FRigidTransform3& StartTM, const FVec3& Dir, const FReal Length, FReal& OutTime, FVec3& OutPosition, FVec3& OutNormal, int32& OutFaceIndex, FVec3& OutFaceNormal, const FReal Thickness = 0, const bool bComputeMTD = false, FVec3 TriMeshScale = FVec3(1.0f)) const;
+		CHAOS_API bool SweepGeom(const TImplicitObjectScaled<FSphere>& QueryGeom, const FRigidTransform3& StartTM, const FVec3& Dir, const FReal Length, FReal& OutTime, FVec3& OutPosition, FVec3& OutNormal, int32& OutFaceIndex, FVec3& OutFaceNormal, const FReal Thickness = 0, const bool bComputeMTD = false, FVec3 TriMeshScale = FVec3(1.0f)) const;
 		CHAOS_API bool SweepGeom(const TImplicitObjectScaled<TBox<FReal, 3>>& QueryGeom, const FRigidTransform3& StartTM, const FVec3& Dir, const FReal Length, FReal& OutTime, FVec3& OutPosition, FVec3& OutNormal, int32& OutFaceIndex, FVec3& OutFaceNormal, const FReal Thickness = 0, const bool bComputeMTD = false, FVec3 TriMeshScale = FVec3(1.0f)) const;
 		CHAOS_API bool SweepGeom(const TImplicitObjectScaled<FCapsule>& QueryGeom, const FRigidTransform3& StartTM, const FVec3& Dir, const FReal Length, FReal& OutTime, FVec3& OutPosition, FVec3& OutNormal, int32& OutFaceIndex, FVec3& OutFaceNormal, const FReal Thickness = 0, const bool bComputeMTD = false, FVec3 TriMeshScale = FVec3(1.0f)) const;
 		CHAOS_API bool SweepGeom(const TImplicitObjectScaled<FConvex>& QueryGeom, const FRigidTransform3& StartTM, const FVec3& Dir, const FReal Length, FReal& OutTime, FVec3& OutPosition, FVec3& OutNormal, int32& OutFaceIndex, FVec3& OutFaceNormal, const FReal Thickness = 0, const bool bComputeMTD = false, FVec3 TriMeshScale = FVec3(1.0f)) const;
 
 		// Sweep used for CCD. Ignores triangles we penetrate by less than IgnorePenetration, and calculate the TOI for a depth of TargetPenetration. If both are zero, this is equivalent to SweepGeom
-		CHAOS_API bool SweepGeomCCD(const TSphere<FReal, 3>& QueryGeom, const FRigidTransform3& StartTM, const FVec3& Dir, const FReal Length, const FReal IgnorePenetration, const FReal TargetPenetration, FReal& OutTOI, FReal& OutPhi, FVec3& OutPosition, FVec3& OutNormal, int32& OutFaceIndex, FVec3& OutFaceNormal, const FVec3& TriMeshScale = FVec3(1.0f)) const;
+		CHAOS_API bool SweepGeomCCD(const FSphere& QueryGeom, const FRigidTransform3& StartTM, const FVec3& Dir, const FReal Length, const FReal IgnorePenetration, const FReal TargetPenetration, FReal& OutTOI, FReal& OutPhi, FVec3& OutPosition, FVec3& OutNormal, int32& OutFaceIndex, FVec3& OutFaceNormal, const FVec3& TriMeshScale = FVec3(1.0f)) const;
 		CHAOS_API bool SweepGeomCCD(const TBox<FReal, 3>& QueryGeom, const FRigidTransform3& StartTM, const FVec3& Dir, const FReal Length, const FReal IgnorePenetration, const FReal TargetPenetration, FReal& OutTOI, FReal& OutPhi, FVec3& OutPosition, FVec3& OutNormal, int32& OutFaceIndex, FVec3& OutFaceNormal, const FVec3& TriMeshScale = FVec3(1.0f)) const;
 		CHAOS_API bool SweepGeomCCD(const FCapsule& QueryGeom, const FRigidTransform3& StartTM, const FVec3& Dir, const FReal Length, const FReal IgnorePenetration, const FReal TargetPenetration, FReal& OutTOI, FReal& OutPhi, FVec3& OutPosition, FVec3& OutNormal, int32& OutFaceIndex, FVec3& OutFaceNormal, const FVec3& TriMeshScale = FVec3(1.0f)) const;
 		CHAOS_API bool SweepGeomCCD(const FConvex& QueryGeom, const FRigidTransform3& StartTM, const FVec3& Dir, const FReal Length, const FReal IgnorePenetration, const FReal TargetPenetration, FReal& OutTOI, FReal& OutPhi, FVec3& OutPosition, FVec3& OutNormal, int32& OutFaceIndex, FVec3& OutFaceNormal, const FVec3& TriMeshScale = FVec3(1.0f)) const;
-		CHAOS_API bool SweepGeomCCD(const TImplicitObjectScaled<TSphere<FReal, 3>>& QueryGeom, const FRigidTransform3& StartTM, const FVec3& Dir, const FReal Length, const FReal IgnorePenetration, const FReal TargetPenetration, FReal& OutTOI, FReal& OutPhi, FVec3& OutPosition, FVec3& OutNormal, int32& OutFaceIndex, FVec3& OutFaceNormal, const FVec3& TriMeshScale = FVec3(1.0f)) const;
+		CHAOS_API bool SweepGeomCCD(const TImplicitObjectScaled<FSphere>& QueryGeom, const FRigidTransform3& StartTM, const FVec3& Dir, const FReal Length, const FReal IgnorePenetration, const FReal TargetPenetration, FReal& OutTOI, FReal& OutPhi, FVec3& OutPosition, FVec3& OutNormal, int32& OutFaceIndex, FVec3& OutFaceNormal, const FVec3& TriMeshScale = FVec3(1.0f)) const;
 		CHAOS_API bool SweepGeomCCD(const TImplicitObjectScaled<TBox<FReal, 3>>& QueryGeom, const FRigidTransform3& StartTM, const FVec3& Dir, const FReal Length, const FReal IgnorePenetration, const FReal TargetPenetration, FReal& OutTOI, FReal& OutPhi, FVec3& OutPosition, FVec3& OutNormal, int32& OutFaceIndex, FVec3& OutFaceNormal, const FVec3& TriMeshScale = FVec3(1.0f)) const;
 		CHAOS_API bool SweepGeomCCD(const TImplicitObjectScaled<FCapsule>& QueryGeom, const FRigidTransform3& StartTM, const FVec3& Dir, const FReal Length, const FReal IgnorePenetration, const FReal TargetPenetration, FReal& OutTOI, FReal& OutPhi, FVec3& OutPosition, FVec3& OutNormal, int32& OutFaceIndex, FVec3& OutFaceNormal, const FVec3& TriMeshScale = FVec3(1.0f)) const;
 		CHAOS_API bool SweepGeomCCD(const TImplicitObjectScaled<FConvex>& QueryGeom, const FRigidTransform3& StartTM, const FVec3& Dir, const FReal Length, const FReal IgnorePenetration, const FReal TargetPenetration, FReal& OutTOI, FReal& OutPhi, FVec3& OutPosition, FVec3& OutNormal, int32& OutFaceIndex, FVec3& OutFaceNormal, const FVec3& TriMeshScale = FVec3(1.0f)) const;
 
-		CHAOS_API bool GJKContactPoint(const TSphere<FReal, 3>& QueryGeom, const FRigidTransform3& QueryTM, const FReal Thickness, FVec3& Location, FVec3& Normal, FReal& Penetration, int32& FaceIndex) const;
+		CHAOS_API bool GJKContactPoint(const FSphere& QueryGeom, const FRigidTransform3& QueryTM, const FReal Thickness, FVec3& Location, FVec3& Normal, FReal& Penetration, int32& FaceIndex) const;
 		CHAOS_API bool GJKContactPoint(const TBox<FReal, 3>& QueryGeom, const FRigidTransform3& QueryTM, const FReal Thickness, FVec3& Location, FVec3& Normal, FReal& Penetration, int32& FaceIndex) const;
 		CHAOS_API bool GJKContactPoint(const FCapsule& QueryGeom, const FRigidTransform3& QueryTM, const FReal Thickness, FVec3& Location, FVec3& Normal, FReal& Penetration, int32& FaceIndex) const;
 		CHAOS_API bool GJKContactPoint(const FConvex& QueryGeom, const FRigidTransform3& QueryTM, const FReal Thickness, FVec3& Location, FVec3& Normal, FReal& Penetration, int32& FaceIndex) const;
 
-		CHAOS_API bool GJKContactPoint(const TImplicitObjectScaled < TSphere<FReal, 3> >& QueryGeom, const FRigidTransform3& QueryTM, const FReal Thickness, FVec3& Location, FVec3& Normal, FReal& Penetration, int32& FaceIndex, FVec3 TriMeshScale = FVec3(1.0f)) const;
+		CHAOS_API bool GJKContactPoint(const TImplicitObjectScaled < FSphere >& QueryGeom, const FRigidTransform3& QueryTM, const FReal Thickness, FVec3& Location, FVec3& Normal, FReal& Penetration, int32& FaceIndex, FVec3 TriMeshScale = FVec3(1.0f)) const;
 		CHAOS_API bool GJKContactPoint(const TImplicitObjectScaled < TBox<FReal, 3> >& QueryGeom, const FRigidTransform3& QueryTM, const FReal Thickness, FVec3& Location, FVec3& Normal, FReal& Penetration, int32& FaceIndex, FVec3 TriMeshScale = FVec3(1.0f)) const;
 		CHAOS_API bool GJKContactPoint(const TImplicitObjectScaled < FCapsule >& QueryGeom, const FRigidTransform3& QueryTM, const FReal Thickness, FVec3& Location, FVec3& Normal, FReal& Penetration, int32& FaceIndex, FVec3 TriMeshScale = FVec3(1.0f)) const;
 		CHAOS_API bool GJKContactPoint(const TImplicitObjectScaled < FConvex >& QueryGeom, const FRigidTransform3& QueryTM, const FReal Thickness, FVec3& Location, FVec3& Normal, FReal& Penetration, int32& FaceIndex, FVec3 TriMeshScale = FVec3(1.0f)) const;
@@ -595,11 +595,16 @@ namespace Chaos
 		CHAOS_API virtual int32 FindMostOpposingFaceScaled(const FVec3& Position, const FVec3& UnitDir, int32 HintFaceIndex, FReal SearchDist, const FVec3& Scale) const override;
 		CHAOS_API virtual FVec3 FindGeometryOpposingNormal(const FVec3& DenormDir, int32 FaceIndex, const FVec3& OriginalNormal) const override;
 
-		virtual const FAABB3 BoundingBox() const
+		virtual const FAABB3 BoundingBox() const override
 		{
 			return MLocalBoundingBox;
 		}
 
+		const FAABB3f BoundingBoxf() const
+		{
+			return FAABB3f(MLocalBoundingBox);
+		}
+
 		static constexpr EImplicitObjectType StaticType()
 		{
 			return ImplicitObjectType::TriangleMesh;
diff --git a/Engine/Source/Runtime/Experimental/Chaos/Public/PhysicsProxy/AnalyticImplicitGroup.h b/Engine/Source/Runtime/Experimental/Chaos/Public/PhysicsProxy/AnalyticImplicitGroup.h
--- a/Engine/Source/Runtime/Experimental/Chaos/Public/PhysicsProxy/AnalyticImplicitGroup.h
+++ b/Engine/Source/Runtime/Experimental/Chaos/Public/PhysicsProxy/AnalyticImplicitGroup.h
@@ -61,7 +61,7 @@ public:
 	{}
 	~FAnalyticImplicitGroup()
 	{
-		for (Chaos::TSphere<Chaos::FReal, 3>* Sphere : Spheres) if (Sphere) delete Sphere;
+		for (Chaos::FSphere* Sphere : Spheres) if (Sphere) delete Sphere;
 		for (Chaos::TBox<Chaos::FReal, 3>* Box : Boxes) if (Box) delete Box;
 		for (Chaos::FCapsule* Capsule : Capsules) if (Capsule) delete Capsule;
 		for (Chaos::FTaperedCylinder* TaperedCylinder : TaperedCylinders) if (TaperedCylinder) delete TaperedCylinder;
@@ -96,7 +96,7 @@ public:
 	void SetRigidBodyState(const EObjectStateTypeEnum State) { RigidBodyState = State; }
 	EObjectStateTypeEnum GetRigidBodyState() const { return RigidBodyState; }
 
-	int32 Add(const FTransform &InitialXf, Chaos::TSphere<Chaos::FReal, 3> *Sphere) { Spheres.Add(Sphere); return Transforms.Insert(InitialXf, Spheres.Num()-1); }
+	int32 Add(const FTransform &InitialXf, Chaos::FSphere *Sphere) { Spheres.Add(Sphere); return Transforms.Insert(InitialXf, Spheres.Num()-1); }
 	int32 Add(const FTransform &InitialXf, Chaos::TBox<Chaos::FReal, 3> *Box) { Boxes.Add(Box); return Transforms.Insert(InitialXf, Spheres.Num()+Boxes.Num()-1); }
 	int32 Add(const FTransform &InitialXf, Chaos::FCapsule *Capsule) { Capsules.Add(Capsule); return Transforms.Insert(InitialXf, Spheres.Num()+Boxes.Num()+Capsules.Num()-1); }
 	int32 Add(const FTransform &InitialXf, Chaos::FTaperedCylinder *TaperedCylinder) { TaperedCylinders.Add(TaperedCylinder); return Transforms.Insert(InitialXf, Spheres.Num()+Boxes.Num()+Capsules.Num()+TaperedCylinders.Num()-1); }
@@ -128,7 +128,7 @@ public:
 		BBoxes.SetNum(Num);
 
 		int32 TransformIndex = 0;
-		for (Chaos::TSphere<Chaos::FReal, 3>* Sphere : Spheres)
+		for (Chaos::FSphere* Sphere : Spheres)
 		{
 			const FTransform& Xf = Transforms[TransformIndex];
 			BBoxes[TransformIndex] = Sphere->BoundingBox().TransformedAABB(Xf);
@@ -152,7 +152,7 @@ public:
 			BBoxes[TransformIndex] = Capsule->BoundingBox().TransformedAABB(Xf);
 			Chaos::FMassProperties &MP = MPArray[TransformIndex++];
 			MP.Volume = Capsule->GetVolume();
-			MP.CenterOfMass = Xf.TransformPositionNoScale(Capsule->GetCenterOfMass());
+			MP.CenterOfMass = Xf.TransformPositionNoScale(FVector3d(Capsule->GetCenterOfMassf()));
 			MP.RotationOfMass = Xf.TransformRotation(Capsule->GetRotationOfMass());
 		}
 		for (Chaos::FTaperedCylinder* TaperedCylinder : TaperedCylinders)
@@ -212,7 +212,7 @@ public:
 
 		TotalMass = 0.f;
 		TransformIndex = 0;
-		for (Chaos::TSphere<Chaos::FReal, 3>* Sphere : Spheres)
+		for (Chaos::FSphere* Sphere : Spheres)
 		{
 			Chaos::FMassProperties &MP = MPArray[TransformIndex++];
 			Chaos::FReal Mass = Density * MP.Volume;
@@ -276,7 +276,7 @@ public:
 		}
 
 		int32 TransformIndex = 0;
-		for (Chaos::TSphere<Chaos::FReal, 3>* Sphere : Spheres)
+		for (Chaos::FSphere* Sphere : Spheres)
 		{
 			TArray<Chaos::FVec3>& Points = CollisionPoints[TransformIndex];
 			if (!Points.Num())
@@ -314,7 +314,7 @@ public:
 			if (!Points.Num())
 			{
 				const Chaos::FAABB3& BBox = Convex->BoundingBox();
-				Chaos::TSphere<Chaos::FReal, 3> Sphere(BBox.Center(), BBox.Extents().Size() / 2);
+				Chaos::FSphere Sphere(BBox.Center(), BBox.Extents().Size() / 2);
 				Points = Sphere.ComputeSamplePoints(ParticlesPerUnitArea, MinParticles, MaxParticles);
 				Chaos::FVec3 Normal;
 				for (Chaos::FVec3 &Pt : Points)
@@ -333,7 +333,7 @@ public:
 			if (!Points.Num())
 			{
 				const Chaos::FAABB3& BBox = LevelSet->BoundingBox();
-				Chaos::TSphere<Chaos::FReal, 3> Sphere(BBox.Center(), BBox.Extents().Size() / 2);
+				Chaos::FSphere Sphere(BBox.Center(), BBox.Extents().Size() / 2);
 				Points = Sphere.ComputeSamplePoints(ParticlesPerUnitArea, MinParticles, MaxParticles);
 				Chaos::FVec3 Normal;
 				for (Chaos::FVec3 &Pt : Points)
@@ -512,7 +512,7 @@ protected:
 	void CullDeepPoints(TArray<Chaos::FVec3>& Points, const int32 SkipIndex)
 	{
 		int32 TransformIndex = 0;
-		for (Chaos::TSphere<Chaos::FReal, 3>* Sphere : Spheres)
+		for (Chaos::FSphere* Sphere : Spheres)
 		{
 			if (TransformIndex != SkipIndex)
 			{
@@ -628,7 +628,7 @@ protected:
 	EObjectStateTypeEnum RigidBodyState;
 
 	// FkSphereElem and FKTaperedCapsuleElem ends
-	TArray<Chaos::TSphere<Chaos::FReal, 3>*> Spheres;
+	TArray<Chaos::FSphere*> Spheres;
 	// FKBoxElem
 	TArray<Chaos::TBox<Chaos::FReal, 3>*> Boxes;
 	// FKSphylElem - Z axis is capsule axis
diff --git a/Engine/Source/Runtime/Experimental/GeometryCollectionEngine/Private/GeometryCollection/GeometryCollectionDebugDraw.cpp b/Engine/Source/Runtime/Experimental/GeometryCollectionEngine/Private/GeometryCollection/GeometryCollectionDebugDraw.cpp
--- a/Engine/Source/Runtime/Experimental/GeometryCollectionEngine/Private/GeometryCollection/GeometryCollectionDebugDraw.cpp
+++ b/Engine/Source/Runtime/Experimental/GeometryCollectionEngine/Private/GeometryCollection/GeometryCollectionDebugDraw.cpp
@@ -27,17 +27,17 @@ namespace GeometryCollectionDebugDraw
 	void AddSphere(FKAggregateGeom& AggGeom, const Chaos::FImplicitSphere3& Sphere, const Chaos::FRigidTransform3& Transform)
 	{
 		FKSphereElem& SphereElem = AggGeom.SphereElems.AddDefaulted_GetRef();
-		SphereElem.Center = Transform.TransformPosition(Sphere.GetCenter());
-		SphereElem.Radius = (Transform.GetScale3D().Z * Sphere.GetRadius());
+		SphereElem.Center = Transform.TransformPosition(FVector3d(Sphere.GetCenterf()));
+		SphereElem.Radius = (Transform.GetScale3D().Z * Sphere.GetRadiusf());
 	}
 
 	void AddCapsule(FKAggregateGeom& AggGeom, const Chaos::FImplicitCapsule3& Capsule, const Chaos::FRigidTransform3& Transform)
 	{
 		FKSphylElem& CapsuleElem = AggGeom.SphylElems.AddDefaulted_GetRef();
-		CapsuleElem.Center = Transform.TransformPosition(Capsule.GetCenter());
+		CapsuleElem.Center = Transform.TransformPosition(FVector3d(Capsule.GetCenterf()));
 		CapsuleElem.Rotation = (Transform.GetRotation() * FRotationMatrix::MakeFromZ(Capsule.GetAxis())).Rotator();
-		CapsuleElem.Length = (Transform.GetScale3D().Z * Capsule.GetHeight());
-		CapsuleElem.Radius = (Transform.GetScale3D().X * Capsule.GetRadius());
+		CapsuleElem.Length = (Transform.GetScale3D().Z * Capsule.GetHeightf());
+		CapsuleElem.Radius = (Transform.GetScale3D().X * Capsule.GetRadiusf());
 	}
 
 	void AddBox(FKAggregateGeom& AggGeom, const Chaos::FImplicitBox3& Box, const Chaos::FRigidTransform3& Transform)
diff --git a/Engine/Source/Runtime/Experimental/GeometryCollectionEngine/Private/GeometryCollection/PhysicsAssetSimulation.cpp b/Engine/Source/Runtime/Experimental/GeometryCollectionEngine/Private/GeometryCollection/PhysicsAssetSimulation.cpp
--- a/Engine/Source/Runtime/Experimental/GeometryCollectionEngine/Private/GeometryCollection/PhysicsAssetSimulation.cpp
+++ b/Engine/Source/Runtime/Experimental/GeometryCollectionEngine/Private/GeometryCollection/PhysicsAssetSimulation.cpp
@@ -140,7 +140,7 @@ void FPhysicsAssetSimulationUtil::BuildParams(const UObject* Caller, const AActo
 				// Translation is in the transform, scaling is baked into the implicit shape.
 				AnalyticShapeGroup->Add(
 					Xf,
-					new Chaos::TSphere<Chaos::FReal, 3>(Chaos::FVec3(0), Radius));
+					new Chaos::FSphere(Chaos::FVec3(0), Radius));
 
 				UE_LOG(USkeletalMeshSimulationComponentLogging, Log,
 					TEXT("USkeletalMeshSimulationComponent::OnCreatePhysicsState() - Caller: %p Actor: '%s' - "
@@ -211,7 +211,7 @@ void FPhysicsAssetSimulationUtil::BuildParams(const UObject* Caller, const AActo
 				// Translation is in the transform, scaling is baked into the implicit shape.
 				AnalyticShapeGroup->Add(
 					Xf,
-					new Chaos::TSphere<Chaos::FReal, 3>(Chaos::FVec3(0), Radius));
+					new Chaos::FSphere(Chaos::FVec3(0), Radius));
 
 				UE_LOG(USkeletalMeshSimulationComponentLogging, Log,
 					TEXT("USkeletalMeshSimulationComponent::OnCreatePhysicsState() - Caller: %p Actor: '%s' - "
@@ -240,12 +240,12 @@ void FPhysicsAssetSimulationUtil::BuildParams(const UObject* Caller, const AActo
 				if (Radius0 > KINDA_SMALL_NUMBER)
 				{
 					FTransform SphereXf = FTransform(Elem.Center - Pt) * Xf;
-					AnalyticShapeGroup->Add(SphereXf, new Chaos::TSphere<Chaos::FReal, 3>(Chaos::FVec3(0), Radius0));
+					AnalyticShapeGroup->Add(SphereXf, new Chaos::FSphere(Chaos::FVec3(0), Radius0));
 				}
 				if (Radius1 > KINDA_SMALL_NUMBER)
 				{
 					FTransform SphereXf = FTransform(Elem.Center + Pt) * Xf;
-					AnalyticShapeGroup->Add(SphereXf, new Chaos::TSphere<Chaos::FReal, 3>(Chaos::FVec3(0), Radius1));
+					AnalyticShapeGroup->Add(SphereXf, new Chaos::FSphere(Chaos::FVec3(0), Radius1));
 				}
 
 				UE_LOG(USkeletalMeshSimulationComponentLogging, Log,
diff --git a/Engine/Source/Runtime/PhysicsCore/Private/ChaosEngineInterface.cpp b/Engine/Source/Runtime/PhysicsCore/Private/ChaosEngineInterface.cpp
--- a/Engine/Source/Runtime/PhysicsCore/Private/ChaosEngineInterface.cpp
+++ b/Engine/Source/Runtime/PhysicsCore/Private/ChaosEngineInterface.cpp
@@ -85,7 +85,7 @@ const Chaos::TBox<Chaos::FReal,3>& FPhysicsGeometryCollection_Chaos::GetBoxGeome
 
 const Chaos::TSphere<Chaos::FReal,3>&  FPhysicsGeometryCollection_Chaos::GetSphereGeometry() const
 {
-	return GetInnerGeometryChecked<Chaos::TSphere<Chaos::FReal, 3>>(Geom);
+	return GetInnerGeometryChecked<Chaos::FSphere>(Geom);
 }
 const Chaos::FCapsule&  FPhysicsGeometryCollection_Chaos::GetCapsuleGeometry() const
 {
diff --git a/Engine/Source/Runtime/PhysicsCore/Private/ChaosInterfaceWrapperCore.cpp b/Engine/Source/Runtime/PhysicsCore/Private/ChaosInterfaceWrapperCore.cpp
--- a/Engine/Source/Runtime/PhysicsCore/Private/ChaosInterfaceWrapperCore.cpp
+++ b/Engine/Source/Runtime/PhysicsCore/Private/ChaosInterfaceWrapperCore.cpp
@@ -44,12 +44,12 @@ namespace ChaosInterface
 
 	Chaos::FReal GetRadius(const Chaos::FCapsule& InCapsule)
 	{
-		return InCapsule.GetRadius();
+		return InCapsule.GetRadiusf();
 	}
 
 	Chaos::FReal GetHalfHeight(const Chaos::FCapsule& InCapsule)
 	{
-		return InCapsule.GetHeight() / 2.;
+		return InCapsule.GetHeightf() / 2.;
 	}
 
 	FCollisionFilterData GetQueryFilterData(const Chaos::FPerShapeData& Shape)
diff --git a/Engine/Source/Runtime/PhysicsCore/Public/Chaos/ChaosEngineInterface.h b/Engine/Source/Runtime/PhysicsCore/Public/Chaos/ChaosEngineInterface.h
--- a/Engine/Source/Runtime/PhysicsCore/Public/Chaos/ChaosEngineInterface.h
+++ b/Engine/Source/Runtime/PhysicsCore/Public/Chaos/ChaosEngineInterface.h
@@ -220,6 +220,7 @@ namespace Chaos
 
 	template <typename T, int>
 	class TSphere;
+	using FSphere = TSphere<FReal, 3>;
 
 	class FConstraintBase;
 	class FJointConstraint;
@@ -309,7 +310,7 @@ struct FPhysicsGeometryCollection_Chaos
 	PHYSICSCORE_API ECollisionShapeType GetType() const;
 	PHYSICSCORE_API const Chaos::FImplicitObject& GetGeometry() const;
 	PHYSICSCORE_API const Chaos::TBox<Chaos::FReal, 3>& GetBoxGeometry() const;
-	PHYSICSCORE_API const Chaos::TSphere<Chaos::FReal, 3>&  GetSphereGeometry() const;
+	PHYSICSCORE_API const Chaos::FSphere&  GetSphereGeometry() const;
 	PHYSICSCORE_API const Chaos::FCapsule&  GetCapsuleGeometry() const;
 	PHYSICSCORE_API const Chaos::FConvex& GetConvexGeometry() const;
 	PHYSICSCORE_API const Chaos::FTriangleMeshImplicitObject& GetTriMeshGeometry() const;
