Geometry collection : fix issue where the wrong crumble notify flag was set when setting them through SetNotifyGlobalCrumblings blueprint function

diff --git a/Engine/Source/Runtime/Experimental/GeometryCollectionEngine/Private/GeometryCollection/GeometryCollectionComponent.cpp b/Engine/Source/Runtime/Experimental/GeometryCollectionEngine/Private/GeometryCollection/GeometryCollectionComponent.cpp
--- a/Engine/Source/Runtime/Experimental/GeometryCollectionEngine/Private/GeometryCollection/GeometryCollectionComponent.cpp
+++ b/Engine/Source/Runtime/Experimental/GeometryCollectionEngine/Private/GeometryCollection/GeometryCollectionComponent.cpp
@@ -1237,7 +1237,7 @@ void UGeometryCollectionComponent::SetNotifyGlobalCrumblings(bool bNewNotifyGlob
 		{
 			PhysicsProxy->SetNotifyGlobalCrumblings_External(bNewNotifyGlobalCrumblings, bGlobalNewCrumblingEventIncludesChildren);
 		}
-		bNotifyCrumblings = bNewNotifyGlobalCrumblings;
+		bNotifyGlobalCrumblings = bNewNotifyGlobalCrumblings;
 		bGlobalCrumblingEventIncludesChildren = bGlobalNewCrumblingEventIncludesChildren;
 	}
 }
