Dataflow : Add connection type dependency system
- Node base class now can handle the type propagation base on the assigned type dependencies groups
- Still compatible with existing node using the OnInputTypeChanged/ OnOutputTypeChanged overrides ( they can be replaced wit hthe type dependency at a later stage )
- Updated the core flow nodes as an example

diff --git a/Engine/Source/Runtime/Experimental/Dataflow/Core/Private/Dataflow/DataflowConnection.cpp b/Engine/Source/Runtime/Experimental/Dataflow/Core/Private/Dataflow/DataflowConnection.cpp
--- a/Engine/Source/Runtime/Experimental/Dataflow/Core/Private/Dataflow/DataflowConnection.cpp
+++ b/Engine/Source/Runtime/Experimental/Dataflow/Core/Private/Dataflow/DataflowConnection.cpp
@@ -102,6 +102,15 @@ bool FDataflowConnection::SupportsType(FName InType) const
 	return (InType == GetType());
 }
 
+FDataflowConnection& FDataflowConnection::SetTypeDependencyGroup(FName DependencyGroupName)
+{
+	if (IsAnyType() && !HasConcreteType())
+	{
+		TypeDependencyGroup = DependencyGroupName;
+	}
+	return *this;
+}
+
 bool FDataflowConnection::IsExtendedType(FName InType) const
 {
 	return (InType.ToString().StartsWith(Type.ToString() + "<"));
diff --git a/Engine/Source/Runtime/Experimental/Dataflow/Core/Private/Dataflow/DataflowCoreNodes.cpp b/Engine/Source/Runtime/Experimental/Dataflow/Core/Private/Dataflow/DataflowCoreNodes.cpp
--- a/Engine/Source/Runtime/Experimental/Dataflow/Core/Private/Dataflow/DataflowCoreNodes.cpp
+++ b/Engine/Source/Runtime/Experimental/Dataflow/Core/Private/Dataflow/DataflowCoreNodes.cpp
@@ -18,9 +18,13 @@ namespace UE::Dataflow
 FDataflowReRouteNode::FDataflowReRouteNode(const UE::Dataflow::FNodeParameters& Param, FGuid InGuid)
 	: Super(Param, InGuid)
 {
-	RegisterInputConnection(&Value);
+	static const FName MainTypeGroup("Main");
+
+	RegisterInputConnection(&Value)
+		.SetTypeDependencyGroup(MainTypeGroup);
 	RegisterOutputConnection(&Value)
-		.SetPassthroughInput(&Value);
+		.SetPassthroughInput(&Value)
+		.SetTypeDependencyGroup(MainTypeGroup);
 }
 
 void FDataflowReRouteNode::Evaluate(UE::Dataflow::FContext& Context, const FDataflowOutput* Out) const
@@ -28,23 +32,18 @@ void FDataflowReRouteNode::Evaluate(UE::Dataflow::FContext& Context, const FData
 	ForwardInput(Context, &Value, &Value);
 }
 
-bool FDataflowReRouteNode::OnInputTypeChanged(const FDataflowInput* Input)
-{
-	return SetOutputConcreteType(&Value, Input->GetType());
-}
-
-bool FDataflowReRouteNode::OnOutputTypeChanged(const FDataflowOutput* Input)
-{
-	return SetInputConcreteType(&Value, Input->GetType());
-}
-
 FDataflowBranchNode::FDataflowBranchNode(const UE::Dataflow::FNodeParameters& Param, FGuid InGuid)
 	: Super(Param, InGuid)
 {
-	RegisterInputConnection(&TrueValue);
-	RegisterInputConnection(&FalseValue);
+	static const FName MainTypeGroup("Main");
+
+	RegisterInputConnection(&TrueValue)
+		.SetTypeDependencyGroup(MainTypeGroup);
+	RegisterInputConnection(&FalseValue)
+		.SetTypeDependencyGroup(MainTypeGroup);
 	RegisterInputConnection(&bCondition);
-	RegisterOutputConnection(&Result);
+	RegisterOutputConnection(&Result)
+		.SetTypeDependencyGroup(MainTypeGroup);
 }
 
 void FDataflowBranchNode::Evaluate(UE::Dataflow::FContext& Context, const FDataflowOutput* Out) const
@@ -65,26 +64,7 @@ void FDataflowBranchNode::Evaluate(UE::Dataflow::FContext& Context, const FDataf
 	}
 }
 
-bool FDataflowBranchNode::OnInputTypeChanged(const FDataflowInput* Input)
-{
-	// using single bitwise | operator to avoid skipping calls to SetInputConcreteType because of the shortcircuiting of ||
-	// need to disable the warning for static analysis ( V792 warning )
-	return bool(
-		SetInputConcreteType(&TrueValue, Input->GetType())
-	  | SetInputConcreteType(&FalseValue, Input->GetType()) //-V792
-	  | SetOutputConcreteType(&Result, Input->GetType()) //-V792
-		);
-}
-
-bool FDataflowBranchNode::OnOutputTypeChanged(const FDataflowOutput* Input)
-{
-	// using single bitwise | operator to avoid skipping calls to SetInputConcreteType because of the shortcircuiting of ||
-	// need to disable the warning for static analysis ( V792 warning )
-	return bool(
-		  SetInputConcreteType(&TrueValue, Input->GetType())
-		| SetInputConcreteType(&FalseValue, Input->GetType()) //-V792
-		);
-}
+const FName FDataflowSelectNode::MainTypeGroup = "Main";
 
 FDataflowSelectNode::FDataflowSelectNode(const UE::Dataflow::FNodeParameters& Param, FGuid InGuid)
 	: Super(Param, InGuid)
@@ -96,7 +76,8 @@ FDataflowSelectNode::FDataflowSelectNode(const UE::Dataflow::FNodeParameters& Pa
 		AddPins();
 	}
 	RegisterOutputConnection(&Result)
-		.SetPassthroughInput(GetConnectionReference(0));
+		.SetPassthroughInput(GetConnectionReference(0))
+		.SetTypeDependencyGroup(MainTypeGroup);
 	check(NumRequiredDataflowInputs + NumInitialInputs == GetNumInputs()); // Update NumRequiredDataflowInputs when adding more inputs. This is used by Serialize
 }
 
@@ -121,34 +102,15 @@ void FDataflowSelectNode::Evaluate(UE::Dataflow::FContext& Context, const FDataf
 	}
 }
 
-bool FDataflowSelectNode::OnInputTypeChanged(const FDataflowInput* Input)
-{
-	bool bResult = SetOutputConcreteType(&Result, Input->GetType());
-	for (int32 Index = 0; Index < Inputs.Num(); ++Index)
-	{
-		bResult |= SetInputConcreteType(GetConnectionReference(Index), Input->GetType());
-	}
-	return bResult;
-}
-
-bool FDataflowSelectNode::OnOutputTypeChanged(const FDataflowOutput* Input)
-{
-	bool bResult = false;
-	for (int32 Index = 0; Index < Inputs.Num(); ++Index)
-	{
-		bResult |= SetInputConcreteType(GetConnectionReference(Index), Input->GetType());
-	}
-	return bResult;
-}
-
 TArray<UE::Dataflow::FPin> FDataflowSelectNode::AddPins()
 {
 	const int32 Index = Inputs.AddDefaulted();
-	const FDataflowInput& Input = RegisterInputArrayConnection(GetConnectionReference(Index));
+	FDataflowInput& Input = RegisterInputArrayConnection(GetConnectionReference(Index));
+	Input.SetTypeDependencyGroup(MainTypeGroup);
 	if (Index > 0)
 	{
 		// Set concrete type the same as Index0.
-		const FDataflowInput* const Input0 = FindInput(GetConnectionReference(0));
+		FDataflowInput* const Input0 = FindInput(GetConnectionReference(0));
 		check(Input0);
 		SetInputConcreteType(GetConnectionReference(Index), Input0->GetType());
 	}
@@ -193,7 +155,8 @@ void FDataflowSelectNode::PostSerialize(const FArchive& Ar)
 
 		for (int32 Index = NumInitialInputs; Index < Inputs.Num(); ++Index)
 		{
-			FindOrRegisterInputArrayConnection(GetConnectionReference(Index));
+			FindOrRegisterInputArrayConnection(GetConnectionReference(Index))
+				.SetTypeDependencyGroup(MainTypeGroup);
 		}
 
 		if (Ar.IsTransacting())
diff --git a/Engine/Source/Runtime/Experimental/Dataflow/Core/Private/Dataflow/DataflowInputOutput.cpp b/Engine/Source/Runtime/Experimental/Dataflow/Core/Private/Dataflow/DataflowInputOutput.cpp
--- a/Engine/Source/Runtime/Experimental/Dataflow/Core/Private/Dataflow/DataflowInputOutput.cpp
+++ b/Engine/Source/Runtime/Experimental/Dataflow/Core/Private/Dataflow/DataflowInputOutput.cpp
@@ -27,6 +27,11 @@ FDataflowInput::FDataflowInput(const UE::Dataflow::FInputParameters& Param)
 {
 }
 
+bool FDataflowInput::IsConnected() const
+{
+	return (Connection != nullptr);
+}
+
 bool FDataflowInput::AddConnection(FDataflowConnection* InOutput)
 {
 	if (Type == InOutput->GetType())
diff --git a/Engine/Source/Runtime/Experimental/Dataflow/Core/Private/Dataflow/DataflowNode.cpp b/Engine/Source/Runtime/Experimental/Dataflow/Core/Private/Dataflow/DataflowNode.cpp
--- a/Engine/Source/Runtime/Experimental/Dataflow/Core/Private/Dataflow/DataflowNode.cpp
+++ b/Engine/Source/Runtime/Experimental/Dataflow/Core/Private/Dataflow/DataflowNode.cpp
@@ -105,6 +105,36 @@ static TArray<FString> GetPinMetaDataFromProperty(const FProperty* Property)
 }
 };
 
+FDataflowConnection* FDataflowNode::FindConnection(const UE::Dataflow::FConnectionKey& Key)
+{
+	FDataflowConnection* Connection = FindInput(Key);
+	if (!Connection)
+	{
+		Connection = FindOutput(Key);
+	}
+	return Connection;
+}
+
+const FDataflowConnection* FDataflowNode::FindConnection(const UE::Dataflow::FConnectionKey& Key) const
+{
+	const FDataflowConnection* Connection = FindInput(Key);
+	if (!Connection)
+	{
+		Connection = FindOutput(Key);
+	}
+	return Connection;
+}
+
+FDataflowConnection* FDataflowNode::FindConnection(const UE::Dataflow::FConnectionReference& Reference)
+{
+	FDataflowConnection* Connection = FindInput(Reference);
+	if (!Connection)
+	{
+		Connection = FindOutput(Reference);
+	}
+	return Connection;
+}
+
 //
 // Inputs
 //
@@ -1258,8 +1288,9 @@ bool FDataflowNode::TrySetConnectionType(FDataflowConnection* Connection, FName
 
 void FDataflowNode::NotifyConnectionTypeChanged(FDataflowConnection* Connection)
 {
-	if (Connection->IsAnyType())
+	if (Connection && Connection->IsAnyType())
 	{
+		OnConnectionTypeChanged(*Connection);
 		if (Connection->GetDirection() == UE::Dataflow::FPin::EDirection::INPUT)
 		{
 			OnInputTypeChanged((FDataflowInput*)Connection);
@@ -1271,6 +1302,68 @@ void FDataflowNode::NotifyConnectionTypeChanged(FDataflowConnection* Connection)
 	}
 }
 
+bool FDataflowNode::OnConnectionTypeChanged(const FDataflowConnection& Connection)
+{
+	// detect if any of the dependent connections is connected
+	// because we do not allow connected connectio nto change type
+	const FName DependencyGroup = Connection.GetTypeDependencyGroup();
+	if (DependencyGroup.IsNone())
+	{
+		return false; // no changes
+	}
+
+	// first check if any connection is connected 
+	bool bIsAnyDependenciesConnected = false;
+	ForEachConnection(
+		[&bIsAnyDependenciesConnected, &Connection](const FDataflowConnection* OtherConnection)
+		{
+			// ignore the connection that is actively being notified
+			if (OtherConnection 
+				&& (OtherConnection != &Connection)
+				&& OtherConnection->IsConnected()
+				)
+			{
+				bIsAnyDependenciesConnected = true;
+				return false; // stop the iteration
+			}
+			return true;
+		});
+	if (bIsAnyDependenciesConnected)
+	{
+		return false; // no changes
+	}
+
+	// Change all connection that share the same type dependecy group
+	bool bTypeChanged = false;
+	const FName NewType = Connection.GetType();
+	ForEachConnection(
+		[&bTypeChanged, &DependencyGroup, &NewType](FDataflowConnection* OtherConnection)
+		{
+			if (OtherConnection
+				&& (OtherConnection->GetTypeDependencyGroup() == DependencyGroup)
+				&& (OtherConnection->GetType() != NewType))
+			{
+				OtherConnection->SetConcreteType(NewType);
+				bTypeChanged = true;
+			}
+			return true;
+		});
+
+	return bTypeChanged;
+}
+
+bool FDataflowNode::SetConnectionConcreteType(const UE::Dataflow::FConnectionKey& ConnectionKey, FName NewType)
+{
+	if (FDataflowConnection* Connection = FindConnection(ConnectionKey))
+	{
+		if (Connection->GetType() != NewType)
+		{
+			return Connection->SetConcreteType(NewType);
+		}
+	}
+	return false;
+}
+
 bool FDataflowNode::SetInputConcreteType(const UE::Dataflow::FConnectionReference& InputReference, FName NewType)
 {
 	if (FDataflowInput* Input = FindInput(InputReference))
diff --git a/Engine/Source/Runtime/Experimental/Dataflow/Core/Public/Dataflow/DataflowConnection.h b/Engine/Source/Runtime/Experimental/Dataflow/Core/Public/Dataflow/DataflowConnection.h
--- a/Engine/Source/Runtime/Experimental/Dataflow/Core/Public/Dataflow/DataflowConnection.h
+++ b/Engine/Source/Runtime/Experimental/Dataflow/Core/Public/Dataflow/DataflowConnection.h
@@ -145,6 +145,11 @@ protected:
 	bool bCanHidePin:1 = false;
 	bool bPinIsHidden:1 = false;
 
+	// Name of the group dependecy the connection is part of 
+	// This applies to anytype connections
+	// All other connection with the same group will inherit type from each other
+	FName TypeDependencyGroup;
+
 	friend struct FDataflowNode;
 	friend class UE::Dataflow::FGraph;
 
@@ -210,6 +215,7 @@ public:
 	virtual void* RealAddress() const { ensure(OwningNode);  return (void*)((size_t)OwningNode + (size_t)GetOffset()); };
 	UE::Dataflow::FContextCacheKey CacheKey() const { return GetTypeHash(Guid); };
 
+	virtual bool IsConnected() const { return false; }
 	virtual bool AddConnection(FDataflowConnection* In) { return false; };
 	virtual bool RemoveConnection(FDataflowConnection* In) { return false; }
 
@@ -219,6 +225,9 @@ public:
 	DATAFLOWCORE_API void SetTypePolicy(IDataflowTypePolicy* InTypePolicy);
 	DATAFLOWCORE_API bool SupportsType(FName InType) const;
 
+	DATAFLOWCORE_API FDataflowConnection& SetTypeDependencyGroup(FName DependencyGroupName);
+	DATAFLOWCORE_API FName GetTypeDependencyGroup() const { return TypeDependencyGroup;	}
+
 	template<class T>
 	bool IsA(const T* InVar) const
 	{
diff --git a/Engine/Source/Runtime/Experimental/Dataflow/Core/Public/Dataflow/DataflowCoreNodes.h b/Engine/Source/Runtime/Experimental/Dataflow/Core/Public/Dataflow/DataflowCoreNodes.h
--- a/Engine/Source/Runtime/Experimental/Dataflow/Core/Public/Dataflow/DataflowCoreNodes.h
+++ b/Engine/Source/Runtime/Experimental/Dataflow/Core/Public/Dataflow/DataflowCoreNodes.h
@@ -27,11 +27,6 @@ public:
 public:
 	UPROPERTY(meta = (DataflowInput, DataflowOutput, DataflowPassthrough = "Value", DisplayName = "Value"))
 	FDataflowAnyType Value;
-
-protected:
-	virtual bool OnInputTypeChanged(const FDataflowInput* Input) override;
-	virtual bool OnOutputTypeChanged(const FDataflowOutput* Input) override;
-
 };
 
 USTRUCT(meta=(Icon="GraphEditor.Branch_16x"))
@@ -57,10 +52,6 @@ public:
 
 	UPROPERTY(meta = (DataflowOutput, DisplayName = "Result"))
 	FDataflowAnyType Result;
-
-private:
-	virtual bool OnInputTypeChanged(const FDataflowInput* Input) override;
-	virtual bool OnOutputTypeChanged(const FDataflowOutput* Input) override;
 };
 
 
@@ -86,8 +77,6 @@ public:
 	FDataflowAnyType Result;
 
 private:
-	virtual bool OnInputTypeChanged(const FDataflowInput* Input) override;
-	virtual bool OnOutputTypeChanged(const FDataflowOutput* Input) override;
 	virtual TArray<UE::Dataflow::FPin> AddPins() override;
 	virtual bool CanAddPin() const override { return true; }
 	virtual bool CanRemovePin() const override { return Inputs.Num() > NumInitialInputs; }
@@ -99,6 +88,7 @@ private:
 
 	static constexpr int32 NumRequiredDataflowInputs = 1;
 	static constexpr int32 NumInitialInputs = 2;
+	static const FName MainTypeGroup;
 };
 
 /** 
diff --git a/Engine/Source/Runtime/Experimental/Dataflow/Core/Public/Dataflow/DataflowInputOutput.h b/Engine/Source/Runtime/Experimental/Dataflow/Core/Public/Dataflow/DataflowInputOutput.h
--- a/Engine/Source/Runtime/Experimental/Dataflow/Core/Public/Dataflow/DataflowInputOutput.h
+++ b/Engine/Source/Runtime/Experimental/Dataflow/Core/Public/Dataflow/DataflowInputOutput.h
@@ -57,6 +57,7 @@ public:
 
 	DATAFLOWCORE_API FDataflowInput(const UE::Dataflow::FInputParameters& Param = {});
 
+	virtual bool IsConnected() const override;
 	virtual bool AddConnection(FDataflowConnection* InOutput) override;
 	virtual bool RemoveConnection(FDataflowConnection* InOutput) override;
 
@@ -160,6 +161,7 @@ public:
 	DATAFLOWCORE_API TArray<FDataflowInput*>& GetConnections();
 	DATAFLOWCORE_API const TArray<FDataflowInput*>& GetConnections() const;
 	bool HasAnyConnections() const { return !Connections.IsEmpty(); }
+	virtual bool IsConnected() const override { return HasAnyConnections(); }
 
 	DATAFLOWCORE_API virtual TArray<FDataflowInput*> GetConnectedInputs();
 	DATAFLOWCORE_API virtual const TArray<const FDataflowInput*> GetConnectedInputs() const;
diff --git a/Engine/Source/Runtime/Experimental/Dataflow/Core/Public/Dataflow/DataflowNode.h b/Engine/Source/Runtime/Experimental/Dataflow/Core/Public/Dataflow/DataflowNode.h
--- a/Engine/Source/Runtime/Experimental/Dataflow/Core/Public/Dataflow/DataflowNode.h
+++ b/Engine/Source/Runtime/Experimental/Dataflow/Core/Public/Dataflow/DataflowNode.h
@@ -164,6 +164,10 @@ PRAGMA_ENABLE_DEPRECATION_WARNINGS
 	DATAFLOWCORE_API bool HasHideableInputs() const;
 	DATAFLOWCORE_API bool HasHiddenInputs() const;
 
+	DATAFLOWCORE_API const FDataflowConnection* FindConnection(const UE::Dataflow::FConnectionKey& Key) const;
+	DATAFLOWCORE_API FDataflowConnection* FindConnection(const UE::Dataflow::FConnectionKey& Key);
+	DATAFLOWCORE_API FDataflowConnection* FindConnection(const UE::Dataflow::FConnectionReference& Reference);
+
 	DATAFLOWCORE_API FDataflowInput* FindInput(FName Name);
 	DATAFLOWCORE_API FDataflowInput* FindInput(const UE::Dataflow::FConnectionKey& Key);
 	/**This version can find array inputs if only the Reference is supplied by searching through all inputs */
@@ -472,6 +476,31 @@ PRAGMA_ENABLE_DEPRECATION_WARNINGS
 		}
 	}
 
+	/*
+	* Go through all connections and execute a function
+	* Signature of the function is bool Function(FDataflowConnection* Connection)
+	*/
+	template <typename TFunction>
+	void ForEachConnection(TFunction Function)
+	{
+		for (TPair<UE::Dataflow::FConnectionKey, FDataflowInput*>& Elem : ExpandedInputs)
+		{
+			if (Function(Elem.Value) == false)
+			{
+				return;
+			}
+		}
+		PRAGMA_DISABLE_DEPRECATION_WARNINGS  // Until Outputs becomes private
+		for (const TPair<int32, FDataflowOutput*>& Elem : Outputs)
+		PRAGMA_ENABLE_DEPRECATION_WARNINGS
+		{
+			if (Function(Elem.Value) == false)
+			{
+				return;
+			}
+		}
+	}
+
 	/**
 	*   IsConnected(...)
 	*
@@ -566,7 +595,8 @@ PRAGMA_ENABLE_DEPRECATION_WARNINGS
 	DATAFLOWCORE_API bool SetOutputConcreteType(const UE::Dataflow::FConnectionReference& OutputReference, FName NewType);
 
 protected:
-	virtual bool OnInputTypeChanged(const FDataflowInput* Input) { return false; };
+	bool OnConnectionTypeChanged(const FDataflowConnection& Connection);
+	virtual bool OnInputTypeChanged(const FDataflowInput* Input) { return false; }
 	virtual bool OnOutputTypeChanged(const FDataflowOutput* Output) { return false; }
 
 	template <typename T>
@@ -581,6 +611,13 @@ protected:
 		return SetOutputConcreteType(OutputReference, TDataflowPolicyTypeName<T>::GetName());
 	}
 
+	/**
+	* Set a connection identified by key to be of a specific type
+	* If the connection is already connected, the method prevents the connection from changing and returns false
+	* returns true if the connection has successfully changed type
+	* */
+	DATAFLOWCORE_API bool SetConnectionConcreteType(const UE::Dataflow::FConnectionKey& ConnectionKey, FName NewType);
+
 	// returns true if any of the types was changed successfully
 	DATAFLOWCORE_API bool SetAllConnectionConcreteType(FName NewType);
 
