[ChaosVD] Added a new macro and required exports to allow implementing custom CVD data channels, outside the chaos Module

diff --git a/Engine/Source/Runtime/Experimental/Chaos/Public/ChaosVisualDebugger/ChaosVDOptionalDataChannel.h b/Engine/Source/Runtime/Experimental/Chaos/Public/ChaosVisualDebugger/ChaosVDOptionalDataChannel.h
--- a/Engine/Source/Runtime/Experimental/Chaos/Public/ChaosVisualDebugger/ChaosVDOptionalDataChannel.h
+++ b/Engine/Source/Runtime/Experimental/Chaos/Public/ChaosVisualDebugger/ChaosVDOptionalDataChannel.h
@@ -12,9 +12,9 @@ class FName;
 /** Set of flags to control how a CVD Data Channel gets initialized*/
 enum class EChaosVDDataChannelInitializationFlags : uint8
 {
-	/** If Set, the data channel will start in a enabled state */
+	/** If Set, the data channel will start in an enabled state */
 	StartEnabled = 1 << 0,
-	/** If Set, the data channel can be toggle on and off as desired */
+	/** If Set, the data channel can be toggled on and off as desired */
 	CanChangeEnabledState = 1 << 1
 };
 ENUM_CLASS_FLAGS(EChaosVDDataChannelInitializationFlags)
@@ -38,16 +38,25 @@ namespace Chaos::VisualDebugger
 		CHAOS_API const FText& GetDisplayName() const;
 
 		/** A name used as ID to find this changed using the CVD Channel Manager */
-		FName GetId() const { return *ChannelId; }
+		FName GetId() const
+		{
+			return *ChannelId;
+		}
 
 		/** Returns true if this data channel is enabled */
-		FORCEINLINE bool IsChannelEnabled() const { return bIsEnabled; }
+		FORCEINLINE bool IsChannelEnabled() const 
+		{ 
+			return bIsEnabled;
+		}
 
 		/** Enables or disabled this Data Channel */
 		CHAOS_API void SetChannelEnabled(bool bNewEnabled);
 
 		/** Returns true if the enabled state of this channel can be changed */
-		bool CanChangeEnabledState() const { return bCanChangeEnabledState;}
+		bool CanChangeEnabledState() const 
+		{ 
+			return bCanChangeEnabledState;
+		}
 
 	private:
 		TSharedRef<FText> LocalizableChannelName;
@@ -56,14 +65,14 @@ namespace Chaos::VisualDebugger
 		bool bCanChangeEnabledState = false;
 	};
 
-	/** Manager that provides a way to iterate trough all existing CVD Data channels or get a specific one */
+	/** Manager that provides a way to iterate through all existing CVD Data channels or get a specific one */
 	class CHAOS_API FChaosVDDataChannelsManager
 	{
 	public:
 
 		static FChaosVDDataChannelsManager& Get();
 
-		/** Iterates trough all the available data channels and executes the provided callback passing each channel as an argument.
+		/** Iterates through all the available data channels and executes the provided callback passing each channel as an argument.
 		 * If the callback returns false, the iteration will be stopped.
 		 */
 		template<typename TCallback>
@@ -105,7 +114,7 @@ namespace Chaos::VisualDebugger
 	void ParseChannelListFromCommandArgument(TArray<FString>& OutParsedChannelList, const FString& InCommandArgument);
 
 	/** Creates a CVD Data Channel Instance - Only to be used by the CVD Macros */
-	TSharedRef<FChaosVDOptionalDataChannel> CreateDataChannel(FName InChannelID, const TSharedRef<FText>& InDisplayName, EChaosVDDataChannelInitializationFlags InitializationFlags);
+	CHAOS_API TSharedRef<FChaosVDOptionalDataChannel> CreateDataChannel(FName InChannelID, const TSharedRef<FText>& InDisplayName, EChaosVDDataChannelInitializationFlags InitializationFlags);
 }
 
 #endif //WITH_CHAOS_VISUAL_DEBUGGER
@@ -131,15 +140,22 @@ namespace Chaos::VisualDebugger
 /** Declares an Optional CVD Data channel to be available globally. The data channel can be accessed by using CVDDC_TheNameOfTheChannelUsedWithThisMacro */
 #ifndef CVD_DECLARE_OPTIONAL_DATA_CHANNEL
 	#define CVD_DECLARE_OPTIONAL_DATA_CHANNEL(DataChannelName) \
-				extern CHAOS_API TSharedRef<Chaos::VisualDebugger::FChaosVDOptionalDataChannel> CVDDC_##DataChannelName;
+			extern CHAOS_API TSharedRef<Chaos::VisualDebugger::FChaosVDOptionalDataChannel> CVDDC_##DataChannelName;
+#endif
+
+/** Declares an Optional CVD Data channel outside the Chaos Module to be available globally. The data channel can be accessed by using CVDDC_TheNameOfTheChannelUsedWithThisMacro */
+#ifndef CVD_DECLARE_OPTIONAL_DATA_CHANNEL_EXTERN
+	#define CVD_DECLARE_OPTIONAL_DATA_CHANNEL_EXTERN(DataChannelName, ModuleAPI) \
+		extern ModuleAPI TSharedRef<Chaos::VisualDebugger::FChaosVDOptionalDataChannel> CVDDC_##DataChannelName;
 #endif
 
 /** Defines and initializes an Optional CVD Data Channel */
 #ifndef CVD_DEFINE_OPTIONAL_DATA_CHANNEL
 	#define CVD_DEFINE_OPTIONAL_DATA_CHANNEL(DataChannelName, InitializationFlags) \
-				TSharedRef<Chaos::VisualDebugger::FChaosVDOptionalDataChannel> CVDDC_##DataChannelName = Chaos::VisualDebugger::CreateDataChannel(#DataChannelName, MakeShared<FText>(NSLOCTEXT(CVD_STRINGIZE(ChaosVisualDebugger), CVD_STRINGIZE(CVD_CONCAT(DataChannelName,_ChannelName)), #DataChannelName)), InitializationFlags);
+			TSharedRef<Chaos::VisualDebugger::FChaosVDOptionalDataChannel> CVDDC_##DataChannelName = Chaos::VisualDebugger::CreateDataChannel(#DataChannelName, MakeShared<FText>(NSLOCTEXT(CVD_STRINGIZE(ChaosVisualDebugger), CVD_STRINGIZE(CVD_CONCAT(DataChannelName,_ChannelName)), #DataChannelName)), InitializationFlags);
 #endif
 
+
 // Declare CVD's Default set of channels
 CVD_DECLARE_OPTIONAL_DATA_CHANNEL(Default);
 CVD_DECLARE_OPTIONAL_DATA_CHANNEL(EvolutionStart);							// Initial particle positions for this tick. Same as end of previous EvolutionEnd but with any external changes applied (set position, velocity, etc).
@@ -167,6 +183,11 @@ CVD_DECLARE_OPTIONAL_DATA_CHANNEL(ClusterParticlesChildData);
 	#define CVD_DECLARE_OPTIONAL_DATA_CHANNEL(DataChannelName)
 #endif
 
+/** Declares an Optional CVD Data channel outside the Chaos Module to be available globally. The data channel can be accessed by using CVDDC_TheNameOfTheChannelUsedWithThisMacro */
+#ifndef CVD_DECLARE_OPTIONAL_DATA_CHANNEL_EXTERN
+	#define CVD_DECLARE_OPTIONAL_DATA_CHANNEL_EXTERN(DataChannelName, ModuleAPI)
+#endif
+
 /** Defines and initializes an Optional CVD Data Channel */
 #ifndef CVD_DEFINE_OPTIONAL_DATA_CHANNEL
 	#define CVD_DEFINE_OPTIONAL_DATA_CHANNEL(DataChannelName, InitializationFlags)
