Chaos Cloth: Add option to use a less optimal vector gather in the ISPC method for TriangleMesh::GetFaceNormals.

We are seeing a crash in this method during the VectorGather that is caused by something going wrong in ISPC that does not occur in the C++ version of this method. It has been proposed that the optimized load + shuffle is reading out of bounds in a way that doesn't occur when doing the slower gather method, so we'd like to try switching to the gather to see if that has an effect.

The switch is controlled by the new cvar p.Chaos.TriangleMesh.ISPCUseHWGather. The default value is to continue using the original load + shuffle.

diff --git a/Engine/Source/Runtime/Experimental/Chaos/Private/Chaos/TriangleMesh.cpp b/Engine/Source/Runtime/Experimental/Chaos/Private/Chaos/TriangleMesh.cpp
--- a/Engine/Source/Runtime/Experimental/Chaos/Private/Chaos/TriangleMesh.cpp
+++ b/Engine/Source/Runtime/Experimental/Chaos/Private/Chaos/TriangleMesh.cpp
@@ -38,6 +38,10 @@ static constexpr bool bChaos_TriangleMesh_ISPC_Enabled = INTEL_ISPC && CHAOS_TRI
 #else
 static bool bChaos_TriangleMesh_ISPC_Enabled = CHAOS_TRIANGLE_MESH_ISPC_ENABLED_DEFAULT;
 static FAutoConsoleVariableRef CVarChaosTriangleMeshISPCEnabled(TEXT("p.Chaos.TriangleMesh.ISPC"), bChaos_TriangleMesh_ISPC_Enabled, TEXT("Whether to use ISPC optimizations in triangle mesh calculations"));
+
+static bool bChaos_TriangleMesh_ISPCUseHWGather = false;
+static FAutoConsoleVariableRef CVarChaosTriangleMeshISPCUseHWGather(TEXT("p.Chaos.TriangleMesh.ISPCUseHWGather"), bChaos_TriangleMesh_ISPCUseHWGather, TEXT("Use (slower) HW gather when computing triangle mesh normals to see if this fixes the ISPC crash."));
+
 #endif
 
 namespace Chaos
@@ -429,11 +433,23 @@ void FTriangleMesh::GetFaceNormals(TArray<TVec3<T>>& Normals, const TConstArrayV
 		if (bChaos_TriangleMesh_ISPC_Enabled && std::is_same_v<T, FRealSingle>)
 		{
 			Normals.SetNumUninitialized(MElements.Num());
-			ispc::GetFaceNormals(
-				(ispc::FVector3f*)Normals.GetData(),
-				(ispc::FVector3f*)Points.GetData(),
-				(ispc::FIntVector*)MElements.GetData(),
-				MElements.Num());
+			if (bChaos_TriangleMesh_ISPCUseHWGather)
+			{
+				ispc::GetFaceNormalsHWGather(
+					(ispc::FVector3f*)Normals.GetData(),
+					(ispc::FVector3f*)Points.GetData(),
+					(ispc::FIntVector*)MElements.GetData(),
+					MElements.Num());
+
+			}
+			else
+			{
+				ispc::GetFaceNormals(
+					(ispc::FVector3f*)Normals.GetData(),
+					(ispc::FVector3f*)Points.GetData(),
+					(ispc::FIntVector*)MElements.GetData(),
+					MElements.Num());
+			}
 		}
 		else
 #endif
diff --git a/Engine/Source/Runtime/Experimental/Chaos/Private/Chaos/TriangleMesh.ispc b/Engine/Source/Runtime/Experimental/Chaos/Private/Chaos/TriangleMesh.ispc
--- a/Engine/Source/Runtime/Experimental/Chaos/Private/Chaos/TriangleMesh.ispc
+++ b/Engine/Source/Runtime/Experimental/Chaos/Private/Chaos/TriangleMesh.ispc
@@ -26,6 +26,31 @@ export void GetFaceNormals(uniform FVector3f Normals[],
 	}
 }
 
+export void GetFaceNormalsHWGather(uniform FVector3f Normals[],
+							const uniform FVector3f Points[],
+							const uniform FIntVector Elements[],
+							const uniform int32 NumElements)
+{
+	foreach(i = 0 ... NumElements)
+	{
+		const varying FIntVector Tri = VectorLoad(&Elements[extract(i, 0)]);
+		
+		#pragma ignore warning(perf)
+		const varying FVector3f p0 = Points[Tri.V[0]];
+		#pragma ignore warning(perf)
+		const varying FVector3f p1 = Points[Tri.V[1]];
+		#pragma ignore warning(perf)
+		const varying FVector3f p2 = Points[Tri.V[2]];
+
+		const varying FVector3f p10 = p1 - p0;
+		const varying FVector3f p20 = p2 - p0;
+		const varying FVector3f Cross = VectorCross(p20, p10);
+		const varying FVector3f CrossNormalized = VectorGetSafeNormal(Cross);
+
+		VectorStore(&Normals[extract(i, 0)], CrossNormalized);
+	}
+}
+
 struct TArrayInt
 {
 	int32 *Data;
