[LLM] Move Filesystem allocations to separate child LLM tags

#rnx
#tests local cooked build run

diff --git a/Engine/Source/Runtime/Core/Private/IO/IoDispatcher.cpp b/Engine/Source/Runtime/Core/Private/IO/IoDispatcher.cpp
--- a/Engine/Source/Runtime/Core/Private/IO/IoDispatcher.cpp
+++ b/Engine/Source/Runtime/Core/Private/IO/IoDispatcher.cpp
@@ -194,7 +194,7 @@ public:
 
 		if (!NextFree)
 		{
-			LLM_SCOPE(ELLMTag::FileSystem);
+			LLM_SCOPE_BYNAME(TEXT("FileSystem/IODispatcher"));
 			//TODO: Virtual alloc
 			FBlock* Block = new FBlock;
 
@@ -388,7 +388,7 @@ public:
 
 	FIoBatchImpl* AllocBatch()
 	{
-		LLM_SCOPE(ELLMTag::FileSystem);
+		LLM_SCOPE_BYNAME(TEXT("FileSystem/IODispatcher"));
 		FIoBatchImpl* Batch = BatchAllocator.Construct();
 
 		return Batch;
@@ -601,7 +601,7 @@ public:
 		}
 		else
 		{
-			LLM_SCOPE(ELLMTag::FileSystem);
+			LLM_SCOPE_BYNAME(TEXT("FileSystem/IODispatcher"));
 			ProcessIncomingRequests();
 			while (PendingIoRequestsCount > 0)
 			{
@@ -859,7 +859,7 @@ private:
 	{
 		FMemory::SetupTLSCachesOnCurrentThread();
 
-		LLM_SCOPE(ELLMTag::FileSystem);
+		LLM_SCOPE_BYNAME(TEXT("FileSystem/IODispatcher"));
 		while (!bStopRequested)
 		{
 			if (PendingIoRequestsCount)
@@ -987,7 +987,7 @@ FIoDispatcher::IsInitialized()
 FIoStatus
 FIoDispatcher::Initialize()
 {
-	LLM_SCOPE(ELLMTag::FileSystem);
+	LLM_SCOPE_BYNAME(TEXT("FileSystem/IODispatcher"));
 	check(!GIoDispatcher);
 	GIoDispatcher = MakeUnique<FIoDispatcher>();
 	return FIoStatus::Ok;
@@ -996,7 +996,7 @@ FIoDispatcher::Initialize()
 void
 FIoDispatcher::InitializePostSettings()
 {
-	LLM_SCOPE(ELLMTag::FileSystem);
+	LLM_SCOPE_BYNAME(TEXT("FileSystem/IODispatcher"));
 	check(GIoDispatcher);
 	GIoDispatcher->Impl->Initialize();
 }
diff --git a/Engine/Source/Runtime/Experimental/IoStore/OnDemand/Private/OnDemandIoStore.cpp b/Engine/Source/Runtime/Experimental/IoStore/OnDemand/Private/OnDemandIoStore.cpp
--- a/Engine/Source/Runtime/Experimental/IoStore/OnDemand/Private/OnDemandIoStore.cpp
+++ b/Engine/Source/Runtime/Experimental/IoStore/OnDemand/Private/OnDemandIoStore.cpp
@@ -498,7 +498,7 @@ FOnDemandIoStore::~FOnDemandIoStore()
 
 FIoStatus FOnDemandIoStore::Initialize()
 {
-	LLM_SCOPE(ELLMTag::FileSystem);
+	LLM_SCOPE_BYNAME(TEXT("FileSystem/OnDemand"));
 	const FIoStatus CacheStatus = InitializeOnDemandInstallCache();
 	if (CacheStatus.GetErrorCode() == EIoErrorCode::PendingFork)
 	{
@@ -615,7 +615,7 @@ void FOnDemandIoStore::Install(
 	FOnDemandInstallProgressed&& OnProgress /*= nullptr*/,
 	const FOnDemandCancellationToken* CancellationToken)
 {
-	LLM_SCOPE(ELLMTag::FileSystem);
+	LLM_SCOPE_BYNAME(TEXT("FileSystem/OnDemand"));
 	FSharedInstallRequest InstallRequest = MakeShared<FInstallRequest>();
 	InstallRequest->Args				= MoveTemp(Args);
 	InstallRequest->OnCompleted			= MoveTemp(OnCompleted);
@@ -632,7 +632,7 @@ void FOnDemandIoStore::Install(
 
 void FOnDemandIoStore::Purge(FOnDemandPurgeArgs&& Args, FOnDemandPurgeCompleted&& OnCompleted)
 {
-	LLM_SCOPE(ELLMTag::FileSystem);
+	LLM_SCOPE_BYNAME(TEXT("FileSystem/OnDemand"));
 	FSharedPurgeRequest PurgeRequest = MakeShared<FPurgeRequest>();
 	PurgeRequest->Args			= MoveTemp(Args);
 	PurgeRequest->OnCompleted	= MoveTemp(OnCompleted);
@@ -647,7 +647,7 @@ void FOnDemandIoStore::Purge(FOnDemandPurgeArgs&& Args, FOnDemandPurgeCompleted&
 
 FIoStatus FOnDemandIoStore::Unmount(FStringView MountId)
 {
-	LLM_SCOPE(ELLMTag::FileSystem);
+	LLM_SCOPE_BYNAME(TEXT("FileSystem/OnDemand"));
 	UE_LOG(LogIoStoreOnDemand, Log, TEXT("Unmounting '%s'"), *WriteToString<256>(MountId));
 
 	bool bPendingMount = false;
@@ -955,7 +955,7 @@ void FOnDemandIoStore::TryEnterTickLoop()
 
 void FOnDemandIoStore::TickLoop()
 {
-	LLM_SCOPE(ELLMTag::FileSystem);
+	LLM_SCOPE_BYNAME(TEXT("FileSystem/OnDemand"));
 	ON_SCOPE_EXIT { UE_LOG(LogIoStoreOnDemand, Verbose, TEXT("Exiting I/O store tick loop")); };
 
 	UE_LOG(LogIoStoreOnDemand, Verbose, TEXT("Entering I/O store tick loop"));
@@ -1706,7 +1706,7 @@ void FOnDemandIoStore::CreateContainersFromToc(
 	TArray<FSharedOnDemandContainer>& Out)
 {
 	static FName AssetClassName(TEXT("OnDemandIoStore"));
-	LLM_SCOPE(ELLMTag::FileSystem);
+	LLM_SCOPE_BYNAME(TEXT("FileSystem/OnDemand"));
 	LLM_TAGSET_SCOPE(FName(MountId), ELLMTagSet::Assets);
 	LLM_TAGSET_SCOPE(AssetClassName, ELLMTagSet::AssetClasses);
 	UE_TRACE_METADATA_SCOPE_ASSET_FNAME(FName(MountId), AssetClassName, FName(TocPath));
@@ -1935,7 +1935,7 @@ FIoStatus FOnDemandIoStore::GetContainersAndPackagesForInstall(
 
 void FOnDemandIoStore::ReleaseContent(FOnDemandInternalContentHandle& ContentHandle)
 {
-	LLM_SCOPE(ELLMTag::FileSystem);
+	LLM_SCOPE_BYNAME(TEXT("FileSystem/OnDemand"));
 	UE_LOG(LogIoStoreOnDemand, Log, TEXT("Releasing content handle '%s'"), *LexToString(ContentHandle));
 
 	UE::TUniqueLock Lock(ContainerMutex);
diff --git a/Engine/Source/Runtime/PakFile/Private/IPlatformFilePak.cpp b/Engine/Source/Runtime/PakFile/Private/IPlatformFilePak.cpp
--- a/Engine/Source/Runtime/PakFile/Private/IPlatformFilePak.cpp
+++ b/Engine/Source/Runtime/PakFile/Private/IPlatformFilePak.cpp
@@ -2183,7 +2183,7 @@ public:
 
 	void NewRequestsToLowerComplete(bool bWasCanceled, IAsyncReadRequest* Request, int32 Index)
 	{
-		LLM_SCOPE(ELLMTag::FileSystem);
+		LLM_SCOPE_BYNAME(TEXT("FileSystem/PakFile"));
 		ClearOldBlockTasks();
 
 		FScopeLock Lock(&CachedFilesScopeLock);
@@ -3324,7 +3324,7 @@ public:
 	}
 	virtual IAsyncReadRequest* ReadRequest(int64 Offset, int64 BytesToRead, EAsyncIOPriorityAndFlags PriorityAndFlags = AIOP_Normal, FAsyncFileCallBack* CompleteCallback = nullptr, uint8* UserSuppliedMemory = nullptr) override
 	{
-		LLM_SCOPE(ELLMTag::FileSystem);
+		LLM_SCOPE_BYNAME(TEXT("FileSystem/PakFile"));
 
 		if (BytesToRead == MAX_int64)
 		{
@@ -5549,7 +5549,7 @@ bool FPakPlatformFile::Initialize(IPlatformFile* Inner, const TCHAR* CmdLine)
 {
 	UE_LOG(LogPakFile, Log, TEXT("Initializing PakPlatformFile"));
 
-	LLM_SCOPE(ELLMTag::FileSystem);
+	LLM_SCOPE_BYNAME(TEXT("FileSystem/PakFile"));
 	SCOPED_BOOT_TIMING("FPakPlatformFile::Initialize");
 	// Inner is required.
 	check(Inner != NULL);
@@ -5813,7 +5813,7 @@ void FPakPlatformFile::OptimizeMemoryUsageForMountedPaks()
 
 bool FPakPlatformFile::Mount(const TCHAR* InPakFilename, uint32 PakOrder, const TCHAR* InPath /*= nullptr*/, bool bLoadIndex /*= true*/, FPakListEntry* OutPakListEntry /*= nullptr*/)
 {
-	LLM_SCOPE(ELLMTag::FileSystem);
+	LLM_SCOPE_BYNAME(TEXT("FileSystem/PakFile"));
 	bool bPakSuccess = false;
 	bool bIoStoreSuccess = true;
 	if (LowerLevel->FileExists(InPakFilename))
@@ -6303,7 +6303,7 @@ void FPakPlatformFile::RegisterEncryptionKey(const FGuid& InGuid, const FAES::FA
 		PendingEncryptedPakFiles.RemoveAll([InGuid](const FPakListDeferredEntry& Entry) { return Entry.EncryptionKeyGuid == InGuid; });
 
 		{
-			LLM_SCOPE(ELLMTag::FileSystem);
+			LLM_SCOPE_BYNAME(TEXT("FileSystem/PakFile"));
 			OptimizeMemoryUsageForMountedPaks();
 		}
 
diff --git a/Engine/Source/Runtime/PakFile/Private/IoDispatcherFileBackend.cpp b/Engine/Source/Runtime/PakFile/Private/IoDispatcherFileBackend.cpp
--- a/Engine/Source/Runtime/PakFile/Private/IoDispatcherFileBackend.cpp
+++ b/Engine/Source/Runtime/PakFile/Private/IoDispatcherFileBackend.cpp
@@ -1365,6 +1365,7 @@ void FFileIoStore::Shutdown()
 
 TIoStatusOr<FIoContainerHeader> FFileIoStore::Mount(const TCHAR* InTocPath, int32 Order, const FGuid& EncryptionKeyGuid, const FAES::FAESKey& EncryptionKey)
 {
+	LLM_SCOPE_BYNAME(TEXT("FileSystem/IOStore"));
 	TUniquePtr<FFileIoStoreReader> Reader(new FFileIoStoreReader(*PlatformImpl, Stats));
 	FIoStatus IoStatus = Reader->Initialize(InTocPath, Order);
 	if (!IoStatus.IsOk())
@@ -1600,7 +1601,7 @@ FAutoConsoleTaskPriority CPrio_IoDispatcherTaskPriority(
 
 void FFileIoStore::ScatterBlock(FFileIoStoreCompressedBlock* CompressedBlock, bool bIsAsync)
 {
-	LLM_SCOPE(ELLMTag::FileSystem);
+	LLM_SCOPE_BYNAME(TEXT("FileSystem/IOStore"));
 	TRACE_CPUPROFILER_EVENT_SCOPE(IoDispatcherScatter);
 	
 	check(!CompressedBlock->bFailed);
@@ -1793,7 +1794,7 @@ namespace FileIoStoreImpl
 
 FIoRequestImpl* FFileIoStore::GetCompletedIoRequests()
 {
-	LLM_SCOPE(ELLMTag::FileSystem);
+	LLM_SCOPE_BYNAME(TEXT("FileSystem/IOStore"));
 	//TRACE_CPUPROFILER_EVENT_SCOPE(GetCompletedRequests);
 	
 	if (!bIsMultithreaded)
diff --git a/Engine/Source/Runtime/PakFile/Private/PakFile.cpp b/Engine/Source/Runtime/PakFile/Private/PakFile.cpp
--- a/Engine/Source/Runtime/PakFile/Private/PakFile.cpp
+++ b/Engine/Source/Runtime/PakFile/Private/PakFile.cpp
@@ -265,7 +265,7 @@ void FPakFile::Initialize(FArchive& Reader, bool bLoadIndex)
 	bool bShouldLoad = false;
 	int32 CompatibleVersion = FPakInfo::PakFile_Version_Latest;
 
-	LLM_SCOPE(ELLMTag::FileSystem);
+	LLM_SCOPE_BYNAME(TEXT("FileSystem/PakFile"));
 
 	// Serialize trailer and check if everything is as expected.
 	// start up one to offset the -- below
diff --git a/Engine/Source/Runtime/StorageServerClient/Private/StorageServerIoDispatcherBackend.cpp b/Engine/Source/Runtime/StorageServerClient/Private/StorageServerIoDispatcherBackend.cpp
--- a/Engine/Source/Runtime/StorageServerClient/Private/StorageServerIoDispatcherBackend.cpp
+++ b/Engine/Source/Runtime/StorageServerClient/Private/StorageServerIoDispatcherBackend.cpp
@@ -68,7 +68,7 @@ void FStorageServerIoDispatcherBackend::Stop()
 
 uint32 FStorageServerIoDispatcherBackend::Run()
 {
-	LLM_SCOPE(ELLMTag::FileSystem);
+	LLM_SCOPE_BYNAME(TEXT("FileSystem/StorageServer"));
 	const int32 BatchCount = GStorageServerIoDispatcherMaxActiveBatchCount;
 	for (int32 BatchIndex = 0; BatchIndex < BatchCount; ++BatchIndex)
 	{
@@ -302,7 +302,7 @@ FStorageServerIoDispatcherBackend::FBatch::FBatch(FStorageServerIoDispatcherBack
 
 void FStorageServerIoDispatcherBackend::FBatch::DoThreadedWork()
 {
-	LLM_SCOPE(ELLMTag::FileSystem);
+	LLM_SCOPE_BYNAME(TEXT("FileSystem/StorageServer"));
 	TRACE_CPUPROFILER_EVENT_SCOPE(StorageServerIoDispatcherProcessBatch);
 	FIoRequestImpl* Request = RequestsHead;
 #if 0
