Fixed non unity/pch compile errors

diff --git a/Engine/Source/Runtime/Engine/Private/Components/SkeletalMeshComponent.cpp b/Engine/Source/Runtime/Engine/Private/Components/SkeletalMeshComponent.cpp
--- a/Engine/Source/Runtime/Engine/Private/Components/SkeletalMeshComponent.cpp
+++ b/Engine/Source/Runtime/Engine/Private/Components/SkeletalMeshComponent.cpp
@@ -11,6 +11,7 @@
 #include "Animation/AnimSingleNodeInstance.h"
 #include "Animation/AnimationSettings.h"
 #include "Animation/PoseSnapshot.h"
+#include "Containers/Ticker.h"
 #include "Engine/SkeletalMesh.h"
 #include "Engine/SkeletalMeshSocket.h"
 #include "AI/NavigationSystemHelpers.h"
diff --git a/Engine/Source/Runtime/Experimental/Chaos/Private/Chaos/PerParticleDampVelocity.cpp b/Engine/Source/Runtime/Experimental/Chaos/Private/Chaos/PerParticleDampVelocity.cpp
--- a/Engine/Source/Runtime/Experimental/Chaos/Private/Chaos/PerParticleDampVelocity.cpp
+++ b/Engine/Source/Runtime/Experimental/Chaos/Private/Chaos/PerParticleDampVelocity.cpp
@@ -5,7 +5,7 @@
 #include "Chaos/SoftsSolverParticlesRange.h"
 #include "GenericPlatform/GenericPlatformMath.h"
 
-#if INTEL_ISPC && !UE_BUILD_SHIPPING
+#if INTEL_ISPC && (!UE_BUILD_SHIPPING || USE_ISPC_KERNEL_CONSOLE_VARIABLES_IN_SHIPPING)
 #include "HAL/IConsoleManager.h"
 #endif
 
