Dataflow - Added node freeze functionality.

[RN] Dataflow - Added a new Freeze functionality to the nodes' right click menu. This saves the current outputs into the node. Then when the frozen node's outputs are queried these saved values will be used instead of re-evaluating the node's outputs. Any following changes made to the frozen node, to its input connections, or to any connected nodes at its inputs, won't affect its resulting outputs until the node is unfrozen. The frozen state is serialized when the Dataflow asset is saved, and then restored when the asset is reloaded. The Unfreeze menu option restores the normal behavior.

diff --git a/Engine/Plugins/ChaosClothAssetEditor/Source/ChaosClothAssetDataflowNodes/Private/ChaosClothAsset/MergeClothCollectionsNode.cpp b/Engine/Plugins/ChaosClothAssetEditor/Source/ChaosClothAssetDataflowNodes/Private/ChaosClothAsset/MergeClothCollectionsNode.cpp
--- a/Engine/Plugins/ChaosClothAssetEditor/Source/ChaosClothAssetDataflowNodes/Private/ChaosClothAsset/MergeClothCollectionsNode.cpp
+++ b/Engine/Plugins/ChaosClothAssetEditor/Source/ChaosClothAssetDataflowNodes/Private/ChaosClothAsset/MergeClothCollectionsNode.cpp
@@ -478,7 +478,10 @@ void FChaosClothAssetMergeClothCollectionsNode_v2::PostSerialize(const FArchive&
 {
 	if (Ar.IsLoading())
 	{
-		check(Collections.Num() >= NumInitialOptionalInputs);
+		if (Collections.Num() < NumInitialOptionalInputs)
+		{
+			Collections.SetNum(NumInitialOptionalInputs);  // In case the FManagedArrayCollection wasn't serialized with the node (pre the WithSerializer trait)
+		}
 
 		for (int32 Index = 0; Index < NumInitialOptionalInputs; ++Index)
 		{
diff --git a/Engine/Plugins/ChaosClothAssetEditor/Source/ChaosClothAssetDataflowNodes/Private/ChaosClothAsset/TerminalNode.cpp b/Engine/Plugins/ChaosClothAssetEditor/Source/ChaosClothAssetDataflowNodes/Private/ChaosClothAsset/TerminalNode.cpp
--- a/Engine/Plugins/ChaosClothAssetEditor/Source/ChaosClothAssetDataflowNodes/Private/ChaosClothAsset/TerminalNode.cpp
+++ b/Engine/Plugins/ChaosClothAssetEditor/Source/ChaosClothAssetDataflowNodes/Private/ChaosClothAsset/TerminalNode.cpp
@@ -330,7 +330,10 @@ void FChaosClothAssetTerminalNode_v2::PostSerialize(const FArchive& Ar)
 
 	if (Ar.IsLoading())
 	{
-		check(CollectionLods.Num() >= NumInitialCollectionLods);
+		if (CollectionLods.Num() < NumInitialCollectionLods)
+		{
+			CollectionLods.SetNum(NumInitialCollectionLods);  // In case the FManagedArrayCollection wasn't serialized with the node (pre the WithSerializer trait)
+		}
 		for (int32 Index = 0; Index < NumInitialCollectionLods; ++Index)
 		{
 			check(FindInput(GetConnectionReference(Index)));
diff --git a/Engine/Plugins/ChaosClothAssetEditor/Source/ChaosClothAssetEditor/Private/ChaosClothAsset/ClothEditorMode.cpp b/Engine/Plugins/ChaosClothAssetEditor/Source/ChaosClothAssetEditor/Private/ChaosClothAsset/ClothEditorMode.cpp
--- a/Engine/Plugins/ChaosClothAssetEditor/Source/ChaosClothAssetEditor/Private/ChaosClothAsset/ClothEditorMode.cpp
+++ b/Engine/Plugins/ChaosClothAssetEditor/Source/ChaosClothAssetEditor/Private/ChaosClothAsset/ClothEditorMode.cpp
@@ -1620,8 +1620,11 @@ void UChaosClothAssetEditorMode::StartToolForSelectedNode(const UObject* Selecte
 	{
 		if (const TSharedPtr<const FDataflowNode> DataflowNode = EdNode->GetDataflowNode())
 		{
-			const FName DataflowNodeType = DataflowNode->GetType();
-			NodeTypeForPendingToolStart = DataflowNodeType;
+			if (!DataflowNode->IsFrozen())
+			{
+				const FName DataflowNodeType = DataflowNode->GetType();
+				NodeTypeForPendingToolStart = DataflowNodeType;
+			}
 		}
 	}
 }
diff --git a/Engine/Plugins/Experimental/Dataflow/Source/DataflowEditor/Private/Dataflow/DataflowEditorCommands.cpp b/Engine/Plugins/Experimental/Dataflow/Source/DataflowEditor/Private/Dataflow/DataflowEditorCommands.cpp
--- a/Engine/Plugins/Experimental/Dataflow/Source/DataflowEditor/Private/Dataflow/DataflowEditorCommands.cpp
+++ b/Engine/Plugins/Experimental/Dataflow/Source/DataflowEditor/Private/Dataflow/DataflowEditorCommands.cpp
@@ -57,6 +57,8 @@ void FDataflowEditorCommandsImpl::RegisterCommands()
 	TBaseCharacterFXEditorCommands::RegisterCommands();
 	
 	UI_COMMAND(EvaluateNode, "Evaluate", "Trigger an evaluation of the selected node.", EUserInterfaceActionType::Button, FInputChord());
+	UI_COMMAND(FreezeNodes, "FreezeNodes", "Freeze the evaluation of the selected nodes.", EUserInterfaceActionType::Button, FInputChord());
+	UI_COMMAND(UnfreezeNodes, "UnfreezeNodes", "Unfreeze the evaluation of the selected nodes.", EUserInterfaceActionType::Button, FInputChord());
 	UI_COMMAND(CreateComment, "CreateComment", "Create a Comment node.", EUserInterfaceActionType::None, FInputChord());
 	UI_COMMAND(ToggleEnabledState, "ToggleEnabledState", "Toggle node between Enabled/Disabled state.", EUserInterfaceActionType::Button, FInputChord());
 	UI_COMMAND(ToggleObjectSelection, "ToggleObjectSelection", "Enable object selection in editor.", EUserInterfaceActionType::ToggleButton, FInputChord());
@@ -332,6 +334,16 @@ void FDataflowEditorCommands::OnPropertyValueChanged(UDataflow* OutDataflow, TSh
 				{
 					if (TSharedPtr<FDataflowNode> DataflowNode = Node->GetDataflowNode())
 					{
+						// Active state update
+						if (InPropertyChangedEvent.GetPropertyName() == GET_MEMBER_NAME_CHECKED(FDataflowNode, bActive))
+						{
+							// Reflect the active state on the drawing of the node
+							if (DataflowNode->bActive != Node->IsNodeEnabled())
+							{
+								Node->SetEnabledState(DataflowNode->bActive ? ENodeEnabledState::Enabled : ENodeEnabledState::Disabled);
+							}
+						}
+						// Invalidate the node and reset the editor timestamp
 						DataflowNode->Invalidate();
 						OutLastNodeTimestamp = UE::Dataflow::FTimestamp::Invalid;
 
@@ -339,12 +351,6 @@ void FDataflowEditorCommands::OnPropertyValueChanged(UDataflow* OutDataflow, TSh
 						{
 							DataflowNode->OnPropertyChanged(*Context, InPropertyChangedEvent);
 						}
-
-						// Reflect the active state on the drawing of the node
-						if (DataflowNode->bActive != Node->IsNodeEnabled())
-						{
-							Node->SetEnabledState(DataflowNode->bActive ? ENodeEnabledState::Enabled : ENodeEnabledState::Disabled);
-						}
 					}
 				}
 			}
@@ -355,13 +361,47 @@ void FDataflowEditorCommands::OnPropertyValueChanged(UDataflow* OutDataflow, TSh
 	}
 }
 
+void FDataflowEditorCommands::FreezeNodes(UE::Dataflow::FContext& Context, const FGraphPanelSelectionSet& SelectedNodes)
+{
+	for (UObject* const SelectedNode : SelectedNodes)
+	{
+		if (UDataflowEdNode* const Node = Cast<UDataflowEdNode>(SelectedNode))
+		{
+			if (const TSharedPtr<FDataflowNode> DataflowNode = Node->GetDataflowNode())
+			{
+				if (!DataflowNode->IsFrozen())
+				{
+					DataflowNode->Freeze(Context);
+				}
+			}
+		}
+	}
+}
+
+void FDataflowEditorCommands::UnfreezeNodes(UE::Dataflow::FContext& Context, const FGraphPanelSelectionSet& SelectedNodes)
+{
+	for (UObject* const SelectedNode : SelectedNodes)
+	{
+		if (UDataflowEdNode* const Node = Cast<UDataflowEdNode>(SelectedNode))
+		{
+			if (const TSharedPtr<FDataflowNode> DataflowNode = Node->GetDataflowNode())
+			{
+				if (DataflowNode->IsFrozen())
+				{
+					DataflowNode->Unfreeze(Context);
+				}
+			}
+		}
+	}
+}
+
 void FDataflowEditorCommands::DeleteNodes(UDataflow* Graph, const FGraphPanelSelectionSet& SelectedNodes)
 {
 	if (ensureMsgf(Graph != nullptr, TEXT("Warning : Failed to find valid graph.")))
 	{
 		for (UObject* Node : SelectedNodes)
 		{
-			if (UDataflowEdNode* EdNode = dynamic_cast<UDataflowEdNode*>(Node))
+			if (UDataflowEdNode* EdNode = Cast<UDataflowEdNode>(Node))
 			{
 				if (const TSharedPtr<UE::Dataflow::FGraph> DataflowGraph = EdNode->GetDataflowGraph())
 				{
@@ -372,7 +412,7 @@ void FDataflowEditorCommands::DeleteNodes(UDataflow* Graph, const FGraphPanelSel
 					}
 				}
 			}
-			else if (UEdGraphNode_Comment* CommentNode = dynamic_cast<UEdGraphNode_Comment*>(Node))
+			else if (UEdGraphNode_Comment* CommentNode = Cast<UEdGraphNode_Comment>(Node))
 			{
 				Graph->RemoveNode(CommentNode);
 			}
diff --git a/Engine/Plugins/Experimental/Dataflow/Source/DataflowEditor/Private/Dataflow/DataflowSNode.cpp b/Engine/Plugins/Experimental/Dataflow/Source/DataflowEditor/Private/Dataflow/DataflowSNode.cpp
--- a/Engine/Plugins/Experimental/Dataflow/Source/DataflowEditor/Private/Dataflow/DataflowSNode.cpp
+++ b/Engine/Plugins/Experimental/Dataflow/Source/DataflowEditor/Private/Dataflow/DataflowSNode.cpp
@@ -17,6 +17,7 @@
 #include "Widgets/Layout/SBorder.h"
 #include "Widgets/Input/SCheckBox.h"
 #include "Widgets/SBoxPanel.h"
+#include "Widgets/Images/SImage.h"
 #include "Editor/Transactor.h"
 #include "GraphEditorSettings.h"
 
@@ -86,14 +87,12 @@ void SDataflowEdNode::Construct(const FArguments& InArgs, UDataflowEdNode* InNod
 				}
 				return ECheckBoxState::Unchecked;
 			})
-		.OnCheckStateChanged_Lambda([&](const ECheckBoxState NewState)
+		.OnCheckStateChanged_Lambda([this](const ECheckBoxState NewState)
 			{
 				if (DataflowGraphNode)
 				{
-					if (NewState == ECheckBoxState::Checked)
-						DataflowGraphNode->SetShouldWireframeRenderNode(true);
-					else
-						DataflowGraphNode->SetShouldWireframeRenderNode(false);
+					const bool bShouldWireframeRenderNode = (NewState == ECheckBoxState::Checked);
+					DataflowGraphNode->SetShouldWireframeRenderNode(bShouldWireframeRenderNode);
 				}
 			})
 		.IsEnabled_Lambda([this]()->bool
@@ -105,6 +104,23 @@ void SDataflowEdNode::Construct(const FArguments& InArgs, UDataflowEdNode* InNod
 				return false;
 			});
 
+	//
+	// Freeze
+	//
+	FreezeImageWidget = SNew(SImage)
+		.Image(FAppStyle::GetBrush("Icons.Lock"))
+		.DesiredSizeOverride(FVector2D(24.f, 24.f))
+		.Visibility_Lambda([this]()->EVisibility
+			{
+				if (DataflowGraphNode && DataflowGraphNode->GetDataflowNode())
+				{
+					if (DataflowGraphNode->GetDataflowNode()->IsFrozen())
+					{
+						return EVisibility::Visible;
+					}
+				}
+				return EVisibility::Collapsed;
+			});
 
 	//
 	//  Cached
@@ -176,19 +192,49 @@ PRAGMA_ENABLE_DEPRECATION_WARNINGS
 	return SGraphNode::CreatePinWidget(Pin);
 }
 
+TSharedRef<SWidget> SDataflowEdNode::CreateTitleRightWidget()
+{
+	// Reserve space for the overlay
+	float RightMargin = 0.f;
+	if (DataflowGraphNode && DataflowGraphNode->GetDataflowNode())
+	{
+		if (DataflowInterface->NodesHaveToggleWidget() && DataflowGraphNode->GetDataflowNode()->GetRenderParameters().Num())
+		{
+			RightMargin += 20.f;
+		}
+	}
+	return SNew(SHorizontalBox)
+		+ SHorizontalBox::Slot()
+		.HAlign(HAlign_Fill)
+		.Padding(FMargin(0.f, 0.f, RightMargin, 0.f));
+}
+
 TArray<FOverlayWidgetInfo> SDataflowEdNode::GetOverlayWidgets(bool bSelected, const FVector2D& WidgetSize) const
 {
 	TArray<FOverlayWidgetInfo> Widgets = SGraphNode::GetOverlayWidgets(bSelected, WidgetSize);
 
-	if (DataflowGraphNode && DataflowInterface->NodesHaveToggleWidget() && DataflowGraphNode->GetDataflowNode() && DataflowGraphNode->GetDataflowNode()->GetRenderParameters().Num())
+	if (DataflowGraphNode && DataflowGraphNode->GetDataflowNode())
 	{
-		const FVector2D ImageSize = RenderCheckBoxWidget->GetDesiredSize();
-
-		FOverlayWidgetInfo Info;
-		Info.OverlayOffset = FVector2D(WidgetSize.X - ImageSize.X - 6.f, 6.f);
-		Info.Widget = RenderCheckBoxWidget;
+		float ImageOffsetX = 0.f;
+		if (DataflowInterface->NodesHaveToggleWidget() && DataflowGraphNode->GetDataflowNode()->GetRenderParameters().Num())
+		{
+			const FVector2D RenderBoxImageSize = RenderCheckBoxWidget->GetDesiredSize();
+			ImageOffsetX += RenderBoxImageSize.X;
+			FOverlayWidgetInfo RenderBoxInfo;
+			RenderBoxInfo.OverlayOffset = FVector2D(WidgetSize.X - ImageOffsetX - 6.f, 6.f);
+			RenderBoxInfo.Widget = RenderCheckBoxWidget;
+			Widgets.Add(RenderBoxInfo);
+		}
 
-		Widgets.Add(Info);
+		if (DataflowInterface->NodesHaveFreezeWidget())
+		{
+			static const FVector2D FreezeImageOverhang(10.f);  // The icon must slightly overhang to make space for the lower output pins
+			const FVector2D FreezeImageSize = FreezeImageWidget->GetDesiredSize();
+			FOverlayWidgetInfo FreezeImageInfo;
+			FreezeImageInfo.OverlayOffset = WidgetSize - FreezeImageSize + FreezeImageOverhang;
+			FreezeImageInfo.Widget = FreezeImageWidget;
+			Widgets.Add(FreezeImageInfo);
+		}
 	}
 
 	/*
diff --git a/Engine/Plugins/Experimental/Dataflow/Source/DataflowEditor/Private/Dataflow/DataflowSchema.cpp b/Engine/Plugins/Experimental/Dataflow/Source/DataflowEditor/Private/Dataflow/DataflowSchema.cpp
--- a/Engine/Plugins/Experimental/Dataflow/Source/DataflowEditor/Private/Dataflow/DataflowSchema.cpp
+++ b/Engine/Plugins/Experimental/Dataflow/Source/DataflowEditor/Private/Dataflow/DataflowSchema.cpp
@@ -62,6 +62,8 @@ void UDataflowSchema::GetContextMenuActions(class UToolMenu* Menu, class UGraphN
 				Section.AddMenuEntry(FDataflowEditorCommands::Get().AddOptionPin);
 				Section.AddMenuEntry(FDataflowEditorCommands::Get().RemoveOptionPin);
 				Section.AddMenuEntry(FDataflowEditorCommands::Get().EvaluateNode);
+				Section.AddMenuEntry(FDataflowEditorCommands::Get().FreezeNodes, FText::FromString("Freeze"));
+				Section.AddMenuEntry(FDataflowEditorCommands::Get().UnfreezeNodes, FText::FromString("Unfreeze"));
 			}
 		}
 		{
diff --git a/Engine/Plugins/Experimental/Dataflow/Source/DataflowEditor/Public/Dataflow/DataflowEditorCommands.h b/Engine/Plugins/Experimental/Dataflow/Source/DataflowEditor/Public/Dataflow/DataflowEditorCommands.h
--- a/Engine/Plugins/Experimental/Dataflow/Source/DataflowEditor/Public/Dataflow/DataflowEditorCommands.h
+++ b/Engine/Plugins/Experimental/Dataflow/Source/DataflowEditor/Public/Dataflow/DataflowEditorCommands.h
@@ -55,6 +55,8 @@ public:
 	 static void UpdateToolCommandBinding(UInteractiveTool* Tool, TSharedPtr<FUICommandList> UICommandList, bool bUnbind = false);
 	
 	TSharedPtr<FUICommandInfo> EvaluateNode;
+	TSharedPtr<FUICommandInfo> FreezeNodes;
+	TSharedPtr<FUICommandInfo> UnfreezeNodes;
 	TSharedPtr<FUICommandInfo> CreateComment;
 	TSharedPtr<FUICommandInfo> ToggleEnabledState;
 	TSharedPtr<FUICommandInfo> ToggleObjectSelection;
@@ -128,6 +130,16 @@ public:
 	*/
 	static void DeleteNodes(UDataflow* Graph, const FGraphPanelSelectionSet& SelectedNodes);
 
+	/*
+	*  FreezeNodes
+	*/
+	static void FreezeNodes(UE::Dataflow::FContext& Context, const FGraphPanelSelectionSet& SelectedNodes);
+
+	/*
+	*  UnfreezeNodes
+	*/
+	static void UnfreezeNodes(UE::Dataflow::FContext& Context, const FGraphPanelSelectionSet& SelectedNodes);
+
 	/*
 	* OnNodeVerifyTitleCommit
 	*/
diff --git a/Engine/Plugins/Experimental/Dataflow/Source/DataflowEditor/Public/Dataflow/DataflowSEditorInterface.h b/Engine/Plugins/Experimental/Dataflow/Source/DataflowEditor/Public/Dataflow/DataflowSEditorInterface.h
--- a/Engine/Plugins/Experimental/Dataflow/Source/DataflowEditor/Public/Dataflow/DataflowSEditorInterface.h
+++ b/Engine/Plugins/Experimental/Dataflow/Source/DataflowEditor/Public/Dataflow/DataflowSEditorInterface.h
@@ -19,7 +19,15 @@ public:
 	/** Dataflow editor content accessors */
 	virtual TSharedPtr<UE::Dataflow::FContext> GetDataflowContext() const = 0;
 
-	virtual bool NodesHaveToggleWidget() const { return true; }
+	virtual bool NodesHaveToggleWidget() const
+	{
+		return true;
+	}
+
+	virtual bool NodesHaveFreezeWidget() const
+	{
+		return true;
+	}
 
 protected:
 	FDataflowSEditorInterface() = default;
diff --git a/Engine/Plugins/Experimental/Dataflow/Source/DataflowEditor/Public/Dataflow/DataflowSNode.h b/Engine/Plugins/Experimental/Dataflow/Source/DataflowEditor/Public/Dataflow/DataflowSNode.h
--- a/Engine/Plugins/Experimental/Dataflow/Source/DataflowEditor/Public/Dataflow/DataflowSNode.h
+++ b/Engine/Plugins/Experimental/Dataflow/Source/DataflowEditor/Public/Dataflow/DataflowSNode.h
@@ -55,7 +55,7 @@ public:
 
 	// SGraphNode interface
 	virtual FReply OnMouseButtonDoubleClick(const FGeometry& InMyGeometry, const FPointerEvent& InMouseEvent) override;
-	TArray<FOverlayWidgetInfo> GetOverlayWidgets(bool bSelected, const FVector2D& WidgetSize) const;
+	virtual TArray<FOverlayWidgetInfo> GetOverlayWidgets(bool bSelected, const FVector2D& WidgetSize) const override;
 	virtual void UpdateErrorInfo() override;
 	virtual void RequestRenameOnSpawn() override { /* No auto rename on spawn, because it can interfere with Copy/Paste and cause a crash */ }
 
@@ -80,11 +80,13 @@ private:
 	virtual EVisibility IsAddPinButtonVisible() const;
 
 	virtual TSharedPtr<SGraphPin> CreatePinWidget(UEdGraphPin* Pin) const override;
+	virtual TSharedRef<SWidget> CreateTitleRightWidget() override;
 
 	TObjectPtr<UDataflowEdNode> DataflowGraphNode = nullptr;	
 
 	FCheckBoxStyle CheckBoxStyle;
 	TSharedPtr<SCheckBox> RenderCheckBoxWidget;
+	TSharedPtr<SImage> FreezeImageWidget;
 
 	//FCheckBoxStyle CacheStatusStyle;
 	//TSharedPtr<SCheckBox> CacheStatus;
diff --git a/Engine/Source/Runtime/Core/Public/UObject/FortniteMainBranchObjectVersion.h b/Engine/Source/Runtime/Core/Public/UObject/FortniteMainBranchObjectVersion.h
--- a/Engine/Source/Runtime/Core/Public/UObject/FortniteMainBranchObjectVersion.h
+++ b/Engine/Source/Runtime/Core/Public/UObject/FortniteMainBranchObjectVersion.h
@@ -538,6 +538,9 @@ struct FFortniteMainBranchObjectVersion
 		// Added support for the GyroscopicTorque flag in CVD
 		GyroscopicTorquesSupportInChaosVisualDebugger,
 
+		// Added managed array property serialization
+		AddManagedArrayCollectionPropertySerialization,
+
 		// -----<new versions can be added above this line>-------------------------------------------------
 		VersionPlusOne,
 		LatestVersion = VersionPlusOne - 1
diff --git a/Engine/Source/Runtime/Experimental/Chaos/Private/GeometryCollection/ManagedArrayCollection.cpp b/Engine/Source/Runtime/Experimental/Chaos/Private/GeometryCollection/ManagedArrayCollection.cpp
--- a/Engine/Source/Runtime/Experimental/Chaos/Private/GeometryCollection/ManagedArrayCollection.cpp
+++ b/Engine/Source/Runtime/Experimental/Chaos/Private/GeometryCollection/ManagedArrayCollection.cpp
@@ -821,6 +821,19 @@ void FManagedArrayCollection::Serialize(Chaos::FChaosArchive& Ar)
 	}
 }
 
+bool FManagedArrayCollection::Serialize(FArchive& Ar)
+{
+	Ar.UsingCustomVersion(FFortniteMainBranchObjectVersion::GUID);
+
+	if (Ar.IsLoading() && Ar.CustomVer(FFortniteMainBranchObjectVersion::GUID) < FFortniteMainBranchObjectVersion::AddManagedArrayCollectionPropertySerialization)
+	{
+		return false;  // No custom serialization yet, default back to SerializeTaggedProperties
+	}
+	Chaos::FChaosArchive ChaosArchive(Ar);
+	Serialize(ChaosArchive);
+	return true;
+}
+
 FArchive& operator<<(FArchive& Ar, FManagedArrayCollection::FGroupInfo& GroupInfo)
 {
 	int Version = 4;
diff --git a/Engine/Source/Runtime/Experimental/Chaos/Public/GeometryCollection/ManagedArrayCollection.h b/Engine/Source/Runtime/Experimental/Chaos/Public/GeometryCollection/ManagedArrayCollection.h
--- a/Engine/Source/Runtime/Experimental/Chaos/Public/GeometryCollection/ManagedArrayCollection.h
+++ b/Engine/Source/Runtime/Experimental/Chaos/Public/GeometryCollection/ManagedArrayCollection.h
@@ -576,6 +576,8 @@ public:
 	*/
 	CHAOS_API virtual void Serialize(Chaos::FChaosArchive& Ar);
 
+	CHAOS_API bool Serialize(FArchive& Ar);
+
 	/**
 	* Cycle Checking. 
 	* Search for TargetNode from the StartingNode. 
@@ -841,6 +843,13 @@ protected:
 
 };
 
+template<> struct TStructOpsTypeTraits<FManagedArrayCollection> : public TStructOpsTypeTraitsBase2<FManagedArrayCollection>
+{
+	enum
+	{
+		WithSerializer = true,
+	};
+};
 
 /*
 *  FManagedArrayInterface
diff --git a/Engine/Source/Runtime/Experimental/Dataflow/Core/Private/Dataflow/DataflowInputOutput.cpp b/Engine/Source/Runtime/Experimental/Dataflow/Core/Private/Dataflow/DataflowInputOutput.cpp
--- a/Engine/Source/Runtime/Experimental/Dataflow/Core/Private/Dataflow/DataflowInputOutput.cpp
+++ b/Engine/Source/Runtime/Experimental/Dataflow/Core/Private/Dataflow/DataflowInputOutput.cpp
@@ -255,31 +255,38 @@ bool FDataflowOutput::Evaluate(UE::Dataflow::FContext& Context) const
 
 bool FDataflowOutput::EvaluateImpl(UE::Dataflow::FContext& Context) const
 {
-	UE_LOG(LogChaosDataflow, VeryVerbose, TEXT("FDataflowOutput::EvaluateImpl(): Node [%s], Output [%s]"), *GetOwningNode()->GetName().ToString(), *GetName().ToString());
-
-	UE::Dataflow::FContextScopedCallstack Callstack(Context, this);
-	if (Callstack.IsLoopDetected())
-	{ 
-		ensureMsgf(false, TEXT("Connection %s is already in the callstack, this is certainly because of a loop in the graph"), *GetName().ToString());
-		return false;
-	}
-
-	// check if the cache has a valid version
-	if(Context.HasData(CacheKey(), OwningNode->GetTimestamp()))
+	if (HasFrozenValue())
 	{
-		UE_LOG(LogChaosDataflow, VeryVerbose, TEXT("FDataflowOutput::EvaluateImpl(): Context has data, NodeTimestamp [%ul], CacheTimestamp [%lu]"), OwningNode->GetTimestamp().Value, Context.GetTimestamp(CacheKey()).Value);
+		UE_LOG(LogChaosDataflow, VeryVerbose, TEXT("FDataflowOutput::EvaluateImpl(): Node [%s], Output [%s] [FROZEN]"), *GetOwningNode()->GetName().ToString(), *GetName().ToString());
 		return true;
 	}
-	// if not, evaluate
-	UE_LOG(LogChaosDataflow, Verbose, TEXT("FDataflowNode::Evaluate(): Node [%s], Output [%s], NodeTimestamp [%lu]"), *GetOwningNode()->GetName().ToString(), *GetName().ToString(), OwningNode->GetTimestamp().Value);
-	OwningNode->Evaluate(Context, this);
-	// Validation
-	if (!Context.HasData(CacheKey()))
+	else
 	{
-		ensureMsgf(false, TEXT("Failed to evaluate output (%s:%s)"), *OwningNode->GetName().ToString(), *GetName().ToString());
-		return false;
-	}
+		UE_LOG(LogChaosDataflow, VeryVerbose, TEXT("FDataflowOutput::EvaluateImpl(): Node [%s], Output [%s]"), *GetOwningNode()->GetName().ToString(), *GetName().ToString());
 
+		UE::Dataflow::FContextScopedCallstack Callstack(Context, this);
+		if (Callstack.IsLoopDetected())
+		{
+			ensureMsgf(false, TEXT("Connection %s is already in the callstack, this is certainly because of a loop in the graph"), *GetName().ToString());
+			return false;
+		}
+
+		// check if the cache has a valid version
+		if (Context.HasData(CacheKey(), OwningNode->GetTimestamp()))
+		{
+			UE_LOG(LogChaosDataflow, VeryVerbose, TEXT("FDataflowOutput::EvaluateImpl(): Context has data, NodeTimestamp [%ul], CacheTimestamp [%lu]"), OwningNode->GetTimestamp().Value, Context.GetTimestamp(CacheKey()).Value);
+			return true;
+		}
+		// if not, evaluate
+		UE_LOG(LogChaosDataflow, Verbose, TEXT("FDataflowNode::Evaluate(): Node [%s], Output [%s], NodeTimestamp [%lu]"), *GetOwningNode()->GetName().ToString(), *GetName().ToString(), OwningNode->GetTimestamp().Value);
+		OwningNode->Evaluate(Context, this);
+		// Validation
+		if (!Context.HasData(CacheKey()))
+		{
+			ensureMsgf(false, TEXT("Failed to evaluate output (%s:%s)"), *OwningNode->GetName().ToString(), *GetName().ToString());
+			return false;
+		}
+	}
 	return true;
 }
 
@@ -288,6 +295,47 @@ TFuture<bool> FDataflowOutput::EvaluateParallel(UE::Dataflow::FContext& Context)
 	return Async(EAsyncExecution::TaskGraph, [&]() -> bool { return this->Evaluate(Context); });
 }
 
+void FDataflowOutput::Freeze(UE::Dataflow::FContext& Context, FInstancedPropertyBag& FrozenProperties)
+{
+	if (GetOwningNode() && Property)
+	{
+		UE_LOG(LogChaosDataflow, Verbose, TEXT("FDataflowOutput::Freeze(): Node [%s], Output [%s]"), *GetOwningNode()->GetName().ToString(), *GetName().ToString());
+
+		if (Evaluate(Context))
+		{
+			// Set the cached value
+			const uint8* const Data = (const uint8*)Context.GetUntypedData(CacheKey(), Property);  // Pretend that the cache element is a property container 
+			const uint8* const ContainerPtr = Data - Property->GetOffset_ForInternal();     // TODO: Replace cache elements by property bags?
+
+			FrozenProperties.AddProperty(GetName(), Property);
+			FrozenProperties.SetValue(GetName(), Property, ContainerPtr);
+		}
+		// Else the output is not evaluate-able (e.g. deactivated node) but is still considered frozen and will return the provided default when GetValue is called.
+	}
+}
+
+bool FDataflowOutput::HasFrozenValue() const
+{
+	return GetOwningNode() && GetOwningNode()->IsFrozen();
+}
+
+const uint8* FDataflowOutput::GetFrozenPropertyValue(const uint8* Default) const
+{
+	if (GetOwningNode())
+	{
+		const FInstancedPropertyBag& FrozenFroperties = GetOwningNode()->GetFrozenFroperties();
+
+		if (const FPropertyBagPropertyDesc* PropertyDesc = FrozenFroperties.FindPropertyDescByName(GetName()))
+		{
+			const FConstStructView Value = FrozenFroperties.GetValue();
+			if (Value.IsValid())
+			{
+				return Value.GetMemory() + PropertyDesc->CachedProperty->GetOffset_ForInternal();
+			}
+		}
+	}
+	return Default;
+}
 
 void FDataflowOutput::ForwardInput(const UE::Dataflow::FConnectionReference& InputReference, UE::Dataflow::FContext& Context) const
 {
@@ -355,4 +403,4 @@ void FDataflowOutput::FixAndPropagateType(FName InType)
 			}
 		}
 	}
-}
\ No newline at end of file
+}
diff --git a/Engine/Source/Runtime/Experimental/Dataflow/Core/Private/Dataflow/DataflowNode.cpp b/Engine/Source/Runtime/Experimental/Dataflow/Core/Private/Dataflow/DataflowNode.cpp
--- a/Engine/Source/Runtime/Experimental/Dataflow/Core/Private/Dataflow/DataflowNode.cpp
+++ b/Engine/Source/Runtime/Experimental/Dataflow/Core/Private/Dataflow/DataflowNode.cpp
@@ -599,6 +599,26 @@ PRAGMA_ENABLE_DEPRECATION_WARNINGS
 	}
 }
 
+void FDataflowNode::Freeze(UE::Dataflow::FContext& Context)
+{
+	bIsFrozen = true;
+	for (FDataflowOutput* const Output : GetOutputs())
+	{
+		if (Output)
+		{
+			Output->Freeze(Context, FrozenProperties);
+		}
+	}
+}
+
+void FDataflowNode::Unfreeze(UE::Dataflow::FContext& Context)
+{
+	FrozenProperties.Reset();
+	bIsFrozen = false;
+
+	Invalidate();
+}
+
 void FDataflowNode::Invalidate(const UE::Dataflow::FTimestamp& InModifiedTimestamp)
 {
 	if (bPauseInvalidations)
diff --git a/Engine/Source/Runtime/Experimental/Dataflow/Core/Private/Dataflow/DataflowNodeParameters.cpp b/Engine/Source/Runtime/Experimental/Dataflow/Core/Private/Dataflow/DataflowNodeParameters.cpp
--- a/Engine/Source/Runtime/Experimental/Dataflow/Core/Private/Dataflow/DataflowNodeParameters.cpp
+++ b/Engine/Source/Runtime/Experimental/Dataflow/Core/Private/Dataflow/DataflowNodeParameters.cpp
@@ -47,11 +47,11 @@ namespace UE::Dataflow
 #endif
 	}
 
-	bool FContext::IsCacheEntryAfterTimestamp(FContextCacheKey InKey, const FTimestamp InTimestamp)
+	bool FContext::IsCacheEntryAfterTimestamp(FContextCacheKey InKey, const FTimestamp InTimestamp) const
 	{
 		if (HasData(InKey))
 		{
-			if (TUniquePtr<FContextCacheElementBase>* CacheEntry = GetDataImpl(InKey))
+			if (const TUniquePtr<FContextCacheElementBase>* const CacheEntry = GetDataImpl(InKey))
 			{
 				if (*CacheEntry)
 				{
@@ -97,7 +97,7 @@ namespace UE::Dataflow
 			// Note: If the node is deactivated and has an output (like above), then the output might still need to be forwarded.
 			//       Therefore the Evaluate method has to be called for whichever value of bActive.
 			//       However if the node is deactivated and has no outputs (like below), now is the time to check its bActive state.
-			else if (Node->bActive)
+			else if (Node->bActive && !Node->IsFrozen())
 			{
 				// TODO: When no outputs are specified, this call to Evaluate should really be removed.
 				//       The purpose of the node evaluation function is to evaluate outputs.
diff --git a/Engine/Source/Runtime/Experimental/Dataflow/Core/Public/Dataflow/DataflowInputOutput.h b/Engine/Source/Runtime/Experimental/Dataflow/Core/Public/Dataflow/DataflowInputOutput.h
--- a/Engine/Source/Runtime/Experimental/Dataflow/Core/Public/Dataflow/DataflowInputOutput.h
+++ b/Engine/Source/Runtime/Experimental/Dataflow/Core/Public/Dataflow/DataflowInputOutput.h
@@ -215,6 +215,10 @@ public:
 	template<class T>
 	const T& GetValue(UE::Dataflow::FContext& Context, const T& Default) const
 	{
+		if (HasFrozenValue())
+		{
+			return GetFrozenValue(Default);
+		}
 		if (!this->Evaluate(Context))
 		{
 			Context.SetData(CacheKey(), Property, Default, GetOwningNodeGuid(), GetOwningNodeValueHash(), GetOwningNodeTimestamp());
@@ -228,7 +232,18 @@ public:
 		return Default;
 	}
 
-	bool HasCachedValue(UE::Dataflow::FContext& Context) const
+	/** Return a pointer to the node's output storage if it has been successfully evaluated and frozen, or the provided default value otherwise. */
+	template<class T>
+	const T& GetFrozenValue(const T& Default) const
+	{
+		return *(const T*)GetFrozenPropertyValue((const uint8*)&Default);
+	}
+	/** Freeze the value of this output to the node property bag. */
+	DATAFLOWCORE_API void Freeze(UE::Dataflow::FContext& Context, struct FInstancedPropertyBag& FrozenProperties);
+	/** Return whether the output has a frozen value in its owner node. */
+	DATAFLOWCORE_API bool HasFrozenValue() const;
+
+	bool HasCachedValue(const UE::Dataflow::FContext& Context) const
 	{
 		return Context.HasData(CacheKey(), GetOwningNodeTimestamp());
 	}
@@ -248,6 +263,7 @@ public:
 	DATAFLOWCORE_API virtual void Invalidate(const UE::Dataflow::FTimestamp& ModifiedTimestamp = UE::Dataflow::FTimestamp::Current()) override;
 
 private:
+	DATAFLOWCORE_API const uint8* GetFrozenPropertyValue(const uint8* Default) const;
 };
 
 template<>
@@ -264,14 +280,18 @@ const T& FDataflowInput::GetValue(UE::Dataflow::FContext& Context, const T& Defa
 {
 	if (const FDataflowOutput* ConnectionOut = GetConnection())
 	{
+		if (ConnectionOut->HasFrozenValue())
+		{
+			return ConnectionOut->GetFrozenValue(Default);
+		}
 		if (!ConnectionOut->Evaluate(Context))
 		{
 			Context.SetData(ConnectionOut->CacheKey(), Property, Default, GetOwningNodeGuid(), GetOwningNodeValueHash(), GetOwningNodeTimestamp());
 		}
 		if (Context.HasData(ConnectionOut->CacheKey()))
 		{
-			const T& data = Context.GetData(ConnectionOut->CacheKey(), Property, Default);
-			return data;
+			const T& Data = Context.GetData(ConnectionOut->CacheKey(), Property, Default);
+			return Data;
 		}
 	}
 	return Default;
@@ -283,6 +303,10 @@ typename TAnyType::FStorageType FDataflowInput::GetValueFromAnyType(UE::Dataflow
 	typename TAnyType::FStorageType ReturnValue = Default;
 	if (const FDataflowOutput* ConnectionOut = GetConnection())
 	{
+		if (ConnectionOut->HasFrozenValue())
+		{
+			return ConnectionOut->GetFrozenValue(Default);
+		}
 		if (ConnectionOut->Evaluate(Context))
 		{
 			if (const TUniquePtr<UE::Dataflow::FContextCacheElementBase>* CacheEntry = Context.GetDataImpl(ConnectionOut->CacheKey()))
diff --git a/Engine/Source/Runtime/Experimental/Dataflow/Core/Public/Dataflow/DataflowNode.h b/Engine/Source/Runtime/Experimental/Dataflow/Core/Public/Dataflow/DataflowNode.h
--- a/Engine/Source/Runtime/Experimental/Dataflow/Core/Public/Dataflow/DataflowNode.h
+++ b/Engine/Source/Runtime/Experimental/Dataflow/Core/Public/Dataflow/DataflowNode.h
@@ -4,12 +4,13 @@
 
 #include "ChaosLog.h"
 #include "CoreMinimal.h"
+#include "Dataflow/DataflowAnyType.h"
 #include "Dataflow/DataflowConnection.h"
 #include "Dataflow/DataflowInputOutput.h"
 #include "Dataflow/DataflowNodeParameters.h"
-#include "Dataflow/DataflowAnyType.h"
-#include "UObject/StructOnScope.h"
 #include "Dataflow/DataflowSettings.h"
+#include "StructUtils/PropertyBag.h"
+#include "UObject/StructOnScope.h"
 
 #include "DataflowNode.generated.h"
 
@@ -74,7 +75,11 @@ PRAGMA_ENABLE_DEPRECATION_WARNINGS
 	virtual ~FDataflowNode() { ClearInputs(); ClearOutputs(); }
 	// Warning: FDataflowNodes aren't actually safe to copy/move yet. These are here to disable deprecation warnings from the implicit operators that were getting created anyways.
 	// (Deleting the operators would require tagging all derived classes since this is a USTRUCT, so also not doing that here).
-	FDataflowNode& operator=(const FDataflowNode&) = default;
+	FDataflowNode& operator=(const FDataflowNode& Other)
+	{
+		checkNoEntry();
+		return *this;
+	}
 
 	FGuid GetGuid() const { return Guid; }
 	FName GetName() const { return Name; }
@@ -322,6 +327,21 @@ PRAGMA_ENABLE_DEPRECATION_WARNINGS
 	//
 	virtual void Evaluate(UE::Dataflow::FContext& Context, const FDataflowOutput*) const { ensure(false); }
 
+	/** Freeze the node's outputs. */
+	DATAFLOWCORE_API void Freeze(UE::Dataflow::FContext& Context);
+	/** Unfreeze the node's outputs. */
+	DATAFLOWCORE_API void Unfreeze(UE::Dataflow::FContext& Context);
+	/** Return whether the node is frozen. */
+	bool IsFrozen() const
+	{
+		return bIsFrozen;
+	}
+	/** Return the frozen output property bag. */
+	const FInstancedPropertyBag& GetFrozenFroperties() const
+	{
+		return FrozenProperties;
+	}
+
 	/**
 	*   GetValue(...)
 	*
@@ -659,6 +679,12 @@ private:
 	TMap<UE::Dataflow::FConnectionKey, FDataflowInput*> ExpandedInputs;
 	TMap<uint32, const FArrayProperty*> InputArrayProperties; // Used to calculate ContainerElementOffsets
 
+	UPROPERTY()
+	FInstancedPropertyBag FrozenProperties;  // Property bag for frozen outputs
+
+	UPROPERTY()
+	bool bIsFrozen = false;
+
 protected:
 	bool bPauseInvalidations = false;
 	UE::Dataflow::FTimestamp PausedModifiedTimestamp = UE::Dataflow::FTimestamp::Invalid; // When unpausing invalidations, Invalidate will be called with this timestamp.
@@ -733,8 +759,7 @@ public:																				\
 	   return new FStructOnScope(TYPE::StaticStruct(), (uint8*)this);}				\
 	virtual void SerializeInternal(FArchive& Ar) override {							\
 		UScriptStruct* const Struct = TYPE::StaticStruct();							\
-		Struct->SerializeTaggedProperties(Ar, (uint8*)this,							\
-		Struct, nullptr);															\
+		Struct->SerializeTaggedProperties(Ar, (uint8*)this, Struct, nullptr);		\
 		Serialize(Ar);																\
 		PostSerialize(Ar);}															\
 	virtual FName GetDisplayName() const override { return TYPE::StaticDisplay(); }	\
diff --git a/Engine/Source/Runtime/Experimental/Dataflow/Core/Public/Dataflow/DataflowNodeParameters.h b/Engine/Source/Runtime/Experimental/Dataflow/Core/Public/Dataflow/DataflowNodeParameters.h
--- a/Engine/Source/Runtime/Experimental/Dataflow/Core/Public/Dataflow/DataflowNodeParameters.h
+++ b/Engine/Source/Runtime/Experimental/Dataflow/Core/Public/Dataflow/DataflowNodeParameters.h
@@ -85,8 +85,10 @@ namespace UE::Dataflow
 		virtual TUniquePtr<FContextCacheElementBase> CreateReference(FContextCacheKey InReferenceDataKey) const = 0;
 
 		template<typename T>
-		inline const T& GetTypedData(FContext& Context, const FProperty* PropertyIn, const T& Default) const;
-		
+		inline const T& GetTypedData(const FContext& Context, const FProperty* PropertyIn, const T& Default) const;
+
+		virtual const void* GetUntypedData(const FContext& Context, const FProperty* PropertyIn) const = 0;
+
 		EType GetType() const {	return Type; }
 
 		const FProperty* GetProperty() const { return Property; }
@@ -120,7 +122,12 @@ namespace UE::Dataflow
 			, Data(Forward<T>(InData))
 		{}
 		
-		inline const T& GetData(FContext& Context, const FProperty* PropertyIn, const T& Default) const;
+		inline const T& GetData(const FContext& Context, const FProperty* PropertyIn, const T& Default) const;
+
+		virtual const void* GetUntypedData(const FContext& Context, const FProperty* PropertyIn) const
+		{
+			return &Data;
+		}
 
 		const T& GetDataDirect() const { return Data; }
 
@@ -142,7 +149,9 @@ namespace UE::Dataflow
 			, DataKey(InDataKey)
 		{}
 
-		inline const T& GetData(FContext& Context, const FProperty* PropertyIn, const T& Default) const;
+		inline const T& GetData(const FContext& Context, const FProperty* PropertyIn, const T& Default) const;
+
+		inline virtual const void* GetUntypedData(const FContext& Context, const FProperty* PropertyIn) const;
 
 		inline virtual TUniquePtr<FContextCacheElementBase> CreateReference(FContextCacheKey InReferenceDataKey) const override;
 
@@ -159,6 +168,11 @@ namespace UE::Dataflow
 			: FContextCacheElementBase(EType::CacheElementNull, InNodeGuid, InProperty, InNodeHash, Timestamp)
 		{}
 
+		virtual const void* GetUntypedData(const FContext& Context, const FProperty* PropertyIn) const
+		{
+			return nullptr;
+		}
+
 		inline virtual TUniquePtr<FContextCacheElementBase> CreateReference(FContextCacheKey InReferenceDataKey) const override;
 	};
 
@@ -168,28 +182,36 @@ namespace UE::Dataflow
 	template<class T>
 	struct TContextCacheElementUObject : public FContextCacheElementBase, public FGCObject
 	{
-		TContextCacheElementUObject(FGuid InNodeGuid, const FProperty* InProperty, T&& InObject, uint32 InNodeHash, FTimestamp Timestamp)
+		TContextCacheElementUObject(FGuid InNodeGuid, const FProperty* InProperty, T&& InObjectPtr, uint32 InNodeHash, FTimestamp Timestamp)
 			: FContextCacheElementBase(EType::CacheElementUObject, InNodeGuid, InProperty, InNodeHash, Timestamp)
-			, Object(InObject)
+			, ObjectPtr(InObjectPtr)
 		{}
 
-		const T& GetData(FContext& /*Context*/, const FProperty* /*PropertyIn*/, const T& /*Default*/) const { return Object; }
+		const T& GetData(const FContext& /*Context*/, const FProperty* /*PropertyIn*/, const T& /*Default*/) const
+		{
+			return ObjectPtr;
+		}
+
+		virtual const void* GetUntypedData(const FContext& /*Context*/, const FProperty* /*PropertyIn*/) const
+		{
+			return (void*)ObjectPtr;
+		}
 
 		inline virtual TUniquePtr<FContextCacheElementBase> CreateReference(FContextCacheKey InReferenceDataKey) const override;
 
 		//~ Begin FGCObject interface
-		virtual void AddReferencedObjects(FReferenceCollector& Collector) override { Collector.AddReferencedObject(Object); }
+		virtual void AddReferencedObjects(FReferenceCollector& Collector) override { Collector.AddReferencedObject(ObjectPtr); }
 		virtual FString GetReferencerName() const override { return TEXT("TContextCacheElementUObject"); }
 		//~ End FGCObject interface
 
 	private:
 		typedef typename TDecay<T>::Type FDataType;  // Using universal references here means T could be either const& or an rvalue reference
-		FDataType Object;                            // Decaying T removes any reference and gets the correct underlying storage data type
+		FDataType ObjectPtr;                         // Decaying T removes any reference and gets the correct underlying storage data type
 	};
 
 	// cache element method implementation 
 	template<class T>
-	const T& FContextCacheElementBase::GetTypedData(FContext& Context, const FProperty* PropertyIn, const T& Default) const
+	const T& FContextCacheElementBase::GetTypedData(const FContext& Context, const FProperty* PropertyIn, const T& Default) const
 	{
 		// check(PropertyIn->IsA<T>()); // @todo(dataflow) compile error for non-class T; find alternatives
 		if (Type == EType::CacheElementTyped)
@@ -295,7 +317,7 @@ namespace UE::Dataflow
 		void SetDataReference(FContextCacheKey Key, const FProperty* Property, FContextCacheKey ReferenceKey)
 		{
 			// find the reference key to get 
-			if (TUniquePtr<FContextCacheElementBase>* CacheElement = GetDataImpl(ReferenceKey))
+			if (const TUniquePtr<FContextCacheElementBase>* CacheElement = GetDataImpl(ReferenceKey))
 			{
 				TUniquePtr<FContextCacheElementBase> CacheReferenceElement = (*CacheElement)->CreateReference(ReferenceKey);
 				SetDataImpl(Key, MoveTemp(CacheReferenceElement));
@@ -315,21 +337,30 @@ namespace UE::Dataflow
 			SetDataImpl(InKey, MoveTemp(CacheNullElement));
 		}
 
-		virtual TUniquePtr<FContextCacheElementBase>* GetDataImpl(FContextCacheKey Key) = 0;
+		virtual const TUniquePtr<FContextCacheElementBase>* GetDataImpl(FContextCacheKey Key) const = 0;
 
 		template<class T>
-		const T& GetData(FContextCacheKey Key, const FProperty* InProperty, const T& Default = T())
+		const T& GetData(FContextCacheKey Key, const FProperty* InProperty, const T& Default = T()) const
 		{
-			if (TUniquePtr<FContextCacheElementBase>* Cache = GetDataImpl(Key))
+			if (const TUniquePtr<FContextCacheElementBase>* Cache = GetDataImpl(Key))
 			{
 				return (*Cache)->GetTypedData<T>(*this, InProperty, Default);
 			}
 			return Default;
 		}
 
-		virtual bool HasDataImpl(FContextCacheKey Key, FTimestamp InTimestamp = FTimestamp::Invalid) = 0;
+		const void* GetUntypedData(FContextCacheKey Key, const FProperty* InProperty) const
+		{
+			if (const TUniquePtr<FContextCacheElementBase>* Cache = GetDataImpl(Key))
+			{
+				return (*Cache)->GetUntypedData(*this, InProperty);
+			}
+			return nullptr;
+		}
+
+		virtual bool HasDataImpl(FContextCacheKey Key, FTimestamp InTimestamp = FTimestamp::Invalid) const = 0;
 		
-		bool HasData(FContextCacheKey Key, FTimestamp InTimestamp = FTimestamp::Invalid)
+		bool HasData(FContextCacheKey Key, FTimestamp InTimestamp = FTimestamp::Invalid) const
 		{
 			FContextCacheKey IntKey = (FContextCacheKey)Key;
 			return HasDataImpl(Key, InTimestamp);
@@ -358,7 +389,7 @@ namespace UE::Dataflow
 		DATAFLOWCORE_API void PopFromCallstack(const FDataflowConnection* Connection);
 		DATAFLOWCORE_API bool IsInCallstack(const FDataflowConnection* Connection) const;
 
-		DATAFLOWCORE_API bool IsCacheEntryAfterTimestamp(FContextCacheKey InKey, const FTimestamp InTimestamp);
+		DATAFLOWCORE_API bool IsCacheEntryAfterTimestamp(FContextCacheKey InKey, const FTimestamp InTimestamp) const;
 
 	private:
 #if DATAFLOW_EDITOR_EVALUATION
@@ -399,12 +430,12 @@ namespace UE::Dataflow
 			DataStore.Emplace(Key, MoveTemp(DataStoreEntry));
 		}
 
-		virtual TUniquePtr<FContextCacheElementBase>* GetDataImpl(FContextCacheKey Key) override
+		virtual const TUniquePtr<FContextCacheElementBase>* GetDataImpl(FContextCacheKey Key) const override
 		{
 			return DataStore.Find(Key);
 		}
 
-		virtual bool HasDataImpl(FContextCacheKey Key, FTimestamp InTimestamp = FTimestamp::Invalid) override
+		virtual bool HasDataImpl(FContextCacheKey Key, FTimestamp InTimestamp = FTimestamp::Invalid) const override
 		{
 			return DataStore.Contains(Key) && DataStore[Key]->GetTimestamp() >= InTimestamp;
 		}
@@ -447,14 +478,14 @@ namespace UE::Dataflow
 			}
 		}
 
-		virtual TUniquePtr<FContextCacheElementBase>* GetDataImpl(FContextCacheKey Key) override
+		virtual const TUniquePtr<FContextCacheElementBase>* GetDataImpl(FContextCacheKey Key) const override
 		{
 			CacheLock->Lock(); ON_SCOPE_EXIT { CacheLock->Unlock(); };
 			
 			return DataStore.Find(Key);
 		}
 
-		virtual bool HasDataImpl(FContextCacheKey Key, FTimestamp InTimestamp = FTimestamp::Invalid) override
+		virtual bool HasDataImpl(FContextCacheKey Key, FTimestamp InTimestamp = FTimestamp::Invalid) const override
 		{
 			CacheLock->Lock(); ON_SCOPE_EXIT { CacheLock->Unlock(); };
 			
@@ -475,7 +506,7 @@ namespace UE::Dataflow
 	// this needs to be after the FContext definition because they access its methods
 
 	template<class T>
-	const T& TContextCacheElement<T>::GetData(FContext& Context, const FProperty* PropertyIn, const T& Default) const
+	const T& TContextCacheElement<T>::GetData(const FContext& Context, const FProperty* PropertyIn, const T& Default) const
 	{
 		return Data;
 	}
@@ -487,11 +518,17 @@ namespace UE::Dataflow
 	}
 
 	template<class T>
-	const T& TContextCacheElementReference<T>::GetData(FContext& Context, const FProperty* PropertyIn, const T& Default) const
+	const T& TContextCacheElementReference<T>::GetData(const FContext& Context, const FProperty* PropertyIn, const T& Default) const
 	{
 		return Context.GetData(DataKey, PropertyIn, Default);
 	}
 
+	template<class T>
+	const void* TContextCacheElementReference<T>::GetUntypedData(const FContext& Context, const FProperty* PropertyIn) const
+	{
+		return Context.GetUntypedData(DataKey, PropertyIn);
+	}
+
 	template<class T>
 	TUniquePtr<FContextCacheElementBase> TContextCacheElementReference<T>::CreateReference(FContextCacheKey InReferenceDataKey) const
 	{
diff --git a/Engine/Source/Runtime/Experimental/Dataflow/Engine/Private/Dataflow/DataflowContent.cpp b/Engine/Source/Runtime/Experimental/Dataflow/Engine/Private/Dataflow/DataflowContent.cpp
--- a/Engine/Source/Runtime/Experimental/Dataflow/Engine/Private/Dataflow/DataflowContent.cpp
+++ b/Engine/Source/Runtime/Experimental/Dataflow/Engine/Private/Dataflow/DataflowContent.cpp
@@ -48,7 +48,7 @@ namespace UE::DataflowContextHelpers
 		{
 			bool bValidKey = false;
 
-			if (TUniquePtr<FContextCacheElementBase>* Data = Context->GetBaseData(Key))
+			if (const TUniquePtr<FContextCacheElementBase>* Data = Context->GetBaseData(Key))
 			{
 				if ((*Data) && !(*Data)->GetProperty())
 				{
@@ -102,7 +102,7 @@ namespace UE::DataflowContextHelpers
 		for (FContextCacheKey Key : Keys)
 		{
 			bool bValidKey = false;
-			if (TUniquePtr<FContextCacheElementBase>* Data = Context->GetBaseData(Key))
+			if (const TUniquePtr<FContextCacheElementBase>* Data = Context->GetBaseData(Key))
 			{
 				if (!(*Data)) return false;
 				if (!(*Data)->GetProperty()) return false;
@@ -140,7 +140,7 @@ namespace UE::DataflowContextHelpers
 		TSet<FContextCacheKey> Keys; Context->GetKeys(Keys);
 		for (FContextCacheKey Key : Keys)
 		{
-			if (TUniquePtr<FContextCacheElementBase>* Data = Context->GetBaseData(Key))
+			if (const TUniquePtr<FContextCacheElementBase>* Data = Context->GetBaseData(Key))
 			{
 				if (*Data)
 				{
diff --git a/Engine/Source/Runtime/Experimental/Dataflow/Engine/Private/Dataflow/DataflowEdNode.cpp b/Engine/Source/Runtime/Experimental/Dataflow/Engine/Private/Dataflow/DataflowEdNode.cpp
--- a/Engine/Source/Runtime/Experimental/Dataflow/Engine/Private/Dataflow/DataflowEdNode.cpp
+++ b/Engine/Source/Runtime/Experimental/Dataflow/Engine/Private/Dataflow/DataflowEdNode.cpp
@@ -26,31 +26,34 @@ DEFINE_LOG_CATEGORY_STATIC(DATAFLOWNODE_LOG, Error, All);
 
 namespace UE::Dataflow::Private
 {
-static UE::Dataflow::FPin::EDirection EdPinDirectionToDataflowDirection(EEdGraphPinDirection EdDirection)
-{
-	if (EdDirection == EEdGraphPinDirection::EGPD_Input)
-	{
-		return FPin::EDirection::INPUT;
-	}
-	if (EdDirection == EEdGraphPinDirection::EGPD_Output)
-	{
-		return FPin::EDirection::OUTPUT;
-	}
-	return FPin::EDirection::NONE;
-}
+	constexpr FLinearColor FrozenTitleColor(0.f, 0.7f, 1.f, 1.f);
+	constexpr FLinearColor FrozenBodyTintColor(0.f, 0.7f, 1.f, 0.5f);
 
-static EEdGraphPinDirection DataflowDirectionToEdPinDirection(UE::Dataflow::FPin::EDirection Direction)
-{
-	if (Direction == FPin::EDirection::INPUT)
+	static UE::Dataflow::FPin::EDirection EdPinDirectionToDataflowDirection(EEdGraphPinDirection EdDirection)
 	{
-		return EEdGraphPinDirection::EGPD_Input;
+		if (EdDirection == EEdGraphPinDirection::EGPD_Input)
+		{
+			return FPin::EDirection::INPUT;
+		}
+		if (EdDirection == EEdGraphPinDirection::EGPD_Output)
+		{
+			return FPin::EDirection::OUTPUT;
+		}
+		return FPin::EDirection::NONE;
 	}
-	if (Direction == FPin::EDirection::OUTPUT)
+
+	static EEdGraphPinDirection DataflowDirectionToEdPinDirection(UE::Dataflow::FPin::EDirection Direction)
 	{
-		return EEdGraphPinDirection::EGPD_Output;
+		if (Direction == FPin::EDirection::INPUT)
+		{
+			return EEdGraphPinDirection::EGPD_Input;
+		}
+		if (Direction == FPin::EDirection::OUTPUT)
+		{
+			return EEdGraphPinDirection::EGPD_Output;
+		}
+		return EEdGraphPinDirection::EGPD_MAX;
 	}
-	return EEdGraphPinDirection::EGPD_MAX;
-}
 }
 
 UDataflowEdNode::UDataflowEdNode(const FObjectInitializer& ObjectInitializer)
@@ -489,15 +492,13 @@ bool UDataflowEdNode::ShowPaletteIconOnNode() const
 
 FLinearColor UDataflowEdNode::GetNodeTitleColor() const
 {
-	if (DataflowGraph)
+	if (const TSharedPtr<const FDataflowNode> DataflowNode = GetDataflowNode())
 	{
-		if (DataflowNodeGuid.IsValid())
+		if (DataflowNode->IsFrozen())
 		{
-			if (TSharedPtr<FDataflowNode> DataflowNode = DataflowGraph->FindBaseNode(DataflowNodeGuid))
-			{
-				return UE::Dataflow::FNodeColorsRegistry::Get().GetNodeTitleColor(DataflowNode->GetCategory());
-			}
+			return UE::Dataflow::Private::FrozenTitleColor;
 		}
+		return UE::Dataflow::FNodeColorsRegistry::Get().GetNodeTitleColor(DataflowNode->GetCategory());
 	}
 	return FDataflowNode::DefaultNodeTitleColor;
 }
@@ -506,7 +507,11 @@ FLinearColor UDataflowEdNode::GetNodeBodyTintColor() const
 {
 	if (const TSharedPtr<const FDataflowNode> DataflowNode = GetDataflowNode())
 	{
-		UE::Dataflow::FNodeColorsRegistry::Get().GetNodeBodyTintColor(DataflowNode->GetCategory());
+		if (DataflowNode->IsFrozen())
+		{
+			return UE::Dataflow::Private::FrozenBodyTintColor;
+		}
+		return UE::Dataflow::FNodeColorsRegistry::Get().GetNodeBodyTintColor(DataflowNode->GetCategory());
 	}
 	return FDataflowNode::DefaultNodeBodyTintColor;
 }
diff --git a/Engine/Source/Runtime/Experimental/Dataflow/Engine/Public/Dataflow/DataflowObjectInterface.h b/Engine/Source/Runtime/Experimental/Dataflow/Engine/Public/Dataflow/DataflowObjectInterface.h
--- a/Engine/Source/Runtime/Experimental/Dataflow/Engine/Public/Dataflow/DataflowObjectInterface.h
+++ b/Engine/Source/Runtime/Experimental/Dataflow/Engine/Public/Dataflow/DataflowObjectInterface.h
@@ -27,7 +27,7 @@ namespace UE::Dataflow
 
 		int32 GetKeys(TSet<FContextCacheKey>& InKeys) const { return Base::GetKeys(InKeys); }
 
-		TUniquePtr<FContextCacheElementBase>* GetBaseData(FContextCacheKey Key) { return Base::GetDataImpl(Key); }
+		const TUniquePtr<FContextCacheElementBase>* GetBaseData(FContextCacheKey Key) const { return Base::GetDataImpl(Key); }
 
 		virtual void Serialize(FArchive& Ar) { Base::Serialize(Ar); }
 
