[LLM] FileSystem LLM tags review followup

#rnx

diff --git a/Engine/Source/Runtime/Experimental/IoStore/OnDemand/Private/OnDemandIoStore.cpp b/Engine/Source/Runtime/Experimental/IoStore/OnDemand/Private/OnDemandIoStore.cpp
--- a/Engine/Source/Runtime/Experimental/IoStore/OnDemand/Private/OnDemandIoStore.cpp
+++ b/Engine/Source/Runtime/Experimental/IoStore/OnDemand/Private/OnDemandIoStore.cpp
@@ -498,7 +498,7 @@ FOnDemandIoStore::~FOnDemandIoStore()
 
 FIoStatus FOnDemandIoStore::Initialize()
 {
-	LLM_SCOPE_BYNAME(TEXT("FileSystem/OnDemand"));
+	LLM_SCOPE_BYNAME(TEXT("FileSystem/OnDemandIoStore"));
 	const FIoStatus CacheStatus = InitializeOnDemandInstallCache();
 	if (CacheStatus.GetErrorCode() == EIoErrorCode::PendingFork)
 	{
@@ -615,7 +615,7 @@ void FOnDemandIoStore::Install(
 	FOnDemandInstallProgressed&& OnProgress /*= nullptr*/,
 	const FOnDemandCancellationToken* CancellationToken)
 {
-	LLM_SCOPE_BYNAME(TEXT("FileSystem/OnDemand"));
+	LLM_SCOPE_BYNAME(TEXT("FileSystem/OnDemandIoStore"));
 	FSharedInstallRequest InstallRequest = MakeShared<FInstallRequest>();
 	InstallRequest->Args				= MoveTemp(Args);
 	InstallRequest->OnCompleted			= MoveTemp(OnCompleted);
@@ -632,7 +632,7 @@ void FOnDemandIoStore::Install(
 
 void FOnDemandIoStore::Purge(FOnDemandPurgeArgs&& Args, FOnDemandPurgeCompleted&& OnCompleted)
 {
-	LLM_SCOPE_BYNAME(TEXT("FileSystem/OnDemand"));
+	LLM_SCOPE_BYNAME(TEXT("FileSystem/OnDemandIoStore"));
 	FSharedPurgeRequest PurgeRequest = MakeShared<FPurgeRequest>();
 	PurgeRequest->Args			= MoveTemp(Args);
 	PurgeRequest->OnCompleted	= MoveTemp(OnCompleted);
@@ -647,7 +647,7 @@ void FOnDemandIoStore::Purge(FOnDemandPurgeArgs&& Args, FOnDemandPurgeCompleted&
 
 FIoStatus FOnDemandIoStore::Unmount(FStringView MountId)
 {
-	LLM_SCOPE_BYNAME(TEXT("FileSystem/OnDemand"));
+	LLM_SCOPE_BYNAME(TEXT("FileSystem/OnDemandIoStore"));
 	UE_LOG(LogIoStoreOnDemand, Log, TEXT("Unmounting '%s'"), *WriteToString<256>(MountId));
 
 	bool bPendingMount = false;
@@ -955,7 +955,7 @@ void FOnDemandIoStore::TryEnterTickLoop()
 
 void FOnDemandIoStore::TickLoop()
 {
-	LLM_SCOPE_BYNAME(TEXT("FileSystem/OnDemand"));
+	LLM_SCOPE_BYNAME(TEXT("FileSystem/OnDemandIoStore"));
 	ON_SCOPE_EXIT { UE_LOG(LogIoStoreOnDemand, Verbose, TEXT("Exiting I/O store tick loop")); };
 
 	UE_LOG(LogIoStoreOnDemand, Verbose, TEXT("Entering I/O store tick loop"));
@@ -1706,7 +1706,7 @@ void FOnDemandIoStore::CreateContainersFromToc(
 	TArray<FSharedOnDemandContainer>& Out)
 {
 	static FName AssetClassName(TEXT("OnDemandIoStore"));
-	LLM_SCOPE_BYNAME(TEXT("FileSystem/OnDemand"));
+	LLM_SCOPE_BYNAME(TEXT("FileSystem/OnDemandIoStore"));
 	LLM_TAGSET_SCOPE(FName(MountId), ELLMTagSet::Assets);
 	LLM_TAGSET_SCOPE(AssetClassName, ELLMTagSet::AssetClasses);
 	UE_TRACE_METADATA_SCOPE_ASSET_FNAME(FName(MountId), AssetClassName, FName(TocPath));
@@ -1935,7 +1935,7 @@ FIoStatus FOnDemandIoStore::GetContainersAndPackagesForInstall(
 
 void FOnDemandIoStore::ReleaseContent(FOnDemandInternalContentHandle& ContentHandle)
 {
-	LLM_SCOPE_BYNAME(TEXT("FileSystem/OnDemand"));
+	LLM_SCOPE_BYNAME(TEXT("FileSystem/OnDemandIoStore"));
 	UE_LOG(LogIoStoreOnDemand, Log, TEXT("Releasing content handle '%s'"), *LexToString(ContentHandle));
 
 	UE::TUniqueLock Lock(ContainerMutex);
diff --git a/Engine/Source/Runtime/PakFile/Private/IoDispatcherFileBackend.cpp b/Engine/Source/Runtime/PakFile/Private/IoDispatcherFileBackend.cpp
--- a/Engine/Source/Runtime/PakFile/Private/IoDispatcherFileBackend.cpp
+++ b/Engine/Source/Runtime/PakFile/Private/IoDispatcherFileBackend.cpp
@@ -1365,7 +1365,7 @@ void FFileIoStore::Shutdown()
 
 TIoStatusOr<FIoContainerHeader> FFileIoStore::Mount(const TCHAR* InTocPath, int32 Order, const FGuid& EncryptionKeyGuid, const FAES::FAESKey& EncryptionKey)
 {
-	LLM_SCOPE_BYNAME(TEXT("FileSystem/IOStore"));
+	LLM_SCOPE_BYNAME(TEXT("FileSystem/FileIoStore"));
 	TUniquePtr<FFileIoStoreReader> Reader(new FFileIoStoreReader(*PlatformImpl, Stats));
 	FIoStatus IoStatus = Reader->Initialize(InTocPath, Order);
 	if (!IoStatus.IsOk())
@@ -1601,7 +1601,7 @@ FAutoConsoleTaskPriority CPrio_IoDispatcherTaskPriority(
 
 void FFileIoStore::ScatterBlock(FFileIoStoreCompressedBlock* CompressedBlock, bool bIsAsync)
 {
-	LLM_SCOPE_BYNAME(TEXT("FileSystem/IOStore"));
+	LLM_SCOPE_BYNAME(TEXT("FileSystem/FileIoStore"));
 	TRACE_CPUPROFILER_EVENT_SCOPE(IoDispatcherScatter);
 	
 	check(!CompressedBlock->bFailed);
@@ -1794,7 +1794,7 @@ namespace FileIoStoreImpl
 
 FIoRequestImpl* FFileIoStore::GetCompletedIoRequests()
 {
-	LLM_SCOPE_BYNAME(TEXT("FileSystem/IOStore"));
+	LLM_SCOPE_BYNAME(TEXT("FileSystem/FileIoStore"));
 	//TRACE_CPUPROFILER_EVENT_SCOPE(GetCompletedRequests);
 	
 	if (!bIsMultithreaded)
