Compute Spheres collision in local space
More than 10% Narrow Phase performance improvement on Switch

diff --git a/Engine/Source/Runtime/Experimental/Chaos/Private/Chaos/Collision/ContactPointsMiscShapes.cpp b/Engine/Source/Runtime/Experimental/Chaos/Private/Chaos/Collision/ContactPointsMiscShapes.cpp
--- a/Engine/Source/Runtime/Experimental/Chaos/Private/Chaos/Collision/ContactPointsMiscShapes.cpp
+++ b/Engine/Source/Runtime/Experimental/Chaos/Private/Chaos/Collision/ContactPointsMiscShapes.cpp
@@ -269,28 +269,29 @@ namespace Chaos
 	}
 
 
-	FContactPoint SphereSphereContactPoint(const FSphere& Sphere1, const FRigidTransform3& Sphere1Transform, const FSphere& Sphere2, const FRigidTransform3& Sphere2Transform, const FReal CullDistance)
+	FContactPoint SphereSphereContactPoint(const FSphere& Sphere1, const FRigidTransform3& Sphere1Transform, const FSphere& Sphere2, const FRigidTransform3& Sphere2Transform, const FRealSingle CullDistance)
 	{
 		FContactPoint Result;
 
-		const FReal R1 = Sphere1.GetRadiusf();
-		const FReal R2 = Sphere2.GetRadiusf();
+		const FRealSingle R1 = Sphere1.GetRadiusf();
+		const FRealSingle R2 = Sphere2.GetRadiusf();
 
-		// World-space contact
-		const FVec3 Center1 = Sphere1Transform.TransformPosition(FVec3(Sphere1.GetCenterf()));
-		const FVec3 Center2 = Sphere2Transform.TransformPosition(FVec3(Sphere2.GetCenterf()));
-		const FVec3 Direction = Center1 - Center2;
-		const FReal SizeSq = Direction.SizeSquared();
-		const FReal CullDistanceSq = FMath::Square(CullDistance + R1 + R2);
+		const FRigidTransform3f Sphere1ToSphere2f = FRigidTransform3f(Sphere1Transform.GetRelativeTransformNoScale(Sphere2Transform));
+
+		const FVec3f Center1 = Sphere1ToSphere2f.TransformPosition(Sphere1.GetCenterf());
+		const FVec3f Center2 = Sphere2.GetCenterf();
+		const FVec3f Direction = Center1 - Center2;
+		const FRealSingle SizeSq = Direction.SizeSquared();
+		const FRealSingle CullDistanceSq = FMath::Square(CullDistance + R1 + R2);
 		if (SizeSq < CullDistanceSq)
 		{
-			const FReal Size = FMath::Sqrt(SizeSq);
-			const FVec3 Normal = Size > UE_SMALL_NUMBER ? Direction / Size : FVec3(0, 0, 1);
-			const FReal NewPhi = Size - (R1 + R2);
+			const FRealSingle Size = FMath::Sqrt(SizeSq);
+			const FVec3f Normal = Size > UE_KINDA_SMALL_NUMBER ? Direction / Size : FVec3f(0, 0, 1);
+			const FRealSingle NewPhi = Size - (R1 + R2);
 
-			Result.ShapeContactPoints[0] = FVec3(Sphere1.GetCenterf()) - Sphere1Transform.InverseTransformVector(R1 * Normal);
-			Result.ShapeContactPoints[1] = FVec3(Sphere2.GetCenterf()) + Sphere2Transform.InverseTransformVector(R2 * Normal);
-			Result.ShapeContactNormal = Sphere2Transform.InverseTransformVector(Normal);
+			Result.ShapeContactPoints[0] = Sphere1.GetCenterf() - Sphere1ToSphere2f.InverseTransformVector(R1 * Normal);
+			Result.ShapeContactPoints[1] = Sphere2.GetCenterf() + R2 * Normal;
+			Result.ShapeContactNormal = Normal;
 			Result.Phi = NewPhi;
 		}
 		return Result;
diff --git a/Engine/Source/Runtime/Experimental/Chaos/Private/Chaos/CollisionOneShotManifoldsMiscShapes.cpp b/Engine/Source/Runtime/Experimental/Chaos/Private/Chaos/CollisionOneShotManifoldsMiscShapes.cpp
--- a/Engine/Source/Runtime/Experimental/Chaos/Private/Chaos/CollisionOneShotManifoldsMiscShapes.cpp
+++ b/Engine/Source/Runtime/Experimental/Chaos/Private/Chaos/CollisionOneShotManifoldsMiscShapes.cpp
@@ -47,7 +47,7 @@ namespace Chaos
 			// @todo(chaos): support manifold maintenance
 			Constraint.ResetActiveManifoldContacts();
 
-			FContactPoint ContactPoint = SphereSphereContactPoint(SphereA, SphereATransform, SphereB, SphereBTransform, Constraint.GetCullDistance());
+			FContactPoint ContactPoint = SphereSphereContactPoint(SphereA, SphereATransform, SphereB, SphereBTransform, Constraint.GetCullDistancef());
 			if (ContactPoint.Phi < Constraint.GetCullDistance())
 			{
 				Constraint.AddOneshotManifoldContact(ContactPoint);
@@ -110,26 +110,26 @@ namespace Chaos
 			Constraint.ResetActiveManifoldContacts();
 
 			// Build a multi-point manifold
-			const FReal NetCullDistance = Sphere.GetRadiusf() + Capsule.GetRadiusf() + Constraint.GetCullDistance();
-			const FReal NetCullDistanceSq = FMath::Square(NetCullDistance);
+			const FRealSingle NetCullDistance = Sphere.GetRadiusf() + Capsule.GetRadiusf() + Constraint.GetCullDistancef();
+			const FRealSingle NetCullDistanceSq = FMath::Square(NetCullDistance);
 
 			// Transform the sphere into capsule space and find the closest point on the capsule line segment
 			// @todo(chaos) this would be much simpler if the sphere's were always at the origin and capsules were at the origin and axis aligned
-			const FRigidTransform3 SphereToCapsuleTransform = SphereTransform.GetRelativeTransformNoScale(CapsuleTransform);
-			const FVec3 SpherePos = SphereToCapsuleTransform.TransformPositionNoScale(FVec3(Sphere.GetCenterf()));
-			const FReal NearPosT = Utilities::ClosestTimeOnLineSegment(FVec3f(SpherePos), Capsule.GetX1f(), Capsule.GetX2f());
+			const FRigidTransform3f SphereToCapsuleTransform = FRigidTransform3f(SphereTransform.GetRelativeTransformNoScale(CapsuleTransform));
+			const FVec3f SpherePos = SphereToCapsuleTransform.TransformPositionNoScale(Sphere.GetCenterf());
+			const FRealSingle NearPosT = Utilities::ClosestTimeOnLineSegment(SpherePos, Capsule.GetX1f(), Capsule.GetX2f());
 
 			// Add the closest contact point to the manifold
-			const FVec3 NearPos = FMath::Lerp(Capsule.GetX1f(), Capsule.GetX2f(), NearPosT);
-			const FVec3 NearPosDelta = SpherePos - NearPos;
-			const FReal NearPosDistanceSq = NearPosDelta.SizeSquared();
+			const FVec3f NearPos = FMath::Lerp(Capsule.GetX1f(), Capsule.GetX2f(), NearPosT);
+			const FVec3f NearPosDelta = SpherePos - NearPos;
+			const FRealSingle NearPosDistanceSq = NearPosDelta.SizeSquared();
 			if (NearPosDistanceSq > UE_SMALL_NUMBER)
 			{
 				if (NearPosDistanceSq < NetCullDistanceSq)
 				{
-					const FReal NearPosDistance = FMath::Sqrt(NearPosDistanceSq);
+					const FRealSingle NearPosDistance = FMath::Sqrt(NearPosDistanceSq);
 					const FVec3 NearPosDir = NearPosDelta / NearPosDistance;
-					const FReal NearPhi = NearPosDistance - Sphere.GetRadiusf() - Capsule.GetRadiusf();
+					const FRealSingle NearPhi = NearPosDistance - Sphere.GetRadiusf() - Capsule.GetRadiusf();
 
 					FContactPoint NearContactPoint;
 					NearContactPoint.ShapeContactPoints[0] = SphereToCapsuleTransform.InverseTransformPositionNoScale(SpherePos - Sphere.GetRadiusf() * NearPosDir);
@@ -141,7 +141,7 @@ namespace Chaos
 					Constraint.AddOneshotManifoldContact(NearContactPoint);
 
 					// If we have a small sphere, just stick with the 1-point manifold
-					const FReal SphereCapsuleSizeThreshold = FReal(Chaos_Collision_Manifold_SphereCapsuleSizeThreshold);
+					const FRealSingle SphereCapsuleSizeThreshold = Chaos_Collision_Manifold_SphereCapsuleSizeThreshold;
 					if (Sphere.GetRadiusf() < SphereCapsuleSizeThreshold * (Capsule.GetHeightf() + Capsule.GetRadiusf()))
 					{
 						return;
@@ -156,7 +156,7 @@ namespace Chaos
 					// If the contact is deep, there's a high chance that pushing one end out will push the other deeper and we also need more contacts.
 					// Note: we only consider the radius of the dynamic object(s) when deciding what "deep" means because the extra contacts are only
 					// to prevent excessive rotation from the single contact we have so far, and only the dynamic objects will rotate.
-					const FReal DeepRadiusFraction = Chaos_Collision_Manifold_CapsuleDeepPenetrationFraction;
+					const FRealSingle DeepRadiusFraction = Chaos_Collision_Manifold_CapsuleDeepPenetrationFraction;
 					const bool bIsDeep = NearPhi < -DeepRadiusFraction * Capsule.GetRadiusf();
 					if (!bIsDeep)
 					{
@@ -167,31 +167,31 @@ namespace Chaos
 					// Calculate the vector orthogonal to the capsule axis that gives the nearest points on the capsule cyclinder to the sphere
 					// The initial length will be proportional to the sine of the angle between the axis and the delta position and will approach
 					// zero when the capsule is end-on to the sphere, in which case we won't add the end caps.
-					constexpr FReal EndCapSinAngleThreshold = FReal(0.35);	// about 20deg
-					constexpr FReal EndCapDistanceThreshold = FReal(0.2);	// fraction
-					FVec3 CapsuleOrthogonal = FVec3::CrossProduct(Capsule.GetAxis(), FVec3::CrossProduct(Capsule.GetAxis(), NearPosDir));
-					const FReal CapsuleOrthogonalLenSq = CapsuleOrthogonal.SizeSquared();
+					constexpr FRealSingle EndCapSinAngleThreshold = 0.35f;	// about 20deg
+					constexpr FRealSingle EndCapDistanceThreshold = 0.2f;	// fraction
+					FVec3f CapsuleOrthogonal = FVec3f::CrossProduct(Capsule.GetAxisf(), FVec3::CrossProduct(Capsule.GetAxisf(), NearPosDir));
+					const FRealSingle CapsuleOrthogonalLenSq = CapsuleOrthogonal.SizeSquared();
 					if (CapsuleOrthogonalLenSq > FMath::Square(EndCapSinAngleThreshold))
 					{
 						// Orthogonal must point towards the sphere, but currently depends on the relative axis orientation
 						CapsuleOrthogonal = CapsuleOrthogonal * FMath::InvSqrt(CapsuleOrthogonalLenSq);
-						if (FVec3::DotProduct(CapsuleOrthogonal, SpherePos - Capsule.GetCenterf()) < FReal(0))
+						if (FVec3f::DotProduct(CapsuleOrthogonal, SpherePos - Capsule.GetCenterf()) < 0.0f)
 						{
 							CapsuleOrthogonal = -CapsuleOrthogonal;
 						}
 
 						if (NearPosT > EndCapDistanceThreshold)
 						{
-							const FVec3 EndCapPos0 = Capsule.GetX1f() + CapsuleOrthogonal * Capsule.GetRadiusf();
-							const FReal EndCapDistance0 = (SpherePos - EndCapPos0).Size();
-							const FReal EndCapPhi0 = EndCapDistance0 - Sphere.GetRadiusf();
+							const FVec3f EndCapPos0 = Capsule.GetX1f() + CapsuleOrthogonal * Capsule.GetRadiusf();
+							const FRealSingle EndCapDistance0 = (SpherePos - EndCapPos0).Size();
+							const FRealSingle EndCapPhi0 = EndCapDistance0 - Sphere.GetRadiusf();
 							
 							if (EndCapPhi0 < Constraint.GetCullDistance())
 							{
-								const FVec3 EndCapPosDir0 = (SpherePos - EndCapPos0) / EndCapDistance0;
-								const FVec3 SpherePos0 = SpherePos - EndCapPosDir0 * Sphere.GetRadiusf();
+								const FVec3f EndCapPosDir0 = (SpherePos - EndCapPos0) / EndCapDistance0;
+								const FVec3f SpherePos0 = SpherePos - EndCapPosDir0 * Sphere.GetRadiusf();
 						
-								FContactPoint EndCapContactPoint0;
+								FContactPointf EndCapContactPoint0;
 								EndCapContactPoint0.ShapeContactPoints[0] = SphereToCapsuleTransform.InverseTransformPositionNoScale(SpherePos0);
 								EndCapContactPoint0.ShapeContactPoints[1] = EndCapPos0;
 								EndCapContactPoint0.ShapeContactNormal = EndCapPosDir0;
@@ -202,18 +202,18 @@ namespace Chaos
 							}
 						}
 
-						if (NearPosT < FReal(1) - EndCapDistanceThreshold)
+						if (NearPosT < 1.0f - EndCapDistanceThreshold)
 						{
-							const FVec3 EndCapPos1 = Capsule.GetX2f() + CapsuleOrthogonal * Capsule.GetRadiusf();
-							const FReal EndCapDistance1 = (SpherePos - EndCapPos1).Size();
-							const FReal EndCapPhi1 = EndCapDistance1 - Sphere.GetRadiusf();
+							const FVec3f EndCapPos1 = Capsule.GetX2f() + CapsuleOrthogonal * Capsule.GetRadiusf();
+							const FRealSingle EndCapDistance1 = (SpherePos - EndCapPos1).Size();
+							const FRealSingle EndCapPhi1 = EndCapDistance1 - Sphere.GetRadiusf();
 
 							if (EndCapPhi1 < Constraint.GetCullDistance())
 							{
-								const FVec3 EndCapPosDir1 = (SpherePos - EndCapPos1) / EndCapDistance1;
-								const FVec3 SpherePos1 = SpherePos - EndCapPosDir1 * Sphere.GetRadiusf();
+								const FVec3f EndCapPosDir1 = (SpherePos - EndCapPos1) / EndCapDistance1;
+								const FVec3f SpherePos1 = SpherePos - EndCapPosDir1 * Sphere.GetRadiusf();
 
-								FContactPoint EndCapContactPoint0;
+								FContactPointf EndCapContactPoint0;
 								EndCapContactPoint0.ShapeContactPoints[0] = SphereToCapsuleTransform.InverseTransformPositionNoScale(SpherePos1);
 								EndCapContactPoint0.ShapeContactPoints[1] = EndCapPos1;
 								EndCapContactPoint0.ShapeContactNormal = EndCapPosDir1;
diff --git a/Engine/Source/Runtime/Experimental/Chaos/Private/Chaos/CollisionResolution.cpp b/Engine/Source/Runtime/Experimental/Chaos/Private/Chaos/CollisionResolution.cpp
--- a/Engine/Source/Runtime/Experimental/Chaos/Private/Chaos/CollisionResolution.cpp
+++ b/Engine/Source/Runtime/Experimental/Chaos/Private/Chaos/CollisionResolution.cpp
@@ -507,7 +507,7 @@ namespace Chaos
 			}
 			else
 			{
-				UpdateContactPoint(Constraint, SphereSphereContactPoint(Sphere1, Sphere1Transform, Sphere2, Sphere2Transform, Constraint.GetCullDistance()), Dt);
+				UpdateContactPoint(Constraint, SphereSphereContactPoint(Sphere1, Sphere1Transform, Sphere2, Sphere2Transform, Constraint.GetCullDistancef()), Dt);
 			}
 		}
 
diff --git a/Engine/Source/Runtime/Experimental/Chaos/Public/Chaos/Collision/ContactPointsMiscShapes.h b/Engine/Source/Runtime/Experimental/Chaos/Public/Chaos/Collision/ContactPointsMiscShapes.h
--- a/Engine/Source/Runtime/Experimental/Chaos/Public/Chaos/Collision/ContactPointsMiscShapes.h
+++ b/Engine/Source/Runtime/Experimental/Chaos/Public/Chaos/Collision/ContactPointsMiscShapes.h
@@ -38,7 +38,7 @@ namespace Chaos
 	template <typename GeometryA, typename GeometryB>
 	FContactPoint GJKImplicitContactPoint(const FImplicitObject& A, const FRigidTransform3& ATransform, const GeometryB& B, const FRigidTransform3& BTransform, const FReal CullDistance);
 
-	FContactPoint SphereSphereContactPoint(const FSphere& Sphere1, const FRigidTransform3& Sphere1Transform, const FSphere& Sphere2, const FRigidTransform3& Sphere2Transform, const FReal CullDistance);
+	FContactPoint SphereSphereContactPoint(const FSphere& Sphere1, const FRigidTransform3& Sphere1Transform, const FSphere& Sphere2, const FRigidTransform3& Sphere2Transform, const FRealSingle CullDistance);
 	FContactPoint SpherePlaneContactPoint(const FSphere& Sphere, const FRigidTransform3& SphereTransform, const TPlane<FReal, 3>& Plane, const FRigidTransform3& PlaneTransform);
 	FContactPoint SphereBoxContactPoint(const FSphere& Sphere, const FRigidTransform3& SphereTransform, const FImplicitBox3& Box, const FRigidTransform3& BoxTransform);
 	FContactPoint SphereCapsuleContactPoint(const FSphere& A, const FRigidTransform3& ATransform, const FCapsule& B, const FRigidTransform3& BTransform);
diff --git a/Engine/Source/Runtime/Experimental/Chaos/Public/Chaos/Collision/PBDCollisionConstraint.h b/Engine/Source/Runtime/Experimental/Chaos/Public/Chaos/Collision/PBDCollisionConstraint.h
--- a/Engine/Source/Runtime/Experimental/Chaos/Public/Chaos/Collision/PBDCollisionConstraint.h
+++ b/Engine/Source/Runtime/Experimental/Chaos/Public/Chaos/Collision/PBDCollisionConstraint.h
@@ -490,6 +490,7 @@ namespace Chaos
 		CHAOS_API FString ToString() const;
 
 		FReal GetCullDistance() const { return CullDistance; }
+		FRealSingle GetCullDistancef() const { return CullDistance; }
 		void SetCullDistance(FReal InCullDistance) { CullDistance = FRealSingle(InCullDistance); }
 
 		FVec3f GetRelativeMovement() const { return RelativeMovement; }
@@ -584,7 +585,7 @@ namespace Chaos
 		CHAOS_API void AddIncrementalManifoldContact(const FContactPoint& ContactPoint);
 
 		// @todo(chaos): remove this and use SetOneShotManifoldContacts
-		inline void AddOneshotManifoldContact(const FContactPoint& ContactPoint)
+		inline void AddOneshotManifoldContact(const FContactPointf& ContactPoint)
 		{
 			if (ContactPoint.IsSet())
 			{
@@ -886,7 +887,7 @@ namespace Chaos
 		CHAOS_API int32 FindSavedManifoldPoint(const int32 ManifoldPointIndex, int32* InOutAllowedSavedPointIndices, int32& InOutNumAllowedSavedPoints) const;
 		CHAOS_API void AssignSavedManifoldPoints();
 
-		inline void InitManifoldPoint(const int32 ManifoldPointIndex, const FContactPoint& ContactPoint)
+		inline void InitManifoldPoint(const int32 ManifoldPointIndex, const FContactPointf& ContactPoint)
 		{
 			FManifoldPoint& ManifoldPoint = ManifoldPoints[ManifoldPointIndex];
 			ManifoldPoint.ContactPoint = ContactPoint;
@@ -899,7 +900,7 @@ namespace Chaos
 			ManifoldPoint.Flags.Reset();
 		}
 
-		inline int32 AddManifoldPoint(const FContactPoint& ContactPoint)
+		inline int32 AddManifoldPoint(const FContactPointf& ContactPoint)
 		{
 			int32 ManifoldPointIndex = ManifoldPoints.AddUninitialized();	// Note: no initialization (see InitManifoldPoint)
 			InitManifoldPoint(ManifoldPointIndex, ContactPoint);
