Remove stall on render fence when releasing geometry collection render data.
Changed to using BeginDestroy()/IsReadyForFinishDestroy() pattern on the owning UGeometryCollection object.
#tests ReplayRun tests to check that we no longer hit the slow path

diff --git a/Engine/Source/Runtime/Experimental/GeometryCollectionEngine/Private/GeometryCollection/GeometryCollectionObject.cpp b/Engine/Source/Runtime/Experimental/GeometryCollectionEngine/Private/GeometryCollection/GeometryCollectionObject.cpp
--- a/Engine/Source/Runtime/Experimental/GeometryCollectionEngine/Private/GeometryCollection/GeometryCollectionObject.cpp
+++ b/Engine/Source/Runtime/Experimental/GeometryCollectionEngine/Private/GeometryCollection/GeometryCollectionObject.cpp
@@ -2015,6 +2015,14 @@ void UGeometryCollection::BeginDestroy()
 {
 	Super::BeginDestroy();
 	ReleaseResources();
+
+	// Render fence to ensure that we complete resource release before destroy.
+	RenderResourceReleaseFence.BeginFence();
+}
+
+bool UGeometryCollection::IsReadyForFinishDestroy()
+{
+	return Super::IsReadyForFinishDestroy() && RenderResourceReleaseFence.IsFenceComplete();
 }
 
 bool UGeometryCollection::HasMeshData() const
diff --git a/Engine/Source/Runtime/Experimental/GeometryCollectionEngine/Private/GeometryCollection/GeometryCollectionRenderData.cpp b/Engine/Source/Runtime/Experimental/GeometryCollectionEngine/Private/GeometryCollection/GeometryCollectionRenderData.cpp
--- a/Engine/Source/Runtime/Experimental/GeometryCollectionEngine/Private/GeometryCollection/GeometryCollectionRenderData.cpp
+++ b/Engine/Source/Runtime/Experimental/GeometryCollectionEngine/Private/GeometryCollection/GeometryCollectionRenderData.cpp
@@ -702,11 +702,8 @@ TUniquePtr<FGeometryCollectionRenderData> FGeometryCollectionRenderData::Create(
 
 FGeometryCollectionRenderData::~FGeometryCollectionRenderData()
 {
-	ReleaseResources();
-
-	FRenderCommandFence Fence;
-	Fence.BeginFence();
-	Fence.Wait();
+	// We should have released any resources before now.
+	check(!bIsInitialized);
 }
 
 void FGeometryCollectionRenderData::Serialize(FArchive& Ar, UGeometryCollection& Owner)
@@ -807,14 +804,7 @@ void FGeometryCollectionRenderData::ReleaseResources()
 
 	if (bHasNaniteData)
 	{
-		if (NaniteResourcesPtr->ReleaseResources())
-		{
-			// HACK: Make sure the renderer is done processing the command, and done using NaniteResource, before we continue.
-			// This code could really use a refactor.
-			FRenderCommandFence Fence;
-			Fence.BeginFence();
-			Fence.Wait();
-		}
+		NaniteResourcesPtr->ReleaseResources();
 	}
 
 	bIsInitialized = false;
diff --git a/Engine/Source/Runtime/Experimental/GeometryCollectionEngine/Public/GeometryCollection/GeometryCollectionObject.h b/Engine/Source/Runtime/Experimental/GeometryCollectionEngine/Public/GeometryCollection/GeometryCollectionObject.h
--- a/Engine/Source/Runtime/Experimental/GeometryCollectionEngine/Public/GeometryCollection/GeometryCollectionObject.h
+++ b/Engine/Source/Runtime/Experimental/GeometryCollectionEngine/Public/GeometryCollection/GeometryCollectionObject.h
@@ -375,6 +375,7 @@ public:
 	GEOMETRYCOLLECTIONENGINE_API virtual void PreSave(FObjectPreSaveContext SaveContext) override;
 	GEOMETRYCOLLECTIONENGINE_API virtual void PostLoad() override;
 	GEOMETRYCOLLECTIONENGINE_API virtual void BeginDestroy() override;
+	GEOMETRYCOLLECTIONENGINE_API virtual bool IsReadyForFinishDestroy() override;
 	/** End UObject Interface */
 
 	GEOMETRYCOLLECTIONENGINE_API void Serialize(FArchive& Ar);
@@ -923,5 +924,8 @@ private:
 	UPROPERTY(EditAnywhere, AdvancedDisplay, Instanced, Category = AssetUserData)
 	TArray<TObjectPtr<UAssetUserData>> AssetUserData;
 
+	/** A fence to track render resource release before destroy. */
+	FRenderCommandFence RenderResourceReleaseFence;
+
 	float GetMassOrDensityInternal(bool& bOutIsDensity, bool bCached) const;
 };
