Physics Mover Updates:
- Option to apply reaction force to ground body when moving in ground plane
- Scale radial force limit rather then applying a velocity to target when walking
- Make query radius editable per mode

diff --git a/Engine/Plugins/Experimental/Mover/Source/Mover/Private/Backends/MoverNetworkPhysicsLiaison.cpp b/Engine/Plugins/Experimental/Mover/Source/Mover/Private/Backends/MoverNetworkPhysicsLiaison.cpp
--- a/Engine/Plugins/Experimental/Mover/Source/Mover/Private/Backends/MoverNetworkPhysicsLiaison.cpp
+++ b/Engine/Plugins/Experimental/Mover/Source/Mover/Private/Backends/MoverNetworkPhysicsLiaison.cpp
@@ -727,7 +727,7 @@ void UMoverNetworkPhysicsLiaisonComponent::ProduceInput_External(int32 PhysicsSt
 		APawn* PawnOwner = Cast<APawn>(GetOwner());
 		bool bProduceInput = PawnOwner ? PawnOwner->IsLocallyControlled() : false;
 
-		if (PawnOwner && !PawnOwner->IsPlayerControlled() && !NetworkPhysicsComponent->GetIsRelayingLocalInputs() && NetworkPhysicsComponent->HasServerWorld())
+		if (PawnOwner && !PawnOwner->IsPlayerControlled() && NetworkPhysicsComponent && !NetworkPhysicsComponent->GetIsRelayingLocalInputs() && NetworkPhysicsComponent->HasServerWorld())
 		{
 			NetworkPhysicsComponent->SetIsRelayingLocalInputs(true);
 		}
@@ -940,6 +940,8 @@ void UMoverNetworkPhysicsLiaisonComponent::OnPreSimulate_Internal(const FPhysics
 		return;
 	}
 
+	ConstraintHandle->SetGroundParticle(nullptr);
+
 	// Update sync state from physics
 	Chaos::FPBDRigidParticleHandle* CharacterParticle = ConstraintHandle->GetCharacterParticle()->CastToRigidParticle();
 	if (!CharacterParticle || CharacterParticle->Disabled())
@@ -964,7 +966,8 @@ void UMoverNetworkPhysicsLiaisonComponent::OnPreSimulate_Internal(const FPhysics
 
 	// Make the sync state velocity relative to the ground if walking
 	FVector LocalGroundVelocity = FVector::ZeroVector;
-	if (Input.SyncState.MovementMode == DefaultModeNames::Walking)
+	const UBaseMovementMode* InputMode = MoverComp->MovementModes[Input.SyncState.MovementMode];
+	if (InputMode->HasGameplayTag(Mover_IsOnGround, true))
 	{
 		if (const UMoverBlackboard* Blackboard = MoverComp->GetSimBlackboard())
 		{
@@ -1040,25 +1043,12 @@ void UMoverNetworkPhysicsLiaisonComponent::OnPreSimulate_Internal(const FPhysics
 	// Add back the ground velocity that was subtracted to but the movement velocity in local space
 	FVector TargetVelocity = OutputSyncState.GetVelocity_WorldSpace() + LocalGroundVelocity;
 
-	// Landed so add the new ground velocity
-	if ((Output.SyncState.MovementMode == DefaultModeNames::Walking) && (Input.SyncState.MovementMode != DefaultModeNames::Walking))
-	{
-		if (const UPhysicsDrivenWalkingMode* WalkingMode = Cast<UPhysicsDrivenWalkingMode>(MoverComp->FindMovementMode(UPhysicsDrivenWalkingMode::StaticClass())))
-		{
-			LocalGroundVelocity = UPhysicsMovementUtils::ComputeGroundVelocityFromHitResult(CharacterParticle->GetX(), Output.FloorResult.HitResult, TickParams.DeltaTimeSeconds);
-			LocalGroundVelocity -= LocalGroundVelocity.ProjectOnToNormal(Output.FloorResult.HitResult.ImpactNormal);
-			TargetVelocity += WalkingMode->FractionalVelocityToTarget * LocalGroundVelocity;
-		}
-	}
-
-	CharacterParticle->SetV(TargetVelocity);
-
-	// Note: Output sync state does not have a target angular velocity so
-	// use the target orientation
 	FRotator DeltaRotation = OutputSyncState.GetOrientation_WorldSpace() - FRotator(CharacterParticle->GetR());
 	FRotator Winding, Remainder;
 	DeltaRotation.GetWindingAndRemainder(Winding, Remainder);
 	float TargetDeltaFacing = FMath::DegreesToRadians(Remainder.Yaw);
+
+	CharacterParticle->SetV(TargetVelocity);
 	if (TickParams.DeltaTimeSeconds > UE_SMALL_NUMBER)
 	{
 		FVector AngularVelocity = (TargetDeltaFacing / TickParams.DeltaTimeSeconds) * Chaos::FVec3::ZAxisVector;
@@ -1083,14 +1073,12 @@ void UMoverNetworkPhysicsLiaisonComponent::OnPreSimulate_Internal(const FPhysics
 		// Set the ground particle on the constraint
 		Chaos::FGeometryParticleHandle* GroundParticle = nullptr;
 
-		if (IPhysicsComponent* PhysicsComp = Cast<IPhysicsComponent>(Output.FloorResult.HitResult.Component))
+		if (Chaos::FPhysicsObjectHandle PhysicsObject = Output.FloorResult.HitResult.PhysicsObject)
 		{
-			if (Chaos::FPhysicsObjectHandle PhysicsObject = PhysicsComp->GetPhysicsObjectById(Output.FloorResult.HitResult.Item))
+			Chaos::FReadPhysicsObjectInterface_Internal Interface = Chaos::FPhysicsObjectInternalInterface::GetRead();
+			if (!Interface.AreAllDisabled({ PhysicsObject }))
 			{
-				Chaos::FReadPhysicsObjectInterface_Internal Interface = Chaos::FPhysicsObjectInternalInterface::GetRead();
 				GroundParticle = Interface.GetParticle(PhysicsObject);
-
-				// Wake the ground particle if it is sleeping
 				WakeParticleIfSleeping(GroundParticle);
 			}
 		}
@@ -1124,7 +1112,7 @@ void UMoverNetworkPhysicsLiaisonComponent::OnPreSimulate_Internal(const FPhysics
 		OutputSyncState.SetMovementBase(nullptr);
 
 		ConstraintHandle->SetData({
-			Chaos::FVec3::ZAxisVector,
+			MoverComp->GetUpDirection(),
 			Chaos::FVec3::ZeroVector,
 			0.0,
 			1.0e10,
diff --git a/Engine/Plugins/Experimental/Mover/Source/Mover/Private/PhysicsMover/Modes/PhysicsDrivenFallingMode.cpp b/Engine/Plugins/Experimental/Mover/Source/Mover/Private/PhysicsMover/Modes/PhysicsDrivenFallingMode.cpp
--- a/Engine/Plugins/Experimental/Mover/Source/Mover/Private/PhysicsMover/Modes/PhysicsDrivenFallingMode.cpp
+++ b/Engine/Plugins/Experimental/Mover/Source/Mover/Private/PhysicsMover/Modes/PhysicsDrivenFallingMode.cpp
@@ -84,13 +84,7 @@ void UPhysicsDrivenFallingMode::OnSimulationTick(const FSimulationTickParams& Pa
 	FFloorCheckResult FloorResult;
 	FWaterCheckResult WaterResult;
 
-	float PawnHalfHeight;
-	float PawnRadius;
-	UpdatedPrimitive->CalcBoundingCylinder(PawnRadius, PawnHalfHeight);
-
 	const float QueryDistance = 1.1f * FMath::Max(TargetHeight, TargetHeight - UpDir.Dot(ProposedMove.LinearVelocity) * DeltaSeconds);
-	const float ShrinkRadius = 1.0f; // TODO - Make this a user setting
-	const float QueryRadius = FMath::Max(PawnRadius - ShrinkRadius, 0.0f);
 
 	UPhysicsMovementUtils::FloorSweep_Internal(StartingSyncState->GetLocation_WorldSpace(), StartingSyncState->GetVelocity_WorldSpace() * DeltaSeconds,
 		UpdatedPrimitive, UpDir, QueryRadius, QueryDistance, CommonLegacySettings->MaxWalkSlopeCosine, TargetHeight, FloorResult, WaterResult);
diff --git a/Engine/Plugins/Experimental/Mover/Source/Mover/Private/PhysicsMover/Modes/PhysicsDrivenSwimmingMode.cpp b/Engine/Plugins/Experimental/Mover/Source/Mover/Private/PhysicsMover/Modes/PhysicsDrivenSwimmingMode.cpp
--- a/Engine/Plugins/Experimental/Mover/Source/Mover/Private/PhysicsMover/Modes/PhysicsDrivenSwimmingMode.cpp
+++ b/Engine/Plugins/Experimental/Mover/Source/Mover/Private/PhysicsMover/Modes/PhysicsDrivenSwimmingMode.cpp
@@ -98,7 +98,7 @@ void UPhysicsDrivenSwimmingMode::OnSimulationTick(const FSimulationTickParams& P
 	const float QueryDistance= 2.0f * PawnHalfHeight;
 	
 	UPhysicsMovementUtils::FloorSweep_Internal(StartingSyncState->GetLocation_WorldSpace(), StartingSyncState->GetVelocity_WorldSpace() * DeltaSeconds,
-		UpdatedPrimitive, UpDir, PawnRadius, QueryDistance, CommonLegacySettings->MaxWalkSlopeCosine, TargetHeight, FloorResult, WaterResult);
+		UpdatedPrimitive, UpDir, QueryRadius, QueryDistance, CommonLegacySettings->MaxWalkSlopeCosine, TargetHeight, FloorResult, WaterResult);
 
 	if (WaterResult.IsSwimmableVolume())
 	{
diff --git a/Engine/Plugins/Experimental/Mover/Source/Mover/Private/PhysicsMover/Modes/PhysicsDrivenWalkingMode.cpp b/Engine/Plugins/Experimental/Mover/Source/Mover/Private/PhysicsMover/Modes/PhysicsDrivenWalkingMode.cpp
--- a/Engine/Plugins/Experimental/Mover/Source/Mover/Private/PhysicsMover/Modes/PhysicsDrivenWalkingMode.cpp
+++ b/Engine/Plugins/Experimental/Mover/Source/Mover/Private/PhysicsMover/Modes/PhysicsDrivenWalkingMode.cpp
@@ -43,6 +43,8 @@ void UPhysicsDrivenWalkingMode::UpdateConstraintSettings(Chaos::FCharacterGround
 	Constraint.SetSwingTorqueLimit(FUnitConversion::Convert(SwingTorqueLimit, EUnit::NewtonMeters, EUnit::KilogramCentimetersSquaredPerSecondSquared));
 	Constraint.SetTargetHeight(TargetHeight);
 	Constraint.SetDampingFactor(GroundDamping);
+	Constraint.SetMotionTargetMassBias(FractionalGroundReaction);
+	Constraint.SetRadialForceMotionTargetScaling(FractionalRadialForceLimitScaling);
 }
 
 void UPhysicsDrivenWalkingMode::OnContactModification_Internal(const FPhysicsMoverSimulationContactModifierParams& Params, Chaos::FCollisionContactModifier& Modifier) const
@@ -137,8 +139,7 @@ void UPhysicsDrivenWalkingMode::FloorCheck(const FMoverDefaultSyncState& SyncSta
 	UpdatedPrimitive->CalcBoundingCylinder(PawnRadius, PawnHalfHeight);
 
 	const float FloorSweepDistance = TargetHeight + CommonLegacySettings->MaxStepHeight;
-	const float ShrinkRadius = 1.0f;
-	const float QueryRadius = FMath::Max(PawnRadius - ShrinkRadius, 0.0f);
+
 	UPhysicsMovementUtils::FloorSweep_Internal(SyncState.GetLocation_WorldSpace(), DeltaPos, UpdatedPrimitive, UpDir,
 		QueryRadius, FloorSweepDistance, CommonLegacySettings->MaxWalkSlopeCosine, TargetHeight, OutFloorResult, OutWaterResult);
 
@@ -152,7 +153,6 @@ void UPhysicsDrivenWalkingMode::FloorCheck(const FMoverDefaultSyncState& SyncSta
 	if (bWalkableFloor)
 	{
 		// Walkable floor found
-		OutFloorResult.bWalkableFloor = true;
 		return;
 	}
 
@@ -318,7 +318,6 @@ void UPhysicsDrivenWalkingMode::OnSimulationTick(const FSimulationTickParams& Pa
 		FVector TargetPosition = StartingSyncState->GetLocation_WorldSpace();
 		if (bHandleVerticalLandingSeparately && FVector::Parallel(StartingSyncState->GetVelocity_WorldSpace().GetSafeNormal(), UpDir))
 		{
-			TargetVelocity = FMath::Lerp(StartingSyncState->GetVelocity_WorldSpace(), ProposedMove.LinearVelocity.ProjectOnToNormal(UpDir), FractionalVelocityToTarget);
 			TargetPosition = StartingSyncState->GetLocation_WorldSpace() - UpDir * (FloorResult.FloorDist - TargetHeight);
 		}
 		else
@@ -334,7 +333,6 @@ void UPhysicsDrivenWalkingMode::OnSimulationTick(const FSimulationTickParams& Pa
 				ProposedNormalVelocity = StartingSyncState->GetVelocity_WorldSpace() - StartingMovePlaneVelocity;
 			}
 
-			TargetVelocity = FMath::Lerp(StartingMovePlaneVelocity, ProposedMovePlaneVelocity, FractionalVelocityToTarget) + ProposedNormalVelocity;
 			TargetPosition += ProposedMovePlaneVelocity * DeltaSeconds;
 		}
 
diff --git a/Engine/Plugins/Experimental/Mover/Source/Mover/Public/PhysicsMover/Modes/PhysicsDrivenFallingMode.h b/Engine/Plugins/Experimental/Mover/Source/Mover/Public/PhysicsMover/Modes/PhysicsDrivenFallingMode.h
--- a/Engine/Plugins/Experimental/Mover/Source/Mover/Public/PhysicsMover/Modes/PhysicsDrivenFallingMode.h
+++ b/Engine/Plugins/Experimental/Mover/Source/Mover/Public/PhysicsMover/Modes/PhysicsDrivenFallingMode.h
@@ -40,4 +40,8 @@ public:
 	// Target height for the character. This is the desired distance from the center of the capsule to the floor
 	UPROPERTY(EditAnywhere, Category = "Physics Mover", meta = (ClampMin = "0", UIMin = "0", ForceUnits = "cm"))
 	float TargetHeight = 95.0f;
+
+	// Radius used for ground queries
+	UPROPERTY(EditAnywhere, Category = "Physics Mover", meta = (ClampMin = "0", UIMin = "0", ForceUnits = "cm"))
+	float QueryRadius = 20.0f;
 };
diff --git a/Engine/Plugins/Experimental/Mover/Source/Mover/Public/PhysicsMover/Modes/PhysicsDrivenSwimmingMode.h b/Engine/Plugins/Experimental/Mover/Source/Mover/Public/PhysicsMover/Modes/PhysicsDrivenSwimmingMode.h
--- a/Engine/Plugins/Experimental/Mover/Source/Mover/Public/PhysicsMover/Modes/PhysicsDrivenSwimmingMode.h
+++ b/Engine/Plugins/Experimental/Mover/Source/Mover/Public/PhysicsMover/Modes/PhysicsDrivenSwimmingMode.h
@@ -32,5 +32,9 @@ public:
 
 	// Target height for the character. This is the desired distance from the center of the capsule to the floor
 	UPROPERTY(EditAnywhere, Category = "Physics Mover", meta = (ClampMin = "0", UIMin = "0", ForceUnits = "cm"))
-	float TargetHeight = 54.0f;
+	float TargetHeight = 95.0f;
+
+	// Radius used for ground queries
+	UPROPERTY(EditAnywhere, Category = "Physics Mover", meta = (ClampMin = "0", UIMin = "0", ForceUnits = "cm"))
+	float QueryRadius = 30.0f;
 };
diff --git a/Engine/Plugins/Experimental/Mover/Source/Mover/Public/PhysicsMover/Modes/PhysicsDrivenWalkingMode.h b/Engine/Plugins/Experimental/Mover/Source/Mover/Public/PhysicsMover/Modes/PhysicsDrivenWalkingMode.h
--- a/Engine/Plugins/Experimental/Mover/Source/Mover/Public/PhysicsMover/Modes/PhysicsDrivenWalkingMode.h
+++ b/Engine/Plugins/Experimental/Mover/Source/Mover/Public/PhysicsMover/Modes/PhysicsDrivenWalkingMode.h
@@ -51,18 +51,28 @@ public:
 	UPROPERTY(EditAnywhere, Category = "Physics Mover", meta = (ClampMin = "0", UIMin = "0", ForceUnits = "cm"))
 	float TargetHeight = 95.0f;
 
+	// Radius used for ground queries
+	UPROPERTY(EditAnywhere, Category = "Physics Mover", meta = (ClampMin = "0", UIMin = "0", ForceUnits = "cm"))
+	float QueryRadius = 30.0f;
+
 	// Damping factor to control the softness of the interaction between the character and the ground
 	// Set to 0 for no damping and 1 for maximum damping
 	UPROPERTY(EditAnywhere, Category = "Physics Mover", meta = (ClampMin = "0", UIMin = "0", ClampMax = "1", UIMax = "1"))
 	float GroundDamping = 0.0f;
 
-	// With this enabled, in addition to the constraint radial force, the mover applies an initial velocity to reach the target
-	// movement. This allows for a separation between pushing force and force to reach the target movement.
-	// A value of 1 means the character will always be able to reach the target motion (if there are no obstacles) regardless of
-	// the radial force limit.
-	// A value of 0 means the only thing moving the character is the constraint
+	// Scaling applied to the radial force limit to raise the limit to always allow the character to
+	// reach the motion target/
+	// A value of 1 means that the radial force limit will be increased as needed to match the force
+	// required to achieve the motion target.
+	// A value of 0 means no scaling will be applied.
+	UPROPERTY(EditAnywhere, Category = "Physics Mover", meta = (ClampMin = "0", UIMin = "0", ClampMax = "1", UIMax = "1"))
+	float FractionalRadialForceLimitScaling = 1.0f;
+
+	// Controls the reaction force applied to the ground in the ground plane when the character is moving
+	// A value of 1 means that the full reaction force is applied
+	// A value of 0 means the character only applies a normal force to the ground and no tangential force
 	UPROPERTY(EditAnywhere, Category = "Physics Mover", meta = (ClampMin = "0", UIMin = "0", ClampMax = "1", UIMax = "1"))
-	float FractionalVelocityToTarget = 1.0f;
+	float FractionalGroundReaction = 1.0f;
 
 	// Controls how much downward velocity is applied to keep the character rooted to the ground when the character
 	// is within MaxStepHeight of the ground surface.
diff --git a/Engine/Plugins/Experimental/Mover/Source/Mover/Public/PhysicsMover/Transitions/PhysicsJumpCheck.h b/Engine/Plugins/Experimental/Mover/Source/Mover/Public/PhysicsMover/Transitions/PhysicsJumpCheck.h
--- a/Engine/Plugins/Experimental/Mover/Source/Mover/Public/PhysicsMover/Transitions/PhysicsJumpCheck.h
+++ b/Engine/Plugins/Experimental/Mover/Source/Mover/Public/PhysicsMover/Transitions/PhysicsJumpCheck.h
@@ -28,7 +28,7 @@ public:
 	* A value of 1 means that the full equal and opposite jump impulse will be applied.
 	*/
 	UPROPERTY(BlueprintReadWrite, EditAnywhere, Category = Parameters, meta = (ClampMin = "0", UIMin = "0", ClampMax = "1", UIMax = "1"))
-	float FractionalGroundReactionImpulse = 0.0f;
+	float FractionalGroundReactionImpulse = 1.0f;
 
 	/** Name of movement mode to transition to when jumping is activated */
 	UPROPERTY(BlueprintReadWrite, EditAnywhere, Category = Parameters)
diff --git a/Engine/Source/Programs/HeadlessChaos/Private/Character/TestCharacterGroundConstraintSolver.cpp b/Engine/Source/Programs/HeadlessChaos/Private/Character/TestCharacterGroundConstraintSolver.cpp
--- a/Engine/Source/Programs/HeadlessChaos/Private/Character/TestCharacterGroundConstraintSolver.cpp
+++ b/Engine/Source/Programs/HeadlessChaos/Private/Character/TestCharacterGroundConstraintSolver.cpp
@@ -136,6 +136,31 @@ namespace ChaosTest
 		EXPECT_VECTOR_FLOAT_EQ(Solver.GetAngularImpulse(Dt), FVec3::ZeroVector);
 	}
 
+	TEST_F(CharacterGroundConstraintSolverTest, TestTwoBody_MotionTarget)
+	{
+		CharacterBody.SetX(FVec3(0.0, 0.0, 10.0));
+		CharacterBody.SetP(FVec3(0.0, 0.0, 10.0));
+		GroundBody.SetX(FVec3(0.0, 0.0, 0.0));
+		GroundBody.SetP(FVec3(0.0, 0.0, 0.0));
+		Data.GroundDistance = 10.0f;
+		Settings.TargetHeight = 10.0f;
+		Settings.RadialForceLimit = 1.0e10f;
+		Data.TargetDeltaPosition = FVec3(10.0, 0.0, 0.0);
+		Settings.SwingTorqueLimit = 0.0f;
+		Settings.TwistTorqueLimit = 0.0f;
+		Settings.AssumedOnGroundHeight = 1.0f;
+
+		CharacterBody.SetInvM(1.0f);
+		CharacterBody.SetInvILocal(FVec3(1.0f, 1.0f, 1.0f));
+		GroundBody.SetInvM(1.0f);
+		GroundBody.SetInvILocal(FVec3(1.0f, 1.0f, 1.0f));
+
+		UpdateTwoBody(3, 0);
+
+		EXPECT_VECTOR_FLOAT_EQ(CharacterBody.P(), FVec3(5.0, 0.0, 10.0));
+		EXPECT_VECTOR_FLOAT_EQ(GroundBody.P(), FVec3(-5.0, 0.0, 0.0));
+	}
+
 	// No overlap or projected overlap but overlap due to other constraints
 	// Solver should fix overlap using displacement not correction
 	TEST_F(CharacterGroundConstraintSolverTest, TestSingleBody_NormalImpulse_OtherConstraintOverlap)
@@ -255,6 +280,7 @@ namespace ChaosTest
 		Data.TargetDeltaFacing = 0.6;
 		Settings.SwingTorqueLimit = 0.0f;
 		Settings.TwistTorqueLimit = (1.0 / (CharacterBody.InvILocal().Z * Dt * Dt)) * 0.5;
+		Settings.RadialForceMotionTargetScaling = 0.0f;
 
 		UpdateSingleBody(5, 0);
 
@@ -363,6 +389,7 @@ namespace ChaosTest
 		Settings.SwingTorqueLimit = 0.0f;
 		Settings.TwistTorqueLimit = 0.0f;
 		Settings.AssumedOnGroundHeight = 1.0f;
+		Settings.MotionTargetMassBias = 0.0f;
 
 		UpdateTwoBody(1, 0);
 
@@ -397,6 +424,9 @@ namespace ChaosTest
 		Settings.TwistTorqueLimit = (1.0 / (CharacterBody.InvILocal().Z * Dt * Dt)) * 100.0;
 		Settings.AssumedOnGroundHeight = 1.0f;
 
+		GroundBody.SetInvM(0.0);
+		GroundBody.SetInvILocal(FVec3(0.0, 0.0, 0.0));
+
 		UpdateTwoBody(1, 0);
 
 		FVec3 ExpectedDeltaPos = GroundBody.Q() * FVec3(10.0, 0.0, 0.0) - FVec3(10.0, 0.0, 0.0);
diff --git a/Engine/Source/Runtime/Experimental/Chaos/Private/Chaos/Character/CharacterGroundConstraintSolver.cpp b/Engine/Source/Runtime/Experimental/Chaos/Private/Chaos/Character/CharacterGroundConstraintSolver.cpp
--- a/Engine/Source/Runtime/Experimental/Chaos/Private/Chaos/Character/CharacterGroundConstraintSolver.cpp
+++ b/Engine/Source/Runtime/Experimental/Chaos/Private/Chaos/Character/CharacterGroundConstraintSolver.cpp
@@ -6,6 +6,13 @@ namespace Chaos
 {
 	namespace Private
 	{
+		static FRotation3 GetRelativeRotation(const FRotation3& From, const FRotation3& To)
+		{
+			FReal DP = FRotation3::DotProduct(From, To);
+			FRotation3 RelRot = To * From.Inverse();
+			return DP < FReal(0.0f) ? -RelRot : RelRot;
+		}
+
 		//////////////////////////////////////////////////////////////////////////
 		/// FBodyData
 
@@ -57,13 +64,12 @@ namespace Chaos
 			: Normal(FSolverVec3(0.0f, 0.0f, 1.0f)) // Need to initialize normal as it is used in GetLinearImpulse
 			, VerticalAxis(FSolverVec3(0.0f, 0.0f, 1.0f)) // Need to vertical axis as it is used in GetAngularImpulse
 			, CharacterInvM(-1.0f)
-			, CharacterInvIZ(-1.0f)
 		{
 		}
 
 		bool FCharacterGroundConstraintSolver::FConstraintData::IsValid()
 		{
-			return (CharacterInvM > 0.0f) && (CharacterInvIZ > 0.0f);
+			return CharacterInvM > 0.0f;
 		}
 
 		//////////////////////////////////////////////////////////////////////////
@@ -93,7 +99,7 @@ namespace Chaos
 
 			const FVec3 PC_Init = Character_X - FVec3(Settings.TargetHeight * Settings.VerticalAxis);
 			const FVec3 PG_Init = Character_X - FVec3(Data.GroundDistance * Settings.VerticalAxis);
-			const FRotation3 RelQuatCharacter = Character_Q * Character_R.Inverse();
+			const FRotation3 RelQuatCharacter = GetRelativeRotation(Character_R, Character_Q);
 			const FReal CharacterAngularDisplacement = RelQuatCharacter.ToRotationVector().Dot(Settings.VerticalAxis);
 			const FVec3 CharacterLinearDisplacement = FVec3(Character_P - Character_X);
 
@@ -101,7 +107,7 @@ namespace Chaos
 			FVec3 ProjectedError = InitialError + CharacterLinearDisplacement;
 
 			const FReal CharacterInvM = BodyData.CharacterBody.InvM();
-			FReal TwoBodyEffectiveMass = 0.0;
+			FReal TwoBodyEffectiveMass = 1.0 / CharacterInvM;
 			FReal GroundInvM = 0.0;
 			FMatrix33 GroundInvI(0.0);
 			FVec3 GroundNormal = Data.GroundNormal;
@@ -111,22 +117,22 @@ namespace Chaos
 			FVec3 GroundLinearDisplacement(ForceInitToZero);
 			FVec3 Ground_X(ForceInitToZero);
 			FVec3 Ground_P(ForceInitToZero);
+			FRotation3 Ground_Q(FRotation3::Identity);
 
 			const bool bTwoBody = BodyData.GroundBody.IsValid();
 			if (bTwoBody)
 			{
 				Ground_X = BodyData.GroundBody.X();
 				Ground_P = BodyData.GroundBody.P();
-
+				Ground_Q = FRotation3(BodyData.GroundBody.Q());
 				const FRotation3 Ground_R(BodyData.GroundBody.R());
-				const FRotation3 Ground_Q(BodyData.GroundBody.Q());
-				RelQuatGround = Ground_Q * Ground_R.Inverse();
+				RelQuatGround = GetRelativeRotation(Ground_R, Ground_Q);
 
 				const FVec3 RC_Init = FVec3(PC_Init - Ground_X);
 				const FVec3 RG_Init = FVec3(PG_Init - Ground_X);
 
 				GroundLinearDisplacement = FVec3(Ground_P - Ground_X) + RelQuatGround * RG_Init - RG_Init;
-				GroundOffset = RelQuatGround * RC_Init;
+				GroundOffset = PC_Init + CharacterLinearDisplacement - Ground_P;
 				GroundNormal = RelQuatGround * GroundNormal;
 				ProjectedError -= GroundLinearDisplacement;
 
@@ -142,14 +148,15 @@ namespace Chaos
 
 			if (WillOverlap)
 			{
-				FReal DampingFactor = FMath::Clamp(Settings.DampingFactor, 0.0f, 1.0f);
+				FReal DampingFactor = FMath::Clamp(Settings.DampingFactor, 0.0, 1.0);
 				FVec3 Diff = DampingFactor * ProjectedOverlap * GroundNormal;
 				InitialError += Diff;
 				ProjectedError += Diff;
 			}
 
 			// Correction
-			const bool CurrentlyOverlapping = InitialError.Dot(GroundNormal) < 0.0f;
+			const FReal InitialErrorN = InitialError.Dot(GroundNormal);
+			const bool CurrentlyOverlapping = InitialErrorN < 0.0f;
 			const bool NeedsCorrection = CurrentlyOverlapping && WillOverlap;
 			CorrectionSolveFunction = NeedsCorrection ? &SolveCorrectionSingleBody : &NoSolve;
 
@@ -157,17 +164,17 @@ namespace Chaos
 			ConstraintData.CharacterInvM = FSolverReal(CharacterInvM);
 			ConstraintData.GroundInvM = FSolverReal(GroundInvM);
 			ConstraintData.GroundInvI = FSolverMatrix33(GroundInvI);
-			ConstraintData.SingleBodyEffectiveMass = FSolverReal(1.0 / CharacterInvM);
-			ConstraintData.TwoBodyEffectiveMass = FSolverReal(TwoBodyEffectiveMass);
+			ConstraintData.EffectiveMassN = FSolverReal(TwoBodyEffectiveMass);
 			ConstraintData.Normal = FSolverVec3(GroundNormal);
 			ConstraintData.GroundOffset = FSolverVec3(GroundOffset);
-			ConstraintData.InitialError = FSolverVec3(InitialError);
-			ConstraintData.InitialProjectedError = FSolverVec3(ProjectedError);
+			ConstraintData.InitialError = FSolverReal(InitialErrorN);
+			ConstraintData.InitialProjectedError = FSolverReal(ProjectedError.Dot(GroundNormal));
 
 			//////////////////////////////////////////////////////////////////////////
 			// Upright constraint
-
-			ConstraintData.CharacterInvI = FSolverMatrix33(Utilities::ComputeWorldSpaceInertia(Character_Q, BodyData.CharacterBody.InvILocal()));
+			const Chaos::FMatrix33 CharacterInvInertia = Utilities::ComputeWorldSpaceInertia(Character_Q, BodyData.CharacterBody.InvILocal());
+			ConstraintData.CharacterInvI = FSolverMatrix33(CharacterInvInertia);
+			ConstraintData.VerticalAxis = FSolverVec3(Settings.VerticalAxis);
 			ConstraintData.CharacterVerticalAxis = FSolverVec3(Character_Q.RotateVector(Settings.VerticalAxis));
 
 			//////////////////////////////////////////////////////////////////////////
@@ -185,7 +192,7 @@ namespace Chaos
 			// If the slope is too steep adjust the target to not point up the slope
 			// and set the radial force to zero
 			const FReal DtSq = Dt * Dt;
-
+			bool bSliding = false;
 			FReal DP = FVec3::DotProduct(Settings.VerticalAxis, GroundNormal);
 			if (DP <= Data.CosMaxWalkableSlopeAngle)
 			{
@@ -197,6 +204,7 @@ namespace Chaos
 					TargetDeltaPos -= UpMotion * FVec3(UpSlope);
 				}
 				ConstraintData.RadialImpulseLimit = FSolverReal(Settings.FrictionForceLimit * DtSq);
+				bSliding = true;
 			}
 			else
 			{
@@ -207,19 +215,51 @@ namespace Chaos
 			const FVec3 MotionTargetError = FVec3(Character_P - MotionTarget);
 			const FReal MotionTargetAngularError = CharacterAngularDisplacement - TargetDeltaFacing;
 
-			const FMatrix33 WorldSpaceCharacterInvI = Utilities::ComputeWorldSpaceInertia(BodyData.CharacterBody.Q(), BodyData.CharacterBody.InvILocal());
-			const FReal CharacterInvI = Settings.VerticalAxis.Dot(WorldSpaceCharacterInvI * Settings.VerticalAxis);
-			
+			FReal MotionTargetErrorSq = MotionTargetError.SizeSquared();
+			FReal MotionTargetDirSizeSq = MotionTargetErrorSq;
+			FVec3 MotionTargetDir = FVec3::ForwardVector;
+			if (MotionTargetDirSizeSq > UE_SMALL_NUMBER)
+			{
+				MotionTargetDir = MotionTargetError * FMath::InvSqrt(MotionTargetDirSizeSq);
+			}
+			else
+			{
+				MotionTargetDirSizeSq = TargetDeltaPos.SizeSquared();
+				if (MotionTargetDirSizeSq > UE_SMALL_NUMBER)
+				{
+					MotionTargetDir = TargetDeltaPos * FMath::InvSqrt(MotionTargetDirSizeSq);
+				}
+			}
+
+			ConstraintData.MassBiasT = FSolverReal(Settings.MotionTargetMassBias);
+
+			FReal EffectiveMassT = 1.0 / CharacterInvM;
+			FReal EffectiveInertiaT = 1.0 / Settings.VerticalAxis.Dot(CharacterInvInertia * Settings.VerticalAxis);
+			if (bTwoBody)
+			{
+				const FVec3 R_Cross_T = GroundOffset.Cross(MotionTargetDir);
+				EffectiveMassT = 1.0 / (CharacterInvM + ConstraintData.MassBiasT * (GroundInvM + R_Cross_T.Dot(GroundInvI * R_Cross_T)));
+			}
+
+			// Scale radial force limit to hit the motion target
+			if (!bSliding && (MotionTargetErrorSq > UE_SMALL_NUMBER))
+			{
+				const FSolverReal MotionTargetImpulse = FSolverReal(EffectiveMassT * Settings.RadialForceMotionTargetScaling * FMath::Sqrt(MotionTargetErrorSq));
+				ConstraintData.RadialImpulseLimit = ConstraintData.RadialImpulseLimit < MotionTargetImpulse ? MotionTargetImpulse : ConstraintData.RadialImpulseLimit;
+			}
 
 			// Write constraint data
-			ConstraintData.AngularImpulseLimit = FSolverReal(Settings.TwistTorqueLimit * DtSq);
+			ConstraintData.AngularTwistImpulseLimit = FSolverReal(Settings.TwistTorqueLimit * DtSq);
 			ConstraintData.AngularSwingImpulseLimit = FSolverReal(Settings.SwingTorqueLimit * DtSq);
-			ConstraintData.VerticalAxis = FSolverVec3(Settings.VerticalAxis);
 			ConstraintData.AssumedOnGroundHeight = FSolverReal(Settings.AssumedOnGroundHeight);
-			ConstraintData.CharacterInvIZ = FSolverReal(CharacterInvI);
-			ConstraintData.EffectiveInertia = FSolverReal(1.0 / ConstraintData.CharacterInvIZ);
 			ConstraintData.MotionTargetError = FSolverVec3(MotionTargetError);
 			ConstraintData.MotionTargetAngularError = FSolverReal(MotionTargetAngularError);
+			ConstraintData.EffectiveMassT = FSolverReal(EffectiveMassT);
+
+			FSolverVec3 FacingConstraintRotationAxis = FSolverVec3(Settings.VerticalAxis);
+			FacingConstraintRotationAxis.Normalize();
+
+			ConstraintData.EffectiveInertiaT = FSolverReal(EffectiveInertiaT);
 
 			PositionSolveFunction = bTwoBody ? &SolvePositionTwoBody : SolvePositionSingleBody;
 		}
diff --git a/Engine/Source/Runtime/Experimental/Chaos/Private/Chaos/Character/CharacterGroundConstraintSolver.h b/Engine/Source/Runtime/Experimental/Chaos/Private/Chaos/Character/CharacterGroundConstraintSolver.h
--- a/Engine/Source/Runtime/Experimental/Chaos/Private/Chaos/Character/CharacterGroundConstraintSolver.h
+++ b/Engine/Source/Runtime/Experimental/Chaos/Private/Chaos/Character/CharacterGroundConstraintSolver.h
@@ -75,65 +75,34 @@ namespace Chaos
 				FConstraintData();
 				bool IsValid();
 
-				/// World space ground body inverse inertia
-				FSolverMatrix33 CharacterInvI;
-
-				/// World space ground body inverse inertia
-				FSolverMatrix33 GroundInvI;
-
-				/// Offset vector from ground body CoM to the constraint position
-				FSolverVec3 GroundOffset;
-
-				/// Constraint error pre-integration
-				FSolverVec3 InitialError;
-
-				/// Projected constraint error post-integration
-				FSolverVec3 InitialProjectedError;
-
-				///	Ground plane normal direction
-				FSolverVec3 Normal;
-
-				/// World space vertical axis
-				FSolverVec3 VerticalAxis;
-
-				/// Vertical axis rotated by the character initial rotation
-				FSolverVec3 CharacterVerticalAxis;
-
-				/// Projected constraint error for motion target constraint
-				FSolverVec3 MotionTargetError;
-
-				/// Projected angular facing error for motion target constraint
-				FSolverReal MotionTargetAngularError;
-
-				/// Character inverse mass
-				FSolverReal CharacterInvM;
-
-				/// Character inverse inertia about the vertical axis
-				FSolverReal CharacterInvIZ;
-
-				/// Ground body inverse mass
-				FSolverReal GroundInvM;
-
-				/// Effective mass for the two body system
-				FSolverReal TwoBodyEffectiveMass;
-
-				/// Effective mass for a single body constraint (the character mass)
-				FSolverReal SingleBodyEffectiveMass;
-
-				/// Effective angular mass to rotate about the vertical axis
-				FSolverReal EffectiveInertia;
-
-				/// Angular impulse from the solver is clamped to this limit
-				FSolverReal AngularImpulseLimit;
-
-				/// Angular swing impulse from the solver is clamped to this limit
-				FSolverReal AngularSwingImpulseLimit;
-
-				/// Radial linear impulse from the solver is clamped to this limit
-				FSolverReal RadialImpulseLimit;
-
-				/// Height below which the character is assumed to be grounded
-				FSolverReal AssumedOnGroundHeight;
+				FSolverMatrix33 CharacterInvI; /// World space ground body inverse inertia
+				FSolverMatrix33 GroundInvI; /// World space ground body inverse inertia
+				FSolverVec3 GroundOffset; /// Offset vector from ground body CoM to the constraint position
+
+				FSolverVec3 Normal; ///	Ground plane normal direction
+				FSolverVec3 VerticalAxis; /// World space vertical axis
+				FSolverVec3 CharacterVerticalAxis; /// Vertical axis rotated by the character initial rotation
+
+				FSolverVec3 MotionTargetError; /// Projected constraint error for motion target constraint
+				FSolverReal MotionTargetAngularError; // Projected angular error for facing constraint
+				FSolverReal InitialError; /// Constraint error pre-integration
+				FSolverReal InitialProjectedError; /// Projected constraint error post-integration
+				
+				FSolverReal CharacterInvM; /// Character inverse mass
+				FSolverReal GroundInvM; /// Ground body inverse mass
+
+				FSolverReal EffectiveMassN; // Effective mass for relative motion in the normal direction
+				FSolverReal EffectiveMassT; // Effective mass for relative motion in the motion target direction
+				FSolverReal EffectiveInertiaT; /// Effective angular mass to rotate the character to the target facing direction 
+
+				FSolverReal MassBiasT; // Mass bias applied to the ground body for the motion target constraint
+				FSolverReal MassBiasF; // Mass bias applied to the ground body for the facing direction angular constraint
+
+				FSolverReal RadialImpulseLimit; /// Radial linear impulse from the solver is clamped to this limit
+				FSolverReal AngularTwistImpulseLimit; /// Angular impulse from the solver is clamped to this limit
+				FSolverReal AngularSwingImpulseLimit; /// Angular swing impulse from the solver is clamped to this limit
+
+				FSolverReal AssumedOnGroundHeight; /// Height below which the character is assumed to be grounded
 			} ConstraintData;
 
 			using FSolveFunctionType = void (*)(const FConstraintData& ConstraintData, FBodyData& BodyData, FImpulseData& ImpulseData);
@@ -149,12 +118,12 @@ namespace Chaos
 		//////////////////////////////////////////////////////////////////////////
 		/// FCharacterGroundConstraintSolver inline functions
 
-		FORCEINLINE FVec3 FCharacterGroundConstraintSolver::GetLinearImpulse(FReal Dt) const
+		FORCEINLINE_DEBUGGABLE FVec3 FCharacterGroundConstraintSolver::GetLinearImpulse(FReal Dt) const
 		{
 			return Dt > UE_SMALL_NUMBER ? FVec3(ImpulseData.LinearPositionImpulse + ImpulseData.LinearCorrectionImpulse * ConstraintData.Normal) / Dt : FVec3::ZeroVector;
 		}
 
-		FORCEINLINE FVec3 FCharacterGroundConstraintSolver::GetAngularImpulse(FReal Dt) const
+		FORCEINLINE_DEBUGGABLE FVec3 FCharacterGroundConstraintSolver::GetAngularImpulse(FReal Dt) const
 		{
 			return Dt > UE_SMALL_NUMBER ? FVec3(ImpulseData.AngularImpulse * ConstraintData.VerticalAxis) / Dt : FVec3::ZeroVector;
 		}
@@ -162,12 +131,12 @@ namespace Chaos
 		//////////////////////////////////////////////////////////////////////////
 		/// Utility functions
 
-		FORCEINLINE FSolverVec3 FCharacterGroundConstraintSolver::ProjectOntoPlane(const FSolverVec3& Vector, const FSolverVec3& PlaneNormal)
+		FORCEINLINE_DEBUGGABLE FSolverVec3 FCharacterGroundConstraintSolver::ProjectOntoPlane(const FSolverVec3& Vector, const FSolverVec3& PlaneNormal)
 		{
 			return Vector - FSolverVec3::DotProduct(Vector, PlaneNormal) * PlaneNormal;
 		}
 
-		FORCEINLINE FSolverVec3 FCharacterGroundConstraintSolver::ClampMagnitude(const FSolverVec3& Vector, const FSolverReal& Max)
+		FORCEINLINE_DEBUGGABLE FSolverVec3 FCharacterGroundConstraintSolver::ClampMagnitude(const FSolverVec3& Vector, const FSolverReal& Max)
 		{
 			const FSolverReal MagSq = Vector.SizeSquared();
 			const FSolverReal MaxSq = Max * Max;
@@ -189,7 +158,7 @@ namespace Chaos
 			}
 		}
 
-		FORCEINLINE FSolverReal FCharacterGroundConstraintSolver::ClampAbs(const FSolverReal& Value, const FSolverReal& Max)
+		FORCEINLINE_DEBUGGABLE FSolverReal FCharacterGroundConstraintSolver::ClampAbs(const FSolverReal& Value, const FSolverReal& Max)
 		{
 			return Value > Max ? Max : (Value < -Max ? -Max : Value);
 		}
@@ -197,7 +166,7 @@ namespace Chaos
 		//////////////////////////////////////////////////////////////////////////
 		/// FBodyData
 
-		FORCEINLINE bool FCharacterGroundConstraintSolver::FBodyData::IsTwoBody()
+		FORCEINLINE_DEBUGGABLE bool FCharacterGroundConstraintSolver::FBodyData::IsTwoBody()
 		{
 			return GroundBody.IsValid() && GroundBody.IsDynamic();
 		}
@@ -205,7 +174,7 @@ namespace Chaos
 		//////////////////////////////////////////////////////////////////////////
 		/// Solve Functions
 
-		FORCEINLINE void FCharacterGroundConstraintSolver::SolvePosition()
+		FORCEINLINE_DEBUGGABLE void FCharacterGroundConstraintSolver::SolvePosition()
 		{
 			check(ConstraintData.IsValid());
 
@@ -216,24 +185,28 @@ namespace Chaos
 			(*PositionSolveFunction)(ConstraintData, BodyData, ImpulseData);
 		}
 
-		FORCEINLINE void FCharacterGroundConstraintSolver::SolveCorrectionSingleBody(const FConstraintData& ConstraintData, FBodyData& BodyData, FImpulseData& ImpulseData)
+		FORCEINLINE_DEBUGGABLE void FCharacterGroundConstraintSolver::SolveCorrectionSingleBody(const FConstraintData& ConstraintData, FBodyData& BodyData, FImpulseData& ImpulseData)
 		{
-			const FSolverReal Error = ConstraintData.Normal.Dot(BodyData.CharacterBody.CP() + ConstraintData.InitialError);
-			if (Error < FSolverReal(0.0f))
+			constexpr FSolverReal Zero(0.0f);
+			const FSolverReal Error = ConstraintData.Normal.Dot(BodyData.CharacterBody.CP()) + ConstraintData.InitialError;
+			if (Error < Zero)
 			{
-				const FSolverReal Delta = -ConstraintData.SingleBodyEffectiveMass * Error;
+				const FSolverReal Delta = -Error / ConstraintData.CharacterInvM;
 				ImpulseData.LinearCorrectionImpulse += Delta;
 				BodyData.CharacterBody.ApplyPositionCorrectionDelta(FSolverVec3(ConstraintData.CharacterInvM * Delta * ConstraintData.Normal));
 			}
 		}
 
-		FORCEINLINE void FCharacterGroundConstraintSolver::SolvePositionSingleBody(const FConstraintData& ConstraintData, FBodyData& BodyData, FImpulseData& ImpulseData)
+		FORCEINLINE_DEBUGGABLE void FCharacterGroundConstraintSolver::SolvePositionSingleBody(const FConstraintData& ConstraintData, FBodyData& BodyData, FImpulseData& ImpulseData)
 		{
+			constexpr FSolverReal Zero(0.0f);
+			constexpr FSolverReal SizeSqTolerance(UE_SMALL_NUMBER);
+
 			// Normal
-			const FSolverReal Error = ConstraintData.Normal.Dot(BodyData.CharacterBody.DP() + ConstraintData.InitialProjectedError);
-			if (Error < FSolverReal(0.0f))
+			const FSolverReal Error = ConstraintData.Normal.Dot(BodyData.CharacterBody.DP()) + ConstraintData.InitialProjectedError;
+			if (Error < Zero)
 			{
-				const FSolverVec3 Delta = -ConstraintData.SingleBodyEffectiveMass * Error * ConstraintData.Normal;
+				const FSolverVec3 Delta = -(Error / ConstraintData.CharacterInvM) * ConstraintData.Normal;
 				ImpulseData.LinearPositionImpulse += Delta;
 				BodyData.CharacterBody.ApplyPositionDelta(FSolverVec3(ConstraintData.CharacterInvM * Delta));
 			}
@@ -243,7 +216,7 @@ namespace Chaos
 			NewCharacterVerticalAxis.Normalize();
 			const FSolverVec3 CrossProd = NewCharacterVerticalAxis.Cross(ConstraintData.VerticalAxis);
 			const FSolverReal SizeSq = CrossProd.SizeSquared();
-			if (SizeSq > UE_SMALL_NUMBER)
+			if (SizeSq > SizeSqTolerance)
 			{
 				const FSolverVec3 AngAxis = CrossProd * FMath::InvSqrt(SizeSq);
 				const FSolverReal AngResistance = FSolverReal(1.0f) / (ConstraintData.CharacterInvI * AngAxis).Dot(AngAxis);
@@ -258,29 +231,32 @@ namespace Chaos
 			if (((NormalImpulse + ImpulseData.LinearCorrectionImpulse) > MinNormalImpulse) || Error < ConstraintData.AssumedOnGroundHeight)
 			{
 				// Target Position
-				const FSolverVec3 MotionTargetError = ProjectOntoPlane(ConstraintData.MotionTargetError + BodyData.CharacterBody.DP(), ConstraintData.Normal);
+				const FSolverVec3 MotionTargetError = ProjectOntoPlane(BodyData.CharacterBody.DP() + ConstraintData.MotionTargetError, ConstraintData.Normal);
 				const FSolverVec3 InitialMotionTargetImpulse = ProjectOntoPlane(ImpulseData.LinearPositionImpulse, ConstraintData.Normal);
-				FSolverVec3 NewMotionTargetImpulse = ClampMagnitude(InitialMotionTargetImpulse - ConstraintData.SingleBodyEffectiveMass * MotionTargetError, ConstraintData.RadialImpulseLimit);
+				FSolverVec3 NewMotionTargetImpulse = ClampMagnitude(InitialMotionTargetImpulse - MotionTargetError / ConstraintData.CharacterInvM, ConstraintData.RadialImpulseLimit);
 				const FSolverVec3 Delta = NewMotionTargetImpulse - InitialMotionTargetImpulse;
 				ImpulseData.LinearPositionImpulse += Delta;
 				BodyData.CharacterBody.ApplyPositionDelta(FSolverVec3(ConstraintData.CharacterInvM * Delta));
 
 				// Target Rotation
 				const FSolverReal MotionTargetAngularError = ConstraintData.MotionTargetAngularError + ConstraintData.VerticalAxis.Dot(BodyData.CharacterBody.DQ());
-				const FSolverReal NewAngularImpulse = ClampAbs(ImpulseData.AngularImpulse - ConstraintData.EffectiveInertia * MotionTargetAngularError, ConstraintData.AngularImpulseLimit);
+				const FSolverReal NewAngularImpulse = ClampAbs(ImpulseData.AngularImpulse - ConstraintData.EffectiveInertiaT * MotionTargetAngularError, ConstraintData.AngularTwistImpulseLimit);
 				const FSolverReal AngularDelta = NewAngularImpulse - ImpulseData.AngularImpulse;
 				ImpulseData.AngularImpulse += AngularDelta;
 				BodyData.CharacterBody.ApplyRotationDelta(FSolverVec3(ConstraintData.CharacterInvI * AngularDelta * ConstraintData.VerticalAxis));
 			}
 		}
 
-		FORCEINLINE void FCharacterGroundConstraintSolver::SolvePositionTwoBody(const FConstraintData& ConstraintData, FBodyData& BodyData, FImpulseData& ImpulseData)
+		FORCEINLINE_DEBUGGABLE void FCharacterGroundConstraintSolver::SolvePositionTwoBody(const FConstraintData& ConstraintData, FBodyData& BodyData, FImpulseData& ImpulseData)
 		{
-			const FSolverReal Error = ConstraintData.Normal.Dot(ConstraintData.InitialProjectedError + BodyData.CharacterBody.DP()
-				- BodyData.GroundBody.DP() - BodyData.GroundBody.DQ().Cross(ConstraintData.GroundOffset));
-			if (Error < FSolverReal(0.0f))
+			constexpr FSolverReal Zero(0.0f);
+			constexpr FSolverReal SizeSqTolerance(UE_SMALL_NUMBER);
+
+			// Normal
+			const FSolverReal Error = ConstraintData.Normal.Dot(BodyData.CharacterBody.DP() - BodyData.GroundBody.DP() - BodyData.GroundBody.DQ().Cross(ConstraintData.GroundOffset)) + ConstraintData.InitialProjectedError;
+			if (Error < Zero)
 			{
-				const FSolverVec3 Delta = -ConstraintData.TwoBodyEffectiveMass * Error * ConstraintData.Normal;
+				const FSolverVec3 Delta = -ConstraintData.EffectiveMassN * Error * ConstraintData.Normal;
 				ImpulseData.LinearPositionImpulse += Delta;
 				BodyData.CharacterBody.ApplyPositionDelta(FSolverVec3(ConstraintData.CharacterInvM * Delta));
 				BodyData.GroundBody.ApplyPositionDelta(FSolverVec3(-ConstraintData.GroundInvM * Delta));
@@ -292,7 +268,7 @@ namespace Chaos
 			NewCharacterVerticalAxis.Normalize();
 			const FSolverVec3 CrossProd = NewCharacterVerticalAxis.Cross(ConstraintData.VerticalAxis);
 			const FSolverReal SizeSq = CrossProd.SizeSquared();
-			if (SizeSq > UE_SMALL_NUMBER)
+			if (SizeSq > SizeSqTolerance)
 			{
 				const FSolverVec3 AngAxis = CrossProd * FMath::InvSqrt(SizeSq);
 				const FSolverReal AngResistance = FSolverReal(1.0f) / (ConstraintData.CharacterInvI * AngAxis).Dot(AngAxis);
@@ -307,18 +283,22 @@ namespace Chaos
 			if (((NormalImpulse + ImpulseData.LinearCorrectionImpulse) > MinNormalImpulse) || Error < ConstraintData.AssumedOnGroundHeight)
 			{
 				// Target Position
-				FSolverVec3 MotionTargetError = ConstraintData.MotionTargetError + BodyData.CharacterBody.DP() - BodyData.GroundBody.DP();
-				MotionTargetError -= BodyData.GroundBody.DQ().Cross(MotionTargetError);
-				MotionTargetError = ProjectOntoPlane(MotionTargetError, ConstraintData.Normal);
+				const FSolverVec3 MotionTargetError = ProjectOntoPlane(BodyData.CharacterBody.DP() - BodyData.GroundBody.DP() - BodyData.GroundBody.DQ().Cross(ConstraintData.GroundOffset) + ConstraintData.MotionTargetError, ConstraintData.Normal);
 				const FSolverVec3 InitialMotionTargetImpulse = ProjectOntoPlane(ImpulseData.LinearPositionImpulse, ConstraintData.Normal);
-				FSolverVec3 NewMotionTargetImpulse = ClampMagnitude(InitialMotionTargetImpulse - ConstraintData.SingleBodyEffectiveMass * MotionTargetError, ConstraintData.RadialImpulseLimit);
+				FSolverVec3 NewMotionTargetImpulse = ClampMagnitude(InitialMotionTargetImpulse - ConstraintData.EffectiveMassT * MotionTargetError, ConstraintData.RadialImpulseLimit);
 				const FSolverVec3 Delta = NewMotionTargetImpulse - InitialMotionTargetImpulse;
 				ImpulseData.LinearPositionImpulse += Delta;
 				BodyData.CharacterBody.ApplyPositionDelta(FSolverVec3(ConstraintData.CharacterInvM * Delta));
+				BodyData.GroundBody.ApplyPositionDelta(FSolverVec3(-ConstraintData.MassBiasT * ConstraintData.GroundInvM * Delta));
+				BodyData.GroundBody.ApplyRotationDelta(FSolverVec3(-ConstraintData.MassBiasT * ConstraintData.GroundInvI * ConstraintData.GroundOffset.Cross(Delta)));
+
+				//FSolverVec3 MotionTargetError = ConstraintData.MotionTargetError + BodyData.CharacterBody.DP() - BodyData.GroundBody.DP();
+				//MotionTargetError -= BodyData.GroundBody.DQ().Cross(MotionTargetError);
+				//MotionTargetError = ProjectOntoPlane(MotionTargetError, ConstraintData.Normal);
 
 				// Target Rotation
 				const FSolverReal MotionTargetAngularError = ConstraintData.MotionTargetAngularError + ConstraintData.VerticalAxis.Dot(BodyData.CharacterBody.DQ() - BodyData.GroundBody.DQ());
-				const FSolverReal NewAngularImpulse = ClampAbs(ImpulseData.AngularImpulse - ConstraintData.EffectiveInertia * MotionTargetAngularError, ConstraintData.AngularImpulseLimit);
+				const FSolverReal NewAngularImpulse = ClampAbs(ImpulseData.AngularImpulse - ConstraintData.EffectiveInertiaT * MotionTargetAngularError, ConstraintData.AngularTwistImpulseLimit);
 				const FSolverReal AngularDelta = NewAngularImpulse - ImpulseData.AngularImpulse;
 				ImpulseData.AngularImpulse += AngularDelta;
 				BodyData.CharacterBody.ApplyRotationDelta(FSolverVec3(ConstraintData.CharacterInvI * AngularDelta * ConstraintData.VerticalAxis));
diff --git a/Engine/Source/Runtime/Experimental/Chaos/Public/Chaos/Character/CharacterGroundConstraint.h b/Engine/Source/Runtime/Experimental/Chaos/Public/Chaos/Character/CharacterGroundConstraint.h
--- a/Engine/Source/Runtime/Experimental/Chaos/Public/Chaos/Character/CharacterGroundConstraint.h
+++ b/Engine/Source/Runtime/Experimental/Chaos/Public/Chaos/Character/CharacterGroundConstraint.h
@@ -126,6 +126,30 @@ namespace Chaos
 		}
 		FReal GetDampingFactor() const { return ConstraintSettings.Read().DampingFactor; }
 
+		/// This is a fractional value used to scale the ground mass in the motion target constraint.
+		/// It controls how much tangential force is applied to the ground when the character is moving.
+		/// 1 means the full reaction force is applied
+		/// 0 means zero reaction force is applied
+		void SetMotionTargetMassBias(const FReal& Value)
+		{
+			ConstraintSettings.Modify(true, DirtyFlags, Proxy, [&Value](FCharacterGroundConstraintSettings& Data) {
+				Data.MotionTargetMassBias = FMath::Clamp(Value, 0.0, 1.0);
+				});
+		}
+		FReal GetMotionTargetMassBias() const { return ConstraintSettings.Read().MotionTargetMassBias; }
+
+		/// This is a fractional value used to scale the radial force so that the character can achieve
+		/// the motion target
+		/// 1 means full scaling is applied and the character will always be able to achieve the motion target
+		/// 0 means no scaling is applied to the radial force limit
+		void SetRadialForceMotionTargetScaling(const FReal& Value)
+		{
+			ConstraintSettings.Modify(true, DirtyFlags, Proxy, [&Value](FCharacterGroundConstraintSettings& Data) {
+				Data.RadialForceMotionTargetScaling = FMath::Clamp(Value, 0.0, 1.0);
+				});
+		}
+		FReal GetRadialForceMotionTargetScaling() const { return ConstraintSettings.Read().RadialForceMotionTargetScaling; }
+
 		/// This is the height below which the character is assumed to be on the ground and can apply
 		/// force to move. This is useful if the character is slightly off the ground but should still
 		/// be considered in a locomotion rather than falling state.
diff --git a/Engine/Source/Runtime/Experimental/Chaos/Public/Chaos/Character/CharacterGroundConstraintSettings.h b/Engine/Source/Runtime/Experimental/Chaos/Public/Chaos/Character/CharacterGroundConstraintSettings.h
--- a/Engine/Source/Runtime/Experimental/Chaos/Public/Chaos/Character/CharacterGroundConstraintSettings.h
+++ b/Engine/Source/Runtime/Experimental/Chaos/Public/Chaos/Character/CharacterGroundConstraintSettings.h
@@ -22,6 +22,8 @@ namespace Chaos
 		FReal CosMaxWalkableSlopeAngle = 0.633;		/// Cosine of the maximum angle in degrees that the character is allowed to walk on
 		FReal DampingFactor = 0.0;					/// Applies a damping to the vertical ground constraint making it softer. Units: /T
 		FReal AssumedOnGroundHeight = 2.0;			/// Below this height the character is assumed to be on the ground and can apply force/torque to reach the target position and facing
+		FReal MotionTargetMassBias = 1.0;			/// Scaling factor applied to reduce the mass of the ground body in the linear motion target constraint
+		FReal RadialForceMotionTargetScaling = 1.0; /// Scale the radial force to always reach this fraction of the motion target
 		void* UserData = nullptr;
 	};
 
