Capsule collisions are computed in relative space
12% improvement on Switch in the Narrow Phase

diff --git a/Engine/Source/Runtime/Core/Private/Math/UnrealMath.cpp b/Engine/Source/Runtime/Core/Private/Math/UnrealMath.cpp
--- a/Engine/Source/Runtime/Core/Private/Math/UnrealMath.cpp
+++ b/Engine/Source/Runtime/Core/Private/Math/UnrealMath.cpp
@@ -326,17 +326,18 @@ UE::Math::TQuat<T> UE::Math::TVector4<T>::ToOrientationQuat() const
 	return RotationQuat;
 }
 
-FVector FMath::ClosestPointOnLine(const FVector& LineStart, const FVector& LineEnd, const FVector& Point)
+template <typename T>
+UE::Math::TVector<T> FMath::ClosestPointOnLine(const UE::Math::TVector<T>& LineStart, const UE::Math::TVector<T>& LineEnd, const UE::Math::TVector<T>& Point)
 {
 	// Solve to find alpha along line that is closest point
 	// Weisstein, Eric W. "Point-Line Distance--3-Dimensional." From MathWorld--A Switchram Web Resource. http://mathworld.wolfram.com/Point-LineDistance3-Dimensional.html 
-	const FVector::FReal A = (LineStart - Point) | (LineEnd - LineStart);
-	const FVector::FReal B = (LineEnd - LineStart).SizeSquared();
+	const T A = (LineStart - Point) | (LineEnd - LineStart);
+	const T B = (LineEnd - LineStart).SizeSquared();
 	// This should be robust to B == 0 (resulting in NaN) because clamp should return 1.
-	const FVector::FReal T = FMath::Clamp<FVector::FReal>(-A/B, 0.f, 1.f);
+	const T Time = FMath::Clamp<T>(-A/B, 0.f, 1.f);
 
 	// Generate closest point
-	FVector ClosestPoint = LineStart + (T * (LineEnd - LineStart));
+	UE::Math::TVector<T> ClosestPoint = LineStart + (Time * (LineEnd - LineStart));
 
 	return ClosestPoint;
 }
@@ -1688,20 +1689,23 @@ CORE_API float FMath::PointDistToLine(const FVector &Point, const FVector &Direc
 	return (float)(OutClosestPoint-Point).Size();		// LWC_TODO: Precision loss
 }
 
-FVector FMath::ClosestPointOnSegment(const FVector &Point, const FVector &StartPoint, const FVector &EndPoint)
+template <typename T>
+UE::Math::TVector<T> FMath::ClosestPointOnSegment(const UE::Math::TVector<T>&Point, const UE::Math::TVector<T>&StartPoint, const UE::Math::TVector<T>&EndPoint)
 {
-	const FVector Segment = EndPoint - StartPoint;
-	const FVector VectToPoint = Point - StartPoint;
+	using VectorType = UE::Math::TVector<T>;
+
+	const VectorType Segment = EndPoint - StartPoint;
+	const VectorType VectToPoint = Point - StartPoint;
 
 	// See if closest point is before StartPoint
-	const FVector::FReal Dot1 = VectToPoint | Segment;
+	const T Dot1 = VectToPoint | Segment;
 	if( Dot1 <= 0 )
 	{
 		return StartPoint;
 	}
 
 	// See if closest point is beyond EndPoint
-	const FVector::FReal Dot2 = Segment | Segment;
+	const T Dot2 = Segment | Segment;
 	if( Dot2 <= Dot1 )
 	{
 		return EndPoint;
@@ -1746,9 +1750,12 @@ float FMath::PointDistToSegmentSquared(const FVector &Point, const FVector &Star
 	return float((Point - ClosestPoint).SizeSquared());	// LWC_TODO: Precision loss
 }
 
+template<typename T>
 struct SegmentDistToSegment_Solver
 {
-	SegmentDistToSegment_Solver(const FVector& InA1, const FVector& InB1, const FVector& InA2, const FVector& InB2):
+	using VectorType = UE::Math::TVector<T>;
+
+	SegmentDistToSegment_Solver(const VectorType& InA1, const VectorType& InB1, const VectorType& InA2, const VectorType& InB2):
 		bLinesAreNearlyParallel(false),
 		A1(InA1),
 		A2(InA2),
@@ -1760,28 +1767,27 @@ struct SegmentDistToSegment_Solver
 
 	bool bLinesAreNearlyParallel;
 
-	const FVector& A1;
-	const FVector& A2;
-
-	const FVector S1;
-	const FVector S2;
-	const FVector S3;
+	const VectorType& A1;
+	const VectorType& A2;
+	const VectorType S1;
+	const VectorType S2;
+	const VectorType S3;
 
-	void Solve(FVector& OutP1, FVector& OutP2)
+	void Solve(VectorType& OutP1, VectorType& OutP2)
 	{
-		const FVector::FReal Dot11 = S1 | S1;
-		const FVector::FReal Dot12 = S1 | S2;
-		const FVector::FReal Dot13 = S1 | S3;
-		const FVector::FReal Dot22 = S2 | S2;
-		const FVector::FReal Dot23 = S2 | S3;
+		const T Dot11 = S1 | S1;
+		const T Dot12 = S1 | S2;
+		const T Dot13 = S1 | S3;
+		const T Dot22 = S2 | S2;
+		const T Dot23 = S2 | S3;
 
-		const FVector::FReal D = Dot11*Dot22 - Dot12*Dot12;
+		const T D = Dot11*Dot22 - Dot12*Dot12;
 
-		FVector::FReal D1 = D;
-		FVector::FReal D2 = D;
+		T D1 = D;
+		T D2 = D;
 
-		FVector::FReal N1;
-		FVector::FReal N2;
+		T N1;
+		T N2;
 
 		if (bLinesAreNearlyParallel || D < UE_KINDA_SMALL_NUMBER)
 		{
@@ -1855,8 +1861,8 @@ struct SegmentDistToSegment_Solver
 		}
 
 		// finally do the division to get the points' location
-		const FVector::FReal T1 = (FMath::Abs(N1) < UE_KINDA_SMALL_NUMBER ? 0.f : N1 / D1);
-		const FVector::FReal T2 = (FMath::Abs(N2) < UE_KINDA_SMALL_NUMBER ? 0.f : N2 / D2);
+		const T T1 = (FMath::Abs(N1) < UE_KINDA_SMALL_NUMBER ? 0.f : N1 / D1);
+		const T T2 = (FMath::Abs(N2) < UE_KINDA_SMALL_NUMBER ? 0.f : N2 / D2);
 
 		// return the closest points
 		OutP1 = A1 + T1 * S1;
@@ -1864,12 +1870,15 @@ struct SegmentDistToSegment_Solver
 	}
 };
 
-void FMath::SegmentDistToSegmentSafe(FVector A1, FVector B1, FVector A2, FVector B2, FVector& OutP1, FVector& OutP2)
+template<typename T>
+void FMath::SegmentDistToSegmentSafe(UE::Math::TVector<T> A1, UE::Math::TVector<T> B1, UE::Math::TVector<T> A2, UE::Math::TVector<T> B2, UE::Math::TVector<T>& OutP1, UE::Math::TVector<T>& OutP2)
 {
-	SegmentDistToSegment_Solver Solver(A1, B1, A2, B2);
+	using VectorType = UE::Math::TVector<T>;
+
+	SegmentDistToSegment_Solver<T> Solver(A1, B1, A2, B2);
 
-	const FVector S1_norm = Solver.S1.GetSafeNormal();
-	const FVector S2_norm = Solver.S2.GetSafeNormal();
+	const VectorType S1_norm = Solver.S1.GetSafeNormal();
+	const VectorType S2_norm = Solver.S2.GetSafeNormal();
 
 	const bool bS1IsPoint = S1_norm.IsZero();
 	const bool bS2IsPoint = S2_norm.IsZero();
@@ -1891,10 +1900,10 @@ void FMath::SegmentDistToSegmentSafe(FVector A1, FVector B1, FVector A2, FVector
 	}
 	else
 	{
-		const	FVector::FReal	Dot11_norm = S1_norm | S1_norm;	// always >= 0
-		const	FVector::FReal	Dot22_norm = S2_norm | S2_norm;	// always >= 0
-		const	FVector::FReal	Dot12_norm = S1_norm | S2_norm;
-		const	FVector::FReal	D_norm	= Dot11_norm*Dot22_norm - Dot12_norm*Dot12_norm;	// always >= 0
+		const T Dot11_norm = S1_norm | S1_norm;	// always >= 0
+		const T Dot22_norm = S2_norm | S2_norm;	// always >= 0
+		const T Dot12_norm = S1_norm | S2_norm;
+		const T D_norm = Dot11_norm * Dot22_norm - Dot12_norm * Dot12_norm;	// always >= 0
 
 		Solver.bLinesAreNearlyParallel = D_norm < UE_KINDA_SMALL_NUMBER;
 		Solver.Solve(OutP1, OutP2);
@@ -3696,3 +3705,9 @@ template struct UE::Math::TVector4<float>;
 template struct UE::Math::TVector4<double>;
 template struct UE::Math::TMatrix<float>;
 template struct UE::Math::TMatrix<double>;
+template void FMath::SegmentDistToSegmentSafe<float>(UE::Math::TVector<float> A1, UE::Math::TVector<float> B1, UE::Math::TVector<float> A2, UE::Math::TVector<float> B2, UE::Math::TVector<float>& OutP1, UE::Math::TVector<float>& OutP2);
+template void FMath::SegmentDistToSegmentSafe<double>(UE::Math::TVector<double> A1, UE::Math::TVector<double> B1, UE::Math::TVector<double> A2, UE::Math::TVector<double> B2, UE::Math::TVector<double>& OutP1, UE::Math::TVector<double>& OutP2);
+template UE::Math::TVector<float> FMath::ClosestPointOnLine(const UE::Math::TVector<float>& LineStart, const UE::Math::TVector<float>& LineEnd, const UE::Math::TVector<float>& Point);
+template UE::Math::TVector<double> FMath::ClosestPointOnLine(const UE::Math::TVector<double>& LineStart, const UE::Math::TVector<double>& LineEnd, const UE::Math::TVector<double>& Point);
+template UE::Math::TVector<float> FMath::ClosestPointOnSegment(const UE::Math::TVector<float>& Point, const UE::Math::TVector<float>& StartPoint, const UE::Math::TVector<float>& EndPoint);
+template UE::Math::TVector<double> FMath::ClosestPointOnSegment(const UE::Math::TVector<double>& Point, const UE::Math::TVector<double>& StartPoint, const UE::Math::TVector<double>& EndPoint);
diff --git a/Engine/Source/Runtime/Core/Public/Math/UnrealMathUtility.h b/Engine/Source/Runtime/Core/Public/Math/UnrealMathUtility.h
--- a/Engine/Source/Runtime/Core/Public/Math/UnrealMathUtility.h
+++ b/Engine/Source/Runtime/Core/Public/Math/UnrealMathUtility.h
@@ -2014,7 +2014,8 @@ public:
 	[[nodiscard]] static CORE_API bool SphereConeIntersection(const FVector& SphereCenter, float SphereRadius, const FVector& ConeAxis, float ConeAngleSin, float ConeAngleCos);
 
 	/** Find the point on the line segment from LineStart to LineEnd which is closest to Point */
-	[[nodiscard]] static CORE_API FVector ClosestPointOnLine(const FVector& LineStart, const FVector& LineEnd, const FVector& Point);
+	template<typename T>
+	[[nodiscard]] static CORE_API UE::Math::TVector<T> ClosestPointOnLine(const UE::Math::TVector<T>& LineStart, const UE::Math::TVector<T>& LineEnd, const UE::Math::TVector<T>& Point);
 
 	/** Find the point on the infinite line between two points (LineStart, LineEnd) which is closest to Point */
 	[[nodiscard]] static CORE_API FVector ClosestPointOnInfiniteLine(const FVector& LineStart, const FVector& LineEnd, const FVector& Point);
@@ -2055,7 +2056,8 @@ public:
 	 *
 	 * @return	point on the segment defined by (StartPoint, EndPoint) that is closest to Point.
 	 */
-	[[nodiscard]] static CORE_API FVector ClosestPointOnSegment(const FVector &Point, const FVector &StartPoint, const FVector &EndPoint);
+	template <typename T>
+	[[nodiscard]] static CORE_API UE::Math::TVector<T> ClosestPointOnSegment(const UE::Math::TVector<T>&Point, const UE::Math::TVector<T>&StartPoint, const UE::Math::TVector<T>&EndPoint);
 
 	/**
 	* FVector2D version of ClosestPointOnSegment.
@@ -2116,7 +2118,8 @@ public:
 	 * @param	OutP1		Closest point on segment 1 to segment 2.
 	 * @param	OutP2		Closest point on segment 2 to segment 1.
 	 */
-	static CORE_API void SegmentDistToSegmentSafe(FVector A1, FVector B1, FVector A2, FVector B2, FVector& OutP1, FVector& OutP2);
+	template<typename T>
+	static CORE_API void SegmentDistToSegmentSafe(UE::Math::TVector<T> A1, UE::Math::TVector<T> B1, UE::Math::TVector<T> A2, UE::Math::TVector<T> B2, UE::Math::TVector<T>& OutP1, UE::Math::TVector<T>& OutP2);
 
 	/**
 	 * returns the time (t) of the intersection of the passed segment and a plane (could be <0 or >1)
@@ -2606,6 +2609,14 @@ public:
 
 };
 
+extern template CORE_API void FMath::SegmentDistToSegmentSafe<float>(UE::Math::TVector<float> A1, UE::Math::TVector<float> B1, UE::Math::TVector<float> A2, UE::Math::TVector<float> B2, UE::Math::TVector<float>& OutP1, UE::Math::TVector<float>& OutP2);
+extern template CORE_API void FMath::SegmentDistToSegmentSafe<double>(UE::Math::TVector<double> A1, UE::Math::TVector<double> B1, UE::Math::TVector<double> A2, UE::Math::TVector<double> B2, UE::Math::TVector<double>& OutP1, UE::Math::TVector<double>& OutP2);
+extern template CORE_API UE::Math::TVector<float> FMath::ClosestPointOnLine(const UE::Math::TVector<float>& LineStart, const UE::Math::TVector<float>& LineEnd, const UE::Math::TVector<float>& Point);
+extern template CORE_API UE::Math::TVector<double> FMath::ClosestPointOnLine(const UE::Math::TVector<double>& LineStart, const UE::Math::TVector<double>& LineEnd, const UE::Math::TVector<double>& Point);
+extern template CORE_API UE::Math::TVector<float> FMath::ClosestPointOnSegment(const UE::Math::TVector<float>& Point, const UE::Math::TVector<float>& StartPoint, const UE::Math::TVector<float>& EndPoint);
+extern template CORE_API UE::Math::TVector<double> FMath::ClosestPointOnSegment(const UE::Math::TVector<double>& Point, const UE::Math::TVector<double>& StartPoint, const UE::Math::TVector<double>& EndPoint);
+
+
 // LWC Conversion helpers
 namespace UE
 {
diff --git a/Engine/Source/Runtime/Experimental/Chaos/Private/Chaos/Collision/ContactPointsMiscShapes.cpp b/Engine/Source/Runtime/Experimental/Chaos/Private/Chaos/Collision/ContactPointsMiscShapes.cpp
--- a/Engine/Source/Runtime/Experimental/Chaos/Private/Chaos/Collision/ContactPointsMiscShapes.cpp
+++ b/Engine/Source/Runtime/Experimental/Chaos/Private/Chaos/Collision/ContactPointsMiscShapes.cpp
@@ -475,28 +475,30 @@ namespace Chaos
 
 	FContactPoint CapsuleCapsuleContactPoint(const FCapsule& A, const FRigidTransform3& ATransform, const FCapsule& B, const FRigidTransform3& BTransform)
 	{
-		FContactPoint Result;
+		FContactPointf Result;
 
-		FVector A1 = ATransform.TransformPosition(FVector(A.GetX1f()));
-		FVector A2 = ATransform.TransformPosition(FVector(A.GetX2f()));
-		FVector B1 = BTransform.TransformPosition(FVector(B.GetX1f()));
-		FVector B2 = BTransform.TransformPosition(FVector(B.GetX2f()));
-		FVector P1, P2;
+		const FRigidTransform3f CapsuleAToCapsuleBf = FRigidTransform3f(ATransform.GetRelativeTransformNoScale(BTransform));
+
+		const FVector3f A1 = CapsuleAToCapsuleBf.TransformPosition(A.GetX1f());
+		const FVector3f A2 = CapsuleAToCapsuleBf.TransformPosition(A.GetX2f());
+		const FVector3f B1 = B.GetX1f();
+		const FVector3f B2 = B.GetX2f();
+		FVector3f P1, P2;
 		FMath::SegmentDistToSegmentSafe(A1, A2, B1, B2, P1, P2);
 
-		FVec3 Delta = P2 - P1;
-		FReal DeltaLen = Delta.Size();
+		const FVec3f Delta = P2 - P1;
+		const FRealSingle DeltaLen = Delta.Size();
 		if (DeltaLen > UE_KINDA_SMALL_NUMBER)
 		{
-			FReal NewPhi = DeltaLen - (A.GetRadiusf() + B.GetRadiusf());
-			FVec3 Dir = Delta / DeltaLen;
-			FVec3 Normal = -Dir;
-			FVec3 LocationA = P1 + Dir * A.GetRadiusf();
-			FVec3 LocationB = P2 - Dir * B.GetRadiusf();
-
-			Result.ShapeContactPoints[0] = ATransform.InverseTransformPosition(LocationA);
-			Result.ShapeContactPoints[1] = BTransform.InverseTransformPosition(LocationB);
-			Result.ShapeContactNormal = BTransform.InverseTransformVector(Normal);
+			const FRealSingle NewPhi = DeltaLen - (A.GetRadiusf() + B.GetRadiusf());
+			const FVec3f Dir = Delta / DeltaLen;
+			const FVec3f Normal = -Dir;
+			const FVec3f LocationA = P1 + Dir * A.GetRadiusf();
+			const FVec3f LocationB = P2 - Dir * B.GetRadiusf();
+
+			Result.ShapeContactPoints[0] = CapsuleAToCapsuleBf.InverseTransformPosition(LocationA);
+			Result.ShapeContactPoints[1] = LocationB;
+			Result.ShapeContactNormal = Normal;
 			Result.Phi = NewPhi;
 		}
 
diff --git a/Engine/Source/Runtime/Experimental/Chaos/Private/Chaos/CollisionOneShotManifoldsMiscShapes.cpp b/Engine/Source/Runtime/Experimental/Chaos/Private/Chaos/CollisionOneShotManifoldsMiscShapes.cpp
--- a/Engine/Source/Runtime/Experimental/Chaos/Private/Chaos/CollisionOneShotManifoldsMiscShapes.cpp
+++ b/Engine/Source/Runtime/Experimental/Chaos/Private/Chaos/CollisionOneShotManifoldsMiscShapes.cpp
@@ -336,18 +336,21 @@ namespace Chaos
 			// @todo(chaos): support manifold maintenance
 			Constraint.ResetActiveManifoldContacts();
 
-			FVec3 AAxis(CapsuleATransform.TransformVector(CapsuleA.GetSegment().GetAxis()));
-			const FVec3 BAxis(CapsuleBTransform.TransformVector(CapsuleB.GetSegment().GetAxis()));
 
-			const FReal AHalfLen = CapsuleA.GetHeightf() / 2.0f;
-			const FReal BHalfLen = CapsuleB.GetHeightf() / 2.0f;
+			const FRigidTransform3f CapsuleAToCapsuleBf = FRigidTransform3f(CapsuleATransform.GetRelativeTransformNoScale(CapsuleBTransform));
+
+			FVec3f AAxis(CapsuleAToCapsuleBf.TransformVector(CapsuleA.GetAxisf()));
+			const FVec3f BAxis = CapsuleB.GetAxisf();
+
+			const FRealSingle AHalfLen = CapsuleA.GetHeightf() / 2.0f;
+			const FRealSingle BHalfLen = CapsuleB.GetHeightf() / 2.0f;
 
 			// Used in a few places below where we need to use the smaller/larger capsule, but always a dynamic one
-			const FReal ADynamicRadius = FConstGenericParticleHandle(Constraint.GetParticle0())->IsDynamic() ? CapsuleA.GetRadiusf() : TNumericLimits<FReal>::Max();
-			const FReal BDynamicRadius = FConstGenericParticleHandle(Constraint.GetParticle1())->IsDynamic() ? CapsuleB.GetRadiusf() : TNumericLimits<FReal>::Max();
+			const FRealSingle ADynamicRadius = FConstGenericParticleHandle(Constraint.GetParticle0())->IsDynamic() ? CapsuleA.GetRadiusf() : TNumericLimits<FRealSingle>::Max();
+			const FRealSingle BDynamicRadius = FConstGenericParticleHandle(Constraint.GetParticle1())->IsDynamic() ? CapsuleB.GetRadiusf() : TNumericLimits<FRealSingle>::Max();
 
 			// Make both capsules point in the same general direction
-			FReal ADotB = FVec3::DotProduct(AAxis, BAxis);
+			FRealSingle ADotB = FVec3f::DotProduct(AAxis, BAxis);
 			if (ADotB < 0)
 			{
 				ADotB = -ADotB;
@@ -356,9 +359,9 @@ namespace Chaos
 
 			// Get the closest points on the two line segments. This is used to generate the closest contact point
 			// which is always added to the manifold (if within CullDistance). We may also add other points.
-			FVector AClosest, BClosest;
-			const FVector ACenter = CapsuleATransform.TransformPosition(FVec3(CapsuleA.GetCenterf()));
-			const FVector BCenter = CapsuleBTransform.TransformPosition(FVec3(CapsuleB.GetCenterf()));
+			FVector3f AClosest, BClosest;
+			const FVector3f ACenter = CapsuleAToCapsuleBf.TransformPosition(CapsuleA.GetCenterf());
+			const FVector3f BCenter = CapsuleB.GetCenterf();
 			FMath::SegmentDistToSegmentSafe(
 				ACenter + AHalfLen * AAxis,
 				ACenter - AHalfLen * AAxis,
@@ -367,18 +370,18 @@ namespace Chaos
 				AClosest,
 				BClosest);
 
-			FVec3 ClosestDelta = BClosest - AClosest;
-			FReal ClosestDeltaLen = ClosestDelta.Size();
+			FVec3f ClosestDelta = BClosest - AClosest;
+			FRealSingle ClosestDeltaLen = ClosestDelta.Size();
 
 			// Stop now if we are beyond the cull distance
-			const FReal ClosestPhi = ClosestDeltaLen - (CapsuleA.GetRadiusf() + CapsuleB.GetRadiusf());
+			const FRealSingle ClosestPhi = ClosestDeltaLen - (CapsuleA.GetRadiusf() + CapsuleB.GetRadiusf());
 			if (ClosestPhi > Constraint.GetCullDistance())
 			{
 				return;
 			}
 
 			// Calculate the normal from the two closest points. Handle exact axis overlaps.
-			FVec3 ClosestNormal;
+			FVec3f ClosestNormal;
 			if (ClosestDeltaLen > UE_KINDA_SMALL_NUMBER)
 			{
 				ClosestNormal = -ClosestDelta / ClosestDeltaLen;
@@ -386,44 +389,44 @@ namespace Chaos
 			else
 			{
 				// Center axes exactly intersect. We'll fake a result that pops the capsules out along the Z axis, with the smaller capsule going up
-				ClosestNormal = (ADynamicRadius <= BDynamicRadius) ? FVec3(0, 0, 1) : FVec3(0, 0, -1);
+				ClosestNormal = (ADynamicRadius <= BDynamicRadius) ? FVec3f(0, 0, 1) : FVec3f(0, 0, -1);
 			}
-			const FVec3 ClosestLocationA = AClosest - ClosestNormal * CapsuleA.GetRadiusf();
-			const FVec3 ClosestLocationB = BClosest + ClosestNormal * CapsuleB.GetRadiusf();
+			const FVec3f ClosestLocationA = AClosest - ClosestNormal * CapsuleA.GetRadiusf();
+			const FVec3f ClosestLocationB = BClosest + ClosestNormal * CapsuleB.GetRadiusf();
 
 			// We always add the closest point to the manifold
 			// We may also add 2 more points generated from the end cap positions of the smaller capsule
 			FContactPoint ClosestContactPoint;
-			ClosestContactPoint.ShapeContactPoints[0] = CapsuleATransform.InverseTransformPosition(ClosestLocationA);
-			ClosestContactPoint.ShapeContactPoints[1] = CapsuleBTransform.InverseTransformPosition(ClosestLocationB);
-			ClosestContactPoint.ShapeContactNormal = CapsuleBTransform.InverseTransformVector(ClosestNormal);
+			ClosestContactPoint.ShapeContactPoints[0] = CapsuleAToCapsuleBf.InverseTransformPosition(ClosestLocationA);
+			ClosestContactPoint.ShapeContactPoints[1] = ClosestLocationB;
+			ClosestContactPoint.ShapeContactNormal = ClosestNormal;
 			ClosestContactPoint.Phi = ClosestPhi;
 			ClosestContactPoint.FaceIndex = INDEX_NONE;
 			ClosestContactPoint.ContactType = EContactPointType::VertexPlane;
 			Constraint.AddOneshotManifoldContact(ClosestContactPoint);
 
 			// We don't generate manifold points within this fraction (of segment length) distance
-			constexpr FReal TDeltaThreshold = FReal(0.2);		// fraction
+			constexpr FRealSingle TDeltaThreshold = 0.2f;		// fraction
 
 			// If the nearest cylinder normal is parallel to the other axis within this tolerance, we stick with 1 manifold point
-			constexpr FReal SinAngleThreshold = FReal(0.35);	// about 20deg (this would be an endcap-versus-cylinderwall collision at >70 degs)
+			constexpr FRealSingle SinAngleThreshold = 0.35f;	// about 20deg (this would be an endcap-versus-cylinderwall collision at >70 degs)
 
 			// If the capsules are in an X configuration, this controls the distance of the manifold points from the closest point
-			const FReal RadialContactFraction = Chaos_Collision_Manifold_CapsuleRadialContactFraction;
+			const FRealSingle RadialContactFraction = Chaos_Collision_Manifold_CapsuleRadialContactFraction;
 
 			// Calculate the line segment times for the nearest point calculate above
 			// NOTE: TA and TB will be in [-1, 1]
-			const FReal TA = FVec3::DotProduct(AClosest - ACenter, AAxis) / AHalfLen;
-			const FReal TB = FVec3::DotProduct(BClosest - BCenter, BAxis) / BHalfLen;
+			const FRealSingle TA = FVec3f::DotProduct(AClosest - ACenter, AAxis) / AHalfLen;
+			const FRealSingle TB = FVec3f::DotProduct(BClosest - BCenter, BAxis) / BHalfLen;
 
 			// If we have an end-end contact with no segment overlap, stick with the single point manifold
 			// This is when we have two capsules laid end to end (as opposed to side-by-side)
 			// NOTE: This test only works because we made the axes point in the same direction above
-			if ((TA < FReal(-1) + TDeltaThreshold) && (TB > FReal(1) - TDeltaThreshold))
+			if ((TA < FRealSingle(-1) + TDeltaThreshold) && (TB > FRealSingle(1) - TDeltaThreshold))
 			{
 				return;
 			}
-			if ((TB < FReal(-1) + TDeltaThreshold) && (TA > FReal(1) - TDeltaThreshold))
+			if ((TB < FRealSingle(-1) + TDeltaThreshold) && (TA > FRealSingle(1) - TDeltaThreshold))
 			{
 				return;
 			}
@@ -432,9 +435,9 @@ namespace Chaos
 			// Also if the contact is deep, there's a high chance that pushing one end out will push the other deeper and we also need more contacts.
 			// Note: we only consider the radius of the dynamic object(s) when deciding what "deep" means because the extra contacts are only
 			// to prevent excessive rotation from the single contact we have so far, and only the dynamic objects will rotate.
-			const FReal AxisDotMinimum = Chaos_Collision_Manifold_CapsuleAxisAlignedThreshold;
-			const FReal DeepRadiusFraction = Chaos_Collision_Manifold_CapsuleDeepPenetrationFraction;
-			const FReal MinDynamicRadius = FMath::Min(ADynamicRadius, BDynamicRadius);
+			const FRealSingle AxisDotMinimum = Chaos_Collision_Manifold_CapsuleAxisAlignedThreshold;
+			const FRealSingle DeepRadiusFraction = Chaos_Collision_Manifold_CapsuleDeepPenetrationFraction;
+			const FRealSingle MinDynamicRadius = FMath::Min(ADynamicRadius, BDynamicRadius);
 			const bool bAreAligned = ADotB > AxisDotMinimum;
 			const bool bIsDeep = ClosestPhi < -DeepRadiusFraction * MinDynamicRadius;
 			if (!bAreAligned && !bIsDeep)
@@ -444,43 +447,42 @@ namespace Chaos
 
 			// Lambda: Create a contact point between a point on the cylinder of FirstCapsule at FirstT, with the nearest point on SecondCapsule
 			auto MakeCapsuleSegmentContact = [](
-				const FReal FirstT,
-				const FVec3& FirstCenter,
-				const FVec3& FirstAxis,
-				const FReal FirstHalfLen,
-				const FReal FirstRadius,
-				const FRigidTransform3& FirstTransform,
-				const FVec3& SecondCenter,
-				const FVec3& SecondAxis,
-				const FReal SecondHalfLen,
-				const FReal SecondRadius,
-				const FRigidTransform3& SecondTransform,
-				const FVec3& Orthogonal,
-				const FReal CullDistance,
+				const FRealSingle FirstT,
+				const FVec3f& FirstCenter,
+				const FVec3f& FirstAxis,
+				const FRealSingle FirstHalfLen,
+				const FRealSingle FirstRadius,
+				const FRigidTransform3f& FirstTransformToSecond,
+				const FVec3f& SecondCenter,
+				const FVec3f& SecondAxis,
+				const FRealSingle SecondHalfLen,
+				const FRealSingle SecondRadius,
+				const FVec3f& Orthogonal,
+				const FRealSingle CullDistance,
 				const bool bSwap) -> FContactPoint
 			{
 				FContactPoint ContactPoint;
 
-				const FVec3 FirstContactPos = FirstCenter + (FirstT * FirstHalfLen) * FirstAxis + Orthogonal * FirstRadius;
-				const FVec3 SecondSegmentPos = FMath::ClosestPointOnLine(SecondCenter - SecondHalfLen * SecondAxis, SecondCenter + SecondHalfLen * SecondAxis, FirstContactPos);
-				const FReal SecondSegmentDist = (FirstContactPos - SecondSegmentPos).Size();
-				const FVec3 SecondSegmentDir = (FirstContactPos - SecondSegmentPos) / SecondSegmentDist;
-				const FVec3 SecondContactPos = SecondSegmentPos + SecondRadius * SecondSegmentDir;
-				const FReal ContactPhi = SecondSegmentDist - SecondRadius;
+				const FVec3f FirstContactPos = FirstCenter + (FirstT * FirstHalfLen) * FirstAxis + Orthogonal * FirstRadius;
+				const FVec3f SecondSegmentPos = FMath::ClosestPointOnLine(SecondCenter - SecondHalfLen * SecondAxis, SecondCenter + SecondHalfLen * SecondAxis, FirstContactPos);
+				const FRealSingle SecondSegmentDist = (FirstContactPos - SecondSegmentPos).Size();
+				const FVec3f SecondSegmentDir = (FirstContactPos - SecondSegmentPos) / SecondSegmentDist;
+				const FVec3f SecondContactPos = SecondSegmentPos + SecondRadius * SecondSegmentDir;
+				const FRealSingle ContactPhi = SecondSegmentDist - SecondRadius;
 
 				if (ContactPhi < CullDistance)
 				{
 					if (!bSwap)
 					{
-						ContactPoint.ShapeContactPoints[0] = FirstTransform.InverseTransformPositionNoScale(FirstContactPos);
-						ContactPoint.ShapeContactPoints[1] = SecondTransform.InverseTransformPositionNoScale(SecondContactPos);
-						ContactPoint.ShapeContactNormal = SecondTransform.InverseTransformVectorNoScale(SecondSegmentDir);
+						ContactPoint.ShapeContactPoints[0] = FirstTransformToSecond.InverseTransformPositionNoScale(FirstContactPos);
+						ContactPoint.ShapeContactPoints[1] = SecondContactPos;
+						ContactPoint.ShapeContactNormal = SecondSegmentDir;
 					}
 					else
 					{
-						ContactPoint.ShapeContactPoints[0] = SecondTransform.InverseTransformPositionNoScale(SecondContactPos);
-						ContactPoint.ShapeContactPoints[1] = FirstTransform.InverseTransformPositionNoScale(FirstContactPos);
-						ContactPoint.ShapeContactNormal = -FirstTransform.InverseTransformVectorNoScale(SecondSegmentDir);
+						ContactPoint.ShapeContactPoints[0] = SecondContactPos;
+						ContactPoint.ShapeContactPoints[1] = FirstTransformToSecond.TransformPositionNoScale(FirstContactPos);
+						ContactPoint.ShapeContactNormal = -FirstTransformToSecond.TransformVectorNoScale(SecondSegmentDir);
 					}
 					ContactPoint.Phi = ContactPhi;
 					ContactPoint.FaceIndex = INDEX_NONE;
@@ -493,51 +495,50 @@ namespace Chaos
 			// Lambda: Add up to 2 more contacts from the cylindrical surface on FirstCylinder, if they are not too close to the existing contact.
 			// The point locations depend on cylinder alignment.
 			auto MakeCapsuleEndPointContacts = [&Constraint, TDeltaThreshold, SinAngleThreshold, RadialContactFraction, &MakeCapsuleSegmentContact](
-				const FReal FirstT,
-				const FVec3& FirstCenter,
-				const FVec3& FirstAxis,
-				const FReal FirstHalfLen,
-				const FReal FirstRadius,
-				const FRigidTransform3& FirstTransform,
-				const FVec3& SecondCenter,
-				const FVec3& SecondAxis,
-				const FReal SecondHalfLen,
-				const FReal SecondRadius,
-				const FRigidTransform3& SecondTransform,
-				const FVec3& ClosestDir,
-				const FReal FirstAxisDotSecondAxis,
+				const FRealSingle FirstT,
+				const FVec3f& FirstCenter,
+				const FVec3f& FirstAxis,
+				const FRealSingle FirstHalfLen,
+				const FRealSingle FirstRadius,
+				const FRigidTransform3f& FirstTransformToSecond,
+				const FVec3f& SecondCenter,
+				const FVec3f& SecondAxis,
+				const FRealSingle SecondHalfLen,
+				const FRealSingle SecondRadius,
+				const FVec3f& ClosestDir,
+				const FRealSingle FirstAxisDotSecondAxis,
 				const bool bSwap) -> void
 			{
 				// Orthogonal: the vector from a point on FirstCapsule's axis to its cylinder surface, in the direction of SecondCapsule
-				FVec3 Orthogonal = FVec3::CrossProduct(FirstAxis, FVec3::CrossProduct(FirstAxis, ClosestDir));
-				const FReal OrthogonalLenSq = Orthogonal.SizeSquared();
+				FVec3f Orthogonal = FVec3f::CrossProduct(FirstAxis, FVec3f::CrossProduct(FirstAxis, ClosestDir));
+				const FRealSingle OrthogonalLenSq = Orthogonal.SizeSquared();
 				if (OrthogonalLenSq > FMath::Square(SinAngleThreshold))
 				{
 					Orthogonal = Orthogonal * FMath::InvSqrt(OrthogonalLenSq);
-					if (FVec3::DotProduct(Orthogonal, SecondCenter - FirstCenter) < FReal(0))
+					if (FVec3f::DotProduct(Orthogonal, SecondCenter - FirstCenter) < FRealSingle(0))
 					{
 						Orthogonal = -Orthogonal;
 					}
 
 					// Clip the FirstCapsule's end points to be within the line segment of SecondCapsule
 					// This is to restrict the extra contacts to the overlapping line segment (e.g, when capsules are lying partly on top of each other)
-					const FReal ProjectedLen = FReal(2) * FirstHalfLen * FirstAxisDotSecondAxis;
-					const FReal ClippedTMin = FVec3::DotProduct((SecondCenter - SecondHalfLen * SecondAxis) - (FirstCenter + FirstHalfLen * FirstAxis), SecondAxis) / ProjectedLen;
-					const FReal ClippedTMax = FVec3::DotProduct((SecondCenter + SecondHalfLen * SecondAxis) - (FirstCenter - FirstHalfLen * FirstAxis), SecondAxis) / ProjectedLen;
+					const FRealSingle ProjectedLen = FRealSingle(2) * FirstHalfLen * FirstAxisDotSecondAxis;
+					const FRealSingle ClippedTMin = FVec3f::DotProduct((SecondCenter - SecondHalfLen * SecondAxis) - (FirstCenter + FirstHalfLen * FirstAxis), SecondAxis) / ProjectedLen;
+					const FRealSingle ClippedTMax = FVec3f::DotProduct((SecondCenter + SecondHalfLen * SecondAxis) - (FirstCenter - FirstHalfLen * FirstAxis), SecondAxis) / ProjectedLen;
 
 					// Clip the FirstCapsules end points to be within some laterial distance od the SecondCapsule's axis
 					// This restricts the contacts to be at a useful location when line segments are perpendicular to each other 
 					// (e.g., when the capsules are on top of each other but in a cross)
 					// As we get more perpendicular, move limits closer to radius fraction
-					const FReal MaxDeltaTRadial = RadialContactFraction * (SecondRadius / FirstHalfLen);
-					const FReal RadialClippedTMax = FMath::Lerp(FReal(MaxDeltaTRadial), FReal(1), FirstAxisDotSecondAxis);
+					const FRealSingle MaxDeltaTRadial = RadialContactFraction * (SecondRadius / FirstHalfLen);
+					const FRealSingle RadialClippedTMax = FMath::Lerp(FRealSingle(MaxDeltaTRadial), FRealSingle(1), FirstAxisDotSecondAxis);
 
-					const FReal TMin = FMath::Max3(FReal(-1), ClippedTMin, -RadialClippedTMax);
-					const FReal TMax = FMath::Min3(FReal(1), ClippedTMax, RadialClippedTMax);
+					const FRealSingle TMin = FMath::Max3(FRealSingle(-1), ClippedTMin, -RadialClippedTMax);
+					const FRealSingle TMax = FMath::Min3(FRealSingle(1), ClippedTMax, RadialClippedTMax);
 
 					if (TMin < FirstT - TDeltaThreshold)
 					{
-						FContactPoint EndContact0 = MakeCapsuleSegmentContact(TMin, FirstCenter, FirstAxis, FirstHalfLen, FirstRadius, FirstTransform, SecondCenter, SecondAxis, SecondHalfLen, SecondRadius, SecondTransform, Orthogonal, Constraint.GetCullDistance(), bSwap);
+						FContactPoint EndContact0 = MakeCapsuleSegmentContact(TMin, FirstCenter, FirstAxis, FirstHalfLen, FirstRadius, FirstTransformToSecond, SecondCenter, SecondAxis, SecondHalfLen, SecondRadius, Orthogonal, Constraint.GetCullDistancef(), bSwap);
 						if (EndContact0.Phi < Constraint.GetCullDistance())
 						{
 							Constraint.AddOneshotManifoldContact(EndContact0);
@@ -545,7 +546,7 @@ namespace Chaos
 					}
 					if (TMax > FirstT + TDeltaThreshold)
 					{
-						FContactPoint EndContact1 = MakeCapsuleSegmentContact(TMax, FirstCenter, FirstAxis, FirstHalfLen, FirstRadius, FirstTransform, SecondCenter, SecondAxis, SecondHalfLen, SecondRadius, SecondTransform, Orthogonal, Constraint.GetCullDistance(), bSwap);
+						FContactPoint EndContact1 = MakeCapsuleSegmentContact(TMax, FirstCenter, FirstAxis, FirstHalfLen, FirstRadius, FirstTransformToSecond, SecondCenter, SecondAxis, SecondHalfLen, SecondRadius, Orthogonal, Constraint.GetCullDistancef(), bSwap);
 						if (EndContact1.Phi < Constraint.GetCullDistance())
 						{
 							Constraint.AddOneshotManifoldContact(EndContact1);
@@ -557,11 +558,11 @@ namespace Chaos
 			// Generate the extra manifold points
 			if (ADynamicRadius <= BDynamicRadius)
 			{
-				MakeCapsuleEndPointContacts(TA, ACenter, AAxis, AHalfLen, CapsuleA.GetRadiusf(), CapsuleATransform, BCenter, BAxis, BHalfLen, CapsuleB.GetRadiusf(), CapsuleBTransform, ClosestNormal, ADotB, false);
+				MakeCapsuleEndPointContacts(TA, ACenter, AAxis, AHalfLen, CapsuleA.GetRadiusf(), CapsuleAToCapsuleBf, BCenter, BAxis, BHalfLen, CapsuleB.GetRadiusf(), ClosestNormal, ADotB, false);
 			}
 			else
 			{
-				MakeCapsuleEndPointContacts(TB, BCenter, BAxis, BHalfLen, CapsuleB.GetRadiusf(), CapsuleBTransform, ACenter, AAxis, AHalfLen, CapsuleA.GetRadiusf(), CapsuleATransform, ClosestNormal, ADotB, true);
+				MakeCapsuleEndPointContacts(TB, BCenter, BAxis, BHalfLen, CapsuleB.GetRadiusf(), CapsuleAToCapsuleBf.Inverse(), ACenter, AAxis, AHalfLen, CapsuleA.GetRadiusf(), ClosestNormal, ADotB, true);
 			}
 		}
 
