Defer delete of GeometryCollection RenderData on the render thread when recreating in place.

diff --git a/Engine/Source/Runtime/Experimental/GeometryCollectionEngine/Private/GeometryCollection/GeometryCollectionObject.cpp b/Engine/Source/Runtime/Experimental/GeometryCollectionEngine/Private/GeometryCollection/GeometryCollectionObject.cpp
--- a/Engine/Source/Runtime/Experimental/GeometryCollectionEngine/Private/GeometryCollection/GeometryCollectionObject.cpp
+++ b/Engine/Source/Runtime/Experimental/GeometryCollectionEngine/Private/GeometryCollection/GeometryCollectionObject.cpp
@@ -1458,9 +1458,20 @@ void UGeometryCollection::RebuildRenderData()
 {
 	if (RenderDataGuid != StateGuid)
 	{
-		ReleaseResources();
+		// Release RenderData render resources (happens on render thread) and defer deletion to the render thread for after ReleaseResources.
+		if (FGeometryCollectionRenderData* RenderDataToDelete = RenderData.Release())
+		{
+			RenderDataToDelete->ReleaseResources();
+			ENQUEUE_RENDER_COMMAND(SetResourceRenderThread)([RenderDataToDelete](FRHICommandListImmediate& RHICmdList)
+			{
+				delete RenderDataToDelete;
+			});
+		}
+
+		// Create new RenderData and initialize render resources (happens on render thread).
 		RenderData = FGeometryCollectionRenderData::Create(*GetGeometryCollection(), EnableNanite, bEnableNaniteFallback, bUseFullPrecisionUVs, bConvertVertexColorsToSRGB);
 		InitResources();
+
 		PropagateMarkDirtyToComponents();
 		RenderDataGuid = StateGuid;
 	}
