Dataflow : add support for variable
- Dataflow asset has now a variable container ( Property bag )
- Added DataflowInstance struct  to make integration of dataflow in other asset easier ( automatic variable sync between asset and dataflow asset )
- Added DataflowInstance UI customization to alolow default overriding in assets using dataflow ( through DataflowInstance )
- Geometry collection implement DataflowInstance interface to support variable
( No support yet for nodes and proper UI workflow in the dfataflow editor , this will be part of a second wave of changes )

diff --git a/Engine/Plugins/Experimental/ChaosEditor/Source/FractureEditor/Private/FractureToolGenerators.cpp b/Engine/Plugins/Experimental/ChaosEditor/Source/FractureEditor/Private/FractureToolGenerators.cpp
--- a/Engine/Plugins/Experimental/ChaosEditor/Source/FractureEditor/Private/FractureToolGenerators.cpp
+++ b/Engine/Plugins/Experimental/ChaosEditor/Source/FractureEditor/Private/FractureToolGenerators.cpp
@@ -904,7 +904,7 @@ AGeometryCollectionActor* UFractureToolGenerateAsset::ConvertActorsToGeometryCol
 
 	if (Dataflow)
 	{
-		FracturedGeometryCollection->DataflowAsset = Dataflow;
+		FracturedGeometryCollection->DataflowInstance.SetDataflowAsset(Dataflow);
 	}
 
 	AddSingleRootNodeIfRequired(FracturedGeometryCollection);
diff --git a/Engine/Plugins/Experimental/Dataflow/Source/DataflowEditor/DataflowEditor.Build.cs b/Engine/Plugins/Experimental/Dataflow/Source/DataflowEditor/DataflowEditor.Build.cs
--- a/Engine/Plugins/Experimental/Dataflow/Source/DataflowEditor/DataflowEditor.Build.cs
+++ b/Engine/Plugins/Experimental/Dataflow/Source/DataflowEditor/DataflowEditor.Build.cs
@@ -68,7 +68,8 @@ namespace UnrealBuildTool.Rules
 					"EditorWidgets",
 					"KismetWidgets",      // SScrubControlPanel
 					"AnimGraph", 
-					"ChaosCaching" // UAnimSingleNodeInstance
+					"ChaosCaching", // UAnimSingleNodeInstance
+					"StructUtilsEditor"
 				}
 			);
 		}
diff --git a/Engine/Plugins/Experimental/Dataflow/Source/DataflowEditor/Private/Dataflow/DataflowEditorModule.cpp b/Engine/Plugins/Experimental/Dataflow/Source/DataflowEditor/Private/Dataflow/DataflowEditorModule.cpp
--- a/Engine/Plugins/Experimental/Dataflow/Source/DataflowEditor/Private/Dataflow/DataflowEditorModule.cpp
+++ b/Engine/Plugins/Experimental/Dataflow/Source/DataflowEditor/Private/Dataflow/DataflowEditorModule.cpp
@@ -16,6 +16,7 @@
 #include "Dataflow/DataflowToolRegistry.h"
 #include "DataflowEditorTools/DataflowEditorWeightMapPaintTool.h"
 #include "Dataflow/DataflowCollectionAddScalarVertexPropertyNode.h"
+#include "Dataflow/DataflowInstanceDetails.h"
 
 #include "PropertyEditorModule.h"
 #include "EditorModeRegistry.h"
@@ -78,6 +79,7 @@ void FDataflowEditorModule::StartupModule()
 	{
 		PropertyModule->RegisterCustomPropertyTypeLayout(ScalarVertexPropertyGroupName, FOnGetPropertyTypeCustomizationInstance::CreateStatic(&UE::Dataflow::FScalarVertexPropertyGroupCustomization::MakeInstance));
 		PropertyModule->RegisterCustomPropertyTypeLayout(DataflowFunctionPropertyName, FOnGetPropertyTypeCustomizationInstance::CreateStatic(&UE::Dataflow::FFunctionPropertyCustomization::MakeInstance));
+		PropertyModule->RegisterCustomPropertyTypeLayout("DataflowInstance", FOnGetPropertyTypeCustomizationInstance::CreateStatic(&FDataflowInstanceDetails::MakeInstance));
 	}
 
 	UE::Dataflow::RenderingCallbacks();
diff --git a/Engine/Plugins/Experimental/Dataflow/Source/DataflowEditor/Private/Dataflow/DataflowEditorToolkit.cpp b/Engine/Plugins/Experimental/Dataflow/Source/DataflowEditor/Private/Dataflow/DataflowEditorToolkit.cpp
--- a/Engine/Plugins/Experimental/Dataflow/Source/DataflowEditor/Private/Dataflow/DataflowEditorToolkit.cpp
+++ b/Engine/Plugins/Experimental/Dataflow/Source/DataflowEditor/Private/Dataflow/DataflowEditorToolkit.cpp
@@ -26,6 +26,7 @@
 #include "Dataflow/DataflowNodeFactory.h"
 #include "Dataflow/DataflowNodeDetailExtension.h"
 #include "Dataflow/DataflowObject.h"
+#include "Dataflow/DataflowInstance.h"
 #include "Dataflow/DataflowObjectInterface.h"
 #include "Dataflow/DataflowEditorPreviewSceneBase.h"
 #include "Dataflow/DataflowRenderingFactory.h"
@@ -343,57 +344,22 @@ void FDataflowEditorToolkit::NotifyPreChange(FEditPropertyChain* PropertyAboutTo
 
 bool FDataflowEditorToolkit::CanOpenDataflowEditor(UObject* ObjectToEdit)
 {
-	if (const UClass* Class = ObjectToEdit->GetClass())
-	{
-		return (Class->FindPropertyByName(FName("DataflowAsset")) != nullptr) ;
-	}
-	return false;
+	return UE::Dataflow::InstanceUtils::HasValidDataflowAsset(ObjectToEdit);
 }
 
 bool FDataflowEditorToolkit::HasDataflowAsset(UObject* ObjectToEdit)
 {
-	if (const UClass* Class = ObjectToEdit->GetClass())
-	{
-		if (FProperty* Property = Class->FindPropertyByName(FName("DataflowAsset")))
-		{
-			return *Property->ContainerPtrToValuePtr<UDataflow*>(ObjectToEdit) != nullptr;
-		}
-	}
-	return false;
+	return (GetDataflowAsset(ObjectToEdit) != nullptr);
 }
 
 UDataflow* FDataflowEditorToolkit::GetDataflowAsset(UObject* ObjectToEdit)
 {
-	UDataflow* DataflowObject = Cast<UDataflow>(ObjectToEdit);
-
-	if (!DataflowObject)
-	{
-		if (const UClass* Class = ObjectToEdit->GetClass())
-		{
-			if (FProperty* Property = Class->FindPropertyByName(FName("DataflowAsset")))
-			{
-				DataflowObject = *Property->ContainerPtrToValuePtr<UDataflow*>(ObjectToEdit);
-			}
-		}
-	}
-	return DataflowObject;
+	return UE::Dataflow::InstanceUtils::GetDataflowAssetFromObject(ObjectToEdit);
 }
 
 const UDataflow* FDataflowEditorToolkit::GetDataflowAsset(const UObject* ObjectToEdit)
 {
-	const UDataflow* DataflowObject = Cast<UDataflow>(ObjectToEdit);
-
-	if (!DataflowObject)
-	{
-		if (const UClass* Class = ObjectToEdit->GetClass())
-		{
-			if (const FProperty* Property = Class->FindPropertyByName(FName("DataflowAsset")))
-			{
-				DataflowObject = *Property->ContainerPtrToValuePtr<const UDataflow*>(ObjectToEdit);
-			}
-		}
-	}
-	return DataflowObject;
+	return UE::Dataflow::InstanceUtils::GetDataflowAssetFromObject(ObjectToEdit);
 }
 
 //~ Begin FBaseCharacterFXEditorToolkit overrides
diff --git a/Engine/Plugins/Experimental/Dataflow/Source/DataflowEditor/Private/Dataflow/DataflowEditorUtil.cpp b/Engine/Plugins/Experimental/Dataflow/Source/DataflowEditor/Private/Dataflow/DataflowEditorUtil.cpp
--- a/Engine/Plugins/Experimental/Dataflow/Source/DataflowEditor/Private/Dataflow/DataflowEditorUtil.cpp
+++ b/Engine/Plugins/Experimental/Dataflow/Source/DataflowEditor/Private/Dataflow/DataflowEditorUtil.cpp
@@ -7,6 +7,7 @@
 #include "Dataflow/DataflowContent.h"
 #include "Dataflow/DataflowEditorToolkit.h"
 #include "Dataflow/DataflowEdNode.h"
+#include "Dataflow/DataflowInstance.h"
 #include "Dataflow/DataflowObject.h"
 #include "Dataflow/DataflowRenderingFactory.h"
 #include "DynamicMesh/MeshNormals.h"
@@ -31,24 +32,12 @@ namespace UE::Dataflow::Private
 	
 	bool HasDataflowAsset(UObject* InObject)
 	{
-		if (UClass* Class = InObject->GetClass())
-		{
-			return Class->FindPropertyByName(FName("DataflowAsset")) &&
-				   Class->FindPropertyByName(FName("DataflowTerminal"));
-		}
-		return false;
+		return UE::Dataflow::InstanceUtils::HasValidDataflowAsset(InObject);
 	}
 	
 	UDataflow* GetDataflowAssetFrom(UObject* InObject)
 	{
-		if (UClass* Class = InObject->GetClass())
-		{
-			if (FProperty* Property = Class->FindPropertyByName(FName("DataflowAsset")))
-			{
-				return *Property->ContainerPtrToValuePtr<UDataflow*>(InObject);
-			}
-		}
-		return nullptr;
+		return UE::Dataflow::InstanceUtils::GetDataflowAssetFromObject(InObject);
 	}
 
 	USkeletalMesh* GetSkeletalMeshFrom(UObject* InObject)
@@ -89,14 +78,7 @@ namespace UE::Dataflow::Private
 
 	FString GetDataflowTerminalFrom(UObject* InObject)
 	{
-		if (UClass* Class = InObject->GetClass())
-		{
-			if (FProperty* Property = Class->FindPropertyByName(FName("DataflowTerminal")))
-			{
-				return *Property->ContainerPtrToValuePtr<FString>(InObject);
-			}
-		}
-		return FString();
+		return UE::Dataflow::InstanceUtils::GetTerminalNodeNameFromObject(InObject).ToString();
 	}
 };
 
diff --git a/Engine/Plugins/Experimental/Dataflow/Source/DataflowEditor/Private/Dataflow/DataflowInstanceDetails.cpp b/Engine/Plugins/Experimental/Dataflow/Source/DataflowEditor/Private/Dataflow/DataflowInstanceDetails.cpp
new file mode 100644
--- /dev/null
+++ b/Engine/Plugins/Experimental/Dataflow/Source/DataflowEditor/Private/Dataflow/DataflowInstanceDetails.cpp
@@ -0,0 +1,144 @@
+// Copyright Epic Games, Inc. All Rights Reserved.
+
+#include "Dataflow/DataflowInstanceDetails.h"
+
+#include "Dataflow/DataflowObject.h"
+#include "Dataflow/DataflowInstance.h"
+
+#include "IDetailChildrenBuilder.h"
+#include "IPropertyUtilities.h"
+#include "PropertyBagDetails.h"
+#include "PropertyCustomizationHelpers.h"
+#include "ScopedTransaction.h"
+
+#define LOCTEXT_NAMESPACE "DataflowEditor"
+
+namespace UE::Dataflow::Private
+{
+	// variables details ( using FInstancedPropertyBag )
+	class FVariablesDetails : public FPropertyBagInstanceDataDetails
+	{
+	public:
+		FVariablesDetails(const TSharedPtr<IPropertyHandle> InDataflowInstanceStructProperty, const TSharedPtr<IPropertyHandle> InVariableStructProperty, const TSharedPtr<IPropertyUtilities>& InPropUtils)
+			: FPropertyBagInstanceDataDetails(InVariableStructProperty, InPropUtils, /*bInFixedLayout*/ true)
+			, DataflowInstanceStructProperty(InDataflowInstanceStructProperty)
+		{
+		}
+
+	protected:
+		struct FVariablesOverrideProvider : public IPropertyBagOverrideProvider
+		{
+			FVariablesOverrideProvider(FDataflowInstance& InDataflowInstance)
+				: DataflowInstance(InDataflowInstance)
+			{
+			}
+
+			virtual bool IsPropertyOverridden(const FGuid PropertyID) const override
+			{
+				return DataflowInstance.IsVariableOverridden(PropertyID);
+			}
+
+			virtual void SetPropertyOverride(const FGuid PropertyID, const bool bIsOverridden) const override
+			{
+				DataflowInstance.SetVariableOverridden(PropertyID, bIsOverridden);
+			}
+
+		private:
+			FDataflowInstance& DataflowInstance;
+		};
+
+		virtual bool HasPropertyOverrides() const override
+		{
+			return true;
+		}
+
+		virtual void PreChangeOverrides() override
+		{
+			check(DataflowInstanceStructProperty);
+			DataflowInstanceStructProperty->NotifyPreChange();
+		}
+
+		virtual void PostChangeOverrides() override
+		{
+			check(DataflowInstanceStructProperty);
+			DataflowInstanceStructProperty->NotifyPostChange(EPropertyChangeType::ValueSet);
+			DataflowInstanceStructProperty->NotifyFinishedChangingProperties();
+		}
+
+		virtual void EnumeratePropertyBags(TSharedPtr<IPropertyHandle> PropertyBagHandle, const EnumeratePropertyBagFuncRef& Func) const override
+		{
+			check(DataflowInstanceStructProperty);
+			DataflowInstanceStructProperty->EnumerateRawData([Func](void* RawData, const int32 /*DataIndex*/, const int32 /*NumDatas*/)
+				{
+					if (FDataflowInstance* DataflowInstance = static_cast<FDataflowInstance*>(RawData))
+					{
+						if (const UDataflow* DataflowAsset = DataflowInstance->GetDataflowAsset())
+						{
+							const FInstancedPropertyBag& DefaultParameters = DataflowAsset->Variables;
+							FInstancedPropertyBag& Parameters = DataflowInstance->GetMutableVariables();
+							FVariablesOverrideProvider OverrideProvider(*DataflowInstance);
+							if (!Func(DefaultParameters, Parameters, OverrideProvider))
+							{
+								return false;
+							}
+						}
+					}
+					return true;
+				});
+		}
+
+	private:
+
+		TSharedPtr<IPropertyHandle> DataflowInstanceStructProperty;
+	};
+}
+
+///=============================================================================================================
+/// 
+/// FDataflowVariablesDetails
+/// 
+///=============================================================================================================
+/// 
+TSharedRef<IPropertyTypeCustomization> FDataflowInstanceDetails::MakeInstance()
+{
+	return MakeShareable(new FDataflowInstanceDetails);
+}
+
+void FDataflowInstanceDetails::CustomizeHeader(TSharedRef<IPropertyHandle> InStructPropertyHandle, FDetailWidgetRow& InHeaderRow, IPropertyTypeCustomizationUtils& InStructCustomizationUtils)
+{
+	// default header row
+	InHeaderRow
+		.NameContent()
+		[
+			InStructPropertyHandle->CreatePropertyNameWidget()
+		];
+}
+
+void FDataflowInstanceDetails::CustomizeChildren(const TSharedRef<IPropertyHandle> InStructPropertyHandle, IDetailChildrenBuilder& InChildrenBuilder, IPropertyTypeCustomizationUtils& InCustomizationUtils)
+{
+	const TSharedPtr<IPropertyHandle> VariablesPropertyHandle = InStructPropertyHandle->GetChildHandle(FDataflowInstance::GetVariablePropertyName());
+	check(VariablesPropertyHandle);
+
+	uint32 NumChildren;
+	if (InStructPropertyHandle->GetNumChildren(NumChildren) == FPropertyAccess::Success)
+	{
+		for (uint32 Index = 0; Index < NumChildren; Index++)
+		{
+			TSharedPtr<IPropertyHandle> ChildProperty = InStructPropertyHandle->GetChildHandle(Index);
+			if (ChildProperty && ChildProperty->GetProperty() == VariablesPropertyHandle->GetProperty())
+			{
+				// customize the variables property
+				const TSharedRef<UE::Dataflow::Private::FVariablesDetails> VariablesDetails = MakeShareable(new  UE::Dataflow::Private::FVariablesDetails(InStructPropertyHandle, VariablesPropertyHandle, InCustomizationUtils.GetPropertyUtilities()));
+				InChildrenBuilder.AddCustomBuilder(VariablesDetails);
+			}
+			else
+			{
+				// add the child property untouched
+				InChildrenBuilder.AddProperty(ChildProperty.ToSharedRef());
+			}
+		}
+	}
+
+}
+
+#undef LOCTEXT_NAMESPACE
diff --git a/Engine/Plugins/Experimental/Dataflow/Source/DataflowEditor/Public/Dataflow/DataflowInstanceDetails.h b/Engine/Plugins/Experimental/Dataflow/Source/DataflowEditor/Public/Dataflow/DataflowInstanceDetails.h
new file mode 100644
--- /dev/null
+++ b/Engine/Plugins/Experimental/Dataflow/Source/DataflowEditor/Public/Dataflow/DataflowInstanceDetails.h
@@ -0,0 +1,23 @@
+// Copyright Epic Games, Inc. All Rights Reserved.
+
+#pragma once
+
+#include "IPropertyTypeCustomization.h"
+
+class IPropertyHandle;
+class USmartObjectDefinition;
+class IPropertyUtilities;
+
+/**
+ * Type customization for FDataflowInstance.
+ */
+class FDataflowInstanceDetails : public IPropertyTypeCustomization
+{
+public:
+	/** Makes a new instance of this detail layout class for a specific detail view requesting it */
+	static TSharedRef<IPropertyTypeCustomization> MakeInstance();
+
+	/** IPropertyTypeCustomization interface */
+	virtual void CustomizeHeader(TSharedRef<IPropertyHandle> InStructPropertyHandle, FDetailWidgetRow& InHeaderRow, IPropertyTypeCustomizationUtils& InCustomizationUtils) override;
+	virtual void CustomizeChildren(TSharedRef<IPropertyHandle> InStructPropertyHandle, IDetailChildrenBuilder& InChildrenBuilder, IPropertyTypeCustomizationUtils& InCustomizationUtils) override;
+};
\ No newline at end of file
diff --git a/Engine/Source/Runtime/Experimental/Dataflow/Engine/Private/Dataflow/DataflowContent.cpp b/Engine/Source/Runtime/Experimental/Dataflow/Engine/Private/Dataflow/DataflowContent.cpp
--- a/Engine/Source/Runtime/Experimental/Dataflow/Engine/Private/Dataflow/DataflowContent.cpp
+++ b/Engine/Source/Runtime/Experimental/Dataflow/Engine/Private/Dataflow/DataflowContent.cpp
@@ -1,6 +1,7 @@
 // Copyright Epic Games, Inc. All Rights Reserved.
 
 #include "Dataflow/DataflowContent.h"
+#include "Dataflow/DataflowInstance.h"
 #include "Dataflow/DataflowObjectInterface.h"
 #include "Dataflow/DataflowObject.h"
 #include "Animation/AnimSingleNodeInstance.h"
@@ -159,18 +160,6 @@ namespace UE::DataflowContextHelpers
 		return true;
 	}
 
-	UDataflow* GetDataflowAssetFrom(UObject* InObject)
-	{
-		if (UClass* Class = InObject->GetClass())
-		{
-			if (FProperty* Property = Class->FindPropertyByName(FName("DataflowAsset")))
-			{
-				return *Property->ContainerPtrToValuePtr<UDataflow*>(InObject);
-			}
-		}
-		return nullptr;
-	}
-
 	/*
 	* CreateNewDataflowContent
 	*
@@ -182,7 +171,7 @@ namespace UE::DataflowContextHelpers
 		bool bNeedsNewAsset = true;
 		TObjectPtr<UObject> Asset = nullptr;
 
-		UDataflow* DataflowAsset = GetDataflowAssetFrom(ContentOwner);
+		UDataflow* DataflowAsset = UE::Dataflow::InstanceUtils::GetDataflowAssetFromObject(ContentOwner);
 		if (bDataflowEnableContextCaching)
 		{
 			UClass* DataflowClass = T::StaticClass();
diff --git a/Engine/Source/Runtime/Experimental/Dataflow/Engine/Private/Dataflow/DataflowInstance.cpp b/Engine/Source/Runtime/Experimental/Dataflow/Engine/Private/Dataflow/DataflowInstance.cpp
new file mode 100644
--- /dev/null
+++ b/Engine/Source/Runtime/Experimental/Dataflow/Engine/Private/Dataflow/DataflowInstance.cpp
@@ -0,0 +1,257 @@
+// Copyright Epic Games, Inc. All Rights Reserved.
+
+#include "Dataflow/DataflowInstance.h"
+#include "Dataflow/DataflowObject.h"
+
+namespace UE::Dataflow::InstanceUtils
+{
+	static const FName DataflowAssetPropertyName("DataflowAsset");
+	static const FName DataflowTerminalPropertyName("DataflowTerminal");
+
+	bool HasValidDataflowAsset(UObject* Obj)
+	{
+		// first check if the object implements IDataflowGeneratedAssetInterface
+		if (const IDataflowInstanceInterface* Interface = Cast<IDataflowInstanceInterface>(Obj))
+		{
+			return true;
+		}
+		// check if the class contains a dataflow asset and terminal node name properties
+		if (UClass* Class = Obj->GetClass())
+		{
+			return Class->FindPropertyByName(DataflowAssetPropertyName) 
+				&& Class->FindPropertyByName(DataflowTerminalPropertyName);
+		}
+		return false;
+
+	}
+
+	template<typename TDataflowType, typename TObjectType>
+	TDataflowType* GetDataflowAssetFromObjectTemplate(TObjectType* Obj)
+	{
+		// first check if the object is itself the dataflow asset
+		TDataflowType* DataflowObject = Cast<TDataflowType>(Obj);
+
+		if (!DataflowObject)
+		{
+			// check if the object implements IDataflowGeneratedAssetInterface
+			if (const IDataflowInstanceInterface* Interface = Cast<IDataflowInstanceInterface>(Obj))
+			{
+				DataflowObject = Interface->GetDataflowInstance().GetDataflowAsset();
+			}
+		}
+
+		if (!DataflowObject && Obj)
+		{
+			// last search if the object has a property named DataflowAsset
+			// @TODO(dataflow)  we should retire this code path eventually in favor of using the interface solution 
+			if (const UClass* Class = Obj->GetClass())
+			{
+				if (FProperty* Property = Class->FindPropertyByName(DataflowAssetPropertyName))
+				{
+					DataflowObject = *Property->ContainerPtrToValuePtr<TDataflowType*>(Obj);
+				}
+			}
+		}
+		return DataflowObject;
+	}
+
+	UDataflow* GetDataflowAssetFromObject(UObject* Obj)
+	{
+		return GetDataflowAssetFromObjectTemplate<UDataflow, UObject>(Obj);
+	}
+
+	const UDataflow* GetDataflowAssetFromObject(const UObject* Obj)
+	{
+		return GetDataflowAssetFromObjectTemplate<const UDataflow, const UObject>(Obj);
+	}
+
+	FName GetTerminalNodeNameFromObject(UObject* Obj)
+	{
+		if (Obj)
+		{
+			// check if the object implements IDataflowGeneratedAssetInterface
+			if (const IDataflowInstanceInterface* Interface = Cast<IDataflowInstanceInterface>(Obj))
+			{
+				return Interface->GetDataflowInstance().GetDataflowTerminal();
+			}
+
+			// last search if the object has a property named DataflowTerminal
+			// @TODO(dataflow)  we should retire this code path eventually in favor of using the interface solution 
+			if (UClass* Class = Obj->GetClass())
+			{
+				if (FProperty* Property = Class->FindPropertyByName(DataflowTerminalPropertyName))
+				{
+					return *Property->ContainerPtrToValuePtr<FName>(Obj);
+				}
+			}
+		}
+		return NAME_None;
+	}
+
+	TArray<FName> GetTerminalNodeNames(const UDataflow* DataflowAsset)
+	{
+		TArray<FName> TerminalNodeNames;
+		if (DataflowAsset)
+		{
+			if (const TSharedPtr<UE::Dataflow::FGraph> DataflowGraph = DataflowAsset->Dataflow)
+			{
+				for (const TSharedPtr<FDataflowNode>& Node : DataflowGraph->GetNodes())
+				{
+					if (Node->AsType<FDataflowTerminalNode>())
+					{
+						TerminalNodeNames.Emplace(Node->GetName());
+					}
+				}
+			}
+		}
+		return TerminalNodeNames;
+	}
+}
+
+//---------------------------------------------------------------------------
+// FDataflowInstance
+//---------------------------------------------------------------------------
+
+FDataflowInstance::FDataflowInstance(UObject* InOwner)
+	: Owner(InOwner)
+{
+#if WITH_EDITOR
+	// let listen to property change on the owner so we can listen to our own changes ( because we are a struct )
+	if (Owner)
+	{
+		FCoreUObjectDelegates::OnObjectPropertyChanged.AddWeakLambda(Owner.Get(), [this](UObject* Obj, struct FPropertyChangedEvent& Event)
+			{
+				if (IDataflowInstanceInterface* Interface = Cast<IDataflowInstanceInterface>(Owner))
+				{
+					if (Obj == Interface->GetDataflowInstance().GetDataflowAsset())
+					{
+						OnOwnerPostEditChangeProperty(Event);
+					}
+				}
+			});
+	}
+#endif
+}
+
+#if WITH_EDITOR
+void FDataflowInstance::OnOwnerPostEditChangeProperty(struct FPropertyChangedEvent& PropertyChangedEvent)
+{
+	const FName PropertyName = PropertyChangedEvent.GetPropertyName();
+	const FName MemberPropertyName = PropertyChangedEvent.GetMemberPropertyName();
+	
+	if (PropertyName == GET_MEMBER_NAME_STRING_CHECKED(UDataflow, Variables))
+	{
+		// Add / remove / move properties / change type
+		// resync variable from the dataflow asset
+		SyncVariables();
+	}
+	else if (Variables.FindPropertyDescByName(MemberPropertyName))
+	{
+		// Specific variable change , nothing to sync 
+		// nothing to do now ( potential may need to re-evaluate the dataflow )
+		SyncVariables();
+	}
+}
+#endif
+
+FName FDataflowInstance::GetVariablePropertyName()
+{
+	return GET_MEMBER_NAME_CHECKED(FDataflowInstance, Variables);
+}
+
+
+void FDataflowInstance::SetDataflowAsset(UDataflow* InDataflowAsset)
+{
+	if (DataflowAsset != InDataflowAsset)
+	{
+		DataflowAsset = InDataflowAsset;
+		SyncVariables();
+	}
+}
+
+UDataflow* FDataflowInstance::GetDataflowAsset() const
+{
+	return DataflowAsset;
+}
+
+void FDataflowInstance::SetDataflowTerminal(FName TerminalNodeName)
+{
+	// todo(ccaillaud)  : in the future we should check that the terminal node is part of the assigned dataflow 
+	DataflowTerminal = TerminalNodeName;
+}
+
+FName FDataflowInstance::GetDataflowTerminal() const
+{
+	return DataflowTerminal;
+}
+
+const FInstancedPropertyBag& FDataflowInstance::GetVariables() const
+{
+	return Variables;
+}
+
+const FInstancedPropertyBag& FDataflowInstance::GetVariables()
+{
+	return Variables;
+}
+
+FInstancedPropertyBag& FDataflowInstance::GetMutableVariables()
+{
+	return Variables;
+}
+
+bool FDataflowInstance::IsVariableOverridden(const FGuid PropertyID) const
+{
+	return OverriddenVariables.Contains(PropertyID);
+}
+
+void FDataflowInstance::SetVariableOverridden(const FGuid PropertyID, const bool bIsOverridden)
+{
+	if (bIsOverridden)
+	{
+		OverriddenVariables.AddUnique(PropertyID);
+	}
+	else
+	{
+		OverriddenVariables.Remove(PropertyID);
+	}
+	SyncVariables();
+}
+
+void FDataflowInstance::SyncVariables()
+{
+	if (DataflowAsset == nullptr)
+	{
+		Variables.Reset();
+	}
+	else
+	{
+		// In editor builds, sync with overrides.
+		Variables.MigrateToNewBagInstanceWithOverrides(DataflowAsset->Variables, OverriddenVariables);
+
+		// clean up differences
+		RemoveOverridenVariablesNotInDataflowAsset();
+	}
+}
+
+void FDataflowInstance::RemoveOverridenVariablesNotInDataflowAsset()
+{
+	// Remove overriden variables that do not exists in the dataflow asset
+	if (!OverriddenVariables.IsEmpty())
+	{
+		if (const UPropertyBag* Bag = Variables.GetPropertyBagStruct())
+		{
+			for (TArray<FGuid>::TIterator It = OverriddenVariables.CreateIterator(); It; ++It)
+			{
+				if (!Bag->FindPropertyDescByID(*It))
+				{
+					It.RemoveCurrentSwap();
+				}
+			}
+		}
+		else
+		{
+			OverriddenVariables.Reset();
+		}
+	}
+}
diff --git a/Engine/Source/Runtime/Experimental/Dataflow/Engine/Private/Dataflow/DataflowObject.cpp b/Engine/Source/Runtime/Experimental/Dataflow/Engine/Private/Dataflow/DataflowObject.cpp
--- a/Engine/Source/Runtime/Experimental/Dataflow/Engine/Private/Dataflow/DataflowObject.cpp
+++ b/Engine/Source/Runtime/Experimental/Dataflow/Engine/Private/Dataflow/DataflowObject.cpp
@@ -82,7 +82,7 @@ void UDataflow::AddReferencedObjects(UObject* InThis, FReferenceCollector& Colle
 
 void UDataflow::PostEditChangeProperty(struct FPropertyChangedEvent& PropertyChangedEvent)
 {
-	UObject::PostEditChangeProperty(PropertyChangedEvent);
+	Super::PostEditChangeProperty(PropertyChangedEvent);
 }
 
 #endif
@@ -236,6 +236,7 @@ bool UDataflow::CanEditChange(const FProperty* InProperty) const
 
 	return true;
 }
+
 #endif
 
 #undef LOCTEXT_NAMESPACE
diff --git a/Engine/Source/Runtime/Experimental/Dataflow/Engine/Public/Dataflow/DataflowInstance.h b/Engine/Source/Runtime/Experimental/Dataflow/Engine/Public/Dataflow/DataflowInstance.h
new file mode 100644
--- /dev/null
+++ b/Engine/Source/Runtime/Experimental/Dataflow/Engine/Public/Dataflow/DataflowInstance.h
@@ -0,0 +1,119 @@
+
+// Copyright Epic Games, Inc. All Rights Reserved.
+
+#pragma once
+
+#include "Dataflow/DataflowContent.h"
+#include "StructUtils/PropertyBag.h"
+#include "UObject/ObjectPtr.h"
+
+#include "DataflowInstance.generated.h"
+
+namespace UE::Dataflow::Private
+{
+	class FVariablesDetails;
+}
+
+class UDataflow;
+
+/**
+* This structure is to be embedded in any asset that need generation from a dataflow
+*/
+USTRUCT()
+struct FDataflowInstance
+{
+	GENERATED_USTRUCT_BODY();
+
+	DATAFLOWENGINE_API FDataflowInstance(UObject* InOwner = nullptr);
+
+public:
+	DATAFLOWENGINE_API void SetDataflowAsset(UDataflow* DataflowAsset);
+
+	DATAFLOWENGINE_API UDataflow* GetDataflowAsset() const;
+
+	DATAFLOWENGINE_API void SetDataflowTerminal(FName TerminalNodeName);
+
+	DATAFLOWENGINE_API FName GetDataflowTerminal() const;
+
+	DATAFLOWENGINE_API const FInstancedPropertyBag& GetVariables() const;
+
+	DATAFLOWENGINE_API const FInstancedPropertyBag& GetVariables();
+
+	/** @return true if the property of specified ID is overridden. */
+	DATAFLOWENGINE_API bool IsVariableOverridden(const FGuid PropertyID) const;
+
+	/** Sets the override status of specified property by ID. */
+	DATAFLOWENGINE_API void SetVariableOverridden(const FGuid PropertyID, const bool bIsOverridden);
+
+	/** Get the variable member property name */
+	DATAFLOWENGINE_API static FName GetVariablePropertyName();
+
+	/** sync variables overriden in the asset with the one in the dataflow asset */
+	DATAFLOWENGINE_API void SyncVariables();
+
+private:
+	/** Dataflow asset to use */
+	UPROPERTY(EditAnywhere, Category = "Dataflow")
+	TObjectPtr<UDataflow> DataflowAsset;
+
+	/** name of the terminal node to use when generating the asset */
+	UPROPERTY(EditAnywhere, Category = "Dataflow")
+	FName DataflowTerminal;
+
+	/** Variables to override */
+	UPROPERTY(EditAnywhere, Category = "Dataflow", meta=(FixedLayout))
+	FInstancedPropertyBag Variables;
+
+	/** Array of overridden variable GUIDs. Non-overridden properties will inherit the values from the DataflowAsset default parameters. */
+	UPROPERTY()
+	TArray<FGuid> OverriddenVariables;
+
+private:
+	friend class UE::Dataflow::Private::FVariablesDetails;
+	DATAFLOWENGINE_API FInstancedPropertyBag& GetMutableVariables();
+
+	/** remove overriden variables that do not exist in the dataflow asset variables */
+	void RemoveOverridenVariablesNotInDataflowAsset();
+
+#if WITH_EDITOR
+	void OnOwnerPostEditChangeProperty(struct FPropertyChangedEvent& PropertyChangedEvent);
+#endif
+
+	// Owning Object
+	TObjectPtr<UObject> Owner = nullptr;
+};
+
+UINTERFACE(MinimalAPI)
+class UDataflowInstanceInterface : public UInterface
+{
+	GENERATED_BODY()
+};
+
+/**
+* Interface to use on asset class that are generated using a data flow
+*/
+class IDataflowInstanceInterface
+{
+	GENERATED_BODY()
+
+public:
+	virtual const FDataflowInstance& GetDataflowInstance() const = 0;
+};
+
+namespace UE::Dataflow::InstanceUtils
+{
+	/** Weather or not a dataflow asset with a non empty terminal node name is assigned to an UObject */
+	DATAFLOWENGINE_API bool HasValidDataflowAsset(UObject* Obj);
+
+	/** Get the dataflow asset from a UObject if available */
+	DATAFLOWENGINE_API UDataflow* GetDataflowAssetFromObject(UObject* Obj);
+
+	/** Get the dataflow asset from a UObject if available (const version) */
+	DATAFLOWENGINE_API const UDataflow* GetDataflowAssetFromObject(const UObject* Obj);
+
+	/** Get the terminal node name from a UObject if available */
+	DATAFLOWENGINE_API FName GetTerminalNodeNameFromObject(UObject* Obj);
+
+	/** get the list of terminal node names for a specific data flow asset */
+	DATAFLOWENGINE_API TArray<FName> GetTerminalNodeNames(const UDataflow* DataflowAsset);
+}
diff --git a/Engine/Source/Runtime/Experimental/Dataflow/Engine/Public/Dataflow/DataflowObject.h b/Engine/Source/Runtime/Experimental/Dataflow/Engine/Public/Dataflow/DataflowObject.h
--- a/Engine/Source/Runtime/Experimental/Dataflow/Engine/Public/Dataflow/DataflowObject.h
+++ b/Engine/Source/Runtime/Experimental/Dataflow/Engine/Public/Dataflow/DataflowObject.h
@@ -15,6 +15,8 @@
 #include "UObject/ObjectPtr.h"
 #include "UObject/UObjectGlobals.h"
 
+#include "StructUtils/PropertyBag.h"
+
 #include "DataflowObject.generated.h"
 
 class FArchive;
@@ -26,7 +28,7 @@ namespace UE::Dataflow { class FGraph; }
 
 
 /**
-*	FFleshAssetEdit
+*	FDataflowAssetEdit
 *     Structured RestCollection access where the scope
 *     of the object controls serialization back into the
 *     dynamic collection
@@ -114,6 +116,10 @@ public:
 	UPROPERTY(EditAnywhere, Category = "Evaluation",meta=(EditConditionHides))
 	EDataflowType Type = EDataflowType::Construction;
 
+	// keep hidden for now until we have a proper UI for it 
+	UPROPERTY()
+	FInstancedPropertyBag Variables;
+
 public:
 	/** UObject Interface */
 	static DATAFLOWENGINE_API void AddReferencedObjects(UObject* InThis, FReferenceCollector& Collector);
diff --git a/Engine/Source/Runtime/Experimental/GeometryCollectionEngine/Private/GeometryCollection/GeometryCollectionObject.cpp b/Engine/Source/Runtime/Experimental/GeometryCollectionEngine/Private/GeometryCollection/GeometryCollectionObject.cpp
--- a/Engine/Source/Runtime/Experimental/GeometryCollectionEngine/Private/GeometryCollection/GeometryCollectionObject.cpp
+++ b/Engine/Source/Runtime/Experimental/GeometryCollectionEngine/Private/GeometryCollection/GeometryCollectionObject.cpp
@@ -93,6 +93,7 @@ namespace GeometryCollectionCookStats
 #endif
 
 static constexpr float DefaultMaxSizeValue = 99999.9;
+static const FName GeometryCollectionDataflowTerminalNodeName("GeometryCollectionTerminal");
 
 UGeometryCollection::UGeometryCollection(const FObjectInitializer& ObjectInitializer)
 	: Super(ObjectInitializer)
@@ -137,6 +138,7 @@ UGeometryCollection::UGeometryCollection(const FObjectInitializer& ObjectInitial
 	, bSlowMovingAsSleeping(true)
 	, SlowMovingVelocityThreshold(1)
 	, EnableRemovePiecesOnFracture_DEPRECATED(false)
+	, DataflowInstance(this)
 	, GeometryCollection(new FGeometryCollection())
 {
 	PersistentGuid = FGuid::NewGuid();
@@ -150,6 +152,9 @@ UGeometryCollection::UGeometryCollection(const FObjectInitializer& ObjectInitial
 
 	// make sure we have at least one size specific entry
 	SizeSpecificData.AddDefaulted();
+
+	// set the default name for the terminal node 
+	DataflowInstance.SetDataflowTerminal(GeometryCollectionDataflowTerminalNodeName);
 }
 
 FGeometryCollectionLevelSetData::FGeometryCollectionLevelSetData()
@@ -951,6 +956,8 @@ void UGeometryCollection::Serialize(FArchive& Ar)
 	TSharedPtr<FGeometryCollection, ESPMode::ThreadSafe> ArchiveGeometryCollection = GeometryCollection;
 
 	TObjectPtr<UDataflow> StrippedDataflowAsset = nullptr;
+	FDataflowInstance StrippedDataflowInstance;
+	bool bStrippedDataflowData = false;
 
 	bool bIsCookedOrCooking = Ar.IsCooking();
 	if ((bIsCookedOrCooking && Ar.IsSaving()) || (Ar.IsCountingMemory() && Ar.IsFilterEditorOnly()))
@@ -986,6 +993,9 @@ void UGeometryCollection::Serialize(FArchive& Ar)
 		// The dataflow asset is only needed for the editor, so we just remove it when cooking 
 		StrippedDataflowAsset = DataflowAsset;
 		DataflowAsset = nullptr;
+		StrippedDataflowInstance = DataflowInstance;
+		DataflowInstance = FDataflowInstance();
+		bStrippedDataflowData = true;
 #endif
 	}
 
@@ -1029,9 +1039,10 @@ void UGeometryCollection::Serialize(FArchive& Ar)
 	}
 
 	// Important : this needs to remain after the call to Super::Serialize
-	if (StrippedDataflowAsset)
+	if (bStrippedDataflowData)
 	{
 		DataflowAsset = StrippedDataflowAsset;
+		DataflowInstance = StrippedDataflowInstance;
 	}
 
 	if ((Ar.IsLoading() || Ar.IsSaving()) && !SizeSpecificData.Num())
@@ -1997,6 +2008,13 @@ void UGeometryCollection::PostLoad()
 	}
 #endif
 
+	// migrate deprecated data if necessary 
+	MigrateDeprecatedRootProxyData();
+	MigrateDeprecatedDataflowData();
+}
+
+void UGeometryCollection::MigrateDeprecatedRootProxyData()
+{
 #if WITH_EDITORONLY_DATA
 	if (!RootProxy_DEPRECATED.IsNull())
 	{
@@ -2022,6 +2040,37 @@ void UGeometryCollection::PostLoad()
 #endif
 }
 
+void UGeometryCollection::MigrateDeprecatedDataflowData()
+{
+#if WITH_EDITORONLY_DATA
+	if (DataflowAsset != nullptr)
+	{
+		DataflowInstance.SetDataflowAsset(DataflowAsset);
+		DataflowInstance.SetDataflowTerminal(FName(DataflowTerminal));
+
+		DataflowAsset = nullptr;
+		DataflowTerminal.Empty();
+
+		// todo(ccaillaud) do not migrate the overrides until we have a solid plan to make sure they won't sync and risk of loosing them 
+	
+		//FInstancedPropertyBag& Variables = DataflowInstance.GetMutableVariables();
+		//// convert all override to be string variables 
+		//for (const TPair<FString, FString>& Pair : Overrides)
+		//{
+		//	const FName VariableName = FName(Pair.Key);
+		//	const FString& VariableValue = Pair.Value;
+
+		//	Variables.AddProperty(VariableName, EPropertyBagPropertyType::String);
+		//	Variables.SetValueString(VariableName, VariableValue);
+		//}
+		// clear all the deprecated values
+
+		// Overrides.Empty();
+	}
+#endif
+}
+
+
 void UGeometryCollection::BeginDestroy()
 {
 	Super::BeginDestroy();
@@ -2132,8 +2181,8 @@ void UGeometryCollection::WriteDataflowContent(const TObjectPtr<UDataflowBaseCon
 {
 	if(const TObjectPtr<UDataflowBaseContent> BaseContent = Cast<UDataflowBaseContent>(DataflowContent))
 	{
-		BaseContent->SetDataflowAsset(DataflowAsset);
-		BaseContent->SetDataflowTerminal(DataflowTerminal);
+		BaseContent->SetDataflowAsset(DataflowInstance.GetDataflowAsset());
+		BaseContent->SetDataflowTerminal(DataflowInstance.GetDataflowTerminal().ToString());
 	}
 }
 
diff --git a/Engine/Source/Runtime/Experimental/GeometryCollectionEngine/Public/GeometryCollection/GeometryCollectionObject.h b/Engine/Source/Runtime/Experimental/GeometryCollectionEngine/Public/GeometryCollection/GeometryCollectionObject.h
--- a/Engine/Source/Runtime/Experimental/GeometryCollectionEngine/Public/GeometryCollection/GeometryCollectionObject.h
+++ b/Engine/Source/Runtime/Experimental/GeometryCollectionEngine/Public/GeometryCollection/GeometryCollectionObject.h
@@ -6,7 +6,7 @@
 #include "Chaos/ChaosSolverActor.h"
 #include "GeometryCollection/GeometryCollectionDamagePropagationData.h"
 #include "GeometryCollection/GeometryCollectionSimulationTypes.h"
-#include "Dataflow/DataflowContent.h"
+#include "Dataflow/DataflowInstance.h"
 #include "GeometryCollection/ManagedArray.h"
 #include "InstanceUniformShaderParameters.h"
 #include "Interfaces/Interface_AssetUserData.h"
@@ -358,7 +358,11 @@ struct FGeometryCollectionRenderResourceSizeInfo
 *
 */
 UCLASS(BlueprintType, customconstructor, MinimalAPI)
-class UGeometryCollection : public UObject, public IInterface_AssetUserData, public IDataflowContentOwner
+class UGeometryCollection : 
+	public UObject, 
+	public IInterface_AssetUserData, 
+	public IDataflowContentOwner,
+	public IDataflowInstanceInterface
 {
 	GENERATED_UCLASS_BODY()
 
@@ -827,15 +831,22 @@ public:
 	//
 	// Dataflow
 	//
-	UPROPERTY(EditAnywhere, BlueprintReadWrite, Category = "Dataflow")
+	UPROPERTY(BlueprintReadWrite, Category = "Dataflow")
 	TObjectPtr<UDataflow> DataflowAsset;
 
-	UPROPERTY(EditAnywhere, Category = "Dataflow")
+	UPROPERTY()
 	FString DataflowTerminal = "GeometryCollectionTerminal";
 
-	UPROPERTY(EditAnywhere, BlueprintReadWrite, Category = "Dataflow", DisplayName = "DataFlow Overrides")
+	UPROPERTY(EditAnywhere, BlueprintReadWrite, Category = "Dataflow", DisplayName = "DataFlow Overrides", meta = (DisplayAfter = "DataflowInstance"))
 	TMap<FString, FString> Overrides;
 
+	UPROPERTY(EditAnywhere, Category = "Dataflow")
+	FDataflowInstance DataflowInstance;
+
+	//~ Begin IDataflowGeneratedAssetInterface
+	virtual const FDataflowInstance& GetDataflowInstance() const { return DataflowInstance; };
+	//~ End IDataflowGeneratedAssetInterface
+	
 	//~ Begin IDataflowContentOwner interface
 	virtual TObjectPtr<UDataflowBaseContent> CreateDataflowContent() override;
 	virtual void WriteDataflowContent(const TObjectPtr<UDataflowBaseContent>& DataflowContent) const override;
@@ -928,4 +939,7 @@ private:
 	FRenderCommandFence RenderResourceReleaseFence;
 
 	float GetMassOrDensityInternal(bool& bOutIsDensity, bool bCached) const;
+
+	void MigrateDeprecatedRootProxyData();
+	void MigrateDeprecatedDataflowData();
 };
