Geometry collection : remove the tick function on the geometry collection actor as it does nothing  and impact performance
- SetRenderStateDirty has been empty for a long time now - It's now deprecated
- Remove the tick function from the actor and all calls to SetRenderStateDirty

diff --git a/Engine/Source/Runtime/Experimental/GeometryCollectionEngine/Private/GeometryCollection/GeometryCollectionActor.cpp b/Engine/Source/Runtime/Experimental/GeometryCollectionEngine/Private/GeometryCollection/GeometryCollectionActor.cpp
--- a/Engine/Source/Runtime/Experimental/GeometryCollectionEngine/Private/GeometryCollection/GeometryCollectionActor.cpp
+++ b/Engine/Source/Runtime/Experimental/GeometryCollectionEngine/Private/GeometryCollection/GeometryCollectionActor.cpp
@@ -41,26 +41,11 @@ AGeometryCollectionActor::AGeometryCollectionActor(const FObjectInitializer& Obj
 	GeometryCollectionDebugDrawComponent_DEPRECATED = nullptr;
 
 	PrimaryActorTick.bCanEverTick = true;
-	SetActorTickEnabled(true);
 
 	bReplicates = true;
 	NetDormancy = DORM_Initial;
 }
 
-
-void AGeometryCollectionActor::Tick(float DeltaTime) 
-{
-	UE_LOG(AGeometryCollectionActorLogging, Verbose, TEXT("AGeometryCollectionActor::Tick()"));
-
-	Super::Tick(DeltaTime);
-
-	if (GeometryCollectionComponent)
-	{
-		GeometryCollectionComponent->SetRenderStateDirty();
-	}
-}
-
-
 const Chaos::FPhysicsSolver* GetSolver(const AGeometryCollectionActor& GeomCollectionActor)
 {
 	return GeomCollectionActor.GetGeometryCollectionComponent()->ChaosSolverActor != nullptr ? GeomCollectionActor.GetGeometryCollectionComponent()->ChaosSolverActor->GetSolver() : GeomCollectionActor.GetWorld()->PhysicsScene_Chaos->GetSolver();
diff --git a/Engine/Source/Runtime/Experimental/GeometryCollectionEngine/Private/GeometryCollection/GeometryCollectionComponent.cpp b/Engine/Source/Runtime/Experimental/GeometryCollectionEngine/Private/GeometryCollection/GeometryCollectionComponent.cpp
--- a/Engine/Source/Runtime/Experimental/GeometryCollectionEngine/Private/GeometryCollection/GeometryCollectionComponent.cpp
+++ b/Engine/Source/Runtime/Experimental/GeometryCollectionEngine/Private/GeometryCollection/GeometryCollectionComponent.cpp
@@ -3644,7 +3644,6 @@ void UGeometryCollectionComponent::ResetDynamicCollection()
 		DynamicCollection->MakeDirty();
 		MarkRenderStateDirty();
 		MarkRenderDynamicDataDirty();
-		SetRenderStateDirty();
 	}
 	else
 	{
diff --git a/Engine/Source/Runtime/Experimental/GeometryCollectionEngine/Public/GeometryCollection/GeometryCollectionActor.h b/Engine/Source/Runtime/Experimental/GeometryCollectionEngine/Public/GeometryCollection/GeometryCollectionActor.h
--- a/Engine/Source/Runtime/Experimental/GeometryCollectionEngine/Public/GeometryCollection/GeometryCollectionActor.h
+++ b/Engine/Source/Runtime/Experimental/GeometryCollectionEngine/Public/GeometryCollection/GeometryCollectionActor.h
@@ -21,10 +21,6 @@ class AGeometryCollectionActor: public AActor
 	GENERATED_UCLASS_BODY()
 
 public:
-
-	/* Game state callback */
-	GEOMETRYCOLLECTIONENGINE_API virtual void Tick(float DeltaSeconds) override;
-
 #if WITH_EDITOR
 	GEOMETRYCOLLECTIONENGINE_API virtual bool GetReferencedContentObjects(TArray<UObject*>& Objects) const override;
 #endif
diff --git a/Engine/Source/Runtime/Experimental/GeometryCollectionEngine/Public/GeometryCollection/GeometryCollectionComponent.h b/Engine/Source/Runtime/Experimental/GeometryCollectionEngine/Public/GeometryCollection/GeometryCollectionComponent.h
--- a/Engine/Source/Runtime/Experimental/GeometryCollectionEngine/Public/GeometryCollection/GeometryCollectionComponent.h
+++ b/Engine/Source/Runtime/Experimental/GeometryCollectionEngine/Public/GeometryCollection/GeometryCollectionComponent.h
@@ -597,7 +597,10 @@ public:
 	GEOMETRYCOLLECTIONENGINE_API virtual bool ShouldCreateRenderState() const override;
 	GEOMETRYCOLLECTIONENGINE_API virtual void CreateRenderState_Concurrent(FRegisterComponentContext* Context) override;
 	GEOMETRYCOLLECTIONENGINE_API virtual void SendRenderDynamicData_Concurrent() override;
+
+	UE_DEPRECATED(5.6, "Render state dirty state is now based on the geometry collection change events")
 	FORCEINLINE void SetRenderStateDirty() { /* Deprecated. */ }
+
 	GEOMETRYCOLLECTIONENGINE_API virtual void SetCollisionObjectType(ECollisionChannel Channel) override;
 	GEOMETRYCOLLECTIONENGINE_API virtual void OnActorEnableCollisionChanged() override;
 	GEOMETRYCOLLECTIONENGINE_API virtual void BeginPlay() override;
