Defer initialization of on-demand install cache and certificates until needed.
This change prevents the I/O store on-demand module initialization to fail for projects that hasn't setup
the persistent storage and isn't using IAD or IAS.

#rnx

diff --git a/Engine/Source/Runtime/Experimental/IoStore/OnDemand/Private/IoStoreOnDemand.cpp b/Engine/Source/Runtime/Experimental/IoStore/OnDemand/Private/IoStoreOnDemand.cpp
--- a/Engine/Source/Runtime/Experimental/IoStore/OnDemand/Private/IoStoreOnDemand.cpp
+++ b/Engine/Source/Runtime/Experimental/IoStore/OnDemand/Private/IoStoreOnDemand.cpp
@@ -28,7 +28,6 @@
 #include "Serialization/CompactBinarySerialization.h"
 #include "Serialization/CompactBinaryWriter.h"
 #include "Serialization/CustomVersion.h"
-#include "Serialization/LargeMemoryWriter.h"
 #include "Serialization/MemoryReader.h"
 #include "Statistics.h"
 #include "String/LexFromString.h"
@@ -348,47 +347,6 @@ static FIasCacheConfig GetIasCacheConfig(const TCHAR* CommandLine)
 	return Ret;
 }
 
-////////////////////////////////////////////////////////////////////////////////
-static void LoadCaCerts()
-{
-	using namespace UE::IoStore::HTTP;
-
-	// The following config option is used when staging to copy root certs PEM
-	const TCHAR* CertSection = TEXT("/Script/Engine.NetworkSettings");
-	const TCHAR* CertKey = TEXT("n.VerifyPeer");
-	bool bExpectCerts = false;
-	if (GConfig != nullptr)
-	{
-		GConfig->GetBool(CertSection, CertKey, bExpectCerts, GEngineIni);
-	}
-
-	// Open the certs file
-	IFileManager& Ifm = IFileManager::Get();
-	FString PemPath = FPaths::EngineContentDir() / TEXT("Certificates/ThirdParty/cacert.pem");
-	FArchive* Reader = Ifm.CreateFileReader(*PemPath);
-	if (!bExpectCerts && Reader == nullptr)
-	{
-		UE_LOG(LogIas, Warning, TEXT("Unable to load '%s'. Maybe it wasn't staged? Ensure '[%s]/%s=true' when staging"), *PemPath, CertSection, CertKey);
-		return;
-	}
-	checkf(Reader != nullptr, TEXT("%s/%s==true but '%s' could not be loaded"), CertSection, CertKey, *PemPath);
-
-	// Buffer certificate data
-	uint32 Size = uint32(Reader->TotalSize());
-	FIoBuffer PemData(Size);
-	FMutableMemoryView PemView = PemData.GetMutableView();
-	Reader->Serialize(PemView.GetData(), Size);
-
-	// Load the certs
-	FCertRoots CaRoots(PemData.GetView());
-
-	uint32 NumCerts = CaRoots.Num();
-	UE_LOG(LogIas, Display, TEXT("CaRoots: %u (%u .pem bytes))"), NumCerts, Size);
-
-	FCertRoots::SetDefault(MoveTemp(CaRoots));
-	delete Reader;
-}
-
 ////////////////////////////////////////////////////////////////////////////////
 /**
  * Utility to create a FArchive capable of reading from disk using the exact same pathing
@@ -1256,8 +1214,6 @@ void FIoStoreOnDemandModule::InitializeInternal()
 		}
 	}
 
-	LoadCaCerts();
-
 	// Make sure we haven't called initialize before
 	check(!HttpIoDispatcherBackend.IsValid());
 
diff --git a/Engine/Source/Runtime/Experimental/IoStore/OnDemand/Private/OnDemandHttpClient.cpp b/Engine/Source/Runtime/Experimental/IoStore/OnDemand/Private/OnDemandHttpClient.cpp
--- a/Engine/Source/Runtime/Experimental/IoStore/OnDemand/Private/OnDemandHttpClient.cpp
+++ b/Engine/Source/Runtime/Experimental/IoStore/OnDemand/Private/OnDemandHttpClient.cpp
@@ -4,9 +4,12 @@
 #include "Statistics.h"
 
 #include "Containers/StringConv.h"
+#include "HAL/FileManager.h"
 #include "HAL/IConsoleManager.h"
 #include "HAL/PlatformTime.h"
 #include "IO/IoStoreOnDemand.h"
+#include "Misc/ConfigCacheIni.h"
+#include "Misc/Paths.h"
 
 namespace UE::IoStore
 {
@@ -26,6 +29,62 @@ static FAutoConsoleVariableRef CVar_IasHttpAllowChunkedXfer(
 	TEXT("Enable/disable IAS' support for chunked transfer encoding")
 );
 
+FIoStatus LoadDefaultHttpCertificates(bool& bWasLoaded)
+{
+	using namespace UE::IoStore::HTTP;
+	
+	bWasLoaded = false;
+
+	static struct FDefaultCerts
+	{
+		FDefaultCerts(bool& bWasLoadedd)
+		{
+			bWasLoadedd = true;
+			Status = FIoStatus::Ok;
+
+			// The following config option is used when staging to copy root certs PEM
+			const TCHAR* CertSection = TEXT("/Script/Engine.NetworkSettings");
+			const TCHAR* CertKey = TEXT("n.VerifyPeer");
+
+			bool bVerifyPeer = false;
+			if (GConfig != nullptr)
+			{
+				GConfig->GetBool(CertSection, CertKey, bVerifyPeer, GEngineIni);
+			}
+
+			// Open the certs file
+			IFileManager& Ifm = IFileManager::Get();
+			const FString PemPath = FPaths::EngineContentDir() / TEXT("Certificates/ThirdParty/cacert.pem");
+			TUniquePtr<FArchive> Reader(Ifm.CreateFileReader(*PemPath));
+
+			if (Reader.IsValid())
+			{
+				// Buffer certificate data
+				const uint64 Size = Reader->TotalSize();
+				FIoBuffer PemData(Size);
+				FMutableMemoryView PemView = PemData.GetMutableView();
+				Reader->Serialize(PemView.GetData(), Size);
+
+				// Load the certs
+				FCertRoots CaRoots(PemData.GetView());
+
+				const uint32 NumCerts = CaRoots.Num();
+				FCertRoots::SetDefault(MoveTemp(CaRoots));
+
+				UE_LOG(LogIoStoreOnDemand, Display, TEXT("Loaded %u certificates from '%s'"), NumCerts, *PemPath);
+			}
+			else if (bVerifyPeer)
+			{
+				Status = FIoStatusBuilder(EIoErrorCode::FileNotOpen) << TEXT("Failed to open certificates file '") << PemPath << TEXT("'");
+			}
+		}
+
+		FIoStatus Status;
+	} DefaultCerts(bWasLoaded);
+
+	return DefaultCerts.Status;
+}
+
 static void LogHttpResult(const TCHAR* Host, const TCHAR* Url, uint32 StatusCode, uint64 DurationMs, uint64 Size, uint64 Offset, const char* Memo = "ok")
 {
 	Size >>= 10;
diff --git a/Engine/Source/Runtime/Experimental/IoStore/OnDemand/Private/OnDemandHttpClient.h b/Engine/Source/Runtime/Experimental/IoStore/OnDemand/Private/OnDemandHttpClient.h
--- a/Engine/Source/Runtime/Experimental/IoStore/OnDemand/Private/OnDemandHttpClient.h
+++ b/Engine/Source/Runtime/Experimental/IoStore/OnDemand/Private/OnDemandHttpClient.h
@@ -16,6 +16,8 @@
 namespace UE::IoStore
 {
 
+FIoStatus LoadDefaultHttpCertificates(bool& bWasLoaded); 
+
 enum class EHttpRedirects
 {
 	/** Redirects will be rejected and handled as failed requests. */
diff --git a/Engine/Source/Runtime/Experimental/IoStore/OnDemand/Private/OnDemandIoStore.cpp b/Engine/Source/Runtime/Experimental/IoStore/OnDemand/Private/OnDemandIoStore.cpp
--- a/Engine/Source/Runtime/Experimental/IoStore/OnDemand/Private/OnDemandIoStore.cpp
+++ b/Engine/Source/Runtime/Experimental/IoStore/OnDemand/Private/OnDemandIoStore.cpp
@@ -481,8 +481,6 @@ FOnDemandIoStore::FOnDemandIoStore()
 
 FOnDemandIoStore::~FOnDemandIoStore()
 {
-	FCoreDelegates::OnPostFork.RemoveAll(this);
-
 	FEncryptionKeyManager::Get().OnKeyAdded().RemoveAll(this);
 
 	if (OnMountPakHandle.IsValid())
@@ -499,19 +497,6 @@ FOnDemandIoStore::~FOnDemandIoStore()
 FIoStatus FOnDemandIoStore::Initialize()
 {
 	LLM_SCOPE_BYNAME(TEXT("FileSystem/OnDemandIoStore"));
-	const FIoStatus CacheStatus = InitializeOnDemandInstallCache();
-	if (CacheStatus.GetErrorCode() == EIoErrorCode::PendingFork)
-	{
-		UE_LOG(LogIoStoreOnDemand, Display, TEXT("Deferring install cache initialization until after process fork"));
-		if (!FCoreDelegates::OnPostFork.IsBoundToObject(this))
-		{
-			FCoreDelegates::OnPostFork.AddRaw(this, &FOnDemandIoStore::OnPostFork);
-		}
-	}
-	else if (!CacheStatus.IsOk())
-	{
-		return CacheStatus;
-	}
 
 #if 0
 	OnMountPakHandle = FCoreInternalDelegates::GetOnPakMountOperation().AddLambda(
@@ -768,24 +753,7 @@ FOnDemandChunkInfo FOnDemandIoStore::GetInstalledChunkInfo(const FIoChunkId& Chu
 	return GetChunkInfo(ChunkId, EOnDemandContainerFlags::Mounted | EOnDemandContainerFlags::InstallOnDemand);
 }
 
-void FOnDemandIoStore::OnPostFork(EForkProcessRole ProcessRole)
-{
-	if (ProcessRole != EForkProcessRole::Child)
-	{
-		UE_LOG(LogIoStoreOnDemand, Display, TEXT("OnPostFork ProcessRole Parent"));
-		return;
-	}
-
-	UE_LOG(LogIoStoreOnDemand, Display, TEXT("OnPostFork ProcessRole Child"));
-
-	const FIoStatus Status = Initialize();
-	if (!Status.IsOk())
-	{
-		UE_LOG(LogIoStoreOnDemand, Fatal, TEXT("Failed to initialize I/O store on demand (post fork), reason '%s'"), *Status.ToString());
-	}
-}
-
-FIoStatus FOnDemandIoStore::InitializeOnDemandInstallCache()
+FIoStatus FOnDemandIoStore::InitializeInstallCache()
 {
 	if (FForkProcessHelper::IsForkRequested() && !FForkProcessHelper::IsForkedChildProcess())
 	{
@@ -907,6 +875,16 @@ FIoStatus FOnDemandIoStore::InitializeOnDemandInstallCache()
 	return EIoErrorCode::Ok;
 }
 
+FIoStatus FOnDemandIoStore::ConditionallyInitializeInstallCache()
+{
+	FIoStatus Status = FIoStatus::Ok;
+	if (InstallCache.IsValid() == false)
+	{
+		Status = InitializeInstallCache();
+	}
+	return Status;
+}
+
 FOnDemandChunkInfo FOnDemandIoStore::GetChunkInfo(const FIoChunkId& ChunkId, EOnDemandContainerFlags ContainerFlags)
 {
 	TUniqueLock Lock(ContainerMutex);
@@ -1054,6 +1032,15 @@ FIoStatus FOnDemandIoStore::TickMountRequest(FMountRequest& MountRequest)
 {
 	UE_LOG(LogIoStoreOnDemand, Verbose, TEXT("Ticking mount request, MountId='%s'"), *MountRequest.Args.MountId);
 
+	bool bWasLoaded = false;
+	if (FIoStatus CertStatus = LoadDefaultHttpCertificates(bWasLoaded); CertStatus.IsOk() == false) 
+	{
+		if (bWasLoaded)
+		{
+			UE_LOG(LogIoStoreOnDemand, Error, TEXT("Failed to load certificates, reason '%s'"), *CertStatus.ToString());
+		}
+	}
+
 	const double StartTime = FPlatformTime::Seconds();
 	ON_SCOPE_EXIT
 	{
@@ -1287,18 +1274,14 @@ FOnDemandInstallResult FOnDemandIoStore::TickInstallRequest(const FInstallReques
 		return OutResult;
 	}
 
-	if (InstallCache.IsValid() == false || PackageStoreBackend.IsValid() == false)
+	Status = ConditionallyInitializeInstallCache();
+	if (Status.IsOk() == false)
 	{
-		if (FForkProcessHelper::IsForkRequested() && !FForkProcessHelper::IsForkedChildProcess())
-		{
-			Status = FIoStatusBuilder(EIoErrorCode::PendingFork) << TEXT("Install cache waiting for fork");
-			return OutResult;
-		}
-
-		Status = FIoStatusBuilder(EIoErrorCode::InvalidCode) << TEXT("Install cache not configured");
 		return OutResult;
 	}
 
+	check(InstallCache.IsValid() && PackageStoreBackend.IsValid()); 
+
 	TAnsiStringBuilder<512> ChunkUrl;
 	TSet<FSharedOnDemandContainer> ContainersForInstallation;
 	TSet<FPackageId> PackageIdsToInstall;
diff --git a/Engine/Source/Runtime/Experimental/IoStore/OnDemand/Private/OnDemandIoStore.h b/Engine/Source/Runtime/Experimental/IoStore/OnDemand/Private/OnDemandIoStore.h
--- a/Engine/Source/Runtime/Experimental/IoStore/OnDemand/Private/OnDemandIoStore.h
+++ b/Engine/Source/Runtime/Experimental/IoStore/OnDemand/Private/OnDemandIoStore.h
@@ -315,8 +315,8 @@ private:
 								const TArray<FPackageId>& PackageIds,
 								TSet<FSharedOnDemandContainer>& OutContainersForInstallation,
 								TSet<FPackageId>& OutPackageIdsToInstall) const;
-	void					OnPostFork(EForkProcessRole ProcessRole);
-	FIoStatus				InitializeOnDemandInstallCache();
+	FIoStatus				InitializeInstallCache();
+	FIoStatus				ConditionallyInitializeInstallCache();
 	FOnDemandChunkInfo		GetChunkInfo(const FIoChunkId& ChunkId, EOnDemandContainerFlags ContainerFlags);
 	void					TryEnterTickLoop();
 	void					TickLoop();
