make external collision for split geometry collection go to the root bone instead of being duplicated on every leaf

diff --git a/Engine/Source/Runtime/Experimental/GeometryCollectionEngine/Private/GeometryCollection/GeometryCollectionEngineConversion.cpp b/Engine/Source/Runtime/Experimental/GeometryCollectionEngine/Private/GeometryCollection/GeometryCollectionEngineConversion.cpp
--- a/Engine/Source/Runtime/Experimental/GeometryCollectionEngine/Private/GeometryCollection/GeometryCollectionEngineConversion.cpp
+++ b/Engine/Source/Runtime/Experimental/GeometryCollectionEngine/Private/GeometryCollection/GeometryCollectionEngineConversion.cpp
@@ -108,6 +108,50 @@
 	return FVector::One();
 }
 
+static void SetExternalCollisions(FVector3d MeshScale, UBodySetup* BodySetup, FGeometryCollection* GeometryCollection, int32 TransformIndex)
+{
+	if (BodySetup)
+	{
+		TArray<Chaos::FImplicitObjectPtr> Geoms;
+		Chaos::FShapesArray Shapes;
+
+		FGeometryAddParams CreateGeometryParams;
+		CreateGeometryParams.bDoubleSided = false;
+		CreateGeometryParams.CollisionData.CollisionFlags.bEnableQueryCollision = true;
+		CreateGeometryParams.CollisionData.CollisionFlags.bEnableSimCollisionComplex = false; // no support for trimesh in destruction
+		CreateGeometryParams.CollisionData.CollisionFlags.bEnableSimCollisionSimple = true;
+		CreateGeometryParams.CollisionTraceType = ECollisionTraceFlag::CTF_UseSimpleAsComplex;
+		CreateGeometryParams.Scale = MeshScale;
+		CreateGeometryParams.LocalTransform = Chaos::FRigidTransform3::Identity;
+		CreateGeometryParams.WorldTransform = Chaos::FRigidTransform3::Identity;
+		CreateGeometryParams.Geometry = &BodySetup->AggGeom;
+		CreateGeometryParams.TriMeshGeometries = MakeArrayView(BodySetup->TriMeshGeometries);
+
+		// todo(chaos) : this currently also create the shape array which is unnecessary ,this could be optimized by having a common function to create only the implicits 
+		ChaosInterface::CreateGeometry(CreateGeometryParams, Geoms, Shapes);
+
+		TManagedArray<Chaos::FImplicitObjectPtr>& ExternalCollisions = GeometryCollection->AddAttribute<Chaos::FImplicitObjectPtr>(FGeometryCollection::ExternalCollisionsAttribute, FGeometryCollection::TransformGroup);
+
+		ExternalCollisions[TransformIndex] = nullptr;
+		for (int32 GeomIndex = 0; GeomIndex < Geoms.Num();)
+		{
+			// make sure we only import box, sphere, capsule or convex
+			if (Geoms[GeomIndex] && IsImportableImplicitObjectType(*Geoms[GeomIndex]))
+			{
+				GeomIndex++;
+			}
+			else
+			{
+				Geoms.RemoveAtSwap(GeomIndex);
+			}
+		}
+		if (Geoms.Num() > 0)
+		{
+			ExternalCollisions[TransformIndex] = MakeImplicitObjectPtr<Chaos::FImplicitObjectUnion>(MoveTemp(Geoms));
+		}
+	}
+}
+
 void FGeometryCollectionEngineConversion::AppendMeshDescription(
 	const FMeshDescription* MeshDescription, const FString& Name, int32 MaterialStartIndex, const FTransform& StaticMeshTransform, 
 	FGeometryCollection* GeometryCollection, UBodySetup* BodySetup, bool ReindexMaterials, bool bAddInternalMaterials, bool bSetInternalFromMaterialIndex)
@@ -294,47 +338,7 @@
 	Transform[TransformIndex1] = FTransform3f(StaticMeshTransform);
 	Transform[TransformIndex1].SetScale3D(FVector3f(1.f, 1.f, 1.f));
 
-	// collisions
-	if (BodySetup)
-	{
-		TArray<Chaos::FImplicitObjectPtr> Geoms;
-		Chaos::FShapesArray Shapes;
-
-		FGeometryAddParams CreateGeometryParams;
-		CreateGeometryParams.bDoubleSided = false;
-		CreateGeometryParams.CollisionData.CollisionFlags.bEnableQueryCollision = true;
-		CreateGeometryParams.CollisionData.CollisionFlags.bEnableSimCollisionComplex = false; // no support for trimesh in destruction
-		CreateGeometryParams.CollisionData.CollisionFlags.bEnableSimCollisionSimple = true;
-		CreateGeometryParams.CollisionTraceType = ECollisionTraceFlag::CTF_UseSimpleAsComplex;
-		CreateGeometryParams.Scale = Scale;
-		CreateGeometryParams.LocalTransform = Chaos::FRigidTransform3::Identity;
-		CreateGeometryParams.WorldTransform = Chaos::FRigidTransform3::Identity;
-		CreateGeometryParams.Geometry = &BodySetup->AggGeom;
-		CreateGeometryParams.TriMeshGeometries = MakeArrayView(BodySetup->TriMeshGeometries);
-
-		// todo(chaos) : this currently also create the shape array which is unnecessary ,this could be optimized by having a common function to create only the implicits 
-		ChaosInterface::CreateGeometry(CreateGeometryParams, Geoms, Shapes);
-
-		TManagedArray<Chaos::FImplicitObjectPtr>& ExternaCollisions = GeometryCollection->AddAttribute<Chaos::FImplicitObjectPtr>(FGeometryCollection::ExternalCollisionsAttribute, FGeometryCollection::TransformGroup);
-
-		ExternaCollisions[TransformIndex1] = nullptr;
-		for (int32 GeomIndex = 0; GeomIndex < Geoms.Num();)
-		{
-			// make sure we only import box, sphere, capsule or convex
-			if (Geoms[GeomIndex] && IsImportableImplicitObjectType(*Geoms[GeomIndex]))
-			{
-				GeomIndex++;
-			}
-			else
-			{
-				Geoms.RemoveAtSwap(GeomIndex);
-			}
-		}
-		if (Geoms.Num() > 0)
-		{
-			ExternaCollisions[TransformIndex1] = MakeImplicitObjectPtr<Chaos::FImplicitObjectUnion>(MoveTemp(Geoms));
-		}
-	}
+	SetExternalCollisions(Scale, BodySetup, GeometryCollection, TransformIndex1);
 		
 	// Bone Hierarchy - Added at root with no common parent
 	TManagedArray<int32>& Parent = GeometryCollection->Parent;
@@ -739,7 +743,8 @@
 
 					for (FMeshDescription& MD : Descriptions)
 					{
-						AppendMeshDescription(&MD, StaticMesh->GetName(), StartMaterialIndex, MeshTransform, GeometryCollection, StaticMesh->GetBodySetup(), false, bAddInternalMaterials, bSetInternalFromMaterialIndex);
+						AppendMeshDescription(&MD, StaticMesh->GetName(), StartMaterialIndex, MeshTransform, GeometryCollection, nullptr /*Body Setup to be set later, on root*/, 
+							false, bAddInternalMaterials, bSetInternalFromMaterialIndex);
 					}
 
 					if (bReindexMaterials)
@@ -1635,6 +1640,13 @@
 		if (bSplitComponents && FGeometryCollectionClusteringUtility::ContainsMultipleRootBones(OutCollectionPtr.Get()))
 		{
 			FGeometryCollectionClusteringUtility::ClusterAllBonesUnderNewRoot(OutCollectionPtr.Get());
+			// If we split components, set external collision (if any) on the new root
+			TArray<int32> RootBones;
+			FGeometryCollectionClusteringUtility::GetRootBones(OutCollectionPtr.Get(), RootBones);
+			if (ensure(RootBones.Num() == 1))
+			{
+				SetExternalCollisions(GetMeshBuildScale3D(*StaticMesh) * MeshTransform.GetScale3D(), StaticMesh->GetBodySetup(), OutCollectionPtr.Get(), RootBones[0]);
+			}
 		}
 
 		NewGeometryCollection->InitializeMaterials();
