Iris -Replaced cvar net.Iris.ReplicationWriterMaxDestroyObjectsPerPacket with net.Iris.ReplicationWriterMaxDestroyObjectsPerFrame to avoid edge cases like when there are only objects pending destroy to replicate. Explicitly disallow writing object destroys in OOB packets.

diff --git a/Engine/Source/Runtime/Experimental/Iris/Core/Private/Iris/ReplicationSystem/ReplicationWriter.cpp b/Engine/Source/Runtime/Experimental/Iris/Core/Private/Iris/ReplicationSystem/ReplicationWriter.cpp
--- a/Engine/Source/Runtime/Experimental/Iris/Core/Private/Iris/ReplicationSystem/ReplicationWriter.cpp
+++ b/Engine/Source/Runtime/Experimental/Iris/Core/Private/Iris/ReplicationSystem/ReplicationWriter.cpp
@@ -66,13 +66,13 @@ static FAutoConsoleVariableRef CVarReplicationWriterMaxHugeObjectsInTransit(TEXT
 	TEXT("How many very large objects, one whose payload doesn't fit in a single packet, is allowed to be scheduled for send. Needs to be at least 1."));
 
 /*
- * net.Iris.ReplicationWriterMaxDestroyObjectsPerPacket
- * Limit the number of root objects to be destroyed per packet to relieve pressure on the ReplicationRecord count and potentially for client performance reasons. The number includes both regular object destruction as well as destruction infos.
+ * net.Iris.ReplicationWriterMaxDestroyObjectsPerFrame
+ * Limit the number of root objects to be destroyed per frame to relieve pressure on the ReplicationRecord count and potentially for client performance reasons. The number includes both regular object destruction as well as destruction infos.
  * Without limiting several hundred objects could fit a packet which would exhaust the replication records long before filling the packet window.
  */
-static int32 GReplicationWriterMaxDestroyObjectsPerPacket = 128;
-static FAutoConsoleVariableRef CVarReplicationWriterMaxDestroyObjectsPerPacket(TEXT("net.Iris.ReplicationWriterMaxDestroyObjectsPerPacket"), GReplicationWriterMaxDestroyObjectsPerPacket,
-	TEXT("How many objects can be replicated for destroy per packet. The number is for regular destruction and destruction infos combined. A value less or equal to zero means unlimited."));
+static int32 GReplicationWriterMaxDestroyObjectsPerFrame = 150;
+static FAutoConsoleVariableRef CVarReplicationWriterMaxDestroyObjectsPerFrame(TEXT("net.Iris.ReplicationWriterMaxDestroyObjectsPerFrame"), GReplicationWriterMaxDestroyObjectsPerFrame,
+	TEXT("How many objects can be replicated for destroy per frame. The number is for regular destruction and destruction infos combined. A value less or equal to zero means unlimited."));
 
 /*
  * net.Iris.ReplicationWriterReplicationRecordStarvationThreshold
@@ -1848,8 +1848,8 @@ uint32 FReplicationWriter::WriteObjectsAndSubObjectsPendingDestroy(FNetSerializa
 
 	Writer.WriteBits(WrittenCount, DestroyObjectBitCount);
 
-	// Can't write if bitstream is overflown. Shouldn't write if we're in replication record starvation mode.
-	if (Writer.IsOverflown() || WriteContext.bIsInReplicationRecordStarvation)
+	// Can't write if bitstream is overflown. Shouldn't write if we're an OOB packet or in replication record starvation mode.
+	if (Writer.IsOverflown() || WriteContext.bIsOOBPacket || WriteContext.bIsInReplicationRecordStarvation)
 	{
 		return 0U;
 	}
@@ -2005,12 +2005,25 @@ uint32 FReplicationWriter::WriteRootObjectsPendingDestroy(FNetSerializationConte
 
 	Writer.WriteBits(WrittenCount, DestroyObjectBitCount);
 
-	// Can't write if bitstream is overflown. Shouldn't write if we're in replication record starvation mode.
-	if (Writer.IsOverflown() || WriteContext.bIsInReplicationRecordStarvation)
+	// Can't write if bitstream is overflown. Shouldn't write if we're an OOB packet or in replication record starvation mode.
+	if (Writer.IsOverflown() || WriteContext.bIsOOBPacket || WriteContext.bIsInReplicationRecordStarvation)
 	{
 		return 0U;
 	}
 
+	// Check whether we're still allowed to write object destroys.
+	uint32 EffectiveMaxDestroyObjectCount = MaxDestroyObjectCount;
+	const int32 MaxDestroyObjectsPerFrame = GReplicationWriterMaxDestroyObjectsPerFrame;
+	if (MaxDestroyObjectsPerFrame > 0)
+	{
+		if (WriteContext.WrittenDestroyObjectCount > static_cast<uint32>(MaxDestroyObjectsPerFrame))
+		{
+			return 0U;
+		}
+
+		EffectiveMaxDestroyObjectCount = static_cast<uint32>(MaxDestroyObjectsPerFrame) - WriteContext.WrittenDestroyObjectCount;
+	}
+
 	auto AreSubObjectsReadyToBeDestroyed = [this](uint32 InternalIndex)
 	{
 		for (uint32 SubObjectIndex : NetRefHandleManager->GetSubObjects(InternalIndex))
@@ -2038,8 +2051,6 @@ uint32 FReplicationWriter::WriteRootObjectsPendingDestroy(FNetSerializationConte
 	};
 
 	bool bWroteAllDestroyedObjects = true;
-
-	const uint32 EffectiveMaxDestroyObjectCount = FPlatformMath::Min((GReplicationWriterMaxDestroyObjectsPerPacket > 0 ? static_cast<uint32>(GReplicationWriterMaxDestroyObjectsPerPacket) : MaxDestroyObjectCount), MaxDestroyObjectCount);
 	for (uint32 InternalIndex = 0U; ((InternalIndex = ObjectsPendingDestroy.FindFirstOne(InternalIndex + 1U)) != FNetBitArray::InvalidIndex) && (WrittenCount < EffectiveMaxDestroyObjectCount); )
 	{
 		FReplicationInfo& Info = GetReplicationInfo(InternalIndex);
@@ -3137,7 +3148,8 @@ int FReplicationWriter::WriteDestructionInfo(FNetSerializationContext& Context,
 {
 	IRIS_PROFILER_SCOPE(FReplicationWriter_WriteDestructionInfo);
 
-	if ((GReplicationWriterMaxDestroyObjectsPerPacket > 0) && (WriteContext.WrittenDestroyObjectCount >= static_cast<uint32>(GReplicationWriterMaxDestroyObjectsPerPacket)))
+	const int32 MaxDestroyObjectsPerFrame = GReplicationWriterMaxDestroyObjectsPerFrame;
+	if ((MaxDestroyObjectsPerFrame > 0) && (WriteContext.WrittenDestroyObjectCount >= static_cast<uint32>(MaxDestroyObjectsPerFrame)))
 	{
 		return 0;
 	}
@@ -3619,12 +3631,14 @@ UDataStream::EWriteResult FReplicationWriter::BeginWrite(const UDataStream::FBeg
 
 	WriteContext.CurrentIndex = 0U;
 	WriteContext.FailedToWriteSmallObjectCount = 0U;
+	WriteContext.WrittenDestroyObjectCount = 0U;
 	WriteContext.SortedObjectCount = 0U;
 	WriteContext.NumWrittenPacketsInThisBatch = 0U;
 	WriteContext.bCanWriteMoreData = Params.bCanWriteMoreData;
 
-	// Set in Write
+	// Updated properly in Write
 	WriteContext.bIsInReplicationRecordStarvation = 0;
+	WriteContext.bIsOOBPacket = 0;
 
 	// Clear net stats. Used for CVS and Network Insights stats.
 	WriteContext.Stats.Reset();
@@ -3700,9 +3714,10 @@ UDataStream::EWriteResult FReplicationWriter::Write(FNetSerializationContext& Co
 	// We have somethings in the WriteContext that must be reset each packet
 	WriteContext.FailedToWriteSmallObjectCount = 0U;
 	WriteContext.WrittenBatchCount = 0U;
-	WriteContext.WrittenDestroyObjectCount = 0U;
 	// Check whether we are running low on replication records or not and need to throttle object replication.
 	WriteContext.bIsInReplicationRecordStarvation = ReplicationRecord.GetUnusedInfoCount() < static_cast<uint32>(GReplicationWriterReplicationRecordStarvationThreshold);
+	// Check whether this is to be considered an OOB packet which should avoid writing things like objects pending destroy
+	WriteContext.bIsOOBPacket = (WriteContext.WriteMode == EDataStreamWriteMode::PostTickDispatch) || (static_cast<int32>(WriteContext.NumWrittenPacketsInThisBatch) >= GReplicationWriterMaxAllowedPacketsIfNotHugeObject);
 
 	FNetBitStreamWriter& Writer = *Context.GetBitStreamWriter();
 
diff --git a/Engine/Source/Runtime/Experimental/Iris/Core/Private/Iris/ReplicationSystem/ReplicationWriter.h b/Engine/Source/Runtime/Experimental/Iris/Core/Private/Iris/ReplicationSystem/ReplicationWriter.h
--- a/Engine/Source/Runtime/Experimental/Iris/Core/Private/Iris/ReplicationSystem/ReplicationWriter.h
+++ b/Engine/Source/Runtime/Experimental/Iris/Core/Private/Iris/ReplicationSystem/ReplicationWriter.h
@@ -255,6 +255,8 @@ private:
 		uint32 bHasHugeObjectToSend : 1;
 		// Whether there are OOB attachments to send or not
 		uint32 bHasOOBAttachmentsToSend : 1;
+		// Whether this packet is mainly intended for OOB and HugeObject data
+		uint32 bIsOOBPacket : 1;
 		// Whether this context is valid or not
 		uint32 bIsValid : 1;
 
