Iris
* Code that handled dispatching the callbacks for subobject creation and destruction  was moved from ReplicationBridge to the NetObjectFactory

diff --git a/Engine/Plugins/Runtime/ReplicationSystemTestPlugin/Source/Private/Tests/ReplicationSystem/ReplicatedTestObject.cpp b/Engine/Plugins/Runtime/ReplicationSystemTestPlugin/Source/Private/Tests/ReplicationSystem/ReplicatedTestObject.cpp
--- a/Engine/Plugins/Runtime/ReplicationSystemTestPlugin/Source/Private/Tests/ReplicationSystem/ReplicatedTestObject.cpp
+++ b/Engine/Plugins/Runtime/ReplicationSystemTestPlugin/Source/Private/Tests/ReplicationSystem/ReplicatedTestObject.cpp
@@ -32,7 +32,6 @@
 
 UReplicatedTestObjectBridge::UReplicatedTestObjectBridge()
 : UObjectReplicationBridge()
-, CreatedObjectsOnNode(nullptr)
 {
 	
 }
@@ -109,28 +108,6 @@ void UReplicatedTestObjectBridge::EndReplication(UReplicatedTestObject* Instance
 	StopReplicatingNetObject(Instance, Flags);
 }
 
-void UReplicatedTestObjectBridge::DestroyInstanceFromRemote(const FDestroyInstanceParams& Params)
-{
-	if (!Params.Instance)
-	{
-		return;
-	}
-
-	if (Params.DestroyReason == EReplicationBridgeDestroyInstanceReason::Destroy)
-	{
-		if (EnumHasAnyFlags(Params.DestroyFlags, EReplicationBridgeDestroyInstanceFlags::AllowDestroyInstanceFromRemote))
-		{
-			// Remove the object from the created objects on the node
-			if (CreatedObjectsOnNode)
-			{
-				CreatedObjectsOnNode->Remove(TStrongObjectPtr<UObject>(Params.Instance));
-			}
-
-			Params.Instance->PreDestroyFromReplication();
-			Params.Instance->MarkAsGarbage();
-		}
-	}
-}
 
 bool UReplicatedTestObjectBridge::IsAllowedToDestroyInstance(const UObject* Instance) const
 {
@@ -152,6 +129,17 @@ void UReplicatedTestObjectBridge::SetExternalPreUpdateFunctor(TFunction<void(TAr
 	SetInstancePreUpdateFunction(PreUpdateFunctor);
 }
 
+void UReplicatedTestObjectBridge::SetCreatedObjectsOnNode(TArray<TStrongObjectPtr<UObject>>* InCreatedObjectsOnNode)
+{
+	using namespace UE::Net;
+	UNetObjectFactory* Factory = GetNetFactory(ReplicatedObjectFactoryId);
+	if (ensure(Factory))
+	{
+		CastChecked<UReplicatedTestObjectFactory>(Factory)->SetCreatedObjectsOnNode(InCreatedObjectsOnNode);
+	}
+}
+
+
 //////////////////////////////////////////////////////////////////////////
 // Implementation for UTestReplicatedIrisPropertyComponent
 //////////////////////////////////////////////////////////////////////////
diff --git a/Engine/Plugins/Runtime/ReplicationSystemTestPlugin/Source/Private/Tests/ReplicationSystem/ReplicatedTestObject.h b/Engine/Plugins/Runtime/ReplicationSystemTestPlugin/Source/Private/Tests/ReplicationSystem/ReplicatedTestObject.h
--- a/Engine/Plugins/Runtime/ReplicationSystemTestPlugin/Source/Private/Tests/ReplicationSystem/ReplicatedTestObject.h
+++ b/Engine/Plugins/Runtime/ReplicationSystemTestPlugin/Source/Private/Tests/ReplicationSystem/ReplicatedTestObject.h
@@ -538,7 +538,7 @@ public:
 
 	virtual void Initialize(UReplicationSystem* InReplicationSystem) override;
 	
-	void SetCreatedObjectsOnNode(TArray<TStrongObjectPtr<UObject>>* InCreatedObjectsOnNode) { CreatedObjectsOnNode = InCreatedObjectsOnNode; }
+	void SetCreatedObjectsOnNode(TArray<TStrongObjectPtr<UObject>>* InCreatedObjectsOnNode);
 
 	// This is the local Interface, it is up to each bridge implementation to define the interface for that type
 	// In this example we have methods that directly uses UReplicatedTestObject;
@@ -577,15 +577,10 @@ public:
 		bool bSuppressCreateInstanceFailedEnsure;
 	};
 
-public:
-
-	TArray<TStrongObjectPtr<UObject>>* CreatedObjectsOnNode;
-
 protected:
 
 	friend FSupressCreateInstanceFailedEnsureScope;
 
-	virtual void DestroyInstanceFromRemote(const FDestroyInstanceParams& Params) override;
 	virtual bool IsAllowedToDestroyInstance(const UObject* Instance) const override;
 
 	TFunction<void(FNetRefHandle NetHandle, const UObject* ReplicatedObject, FVector& OutLocation, float& OutCullDistance)> WorldLocationUpdateFunc;
diff --git a/Engine/Plugins/Runtime/ReplicationSystemTestPlugin/Source/Private/Tests/ReplicationSystem/ReplicatedTestObjectFactory.cpp b/Engine/Plugins/Runtime/ReplicationSystemTestPlugin/Source/Private/Tests/ReplicationSystem/ReplicatedTestObjectFactory.cpp
--- a/Engine/Plugins/Runtime/ReplicationSystemTestPlugin/Source/Private/Tests/ReplicationSystem/ReplicatedTestObjectFactory.cpp
+++ b/Engine/Plugins/Runtime/ReplicationSystemTestPlugin/Source/Private/Tests/ReplicationSystem/ReplicatedTestObjectFactory.cpp
@@ -168,9 +168,9 @@ UNetObjectFactory::FInstantiateResult UReplicatedTestObjectFactory::InstantiateR
 	UReplicatedTestObjectBridge* TestBridge = Cast<UReplicatedTestObjectBridge>(Bridge);
 
 	// Store the object so that we can find detached/torn off instances from tests
-	if (TestBridge->CreatedObjectsOnNode)
+	if (CreatedObjectsOnNode)
 	{
-		TestBridge->CreatedObjectsOnNode->Add(TStrongObjectPtr<UObject>(CreatedObject));
+		CreatedObjectsOnNode->Add(TStrongObjectPtr<UObject>(CreatedObject));
 	}
 	
 	FInstantiateResult InstantiateResult;
@@ -184,4 +184,27 @@ void UReplicatedTestObjectFactory::PostInit(const FPostInitContext& Context)
 	UReplicatedTestObject* Instance = CastChecked<UReplicatedTestObject>(Context.Instance);
 
 	Instance->NetRefHandle = Context.Handle;
+}
+
+void UReplicatedTestObjectFactory::DestroyedFromReplication(const FDestroyedContext& Context)
+{
+	if (!Context.DestroyedInstance)
+	{
+		return;
+	}
+
+	if (Context.DestroyReason == EReplicationBridgeDestroyInstanceReason::Destroy)
+	{
+		if (EnumHasAnyFlags(Context.DestroyFlags, EReplicationBridgeDestroyInstanceFlags::AllowDestroyInstanceFromRemote))
+		{
+			// Remove the object from the created objects on the node
+			if (CreatedObjectsOnNode)
+			{
+				CreatedObjectsOnNode->Remove(TStrongObjectPtr<UObject>(Context.DestroyedInstance));
+			}
+
+			Context.DestroyedInstance->PreDestroyFromReplication();
+			Context.DestroyedInstance->MarkAsGarbage();
+		}
+	}
 }
\ No newline at end of file
diff --git a/Engine/Plugins/Runtime/ReplicationSystemTestPlugin/Source/Private/Tests/ReplicationSystem/ReplicatedTestObjectFactory.h b/Engine/Plugins/Runtime/ReplicationSystemTestPlugin/Source/Private/Tests/ReplicationSystem/ReplicatedTestObjectFactory.h
--- a/Engine/Plugins/Runtime/ReplicationSystemTestPlugin/Source/Private/Tests/ReplicationSystem/ReplicatedTestObjectFactory.h
+++ b/Engine/Plugins/Runtime/ReplicationSystemTestPlugin/Source/Private/Tests/ReplicationSystem/ReplicatedTestObjectFactory.h
@@ -40,6 +40,16 @@ public:
 	virtual FInstantiateResult InstantiateReplicatedObjectFromHeader(const FInstantiateContext& Context, const UE::Net::FNetObjectCreationHeader* InHeader) override;
 
 	virtual void PostInit(const FPostInitContext& Context) override;
+	virtual void DestroyedFromReplication(const FDestroyedContext& Context) override;
+
+	void SetCreatedObjectsOnNode(TArray<TStrongObjectPtr<UObject>>* InCreatedObjectsOnNode)
+	{
+		CreatedObjectsOnNode = InCreatedObjectsOnNode;
+	}
+
+public:
+
+	TArray<TStrongObjectPtr<UObject>>* CreatedObjectsOnNode = nullptr;
 
 protected:
 
diff --git a/Engine/Source/Runtime/Engine/Private/Net/Iris/ReplicationSystem/EngineReplicationBridge.cpp b/Engine/Source/Runtime/Engine/Private/Net/Iris/ReplicationSystem/EngineReplicationBridge.cpp
--- a/Engine/Source/Runtime/Engine/Private/Net/Iris/ReplicationSystem/EngineReplicationBridge.cpp
+++ b/Engine/Source/Runtime/Engine/Private/Net/Iris/ReplicationSystem/EngineReplicationBridge.cpp
@@ -288,8 +288,8 @@ UE::Net::FNetRefHandle UEngineReplicationBridge::StartReplicatingActor(AActor* A
 
 	// Create handles for the registered fragments
 	Super::FRootObjectReplicationParams RootObjectParams = {
-		.bNeedsPreUpdate = 1U,
-		.bNeedsWorldLocationUpdate = 1U,
+		.bNeedsPreUpdate = true,
+		.bNeedsWorldLocationUpdate = true,
 		.bIsDormant = Actor->NetDormancy > DORM_Awake,
 		.StaticPriority = (Actor->bAlwaysRelevant || Actor->bOnlyRelevantToOwner) ? Actor->NetPriority : 0.0f,
 		.PollFrequency = Actor->GetNetUpdateFrequency()
@@ -524,61 +524,6 @@ void UEngineReplicationBridge::StopReplicatingComponent(UActorComponent* ActorCo
 	}
 }
 
-//$IRIS todo: move this to be a factory callback
-void UEngineReplicationBridge::OnSubObjectCreatedFromReplication(FNetRefHandle SubObjectHandle)
-{
-	AActor* RootObject = Cast<AActor>(GetReplicatedObject(GetRootObjectOfSubObject(SubObjectHandle)));
-	UObject* SubObject = GetReplicatedObject(SubObjectHandle);
-	if (IsValid(RootObject) && IsValid(SubObject))
-	{
-		RootObject->OnSubobjectCreatedFromReplication(SubObject);
-	}
-}
-
-void UEngineReplicationBridge::DestroyInstanceFromRemote(const FDestroyInstanceParams& Params)
-{
-	if (Params.DestroyReason == EReplicationBridgeDestroyInstanceReason::DoNotDestroy)
-	{
-		return;
-	}
-
-	if (AActor* Actor = Cast<AActor>(Params.Instance))
-	{
-		if ((Params.DestroyReason == EReplicationBridgeDestroyInstanceReason::TearOff) && !NetDriver->ShouldClientDestroyTearOffActors())
-		{
-			NetDriver->ClientSetActorTornOff(Actor);
-		}
-		else if (EnumHasAnyFlags(Params.DestroyFlags, EReplicationBridgeDestroyInstanceFlags::AllowDestroyInstanceFromRemote))
-		{
-			// Any subobjects have already been detached from ReplicationBridge
-			Actor->PreDestroyFromReplication();
-			Actor->Destroy(true);
-		}
-	}
-	else
-	{
-		// If the SubObject is being torn off it is up to owning actor to clean it up properly
-		if (Params.DestroyReason == EReplicationBridgeDestroyInstanceReason::TearOff)
-		{
-			return;
-		}
-
-		if (!EnumHasAnyFlags(Params.DestroyFlags, EReplicationBridgeDestroyInstanceFlags::AllowDestroyInstanceFromRemote))
-		{
-			return;
-		}
-
-		AActor* Owner = Cast<AActor>(Params.RootObject);
-		if (ensureMsgf(IsValid(Owner) && !Owner->IsUnreachable(), TEXT("UEngineReplicationBridge::DestroyInstanceFromRemote Destroyed subobject: %s has an invalid owner: %s"), *GetNameSafe(Params.Instance), *GetPathNameSafe(Params.RootObject)))
-		{
-			Owner->OnSubobjectDestroyFromReplication(Params.Instance);
-		}
-
-		Params.Instance->PreDestroyFromReplication();
-		Params.Instance->MarkAsGarbage();
-	}
-}
-
 void UEngineReplicationBridge::GetInitialDependencies(FNetRefHandle Handle, FNetDependencyInfoArray& OutDependencies) const
 {
 	using namespace UE::Net;
diff --git a/Engine/Source/Runtime/Engine/Private/Net/Iris/ReplicationSystem/NetActorFactory.cpp b/Engine/Source/Runtime/Engine/Private/Net/Iris/ReplicationSystem/NetActorFactory.cpp
--- a/Engine/Source/Runtime/Engine/Private/Net/Iris/ReplicationSystem/NetActorFactory.cpp
+++ b/Engine/Source/Runtime/Engine/Private/Net/Iris/ReplicationSystem/NetActorFactory.cpp
@@ -483,6 +483,36 @@ void UNetActorFactory::PostInit(const FPostInitContext& Context)
 	}
 }
 
+void UNetActorFactory::SubObjectCreatedFromReplication(UE::Net::FNetRefHandle RootObject, UE::Net::FNetRefHandle SubObjectCreated)
+{
+	AActor* RootActor = Cast<AActor>(Bridge->GetReplicatedObject(RootObject));
+	UObject* SubObject = Bridge->GetReplicatedObject(SubObjectCreated);
+	if (IsValid(RootActor) && IsValid(SubObject))
+	{
+		RootActor->OnSubobjectCreatedFromReplication(SubObject);
+	}
+}
+
+void UNetActorFactory::DestroyedFromReplication(const FDestroyedContext& Context)
+{
+	UEngineReplicationBridge* ActorBridge = CastChecked<UEngineReplicationBridge>(Bridge);
+	UNetDriver* NetDriver = ActorBridge->GetNetDriver();
+
+	if (AActor* Actor = Cast<AActor>(Context.DestroyedInstance))
+	{
+		if ((Context.DestroyReason == EReplicationBridgeDestroyInstanceReason::TearOff) && !NetDriver->ShouldClientDestroyTearOffActors())
+		{
+			NetDriver->ClientSetActorTornOff(Actor);
+		}
+		else if (EnumHasAnyFlags(Context.DestroyFlags, EReplicationBridgeDestroyInstanceFlags::AllowDestroyInstanceFromRemote))
+		{
+			// Any subobjects have already been detached from ReplicationBridge
+			Actor->PreDestroyFromReplication();
+			Actor->Destroy(true);
+		}
+	}	
+}
+
  namespace UE::Net
  {
 
diff --git a/Engine/Source/Runtime/Engine/Private/Net/Iris/ReplicationSystem/NetSubObjectFactory.cpp b/Engine/Source/Runtime/Engine/Private/Net/Iris/ReplicationSystem/NetSubObjectFactory.cpp
--- a/Engine/Source/Runtime/Engine/Private/Net/Iris/ReplicationSystem/NetSubObjectFactory.cpp
+++ b/Engine/Source/Runtime/Engine/Private/Net/Iris/ReplicationSystem/NetSubObjectFactory.cpp
@@ -197,6 +197,35 @@ bool UNetSubObjectFactory::SerializeHeader(const UE::Net::FCreationHeaderContext
 	return SubObjectHeader->Serialize(Context);
 }
 
+void UNetSubObjectFactory::SubObjectCreatedFromReplication(UE::Net::FNetRefHandle RootObject, UE::Net::FNetRefHandle SubObjectCreated)
+{
+	ensureMsgf(false, TEXT("NetSubObjectFactory::SubObjectCreatedFromReplication should never be called since subobjects cannot have their own subobject list. RootObject: %s, SubObjectCreated: %s"), *Bridge->PrintObjectFromNetRefHandle(RootObject), *Bridge->PrintObjectFromNetRefHandle(SubObjectCreated));
+}
+
+void UNetSubObjectFactory::DestroyedFromReplication(const FDestroyedContext& Context)
+{
+	// If the SubObject is being torn off it is up to owning actor to clean it up properly
+	if (Context.DestroyReason == EReplicationBridgeDestroyInstanceReason::TearOff)
+	{
+		return;
+	}
+
+	if (!EnumHasAnyFlags(Context.DestroyFlags, EReplicationBridgeDestroyInstanceFlags::AllowDestroyInstanceFromRemote))
+	{
+		return;
+	}
+
+	AActor* Owner = Cast<AActor>(Context.RootObject);
+
+	if (ensureMsgf(IsValid(Owner) && !Owner->IsUnreachable(), TEXT("UEngineReplicationBridge::DestroyInstanceFromRemote Destroyed subobject: %s has an invalid owner: %s"), *GetNameSafe(Context.DestroyedInstance), *GetPathNameSafe(Context.RootObject)))
+	{
+		Owner->OnSubobjectDestroyFromReplication(Context.DestroyedInstance);
+	}
+
+	Context.DestroyedInstance->PreDestroyFromReplication();
+	Context.DestroyedInstance->MarkAsGarbage();
+}
+
 //------------------------------------------------------------------------
 
 namespace UE::Net
diff --git a/Engine/Source/Runtime/Engine/Public/Net/Iris/ReplicationSystem/EngineReplicationBridge.h b/Engine/Source/Runtime/Engine/Public/Net/Iris/ReplicationSystem/EngineReplicationBridge.h
--- a/Engine/Source/Runtime/Engine/Public/Net/Iris/ReplicationSystem/EngineReplicationBridge.h
+++ b/Engine/Source/Runtime/Engine/Public/Net/Iris/ReplicationSystem/EngineReplicationBridge.h
@@ -103,8 +103,6 @@ protected:
 	// UObjectReplicationBridge
 	virtual void Initialize(UReplicationSystem* ReplicationSystem) override;
 	virtual void Deinitialize() override;
-	virtual void OnSubObjectCreatedFromReplication(FNetRefHandle SubObjectHandle) override;
-	virtual void DestroyInstanceFromRemote(const FDestroyInstanceParams& Params) override;
 	virtual void GetInitialDependencies(FNetRefHandle Handle, FNetDependencyInfoArray& OutDependencies) const override;
 	virtual bool RemapPathForPIE(uint32 ConnectionId, FString& Path, bool bReading) const override;
 	virtual bool ObjectLevelHasFinishedLoading(UObject* Object) const override;
diff --git a/Engine/Source/Runtime/Engine/Public/Net/Iris/ReplicationSystem/NetActorFactory.h b/Engine/Source/Runtime/Engine/Public/Net/Iris/ReplicationSystem/NetActorFactory.h
--- a/Engine/Source/Runtime/Engine/Public/Net/Iris/ReplicationSystem/NetActorFactory.h
+++ b/Engine/Source/Runtime/Engine/Public/Net/Iris/ReplicationSystem/NetActorFactory.h
@@ -123,6 +123,11 @@ public:
 
 	virtual void PostInit(const FPostInitContext& Context) override;
 
+	virtual void SubObjectCreatedFromReplication(UE::Net::FNetRefHandle RootObject, UE::Net::FNetRefHandle SubObjectCreated) override;
+
+	virtual void DestroyedFromReplication(const FDestroyedContext& Context) override;
+
+
 protected:
 
 	virtual TUniquePtr<UE::Net::FNetObjectCreationHeader> CreateAndFillHeader(UE::Net::FNetRefHandle Handle) override;
diff --git a/Engine/Source/Runtime/Engine/Public/Net/Iris/ReplicationSystem/NetSubObjectFactory.h b/Engine/Source/Runtime/Engine/Public/Net/Iris/ReplicationSystem/NetSubObjectFactory.h
--- a/Engine/Source/Runtime/Engine/Public/Net/Iris/ReplicationSystem/NetSubObjectFactory.h
+++ b/Engine/Source/Runtime/Engine/Public/Net/Iris/ReplicationSystem/NetSubObjectFactory.h
@@ -29,6 +29,10 @@ public:
 
 	virtual FInstantiateResult InstantiateReplicatedObjectFromHeader(const FInstantiateContext& Context, const UE::Net::FNetObjectCreationHeader* Header) override;
 
+	virtual void SubObjectCreatedFromReplication(UE::Net::FNetRefHandle RootObject, UE::Net::FNetRefHandle SubObjectCreated) override;
+
+	virtual void DestroyedFromReplication(const FDestroyedContext& Context) override;
+
 protected:
 
 	virtual TUniquePtr<UE::Net::FNetObjectCreationHeader> CreateAndFillHeader(UE::Net::FNetRefHandle Handle) override;
diff --git a/Engine/Source/Runtime/Experimental/Iris/Core/Private/Iris/ReplicationSystem/ObjectReplicationBridge.cpp b/Engine/Source/Runtime/Experimental/Iris/Core/Private/Iris/ReplicationSystem/ObjectReplicationBridge.cpp
--- a/Engine/Source/Runtime/Experimental/Iris/Core/Private/Iris/ReplicationSystem/ObjectReplicationBridge.cpp
+++ b/Engine/Source/Runtime/Experimental/Iris/Core/Private/Iris/ReplicationSystem/ObjectReplicationBridge.cpp
@@ -753,36 +753,42 @@ void UObjectReplicationBridge::StopReplicatingNetObject(UObject* Instance, EEndR
 	}
 }
 
-void UObjectReplicationBridge::DetachInstanceFromRemote(FNetRefHandle Handle, EReplicationBridgeDestroyInstanceReason DestroyReason, EReplicationBridgeDestroyInstanceFlags DestroyFlags)
+void UObjectReplicationBridge::DetachInstanceFromRemote(FNetRefHandle Handle, EReplicationBridgeDestroyInstanceReason DestroyReason, EReplicationBridgeDestroyInstanceFlags DestroyFlags, UE::Net::FNetObjectFactoryId NetFactoryId)
 {
-	UE_LOG_OBJECTREPLICATIONBRIDGE(Verbose, TEXT("DetachInstanceFromRemote %s DestroyReason: %s DestroyFlags: %u"), *PrintObjectFromNetRefHandle(Handle), ToCStr(LexToString(DestroyReason)), unsigned(DestroyFlags));
-
+	using namespace UE::Net;
 	using namespace UE::Net::Private;
 
-	FDestroyInstanceParams DestroyInstanceParams;
-	
-	DestroyInstanceParams.Instance = GetObjectFromReferenceHandle(Handle);
-	DestroyInstanceParams.DestroyReason = DestroyReason;
-	DestroyInstanceParams.DestroyFlags = DestroyFlags;
+	UE_LOG_OBJECTREPLICATIONBRIDGE(Verbose, TEXT("DetachInstanceFromRemote %s DestroyReason: %s DestroyFlags: %u"), *PrintObjectFromNetRefHandle(Handle), ToCStr(LexToString(DestroyReason)), unsigned(DestroyFlags));
 
+	// Tell factory about the imminent destruction
 	if (DestroyReason != EReplicationBridgeDestroyInstanceReason::DoNotDestroy)
 	{
-		const FInternalNetRefIndex ObjectNetIndex = NetRefHandleManager->GetInternalIndex(Handle);
-		const FNetRefHandleManager::FReplicatedObjectData& RepObjectData = NetRefHandleManager->GetReplicatedObjectData(ObjectNetIndex);
-		DestroyInstanceParams.RootObject = RepObjectData.SubObjectRootIndex != FNetRefHandleManager::InvalidInternalIndex ? NetRefHandleManager->GetReplicatedObjectInstance(RepObjectData.SubObjectRootIndex) : nullptr;
-	}
-	
-	UnregisterInstance(Handle);
+		UNetObjectFactory::FDestroyedContext DestroyInstanceParams
+		{
+			.DestroyedInstance = GetObjectFromReferenceHandle(Handle),
+			.DestroyReason = DestroyReason,
+			.DestroyFlags = DestroyFlags,
+        };
 
-	// Destroy instance if requested
-	if (DestroyInstanceParams.Instance && DestroyReason != EReplicationBridgeDestroyInstanceReason::DoNotDestroy)
-	{
-		DestroyInstanceFromRemote(DestroyInstanceParams);
+		// Tell the factory to handle the object if the instance exists
+        if (ensureMsgf(NetFactoryId != InvalidNetObjectFactoryId, TEXT("Received invalid netfactory for %s: %s"), *Handle.ToString(), *GetPathNameSafe(DestroyInstanceParams.DestroyedInstance)))
+		{
+            if (DestroyInstanceParams.DestroyedInstance)
+            {
+                const FInternalNetRefIndex DestroyedObjectIndex = NetRefHandleManager->GetInternalIndex(Handle);
+                const FNetRefHandleManager::FReplicatedObjectData& DestroyedObjectData = NetRefHandleManager->GetReplicatedObjectData(DestroyedObjectIndex);
+
+                DestroyInstanceParams.RootObject = DestroyedObjectData.SubObjectRootIndex != FNetRefHandleManager::InvalidInternalIndex ? NetRefHandleManager->GetReplicatedObjectInstance(DestroyedObjectData.SubObjectRootIndex) : nullptr;
+
+                UNetObjectFactory* Factory = GetNetFactory(NetFactoryId);
+                Factory->DestroyedFromReplication(DestroyInstanceParams);
+            }
+		}
 	}
+	
+	UnregisterInstance(Handle);	
 
 	// $IRIS TODO: Cleanup any pending creation data if we have not yet instantiated the instance.
-
-	Super::DetachInstanceFromRemote(Handle, DestroyReason, DestroyFlags);
 }
 
 void UObjectReplicationBridge::DetachInstance(FNetRefHandle RefHandle)
@@ -958,9 +964,15 @@ FReplicationBridgeCreateNetRefHandleResult UObjectReplicationBridge::CreateNetRe
 	return CreateResult;
 }
 
-void UObjectReplicationBridge::SubObjectCreatedFromReplication(FNetRefHandle Handle)
+void UObjectReplicationBridge::SubObjectCreatedFromReplication(UE::Net::Private::FInternalNetRefIndex RootObjectIndex, FNetRefHandle SubObjectCreated)
 {
-	OnSubObjectCreatedFromReplication(Handle);
+	using namespace UE::Net;
+	using namespace UE::Net::Private;
+
+	const FNetObjectFactoryId FactoryId = NetRefHandleManager->GetReplicatedObjectDataNoCheck(RootObjectIndex).NetFactoryId;
+	UNetObjectFactory* Factory = GetNetFactory(FactoryId);
+
+	Factory->SubObjectCreatedFromReplication(NetRefHandleManager->GetNetRefHandleFromInternalIndex(RootObjectIndex), SubObjectCreated);
 }
 
 void UObjectReplicationBridge::PostApplyInitialState(UE::Net::Private::FInternalNetRefIndex InternalObjectIndex)
diff --git a/Engine/Source/Runtime/Experimental/Iris/Core/Private/Iris/ReplicationSystem/ReplicationBridge.cpp b/Engine/Source/Runtime/Experimental/Iris/Core/Private/Iris/ReplicationSystem/ReplicationBridge.cpp
--- a/Engine/Source/Runtime/Experimental/Iris/Core/Private/Iris/ReplicationSystem/ReplicationBridge.cpp
+++ b/Engine/Source/Runtime/Experimental/Iris/Core/Private/Iris/ReplicationSystem/ReplicationBridge.cpp
@@ -100,11 +100,11 @@ void UReplicationBridge::UpdateInstancesWorldLocation()
 {
 }
 
-void UReplicationBridge::SubObjectCreatedFromReplication(FNetRefHandle SubObjectRefHandle)
+void UReplicationBridge::SubObjectCreatedFromReplication(UE::Net::Private::FInternalNetRefIndex RootObjectIndex, FNetRefHandle SubObjectRefHandle)
 {
 }
 
-void UReplicationBridge::DetachInstanceFromRemote(FNetRefHandle Handle, EReplicationBridgeDestroyInstanceReason DestroyReason, EReplicationBridgeDestroyInstanceFlags DestroyFlags)
+void UReplicationBridge::DetachInstanceFromRemote(FNetRefHandle Handle, EReplicationBridgeDestroyInstanceReason DestroyReason, EReplicationBridgeDestroyInstanceFlags DestroyFlags,  UE::Net::FNetObjectFactoryId NetFactoryId)
 {
 }
 
@@ -165,7 +165,7 @@ void UReplicationBridge::DetachSubObjectInstancesFromRemote(FNetRefHandle OwnerH
 				if (bAlwaysDestroyDynamicSubObjectInstancesOnDetachFromStaticRoot)
 				{
 					EnumAddFlags(SubObjectDestroyFlags, EReplicationBridgeDestroyInstanceFlags::AllowDestroyInstanceFromRemote);
-					CallDetachInstanceFromRemote(SubObjectHandle, EReplicationBridgeDestroyInstanceReason::Destroy, SubObjectDestroyFlags);
+					CallDetachInstanceFromRemote(SubObjectHandle, EReplicationBridgeDestroyInstanceReason::Destroy, SubObjectDestroyFlags, SubObjectData.NetFactoryId);
 					continue;
 				}
 				else
@@ -174,7 +174,7 @@ void UReplicationBridge::DetachSubObjectInstancesFromRemote(FNetRefHandle OwnerH
 				}
 			}
 
-			CallDetachInstanceFromRemote(SubObjectHandle, DestroyReason, SubObjectDestroyFlags);
+			CallDetachInstanceFromRemote(SubObjectHandle, DestroyReason, SubObjectDestroyFlags, SubObjectData.NetFactoryId);
 		}
 	}
 }
@@ -197,7 +197,7 @@ void UReplicationBridge::DestroyNetObjectFromRemote(FNetRefHandle Handle, ERepli
 		DetachSubObjectInstancesFromRemote(Handle, DestroyReason, DestroyFlags);
 		
 		// Allow derived bridges to cleanup any instance info they have stored
-		CallDetachInstanceFromRemote(Handle, DestroyReason, DestroyFlags);
+		CallDetachInstanceFromRemote(Handle, DestroyReason, DestroyFlags, ObjectData.NetFactoryId);
 
 		// Detach instance protocol
 		InternalDetachInstanceFromNetRefHandle(Handle);
@@ -230,9 +230,9 @@ void UReplicationBridge::CallDetachInstance(FNetRefHandle Handle)
 	DetachInstance(Handle);
 }
 
-void UReplicationBridge::CallDetachInstanceFromRemote(FNetRefHandle Handle, EReplicationBridgeDestroyInstanceReason DestroyReason, EReplicationBridgeDestroyInstanceFlags DestroyFlags)
+void UReplicationBridge::CallDetachInstanceFromRemote(FNetRefHandle Handle, EReplicationBridgeDestroyInstanceReason DestroyReason, EReplicationBridgeDestroyInstanceFlags DestroyFlags, UE::Net::FNetObjectFactoryId NetFactoryId)
 {
-	DetachInstanceFromRemote(Handle, DestroyReason, DestroyFlags);
+	DetachInstanceFromRemote(Handle, DestroyReason, DestroyFlags, NetFactoryId);
 }
 
 void UReplicationBridge::CallPruneStaleObjects()
@@ -266,12 +266,12 @@ bool UReplicationBridge::CallWriteNetRefHandleDestructionInfo(FReplicationBridge
 	{
 		// Write destruction info
 		WriteFullNetObjectReference(Context.SerializationContext, Info->StaticRef);
+		Context.SerializationContext.GetBitStreamWriter()->WriteBits(Info->NetFactoryId, FNetObjectFactoryRegistry::GetMaxBits());
 	}
 	else
 	{
+		Context.SerializationContext.SetError(TEXT("DestructionInfoNotFound"), false);
 		UE_LOG_REPLICATIONBRIDGE(Error, TEXT("Failed to write destructionInfo for %s"), *Handle.ToString());
-		// Write invalid reference
-		WriteFullNetObjectReference(Context.SerializationContext, FNetObjectReference());
 	}
 
 	return !Context.SerializationContext.HasErrorOrOverflow();
@@ -310,6 +310,8 @@ void UReplicationBridge::ReadAndExecuteDestructionInfoFromRemote(FReplicationBri
 	FNetObjectReference ReferenceToDestroy;
 	ReadFullNetObjectReference(Context.SerializationContext, ReferenceToDestroy);
 
+	const FNetObjectFactoryId NetFactoryId = IntCastChecked<FNetObjectFactoryId>(Context.SerializationContext.GetBitStreamReader()->ReadBits(FNetObjectFactoryRegistry::GetMaxBits()));
+
 	// Destroy the reference
 	// Resolve the reference in order to be able to destroy it
 	if (const UObject* Instance = ObjectReferenceCache->ResolveObjectReference(ReferenceToDestroy, Context.SerializationContext.GetInternalContext()->ResolveContext))
@@ -327,13 +329,13 @@ void UReplicationBridge::ReadAndExecuteDestructionInfoFromRemote(FReplicationBri
 		// We keep tracking them internally
 		DetachSubObjectInstancesFromRemote(ReferenceToDestroy.GetRefHandle(), DestroyReason, DestroyFlags);
 
-		CallDetachInstanceFromRemote(ReferenceToDestroy.GetRefHandle(), DestroyReason, DestroyFlags);
+		CallDetachInstanceFromRemote(ReferenceToDestroy.GetRefHandle(), DestroyReason, DestroyFlags, NetFactoryId);
 	}
 }
 
-void UReplicationBridge::CallSubObjectCreatedFromReplication(FNetRefHandle SubObjectHandle)
+void UReplicationBridge::CallSubObjectCreatedFromReplication(UE::Net::Private::FInternalNetRefIndex RootObjectIndex, FNetRefHandle SubObjectCreated)
 {
-	SubObjectCreatedFromReplication(SubObjectHandle);
+	SubObjectCreatedFromReplication(RootObjectIndex, SubObjectCreated);
 }
 
 void UReplicationBridge::CallPostApplyInitialState(UE::Net::Private::FInternalNetRefIndex InternalObjectIndex)
@@ -979,6 +981,9 @@ UE::Net::FNetRefHandle UReplicationBridge::StoreDestructionInfo(FNetRefHandle Ha
 
 	FReplicationProtocolManager* ProtocolManager = GetReplicationProtocolManager();
 
+	// $IRIS TODO: NetFactoryID can be invalid when static objects have not been replicated yet and assigned a NetFactory via BeginReplication.  Need to add an official way for the netFactory be given by the user code in the NetDriver.
+	const FNetObjectFactoryId NetFactoryId = IntCastChecked<FNetObjectFactoryId>(NetRefHandleManager->GetReplicatedObjectData(NetRefHandleManager->GetInternalIndex(Handle)).NetFactoryId);
+
 	// Create handle for DestructionInfo to be able to scope destruction infos
 	const FNetObjectGroupHandle LevelGroupHandle = GetLevelGroup(Parameters.Level);
 	FNetRefHandle DestructionInfoHandle = NetRefHandleManager->CreateHandleForDestructionInfo(Handle, DestructionInfoProtocol);
@@ -991,10 +996,13 @@ UE::Net::FNetRefHandle UReplicationBridge::StoreDestructionInfo(FNetRefHandle Ha
 	const FInternalNetRefIndex InternalReplicationIndex = NetRefHandleManager->GetInternalIndex(DestructionInfoHandle);
 
 	// We also need to store the actual data we send to destroy static references when they are scoped
-	FDestructionInfo PersistentDestructionInfo;
-	PersistentDestructionInfo.LevelGroupHandle = LevelGroupHandle;
-	PersistentDestructionInfo.StaticRef = FObjectReferenceCache::MakeNetObjectReference(Handle);
-	PersistentDestructionInfo.InternalReplicationIndex = InternalReplicationIndex;
+	FDestructionInfo PersistentDestructionInfo
+	{
+		.StaticRef = FObjectReferenceCache::MakeNetObjectReference(Handle),
+		.LevelGroupHandle = LevelGroupHandle,
+		.NetFactoryId = NetFactoryId,
+		.InternalReplicationIndex = InternalReplicationIndex
+	};
 
 	StaticObjectsPendingDestroy.Add(DestructionInfoHandle, PersistentDestructionInfo);
 
diff --git a/Engine/Source/Runtime/Experimental/Iris/Core/Private/Iris/ReplicationSystem/ReplicationReader.cpp b/Engine/Source/Runtime/Experimental/Iris/Core/Private/Iris/ReplicationSystem/ReplicationReader.cpp
--- a/Engine/Source/Runtime/Experimental/Iris/Core/Private/Iris/ReplicationSystem/ReplicationReader.cpp
+++ b/Engine/Source/Runtime/Experimental/Iris/Core/Private/Iris/ReplicationSystem/ReplicationReader.cpp
@@ -1987,7 +1987,7 @@ void FReplicationReader::DispatchStateData(FNetSerializationContext& Context)
 		{
 			if (ObjectData.SubObjectRootIndex != FNetRefHandleManager::InvalidInternalIndex)
 			{
-				ReplicationBridge->CallSubObjectCreatedFromReplication(ObjectData.RefHandle);
+				ReplicationBridge->CallSubObjectCreatedFromReplication(ObjectData.SubObjectRootIndex, ObjectData.RefHandle);
 			}
 		}
 
diff --git a/Engine/Source/Runtime/Experimental/Iris/Core/Public/Iris/ReplicationSystem/NetObjectFactory.h b/Engine/Source/Runtime/Experimental/Iris/Core/Public/Iris/ReplicationSystem/NetObjectFactory.h
--- a/Engine/Source/Runtime/Experimental/Iris/Core/Public/Iris/ReplicationSystem/NetObjectFactory.h
+++ b/Engine/Source/Runtime/Experimental/Iris/Core/Public/Iris/ReplicationSystem/NetObjectFactory.h
@@ -85,6 +85,7 @@ public:
 	struct FInstantiateContext;
 	struct FPostInstantiationContext;
 	struct FPostInitContext;
+	struct FDestroyedContext;
 
 	/**
 	 * Creates the header containing all information required to instantiate a remote version of the object represented by the handle.
@@ -132,6 +133,22 @@ public:
 	 */
 	virtual void PostInit(const FPostInitContext& Context) {}
 
+	/** 
+	 * Optional callback triggered when a root object managed by this factory gets assigned a dynamic subobject. 
+	 * This callback is called on remotes only. 
+	 * At the time of the callback the RootObject will have the latest replicated properties set, but the subobject will only be default constructed and won't be assigned the replicated properties received alongside the creation request. 
+	 * @param RootObject The root object that owns the subobject
+	 * @param SubObjectCreated The subobject that was just instantiated
+	*/
+
+	virtual void SubObjectCreatedFromReplication(UE::Net::FNetRefHandle RootObject, UE::Net::FNetRefHandle SubObjectCreated) {}
+
+	/** 
+	 * Optional callback triggered when an remote object is about to be destroyed.
+	 * @param Context Gives access to the instance that is about to be destroyed and details on what type of destruction is requested.
+	*/
+	virtual void DestroyedFromReplication(const FDestroyedContext& Context) {}
+
 public:
 
 	/** Result of the instantiate request */
@@ -181,6 +198,17 @@ public:
 		UE::Net::FNetRefHandle Handle;
 	};
 
+	/** Contextual information for the DestroyedFromReplication callback*/
+	struct FDestroyedContext
+	{
+		/** The object about to be destroyed */
+		UObject* DestroyedInstance = nullptr;
+		/** Optional pointer to the root object when the destroyed object is a subobject */
+		UObject* RootObject = nullptr;
+		EReplicationBridgeDestroyInstanceReason DestroyReason = EReplicationBridgeDestroyInstanceReason::DoNotDestroy;
+		EReplicationBridgeDestroyInstanceFlags DestroyFlags = EReplicationBridgeDestroyInstanceFlags::None;
+	};
+
 protected:
 
 	virtual void OnInit() {}
diff --git a/Engine/Source/Runtime/Experimental/Iris/Core/Public/Iris/ReplicationSystem/ObjectReplicationBridge.h b/Engine/Source/Runtime/Experimental/Iris/Core/Public/Iris/ReplicationSystem/ObjectReplicationBridge.h
--- a/Engine/Source/Runtime/Experimental/Iris/Core/Public/Iris/ReplicationSystem/ObjectReplicationBridge.h
+++ b/Engine/Source/Runtime/Experimental/Iris/Core/Public/Iris/ReplicationSystem/ObjectReplicationBridge.h
@@ -242,36 +242,23 @@ protected:
 	IRISCORE_API virtual bool WriteNetRefHandleCreationInfo(FReplicationBridgeSerializationContext& Context, FNetRefHandle Handle) override;
 	IRISCORE_API virtual bool CacheNetRefHandleCreationInfo(FNetRefHandle Handle) override;
 	IRISCORE_API virtual FReplicationBridgeCreateNetRefHandleResult CreateNetRefHandleFromRemote(FNetRefHandle RootObjectNetHandle, FNetRefHandle WantedNetHandle, FReplicationBridgeSerializationContext& Context) override;
-	IRISCORE_API virtual void SubObjectCreatedFromReplication(FNetRefHandle SubObjectHandle) override;
+	IRISCORE_API virtual void SubObjectCreatedFromReplication(UE::Net::Private::FInternalNetRefIndex RootObjectIndex, FNetRefHandle SubObjectCreated) override;
 	IRISCORE_API virtual void PostApplyInitialState(UE::Net::Private::FInternalNetRefIndex InternalObjectIndex) final;
-	IRISCORE_API virtual void DetachInstanceFromRemote(FNetRefHandle Handle, EReplicationBridgeDestroyInstanceReason DestroyReason, EReplicationBridgeDestroyInstanceFlags DestroyFlags) override;
+	IRISCORE_API virtual void DetachInstanceFromRemote(FNetRefHandle Handle, EReplicationBridgeDestroyInstanceReason DestroyReason, EReplicationBridgeDestroyInstanceFlags DestroyFlags,  UE::Net::FNetObjectFactoryId NetFactoryId) override;
 	IRISCORE_API virtual void DetachInstance(FNetRefHandle Handle) override;
 	IRISCORE_API virtual void OnProtocolMismatchReported(FNetRefHandle RefHandle, uint32 ConnectionId) override;
 	IRISCORE_API virtual void OnErrorWithNetRefHandleReported(UE::Net::ENetRefHandleError ErrorType, FNetRefHandle RefHandle, uint32 ConnectionId) override;
 	
 protected:
 	
-	/** Invoked for new replicated SubObjects after state has been applied to owner */
-	virtual void OnSubObjectCreatedFromReplication(FNetRefHandle SubObjectHandle) {};
-
-	struct FDestroyInstanceParams
-	{
-		UObject* Instance = nullptr;
-		UObject* RootObject = nullptr;
-		EReplicationBridgeDestroyInstanceReason DestroyReason = EReplicationBridgeDestroyInstanceReason::DoNotDestroy;
-		EReplicationBridgeDestroyInstanceFlags DestroyFlags = EReplicationBridgeDestroyInstanceFlags::None;
-	};
-
-	/** Destroy or tear-off the game instance on request from remote. */
-	virtual void DestroyInstanceFromRemote(const FDestroyInstanceParams& Params) {}
-
 	/** Called when we found a divergence between the local and remote protocols when trying to instantiate a remote replicated object. */
 	virtual void OnProtocolMismatchDetected(FNetRefHandle ObjectHandle) {}
 
-protected:
 	/** Lookup the UObject associated with the provided Handle. This function will not try to resolve the reference. */
 	IRISCORE_API UObject* GetObjectFromReferenceHandle(FNetRefHandle RefHandle) const;
 
+	IRISCORE_API UNetObjectFactory* GetNetFactory(UE::Net::FNetObjectFactoryId FactoryId) const;
+
 	/** Helper method that calls provided PreUpdateFunction and polls state data for all replicated instances with the NeedsPoll trait. */
 	using FInstancePreUpdateFunction = TFunction<void(TArrayView<UObject*>, const UReplicationBridge*)>;
 
@@ -411,8 +398,6 @@ private:
 
 	void OnMaxInternalNetRefIndexIncreased(UE::Net::Private::FInternalNetRefIndex NewMaxInternalIndex);
 
-	UNetObjectFactory* GetNetFactory(UE::Net::FNetObjectFactoryId FactoryId) const;
-
 private:
 
 	friend UE::Net::Private::FObjectPoller;
diff --git a/Engine/Source/Runtime/Experimental/Iris/Core/Public/Iris/ReplicationSystem/ReplicationBridge.h b/Engine/Source/Runtime/Experimental/Iris/Core/Public/Iris/ReplicationSystem/ReplicationBridge.h
--- a/Engine/Source/Runtime/Experimental/Iris/Core/Public/Iris/ReplicationSystem/ReplicationBridge.h
+++ b/Engine/Source/Runtime/Experimental/Iris/Core/Public/Iris/ReplicationSystem/ReplicationBridge.h
@@ -7,6 +7,7 @@
 #include "Containers/Map.h"
 #include "Iris/Core/NetObjectReference.h"
 #include "Iris/ReplicationSystem/NetRefHandle.h"
+#include "Iris/ReplicationSystem/NetObjectFactory.h"
 #include "Iris/ReplicationSystem/NetObjectGroupHandle.h"
 #include "Iris/ReplicationSystem/ReplicationBridgeTypes.h"
 #include "Misc/EnumClassFlags.h"
@@ -165,8 +166,8 @@ protected:
 	/** Read data required to instantiate NetObject from bitstream. */
 	IRISCORE_API virtual FReplicationBridgeCreateNetRefHandleResult CreateNetRefHandleFromRemote(FNetRefHandle RootObjectOfSubObject, FNetRefHandle WantedNetHandle, FReplicationBridgeSerializationContext& Context);
 
-	/** Invoked right before we apply the state for a new received subobject but after we have applied state for owning/root object in order to behave like old replication system */
-	IRISCORE_API virtual void SubObjectCreatedFromReplication(FNetRefHandle SubObjectRefHandle);
+	/** Invoked right before we apply the state for a new received subobject but after we have applied state on the root object. */
+ 	IRISCORE_API virtual void SubObjectCreatedFromReplication(UE::Net::Private::FInternalNetRefIndex RootObjectIndex,FNetRefHandle SubObjectCreated);
 
 	/** Invoked after we have applied the initial state for an object.*/
 	virtual void PostApplyInitialState(UE::Net::Private::FInternalNetRefIndex InternalObjectIndex) {}
@@ -177,7 +178,7 @@ protected:
 	 * @param DestroyReason Reason for destroying the instance. 
 	 * @param DestroyFlags Special flags such as whether the instance may be destroyed when reason is TearOff, which may revert to destroying, or Destroy.
 	 */
-	IRISCORE_API virtual void DetachInstanceFromRemote(FNetRefHandle Handle, EReplicationBridgeDestroyInstanceReason DestroyReason, EReplicationBridgeDestroyInstanceFlags DestroyFlags);
+	IRISCORE_API virtual void DetachInstanceFromRemote(FNetRefHandle Handle, EReplicationBridgeDestroyInstanceReason DestroyReason, EReplicationBridgeDestroyInstanceFlags DestroyFlags,  UE::Net::FNetObjectFactoryId NetFactoryId);
 
 	/** Called when we detach instance protocol from the local instance */
 	IRISCORE_API virtual void DetachInstance(FNetRefHandle Handle);
@@ -259,12 +260,12 @@ private:
 	bool CallCacheNetRefHandleCreationInfo(FNetRefHandle Handle);
 	bool CallWriteNetRefHandleCreationInfo(FReplicationBridgeSerializationContext& Context, FNetRefHandle Handle);
 	bool CallWriteNetRefHandleDestructionInfo(FReplicationBridgeSerializationContext& Context, FNetRefHandle Handle);
-	void CallSubObjectCreatedFromReplication(FNetRefHandle SubObjectHandle);
+	void CallSubObjectCreatedFromReplication(UE::Net::Private::FInternalNetRefIndex RootObjectIndex, FNetRefHandle SubObjectCreated);
 	void CallPostApplyInitialState(UE::Net::Private::FInternalNetRefIndex InternalObjectIndex);
 	void CallPruneStaleObjects();
 	void CallGetInitialDependencies(FNetRefHandle Handle, FNetDependencyInfoArray& OutDependencies) const;
 	void CallDetachInstance(FNetRefHandle Handle);
-	void CallDetachInstanceFromRemote(FNetRefHandle Handle, EReplicationBridgeDestroyInstanceReason DestroyReason, EReplicationBridgeDestroyInstanceFlags DestroyFlags);
+	void CallDetachInstanceFromRemote(FNetRefHandle Handle, EReplicationBridgeDestroyInstanceReason DestroyReason, EReplicationBridgeDestroyInstanceFlags DestroyFlags, UE::Net::FNetObjectFactoryId NetFactoryId);
 
 private:
 
@@ -327,6 +328,7 @@ private:
 	{
 		UE::Net::FNetObjectReference StaticRef;
 		UE::Net::FNetObjectGroupHandle LevelGroupHandle;
+		UE::Net::FNetObjectFactoryId NetFactoryId = UE::Net::InvalidNetObjectFactoryId;
 		UE::Net::Private::FInternalNetRefIndex InternalReplicationIndex;
 	};
 
