Added Mount and Install console commands for I/O store on-demand.

#rnx

diff --git a/Engine/Source/Runtime/Experimental/IoStore/OnDemand/Private/IoStoreOnDemand.cpp b/Engine/Source/Runtime/Experimental/IoStore/OnDemand/Private/IoStoreOnDemand.cpp
--- a/Engine/Source/Runtime/Experimental/IoStore/OnDemand/Private/IoStoreOnDemand.cpp
+++ b/Engine/Source/Runtime/Experimental/IoStore/OnDemand/Private/IoStoreOnDemand.cpp
@@ -31,6 +31,9 @@
 #include "Serialization/MemoryReader.h"
 #include "Statistics.h"
 #include "String/LexFromString.h"
+#if !UE_BUILD_SHIPPING
+#include "Misc/PackageName.h"
+#endif //!UE_BUILD_SHIPPING
 
 #if PLATFORM_WINDOWS
 #	include <Windows/AllowWindowsPlatformTypes.h>
@@ -45,61 +48,309 @@ namespace UE::IoStore
 {
 
 ///////////////////////////////////////////////////////////////////////////////
-static FAutoConsoleCommand OnDemandPurgeCacheCommand(
-	TEXT("iostore.PurgeOnDemandInstallCache"),
-	TEXT("Purge On Demand Install Cache"),
-	FConsoleCommandDelegate::CreateStatic([]()
+#if !UE_BUILD_SHIPPING
+namespace Commands
+{
+
+////////////////////////////////////////////////////////////////////////////////
+static TIoStatusOr<FOnDemandHostGroup> SplitHostAndTocUrl(FStringView Url, FStringView& OutTocRelativUrl)
+{
+	if (Url.StartsWith(TEXTVIEW("http")) == false)
+	{
+		return FIoStatus(EIoErrorCode::InvalidParameter, TEXT("Invalid URL protocol"));
+	}
+
+	int32 Delim = INDEX_NONE;
+	if (Url.FindChar(':', Delim) == false)
+	{
+		return FIoStatus(EIoErrorCode::InvalidParameter, TEXT("Invalid URL"));
+	}
+
+	const int32 ProtocolDelim = Delim + 3;
+	if (Url.RightChop(ProtocolDelim).FindChar('/', Delim) == false)
+	{
+		return FIoStatus(EIoErrorCode::InvalidParameter, TEXT("Failed to find host and TOC path delimiter"));
+	}
+
+	const FStringView Host	= Url.Left(ProtocolDelim + Delim);
+	OutTocRelativUrl		= Url.RightChop(Host.Len());
+
+	return FOnDemandHostGroup::Create(Host);
+}
+
+////////////////////////////////////////////////////////////////////////////////
+static void PurgeInstallCache()
+{
+	FIoStoreOnDemandModule* IOStoreOnDemandModule = FModuleManager::Get().GetModulePtr<FIoStoreOnDemandModule>(TEXT("IoStoreOnDemand"));
+	if(!IOStoreOnDemandModule)
 	{
-		FIoStoreOnDemandModule* IOStoreOnDemandModule = FModuleManager::Get().GetModulePtr<FIoStoreOnDemandModule>(TEXT("IoStoreOnDemand"));
-		if(!IOStoreOnDemandModule)
+		UE_LOG(LogIoStoreOnDemand, Error, TEXT("Could not find IoStoreOnDemand module"));
+		return;
+	}
+
+	UE_LOG(LogIoStoreOnDemand, Display, TEXT("Purging on demand install cache"));
+	IOStoreOnDemandModule->Purge(FOnDemandPurgeArgs(), [](const FOnDemandPurgeResult& Result)
+	{
+		if (Result.Status.IsOk())
 		{
-			UE_LOG(LogIoStoreOnDemand, Error, TEXT("Could not find IoStoreOnDemand module"));
-			return;
+			UE_LOG(LogIoStoreOnDemand, Display, TEXT("Purged on demand install cache"));
+		}
+		else
+		{
+			UE_LOG(LogIoStoreOnDemand, Error, TEXT("Failed Purged on demand install cache: %s"), *Result.Status.ToString());
 		}
+	});
+}
+
+////////////////////////////////////////////////////////////////////////////////
+static void PrintCacheUsage()
+{
+	FIoStoreOnDemandModule* IOStoreOnDemandModule = FModuleManager::Get().GetModulePtr<FIoStoreOnDemandModule>(TEXT("IoStoreOnDemand"));
+	if (!IOStoreOnDemandModule)
+	{
+		UE_LOG(LogIoStoreOnDemand, Error, TEXT("Could not find IoStoreOnDemand module"));
+		return;
+	}
+
+	TIoStatusOr<FOnDemandCacheUsage> MaybeUsage = IOStoreOnDemandModule->GetCacheUsage();
+	if (!MaybeUsage.IsOk())
+	{
+		UE_LOG(LogIoStoreOnDemand, Error, TEXT("iostore.CacheUsage failed: %s"), *MaybeUsage.Status().ToString());
+		return;
+	}
+
+	const FOnDemandCacheUsage CacheUsage = MaybeUsage.ConsumeValueOrDie();
+	UE_LOG(LogIoStoreOnDemand, Display, TEXT("iostore.CacheUsage"));
+	TStringBuilder<512> Sb;
+	Sb << TEXTVIEW("InstallCache: ") << CacheUsage.InstallCache << TEXTVIEW(", StreamingCache: ") << CacheUsage.StreamingCache;
+	UE_LOG(LogIoStoreOnDemand, Display, TEXT("%s"), Sb.ToString());
+}
+
+////////////////////////////////////////////////////////////////////////////////
+static void MountUrl(const TArray<FString>& Args)
+{
+	FIoStoreOnDemandModule* IoStoreOnDemand = FModuleManager::Get().GetModulePtr<FIoStoreOnDemandModule>(TEXT("IoStoreOnDemand"));
+	if (IoStoreOnDemand == nullptr)
+	{
+		UE_LOG(LogIoStoreOnDemand, Error, TEXT("Failed to load I/O store on-demand module."));
+		return;
+	}
+
+	if (Args.Num() < 2)
+	{
+		UE_LOG(LogIoStoreOnDemand, Error, TEXT("Not enough arguments."));
+		return;
+	}
+
+	FString Url = Args[0];
+	Url.TrimQuotesInline();
+
+	FStringView TocRelativeUrl;
+	TIoStatusOr<FOnDemandHostGroup> HostGroup = SplitHostAndTocUrl(Url, TocRelativeUrl);
+	if (HostGroup.IsOk() == false)
+	{
+		UE_LOG(LogIoStoreOnDemand, Error, TEXT("%s"), *HostGroup.Status().ToString());
+		return;
+	}
+
+	EOnDemandMountOptions Options = EOnDemandMountOptions::StreamOnDemand;
+	const FString& Arg = Args[1];
+	if (Arg.ToLower().Contains(TEXT("install")))
+	{
+		Options = EOnDemandMountOptions::InstallOnDemand; 
+	}
+	const FString& MountId = Args.Num() > 2 ? Args[2] : Url; 
+
+	IoStoreOnDemand->Mount(
+		FOnDemandMountArgs
+		{
+			.MountId = MountId,
+			.TocRelativeUrl = FString(TocRelativeUrl),
+			.HostGroup = HostGroup.ConsumeValueOrDie(),
+			.Options = Options
+		},
+		[](FOnDemandMountResult MountResult)
+		{
+			UE_CLOG(!MountResult.Status.IsOk(), LogIas, Error,
+				TEXT("Failed to mount TOC for '%s', reason '%s'"), *MountResult.MountId, *MountResult.Status.ToString());
+		});
+}
+
+////////////////////////////////////////////////////////////////////////////////
+static void MountFile(const TArray<FString>& Args)
+{
+	FIoStoreOnDemandModule* IoStoreOnDemand = FModuleManager::Get().GetModulePtr<FIoStoreOnDemandModule>(TEXT("IoStoreOnDemand"));
+	if (IoStoreOnDemand == nullptr)
+	{
+		UE_LOG(LogIoStoreOnDemand, Error, TEXT("Failed to load I/O store on-demand module."));
+		return;
+	}
+
+	if (Args.Num() < 3)
+	{
+		UE_LOG(LogIoStoreOnDemand, Error, TEXT("Not enough arguments."));
+		return;
+	}
+
+	TArray<FOnDemandMountArgs> AllMountArgs;
+	FString FilenameOrWildcard = Args[0];
+	FilenameOrWildcard.TrimQuotesInline();
+	TArray<FString> FilePaths;
+
+	const FString ContentDir = FString::Printf(TEXT("%sPaks/"), *FPaths::ProjectContentDir());
+	IFileManager& Ifm = IFileManager::Get();
+	Ifm.IterateDirectoryRecursively(
+		*ContentDir,
+		[&FilenameOrWildcard, &FilePaths, &ContentDir](const TCHAR* FilenameOrDirectory, bool bIsDirectory)
+		{
+			if (bIsDirectory == false)
+			{
+				FString FilePath(FilenameOrDirectory);
+				if (FPathViews::GetExtension(FilePath) == TEXT("uondemandtoc"))
+				{
+					FStringView Filename = FPathViews::GetBaseFilename(FilePath);
+					if (Filename == FilenameOrWildcard || FString(Filename).MatchesWildcard(FilenameOrWildcard))
+					{
+						FilePaths.Add(FilenameOrDirectory);
+					}
+				}
+			}
+			return true;
+		});
+
+	if (FilePaths.IsEmpty())
+	{
+		UE_LOG(LogIoStoreOnDemand, Error, TEXT("Failed to find any on-demand TOC file(s) matching '%s'"), *FilenameOrWildcard);
+		return;
+	}
+
+	EOnDemandMountOptions Options = EOnDemandMountOptions::StreamOnDemand;
+	if (Args[1].ToLower().Contains(TEXT("install")))
+	{
+		Options = EOnDemandMountOptions::InstallOnDemand; 
+	}
+
+	FString HostPath = Args[2];
+	HostPath.TrimQuotesInline();
+	FStringView TocRelativeUrl;
+	TIoStatusOr<FOnDemandHostGroup> HostGroup = SplitHostAndTocUrl(HostPath, TocRelativeUrl);
+	if (HostGroup.IsOk() == false)
+	{
+		UE_LOG(LogIoStoreOnDemand, Error, TEXT("%s"), *HostGroup.Status().ToString());
+		return;
+	}
+
+	const FString& MountId = Args.Num() > 3 ? Args[3] : FilePaths[0];
+	for (const FString& FilePath : FilePaths)
+	{
+		FOnDemandMountArgs& MountArgs	= AllMountArgs.AddDefaulted_GetRef();
+		MountArgs.FilePath				= FilePath;
+		MountArgs.HostGroup				= HostGroup.ConsumeValueOrDie();
+		MountArgs.TocRelativeUrl		= TocRelativeUrl;
+		MountArgs.MountId				= MountId;
+	}
+
+	if (AllMountArgs.IsEmpty())
+	{
+		UE_LOG(LogIoStoreOnDemand, Error, TEXT("Failed to parse mount arguments"));
+		return;
+	}
+
+	for (FOnDemandMountArgs& MountArgs : AllMountArgs)
+	{
+		MountArgs.Options = Options;
+		IoStoreOnDemand->Mount(
+			MoveTemp(MountArgs),
+			[](FOnDemandMountResult MountResult)
+			{
+				UE_CLOG(!MountResult.Status.IsOk(), LogIas, Error,
+					TEXT("Failed to mount TOC for '%s', reason '%s'"), *MountResult.MountId, *MountResult.Status.ToString());
+			});
+	}
+}
+
+////////////////////////////////////////////////////////////////////////////////
+static void InstallPackage(const TArray<FString>& Args)
+{
+	if (Args.Num() < 1)
+	{
+		UE_LOG(LogIoStoreOnDemand, Error, TEXT("Not enough arguments."));
+		return;
+	}
+
+	FIoStoreOnDemandModule* IoStoreOnDemand = FModuleManager::Get().GetModulePtr<FIoStoreOnDemandModule>(TEXT("IoStoreOnDemand"));
+	if (IoStoreOnDemand == nullptr)
+	{
+		UE_LOG(LogIoStoreOnDemand, Error, TEXT("Failed to load I/O store on-demand module."));
+		return;
+	}
+
+	FString PackageName = Args[0];
+	PackageName.TrimQuotesInline();
+	FPackageId PackageId;
+	if (FPackageName::IsValidLongPackageName(PackageName))
+	{
+		PackageId = FPackageId::FromName(FName(*PackageName));
+	}
+
+	if (PackageId.IsValid() == false)
+	{
+		UE_LOG(LogIoStoreOnDemand, Error, TEXT("Invalid package name '%s'"), *PackageName);
+		return;
+	}
+
+	static FOnDemandContentHandle DefaultContentHandle = FOnDemandContentHandle::Create(TEXT("ConsoleCommand"));
+
+	FOnDemandInstallArgs InstallArgs;
+	InstallArgs.PackageIds.Add(PackageId);
+	InstallArgs.ContentHandle = DefaultContentHandle;
 
-		UE_LOG(LogIoStoreOnDemand, Display, TEXT("Purging on demand install cache"));
-		IOStoreOnDemandModule->Purge(FOnDemandPurgeArgs(), [](const FOnDemandPurgeResult& Result)
+	IoStoreOnDemand->Install(
+		MoveTemp(InstallArgs),
+		[PackageName = MoveTemp(PackageName)](FOnDemandInstallResult InstallResult)
 		{
-			if (Result.Status.IsOk())
+			if (InstallResult.Status.IsOk())
 			{
-				UE_LOG(LogIoStoreOnDemand, Display, TEXT("Purged on demand install cache"));
+				UE_LOG(LogIoStoreOnDemand, Log, TEXT("Successfully installed package '%s'"), *PackageName);
 			}
 			else
 			{
-				UE_LOG(LogIoStoreOnDemand, Error, TEXT("Failed Purged on demand install cache: %s"), *Result.Status.ToString());
+				UE_LOG(LogIoStoreOnDemand, Error, TEXT("Failed to install package '%s', reason '%s'"),
+					*PackageName, *InstallResult.Status.ToString());
 			}
 		});
-	}),
-	ECVF_Cheat
-);
+}
 
-static FAutoConsoleCommand OnDemandCacheUsageCommand(
+////////////////////////////////////////////////////////////////////////////////
+static FAutoConsoleCommand PurgeCacheCommand(
+	TEXT("iostore.PurgeInstallCache"),
+	TEXT("Purge On Demand Install Cache"),
+	FConsoleCommandDelegate::CreateStatic(PurgeInstallCache),
+	ECVF_Cheat);
+
+static FAutoConsoleCommand PrintCacheUsageCommand(
 	TEXT("iostore.CacheUsage"),
 	TEXT("print cache usage"),
-	FConsoleCommandDelegate::CreateStatic([]()
-	{
-		FIoStoreOnDemandModule* IOStoreOnDemandModule = FModuleManager::Get().GetModulePtr<FIoStoreOnDemandModule>(TEXT("IoStoreOnDemand"));
-		if (!IOStoreOnDemandModule)
-		{
-			UE_LOG(LogIoStoreOnDemand, Error, TEXT("Could not find IoStoreOnDemand module"));
-			return;
-		}
-
-		TIoStatusOr<FOnDemandCacheUsage> MaybeUsage = IOStoreOnDemandModule->GetCacheUsage();
-		if (!MaybeUsage.IsOk())
-		{
-			UE_LOG(LogIoStoreOnDemand, Error, TEXT("iostore.CacheUsage failed: %s"), *MaybeUsage.Status().ToString());
-			return;
-		}
-
-		const FOnDemandCacheUsage CacheUsage = MaybeUsage.ConsumeValueOrDie();
-		UE_LOG(LogIoStoreOnDemand, Display, TEXT("iostore.CacheUsage"));
-		TStringBuilder<512> Sb;
-		Sb << TEXTVIEW("InstallCache: ") << CacheUsage.InstallCache << TEXTVIEW(", StreamingCache: ") << CacheUsage.StreamingCache;
-		UE_LOG(LogIoStoreOnDemand, Display, TEXT("%s"), Sb.ToString());
-	}),
-	ECVF_Cheat
-);
+	FConsoleCommandDelegate::CreateStatic(PrintCacheUsage),
+	ECVF_Cheat);
+
+static FAutoConsoleCommand MountUrlCommand(
+	TEXT("iostore.MountUrl"),
+	TEXT("<URL> <Install|Stream> <MountId>"),
+	FConsoleCommandWithArgsDelegate::CreateStatic(MountUrl));
+
+static FAutoConsoleCommand MountFileCommand(
+	TEXT("iostore.MountFile"),
+	TEXT("<Filename|Wildcard> <Install|Stream> <HostPath> <MountId>"),
+	FConsoleCommandWithArgsDelegate::CreateStatic(MountFile));
+
+static FAutoConsoleCommand InstallPackageCommand(
+	TEXT("iostore.InstallPackage"),
+	TEXT("<PackageName>"),
+	FConsoleCommandWithArgsDelegate::CreateStatic(InstallPackage));
+
+} // namespace Commands
+#endif // !UE_BUILD_SHIPPING
 
 ////////////////////////////////////////////////////////////////////////////////
 FString GIasOnDemandTocExt = TEXT(".uondemandtoc");
