Fix transient snapshotting

[RN] JsonObjectGraph::Stringify now properly encodes subobjects of objects in the transient package

diff --git a/Engine/Source/Runtime/Experimental/JsonObjectGraph/Private/JsonStringifyImpl.cpp b/Engine/Source/Runtime/Experimental/JsonObjectGraph/Private/JsonStringifyImpl.cpp
--- a/Engine/Source/Runtime/Experimental/JsonObjectGraph/Private/JsonStringifyImpl.cpp
+++ b/Engine/Source/Runtime/Experimental/JsonObjectGraph/Private/JsonStringifyImpl.cpp
@@ -27,7 +27,6 @@ struct FPackageReferenceFinder : public FArchiveUObject
 	FPackageReferenceFinder(const UObject* Obj, TArray<const UObject*>& InReferences, bool bFilterEditorOnly)
 		: FArchiveUObject()
 		, References(InReferences)
-		, OwningPackage(Obj->GetOutermost())
 	{
 		// Copying FPackageHarvester:
 		SetIsPersistent(true);
@@ -50,8 +49,7 @@ private:
 	{
 		if (ObjRef != nullptr &&
 			(!ObjRef->HasAnyFlags(RF_Transient) || ObjRef->IsNative()) &&
-			!ObjRef->IsIn(GetTransientPackage()) &&
-			ObjRef != OwningPackage)
+			!ObjRef->IsA<UPackage>())
 		{
 			// Set to null any pointer to an external asset
 			References.Add(ObjRef);
@@ -61,7 +59,6 @@ private:
 	}
 
 	TArray<const UObject*>& References;
-	const UPackage* const OwningPackage;
 };
 
 static const UClass* GetFirstNativeClass(const UClass* Class)
@@ -125,18 +122,9 @@ static void GatherExports(
 	RefsProcessed.Append(PendingRefs);
 
 	TArray<const UObject*> ScratchRefs; // just keeping this allocation alive across iterations
-
-	// Roots is invalid after this point!
-
-	UObject* TransientPackage = GetTransientPackage();
 	while (PendingRefs.Num())
 	{
 		const UObject* Iter = PendingRefs.Pop();
-		if (Iter->GetOutermost() == TransientPackage)
-		{
-			continue;
-		}
-
 		if (AllowedRoots.Contains(Iter) && Iter->IsA<UPackage>())
 		{
 			// add all immediate roots of a package, if filter 'editor only' objects
