Load default HTTP certificates before resolving endpoints.

#rnx

diff --git a/Engine/Source/Runtime/Experimental/IoStore/OnDemand/Private/IoStoreOnDemand.cpp b/Engine/Source/Runtime/Experimental/IoStore/OnDemand/Private/IoStoreOnDemand.cpp
--- a/Engine/Source/Runtime/Experimental/IoStore/OnDemand/Private/IoStoreOnDemand.cpp
+++ b/Engine/Source/Runtime/Experimental/IoStore/OnDemand/Private/IoStoreOnDemand.cpp
@@ -1518,6 +1518,16 @@ void FIoStoreOnDemandModule::InitializeInternal()
 		UE_LOG(LogIas, Log, TEXT("File cache disabled - streaming only (zero-quota)"));
 	}
 
+	bool bFirstAttempt = false;
+	if (FIoStatus Status = LoadDefaultHttpCertificates(bFirstAttempt); !Status.IsOk())
+	{
+		if (EndpointConfig.RequiresTls())
+		{
+			UE_LOG(LogIas, Error, TEXT("Endpoint configuration requires TLS but certificates was not loaded, reason '%s'"),
+				*Status.ToString());
+			return;
+		}
+	}
 	HttpIoDispatcherBackend = MakeOnDemandIoDispatcherBackend(EndpointConfig, *IoStore, MoveTemp(Cache));
 
 	int32 BackendPriority = -10;
diff --git a/Engine/Source/Runtime/Experimental/IoStore/OnDemand/Private/OnDemandIoDispatcherBackend.h b/Engine/Source/Runtime/Experimental/IoStore/OnDemand/Private/OnDemandIoDispatcherBackend.h
--- a/Engine/Source/Runtime/Experimental/IoStore/OnDemand/Private/OnDemandIoDispatcherBackend.h
+++ b/Engine/Source/Runtime/Experimental/IoStore/OnDemand/Private/OnDemandIoDispatcherBackend.h
@@ -2,6 +2,7 @@
 
 #pragma once
 
+#include "Containers/StringView.h"
 #include "IO/IoDispatcherBackend.h"
 #include "Templates/UniquePtr.h"
 
@@ -49,6 +50,25 @@ struct FOnDemandEndpointConfig
 	{
 		return (DistributionUrl.Len() > 0 || ServiceUrls.Num() > 0) && TocPath.Len() > 0;
 	}
+
+	bool RequiresTls() const
+	{
+		auto IsHttps = [](FStringView Url) -> bool { return Url.StartsWith(TEXTVIEW("https"), ESearchCase::IgnoreCase); };
+
+		if (IsHttps(DistributionUrl) || IsHttps(FallbackUrl))
+		{
+			return true;
+		}
+		for (const FString& Url : ServiceUrls)
+		{
+			if (IsHttps(Url))
+			{
+				return true;
+			}
+		}
+
+		return false;
+	}
 };
 
 class IOnDemandIoDispatcherBackend
