Use UPS when gathering references for json stringification, this should be faster and also avoid unintended conflicts from layering delta serialization logic

#rnx

diff --git a/Engine/Source/Runtime/Experimental/JsonObjectGraph/Private/JsonStringifyImpl.cpp b/Engine/Source/Runtime/Experimental/JsonObjectGraph/Private/JsonStringifyImpl.cpp
--- a/Engine/Source/Runtime/Experimental/JsonObjectGraph/Private/JsonStringifyImpl.cpp
+++ b/Engine/Source/Runtime/Experimental/JsonObjectGraph/Private/JsonStringifyImpl.cpp
@@ -32,6 +32,7 @@ struct FPackageReferenceFinder : public FArchiveUObject
 		SetIsPersistent(true);
 		SetIsSaving(true);
 		SetFilterEditorOnly(bFilterEditorOnly);
+		SetUseUnversionedPropertySerialization(true);
 		ArIsObjectReferenceCollector = true;
 		ArShouldSkipBulkData = true;
 		if (Obj->HasAnyFlags(RF_ClassDefaultObject))
