IasTool - skip bucket name when downloading TOC from non S3 endpoint.

#rnx

diff --git a/Engine/Source/Runtime/Experimental/IoStore/OnDemand/Private/Tool/ListTocs.cpp b/Engine/Source/Runtime/Experimental/IoStore/OnDemand/Private/Tool/ListTocs.cpp
--- a/Engine/Source/Runtime/Experimental/IoStore/OnDemand/Private/Tool/ListTocs.cpp
+++ b/Engine/Source/Runtime/Experimental/IoStore/OnDemand/Private/Tool/ListTocs.cpp
@@ -81,6 +81,10 @@ struct FListTocsParams
 
 FIoStatus FListTocsParams::Validate() const
 {
+	if (!TocUrl.IsEmpty())
+	{
+		return FIoStatus::Ok;
+	}
 	if (Bucket.IsEmpty())
 	{
 		return FIoStatus(EIoErrorCode::InvalidParameter, TEXT("Invalid bucket name"));
@@ -432,8 +436,7 @@ FIoStatus ListTocs(const FListTocsParams& Params)
 		UE_LOG(LogIas, Display, TEXT("Fetching TOC '%s/%s'"), *Client.GetConfig().ServiceUrl, *Params.TocKey);
 		FS3GetObjectResponse TocResponse = Client.GetObject(FS3GetObjectRequest
 		{
-			TEXT("/"),
-			Params.TocKey
+			.Key = Params.TocKey
 		});
 
 		if (TocResponse.IsOk())
