[Chaos] Only calculate PseudoFraction for KinematicTarget when using substepping.
Fixes issue with non-substepped async physics where if the GameThread kicks off the PhysicsThread to tick multiple steps then KinematicTarget would get applied as if it was substepping, applying only a fraction of the movement.

diff --git a/Engine/Source/Runtime/Experimental/Chaos/Private/Chaos/Framework/PhysicsSolverBase.cpp b/Engine/Source/Runtime/Experimental/Chaos/Private/Chaos/Framework/PhysicsSolverBase.cpp
--- a/Engine/Source/Runtime/Experimental/Chaos/Private/Chaos/Framework/PhysicsSolverBase.cpp
+++ b/Engine/Source/Runtime/Experimental/Chaos/Private/Chaos/Framework/PhysicsSolverBase.cpp
@@ -166,9 +166,9 @@ namespace Chaos
 		FPhysicsThreadContextScope Scope(/*IsPhysicsThreadContext=*/true);
 #endif
 
-		// StepFraction: how much of the remaining time this step represents, used to interpolate kinematic targets
+		// StepFraction: how much of the remaining time this step represents while substepping, used to interpolate kinematic targets
 		// E.g., for 4 steps this will be: 1/4, 1/3, 1/2, 1
-		const FReal PseudoFraction = (FReal)1 / (FReal)(PushData->IntervalNumSteps - PushData->IntervalStep);
+		const FReal PseudoFraction = PushData->bSolverSubstepped ? (FReal)1 / (FReal)(PushData->IntervalNumSteps - PushData->IntervalStep) : 1.0f;
 
 		Solver.AdvanceSolverBy(FSubStepInfo{ PseudoFraction, PushData->IntervalStep, PushData->IntervalNumSteps, PushData->bSolverSubstepped });
 
