Improvements to ias.DisplayOnScreenStatistics.

#rnx

- Use FText::AsMemory to format byte values rather than always hard coding to KiB, this will make larger values easier to read.
- Added a new entry to show the number of bytes that have been rejected from the cache.
- Added "Avg" to the latency timing to make it clear that we are displaying an average.
- The feature is now compiled out in shipping builds, it was only ever intended to be used during development and/or debugging.

diff --git a/Engine/Source/Runtime/Experimental/IoStore/OnDemand/Private/Statistics.cpp b/Engine/Source/Runtime/Experimental/IoStore/OnDemand/Private/Statistics.cpp
--- a/Engine/Source/Runtime/Experimental/IoStore/OnDemand/Private/Statistics.cpp
+++ b/Engine/Source/Runtime/Experimental/IoStore/OnDemand/Private/Statistics.cpp
@@ -15,6 +15,8 @@ LLM_DEFINE_TAG(Ias);
 
 #if IAS_WITH_STATISTICS
 
+#define UE_ENABLE_ONSCREEN_STATISTICS !UE_BUILD_SHIPPING
+
 namespace UE::IoStore
 {
 
@@ -24,12 +26,6 @@ static FAutoConsoleVariableRef CVar_StatisticsLogInterval(
 	GIasStatisticsLogInterval,
 	TEXT("Enables and sets interval for periodic logging of statistics"));
 
-bool GIasDisplayOnScreenStatistics = false;
-static FAutoConsoleVariableRef CVar_DisplayOnScreenStatistics(
-	TEXT("ias.DisplayOnScreenStatistics"),
-	GIasDisplayOnScreenStatistics,
-	TEXT("Enables display of Ias on screen statistics"));
-
 bool GIasReportHttpAnalyticsEnabled = true;
 static FAutoConsoleVariableRef CVar_ReportHttpAnalytics(
 	TEXT("ias.ReportHttpAnalytics"),
@@ -42,6 +38,16 @@ static FAutoConsoleVariableRef CVar_ReportCacheAnalytics(
 	GIasReportCacheAnalyticsEnabled,
 	TEXT("Enables reporting statics on our file cache usage to the analytics system"));
 
+#if UE_ENABLE_ONSCREEN_STATISTICS
+
+bool GIasDisplayOnScreenStatistics = false;
+static FAutoConsoleVariableRef CVar_DisplayOnScreenStatistics(
+	TEXT("ias.DisplayOnScreenStatistics"),
+	GIasDisplayOnScreenStatistics,
+	TEXT("Enables display of Ias on screen statistics"));
+
+#endif // UE_ENABLE_ONSCREEN_STATISTICS
+
 ////////////////////////////////////////////////////////////////////////////////
 static int32 BytesToApproxMB(uint64 Bytes) { return int32(Bytes >> 20); }
 static int32 BytesToApproxKB(uint64 Bytes) { return int32(Bytes >> 10); }
@@ -411,6 +417,7 @@ FOnDemandIoBackendStats::FOnDemandIoBackendStats(EStatsFlags InFlags)
 		}
 	});
 
+#if UE_ENABLE_ONSCREEN_STATISTICS
 	GStatisticsOnScreenDelegateHandle = FCoreDelegates::OnGetOnScreenMessages.AddLambda(
 		[this] (FCoreDelegates::FSeverityMessageMap& Out)
 		{
@@ -421,12 +428,10 @@ FOnDemandIoBackendStats::FOnDemandIoBackendStats(EStatsFlags InFlags)
 				TStringBuilder<256> CachingText;
 				if (!EnumHasAllFlags(Flags, EStatsFlags::CachingDisabled))
 				{
-					CachingText.Appendf(
-						TEXT("Cached:%d KiB | Read:%d KiB (%d)"),
-						GCacheCachedBytes.Get() >> 10,
-						GCacheReadBytes.Get() >> 10,
-						GCacheGetCount.Get()
-					);
+					CachingText << TEXT("Cached: ") << FText::AsMemory(GCacheCachedBytes.Get()).ToString();
+					CachingText << TEXT(" | Rejected: ") << FText::AsMemory(GCacheRejectBytes.Get()).ToString();
+					CachingText << TEXT(" | Read: ") << FText::AsMemory(GCacheReadBytes.Get()).ToString();
+					CachingText << TEXT(" (") << GCacheGetCount.Get() << TEXT(")");
 				}
 				else
 				{
@@ -435,10 +440,10 @@ FOnDemandIoBackendStats::FOnDemandIoBackendStats(EStatsFlags InFlags)
 
 				TStringBuilder<256> Builder;
 				Builder.Appendf(
-					TEXT("IAS - %s: %s | Downloaded:%d KiB (%d) %d ms | Retries:%d | Pending:%d"),
+					TEXT("IAS - %s: %s | Downloaded:%s (%d) Avg %d ms | Retries:%d | Pending:%d"),
 					bIsConnected ? TEXT("Connected") : TEXT("Disconnected"),
 					CachingText.ToString(),
-					GHttpDownloadedBytes.Get() >> 10,
+					*FText::AsMemory(GHttpDownloadedBytes.Get()).ToString(),
 					GHttpGetCount.Get(),
 					GHttpDurationMsAvg.Get(),
 					GHttpRetryCount.Get(),
@@ -459,6 +464,7 @@ FOnDemandIoBackendStats::FOnDemandIoBackendStats(EStatsFlags InFlags)
 				}
 			}
 		});
+#endif // UE_ENABLE_ONSCREEN_STATISTICS
 }
 
 FOnDemandIoBackendStats::~FOnDemandIoBackendStats()
@@ -701,4 +707,6 @@ void FOnDemandIoBackendStats::OnHttpError()
 
 } // namespace UE::IoStore
 
+#undef UE_ENABLE_ONSCREEN_STATISTICS
+
 #endif // IAS_WITH_STATISTICS
