Iris - ReplicationDataStream runtime togglable debug features. Cvar net.Iris.ReplicationWriter.WriteBatchSizePerObject 0/1 controls whether sizes are sent per object or not rather than just one size for an object batch. This allows the reader to more quickly detect bitstream overflow and underflow errors. Cvar net.Iris.ReplicationWriter.WriteSentinels 0/1 controls whether sentinels are injected at certain point in the object replication to more quickly detect errors within a single object. The cvars can be changed at any point in time and both the reader and writer will adjust accordingly. Different build configurations are still 100% compatible with eachother but the code to handle the debug features is compiled out in shipping. If a non-shipping executable is connected to a shipping executable the non-shipping executable will cause a disconnect of the shippnig executable if any of the debug features is enabled.

#okforgihub public

diff --git a/Engine/Source/Runtime/Experimental/Iris/Core/Private/Iris/IrisConfigInternal.h b/Engine/Source/Runtime/Experimental/Iris/Core/Private/Iris/IrisConfigInternal.h
--- a/Engine/Source/Runtime/Experimental/Iris/Core/Private/Iris/IrisConfigInternal.h
+++ b/Engine/Source/Runtime/Experimental/Iris/Core/Private/Iris/IrisConfigInternal.h
@@ -31,12 +31,6 @@
 #endif 
 #endif
 
-// When enabled will add bit stream guards to detect overflow or serialization errors.
-// Increases bandwidth and replication time so should only be enabled for local or unit testing.
-#ifndef UE_NET_USE_READER_WRITER_SENTINEL
-#	define UE_NET_USE_READER_WRITER_SENTINEL 0
-#endif
-
 #ifndef UE_NET_VALIDATE_DC_BASELINES
 #if (UE_BUILD_SHIPPING || UE_BUILD_TEST)
 #	define UE_NET_VALIDATE_DC_BASELINES 0
diff --git a/Engine/Source/Runtime/Experimental/Iris/Core/Private/Iris/ReplicationSystem/NetBlob/ReliableNetBlobQueue.cpp b/Engine/Source/Runtime/Experimental/Iris/Core/Private/Iris/ReplicationSystem/NetBlob/ReliableNetBlobQueue.cpp
--- a/Engine/Source/Runtime/Experimental/Iris/Core/Private/Iris/ReplicationSystem/NetBlob/ReliableNetBlobQueue.cpp
+++ b/Engine/Source/Runtime/Experimental/Iris/Core/Private/Iris/ReplicationSystem/NetBlob/ReliableNetBlobQueue.cpp
@@ -116,13 +116,6 @@ uint32 FReliableNetBlobQueue::SerializeInternal(FNetSerializationContext& Contex
 		// GetRefCount() is not const.
 		TRefCountPtr<FNetBlob>& Attachment = NetBlobs[Index];
 
-#if UE_NET_USE_READER_WRITER_SENTINEL
-		{
-			UE_NET_TRACE_SCOPE(Sentintel, *Writer, Context.GetTraceCollector(), ENetTraceVerbosity::Verbose);
-			Writer->WriteBits(0x5e4714e1, 32);
-		}
-#endif
-
 		// If this sequence is disjoint from the previous sequence we need to serialize the full index.
 		// It's important that the sequence number is sent first so we can validate it before receiving exports and payload.
 		if (Writer->WriteBool(Seq != PrevWrittenSeq + 1U))
@@ -249,20 +242,6 @@ uint32 FReliableNetBlobQueue::DeserializeInternal(FNetSerializationContext& Cont
 	bool bHasMoreBlobs = false;
 	do
 	{
-
-#if UE_NET_USE_READER_WRITER_SENTINEL
-		{
-			UE_NET_TRACE_SCOPE(Sentintel, *Reader, Context.GetTraceCollector(), ENetTraceVerbosity::Verbose);
-			const uint32 Sentinel = Reader->ReadBits(32);
-			if (Sentinel != 0x5e4714e1U)
-			{
-				UE_LOG(LogIris, Error, TEXT("Wrong sentinel %u != %u"), Sentinel, 0x5e4714e1U);
-				Context.SetError(GNetError_BitStreamError);
-				return 0;
-			}
-		}
-#endif
-
 		if (const bool bIsNewSequence = Reader->ReadBool())
 		{
 			Index = Reader->ReadBits(IndexBitCount);
diff --git a/Engine/Source/Runtime/Experimental/Iris/Core/Private/Iris/ReplicationSystem/PendingBatchData.h b/Engine/Source/Runtime/Experimental/Iris/Core/Private/Iris/ReplicationSystem/PendingBatchData.h
--- a/Engine/Source/Runtime/Experimental/Iris/Core/Private/Iris/ReplicationSystem/PendingBatchData.h
+++ b/Engine/Source/Runtime/Experimental/Iris/Core/Private/Iris/ReplicationSystem/PendingBatchData.h
@@ -5,6 +5,13 @@
 #include "Containers/Array.h"
 #include "Iris/ReplicationSystem/NetRefHandle.h"
 
+namespace UE::Net::Private
+{
+
+enum EReplicationDataStreamDebugFeatures : uint32;
+
+}
+
 namespace UE::Net::Private
 {
 // Queued data chunk
@@ -22,6 +29,7 @@ struct FQueuedDataChunk
 	uint32 NumBits : 30;
 	uint32 bHasBatchOwnerData : 1;
 	uint32 bIsEndReplicationChunk : 1;
+	EReplicationDataStreamDebugFeatures StreamDebugFeatures = static_cast<EReplicationDataStreamDebugFeatures>(0U);
 };
 
 // Struct to contain storage and required data for queued batches pending must be mapped references
diff --git a/Engine/Source/Runtime/Experimental/Iris/Core/Private/Iris/ReplicationSystem/ReplicationDataStreamDebug.h b/Engine/Source/Runtime/Experimental/Iris/Core/Private/Iris/ReplicationSystem/ReplicationDataStreamDebug.h
new file mode 100644
--- /dev/null
+++ b/Engine/Source/Runtime/Experimental/Iris/Core/Private/Iris/ReplicationSystem/ReplicationDataStreamDebug.h
@@ -0,0 +1,44 @@
+// Copyright Epic Games, Inc. All Rights Reserved.
+
+#pragma once
+
+#include "HAL/Platform.h"
+#include "Misc/EnumClassFlags.h"
+#include "Iris/Serialization/NetBitStreamReader.h"
+#include "Iris/Serialization/NetBitStreamWriter.h"
+
+// Whether said features should be compiled in or not. Individual features should be togglable by cvars.
+#ifndef UE_NET_REPLICATIONDATASTREAM_DEBUG
+	#define UE_NET_REPLICATIONDATASTREAM_DEBUG !(UE_BUILD_SHIPPING)
+#endif
+
+namespace UE::Net::Private
+{
+
+// CVar net.Iris.
+extern bool bReplicationDataStreamDebugBatchSizePerObjectEnabled;
+// CVar net.Iris.
+extern bool bReplicationDataStreamDebugSentinelsEnabled;
+
+enum EReplicationDataStreamDebugFeatures : uint32
+{
+	None = 0U,
+	BatchSizePerObject = 1U << 0U,
+	Sentinels = 1U << 1U,
+};
+ENUM_CLASS_FLAGS(EReplicationDataStreamDebugFeatures);
+
+static constexpr uint32 ReplicationDataStreamDebugFeaturesBitCount = 2U;
+
+inline void WriteReplicationDataStreamDebugFeatures(FNetBitStreamWriter* Writer, EReplicationDataStreamDebugFeatures Features)
+{
+	Writer->WriteBits(static_cast<uint32>(Features), ReplicationDataStreamDebugFeaturesBitCount);
+}
+
+inline EReplicationDataStreamDebugFeatures ReadReplicationDataStreamDebugFeatures(FNetBitStreamReader* Reader)
+{
+	EReplicationDataStreamDebugFeatures Features = static_cast<EReplicationDataStreamDebugFeatures>(Reader->ReadBits(ReplicationDataStreamDebugFeaturesBitCount));
+	return Features;
+}
+
+}
diff --git a/Engine/Source/Runtime/Experimental/Iris/Core/Private/Iris/ReplicationSystem/ReplicationReader.cpp b/Engine/Source/Runtime/Experimental/Iris/Core/Private/Iris/ReplicationSystem/ReplicationReader.cpp
--- a/Engine/Source/Runtime/Experimental/Iris/Core/Private/Iris/ReplicationSystem/ReplicationReader.cpp
+++ b/Engine/Source/Runtime/Experimental/Iris/Core/Private/Iris/ReplicationSystem/ReplicationReader.cpp
@@ -376,6 +376,22 @@ FNetRefHandle FReplicationReader::ReadNetRefHandleId(FNetSerializationContext& C
 	return RefHandle;
 }
 	
+void FReplicationReader::ReadStreamDebugFeatures(FNetSerializationContext& Context)
+{
+	FNetBitStreamReader* Reader = Context.GetBitStreamReader();
+	StreamDebugFeatures = ReadReplicationDataStreamDebugFeatures(Reader);
+
+#if !UE_NET_REPLICATIONDATASTREAM_DEBUG
+	// If datastream debug features isn't enabled in this build the debug functions won't do any useful work and will cause bitstream errors.
+	if (!ensure(StreamDebugFeatures != EReplicationDataStreamDebugFeatures::None))
+	{
+		UE_LOG(LogIris, Error, TEXT("StreamDebugFeatures enabled on sending side despite receiving side not being able to use them."));
+		Context.SetError(GNetError_InvalidValue);
+		return;
+	}
+#endif
+}
+
 uint32 FReplicationReader::ReadObjectsPendingDestroy(FNetSerializationContext& Context)
 {
 	FNetBitStreamReader& Reader = *Context.GetBitStreamReader();
@@ -760,7 +776,7 @@ FPendingBatchData* FReplicationReader::UpdateUnresolvedMustBeMappedReferences(FN
 	return nullptr;
 }
 
-uint32 FReplicationReader::ReadObjectsInBatch(FNetSerializationContext& Context, FNetRefHandle IncompleteHandle, bool bHasBatchOwnerData, uint32 BatchEndBitPosition)
+uint32 FReplicationReader::ReadObjectsInBatchWithoutSizes(FNetSerializationContext& Context, FNetRefHandle IncompleteHandle, bool bHasBatchOwnerData, uint32 BatchEndBitPosition)
 {
 	uint32 ReadObjectCount = 0;
 	FNetBitStreamReader& Reader = *Context.GetBitStreamReader();
@@ -776,7 +792,11 @@ uint32 FReplicationReader::ReadObjectsInBatch(FNetSerializationContext& Context,
 		++ReadObjectCount;
 	}
 
-	ensure(Reader.GetPosBits() <= BatchEndBitPosition);
+	if (!ensure(Reader.GetPosBits() <= BatchEndBitPosition))
+	{
+		Context.SetError(GNetError_BitStreamOverflow);
+		return 0U;
+	}
 
 	// ReadSubObjects 
 	while (Reader.GetPosBits() < BatchEndBitPosition)
@@ -792,6 +812,91 @@ uint32 FReplicationReader::ReadObjectsInBatch(FNetSerializationContext& Context,
 	return ReadObjectCount;
 }
 
+uint32 FReplicationReader::ReadObjectsInBatchWithSizes(FNetSerializationContext& Context, FNetRefHandle IncompleteHandle, const bool bHasBatchOwnerData, uint32 BatchEndBitPosition)
+{
+#if !UE_NET_REPLICATIONDATASTREAM_DEBUG
+	return 0;
+#else
+	const uint32 BatchSize = BatchEndBitPosition - Context.GetBitStreamReader()->GetPosBits();
+	FNetBitStreamReader BatchReader = Context.GetBitStreamReader()->CreateSubstream(BatchSize);
+
+	FNetSerializationContext ObjectContext;
+	FNetBitStreamReader ObjectReader;
+
+	bool bIsSubObject = !bHasBatchOwnerData;
+	uint32 ReadObjectCount = 0;
+	while (BatchReader.GetBitsLeft() > 0U)
+	{
+		{
+			UE_NET_TRACE_SCOPE(BatchSize, BatchReader, Context.GetTraceCollector(), ENetTraceVerbosity::Trace);
+			const uint32 ObjectSize = BatchReader.ReadBits(Parameters.NumBitsUsedForHugeObjectBatchSize);
+			ObjectReader = BatchReader.CreateSubstream(ObjectSize);
+
+			// If we were unable to create a substream of the correct size we have bitstream corruption.
+			if (ObjectReader.GetBitsLeft() != ObjectSize)
+			{
+				BatchReader.DiscardSubstream(ObjectReader);
+				Context.GetBitStreamReader()->DiscardSubstream(BatchReader);
+				Context.SetError(GNetError_BitStreamError);
+				break;
+			}
+		}
+
+		ObjectContext = Context.MakeSubContext(&ObjectReader);
+		ReadObjectInBatch(ObjectContext, IncompleteHandle, bIsSubObject);
+
+		// Check for error
+		if (ObjectContext.HasErrorOrOverflow())
+		{
+			BatchReader.DiscardSubstream(ObjectReader);
+			Context.GetBitStreamReader()->DiscardSubstream(BatchReader);
+			Context.SetError(ObjectContext.GetError(), ObjectReader.IsOverflown());
+			break;
+		}
+
+		// Check for underflow
+		if (ObjectReader.GetBitsLeft() != 0U)
+		{
+			BatchReader.DiscardSubstream(ObjectReader);
+			Context.GetBitStreamReader()->DiscardSubstream(BatchReader);
+			Context.SetError(GNetError_BitStreamError);
+			UE_LOG(LogIris, Error, TEXT("FReplicationReader::ReadObjectsInBatchWithSizes Bitstream underflow after reading replicated object with Handle: %s. %s"), ToCStr(ObjectContext.GetErrorHandleContext().ToString()), ToCStr(ObjectContext.PrintReadJournal()));
+			break;
+		}
+
+		BatchReader.CommitSubstream(ObjectReader);
+
+		++ReadObjectCount;
+		// Only the first batch can be the subobject owner. The rest will be subobjects.
+		bIsSubObject = true;
+	}
+
+	if (Context.HasErrorOrOverflow())
+	{
+		return 0U;
+	}
+	else
+	{
+		Context.GetBitStreamReader()->CommitSubstream(BatchReader);
+		return ReadObjectCount;
+	}
+#endif
+}
+
+uint32 FReplicationReader::ReadObjectsInBatch(FNetSerializationContext& Context, FNetRefHandle IncompleteHandle, bool bHasBatchOwnerData, uint32 BatchEndBitPosition)
+{
+#if UE_NET_REPLICATIONDATASTREAM_DEBUG
+	if (EnumHasAnyFlags(StreamDebugFeatures, EReplicationDataStreamDebugFeatures::BatchSizePerObject))
+	{
+		return ReadObjectsInBatchWithSizes(Context, IncompleteHandle, bHasBatchOwnerData, BatchEndBitPosition);
+	}
+	else
+#endif
+	{
+		return ReadObjectsInBatchWithoutSizes(Context, IncompleteHandle, bHasBatchOwnerData, BatchEndBitPosition);
+	}
+}
+
 uint32 FReplicationReader::ReadObjectBatch(FNetSerializationContext& Context, uint32 ReadObjectFlags)
 {
 	FNetBitStreamReader& Reader = *Context.GetBitStreamReader();
@@ -809,9 +914,7 @@ uint32 FReplicationReader::ReadObjectBatch(FNetSerializationContext& Context, ui
 		FForceInlineExportScope ForceInlineExportScope(Context.GetInternalContext());
 		ReplicationBridge->ReadAndExecuteDestructionInfoFromRemote(BridgeContext);
 
-#if UE_NET_USE_READER_WRITER_SENTINEL
-		ReadAndVerifySentinelBits(&Reader, TEXT("DestructionInfo"), 8);
-#endif
+		ReadSentinel(Context, TEXT("DestructionInfo"));
 
 		if (Context.HasErrorOrOverflow())
 		{
@@ -822,10 +925,13 @@ uint32 FReplicationReader::ReadObjectBatch(FNetSerializationContext& Context, ui
 
 		return 1U;
 	}
-
-#if UE_NET_USE_READER_WRITER_SENTINEL
-	ReadAndVerifySentinelBits(&Reader, TEXT("ReadObject"), 8);
-#endif
+	
+	if (!ReadSentinel(Context, TEXT("Object")))
+	{
+		UE_LOG(LogIris, Error, TEXT("FReplicationReader::ReadObject Failed to read object sentinel. \n%s"), *Context.PrintReadJournal());
+		Context.SetError(GNetError_BitStreamError);
+		return 0U;
+	}
 
 	uint32 ObjectsReadInBatch = 0U;
 	
@@ -920,6 +1026,7 @@ uint32 FReplicationReader::ReadObjectBatch(FNetSerializationContext& Context, ui
 		DataChunk.StorageOffset = PendingBatchData->DataChunkStorage.Num();
 		DataChunk.bHasBatchOwnerData = bHasBatchOwnerData;
 		DataChunk.bIsEndReplicationChunk = false;
+		DataChunk.StreamDebugFeatures = StreamDebugFeatures;
 
 		// Make sure we have space
 		PendingBatchData->DataChunkStorage.AddUninitialized(NumDataWords);
@@ -991,9 +1098,11 @@ void FReplicationReader::ReadObjectInBatch(FNetSerializationContext& Context, FN
 	const bool bHasState = Reader.ReadBool();
 	if (bHasState)
 	{
-#if UE_NET_USE_READER_WRITER_SENTINEL
-		ReadAndVerifySentinelBits(&Reader, TEXT("HasState"), 8);
-#endif
+		if (!ReadSentinel(Context, TEXT("HasState")))
+		{
+			UE_LOG(LogIris, Error, TEXT("FReplicationReader::ReadObject Failed to read replicated object with Handle: %s. Error '%s'. %s"), ToCStr(IncompleteHandle.ToString()), (Context.HasError() ? ToCStr(Context.GetError().ToString()) : TEXT("BitStream Overflow")), ToCStr(Context.PrintReadJournal()));
+			return;
+		}
 	}
 
 	uint32 NewBaselineIndex = FDeltaCompressionBaselineManager::InvalidBaselineIndex;
@@ -1187,9 +1296,11 @@ void FReplicationReader::ReadObjectInBatch(FNetSerializationContext& Context, FN
 				}
 			}
 
-#if UE_NET_USE_READER_WRITER_SENTINEL
-			ReadAndVerifySentinelBits(&Reader, TEXT("HasStateEnd"), 8);
-#endif
+			if (!ReadSentinel(Context, TEXT("HasStateEnd")))
+			{
+				bHasErrors = true;
+				goto ErrorHandling;
+			}
 
 			// Should we store a new baseline?
 			if (NewBaselineIndex != FDeltaCompressionBaselineManager::InvalidBaselineIndex)
@@ -2374,6 +2485,7 @@ void FReplicationReader::ProcessHugeObjectAttachment(FNetSerializationContext& C
 	// Find out how many objects to read so we can reserve object dispatch infos.
 	FNetObjectBlob::FHeader HugeObjectHeader = {};
 	FNetObjectBlob::DeserializeHeader(HugeObjectSerializationContext, HugeObjectHeader);
+	ReadStreamDebugFeatures(HugeObjectSerializationContext);
 	if (HugeObjectSerializationContext.HasErrorOrOverflow() || HugeObjectHeader.ObjectCount < 1U)
 	{
 		if (!Context.HasError())
@@ -2555,6 +2667,9 @@ void FReplicationReader::ProcessQueuedBatches()
 				// Need to set this pointer as we are dealing with temporary linear allocations
 				ObjectsToDispatchArray = &TempObjectsToDispatchArray;
 
+				// Use whatver StreamDebugFeatures were set when receiving the chunk
+				StreamDebugFeatures = CurrentChunk.StreamDebugFeatures;
+
 				// $IRIS: $TODO: Implement special dispatch to defer RepNotifies if we are processing multiple batches for the same object.
 				ReadObjectsInBatch(Context, PendingBatchData.Handle, CurrentChunk.bHasBatchOwnerData, CurrentChunk.NumBits);
 
@@ -2649,6 +2764,12 @@ void FReplicationReader::Read(FNetSerializationContext& Context)
 
 	UE_NET_TRACE_SCOPE(ReplicationData, *Context.GetBitStreamReader(), Context.GetTraceCollector(), ENetTraceVerbosity::Trace);
 
+	ReadStreamDebugFeatures(Context);
+	if (Context.HasErrorOrOverflow())
+	{
+		return;
+	}
+
 	FMemMark TempAllocatorScope(TempLinearAllocator);
 
 	// Sanity check received object count
@@ -2830,4 +2951,22 @@ void FReplicationReader::ResolveAndDispatchAttachments(FNetSerializationContext&
 	ReplicationInfo->bHasAttachments = bHasUnresolvedReferences;
 }
 
+bool FReplicationReader::ReadSentinel(FNetSerializationContext& Context, const TCHAR* DebugName)
+{
+	bool bSuccess = true;
+
+#if UE_NET_REPLICATIONDATASTREAM_DEBUG
+	if (EnumHasAnyFlags(StreamDebugFeatures, EReplicationDataStreamDebugFeatures::Sentinels))
+	{
+		bSuccess = ReadAndVerifySentinelBits(Context.GetBitStreamReader(), DebugName, 8U);
+		if (!bSuccess)
+		{
+			Context.SetError(GNetError_BitStreamError);
+		}
+	}
+#endif
+
+	return bSuccess;
+}
+
 } // end namespace UE::Net::Private
diff --git a/Engine/Source/Runtime/Experimental/Iris/Core/Private/Iris/ReplicationSystem/ReplicationReader.h b/Engine/Source/Runtime/Experimental/Iris/Core/Private/Iris/ReplicationSystem/ReplicationReader.h
--- a/Engine/Source/Runtime/Experimental/Iris/Core/Private/Iris/ReplicationSystem/ReplicationReader.h
+++ b/Engine/Source/Runtime/Experimental/Iris/Core/Private/Iris/ReplicationSystem/ReplicationReader.h
@@ -7,8 +7,9 @@
 #include "Iris/ReplicationSystem/ChangeMaskUtil.h"
 #include "Iris/ReplicationSystem/NetRefHandleManager.h"
 #include "Iris/ReplicationSystem/ObjectReferenceCache.h"
-#include "Iris/ReplicationSystem/ReplicationTypes.h"
 #include "Iris/ReplicationSystem/PendingBatchData.h"
+#include "Iris/ReplicationSystem/ReplicationDataStreamDebug.h"
+#include "Iris/ReplicationSystem/ReplicationTypes.h"
 #include "Containers/Map.h"
 #include "Misc/MemStack.h"
 
@@ -141,7 +142,12 @@ private:
 	void ReadObjectInBatch(FNetSerializationContext& Context, FNetRefHandle BatchHandle, bool bIsSubObject);
 
 	uint32 ReadObjectsInBatch(FNetSerializationContext& Context, FNetRefHandle InCompleteHandle, bool bHasBatchOwnerData, uint32 BatchEndBitPosition);
+	uint32 ReadObjectsInBatchWithoutSizes(FNetSerializationContext& Context, FNetRefHandle InCompleteHandle, bool bHasBatchOwnerData, uint32 BatchEndBitPosition);
+	uint32 ReadObjectsInBatchWithSizes(FNetSerializationContext& Context, FNetRefHandle InCompleteHandle, bool bHasBatchOwnerData, uint32 BatchEndBitPosition);
+
 
+	// Read stream debug features. Use StreamDebugFeatures to see if a feature is enabled.
+	void ReadStreamDebugFeatures(FNetSerializationContext& Context);
 	// Read objects pending destroy. This will call ReadObjectsAndSubObjectsPendingDestroy or ReadRootObjectsPendingDestroy depending on sending side settings.
 	uint32 ReadObjectsPendingDestroy(FNetSerializationContext& Context);
 	// Receive both objects and subobject destruction. Legacy Iris behavior.
@@ -216,6 +222,9 @@ private:
 	// Remove a handle from the hot and cold unresolved caches used by ResolveAndDispatchUnresolvedReferences(). If this handle is marked as unresolved
 	// again, it will be added to the hot cache.
 	void RemoveFromUnresolvedCache(const FNetRefHandle Handle);
+
+	// Reads and verifies the sentinel was read properly. If an error is detected the context will have an error set. The function will return true if the sentinel matches, false otherwise.
+	bool ReadSentinel(FNetSerializationContext& Context, const TCHAR* DebugName);
 	
 private:
 
@@ -282,6 +291,9 @@ private:
 
 	IConsoleVariable const* DelayAttachmentsWithUnresolvedReferences;
 	uint32 NumHandlesPendingResolveLastUpdate = 0U;
+
+	// Features that can aid in tracking down bitstream errors etc.
+	EReplicationDataStreamDebugFeatures StreamDebugFeatures = EReplicationDataStreamDebugFeatures::None;
 };
 
 }
diff --git a/Engine/Source/Runtime/Experimental/Iris/Core/Private/Iris/ReplicationSystem/ReplicationWriter.cpp b/Engine/Source/Runtime/Experimental/Iris/Core/Private/Iris/ReplicationSystem/ReplicationWriter.cpp
--- a/Engine/Source/Runtime/Experimental/Iris/Core/Private/Iris/ReplicationSystem/ReplicationWriter.cpp
+++ b/Engine/Source/Runtime/Experimental/Iris/Core/Private/Iris/ReplicationSystem/ReplicationWriter.cpp
@@ -89,6 +89,15 @@ static FAutoConsoleVariableRef CvarValidateObjectsWithDirtyChanges(TEXT("net.Iri
 static bool bDestroyRootAndSubObjectsIndividually = false;
 static FAutoConsoleVariableRef CvarDestroyRootAndSubObjectsIndividually(TEXT("net.Iris.DestroyRootAndSubObjectsIndividually"), bDestroyRootAndSubObjectsIndividually, TEXT("TEMPORARY. Whether to send destroy on both root and subobjects (Iris legacy behavior) or just root objects and have the receiving end destroy the subobjects with it. Default is false, only sending destroy for root objects."));
 
+// net.Iris.ReplicationWriter.WriteBatchSizePerObject
+static bool bDebugBatchSizePerObjectEnabled = false;
+static FAutoConsoleVariableRef CvarWriteBatchSizePerObject(TEXT("net.Iris.ReplicationWriter.WriteBatchSizePerObject"), bDebugBatchSizePerObjectEnabled, TEXT("Write batch size per object. Helps tracking down bitstream errors. Requires code to be compiled with UE_NET_REPLICATIONDATASTREAM_DEBUG to be enabled."));
+
+// net.Iris.ReplicationWriter.WriteSentinels
+static bool bDebugSentinelsEnabled = false;
+static FAutoConsoleVariableRef CvarWriteSentinels(TEXT("net.Iris.ReplicationWriter.WriteSentinels"), bDebugSentinelsEnabled, TEXT("Write sentinels at carefully chosen points in the stream. Helps tracking down bitstream errors. Requires code to be compiled with UE_NET_REPLICATIONDATASTREAM_DEBUG to be enabled."));
+
+
 static const FName NetError_ObjectStateTooLarge("Object state is too large to be split.");
 
 const TCHAR* FReplicationWriter::LexToString(const EReplicatedObjectState State)
@@ -1814,6 +1823,25 @@ void FReplicationWriter::CommitBatchRecord(const FBatchRecord& BatchRecord)
 	}
 }
 
+void FReplicationWriter::UpdateStreamDebugFeatures()
+{
+	StreamDebugFeatures = EReplicationDataStreamDebugFeatures::None;
+#if UE_NET_REPLICATIONDATASTREAM_DEBUG
+	// Enable these features for object replication. This isn't the best test.
+	if (Parameters.ReplicationSystem && Parameters.ReplicationSystem->IsServer())
+	{
+		StreamDebugFeatures |= (bDebugBatchSizePerObjectEnabled ? EReplicationDataStreamDebugFeatures::BatchSizePerObject : EReplicationDataStreamDebugFeatures::None);
+		StreamDebugFeatures |= (bDebugSentinelsEnabled ? EReplicationDataStreamDebugFeatures::Sentinels : EReplicationDataStreamDebugFeatures::None);
+	}
+#endif
+}
+
+void FReplicationWriter::WriteStreamDebugFeatures(FNetSerializationContext& Context)
+{
+	FNetBitStreamWriter* Writer = Context.GetBitStreamWriter();
+	WriteReplicationDataStreamDebugFeatures(Writer, StreamDebugFeatures);
+}
+
 uint32 FReplicationWriter::WriteObjectsPendingDestroy(FNetSerializationContext& Context)
 {
 	FNetBitStreamWriter& Writer = *Context.GetBitStreamWriter();
@@ -2313,7 +2341,7 @@ FReplicationWriter::EWriteObjectStatus FReplicationWriter::WriteObjectAndSubObje
 			NetRefHandleForTraceScope = NetRefHandleManager->GetReplicatedObjectDataNoCheck(HugeObjectInternalIndex).RefHandle;
 		}
 	}
-	UE_NET_TRACE_OBJECT_SCOPE(NetRefHandleForTraceScope, Writer, Context.GetTraceCollector(), ENetTraceVerbosity::Trace);
+	UE_NET_TRACE_NAMED_OBJECT_SCOPE(ObjectTraceScope, NetRefHandleForTraceScope, Writer, Context.GetTraceCollector(), ENetTraceVerbosity::Trace);
 #endif
 
 	// We only need to write batch info for root objects
@@ -2325,12 +2353,10 @@ FReplicationWriter::EWriteObjectStatus FReplicationWriter::WriteObjectAndSubObje
 	if (bWriteBatchInfo)
 	{
 		// Write bit indicating that we are not a destruction info.
-		const bool bIsDestructionInfo = false;
+		constexpr bool bIsDestructionInfo = false;
 		Writer.WriteBool(bIsDestructionInfo);
 
-	#if UE_NET_USE_READER_WRITER_SENTINEL
-		WriteSentinelBits(&Writer, 8);
-	#endif
+		WriteSentinel(&Writer, TEXT("DestructionInfo"));
 
 		// A batch starts with (RefHandleId | BatchSize | bHasBatchObjectData | bHasExports)
 		// We write the header up front, and then we seek back and update relevant info is the object + subobjects is successfully serialized along with necessary exports
@@ -2431,6 +2457,16 @@ FReplicationWriter::EWriteObjectStatus FReplicationWriter::WriteObjectAndSubObje
 
 	if (bHasState | bWriteAttachments | bSentTearOff | Info.SubObjectPendingDestroy)
 	{
+#if UE_NET_REPLICATIONDATASTREAM_DEBUG
+		uint32 BatchSizePos = Writer.GetPosBits();
+		const uint32 BatchSizeBits = EnumHasAnyFlags(StreamDebugFeatures, EReplicationDataStreamDebugFeatures::BatchSizePerObject) ? Parameters.NumBitsUsedForHugeObjectBatchSize : 0U;
+		if (BatchSizeBits)
+		{
+			UE_NET_TRACE_SCOPE(BatchSize, Writer, Context.GetTraceCollector(), ENetTraceVerbosity::Trace);
+			Writer.WriteBits(0U, BatchSizeBits);
+		}
+#endif
+
 		// Only need to write the handle if this is a subobject
 		if (Info.IsSubObject)
 		{
@@ -2453,9 +2489,7 @@ FReplicationWriter::EWriteObjectStatus FReplicationWriter::WriteObjectAndSubObje
 
 		if (Writer.WriteBool(bHasState))
 		{
-#if UE_NET_USE_READER_WRITER_SENTINEL
-			WriteSentinelBits(&Writer, 8);
-#endif
+			WriteSentinel(&Writer, TEXT("HasState"));
 
 			BatchEntry.bSentState = 1;
 
@@ -2576,9 +2610,8 @@ FReplicationWriter::EWriteObjectStatus FReplicationWriter::WriteObjectAndSubObje
 					FReplicationProtocolOperations::SerializeWithMask(Context, Info.GetChangeMaskStoragePointer(), ReplicatedObjectStateBuffer, ObjectData.Protocol);
 				}
 			}
-#if UE_NET_USE_READER_WRITER_SENTINEL
-			WriteSentinelBits(&Writer, 8);
-#endif
+
+			WriteSentinel(&Writer, TEXT("HasStateEnd"));
 		}
 
 		{
@@ -2661,6 +2694,16 @@ FReplicationWriter::EWriteObjectStatus FReplicationWriter::WriteObjectAndSubObje
 			return EWriteObjectStatus::BitStreamOverflow;
 		}
 
+#if UE_NET_REPLICATIONDATASTREAM_DEBUG
+		// Write batch size if we didn't roll back the whole thing already.
+		if (BatchSizeBits && (Writer.GetPosBits() > BatchSizePos))
+		{
+			const uint32 BatchSize = Writer.GetPosBits()  - BatchSizePos - BatchSizeBits;
+			FNetBitStreamWriteScope WriteScope(Writer, BatchSizePos);
+			Writer.WriteBits(BatchSize, BatchSizeBits);
+		}
+#endif
+
 		if (bWriteReplicatedDestroyHeader)
 		{
 			// Rewrite destroy header if necessary	
@@ -2690,6 +2733,8 @@ FReplicationWriter::EWriteObjectStatus FReplicationWriter::WriteObjectAndSubObje
 		}
 	}
 
+	UE_NET_TRACE_EXIT_NAMED_SCOPE(ObjectTraceScope);
+
 	// Success so far. Fill in batch entry. Keep index to update info later as the array can resize.
 	const int ParentBatchEntryIndex = OutBatchInfo.ObjectInfos.Num();
 	{
@@ -2904,6 +2949,9 @@ int FReplicationWriter::PrepareAndSendHugeObjectPayload(FNetSerializationContext
 	FNetObjectBlob::SerializeHeader(HugeObjectSerializationContext, HugeObjectHeader);
 	const uint32 PastHeaderPos = HugeObjectWriter.GetPosBits();
 
+	// As the huge object will most likely be processed out of order we need to write the stream debug features again
+	WriteStreamDebugFeatures(HugeObjectSerializationContext);
+
 	FHugeObjectContext HugeObjectContext;
 
 	FBatchInfo BatchInfo;
@@ -3180,12 +3228,7 @@ int FReplicationWriter::WriteDestructionInfo(FNetSerializationContext& Context,
 		return 0;
 	}
 
-#if UE_NET_USE_READER_WRITER_SENTINEL
-	{
-		UE_NET_TRACE_SCOPE(Sentinel, Writer, Context.GetTraceCollector(), ENetTraceVerbosity::Trace);
-		WriteSentinelBits(&Writer, 8);
-	}
-#endif
+	WriteSentinel(&Writer, TEXT("DestructionInfo"));
 
 	// Push record
 	if (!Writer.IsOverflown())
@@ -3731,9 +3774,14 @@ UDataStream::EWriteResult FReplicationWriter::Write(FNetSerializationContext& Co
 	WriteBitStreamInfo.ReplicationStartPos = Writer.GetPosBits();
 	WriteBitStreamInfo.ReplicationCapacity = Writer.GetBitsLeft();
 
+	UpdateStreamDebugFeatures();
+
 	UE_NET_TRACE_SCOPE(ReplicationData, Writer, Context.GetTraceCollector(), ENetTraceVerbosity::Trace);
 	
 	FNetBitStreamRollbackScope Rollback(Writer);
+
+	WriteStreamDebugFeatures(Context);
+
 	const uint32 HeaderPos = Writer.GetPosBits();
 
 	uint32 WrittenObjectCount = 0;
@@ -4042,6 +4090,16 @@ void FReplicationWriter::MarkObjectDirty(FInternalNetRefIndex InternalIndex, con
 	ObjectsWithDirtyChanges.SetBit(InternalIndex);
 }
 
+void FReplicationWriter::WriteSentinel(FNetBitStreamWriter* Writer, const TCHAR* DebugName)
+{
+#if UE_NET_REPLICATIONDATASTREAM_DEBUG
+	if (EnumHasAnyFlags(StreamDebugFeatures, EReplicationDataStreamDebugFeatures::Sentinels))
+	{
+		WriteSentinelBits(Writer, 8U);
+	}
+#endif
+}
+
 FReplicationWriter::FHugeObjectContext::FHugeObjectContext() = default;
 
 FReplicationWriter::FHugeObjectContext::~FHugeObjectContext() = default;
diff --git a/Engine/Source/Runtime/Experimental/Iris/Core/Private/Iris/ReplicationSystem/ReplicationWriter.h b/Engine/Source/Runtime/Experimental/Iris/Core/Private/Iris/ReplicationSystem/ReplicationWriter.h
--- a/Engine/Source/Runtime/Experimental/Iris/Core/Private/Iris/ReplicationSystem/ReplicationWriter.h
+++ b/Engine/Source/Runtime/Experimental/Iris/Core/Private/Iris/ReplicationSystem/ReplicationWriter.h
@@ -13,6 +13,7 @@
 #include "Iris/ReplicationSystem/NetRefHandle.h"
 #include "Iris/ReplicationSystem/NetExports.h"
 #include "Iris/ReplicationSystem/ObjectReferenceCache.h"
+#include "Iris/ReplicationSystem/ReplicationDataStreamDebug.h"
 #include "Iris/Stats/NetStats.h"
 #include "Containers/Array.h"
 #include "Containers/List.h"
@@ -468,6 +469,10 @@ private:
 	// Partial sort of OutScheduledObjectIndices, will sort at most PartialSortObjectCount objects
 	uint32 SortScheduledObjects(FScheduleObjectInfo* ScheduledObjectIndices, uint32 ScheduledObjectCount, uint32 StartIndex);
 
+	// Update the active set of stream debug features based on cvars, build configuration etc. Enabling certain debug features can help track down bitstream errors. 
+	void UpdateStreamDebugFeatures();
+	// Write stream debug features.
+	void WriteStreamDebugFeatures(FNetSerializationContext& Context);
 	// Write all objects pending destroy (or as many as we fit in the current packet). Will call one of WriteObjectsAndSubObjectsPendingDestroy and WriteRootObjectsPendingDestroy depending on cvar net.Iris.DestroyRootAndSubObjectsIndividually
 	uint32 WriteObjectsPendingDestroy(FNetSerializationContext& Context);
 	// Used when net.Iris.DestroyRootAndSubObjectsIndividually is true. Will send destroy for every root and subobject. Leagcy Iris behavior.
@@ -557,6 +562,10 @@ private:
 
 	void MarkObjectDirty(FInternalNetRefIndex InternalIndex, const char* Caller);
 
+	// Writes a sentinel naking it easier for the receiver to detect bitstream errors.
+	void WriteSentinel(FNetBitStreamWriter* Writer, const TCHAR* DebugName);
+
+
 private:
 	// Replication parameters
 	FReplicationParameters Parameters;
@@ -601,6 +610,8 @@ private:
 	FWriteContext WriteContext;
 	FBitStreamInfo WriteBitStreamInfo;
 	FHugeObjectSendQueue HugeObjectSendQueue;
+	// Features that can aid in tracking down bitstream errors etc.
+	EReplicationDataStreamDebugFeatures StreamDebugFeatures = EReplicationDataStreamDebugFeatures::None;
 
 	// Is replication enabled?
 	bool bReplicationEnabled = false;
