Add anytype registration to allow string only type compatibility checks
- Added AnyType registry with corresponding macro
- Added a way to query if two types are compatible through the registry

diff --git a/Engine/Source/Runtime/Experimental/Dataflow/Core/Private/Dataflow/DataflowAnyType.cpp b/Engine/Source/Runtime/Experimental/Dataflow/Core/Private/Dataflow/DataflowAnyType.cpp
new file mode 100644
--- /dev/null
+++ b/Engine/Source/Runtime/Experimental/Dataflow/Core/Private/Dataflow/DataflowAnyType.cpp
@@ -0,0 +1,33 @@
+// Copyright Epic Games, Inc. All Rights Reserved.
+
+#include "Dataflow/DataflowAnyType.h"
+#include "Dataflow/DataflowAnyTypeRegistry.h"
+
+namespace UE::Dataflow
+{
+	void RegisterAnyTypes()
+	{
+		UE_DATAFLOW_REGISTER_ANYTYPE(FDataflowAnyType);
+		UE_DATAFLOW_REGISTER_ANYTYPE(FDataflowAllTypes);
+		UE_DATAFLOW_REGISTER_ANYTYPE(FDataflowNumericTypes);
+		UE_DATAFLOW_REGISTER_ANYTYPE(FDataflowVectorTypes);
+		UE_DATAFLOW_REGISTER_ANYTYPE(FDataflowStringTypes);
+		UE_DATAFLOW_REGISTER_ANYTYPE(FDataflowBoolTypes);
+		UE_DATAFLOW_REGISTER_ANYTYPE(FDataflowTransformTypes);
+		UE_DATAFLOW_REGISTER_ANYTYPE(FDataflowStringConvertibleTypes);
+		UE_DATAFLOW_REGISTER_ANYTYPE(FDataflowUObjectConvertibleTypes);
+		UE_DATAFLOW_REGISTER_ANYTYPE(FDataflowSelectionTypes);
+
+		// array types 
+		UE_DATAFLOW_REGISTER_ANYTYPE(FDataflowVectorArrayTypes);
+		UE_DATAFLOW_REGISTER_ANYTYPE(FDataflowNumericArrayTypes);
+		UE_DATAFLOW_REGISTER_ANYTYPE(FDataflowStringArrayTypes);
+		UE_DATAFLOW_REGISTER_ANYTYPE(FDataflowBoolArrayTypes);
+		UE_DATAFLOW_REGISTER_ANYTYPE(FDataflowTransformArrayTypes);
+	}
+
+	bool AreTypesCompatible(FName TypeA, FName TypeB)
+	{
+		return FAnyTypesRegistry::AreTypesCompatibleStatic(TypeA, TypeB);
+	}
+}
\ No newline at end of file
diff --git a/Engine/Source/Runtime/Experimental/Dataflow/Core/Private/Dataflow/DataflowAnyTypeRegistry.cpp b/Engine/Source/Runtime/Experimental/Dataflow/Core/Private/Dataflow/DataflowAnyTypeRegistry.cpp
new file mode 100644
--- /dev/null
+++ b/Engine/Source/Runtime/Experimental/Dataflow/Core/Private/Dataflow/DataflowAnyTypeRegistry.cpp
@@ -0,0 +1,40 @@
+// Copyright Epic Games, Inc. All Rights Reserved.
+
+#include "Dataflow/DataflowAnyTypeRegistry.h"
+
+namespace UE::Dataflow
+{
+	bool FAnyTypesRegistry::AreTypesCompatibleStatic(FName TypeA, FName TypeB)
+	{
+		return GetInstance().AreTypesCompatible(TypeA, TypeB);
+	}
+	
+	FAnyTypesRegistry& FAnyTypesRegistry::GetInstance()
+	{
+		static FAnyTypesRegistry Instance;
+		return Instance;
+	}
+
+	bool FAnyTypesRegistry::AreTypesCompatible(FName TypeA, FName TypeB) const
+	{
+		if (TypeA == TypeB)
+		{
+			return true;
+		}
+		if (const FSupportTypeFunction* SupportFuncA = SupportTypeFunctionByName.Find(TypeA))
+		{
+			if ((*SupportFuncA)(TypeB))
+			{
+				return true;
+			}
+		}
+		if (const FSupportTypeFunction* SupportFuncB = SupportTypeFunctionByName.Find(TypeB))
+		{
+			if ((*SupportFuncB)(TypeA))
+			{
+				return true;
+			}
+		}
+		return false;
+	}
+};
diff --git a/Engine/Source/Runtime/Experimental/Dataflow/Core/Private/Dataflow/DataflowNodeFactory.cpp b/Engine/Source/Runtime/Experimental/Dataflow/Core/Private/Dataflow/DataflowNodeFactory.cpp
--- a/Engine/Source/Runtime/Experimental/Dataflow/Core/Private/Dataflow/DataflowNodeFactory.cpp
+++ b/Engine/Source/Runtime/Experimental/Dataflow/Core/Private/Dataflow/DataflowNodeFactory.cpp
@@ -20,6 +20,7 @@ namespace UE::Dataflow
 		RegisterCoreNodes();
 		RegisterDataflowMathNodes();
 		RegisterDataflowVectorNodes();
+		UE::Dataflow::RegisterAnyTypes();
 	}
 
 	TSharedPtr<FDataflowNode> FNodeFactory::NewNodeFromRegisteredType(FGraph& Graph, const FNewNodeParameters& Param)
diff --git a/Engine/Source/Runtime/Experimental/Dataflow/Core/Public/Dataflow/DataflowAnyType.h b/Engine/Source/Runtime/Experimental/Dataflow/Core/Public/Dataflow/DataflowAnyType.h
--- a/Engine/Source/Runtime/Experimental/Dataflow/Core/Public/Dataflow/DataflowAnyType.h
+++ b/Engine/Source/Runtime/Experimental/Dataflow/Core/Public/Dataflow/DataflowAnyType.h
@@ -205,3 +205,9 @@ struct FDataflowTransformArrayTypes : public FDataflowAnyType
 	TArray<FTransform> Value;
 };
 
+namespace UE::Dataflow
+{
+	DATAFLOWCORE_API bool AreTypesCompatible(FName TypeA, FName TypeB);
+
+	DATAFLOWCORE_API void RegisterAnyTypes();
+}
\ No newline at end of file
diff --git a/Engine/Source/Runtime/Experimental/Dataflow/Core/Public/Dataflow/DataflowAnyTypeRegistry.h b/Engine/Source/Runtime/Experimental/Dataflow/Core/Public/Dataflow/DataflowAnyTypeRegistry.h
new file mode 100644
--- /dev/null
+++ b/Engine/Source/Runtime/Experimental/Dataflow/Core/Public/Dataflow/DataflowAnyTypeRegistry.h
@@ -0,0 +1,38 @@
+// Copyright Epic Games, Inc. All Rights Reserved.
+
+#pragma once
+
+#include "UObject/NameTypes.h"
+#include "Containers/Map.h"
+
+namespace UE::Dataflow
+{
+	struct FAnyTypesRegistry
+	{
+	public:
+		template<typename T>
+		static void RegisterTypeStatic(FName TypeName)
+		{
+			GetInstance().RegisterType<T>(TypeName);
+		}
+		static bool AreTypesCompatibleStatic(FName TypeA, FName TypeB);
+		static FAnyTypesRegistry& GetInstance();
+
+	private:
+		typedef bool(*FSupportTypeFunction)(FName);
+
+		FAnyTypesRegistry() {};
+
+		template<typename T>
+		void RegisterType(FName TypeName)
+		{
+			SupportTypeFunctionByName.Add(TypeName, &T::FPolicyType::SupportsTypeStatic);
+		}
+
+		bool AreTypesCompatible(FName TypeA, FName TypeB) const;
+
+		TMap<FName,FSupportTypeFunction> SupportTypeFunctionByName;
+	};
+};
+
+#define UE_DATAFLOW_REGISTER_ANYTYPE(TYPE) UE::Dataflow::FAnyTypesRegistry::RegisterTypeStatic<TYPE>(#TYPE)
diff --git a/Engine/Source/Runtime/Experimental/Dataflow/Core/Public/Dataflow/DataflowTypePolicy.h b/Engine/Source/Runtime/Experimental/Dataflow/Core/Public/Dataflow/DataflowTypePolicy.h
--- a/Engine/Source/Runtime/Experimental/Dataflow/Core/Public/Dataflow/DataflowTypePolicy.h
+++ b/Engine/Source/Runtime/Experimental/Dataflow/Core/Public/Dataflow/DataflowTypePolicy.h
@@ -86,6 +86,11 @@ struct FDataflowAllTypesPolicy : public IDataflowTypePolicy
 		return true;
 	}
 
+	static bool SupportsTypeStatic(FName InType)
+	{
+		return true;
+	}
+
 	static IDataflowTypePolicy* GetInterface()
 	{
 		static FDataflowAllTypesPolicy Instance;
