Dataflow :expose Add SetNullValue to Datflow output API

diff --git a/Engine/Source/Runtime/Experimental/Dataflow/Core/Private/Dataflow/DataflowInputOutput.cpp b/Engine/Source/Runtime/Experimental/Dataflow/Core/Private/Dataflow/DataflowInputOutput.cpp
--- a/Engine/Source/Runtime/Experimental/Dataflow/Core/Private/Dataflow/DataflowInputOutput.cpp
+++ b/Engine/Source/Runtime/Experimental/Dataflow/Core/Private/Dataflow/DataflowInputOutput.cpp
@@ -364,11 +364,20 @@ void FDataflowOutput::ForwardInput(const FDataflowInput* Input, UE::Dataflow::FC
 		{
 			// if there's no connection we make a invalid reference 
 			// so when the input is going to pull the cached value , it will return a default value instead
-			Context.SetNullData(CacheKey(), Property, GetOwningNodeGuid(), GetOwningNodeValueHash(), GetOwningNodeTimestamp());
+			SetNullValue(Context);
 		}
 	}
 }
 
+void FDataflowOutput::SetNullValue(UE::Dataflow::FContext& Context) const
+{
+	if (Property && OwningNode)
+	{
+		// so when the input is going to pull the cached value , it will return a default value instead
+		Context.SetNullData(CacheKey(), Property, GetOwningNodeGuid(), GetOwningNodeValueHash(), GetOwningNodeTimestamp());
+	}
+}
+
 void FDataflowOutput::FixAndPropagateType(FName InType)
 {
 	check(InType.ToString().StartsWith(Type.ToString()));
diff --git a/Engine/Source/Runtime/Experimental/Dataflow/Core/Public/Dataflow/DataflowInputOutput.h b/Engine/Source/Runtime/Experimental/Dataflow/Core/Public/Dataflow/DataflowInputOutput.h
--- a/Engine/Source/Runtime/Experimental/Dataflow/Core/Public/Dataflow/DataflowInputOutput.h
+++ b/Engine/Source/Runtime/Experimental/Dataflow/Core/Public/Dataflow/DataflowInputOutput.h
@@ -212,6 +212,9 @@ public:
 			});
 	}
 
+	/** Set a null value : this means that the connected input will get a default value and that this output will not be set re-evaluate next time */
+	DATAFLOWCORE_API void SetNullValue(UE::Dataflow::FContext& Context) const;
+
 	template<class T>
 	const T& GetValue(UE::Dataflow::FContext& Context, const T& Default) const
 	{
