[NetPhysics] Implement per-actor settings for post-resim error correction.
- Also fix bug where the render interpolation would not apply the error offset for the last physics tick.
- Fixed the logic for directional decay which was not working as intended

diff --git a/Engine/Source/Runtime/Engine/Classes/Components/PrimitiveComponent.h b/Engine/Source/Runtime/Engine/Classes/Components/PrimitiveComponent.h
--- a/Engine/Source/Runtime/Engine/Classes/Components/PrimitiveComponent.h
+++ b/Engine/Source/Runtime/Engine/Classes/Components/PrimitiveComponent.h
@@ -2873,11 +2873,11 @@ public:
 	 */
 	ENGINE_API void SetRigidBodyReplicatedTarget(FRigidBodyState& UpdatedState, const FName BoneName = NAME_None, int32 ServerFrame = 0, int32 ServerHandle = 0);
 
-	/** Register this components physics object as a focal particle in Physics Repliocation LOD */
+	/** Register this components physics object as a focal particle in Physics Replication LOD */
 	UFUNCTION(BlueprintCallable, Category = Replication)
 	ENGINE_API void RegisterAsFocalPointInPhysicsReplicationLOD() const;
 
-	/** Unregister this components physics object from being a focal particle in Physics Repliocation LOD */
+	/** Unregister this components physics object from being a focal particle in Physics Replication LOD */
 	UFUNCTION(BlueprintCallable, Category = Replication)
 	ENGINE_API void UnregisterAsFocalPointInPhysicsReplicationLOD() const;
 
diff --git a/Engine/Source/Runtime/Engine/Private/PhysicsEngine/NetworkPhysicsSettingsComponent.cpp b/Engine/Source/Runtime/Engine/Private/PhysicsEngine/NetworkPhysicsSettingsComponent.cpp
--- a/Engine/Source/Runtime/Engine/Private/PhysicsEngine/NetworkPhysicsSettingsComponent.cpp
+++ b/Engine/Source/Runtime/Engine/Private/PhysicsEngine/NetworkPhysicsSettingsComponent.cpp
@@ -12,7 +12,10 @@
 #include "PBDRigidsSolver.h"
 #include "PhysicsReplication.h"
 #include "Components/PrimitiveComponent.h"
-#include "Chaos/PhysicsObjectInternalInterface.h"
+#include "PhysicsEngine/PhysicsObjectExternalInterface.h"
+#include "PhysicsProxy/SingleParticlePhysicsProxy.h"
+#include "PhysicsProxy/ClusterUnionPhysicsProxy.h"
+#include "PhysicsProxy/GeometryCollectionPhysicsProxy.h"
 
 namespace PhysicsReplicationCVars
 {
@@ -79,7 +82,7 @@ void UNetworkPhysicsSettingsComponent::InitializeComponent()
 					{
 						if (UPrimitiveComponent* RootPrimComp = Cast<UPrimitiveComponent>(Owner->GetRootComponent()))
 						{
-							if (Chaos::FConstPhysicsObjectHandle PhysicsObject = RootPrimComp->GetPhysicsObjectByName(NAME_None))
+							if (Chaos::FPhysicsObjectHandle PhysicsObject = RootPrimComp->GetPhysicsObjectByName(NAME_None))
 							{
 								if (FNetworkPhysicsSettingsAsyncInput* AsyncInput = NetworkPhysicsSettings_Internal->GetProducerInputData_External())
 								{
@@ -94,6 +97,9 @@ void UNetworkPhysicsSettingsComponent::InitializeComponent()
 									};
 									AsyncInput->Settings = Settings;
 								}
+
+								// Apply resimulation error correction settings for render interpolation to the physics proxy
+								ResimulationSettings.ResimulationErrorCorrectionSettings.ApplySettings_External(PhysicsObject);
 							}
 						}
 					}
@@ -152,7 +158,7 @@ void UNetworkPhysicsSettingsComponent::OnComponentPhysicsStateChanged(UPrimitive
 {
 	if (StateChange == EComponentPhysicsStateChange::Created)
 	{
-		if (Chaos::FConstPhysicsObjectHandle PhysicsObject = ChangedComponent->GetPhysicsObjectByName(NAME_None))
+		if (Chaos::FPhysicsObjectHandle PhysicsObject = ChangedComponent->GetPhysicsObjectByName(NAME_None))
 		{
 			if (NetworkPhysicsSettings_Internal)
 			{
@@ -161,6 +167,9 @@ void UNetworkPhysicsSettingsComponent::OnComponentPhysicsStateChanged(UPrimitive
 					AsyncInput->PhysicsObject = PhysicsObject;
 				}
 			}
+
+			// Apply resimulation error correction settings for render interpolation to the physics proxy
+			ResimulationSettings.ResimulationErrorCorrectionSettings.ApplySettings_External(PhysicsObject);
 		}
 
 		RegisterInPhysicsReplicationLOD();
@@ -189,6 +198,76 @@ UNetworkPhysicsSettingsComponent* UNetworkPhysicsSettingsComponent::GetSettingsF
 	return Value ? *Value : nullptr;
 }
 
+void FNetworkPhysicsSettingsResimulationErrorCorrection::ApplySettings_External(Chaos::FPhysicsObjectHandle PhysicsObject)
+{
+	if (!PhysicsObject)
+	{
+		return;
+	}
+
+	if (bOverrideResimErrorInterpolationSettings)
+	{
+		FLockedWritePhysicsObjectExternalInterface Interface = FPhysicsObjectExternalInterface::LockWrite({ &PhysicsObject, 1 });
+
+		if (const Chaos::FGeometryParticle* Particle = Interface->GetParticle(PhysicsObject))
+		{
+			if (!Particle->GetProxy())
+			{
+				return;
+			}
+
+			auto ApplySettingsHelper = [this](auto Proxy)
+			{
+				// Get the proxies error interpolation, create one if not already present
+				if (FProxyInterpolationError* InterpolationError = Proxy->template GetOrCreateErrorInterpolationData<FProxyInterpolationError>())
+				{
+					// Apply the custom settings
+					FErrorInterpolationSettings& Settings = InterpolationError->GetOrCreateErrorInterpolationSettings();
+					Settings.ErrorCorrectionDuration = ResimErrorCorrectionDuration;
+					Settings.MaximumErrorCorrectionBeforeSnapping = ResimErrorMaximumDistanceBeforeSnapping;
+					Settings.ErrorDirectionalDecayMultiplier = ResimErrorDirectionalDecayMultiplier;
+				}
+			};
+
+			switch (Particle->GetProxy()->GetType())
+			{
+				case EPhysicsProxyType::SingleParticleProxy:
+				{
+					Chaos::FSingleParticlePhysicsProxy* SPProxy = static_cast<Chaos::FSingleParticlePhysicsProxy*>(Particle->GetProxy());
+					if (ensure(SPProxy))
+					{
+						ApplySettingsHelper(SPProxy);
+					}
+					break;
+				}
+				case EPhysicsProxyType::ClusterUnionProxy:
+				{
+					Chaos::FClusterUnionPhysicsProxy* CUProxy = static_cast<Chaos::FClusterUnionPhysicsProxy*>(Particle->GetProxy());
+					if (ensure(CUProxy))
+					{
+						ApplySettingsHelper(CUProxy);
+					}
+					break;
+				}
+				case EPhysicsProxyType::GeometryCollectionType:
+				{
+					FGeometryCollectionPhysicsProxy* GCProxy = static_cast<FGeometryCollectionPhysicsProxy*>(Particle->GetProxy());
+					if (ensure(GCProxy))
+					{
+						ApplySettingsHelper(GCProxy);
+					}
+					break;
+				}
+				default:
+				{
+					ensure(false);
+					break;
+				}
+			};
+		}
+	}
+}
+
 #pragma region // FNetworkPhysicsSettingsComponentAsync
 
 FNetworkPhysicsSettingsComponentAsync::FNetworkPhysicsSettingsComponentAsync() : TSimCallbackObject()
diff --git a/Engine/Source/Runtime/Engine/Public/Physics/NetworkPhysicsSettingsComponent.h b/Engine/Source/Runtime/Engine/Public/Physics/NetworkPhysicsSettingsComponent.h
--- a/Engine/Source/Runtime/Engine/Public/Physics/NetworkPhysicsSettingsComponent.h
+++ b/Engine/Source/Runtime/Engine/Public/Physics/NetworkPhysicsSettingsComponent.h
@@ -20,10 +20,7 @@ class FNetworkPhysicsSettingsComponentAsync;
 using FDefaultReplicationSettings		= FNetworkPhysicsSettingsDefaultReplication;
 using FPredictiveInterpolationSettings	= FNetworkPhysicsSettingsPredictiveInterpolation;
 using FResimulationSettings				= FNetworkPhysicsSettingsResimulation;
-/*
-using FRewindSettings					= FNetworkPhysicsSettingsRewindData;
-using FRenderInterpolationSettings		= FNetworkPhysicsSettingsRenderInterpolation;
-*/
+using FRenderInterpolationSettings		= FNetworkPhysicsSettingsResimulationErrorCorrection;
 
 namespace PhysicsReplicationCVars
 {
@@ -248,12 +245,46 @@ struct FNetworkPhysicsSettingsPredictiveInterpolation
 
 };
 
+USTRUCT()
+struct FNetworkPhysicsSettingsResimulationErrorCorrection
+{
+	FNetworkPhysicsSettingsResimulationErrorCorrection()
+	: bOverrideResimErrorInterpolationSettings(0)
+	, ResimErrorCorrectionDuration(0.5f)
+	, ResimErrorMaximumDistanceBeforeSnapping(250.0f)
+	, ResimErrorDirectionalDecayMultiplier(0.0f)
+	{};
+
+	GENERATED_BODY()
+
+	/** Apply these settings to Physics Object */
+	void ApplySettings_External(Chaos::FPhysicsObjectHandle PhysicsObject);
+
+	/** Enable override for post-resimulation error correction settings during render interpolation
+	* NOTE: This currently does not work if the experimental p.RenderInterp.ErrorVelocityCorrection CVar is set to true (false by default).*/
+	UPROPERTY(config, EditDefaultsOnly, Category = "Overrides")
+	uint32 bOverrideResimErrorInterpolationSettings : 1;
+
+	// Overrides CVar: p.RenderInterp.ErrorCorrectionDuration -- How long in seconds to apply error correction over
+	UPROPERTY(config, EditDefaultsOnly, Category = "Overrides", Meta = (EditCondition = "bOverrideResimErrorInterpolationSettings", ClampMin = "0.0", UIMin = "0.0"))
+	float ResimErrorCorrectionDuration;
+
+	// Overrides CVar: p.RenderInterp.MaximumErrorCorrectionBeforeSnapping -- Maximum error correction in cm before we stop interpolating and snap to target
+	UPROPERTY(config, EditDefaultsOnly, Category = "Overrides", Meta = (EditCondition = "bOverrideResimErrorInterpolationSettings", ClampMin = "0.0", UIMin = "0.0"))
+	float ResimErrorMaximumDistanceBeforeSnapping;
+
+	// Overrides CVar: p.RenderInterp.DirectionalDecayMultiplier -- Decay error offset in the direction that the physics object is moving, value is multiplier of projected offset direction, 0.25 means a 25% decay of the magnitude in the direction of physics travel. Deactivate by setting to 0..
+	UPROPERTY(config, EditDefaultsOnly, Category = "Overrides", Meta = (EditCondition = "bOverrideResimErrorInterpolationSettings", ClampMin = "0.0", UIMin = "0.0", ClampMax = "1.0", UIMax = "1.0"))
+	float ResimErrorDirectionalDecayMultiplier;
+};
+
 USTRUCT()
 struct FNetworkPhysicsSettingsResimulation
 {
 	FNetworkPhysicsSettingsResimulation()
 	: bOverrideResimulationErrorThreshold_DEPRECATED(0)
 	, ResimulationErrorThreshold_DEPRECATED(10)
+	, ResimulationErrorCorrectionSettings()
 	{};
 
 	GENERATED_BODY()
@@ -345,6 +376,9 @@ struct FNetworkPhysicsSettingsResimulation
 	UPROPERTY(config, EditDefaultsOnly, Category = "Overrides", Meta = (EditCondition = "bOverrideAngVelStabilityMultiplier"))
 	float AngVelStabilityMultiplier = PhysicsReplicationCVars::ResimulationCVars::AngVelStabilityMultiplier;
 	float GetAngVelStabilityMultiplier() { return bOverrideAngVelStabilityMultiplier ? AngVelStabilityMultiplier : PhysicsReplicationCVars::ResimulationCVars::AngVelStabilityMultiplier; }
+
+	UPROPERTY(config, EditDefaultsOnly, Category = "Overrides", meta = (DisplayName = "Post-Resimulation Error Correction Render Interpolation Settings"))
+	FNetworkPhysicsSettingsResimulationErrorCorrection ResimulationErrorCorrectionSettings;
 };
 
 USTRUCT()
@@ -559,7 +593,7 @@ struct FNetworkPhysicsSettingsAsyncInput : public Chaos::FSimCallbackInput
 	}
 };
 
-class FNetworkPhysicsSettingsComponentAsync : public Chaos::TSimCallbackObject<FNetworkPhysicsSettingsAsyncInput>
+class FNetworkPhysicsSettingsComponentAsync : public Chaos::TSimCallbackObject<FNetworkPhysicsSettingsAsyncInput, Chaos::FSimCallbackNoOutput, Chaos::ESimCallbackOptions::Presimulate>
 {
 public:
 	FNetworkPhysicsSettingsComponentAsync();
diff --git a/Engine/Source/Runtime/Experimental/Chaos/Private/PhysicsProxy/ClusterUnionPhysicsProxy.cpp b/Engine/Source/Runtime/Experimental/Chaos/Private/PhysicsProxy/ClusterUnionPhysicsProxy.cpp
--- a/Engine/Source/Runtime/Experimental/Chaos/Private/PhysicsProxy/ClusterUnionPhysicsProxy.cpp
+++ b/Engine/Source/Runtime/Experimental/Chaos/Private/PhysicsProxy/ClusterUnionPhysicsProxy.cpp
@@ -489,12 +489,19 @@ namespace Chaos
 			InterpData = GetOrCreateErrorInterpolationData<FProxyInterpolationError>();
 			check(InterpData);
 
+			FReal MaxErrorCorrection = RenderInterpolationCVars::RenderInterpMaximumErrorCorrectionBeforeSnapping;
+			FReal ErrorCorrectionDuration = RenderInterpolationCVars::RenderInterpErrorCorrectionDuration;
+			if (FErrorInterpolationSettings* ErrorSettings = InterpData->GetErrorInterpolationSettings())
+			{
+				MaxErrorCorrection = ErrorSettings->MaximumErrorCorrectionBeforeSnapping;
+				ErrorCorrectionDuration = ErrorSettings->ErrorCorrectionDuration;
+			}
+
 			// If error is within interpolation limit, set the number of physics frames to interpolate over, else leave at 0 frames to instantly correct the error.
-			const FReal MaxErrorCorrection = RenderInterpolationCVars::RenderInterpMaximumErrorCorrectionBeforeSnapping;
 			int32 RenderInterpErrorCorrectionDurationTicks = 0;
 			if (Error->ErrorX.SizeSquared() < (MaxErrorCorrection * MaxErrorCorrection))
 			{
-				RenderInterpErrorCorrectionDurationTicks = FMath::FloorToInt32(RenderInterpolationCVars::RenderInterpErrorCorrectionDuration / AsyncFixedTimeStep); // Convert duration from seconds to simulation ticks
+				RenderInterpErrorCorrectionDurationTicks = FMath::FloorToInt32(ErrorCorrectionDuration / AsyncFixedTimeStep); // Convert duration from seconds to simulation ticks
 			}
 			InterpData->AccumlateErrorXR(Error->ErrorX, Error->ErrorR, SolverSyncTimestamp, RenderInterpErrorCorrectionDurationTicks);
 		}
@@ -578,9 +585,9 @@ namespace Chaos
 
 			if (const FVec3* Prev = LerpHelper(PullData.X, ProxyTimestamp->OverWriteX))
 			{
-				if (bHasInterpolationData && RenderInterpolationCVars::RenderInterpErrorDirectionalDecayMultiplier > 0.0f)
+				if (bHasInterpolationData)
 				{
-					DirectionalDecayPerformed = InterpData->DirectionalDecay(NextPullData->X - *Prev);
+					DirectionalDecayPerformed = InterpData->DirectionalDecay(NextPullData->X - *Prev, RenderInterpolationCVars::RenderInterpErrorDirectionalDecayMultiplier);
 				}
 
 				const FVec3 NewX = bIsReplicationErrorSmoothing ?
diff --git a/Engine/Source/Runtime/Experimental/Chaos/Private/PhysicsProxy/GeometryCollectionPhysicsProxy.cpp b/Engine/Source/Runtime/Experimental/Chaos/Private/PhysicsProxy/GeometryCollectionPhysicsProxy.cpp
--- a/Engine/Source/Runtime/Experimental/Chaos/Private/PhysicsProxy/GeometryCollectionPhysicsProxy.cpp
+++ b/Engine/Source/Runtime/Experimental/Chaos/Private/PhysicsProxy/GeometryCollectionPhysicsProxy.cpp
@@ -4303,12 +4303,19 @@ bool FGeometryCollectionPhysicsProxy::PullFromPhysicsState(const Chaos::FDirtyGe
 		InterpData = GetOrCreateErrorInterpolationData<FProxyInterpolationError>();
 		check(InterpData);
 
+		Chaos::FReal MaxErrorCorrection = RenderInterpolationCVars::RenderInterpMaximumErrorCorrectionBeforeSnapping;
+		Chaos::FReal ErrorCorrectionDuration = RenderInterpolationCVars::RenderInterpErrorCorrectionDuration;
+		if (FErrorInterpolationSettings* ErrorSettings = InterpData->GetErrorInterpolationSettings())
+		{
+			MaxErrorCorrection = ErrorSettings->MaximumErrorCorrectionBeforeSnapping;
+			ErrorCorrectionDuration = ErrorSettings->ErrorCorrectionDuration;
+		}
+
 		// If error is within interpolation limit, set the number of physics frames to interpolate over, else leave at 0 frames to instantly correct the error.
-		const Chaos::FReal MaxErrorCorrection = RenderInterpolationCVars::RenderInterpMaximumErrorCorrectionBeforeSnapping;
 		int32 RenderInterpErrorCorrectionDurationTicks = 0;
 		if (Error->ErrorX.SizeSquared() < (MaxErrorCorrection * MaxErrorCorrection))
 		{
-			RenderInterpErrorCorrectionDurationTicks = FMath::FloorToInt32(RenderInterpolationCVars::RenderInterpErrorCorrectionDuration / AsyncFixedTimeStep); // Convert duration from seconds to simulation ticks
+			RenderInterpErrorCorrectionDurationTicks = FMath::FloorToInt32(ErrorCorrectionDuration / AsyncFixedTimeStep); // Convert duration from seconds to simulation ticks
 		}
 		InterpData->AccumlateErrorXR(Error->ErrorX, Error->ErrorR, SolverSyncTimestamp, RenderInterpErrorCorrectionDurationTicks);
 	}
@@ -4343,9 +4350,9 @@ bool FGeometryCollectionPhysicsProxy::PullFromPhysicsState(const Chaos::FDirtyGe
 			{
 				InterpData->UpdateError(SolverSyncTimestamp, AsyncFixedTimeStep);
 
-				if (RenderInterpolationCVars::RenderInterpErrorDirectionalDecayMultiplier > 0.0f && PrevResults.GetNumEntries() && NextResults.GetNumEntries())
+				if (PrevResults.GetNumEntries() && NextResults.GetNumEntries())
 				{
-					InterpData->DirectionalDecay(NextResults.GetPositions(0).ParticleX - PrevResults.GetPositions(0).ParticleX);
+					InterpData->DirectionalDecay(NextResults.GetPositions(0).ParticleX - PrevResults.GetPositions(0).ParticleX, RenderInterpolationCVars::RenderInterpErrorDirectionalDecayMultiplier);
 				}
 			}
 
diff --git a/Engine/Source/Runtime/Experimental/Chaos/Private/PhysicsProxy/SingleParticlePhysicsProxy.cpp b/Engine/Source/Runtime/Experimental/Chaos/Private/PhysicsProxy/SingleParticlePhysicsProxy.cpp
--- a/Engine/Source/Runtime/Experimental/Chaos/Private/PhysicsProxy/SingleParticlePhysicsProxy.cpp
+++ b/Engine/Source/Runtime/Experimental/Chaos/Private/PhysicsProxy/SingleParticlePhysicsProxy.cpp
@@ -305,13 +305,20 @@ bool FSingleParticlePhysicsProxy::PullFromPhysicsState(const Chaos::FDirtyRigidP
 				InterpData = GetOrCreateErrorInterpolationData<FProxyInterpolationError>();
 			}
 			check(InterpData);
+			
+			FReal MaxErrorCorrection = RenderInterpolationCVars::RenderInterpMaximumErrorCorrectionBeforeSnapping;
+			FReal ErrorCorrectionDuration = RenderInterpolationCVars::RenderInterpErrorCorrectionDuration;
+			if (FErrorInterpolationSettings * ErrorSettings = InterpData->GetErrorInterpolationSettings())
+			{
+				MaxErrorCorrection = ErrorSettings->MaximumErrorCorrectionBeforeSnapping;
+				ErrorCorrectionDuration = ErrorSettings->ErrorCorrectionDuration;
+			}
 
 			// If error is within interpolation limit, set the number of physics frames to interpolate over, else leave at 0 frames to instantly correct the error.
 			int32 RenderInterpErrorCorrectionDurationTicks = 0;
-			const FReal MaxErrorCorrection = RenderInterpolationCVars::RenderInterpMaximumErrorCorrectionBeforeSnapping;
 			if (Error->ErrorX.SizeSquared() < (MaxErrorCorrection * MaxErrorCorrection))
 			{
-				RenderInterpErrorCorrectionDurationTicks = FMath::FloorToInt32(RenderInterpolationCVars::RenderInterpErrorCorrectionDuration / AsyncFixedTimeStep); // Convert duration from seconds to simulation ticks
+				RenderInterpErrorCorrectionDurationTicks = FMath::FloorToInt32(ErrorCorrectionDuration / AsyncFixedTimeStep); // Convert duration from seconds to simulation ticks
 			}
 			InterpData->AccumlateErrorXR(Error->ErrorX, Error->ErrorR, SolverSyncTimestamp, RenderInterpErrorCorrectionDurationTicks);
 			
@@ -347,10 +354,7 @@ bool FSingleParticlePhysicsProxy::PullFromPhysicsState(const Chaos::FDirtyRigidP
 					FVec3 Target = FMath::Lerp(*Prev, NextPullData->X, *Alpha);
 					if (bIsReplicationErrorSmoothing)
 					{
-						if (RenderInterpolationCVars::RenderInterpErrorDirectionalDecayMultiplier > 0.0f)
-						{
-							DirectionalDecayPerformed = InterpData->DirectionalDecay(NextPullData->X - *Prev);
-						}
+						DirectionalDecayPerformed = InterpData->DirectionalDecay(NextPullData->X - *Prev, RenderInterpolationCVars::RenderInterpErrorDirectionalDecayMultiplier);
 
 						Target += InterpData->GetErrorX(*Alpha);
 
diff --git a/Engine/Source/Runtime/Experimental/Chaos/Public/Chaos/Framework/PhysicsProxyBase.h b/Engine/Source/Runtime/Experimental/Chaos/Public/Chaos/Framework/PhysicsProxyBase.h
--- a/Engine/Source/Runtime/Experimental/Chaos/Public/Chaos/Framework/PhysicsProxyBase.h
+++ b/Engine/Source/Runtime/Experimental/Chaos/Public/Chaos/Framework/PhysicsProxyBase.h
@@ -203,6 +203,24 @@ struct PhysicsProxyWrapper
 	EPhysicsProxyType Type;
 };
 
+struct FErrorInterpolationSettings
+{
+	FErrorInterpolationSettings()
+		: ErrorCorrectionDuration(0.5f)
+		, MaximumErrorCorrectionBeforeSnapping(250.0f)
+		, ErrorDirectionalDecayMultiplier(0.0f)
+	{}
+	~FErrorInterpolationSettings() { }
+
+	// How long in seconds to apply error correction over
+	float ErrorCorrectionDuration;
+
+	// Maximum error correction in cm before we stop interpolating and snap to target
+	float MaximumErrorCorrectionBeforeSnapping;
+
+	// Decay error offset in the direction that the physics object is moving, value is multiplier of projected offset direction, 0.25 means a 25% decay of the magnitude in the direction of physics travel. Deactivate by setting to 0
+	float ErrorDirectionalDecayMultiplier;
+};
 
 struct FProxyInterpolationBase
 {
@@ -247,7 +265,12 @@ public:
 	virtual const bool UpdateError(const int32 CurrentSimTick, const Chaos::FReal AsyncFixedTimeStep) { return false; }
 
 	/** EXPERIMENTAL - Decay error based on moved direction and distance */
-	virtual const bool DirectionalDecay(Chaos::FVec3 Direction) { return false; }
+	UE_DEPRECATED(5.6, "Deprecated, use DirectionalDecay(Direction, ErrorDirectionalDecayMultiplier) instead")
+	virtual const bool DirectionalDecay(Chaos::FVec3 Direction) { return DirectionalDecay(Direction, 0.0f); }
+
+	/** EXPERIMENTAL - Decay error based on moved direction and distance
+	* @param ErrorDirectionalDecayMultiplier is a multiplier where 0.25 means a 25% decay along the direction of physics movement that aligns with the error offset direction, parameter gets overridden if there are custom settings for this particle */
+	virtual const bool DirectionalDecay(Chaos::FVec3 Direction, const float ErrorDirectionalDecayMultiplier) { return false; }
 
 	/** EXPERIMENTAL - If currently correcting error while taking velocity into account */
 	virtual const bool IsErrorVelocitySmoothing() const { return false; }
@@ -260,6 +283,9 @@ public:
 
 	/** EXPERIMENTAL - Register the current velocity and position for use in velocity correction calculations */
 	virtual void SetVelocitySmoothing(const Chaos::FVec3 CurrV, const Chaos::FVec3 CurrX, const int32 ErrorVelocitySmoothDuration) { };
+
+	/** Get FErrorInterpolationSettings which stores custom settings for render interpolation error corrections, returns nullptr if there are no custom settings */
+	virtual FErrorInterpolationSettings* GetErrorInterpolationSettings() { return nullptr; }
 };
 
 // Render interpolation that can correct errors from resimulation / repositions through a linear decay over N simulation tick.
@@ -274,7 +300,7 @@ struct FProxyInterpolationError : FProxyInterpolationBase
 	static const EProxyInterpolationType InterpolationType = EProxyInterpolationType::ErrorLinear;
 	virtual const EProxyInterpolationType GetInterpolationType() const override { return FProxyInterpolationError::InterpolationType; }
 
-	virtual const bool IsErrorSmoothing() const override { return ErrorSmoothingCount > 0; }
+	virtual const bool IsErrorSmoothing() const override { return LastSimTick < EndDecayTick; }
 
 	virtual const Chaos::FVec3 GetErrorX(const Chaos::FRealSingle Alpha) const override { return FMath::Lerp(ErrorXPrev, ErrorX, Alpha); }
 
@@ -282,8 +308,8 @@ struct FProxyInterpolationError : FProxyInterpolationBase
 
 	virtual void AccumlateErrorXR(const Chaos::FVec3 X, const FQuat R, const int32 CurrentSimTick, const int32 ErrorSmoothDuration) override
 	{
-		ErrorSmoothingCount = ErrorSmoothDuration; // How many simulation ticks to correct error over
 		LastSimTick = CurrentSimTick - 1; // Error is from the previous simulation tick, not the current
+		EndDecayTick = LastSimTick + ErrorSmoothDuration;
 		SimTicks = 0;
 
 		if (IsErrorSmoothing())
@@ -312,16 +338,24 @@ struct FProxyInterpolationError : FProxyInterpolationBase
 		return false;
 	}
 
-	virtual const bool DirectionalDecay(Chaos::FVec3 Direction) override
+	virtual const bool DirectionalDecay(Chaos::FVec3 Direction, float ErrorDirectionalDecayMultiplier) override
 	{
-		if (IsErrorSmoothing() && SimTicks > 0)
+		ErrorDirectionalDecayMultiplier = ErrorInterpolationSettings.Get() ? ErrorInterpolationSettings.Get()->ErrorDirectionalDecayMultiplier : ErrorDirectionalDecayMultiplier;
+		if (ErrorDirectionalDecayMultiplier > 0.0f && IsErrorSmoothing() && SimTicks > 0)
 		{
-			const Chaos::FVec3 DirectionNormal = Direction.GetSafeNormal();
-			Chaos::FRealDouble DotProd = Chaos::FVec3::DotProduct(DirectionNormal, ErrorX);
+			Chaos::FVec3 DirProjection = Direction.ProjectOnTo(ErrorX) * ErrorDirectionalDecayMultiplier;
+			const Chaos::FRealDouble DotProd = Chaos::FVec3::DotProduct(DirProjection, ErrorX);
+			
 			if (DotProd > 0.0f)
 			{
-				Chaos::FVec3 DirProjection = ErrorX.ProjectOnToNormal(DirectionNormal) * RenderInterpolationCVars::RenderInterpErrorDirectionalDecayMultiplier;
-				ErrorX -= DirProjection;
+				if (DirProjection.SizeSquared() < ErrorX.SizeSquared())
+				{
+					ErrorX -= DirProjection;
+				}
+				else
+				{
+					ErrorX = Chaos::FVec3::Zero();
+				}
 				return true;
 			}
 		}
@@ -331,6 +365,11 @@ struct FProxyInterpolationError : FProxyInterpolationBase
 protected:
 	virtual const bool DecayError()
 	{
+		if (SimTicks <= 0)
+		{
+			return false;
+		}
+
 		if (!IsErrorSmoothing())
 		{
 			Reset();
@@ -343,7 +382,9 @@ protected:
 		// Second step: 8/9  = 0.888 |  90 * 0.888 = 80 error
 		// Third step: 7/8  = 0.875 |  80 * 0.875 = 70 error
 		// etc.
-		for (int32 Ticks = FMath::Clamp(SimTicks, 0, ErrorSmoothingCount); Ticks > 0; Ticks--)
+		const int32 ErrorSmoothingCount = EndDecayTick - LastSimTick;
+		int32 Ticks = FMath::Clamp(SimTicks, 0, ErrorSmoothingCount);
+		for (; Ticks > 0; Ticks--)
 		{
 			Chaos::FRealSingle Alpha = Chaos::FRealSingle(ErrorSmoothingCount - 1) / Chaos::FRealSingle(ErrorSmoothingCount);
 			Alpha = FMath::Clamp(Alpha, 0.f, 1.f);
@@ -352,11 +393,9 @@ protected:
 			ErrorX *= Alpha;
 			ErrorRPrev = ErrorR;
 			ErrorR = FMath::Lerp(FQuat::Identity, ErrorR, Alpha);
-
-			ErrorSmoothingCount = FMath::Max(ErrorSmoothingCount - 1, 0);
 		}
 
-		return SimTicks > 0;
+		return true;
 	}
 
 	virtual void Reset()
@@ -365,20 +404,39 @@ protected:
 		ErrorXPrev = Chaos::FVec3::ZeroVector;
 		ErrorR = FQuat::Identity;
 		ErrorRPrev = FQuat::Identity;
-		ErrorSmoothingCount = 0;
+		EndDecayTick = 0;
 		LastSimTick = 0;
 		SimTicks = 0;
 	}
 
+public:
+	/** Get or create FErrorInterpolationSettings to stores custom settings for render interpolation error corrections */
+	FErrorInterpolationSettings& GetOrCreateErrorInterpolationSettings()
+	{
+		if (!ErrorInterpolationSettings.IsValid())
+		{
+			ErrorInterpolationSettings = MakeUnique<FErrorInterpolationSettings>();
+		}
+		return *ErrorInterpolationSettings.Get();
+	}
+
+	/** Get FErrorInterpolationSettings which stores custom settings for render interpolation error corrections, returns nullptr if there are no custom settings */
+	virtual FErrorInterpolationSettings* GetErrorInterpolationSettings() override
+	{
+		return ErrorInterpolationSettings.Get();
+	}
+
 protected:
 	int32 LastSimTick = 0;
 	int32 SimTicks = 0;
+	int32 EndDecayTick = 0;
 
-	Chaos::FVec3 ErrorX = { 0,0,0 };
-	Chaos::FVec3 ErrorXPrev = { 0,0,0 };
+	Chaos::FVec3 ErrorX = Chaos::FVec3::Zero();
+	Chaos::FVec3 ErrorXPrev = Chaos::FVec3::Zero();
 	FQuat ErrorR = FQuat::Identity;
 	FQuat ErrorRPrev = FQuat::Identity;
-	int32 ErrorSmoothingCount = 0;
+
+	TUniquePtr<FErrorInterpolationSettings> ErrorInterpolationSettings = nullptr;
 };
 
 
@@ -431,8 +489,8 @@ protected:
 		}
 	}
 
-	Chaos::FVec3 ErrorVelocitySmoothingV = { 0,0,0 };
-	Chaos::FVec3 ErrorVelocitySmoothingX = { 0,0,0 };
-	Chaos::FVec3 ErrorVelocitySmoothingXPrev = { 0,0,0 };
+	Chaos::FVec3 ErrorVelocitySmoothingV = Chaos::FVec3::Zero();
+	Chaos::FVec3 ErrorVelocitySmoothingX = Chaos::FVec3::Zero();
+	Chaos::FVec3 ErrorVelocitySmoothingXPrev = Chaos::FVec3::Zero();
 	int32 ErrorVelocitySmoothingCount = 0;
-};
\ No newline at end of file
+};
diff --git a/Engine/Source/Runtime/Experimental/Chaos/Public/PhysicsProxy/ClusterUnionPhysicsProxy.h b/Engine/Source/Runtime/Experimental/Chaos/Public/PhysicsProxy/ClusterUnionPhysicsProxy.h
--- a/Engine/Source/Runtime/Experimental/Chaos/Public/PhysicsProxy/ClusterUnionPhysicsProxy.h
+++ b/Engine/Source/Runtime/Experimental/Chaos/Public/PhysicsProxy/ClusterUnionPhysicsProxy.h
@@ -171,7 +171,8 @@ namespace Chaos
 
 		TUniquePtr<FProxyInterpolationBase> InterpolationData;
 
-		/** Get or create a derived FProxyInterpolationBase that handles error corrections */
+	public:
+		/** Get or create a derived FProxyInterpolationBase that handles render interpolation error corrections */
 		template<typename ErrorDataType>
 		ErrorDataType* GetOrCreateErrorInterpolationData()
 		{
@@ -187,6 +188,7 @@ namespace Chaos
 			return static_cast<ErrorDataType*>(InterpolationData.Get());
 		}
 
+	private:
 		// An array of a particles that exist in the external implicit object union.
 		// Note that this array should only be used for book-keeping. It is generally
 		// unsafe to try and access the particles within this array.
diff --git a/Engine/Source/Runtime/Experimental/Chaos/Public/PhysicsProxy/GeometryCollectionPhysicsProxy.h b/Engine/Source/Runtime/Experimental/Chaos/Public/PhysicsProxy/GeometryCollectionPhysicsProxy.h
--- a/Engine/Source/Runtime/Experimental/Chaos/Public/PhysicsProxy/GeometryCollectionPhysicsProxy.h
+++ b/Engine/Source/Runtime/Experimental/Chaos/Public/PhysicsProxy/GeometryCollectionPhysicsProxy.h
@@ -762,7 +762,8 @@ private:
 
 	TUniquePtr<FProxyInterpolationBase> InterpolationData;
 
-	/** Get or create a derived FProxyInterpolationBase that handles error corrections */
+public:
+	/** Get or create a derived FProxyInterpolationBase that handles render interpolation error corrections */
 	template<typename ErrorDataType>
 	ErrorDataType* GetOrCreateErrorInterpolationData()
 	{
@@ -778,6 +779,7 @@ private:
 		return static_cast<ErrorDataType*>(InterpolationData.Get());
 	}
 
+private:
 #if WITH_EDITORONLY_DATA
 	// this is used as a unique ID when collecting data from runtime
 	FGuid CollectorGuid;
diff --git a/Engine/Source/Runtime/Experimental/Chaos/Public/PhysicsProxy/SingleParticlePhysicsProxy.h b/Engine/Source/Runtime/Experimental/Chaos/Public/PhysicsProxy/SingleParticlePhysicsProxy.h
--- a/Engine/Source/Runtime/Experimental/Chaos/Public/PhysicsProxy/SingleParticlePhysicsProxy.h
+++ b/Engine/Source/Runtime/Experimental/Chaos/Public/PhysicsProxy/SingleParticlePhysicsProxy.h
@@ -172,17 +172,7 @@ public:
 		return Reference.Get();
 	}
 
-protected:
-	TUniquePtr<PARTICLE_TYPE> Particle;
-	FParticleHandle* Handle;
-	FPhysicsObjectUniquePtr Reference;
-
-	int32 GravityGroupIndex;
-
-private:
-	TUniquePtr<FProxyInterpolationBase> InterpolationData;
-
-	/** Get or create a derived FProxyInterpolationBase that handles error corrections */
+	/** Get or create a derived FProxyInterpolationBase that handles render interpolation error corrections */
 	template<typename ErrorDataType>
 	ErrorDataType* GetOrCreateErrorInterpolationData()
 	{
@@ -198,6 +188,16 @@ private:
 		return static_cast<ErrorDataType*>(InterpolationData.Get());
 	}
 
+protected:
+	TUniquePtr<PARTICLE_TYPE> Particle;
+	FParticleHandle* Handle;
+	FPhysicsObjectUniquePtr Reference;
+
+	int32 GravityGroupIndex;
+
+private:
+	TUniquePtr<FProxyInterpolationBase> InterpolationData;
+
 	//use static Create
 	CHAOS_API FSingleParticlePhysicsProxy(TUniquePtr<PARTICLE_TYPE>&& InParticle, FParticleHandle* InHandle, UObject* InOwner = nullptr);
 };
