Provide more info when ''ias.DisplayOnScreenStatistics' is enabled.

#rnx

- Moved FBackendStatus out to it's own code file.
-- I left the associated cvars in OnDemandIoDispatcherBackend.cpp to make them easier to find.
- Removed EStatsFlags from FOnDemandIoBackendStats, we now pass in a reference to the owning backend's FBackendStatus instead. This allows the statistics system to gain access to the full array of statuses without replicating them all via it's own enum.
- Added a new line to 'ias.DisplayOnScreenStatistics' that details the current status of the backend. This line will be written as an error if the http connection is dropped.
- Removed the connection status from the stats line as that is now covered by the new line.

diff --git a/Engine/Source/Runtime/Experimental/IoStore/OnDemand/Private/OnDemandBackendStatus.cpp b/Engine/Source/Runtime/Experimental/IoStore/OnDemand/Private/OnDemandBackendStatus.cpp
new file mode 100644
--- /dev/null
+++ b/Engine/Source/Runtime/Experimental/IoStore/OnDemand/Private/OnDemandBackendStatus.cpp
@@ -0,0 +1,116 @@
+// Copyright Epic Games, Inc. All Rights Reserved.
+
+#include "OnDemandBackendStatus.h"
+
+#include "IO/IoStoreOnDemand.h"
+#include "Logging/StructuredLog.h"
+
+namespace UE::IoStore
+{
+
+extern bool GIasHttpEnabled;
+extern bool GIasHttpOptionalBulkDataEnabled;
+
+bool FBackendStatus::IsHttpEnabled(EIoChunkType ChunkType) const
+{
+	const uint8 CurrentFlags = Flags.load(std::memory_order_relaxed);
+	return IsHttpEnabled(CurrentFlags) &&
+		(ChunkType != EIoChunkType::OptionalBulkData ||
+		((CurrentFlags & uint8(EFlags::HttpBulkOptionalDisabled)) == 0 && GIasHttpOptionalBulkDataEnabled));
+}
+
+void FBackendStatus::ToString(FStringBuilderBase& Builder)
+{
+	if (IsCacheEnabled())
+	{
+		Builder << TEXT("Caching - Enabled");
+	}
+	else
+	{
+		Builder << TEXT("Caching - Disabled");
+	}
+
+	Builder << TEXT(" | ");
+
+	if (IsHttpError())
+	{
+		Builder << TEXT("Http - Disconnected");
+	}
+	else if (IsHttpEnabled())
+	{
+		Builder << TEXT("Http - Enabled");
+	}
+	else
+	{
+		Builder << TEXT("Http - Disabled");
+	}
+
+	Builder << TEXT(" | ");
+
+	if (IsHttpEnabled(EIoChunkType::OptionalBulkData))
+	{
+		Builder << TEXT("Optional Mips Enabled");
+	}
+	else
+	{
+		Builder << TEXT("Optional Mips Disabled");
+	}
+}
+
+bool FBackendStatus::IsHttpEnabled(uint8 FlagsToTest)
+{
+	constexpr uint8 HttpFlags = uint8(EFlags::HttpEnabled) | uint8(EFlags::HttpError);
+	return ((FlagsToTest & HttpFlags) == uint8(EFlags::HttpEnabled)) && GIasHttpEnabled;
+}
+
+void FBackendStatus::AddOrRemoveFlags(EFlags FlagsToAddOrRemove, bool bValue, const TCHAR* DebugText)
+{
+	const uint8 PrevFlags = AddOrRemoveFlags(FlagsToAddOrRemove, bValue);
+	
+	TStringBuilder<128> Sb;
+
+	Sb.Append(DebugText)
+		<< TEXT(" '");
+
+	Sb.Append(bValue ? TEXT("true") : TEXT("false"))
+		<< TEXT("', backend status '(")
+		<< EFlags(PrevFlags)
+		<< TEXT(") -> (")
+		<< EFlags(Flags.load(std::memory_order_relaxed))
+		<< TEXT(")'");
+
+	UE_LOGFMT(LogIas, Log, "{Message}", Sb);
+}
+
+FStringBuilderBase& operator<<(FStringBuilderBase& Sb, FBackendStatus::EFlags StatusFlags)
+{
+	if (StatusFlags == FBackendStatus::EFlags::None)
+	{
+		Sb.Append(TEXT("None"));
+		return Sb;
+	}
+
+	bool bFirst = true;
+	auto AppendIf = [StatusFlags, &Sb, &bFirst](FBackendStatus::EFlags Contains, const TCHAR* Str)
+		{
+			if (uint8(StatusFlags) & uint8(Contains))
+			{
+				if (!bFirst)
+				{
+					Sb.AppendChar(TEXT('|'));
+				}
+				Sb.Append(Str);
+				bFirst = false;
+			}
+		};
+
+	AppendIf(FBackendStatus::EFlags::CacheEnabled, TEXT("CacheEnabled"));
+	AppendIf(FBackendStatus::EFlags::HttpEnabled, TEXT("HttpEnabled"));
+	AppendIf(FBackendStatus::EFlags::HttpError, TEXT("HttpError"));
+	AppendIf(FBackendStatus::EFlags::HttpBulkOptionalDisabled, TEXT("HttpBulkOptionalDisabled"));
+	AppendIf(FBackendStatus::EFlags::AbandonCache, TEXT("AbandonCache"));
+
+	return Sb;
+}
+
+} // namespace UE::IoStore
\ No newline at end of file
diff --git a/Engine/Source/Runtime/Experimental/IoStore/OnDemand/Private/OnDemandBackendStatus.h b/Engine/Source/Runtime/Experimental/IoStore/OnDemand/Private/OnDemandBackendStatus.h
new file mode 100644
--- /dev/null
+++ b/Engine/Source/Runtime/Experimental/IoStore/OnDemand/Private/OnDemandBackendStatus.h
@@ -0,0 +1,128 @@
+// Copyright Epic Games, Inc. All Rights Reserved.
+
+#pragma once
+
+#include "GenericPlatform/GenericPlatformCrashContext.h"
+#include "IO/IoChunkId.h"
+#include "Misc/StringBuilder.h"
+
+#include <atomic>
+
+namespace UE::IoStore
+{
+
+class FBackendStatus
+{
+public:
+
+	enum class EFlags : uint8
+	{
+		None						= 0,
+		CacheEnabled				= 1 << 0,
+		HttpEnabled					= 1 << 1,
+		HttpError					= 1 << 2,
+		HttpBulkOptionalDisabled	= 1 << 3,
+		AbandonCache				= 1 << 4,
+
+		// When adding new values here, remember to update operator<<(FStringBuilderBase& Sb, EFlags StatusFlags) in the cpp!
+	};
+
+	bool IsHttpEnabled() const
+	{
+		return IsHttpEnabled(Flags.load(std::memory_order_relaxed));
+	}
+
+	bool IsHttpEnabled(EIoChunkType ChunkType) const;
+
+	bool IsHttpError() const
+	{
+		return HasAnyFlags(EFlags::HttpError);
+	}
+
+	bool IsCacheEnabled() const
+	{
+		return HasAnyFlags(EFlags::CacheEnabled);
+	}
+
+	bool IsCacheWriteable() const
+	{
+		const uint8 CurrentFlags = Flags.load(std::memory_order_relaxed);
+		return (CurrentFlags & uint8(EFlags::CacheEnabled)) && IsHttpEnabled(CurrentFlags);
+	}
+
+	bool IsCacheReadOnly() const
+	{
+		const uint8 CurrentFlags = Flags.load(std::memory_order_relaxed);
+		return (CurrentFlags & uint8(EFlags::CacheEnabled)) && !IsHttpEnabled(CurrentFlags);
+	}
+
+	bool ShouldAbandonCache() const
+	{
+		return HasAnyFlags(EFlags::AbandonCache);
+	}
+
+	void SetHttpEnabled(bool bEnabled)
+	{
+		AddOrRemoveFlags(EFlags::HttpEnabled, bEnabled, TEXT("HTTP streaming enabled"));
+		FGenericCrashContext::SetEngineData(TEXT("IAS.Enabled"), bEnabled ? TEXT("true") : TEXT("false"));
+	}
+
+	void SetHttpOptionalBulkEnabled(bool bEnabled)
+	{
+		AddOrRemoveFlags(EFlags::HttpBulkOptionalDisabled, bEnabled == false, TEXT("HTTP streaming of optional bulk data disabled"));
+	}
+
+	void SetCacheEnabled(bool bEnabled)
+	{
+		AddOrRemoveFlags(EFlags::CacheEnabled, bEnabled, TEXT("Cache enabled"));
+	}
+
+	void SetHttpError(bool bError)
+	{
+		AddOrRemoveFlags(EFlags::HttpError, bError, TEXT("HTTP streaming error"));
+	}
+
+	void SetAbandonCache(bool bAbandon)
+	{
+		AddOrRemoveFlags(EFlags::AbandonCache, bAbandon, TEXT("Abandon cache"));
+	}
+
+	void ToString(FStringBuilderBase& Builder);
+
+private:
+
+	static bool IsHttpEnabled(uint8 FlagsToTest);
+
+	bool HasAnyFlags(uint8 Contains) const
+	{
+		return (Flags.load(std::memory_order_relaxed) & Contains) != 0;
+	}
+
+	bool HasAnyFlags(EFlags Contains) const
+	{
+		return HasAnyFlags(uint8(Contains));
+	}
+
+	uint8 AddFlags(EFlags FlagsToAdd)
+	{
+		return Flags.fetch_or(uint8(FlagsToAdd));
+	}
+
+	uint8 RemoveFlags(EFlags FlagsToRemove)
+	{
+		return Flags.fetch_and(~uint8(FlagsToRemove));
+	}
+
+	uint8 AddOrRemoveFlags(EFlags FlagsToAddOrRemove, bool bValue)
+	{
+		return bValue ? AddFlags(FlagsToAddOrRemove) : RemoveFlags(FlagsToAddOrRemove);
+	}
+
+	void AddOrRemoveFlags(EFlags FlagsToAddOrRemove, bool bValue, const TCHAR* DebugText);
+
+	friend FStringBuilderBase& operator<<(FStringBuilderBase& Sb, EFlags StatusFlags);
+
+	std::atomic<uint8> Flags{ 0 };
+};
+
+} // namespace UE::IoStore
diff --git a/Engine/Source/Runtime/Experimental/IoStore/OnDemand/Private/OnDemandIoDispatcherBackend.cpp b/Engine/Source/Runtime/Experimental/IoStore/OnDemand/Private/OnDemandIoDispatcherBackend.cpp
--- a/Engine/Source/Runtime/Experimental/IoStore/OnDemand/Private/OnDemandIoDispatcherBackend.cpp
+++ b/Engine/Source/Runtime/Experimental/IoStore/OnDemand/Private/OnDemandIoDispatcherBackend.cpp
@@ -6,7 +6,6 @@
 #include "Containers/BitArray.h"
 #include "Containers/StringView.h"
 #include "DistributionEndpoints.h"
-#include "GenericPlatform/GenericPlatformCrashContext.h"
 #include "HAL/Event.h"
 #include "HAL/LowLevelMemTracker.h"
 #include "HAL/Platform.h"
@@ -14,6 +13,7 @@
 #include "HAL/PreprocessorHelpers.h"
 #include "HAL/Runnable.h"
 #include "HAL/RunnableThread.h"
+#include "IO/IasCache.h"
 #include "IO/IoAllocators.h"
 #include "IO/IoChunkEncoding.h"
 #include "IO/IoDispatcher.h"
@@ -21,7 +21,6 @@
 #include "IO/IoStatus.h"
 #include "IO/IoStore.h"
 #include "IO/IoStoreOnDemand.h"
-#include "IO/IasCache.h"
 #include "Logging/StructuredLog.h"
 #include "Math/NumericLimits.h"
 #include "Misc/ConfigCacheIni.h"
@@ -30,11 +29,12 @@
 #include "Misc/Paths.h"
 #include "Misc/ScopeLock.h"
 #include "Misc/ScopeRWLock.h"
+#include "OnDemandBackendStatus.h"
 #include "OnDemandHttpClient.h"
 #include "OnDemandIoStore.h"
+#include "ProfilingDebugging/IoStoreTrace.h"
 #include "Statistics.h"
 #include "Tasks/Task.h"
-#include "ProfilingDebugging/IoStoreTrace.h"
 
 #include <atomic>
 
@@ -707,167 +707,6 @@ static void LogIoResult(
 		PrioToString(Priority));
 };
 
-///////////////////////////////////////////////////////////////////////////////
-struct FBackendStatus
-{
-	enum class EFlags : uint8
-	{
-		None						= 0,
-		CacheEnabled				= (1 << 0),
-		HttpEnabled					= (1 << 1),
-		HttpError					= (1 << 2),
-		HttpBulkOptionalDisabled	= (1 << 3),
-		AbandonCache				= (1 << 4),
-
-		// When adding new values here, remember to update operator<<(FStringBuilderBase& Sb, EFlags StatusFlags) below!
-	};
-
-	bool IsHttpEnabled() const
-	{
-		return IsHttpEnabled(Flags.load(std::memory_order_relaxed));
-	}
-
-	bool IsHttpEnabled(EIoChunkType ChunkType) const
-	{
-		const uint8 CurrentFlags = Flags.load(std::memory_order_relaxed);
-		return IsHttpEnabled(CurrentFlags) &&
-			(ChunkType != EIoChunkType::OptionalBulkData ||
-				((CurrentFlags & uint8(EFlags::HttpBulkOptionalDisabled)) == 0 && GIasHttpOptionalBulkDataEnabled));
-	}
-
-	bool IsHttpError() const
-	{
-		return HasAnyFlags(EFlags::HttpError);
-	}
-
-	bool IsCacheEnabled() const
-	{
-		return HasAnyFlags(EFlags::CacheEnabled);
-	}
-
-	bool IsCacheWriteable() const
-	{
-		const uint8 CurrentFlags = Flags.load(std::memory_order_relaxed);
-		return (CurrentFlags & uint8(EFlags::CacheEnabled)) && IsHttpEnabled(CurrentFlags); 
-	}
-
-	bool IsCacheReadOnly() const
-	{
-		const uint8 CurrentFlags = Flags.load(std::memory_order_relaxed);
-		return (CurrentFlags & uint8(EFlags::CacheEnabled)) && !IsHttpEnabled(CurrentFlags);
-	}
-
-	bool ShouldAbandonCache() const
-	{
-		return HasAnyFlags(EFlags::AbandonCache);
-	}
-
-	void SetHttpEnabled(bool bEnabled)
-	{
-		AddOrRemoveFlags(EFlags::HttpEnabled, bEnabled, TEXT("HTTP streaming enabled"));
-		FGenericCrashContext::SetEngineData(TEXT("IAS.Enabled"), bEnabled ? TEXT("true") : TEXT("false"));
-	}
-
-	void SetHttpOptionalBulkEnabled(bool bEnabled)
-	{
-		AddOrRemoveFlags(EFlags::HttpBulkOptionalDisabled, bEnabled == false, TEXT("HTTP streaming of optional bulk data disabled"));
-	}
-
-	void SetCacheEnabled(bool bEnabled)
-	{
-		AddOrRemoveFlags(EFlags::CacheEnabled, bEnabled, TEXT("Cache enabled"));
-	}
-
-	void SetHttpError(bool bError)
-	{
-		AddOrRemoveFlags(EFlags::HttpError, bError, TEXT("HTTP streaming error"));
-	}
-
-	void SetAbandonCache(bool bAbandon)
-	{
-		AddOrRemoveFlags(EFlags::AbandonCache, bAbandon, TEXT("Abandon cache"));
-	}
-
-private:
-	static bool IsHttpEnabled(uint8 FlagsToTest)
-	{
-		constexpr uint8 HttpFlags = uint8(EFlags::HttpEnabled) | uint8(EFlags::HttpError);
-		return ((FlagsToTest & HttpFlags) == uint8(EFlags::HttpEnabled)) && GIasHttpEnabled;
-	}
-
-	bool HasAnyFlags(uint8 Contains) const
-	{
-		return (Flags.load(std::memory_order_relaxed) & Contains) != 0;
-	}
-	
-	bool HasAnyFlags(EFlags Contains) const
-	{
-		return HasAnyFlags(uint8(Contains));
-	}
-
-	uint8 AddFlags(EFlags FlagsToAdd)
-	{
-		return Flags.fetch_or(uint8(FlagsToAdd));
-	}
-
-	uint8 RemoveFlags(EFlags FlagsToRemove)
-	{
-		return Flags.fetch_and(~uint8(FlagsToRemove));
-	}
-
-	uint8 AddOrRemoveFlags(EFlags FlagsToAddOrRemove, bool bValue)
-	{
-		return bValue ? AddFlags(FlagsToAddOrRemove) : RemoveFlags(FlagsToAddOrRemove);
-	}
-
-	void AddOrRemoveFlags(EFlags FlagsToAddOrRemove, bool bValue, const TCHAR* DebugText)
-	{
-		const uint8 PrevFlags = AddOrRemoveFlags(FlagsToAddOrRemove, bValue);
-		TStringBuilder<128> Sb;
-		Sb.Append(DebugText)
-			<< TEXT(" '");
-		Sb.Append(bValue ? TEXT("true") : TEXT("false"))
-			<< TEXT("', backend status '(")
-			<< EFlags(PrevFlags)
-			<< TEXT(") -> (")
-			<< EFlags(Flags.load(std::memory_order_relaxed))
-			<< TEXT(")'");
-		UE_LOG(LogIas, Log, TEXT("%s"), Sb.ToString());
-	}
-
-	friend FStringBuilderBase& operator<<(FStringBuilderBase& Sb, EFlags StatusFlags)
-	{
-		if (StatusFlags == EFlags::None)
-		{
-			Sb.Append(TEXT("None"));
-			return Sb;
-		}
-
-		bool bFirst = true;
-		auto AppendIf = [StatusFlags, &Sb, &bFirst](EFlags Contains, const TCHAR* Str)
-		{
-			if (uint8(StatusFlags) & uint8(Contains))
-			{
-				if (!bFirst)
-				{
-					Sb.AppendChar(TEXT('|'));
-				}
-				Sb.Append(Str);
-				bFirst = false;
-			}
-		};
-
-		AppendIf(EFlags::CacheEnabled, TEXT("CacheEnabled"));
-		AppendIf(EFlags::HttpEnabled, TEXT("HttpEnabled"));
-		AppendIf(EFlags::HttpError, TEXT("HttpError"));
-		AppendIf(EFlags::HttpBulkOptionalDisabled, TEXT("HttpBulkOptionalDisabled"));
-		AppendIf(EFlags::AbandonCache, TEXT("AbandonCache"));
-
-		return Sb;
-	}
-
-	std::atomic<uint8> Flags{0};
-};
 ///////////////////////////////////////////////////////////////////////////////
 class FOnDemandIoBackend final
 	: public FRunnable
@@ -1118,7 +957,7 @@ private:
 FOnDemandIoBackend::FOnDemandIoBackend(const FOnDemandEndpointConfig& Config, FOnDemandIoStore& InIoStore, TUniquePtr<IIasCache>&& InCache)
 	: IoStore(InIoStore)
 	, Cache(MoveTemp(InCache))
-	, Stats(Cache.IsValid() ? EStatsFlags::None : EStatsFlags::CachingDisabled)
+	, Stats(BackendStatus)
 	, CurrentThreadPriority(EThreadPriority::TPri_Num)
 {
 	EndpointTestPath = Config.TocPath;
diff --git a/Engine/Source/Runtime/Experimental/IoStore/OnDemand/Private/Statistics.cpp b/Engine/Source/Runtime/Experimental/IoStore/OnDemand/Private/Statistics.cpp
--- a/Engine/Source/Runtime/Experimental/IoStore/OnDemand/Private/Statistics.cpp
+++ b/Engine/Source/Runtime/Experimental/IoStore/OnDemand/Private/Statistics.cpp
@@ -347,8 +347,8 @@ static FOnDemandIoBackendStats* GStatistics = nullptr;
 static FDelegateHandle GStatisticsEndFrameDelegateHandle;
 static FDelegateHandle GStatisticsOnScreenDelegateHandle;
 
-FOnDemandIoBackendStats::FOnDemandIoBackendStats(EStatsFlags InFlags)
-	: Flags(InFlags)
+FOnDemandIoBackendStats::FOnDemandIoBackendStats(FBackendStatus& InStatus)
+	: BackendStatus(InStatus)
 {
 	static constexpr float OneOver1024 = 1.0f / 1024.0f;
 
@@ -400,7 +400,7 @@ FOnDemandIoBackendStats::FOnDemandIoBackendStats(EStatsFlags InFlags)
 			static double LastLogTime = 0.0;
 			if (double Time = FPlatformTime::Seconds(); Time - LastLogTime > (double)GIasStatisticsLogInterval)
 			{
-				if (!EnumHasAllFlags(Flags,EStatsFlags::CachingDisabled))
+				if (BackendStatus.IsCacheEnabled())
 				{
 					UE_LOG(LogIas, Log, TEXT("CacheStats: CachedKiB=%d, WrittenKiB=%d, ReadKiB=%d, RejectedKiB=%d, Get=%d, Error=%d, Put=%d, PutReject=%d, PutExisting=%d"),
 						(int32)CCachedKiB, (int32)CWrittenKiB, (int32)CReadKiB, (int32)CRejectedKiB, CGetCount, CErrorCount, CPutCount, CPutRejectCount, CPutExistingCount);
@@ -423,10 +423,19 @@ FOnDemandIoBackendStats::FOnDemandIoBackendStats(EStatsFlags InFlags)
 		{
 			if (GIasDisplayOnScreenStatistics)
 			{
-				const bool bIsConnected = GHttpConnectCount.Get() > GHttpDisconnectCount.Get();
+				{
+					TStringBuilder<256> BackendStatusText;
+					BackendStatusText << TEXT("IAS Backend Status: ");
+					BackendStatus.ToString(BackendStatusText);
+
+					FCoreDelegates::EOnScreenMessageSeverity Verbosity = BackendStatus.IsHttpError() ?	FCoreDelegates::EOnScreenMessageSeverity::Error :
+																										FCoreDelegates::EOnScreenMessageSeverity::Info;
+
+					Out.Add(Verbosity, FText::FromStringView(BackendStatusText.ToView()));
+				}
 
 				TStringBuilder<256> CachingText;
-				if (!EnumHasAllFlags(Flags, EStatsFlags::CachingDisabled))
+				if (BackendStatus.IsCacheEnabled())
 				{
 					CachingText << TEXT("Cached: ") << FText::AsMemory(GCacheCachedBytes.Get()).ToString();
 					CachingText << TEXT(" | Rejected: ") << FText::AsMemory(GCacheRejectBytes.Get()).ToString();
@@ -440,8 +449,7 @@ FOnDemandIoBackendStats::FOnDemandIoBackendStats(EStatsFlags InFlags)
 
 				TStringBuilder<256> Builder;
 				Builder.Appendf(
-					TEXT("IAS - %s: %s | Downloaded:%s (%d) Avg %d ms | Retries:%d | Pending:%d"),
-					bIsConnected ? TEXT("Connected") : TEXT("Disconnected"),
+					TEXT("IAS Backend Stats:   %s | Downloaded: %s (%d) Avg %d ms | Retries: %d | Pending: %d"),
 					CachingText.ToString(),
 					*FText::AsMemory(GHttpDownloadedBytes.Get()).ToString(),
 					GHttpGetCount.Get(),
@@ -460,7 +468,7 @@ FOnDemandIoBackendStats::FOnDemandIoBackendStats(EStatsFlags InFlags)
 						GHttpErrorCount.Get()
 					);
 
-					Out.Add(FCoreDelegates::EOnScreenMessageSeverity::Error, FText::FromStringView(Builder.ToView()));
+					Out.Add(FCoreDelegates::EOnScreenMessageSeverity::Error, FText::FromStringView(Warning.ToView()));
 				}
 			}
 		});
@@ -528,7 +536,7 @@ void FOnDemandIoBackendStats::ReportEndPointAnalytics(TArray<FAnalyticsEventAttr
 
 		AppendAnalyticsEventAttributeArray(OutAnalyticsArray
 
-			,TEXT("IasCacheEnabled"), !EnumHasAllFlags(Flags, EStatsFlags::CachingDisabled)
+			,TEXT("IasCacheEnabled"), BackendStatus.IsCacheEnabled()
 
 			,TRACK_DELTA("IasCacheTotalCount", CacheTotalCount)
 			,TRACK_DELTA("IasCacheErrorCount", GCacheErrorCount.Get())
diff --git a/Engine/Source/Runtime/Experimental/IoStore/OnDemand/Private/Statistics.h b/Engine/Source/Runtime/Experimental/IoStore/OnDemand/Private/Statistics.h
--- a/Engine/Source/Runtime/Experimental/IoStore/OnDemand/Private/Statistics.h
+++ b/Engine/Source/Runtime/Experimental/IoStore/OnDemand/Private/Statistics.h
@@ -3,6 +3,7 @@
 #pragma once
 
 #include "HAL/LowLevelMemTracker.h"
+#include "OnDemandBackendStatus.h"
 #include "ProfilingDebugging/CountersTrace.h"
 #include "ProfilingDebugging/CsvProfiler.h"
 
@@ -25,20 +26,12 @@ namespace UE::IoStore
 #	define IAS_STATISTICS_IMPL(...) { return __VA_ARGS__; }
 #endif
 
-enum class EStatsFlags
-{
-	None				= 0,
-	CachingDisabled		= 1 << 0
-};
-
-ENUM_CLASS_FLAGS(EStatsFlags);
-
 class FOnDemandIoBackendStats
 {
 public:
 	static FOnDemandIoBackendStats* Get() IAS_STATISTICS_IMPL(nullptr)
 
-	FOnDemandIoBackendStats(EStatsFlags Flags) IAS_STATISTICS_IMPL()
+	FOnDemandIoBackendStats(FBackendStatus& InStatus) IAS_STATISTICS_IMPL()
 	~FOnDemandIoBackendStats() IAS_STATISTICS_IMPL()
 
 	/** Report analytics not directly associated with a specific endpoint */
@@ -76,7 +69,7 @@ public:
 private:
 
 #if IAS_WITH_STATISTICS
-	EStatsFlags Flags;
+	FBackendStatus& BackendStatus;
 #endif //IAS_WITH_STATISTICS
 };
 
