Fix more latest clang issues

#rnx

diff --git a/Engine/Source/Programs/UnrealBuildTool/ToolChain/ClangWarnings.cs b/Engine/Source/Programs/UnrealBuildTool/ToolChain/ClangWarnings.cs
--- a/Engine/Source/Programs/UnrealBuildTool/ToolChain/ClangWarnings.cs
+++ b/Engine/Source/Programs/UnrealBuildTool/ToolChain/ClangWarnings.cs
@@ -60,6 +60,8 @@ namespace UnrealBuildTool
 			if (ClangVersion >= new VersionNumber(19))
 			{
 				Arguments.Add("-Wno-extra-qualification");              // https://clang.llvm.org/docs/DiagnosticsReference.html#wextra-qualification					// new warning for clang 19
+				Arguments.Add("-Wno-cast-function-type-mismatch");
+				Arguments.Add("-Wno-missing-template-arg-list-after-template-kw");
 			}
 
 			Arguments.Add("-Wno-gnu-string-literal-operator-template"); // https://clang.llvm.org/docs/DiagnosticsReference.html#wgnu-string-literal-operator-template	// We use this feature to allow static FNames.
diff --git a/Engine/Source/Runtime/Experimental/Chaos/Public/Chaos/ImplicitObjectScaled.h b/Engine/Source/Runtime/Experimental/Chaos/Public/Chaos/ImplicitObjectScaled.h
--- a/Engine/Source/Runtime/Experimental/Chaos/Public/Chaos/ImplicitObjectScaled.h
+++ b/Engine/Source/Runtime/Experimental/Chaos/Public/Chaos/ImplicitObjectScaled.h
@@ -69,7 +69,7 @@ public:
 		this->bIsConvex = MObject->IsConvex();
 		this->bDoCollide = MObject->GetDoCollide();
 		this->OuterMargin = FRealSingle(InMargin);
-		SetMargin(OuterMargin + MObject->GetMargin());
+		SetMargin(OuterMargin + MObject->GetMarginf());
 	}
 
 	TImplicitObjectInstanced(const ObjectType& Object, const FReal InMargin = 0)
@@ -92,7 +92,7 @@ public:
 		this->bIsConvex = MObject->IsConvex();
 		this->bDoCollide = MObject->GetDoCollide();
 		this->OuterMargin = Other.OuterMargin;
-		SetMargin(Other.GetMargin());
+		SetMargin(Other.GetMarginf());
 	}
 
 	static constexpr EImplicitObjectType StaticType()
diff --git a/Engine/Source/Runtime/Experimental/Chaos/Public/Chaos/Sphere.h b/Engine/Source/Runtime/Experimental/Chaos/Public/Chaos/Sphere.h
--- a/Engine/Source/Runtime/Experimental/Chaos/Public/Chaos/Sphere.h
+++ b/Engine/Source/Runtime/Experimental/Chaos/Public/Chaos/Sphere.h
@@ -47,7 +47,7 @@ namespace Chaos
 		    : FImplicitObject(EImplicitObject::IsConvex | EImplicitObject::HasBoundingBox, ImplicitObjectType::Sphere)
 		    , Center(MoveTemp(Other.Center))
 		{
-			SetRadius(Other.GetRadius());
+			SetRadius(Other.GetRadiusf());
 		}
 
 		TSphere& operator=(TSphere<T, d>&& InSteal)
