Make ISMPool clear the StaticMesh property of components in the free list. Prevents keeping hard references to static meshes that are unused. Can be disabled via cvar r.ISMPool.ClearComponentMeshOnRecycle.

diff --git a/Engine/Source/Runtime/Experimental/ISMPool/Private/ISMPool/ISMPoolComponent.cpp b/Engine/Source/Runtime/Experimental/ISMPool/Private/ISMPool/ISMPoolComponent.cpp
--- a/Engine/Source/Runtime/Experimental/ISMPool/Private/ISMPool/ISMPoolComponent.cpp
+++ b/Engine/Source/Runtime/Experimental/ISMPool/Private/ISMPool/ISMPoolComponent.cpp
@@ -31,6 +31,12 @@ FAutoConsoleVariableRef CVarISMPoolComponentRecycle(
 	GComponentRecycle,
 	TEXT("Recycle ISM components to a free list for reuse when all their instances are removed."));
 
+static bool GISMPoolClearComponentMeshOnRecycle = true;
+FAutoConsoleVariableRef CVarISMPoolClearComponentMeshOnRecycle(
+	TEXT("r.ISMPool.ClearComponentMeshOnRecycle"),
+	GISMPoolClearComponentMeshOnRecycle,
+	TEXT("If true, ISM components on the free list will have their StaticMesh property cleared - to prevent holding a reference to an unused mesh"));
+
 // Target free list size when recycling ISM components.
 // We try to maintain a pool of free components for fast allocation, but want to clean up when numbers get too high.
 static int32 GComponentFreeListTargetSize = 50;
@@ -456,6 +462,12 @@ void FISMPool::RemoveISM(FISMIndex ISMIndex, bool bKeepAlive, bool bRecycle)
 #if WITH_EDITOR
 		ISM.ISMComponent->Rename(nullptr);
 #endif
+
+		if (GISMPoolClearComponentMeshOnRecycle)
+		{
+			check(ISM.ISMComponent);
+			ISM.ISMComponent->SetStaticMesh(nullptr);
+		}
 		FreeListISM.Add(ISMIndex);
 	}
 	else
