Iris
- Implemented support for "dynamic"-DataStreams which can be opened and closed on request.
- Exposed access to DataStreams through ReplicationSystem facade.
- Refactored initialization/deinitialization of datastreams to be able to propertly intiialize external DataStreams

diff --git a/Engine/Plugins/Runtime/ReplicationSystemTestPlugin/Source/Private/Tests/DataStream/TestDataStream.cpp b/Engine/Plugins/Runtime/ReplicationSystemTestPlugin/Source/Private/Tests/DataStream/TestDataStream.cpp
--- a/Engine/Plugins/Runtime/ReplicationSystemTestPlugin/Source/Private/Tests/DataStream/TestDataStream.cpp
+++ b/Engine/Plugins/Runtime/ReplicationSystemTestPlugin/Source/Private/Tests/DataStream/TestDataStream.cpp
@@ -12,6 +12,12 @@
 #include "Iris/Serialization/NetBitStreamReader.h"
 #include "Iris/Serialization/NetBitStreamWriter.h"
 #include "Iris/Serialization/NetSerializationContext.h"
+#include "Tests/ReplicationSystem/ReplicationSystemServerClientTestFixture.h"
+#include "Iris/ReplicationSystem/ReplicationSystem.h"
+#include "Iris/ReplicationSystem/NetTokenStore.h"
+#include "Iris/ReplicationSystem/StringTokenStore.h"
+#include "Iris/ReplicationSystem/NetTokenDataStream.h"
+#include "Templates/Casts.h"
 
 namespace UE::Net::Private
 {
@@ -204,7 +210,7 @@ FTestDataStream::FTestDataStream()
 
 void FTestDataStream::SetUp()
 {
-	FDataStreamManagerInitParams InitParams;
+	UDataStreamManager::FInitParameters InitParams;
 	InitParams.PacketWindowSize = 256;
 	DataStreamManager = NewObject<UDataStreamManager>();
 	DataStreamManager->Init(InitParams);
@@ -214,6 +220,7 @@ void FTestDataStream::SetUp()
 void FTestDataStream::TearDown()
 {
 	RestoreDataStreamDefinitions();
+	DataStreamManager->Deinit();
 	DataStreamManager->MarkAsGarbage();
 	DataStreamManager = nullptr;
 }
@@ -297,4 +304,500 @@ FTestMessage& operator<<(FTestMessage& TestMessage, const UDataStream::EWriteRes
 	}
 }
 
+class FTestDynamicCreateDataStreamFixture : public FReplicationSystemServerClientTestFixture
+{
+protected:
+	virtual void SetUp() override
+	{
+		FReplicationSystemServerClientTestFixture::SetUp();
+
+		// We can overrida what we want as long as we do it before any connections are created.
+		DataStreamUtil.SetUp();
+
+		// Add a dynamic DataStream
+		{
+			FDataStreamTestUtil::FAddDataStreamDefinitionParams Params;
+			Params.bDynamicCreate = true;
+			DataStreamUtil.AddDataStreamDefinition(TEXT("DynamicNetToken"), TEXT("/Script/IrisCore.NetTokenDataStream"), Params);	
+		}
+
+		{
+			DataStreamUtil.AddDataStreamDefinition(TEXT("Replication"), TEXT("/Script/IrisCore.ReplicationDataStream"));
+		}
+
+		DataStreamUtil.FixupDefinitions();
+
+		// Add a client
+		Client = CreateClient();
+
+		// Cache some data
+		// Server
+		{
+			FNetTokenStore* TokenStore = Server->GetReplicationSystem()->GetNetTokenStore();
+			ServerStringTokenStore = TokenStore->GetDataStore<FStringTokenStore>();
+			ServerRemoteNetTokenStoreState = TokenStore->GetRemoteNetTokenStoreState(Client->ConnectionIdOnServer);
+		}
+
+		// Client
+		{
+			FNetTokenStore* TokenStore = Client->GetReplicationSystem()->GetNetTokenStore();
+			ClientStringTokenStore = TokenStore->GetDataStore<FStringTokenStore>();
+			ClientRemoteNetTokenStoreState = TokenStore->GetRemoteNetTokenStoreState(Client->LocalConnectionId);
+		}
+	}
+
+	FNetToken CreateAndExportTokenToClient(const FString& TokenString)
+	{
+		const FNetToken Token = ServerStringTokenStore->GetOrCreateToken(TokenString);
+		UNetTokenDataStream* NetTokenDataStream = Cast<UNetTokenDataStream>(Server->GetReplicationSystem()->GetDataStream(Client->ConnectionIdOnServer, DataStreamName));
+
+		if (NetTokenDataStream)
+		{
+			NetTokenDataStream->AddNetTokenForExplicitExport(Token);
+		}
+
+		return Token;
+	}
+
+	FNetToken CreateAndExportTokenToServer(const FString& TokenString)
+	{
+		const FNetToken Token = ClientStringTokenStore->GetOrCreateToken(TokenString);
+		UNetTokenDataStream* NetTokenDataStream = Cast<UNetTokenDataStream>(Client->GetReplicationSystem()->GetDataStream(Client->LocalConnectionId, DataStreamName));
+
+		if (NetTokenDataStream)
+		{
+			NetTokenDataStream->AddNetTokenForExplicitExport(Token);
+		}
+
+		return Token;
+	}
+
+	UDataStream::EDataStreamState GetDataStreamStateOnServer() const
+	{
+		return Server->GetConnectionInfo(Client->ConnectionIdOnServer).DataStreamManager->GetStreamState(DataStreamName);
+	}
+
+	UDataStream::EDataStreamState GetDataStreamStateOnClient() const
+	{
+		return Client->GetConnectionInfo(Client->LocalConnectionId).DataStreamManager->GetStreamState(DataStreamName);
+	}
+
+	void RoundTrip()
+	{
+		Server->UpdateAndSend({Client}, true);
+		Client->UpdateAndSend({Server}, true);
+	}
+	const FName DataStreamName = "DynamicNetToken";
+
+	FReplicationSystemTestClient* Client = nullptr;
+
+	FStringTokenStore* ServerStringTokenStore = nullptr;
+	FStringTokenStore* ClientStringTokenStore = nullptr;
+
+	const FNetTokenStoreState* ClientRemoteNetTokenStoreState = nullptr;
+	const FNetTokenStoreState* ServerRemoteNetTokenStoreState = nullptr;
+
+	UDataStreamManager* ServerDataStreamManager = nullptr;
+	UDataStreamManager* ClientDataStreamManager = nullptr;
+
+	FDataStreamTestUtil DataStreamUtil;
+};
+
+// Basic functionality test
+UE_NET_TEST_FIXTURE(FTestDynamicCreateDataStreamFixture, TestDynamicDataStream)
+{
+	UReplicationSystem* ServerReplicationSystem = Server->ReplicationSystem;
+	UReplicationSystem* ClientReplicationSystem = Client->ReplicationSystem;
+
+	// Verify that we cannot find DataStream
+	{
+		UDataStream* ServerNetTokenStream = ServerReplicationSystem->GetDataStream(Client->ConnectionIdOnServer, DataStreamName);
+		UDataStream* ClientNetTokenStream = ClientReplicationSystem->GetDataStream(Client->LocalConnectionId, DataStreamName);
+
+		// It should not exist as it is a dynamic DataStream
+		UE_NET_ASSERT_EQ(ServerNetTokenStream, nullptr);
+		UE_NET_ASSERT_EQ(ClientNetTokenStream, nullptr);
+	
+		UE_NET_ASSERT_TRUE(GetDataStreamStateOnServer() == UDataStream::EDataStreamState::Invalid);
+		UE_NET_ASSERT_TRUE(GetDataStreamStateOnClient() == UDataStream::EDataStreamState::Invalid);		
+	}
+
+	// Open dynamic stream, and verify that it now exists on server
+	{
+		UDataStream* ServerNetTokenStream = ServerReplicationSystem->OpenDataStream(Client->ConnectionIdOnServer, DataStreamName);
+		UDataStream* ClientNetTokenStream = ClientReplicationSystem->GetDataStream(Client->LocalConnectionId, DataStreamName);
+
+		// Should now exist on server
+		UE_NET_ASSERT_NE(ServerNetTokenStream, nullptr);
+		UE_NET_ASSERT_TRUE(GetDataStreamStateOnServer() == UDataStream::EDataStreamState::PendingCreate);
+
+		// But not on client
+		UE_NET_ASSERT_EQ(ClientNetTokenStream, nullptr);
+	}
+
+	// Roundtrip
+	RoundTrip();
+
+	// Now we expect it to be created on client as well.
+	{
+		UDataStream* ServerNetTokenStream = ServerReplicationSystem->GetDataStream(Client->ConnectionIdOnServer, DataStreamName);
+		UDataStream* ClientNetTokenStream = ClientReplicationSystem->GetDataStream(Client->LocalConnectionId, DataStreamName);
+
+		// Server state should now be UDataStream::EDataStreamState::Open
+		UE_NET_ASSERT_NE(ServerNetTokenStream, nullptr);
+		UE_NET_ASSERT_TRUE(GetDataStreamStateOnServer() == UDataStream::EDataStreamState::Open);
+
+		// Should now exist on client and be in the WaitOnCreateConfirmation
+		UE_NET_ASSERT_NE(ClientNetTokenStream, nullptr);
+		UE_NET_ASSERT_TRUE(GetDataStreamStateOnClient() == UDataStream::EDataStreamState::WaitOnCreateConfirmation);
+	}
+
+	// Send some data to client on the now open stream
+	FNetToken ServerHelloToken = CreateAndExportTokenToClient(TEXT("Hello"));
+	
+	// Roundtrip
+	RoundTrip();
+
+	// We should be able to resolve this on client now.
+	{
+		// Client state should now be open as well
+		UE_NET_ASSERT_TRUE(GetDataStreamStateOnClient() == UDataStream::EDataStreamState::Open);
+
+		const TCHAR* RcvdTokenString = ClientStringTokenStore->ResolveToken(ServerHelloToken, ClientRemoteNetTokenStoreState);
+		UE_NET_ASSERT_NE(RcvdTokenString, nullptr);
+	}
+
+	// Send some data from client
+	FNetToken ClientHelloToken = CreateAndExportTokenToServer(TEXT("HelloFromClient"));
+	
+	// Roundtrip
+	RoundTrip();
+
+	// We should be able to resolve this on server now.
+	{
+		const TCHAR* RcvdTokenString = ServerStringTokenStore->ResolveToken(ClientHelloToken, ServerRemoteNetTokenStoreState);
+		UE_NET_ASSERT_NE(RcvdTokenString, nullptr);
+	}
+
+	// Close from server
+	ServerReplicationSystem->CloseDataStream(Client->ConnectionIdOnServer, DataStreamName);
+
+	// Double roundtrip and we should be done
+	RoundTrip();
+	RoundTrip();
+
+	// Verify that we cannot find DataStream as it should be closed
+	{
+		UDataStream* ServerNetTokenStream = ServerReplicationSystem->GetDataStream(Client->ConnectionIdOnServer, DataStreamName);
+		UDataStream* ClientNetTokenStream = ClientReplicationSystem->GetDataStream(Client->LocalConnectionId, DataStreamName);
+
+		// It should not exist as it is a dynamic DataStream
+		UE_NET_ASSERT_EQ(ServerNetTokenStream, nullptr);
+		UE_NET_ASSERT_EQ(ClientNetTokenStream, nullptr);
+	}
+}
+
+// Verify that stream gets created even if we drop create request
+UE_NET_TEST_FIXTURE(FTestDynamicCreateDataStreamFixture, TestDynamicDataStream_DropPendingCreateFromServer)
+{
+	UReplicationSystem* ServerReplicationSystem = Server->ReplicationSystem;
+	UReplicationSystem* ClientReplicationSystem = Client->ReplicationSystem;
+
+	// Open dynamic stream, and verify that it now exists on server but not on client
+	{
+		UDataStream* ServerNetTokenStream = ServerReplicationSystem->OpenDataStream(Client->ConnectionIdOnServer, DataStreamName);
+		UDataStream* ClientNetTokenStream = ClientReplicationSystem->GetDataStream(Client->LocalConnectionId, DataStreamName);
+
+		UE_NET_ASSERT_NE(ServerNetTokenStream, nullptr);
+		UE_NET_ASSERT_EQ(ClientNetTokenStream, nullptr);
+	}
+
+	// Drop PendingCreate
+	Server->UpdateAndSend({Client}, false);
+
+	// DataStream should not exist on client
+	{
+		UDataStream* ClientNetTokenStream = ClientReplicationSystem->GetDataStream(Client->LocalConnectionId, DataStreamName);
+		UE_NET_ASSERT_EQ(ClientNetTokenStream, nullptr);
+	}
+
+	// Send again
+	Server->UpdateAndSend({Client}, true);
+
+	// Should now be created on client
+	{
+		UDataStream* ClientNetTokenStream = ClientReplicationSystem->GetDataStream(Client->LocalConnectionId, DataStreamName);
+		UE_NET_ASSERT_NE(ClientNetTokenStream, nullptr);
+	}
+}
+
+// Verify that stream gets to open even if we drop create request/confirmation from remote
+UE_NET_TEST_FIXTURE(FTestDynamicCreateDataStreamFixture, TestDynamicDataStream_DropPendingCreateFromClient)
+{
+	UReplicationSystem* ServerReplicationSystem = Server->ReplicationSystem;
+	UReplicationSystem* ClientReplicationSystem = Client->ReplicationSystem;
+
+	// Open it, and verify that it now exists on server
+	{
+		UDataStream* ServerNetTokenStream = ServerReplicationSystem->OpenDataStream(Client->ConnectionIdOnServer, DataStreamName);
+		UDataStream* ClientNetTokenStream = ClientReplicationSystem->GetDataStream(Client->LocalConnectionId, DataStreamName);
+
+		// Should now be valid on server
+		UE_NET_ASSERT_NE(ServerNetTokenStream, nullptr);
+		// But not on client
+		UE_NET_ASSERT_EQ(ClientNetTokenStream, nullptr);
+	}
+
+	// Deliver PendingCreate from server
+	Server->UpdateAndSend({Client}, true);
+
+	// Client should have created it.
+	{
+		UDataStream* ClientNetTokenStream = ClientReplicationSystem->GetDataStream(Client->LocalConnectionId, DataStreamName);
+		UE_NET_ASSERT_NE(ClientNetTokenStream, nullptr);
+	}
+
+	// Drop PendingCreate from client
+	Client->UpdateAndSend({Server}, false);
+	
+	// Server state should be WaitForCreateConfirmation
+	UE_NET_ASSERT_TRUE(GetDataStreamStateOnServer() == UDataStream::EDataStreamState::WaitOnCreateConfirmation);
+
+	// Client state should be PendingCreate
+	UE_NET_ASSERT_TRUE(GetDataStreamStateOnClient() == UDataStream::EDataStreamState::PendingCreate);
+
+	// Send and deliver to server
+	Client->UpdateAndSend({Server}, true);
+
+	// Server state should now be open
+	UE_NET_ASSERT_TRUE(GetDataStreamStateOnServer() == UDataStream::EDataStreamState::Open);
+
+	// Send and deliver to client
+	Server->UpdateAndSend({Client}, true);
+
+	// Client state should now also be open
+	UE_NET_ASSERT_TRUE(GetDataStreamStateOnClient() == UDataStream::EDataStreamState::Open);
+}
+
+// Verify that stream can be closed from client
+UE_NET_TEST_FIXTURE(FTestDynamicCreateDataStreamFixture, TestDynamicDataStream_RequestCloseFromClient)
+{
+	UReplicationSystem* ServerReplicationSystem = Server->ReplicationSystem;
+	UReplicationSystem* ClientReplicationSystem = Client->ReplicationSystem;
+
+	UDataStream* ServerNetTokenStream = ServerReplicationSystem->OpenDataStream(Client->ConnectionIdOnServer, DataStreamName);
+
+	RoundTrip();
+	RoundTrip();
+
+	// Verify that stream is open on both sides
+	UE_NET_ASSERT_TRUE(GetDataStreamStateOnServer() == UDataStream::EDataStreamState::Open);
+	UE_NET_ASSERT_TRUE(GetDataStreamStateOnClient() == UDataStream::EDataStreamState::Open);
+
+	// Request close from client
+	ClientReplicationSystem->CloseDataStream(Client->ConnectionIdOnServer, DataStreamName);
+
+	// Verify that stream is open on both sides
+	UE_NET_ASSERT_TRUE(GetDataStreamStateOnServer() == UDataStream::EDataStreamState::Open);
+	UE_NET_ASSERT_TRUE(GetDataStreamStateOnClient() == UDataStream::EDataStreamState::PendingClose);
+
+	// Send and deliver to server
+	Client->UpdateAndSend({Server}, true);
+	
+	UE_NET_ASSERT_TRUE(GetDataStreamStateOnServer() == UDataStream::EDataStreamState::PendingClose);
+	UE_NET_ASSERT_TRUE(GetDataStreamStateOnClient() == UDataStream::EDataStreamState::WaitOnCloseConfirmation);
+
+	// Double roundtrip and we should be done
+	RoundTrip();
+	RoundTrip();
+
+	// Verify that stream is invalidated
+	UE_NET_ASSERT_TRUE(GetDataStreamStateOnServer() == UDataStream::EDataStreamState::Invalid);
+	UE_NET_ASSERT_TRUE(GetDataStreamStateOnClient() == UDataStream::EDataStreamState::Invalid);	
+}
+
+// Verify that stream can be closed from client even if we drop request
+UE_NET_TEST_FIXTURE(FTestDynamicCreateDataStreamFixture, TestDynamicDataStream_RequestCloseFromClientIsResentIfDropped)
+{
+	UReplicationSystem* ServerReplicationSystem = Server->ReplicationSystem;
+	UReplicationSystem* ClientReplicationSystem = Client->ReplicationSystem;
+
+	UDataStream* ServerNetTokenStream = ServerReplicationSystem->OpenDataStream(Client->ConnectionIdOnServer, DataStreamName);
+
+	RoundTrip();
+	RoundTrip();
+
+	// Verify that stream is open on both sides
+	UE_NET_ASSERT_TRUE(GetDataStreamStateOnServer() == UDataStream::EDataStreamState::Open);
+	UE_NET_ASSERT_TRUE(GetDataStreamStateOnClient() == UDataStream::EDataStreamState::Open);
+
+	// Request close from client
+	ClientReplicationSystem->CloseDataStream(Client->ConnectionIdOnServer, DataStreamName);
+
+	// Verify expected stream state
+	UE_NET_ASSERT_TRUE(GetDataStreamStateOnServer() == UDataStream::EDataStreamState::Open);
+	UE_NET_ASSERT_TRUE(GetDataStreamStateOnClient() == UDataStream::EDataStreamState::PendingClose);
+
+	// Drop send to server
+	Client->UpdateAndSend({Server}, false);
+	
+	// Verify expected stream state
+	UE_NET_ASSERT_TRUE(GetDataStreamStateOnServer() == UDataStream::EDataStreamState::Open);
+	UE_NET_ASSERT_TRUE(GetDataStreamStateOnClient() == UDataStream::EDataStreamState::PendingClose);
+
+	// Send and deliver data
+	Client->UpdateAndSend({Server}, true);
+
+	// Verify expected stream state
+	UE_NET_ASSERT_TRUE(GetDataStreamStateOnServer() == UDataStream::EDataStreamState::PendingClose);
+	UE_NET_ASSERT_TRUE(GetDataStreamStateOnClient() == UDataStream::EDataStreamState::WaitOnCloseConfirmation);
+
+	// Double roundtrip and we should be done
+	RoundTrip();
+	RoundTrip();
+
+	// Verify that stream is invalidated
+	UE_NET_ASSERT_TRUE(GetDataStreamStateOnServer() == UDataStream::EDataStreamState::Invalid);
+	UE_NET_ASSERT_TRUE(GetDataStreamStateOnClient() == UDataStream::EDataStreamState::Invalid);	
+}
+
+// Verify that stream gets properly closed when changing state with create data in flight
+UE_NET_TEST_FIXTURE(FTestDynamicCreateDataStreamFixture, TestDynamicDataStream_CloseWhileWaitingForCreateConfirmation)
+{
+	UReplicationSystem* ServerReplicationSystem = Server->ReplicationSystem;
+	UReplicationSystem* ClientReplicationSystem = Client->ReplicationSystem;
+
+	// Open Stream
+	UDataStream* ServerNetTokenStream = ServerReplicationSystem->OpenDataStream(Client->ConnectionIdOnServer, DataStreamName);
+
+	// Put data in flight containing create
+	Server->PreSendUpdate();
+	Server->SendTo(Client);
+	Server->PostSendUpdate();
+
+	// Server state should be WaitForCreateConfirmation
+	UE_NET_ASSERT_TRUE(GetDataStreamStateOnServer() == UDataStream::EDataStreamState::WaitOnCreateConfirmation);
+
+	// Request close on server
+	ServerReplicationSystem->CloseDataStream(Client->ConnectionIdOnServer, DataStreamName);
+
+	// Server state should be PendingClose
+	UE_NET_ASSERT_TRUE(GetDataStreamStateOnServer() == UDataStream::EDataStreamState::PendingClose);
+
+	// Put data in flight containing close?
+	Server->PreSendUpdate();
+	Server->SendTo(Client);
+	Server->PostSendUpdate();
+
+	// Deliver packet with create
+	Server->DeliverTo(Client, true);
+
+	// Deliver packet with close
+	Server->DeliverTo(Client, true);
+
+	// Double roundtrip and we should be done
+	RoundTrip();
+	RoundTrip();
+
+	// Verify that stream is invalidated
+	UE_NET_ASSERT_TRUE(GetDataStreamStateOnServer() == UDataStream::EDataStreamState::Invalid);
+	UE_NET_ASSERT_TRUE(GetDataStreamStateOnClient() == UDataStream::EDataStreamState::Invalid);
+}
+
+// Verify that stream gets properly closed when changing state with create data in flight
+UE_NET_TEST_FIXTURE(FTestDynamicCreateDataStreamFixture, TestDynamicDataStream_CloseBeforeFirstSend)
+{
+	UReplicationSystem* ServerReplicationSystem = Server->ReplicationSystem;
+	UReplicationSystem* ClientReplicationSystem = Client->ReplicationSystem;
+
+	// Open Stream
+	UDataStream* ServerNetTokenStream = ServerReplicationSystem->OpenDataStream(Client->ConnectionIdOnServer, DataStreamName);
+	// Close Stream
+	ServerReplicationSystem->CloseDataStream(Client->ConnectionIdOnServer, DataStreamName);
+
+	// Run a few updates and make sure stream is propertly closed.
+	RoundTrip();
+	RoundTrip();
+
+	// Verify that stream is invalidated
+	UE_NET_ASSERT_TRUE(GetDataStreamStateOnServer() == UDataStream::EDataStreamState::Invalid);
+	UE_NET_ASSERT_TRUE(GetDataStreamStateOnClient() == UDataStream::EDataStreamState::Invalid);
+}
+
+UE_NET_TEST_FIXTURE(FTestDynamicCreateDataStreamFixture, TestDynamicDataStreamCloseRespectsHasAcknowledgedAllReliableData)
+{
+	UReplicationSystem* ServerReplicationSystem = Server->ReplicationSystem;
+	UReplicationSystem* ClientReplicationSystem = Client->ReplicationSystem;
+
+	// Open dynamic stream, and verify that it now exists on server
+	UDataStream* ServerNetTokenStream = ServerReplicationSystem->OpenDataStream(Client->ConnectionIdOnServer, DataStreamName);
+
+	// Roundtrip
+	RoundTrip();
+
+	// Put some data in flight in separate packets
+	FNetToken ServerHelloToken = CreateAndExportTokenToClient(TEXT("Hello"));
+	Server->PreSendUpdate();
+	Server->SendTo(Client, TEXT("Hello"));
+	Server->PostSendUpdate();
+
+	FNetToken ServerHello2Token = CreateAndExportTokenToClient(TEXT("Hello2"));
+	Server->PreSendUpdate();
+	Server->SendTo(Client, TEXT("Hello2"));
+	Server->PostSendUpdate();
+
+	FNetToken ServerHello3Token = CreateAndExportTokenToClient(TEXT("Hello3"));
+	Server->PreSendUpdate();
+	Server->SendTo(Client, TEXT("Hello3"));
+	Server->PostSendUpdate();
+	
+	// Close from server
+	ServerReplicationSystem->CloseDataStream(Client->ConnectionIdOnServer, DataStreamName);
+	Server->PreSendUpdate();
+	Server->SendTo(Client, TEXT("Close"));
+	Server->PostSendUpdate();
+
+	// We expect the stream to be PendingClose as we still have important data in flight
+	UE_NET_ASSERT_TRUE(GetDataStreamStateOnServer() == UDataStream::EDataStreamState::PendingClose);
+
+	// Drop a few packets
+	Server->DeliverTo(Client, false);
+	Server->DeliverTo(Client, false);
+	Server->DeliverTo(Client, false);
+
+	// Deliver PendingClose
+	Server->DeliverTo(Client, true);
+
+	// Acknowledge pending close
+	Client->UpdateAndSend({Server}, true);
+
+	// We expect the stream to be PendingClose as we still have important data in flight
+	UE_NET_ASSERT_TRUE(GetDataStreamStateOnServer() == UDataStream::EDataStreamState::PendingClose);
+
+	// Deliver PendingClose
+	Server->UpdateAndSend({Client}, true);
+
+	// Acknowledge pending close
+	Client->UpdateAndSend({Server}, true);
+
+	// Deliver PendingClose
+	RoundTrip();
+
+	// Validate that the tokens are resolvable on the client
+	UE_NET_ASSERT_NE(ClientStringTokenStore->ResolveToken(ServerHelloToken, ClientRemoteNetTokenStoreState), nullptr);
+	UE_NET_ASSERT_NE(ClientStringTokenStore->ResolveToken(ServerHello2Token, ClientRemoteNetTokenStoreState), nullptr);
+	UE_NET_ASSERT_NE(ClientStringTokenStore->ResolveToken(ServerHello3Token, ClientRemoteNetTokenStoreState), nullptr);
+
+	// And streams are destroyed on both server and client	
+	UE_NET_ASSERT_TRUE(GetDataStreamStateOnServer() == UDataStream::EDataStreamState::Invalid);
+	UE_NET_ASSERT_TRUE(GetDataStreamStateOnClient() == UDataStream::EDataStreamState::Invalid);
+}
+
+
+
+
+
+
+
+
 }
diff --git a/Engine/Plugins/Runtime/ReplicationSystemTestPlugin/Source/Private/Tests/ReplicationSystem/ReplicationSystemServerClientTestFixture.cpp b/Engine/Plugins/Runtime/ReplicationSystemTestPlugin/Source/Private/Tests/ReplicationSystem/ReplicationSystemServerClientTestFixture.cpp
--- a/Engine/Plugins/Runtime/ReplicationSystemTestPlugin/Source/Private/Tests/ReplicationSystem/ReplicationSystemServerClientTestFixture.cpp
+++ b/Engine/Plugins/Runtime/ReplicationSystemTestPlugin/Source/Private/Tests/ReplicationSystem/ReplicationSystemServerClientTestFixture.cpp
@@ -50,15 +50,16 @@ void FDataStreamTestUtil::RestoreDataStreamDefinitions()
 	*PointerToFixupComplete = false;
 }
 
-void FDataStreamTestUtil::AddDataStreamDefinition(const TCHAR* StreamName, const TCHAR* ClassPath, bool bValid)
+void FDataStreamTestUtil::AddDataStreamDefinition(const TCHAR* StreamName, const TCHAR* ClassPath, const FAddDataStreamDefinitionParams& Params)
 {
 	FDataStreamDefinition Definition = {};
 
 	Definition.DataStreamName = FName(StreamName);
-	Definition.ClassName = bValid ? FName(ClassPath): FName();
+	Definition.ClassName = Params.bValid ? FName(ClassPath): FName();
 	Definition.Class = nullptr;
 	Definition.DefaultSendStatus = EDataStreamSendStatus::Send;
-	Definition.bAutoCreate = false;
+	Definition.bAutoCreate = Params.bAutoCreate;
+	Definition.bDynamicCreate = Params.bDynamicCreate;
 
 	CurrentDataStreamDefinitions->Add(Definition);
 }
@@ -245,32 +246,28 @@ uint32 FReplicationSystemTestNode::AddConnection()
 	FConnectionInfo Connection;
 	Connection.ConnectionId = Connections.Num() + 1U;
 
-	// Create DataStreams
-	FDataStreamManagerInitParams InitParams;
-	InitParams.PacketWindowSize = 256;
-	UDataStreamManager* DataStreamManager = NewObject<UDataStreamManager>();
-	DataStreamManager->Init(InitParams);
-	Connection.DataStreamManager = DataStreamManager;
-	CreatedObjects.Add(TStrongObjectPtr<UObject>(Connection.DataStreamManager));
+	UE_NET_TRACE_CONNECTION_CREATED(GetNetTraceId(), Connection.ConnectionId);
+	UE_NET_TRACE_CONNECTION_STATE_UPDATED(GetNetTraceId(), Connection.ConnectionId, static_cast<uint8>(3));
 
 	// Init and store RemoteNetTokenStoreState
 	ReplicationSystem->GetNetTokenStore()->InitRemoteNetTokenStoreState(Connection.ConnectionId);
-	Connection.RemoteNetTokenStoreState = ReplicationSystem->GetNetTokenStore()->GetRemoteNetTokenStoreState(Connection.ConnectionId);
+	//Connection.RemoteNetTokenStoreState = ReplicationSystem->GetNetTokenStore()->GetRemoteNetTokenStoreState(Connection.ConnectionId);
+	
+	// Add a connection
+	ReplicationSystem->AddConnection(Connection.ConnectionId);
+
+	Connection.DataStreamManager = NewObject<UDataStreamManager>();
+
+	// Make ReplicationSystem aware of new DataStreamManager associated with connection.
+	ReplicationSystem->InitDataStreamManager(Connection.ConnectionId, Connection.DataStreamManager);
+
+	CreatedObjects.Add(TStrongObjectPtr<UObject>(Connection.DataStreamManager));
 
 	// Streams created based on config
 	Connection.DataStreamManager->CreateStream("NetToken");
-	Connection.NetTokenDataStream = StaticCast<UNetTokenDataStream*>(Connection.DataStreamManager->GetStream("NetToken"));
 	Connection.DataStreamManager->CreateStream("Replication");
-	Connection.ReplicationDataStream = StaticCast<UReplicationDataStream*>(Connection.DataStreamManager->GetStream("Replication"));
-
-	UE_NET_TRACE_CONNECTION_CREATED(GetNetTraceId(), Connection.ConnectionId);
-	UE_NET_TRACE_CONNECTION_STATE_UPDATED(GetNetTraceId(), Connection.ConnectionId, static_cast<uint8>(3));
-	
-	// Add a connection
-	ReplicationSystem->AddConnection(Connection.ConnectionId);
 
-	// Initialize Streams
-	ReplicationSystem->InitDataStreams(Connection.ConnectionId, Connection.DataStreamManager);
+	// Enable replication
 	ReplicationSystem->SetReplicationEnabledForConnection(Connection.ConnectionId, true);
 
 	// Add view
diff --git a/Engine/Plugins/Runtime/ReplicationSystemTestPlugin/Source/Private/Tests/ReplicationSystem/ReplicationSystemServerClientTestFixture.h b/Engine/Plugins/Runtime/ReplicationSystemTestPlugin/Source/Private/Tests/ReplicationSystem/ReplicationSystemServerClientTestFixture.h
--- a/Engine/Plugins/Runtime/ReplicationSystemTestPlugin/Source/Private/Tests/ReplicationSystem/ReplicationSystemServerClientTestFixture.h
+++ b/Engine/Plugins/Runtime/ReplicationSystemTestPlugin/Source/Private/Tests/ReplicationSystem/ReplicationSystemServerClientTestFixture.h
@@ -37,7 +37,13 @@ public:
 
 	void StoreDataStreamDefinitions();
 	void RestoreDataStreamDefinitions();
-	void AddDataStreamDefinition(const TCHAR* StreamName, const TCHAR* ClassPath, bool bValid = true);
+	struct FAddDataStreamDefinitionParams
+	{
+		bool bAutoCreate = false;
+		bool bDynamicCreate = false;
+		bool bValid = true;
+	};
+	void AddDataStreamDefinition(const TCHAR* StreamName, const TCHAR* ClassPath, const FAddDataStreamDefinitionParams& Params = FAddDataStreamDefinitionParams());
 	void FixupDefinitions() { DataStreamDefinitions->FixupDefinitions(); }
 
 protected:
@@ -66,9 +72,6 @@ public:
 	struct FConnectionInfo
 	{
 		UDataStreamManager* DataStreamManager = nullptr;
-		UReplicationDataStream* ReplicationDataStream = nullptr;
-		UNetTokenDataStream* NetTokenDataStream = nullptr;
-		FNetTokenStoreState* RemoteNetTokenStoreState = nullptr;
 		TResizableCircularQueue<const FDataStreamRecord*> WriteRecords;
 		TResizableCircularQueue<FPacketData> WrittenPackets;
 		uint32 ConnectionId = 0;
diff --git a/Engine/Plugins/Runtime/ReplicationSystemTestPlugin/Source/Private/Tests/ReplicationSystem/TestReplicationSystem.cpp b/Engine/Plugins/Runtime/ReplicationSystemTestPlugin/Source/Private/Tests/ReplicationSystem/TestReplicationSystem.cpp
--- a/Engine/Plugins/Runtime/ReplicationSystemTestPlugin/Source/Private/Tests/ReplicationSystem/TestReplicationSystem.cpp
+++ b/Engine/Plugins/Runtime/ReplicationSystemTestPlugin/Source/Private/Tests/ReplicationSystem/TestReplicationSystem.cpp
@@ -604,7 +604,11 @@ public:
 	FNetToken CreateAndExportNetToken(const FString& TokenString)
 	{
 		FNetToken Token = ServerStringTokenStore->GetOrCreateToken(TokenString);
-		Server->GetConnectionInfo(Client->ConnectionIdOnServer).NetTokenDataStream->AddNetTokenForExplicitExport(Token);
+		UNetTokenDataStream* NetTokenDataStream = Cast<UNetTokenDataStream>(Server->GetReplicationSystem()->GetDataStream(Client->ConnectionIdOnServer, FName("NetToken")));
+		if (NetTokenDataStream)
+		{
+			NetTokenDataStream->AddNetTokenForExplicitExport(Token);
+		}
 
 		return Token;
 	}
@@ -612,7 +616,11 @@ public:
 	FNetToken CreateAndExportNetTokenOnClient(const FString& TokenString)
 	{
 		FNetToken Token = ClientStringTokenStore->GetOrCreateToken(TokenString);
-		Client->GetConnectionInfo(Client->LocalConnectionId).NetTokenDataStream->AddNetTokenForExplicitExport(Token);
+		UNetTokenDataStream* NetTokenDataStream = Cast<UNetTokenDataStream>(Client->GetReplicationSystem()->GetDataStream(Client->LocalConnectionId, FName("NetToken")));
+		if (NetTokenDataStream)
+		{
+			NetTokenDataStream->AddNetTokenForExplicitExport(Token);
+		}
 
 		return Token;
 	}
@@ -622,11 +630,16 @@ public:
 		FReplicationSystemServerClientTestFixture::SetUp();
 
 		Client = CreateClient();
-
-		ServerStringTokenStore = Server->GetReplicationSystem()->GetNetTokenStore()->GetDataStore<FStringTokenStore>();
-		ServerRemoteNetTokenState = Server->GetConnectionInfo(Client->ConnectionIdOnServer).NetTokenDataStream->GetRemoteNetTokenStoreState();
-		ClientStringTokenStore = Client->GetReplicationSystem()->GetNetTokenStore()->GetDataStore<FStringTokenStore>();
-		ClientRemoteNetTokenState = Client->GetConnectionInfo(Client->LocalConnectionId).NetTokenDataStream->GetRemoteNetTokenStoreState();
+		{
+			FNetTokenStore* ServerTokenStore = Server->GetReplicationSystem()->GetNetTokenStore();
+			ServerStringTokenStore = ServerTokenStore->GetDataStore<FStringTokenStore>();		
+			ServerRemoteNetTokenState = ServerTokenStore->GetRemoteNetTokenStoreState(Client->ConnectionIdOnServer);
+		}
+		{
+			FNetTokenStore* ClientTokenStore = Client->GetReplicationSystem()->GetNetTokenStore();
+			ClientStringTokenStore = ClientTokenStore->GetDataStore<FStringTokenStore>();
+			ClientRemoteNetTokenState = ClientTokenStore->GetRemoteNetTokenStoreState(Client->LocalConnectionId);
+		}
 	}
 };
 
diff --git a/Engine/Plugins/Runtime/ReplicationSystemTestPlugin/Source/Private/Tests/Serialization/TestFieldPathNetSerializer.cpp b/Engine/Plugins/Runtime/ReplicationSystemTestPlugin/Source/Private/Tests/Serialization/TestFieldPathNetSerializer.cpp
--- a/Engine/Plugins/Runtime/ReplicationSystemTestPlugin/Source/Private/Tests/Serialization/TestFieldPathNetSerializer.cpp
+++ b/Engine/Plugins/Runtime/ReplicationSystemTestPlugin/Source/Private/Tests/Serialization/TestFieldPathNetSerializer.cpp
@@ -2,6 +2,7 @@
 
 #include "TestFieldPathNetSerializer.h"
 #include "Tests/ReplicationSystem/ReplicationSystemServerClientTestFixture.h"
+#include "Iris/ReplicationSystem/NetTokenStore.h"
 #include "Iris/ReplicationState/ReplicationStateDescriptorBuilder.h"
 #include "Iris/Serialization/NetBitStreamReader.h"
 #include "Iris/Serialization/NetBitStreamWriter.h"
@@ -111,11 +112,12 @@ void FTestFieldPathNetSerializer::SetUp()
 
 		{
 			const FReplicationSystemTestNode::FConnectionInfo& ConnectionInfo = Server->GetConnectionInfo(Client->ConnectionIdOnServer);
+			FNetTokenStore* ServerTokenStore = Server->GetReplicationSystem()->GetNetTokenStore();
 
 			FInternalNetSerializationContext TempServerInternalContext;
 			FInternalNetSerializationContext::FInitParameters TempServerInternalContextInitParams;
 			TempServerInternalContextInitParams.ReplicationSystem = Server->GetReplicationSystem();
-			TempServerInternalContextInitParams.ObjectResolveContext.RemoteNetTokenStoreState = ConnectionInfo.RemoteNetTokenStoreState;
+			TempServerInternalContextInitParams.ObjectResolveContext.RemoteNetTokenStoreState = ServerTokenStore->GetRemoteNetTokenStoreState(Client->ConnectionIdOnServer);
 			TempServerInternalContextInitParams.ObjectResolveContext.ConnectionId = ConnectionInfo.ConnectionId;
 			TempServerInternalContext.Init(TempServerInternalContextInitParams);
 
@@ -125,10 +127,12 @@ void FTestFieldPathNetSerializer::SetUp()
 		{
 			const FReplicationSystemTestNode::FConnectionInfo& ConnectionInfo = Client->GetConnectionInfo(Client->LocalConnectionId);
 
+			FNetTokenStore* ClientTokenStore = Client->GetReplicationSystem()->GetNetTokenStore();
+
 			FInternalNetSerializationContext TempClientInternalContext;
 			FInternalNetSerializationContext::FInitParameters TempClientInternalContextInitParams;
 			TempClientInternalContextInitParams.ReplicationSystem = Client->GetReplicationSystem();
-			TempClientInternalContextInitParams.ObjectResolveContext.RemoteNetTokenStoreState = ConnectionInfo.RemoteNetTokenStoreState;
+			TempClientInternalContextInitParams.ObjectResolveContext.RemoteNetTokenStoreState = ClientTokenStore->GetRemoteNetTokenStoreState(Client->LocalConnectionId);
 			TempClientInternalContextInitParams.ObjectResolveContext.ConnectionId = ConnectionInfo.ConnectionId;
 			TempClientInternalContext.Init(TempClientInternalContextInitParams);
 
diff --git a/Engine/Source/Runtime/Engine/Private/Net/Experimental/Iris/DataStreamChannel.cpp b/Engine/Source/Runtime/Engine/Private/Net/Experimental/Iris/DataStreamChannel.cpp
--- a/Engine/Source/Runtime/Engine/Private/Net/Experimental/Iris/DataStreamChannel.cpp
+++ b/Engine/Source/Runtime/Engine/Private/Net/Experimental/Iris/DataStreamChannel.cpp
@@ -59,37 +59,36 @@ UDataStreamChannel::UDataStreamChannel(const FObjectInitializer& ObjectInitializ
 	ChName = FName("DataStream");
 }
 
-void UDataStreamChannel::Init(UNetConnection* InConnection, int32 InChIndex, EChannelCreateFlags CreateFlags)
+void UDataStreamChannel::InitDataStreamManager()
 {
-	Super::Init(InConnection, InChIndex, CreateFlags);
+	check(Connection);
 
 #if UE_WITH_IRIS
+	if (UReplicationSystem* ReplicationSystem = Connection->Driver->GetReplicationSystem())
 	{
-		FDataStreamManagerInitParams InitParams;
-		InitParams.PacketWindowSize = MaxPacketsInFlightCount;
 		DataStreamManager = NewObject<UDataStreamManager>();
-		DataStreamManager->Init(InitParams);
-	}
 
-	if (UReplicationSystem* ReplicationSystem = InConnection->Driver->GetReplicationSystem())
-	{
+		// Make ReplicationSystem aware of new DataStreamManager associated with connection.
+		ReplicationSystem->InitDataStreamManager(Connection->GetConnectionHandle().GetParentConnectionId(), DataStreamManager);
+
 		bIsReadyToHandshake = 1U;
-		ReplicationSystem->InitDataStreams(InConnection->GetConnectionHandle().GetParentConnectionId(), DataStreamManager);
 	}
-#endif // UE_WITH_IRIS
+	else
+	{
+		bIsReadyToHandshake = 0U;
+	}
+#endif
 }
 
-void UDataStreamChannel::ReInit()
+void UDataStreamChannel::Init(UNetConnection* InConnection, int32 InChIndex, EChannelCreateFlags CreateFlags)
 {
-	check(Connection);
+	Super::Init(InConnection, InChIndex, CreateFlags);
+	InitDataStreamManager();
+}
 
-#if UE_WITH_IRIS
-	if (UReplicationSystem* ReplicationSystem = Connection->Driver->GetReplicationSystem())
-	{
-		bIsReadyToHandshake = 1U;
-		ReplicationSystem->InitDataStreams(Connection->GetConnectionHandle().GetParentConnectionId(), DataStreamManager);
-	}
-#endif
+void UDataStreamChannel::ReInit()
+{
+	InitDataStreamManager();
 }
 
 bool UDataStreamChannel::CleanUp(const bool bForDestroy, EChannelCloseReason CloseReason)
diff --git a/Engine/Source/Runtime/Engine/Private/Net/Experimental/Iris/DataStreamChannel.h b/Engine/Source/Runtime/Engine/Private/Net/Experimental/Iris/DataStreamChannel.h
--- a/Engine/Source/Runtime/Engine/Private/Net/Experimental/Iris/DataStreamChannel.h
+++ b/Engine/Source/Runtime/Engine/Private/Net/Experimental/Iris/DataStreamChannel.h
@@ -63,6 +63,8 @@ private:
 	ENGINE_API virtual bool HasAcknowledgedAllReliableData() const override;
 
 private:
+	void InitDataStreamManager();
+
 	enum : uint32
 	{
 		MaxPacketsInFlightCount = 256U,
diff --git a/Engine/Source/Runtime/Experimental/Iris/Core/Private/Iris/DataStream/DataStream.cpp b/Engine/Source/Runtime/Experimental/Iris/Core/Private/Iris/DataStream/DataStream.cpp
--- a/Engine/Source/Runtime/Experimental/Iris/Core/Private/Iris/DataStream/DataStream.cpp
+++ b/Engine/Source/Runtime/Experimental/Iris/Core/Private/Iris/DataStream/DataStream.cpp
@@ -15,3 +15,25 @@ void UDataStream::EndWrite()
 {
 }
 
+void UDataStream::Init(const FInitParameters& Params)
+{
+}
+
+void UDataStream::Deinit()
+{
+}
+
+const TCHAR* LexToString(const UDataStream::EDataStreamState State)
+{
+	static const TCHAR* Names[] = {
+		TEXT("Invalid"),
+		TEXT("PendingCreate"),
+		TEXT("WaitOnCreateConfirmation"),
+		TEXT("Open"),
+		TEXT("PendingClose"),
+		TEXT("WaitOnCloseConfirmation"),
+	};
+	static_assert(UE_ARRAY_COUNT(Names) == uint32(UDataStream::EDataStreamState::Count), "Missing names for one or more values of EDataStreamState.");
+
+	return State < UDataStream::EDataStreamState::Count ? Names[(uint32)State] : TEXT("");
+}
diff --git a/Engine/Source/Runtime/Experimental/Iris/Core/Private/Iris/DataStream/DataStreamDefinitions.cpp b/Engine/Source/Runtime/Experimental/Iris/Core/Private/Iris/DataStream/DataStreamDefinitions.cpp
--- a/Engine/Source/Runtime/Experimental/Iris/Core/Private/Iris/DataStream/DataStreamDefinitions.cpp
+++ b/Engine/Source/Runtime/Experimental/Iris/Core/Private/Iris/DataStream/DataStreamDefinitions.cpp
@@ -17,6 +17,7 @@ void UDataStreamDefinitions::FixupDefinitions()
 		return;
 	}
 
+	int32 CurrentStreamIndex = 0;
 	for (FDataStreamDefinition& Definition : DataStreamDefinitions)
 	{
 		UE_CLOG(DataStreamDefinitions.ContainsByPredicate([Name = Definition.DataStreamName, &Definition](const FDataStreamDefinition& ExistingDefinition) { return Name == ExistingDefinition.DataStreamName && &Definition != &ExistingDefinition; }), LogIris, Error, TEXT("DataStream name is defined multiple times: %s."), *Definition.DataStreamName.GetPlainNameString());
@@ -25,21 +26,33 @@ void UDataStreamDefinitions::FixupDefinitions()
 		Definition.Class = StaticLoadClass(UDataStream::StaticClass(), nullptr, *Definition.ClassName.ToString(), nullptr, LOAD_Quiet);
 
 		UE_CLOG(Definition.Class == nullptr, LogIris, Error, TEXT("DataStream class could not be loaded: %s"), *Definition.ClassName.GetPlainNameString());
+
+		Definition.StreamIndex = CurrentStreamIndex++;
 	}
 
 	bFixupComplete = true;
 }
 
+int32 UDataStreamDefinitions::GetStreamIndex(const FDataStreamDefinition& Definition)
+{
+	return Definition.StreamIndex;
+}
+
 const FDataStreamDefinition* UDataStreamDefinitions::FindDefinition(const FName Name) const
 {
 	return DataStreamDefinitions.FindByPredicate([Name](const FDataStreamDefinition& Definition) { return Name == Definition.DataStreamName; });
 }
 
-void UDataStreamDefinitions::GetAutoCreateStreamNames(TArray<FName>& OutStreamNames) const
+const FDataStreamDefinition* UDataStreamDefinitions::FindDefinition(int32 StreamIndex) const
+{
+	return DataStreamDefinitions.FindByPredicate([StreamIndex](const FDataStreamDefinition& Definition) { return StreamIndex == Definition.StreamIndex; });
+}
+
+void UDataStreamDefinitions::GetStreamNamesToAutoCreateOrRegister(TArray<FName>& OutStreamNames) const
 {
 	for (const FDataStreamDefinition& Definition : DataStreamDefinitions)
 	{
-		if (Definition.bAutoCreate)
+		if (Definition.bAutoCreate || Definition.bDynamicCreate)
 		{
 			OutStreamNames.Add(Definition.DataStreamName);
 		}
diff --git a/Engine/Source/Runtime/Experimental/Iris/Core/Private/Iris/DataStream/DataStreamDefinitions.h b/Engine/Source/Runtime/Experimental/Iris/Core/Private/Iris/DataStream/DataStreamDefinitions.h
--- a/Engine/Source/Runtime/Experimental/Iris/Core/Private/Iris/DataStream/DataStreamDefinitions.h
+++ b/Engine/Source/Runtime/Experimental/Iris/Core/Private/Iris/DataStream/DataStreamDefinitions.h
@@ -31,6 +31,17 @@ struct FDataStreamDefinition
 	// Whether the DataStream should be auto created for each connection. If not then CreateStream need be called manually.
 	UPROPERTY()
 	bool bAutoCreate = false;
+
+	// If bDynamicCreate is set to true we will reserve a slot for the stream allowing it to be openened and closed on demand
+	UPROPERTY()
+	bool bDynamicCreate = false;
+
+	// Get the assigned stream index
+	int32 GetStreamIndex() const;
+
+private:
+	friend class UDataStreamDefinitions;
+	int32 StreamIndex = -1;
 };
 
 UCLASS(transient, config=Engine)
@@ -46,7 +57,9 @@ protected:
 
 	void FixupDefinitions();
 	const FDataStreamDefinition* FindDefinition(const FName Name) const;
-	void GetAutoCreateStreamNames(TArray<FName>& OutStreamNames) const;
+	const FDataStreamDefinition* FindDefinition(int32 StreamIndex) const;
+	static int32 GetStreamIndex(const FDataStreamDefinition& Definition);
+	void GetStreamNamesToAutoCreateOrRegister(TArray<FName>& OutStreamNames) const;
 
 private:
 	UPROPERTY(Config)
diff --git a/Engine/Source/Runtime/Experimental/Iris/Core/Private/Iris/DataStream/DataStreamManager.cpp b/Engine/Source/Runtime/Experimental/Iris/Core/Private/Iris/DataStream/DataStreamManager.cpp
--- a/Engine/Source/Runtime/Experimental/Iris/Core/Private/Iris/DataStream/DataStreamManager.cpp
+++ b/Engine/Source/Runtime/Experimental/Iris/Core/Private/Iris/DataStream/DataStreamManager.cpp
@@ -9,15 +9,30 @@
 #include "Iris/Serialization/NetBitStreamWriter.h"
 #include "Iris/Serialization/NetSerializationContext.h"
 #include "Iris/ReplicationSystem/NetExports.h"
+#include "Iris/Serialization/NetBitStreamUtil.h"
 #include "UObject/Package.h"
 
+#define UE_LOG_DATASTREAM_CONN(Verbosity, Format, ...)  UE_LOG(LogIris, Verbosity, TEXT("DataStreamManager: R:%u :C%u ") Format, InitParameters.ReplicationSystemId, InitParameters.ConnectionId, ##__VA_ARGS__)
+
 class UDataStreamManager::FImpl
 {
+public:
+	using EDataStreamState = UDataStream::EDataStreamState;
+
+	enum class ECreateDataStreamFlags : uint8
+	{
+		// Create stream
+		None,
+		// Streams marked with bDynamicCreate will only be registered if this flag is set.
+		RegisterIfStreamIsDynamic,
+	};
+	FRIEND_ENUM_CLASS_FLAGS(ECreateDataStreamFlags);
+
 public:
 	FImpl();
 	~FImpl();
 
-	void Init(const FDataStreamManagerInitParams& InitParams);
+	void Init(const UDataStream::FInitParameters& InitParams);
 	void Deinit();
 
 	EWriteResult BeginWrite(const UDataStream::FBeginWriteParameters& Params);
@@ -26,8 +41,9 @@ public:
 	void ReadData(UE::Net::FNetSerializationContext& context);
 	void ProcessPacketDeliveryStatus(UE::Net::EPacketDeliveryStatus Status, FDataStreamRecord const* Record);
 	bool HasAcknowledgedAllReliableData() const;
-
-	ECreateDataStreamResult CreateStream(const FName StreamName);
+	
+	ECreateDataStreamResult CreateStream(const FName StreamName, ECreateDataStreamFlags Flags = ECreateDataStreamFlags::None);
+	void CloseStream(const FName StreamName);
 
 	const UDataStream* GetStream(const FName StreamName) const;
 	UDataStream* GetStream(const FName StreamName);
@@ -35,6 +51,8 @@ public:
 	void SetSendStatus(const FName StreamName, EDataStreamSendStatus Status);
 	EDataStreamSendStatus GetSendStatus(const FName StreamName) const;
 
+	EDataStreamState GetStreamState(const FName StreamName) const;
+
 	UE::Net::Private::FNetExports& GetNetExports();
 
 	void AddReferencedObjects(FReferenceCollector& Collector);
@@ -44,7 +62,7 @@ private:
 	{
 		inline FindStreamByName(FName InName) : Name(InName) {}
 		
-		inline bool operator()(const UDataStream* Stream) const { return Name.IsEqual(Stream->GetFName(), ENameCase::IgnoreCase, false); }
+		inline bool operator()(const UDataStream* Stream) const { return Stream && Name.IsEqual(Stream->GetFName(), ENameCase::IgnoreCase, false); }
 		inline bool operator()(const FDataStreamDefinition& Definition) const { return Name == Definition.DataStreamName; }
 
 	private:
@@ -55,25 +73,40 @@ private:
 	{
 		TArray<const FDataStreamRecord*, TInlineAllocator<8>> DataStreamRecords;
 		uint32 DataStreamMask;
+		// What streams carried state changes in last record
+		uint32 DataStreamStateMask;
 	};
 
 	void InitRecordStorage();
+	void InitStream(UDataStream* Stream);
 	void InitStreams();
+	void DestroyStream(uint32 StreamIndex);
+	void SetStreamState(uint32 StreamIndex, EDataStreamState NewState);
+	EDataStreamState GetStreamState(uint32 StreamIndex) const;
+	ECreateDataStreamResult CreateStreamFromDefinition(const FDataStreamDefinition& Definition, ECreateDataStreamFlags Flags);
+	ECreateDataStreamResult CreateStreamFromIndex(int32 StreamIndex);
+	void MarkStreamStateDirty(uint32 StreamIndex);
+	void HandleReceivedStreamState(UE::Net::FNetSerializationContext& Context, uint32 StreamIndex, EDataStreamState RecvdState);
 
 private:
 	static constexpr uint32 MaxStreamCount = 32U;
 	static constexpr uint32 StreamCountBitCount = 5U; // Enough for 32 streams
+	static constexpr uint32 StreamStateBitCount = 4U; // Enough for 16 states
 
 	UE::Net::Private::FNetExports NetExports;
 
 	// We can afford reserving space for a few pointers It's unlikely we will create anything close to 16 streams.
 	TArray<TObjectPtr<UDataStream>> Streams;
 	TArray<EDataStreamSendStatus> StreamSendStatus;
+	TArray<EDataStreamState> StreamState;
 	TArray<FRecord> RecordStorage;
 	TResizableCircularQueue<FRecord*> Records;
-	uint32 PacketWindowSize = 0;
+	UDataStream::FInitParameters InitParameters;
+	uint32 DirtyStreamsMask = 0U;
 };
 
+static_assert((uint32)(UDataStream::EDataStreamState::Count) <= 15U, "EDataStreamState must fit in 4 bits.");
+
 UDataStreamManager::UDataStreamManager()
 {
 	if (!HasAnyFlags(RF_ClassDefaultObject))
@@ -86,7 +119,7 @@ UDataStreamManager::~UDataStreamManager()
 {
 }
 
-void UDataStreamManager::Init(const FDataStreamManagerInitParams& InitParams)
+void UDataStreamManager::Init(const UDataStream::FInitParameters& InitParams)
 {
 	Impl->Init(InitParams);
 }
@@ -136,6 +169,21 @@ UDataStream* UDataStreamManager::GetStream(const FName StreamName)
 	return Impl->GetStream(StreamName);
 }
 
+const UDataStream* UDataStreamManager::GetStream(const FName StreamName) const
+{
+	return Impl->GetStream(StreamName);
+}
+
+void UDataStreamManager::CloseStream(const FName StreamName)
+{
+	Impl->CloseStream(StreamName);
+}
+
+UDataStream::EDataStreamState UDataStreamManager::GetStreamState(const FName StreamName) const
+{
+	return Impl->GetStreamState(StreamName);
+}
+
 void UDataStreamManager::SetSendStatus(const FName StreamName, EDataStreamSendStatus Status)
 {
 	return Impl->SetSendStatus(StreamName, Status);
@@ -171,9 +219,11 @@ UDataStreamManager::FImpl::~FImpl()
 {
 }
 
-void UDataStreamManager::FImpl::Init(const FDataStreamManagerInitParams& InitParams)
+void UDataStreamManager::FImpl::Init(const UDataStreamManager::FInitParameters& InitParams)
 {
-	PacketWindowSize = InitParams.PacketWindowSize;
+	InitParameters = InitParams;
+	InitParameters.NetExports = &NetExports;
+
 	InitRecordStorage();
 	InitStreams();
 }
@@ -191,12 +241,27 @@ void UDataStreamManager::FImpl::Deinit()
 	{
 		if (IsValid(Stream))
 		{
+			Stream->Deinit();
 			Stream->MarkAsGarbage();
 		}
 	}
 
 	Streams.Reset();
 	StreamSendStatus.Reset();
+	StreamState.Reset();
+}
+
+void UDataStreamManager::FImpl::DestroyStream(uint32 StreamIndex)
+{
+	UDataStream* Stream = Streams[StreamIndex];
+	if (IsValid(Stream))
+	{
+		Stream->Deinit();
+		Stream->MarkAsGarbage();
+		Streams[StreamIndex] = nullptr;
+		StreamState[StreamIndex] = EDataStreamState::Invalid;
+		StreamSendStatus[StreamIndex] = EDataStreamSendStatus::Pause;
+	}
 }
 
 UDataStreamManager::EWriteResult UDataStreamManager::FImpl::BeginWrite(const UDataStream::FBeginWriteParameters& Params)
@@ -209,7 +274,9 @@ UDataStreamManager::EWriteResult UDataStreamManager::FImpl::BeginWrite(const UDa
 
 	TObjectPtr<UDataStream>* StreamData = Streams.GetData();
 	const EDataStreamSendStatus* SendStatusData = StreamSendStatus.GetData();
-	UDataStream::EWriteResult CombinedWriteResult = UDataStream::EWriteResult::NoData;
+	const EDataStreamState* StreamStateData = StreamState.GetData();
+
+	UDataStream::EWriteResult CombinedWriteResult = DirtyStreamsMask == 0U ? UDataStream::EWriteResult::NoData : UDataStream::EWriteResult::HasMoreData;
 
 	for (SIZE_T StreamIt = 0, StreamEndIt = StreamCount; StreamIt != StreamEndIt; ++StreamIt)
 	{
@@ -277,8 +344,61 @@ UDataStreamManager::EWriteResult UDataStreamManager::FImpl::WriteData(UE::Net::F
 	FDataStreamRecord const** TempStreamRecords = TempRecord.DataStreamRecords.GetData();
 
 	FNetBitStreamWriter ManagerStream = Context.GetBitStreamWriter()->CreateSubstream();
+	// This will write the number of bits required for the StreamBitCount
 	ManagerStream.WriteBits(0U, StreamCountBitCount);
+	// Will be rewritten later to contain a bit mask for all streams that have written data.
 	ManagerStream.WriteBits(0U, StreamCount);
+
+	const bool bHasStreamsWithDirtyState = DirtyStreamsMask != 0U;
+	if (ManagerStream.WriteBool(bHasStreamsWithDirtyState))
+	{
+		ManagerStream.WriteBits(DirtyStreamsMask, StreamCount);
+	}
+
+	TObjectPtr<UDataStream>* StreamData = Streams.GetData();
+	const EDataStreamSendStatus* SendStatusData = StreamSendStatus.GetData();
+
+	UDataStream::EWriteResult CombinedWriteResult = bHasStreamsWithDirtyState ? UDataStream::EWriteResult::Ok : UDataStream::EWriteResult::NoData;
+
+	// Write rare datastream state changes
+	if (bHasStreamsWithDirtyState)
+	{
+		TArray<EDataStreamState> OldStreamState = StreamState;
+		for (uint32 StreamIt = 0U, StreamEndIt = StreamCount, CurrentStreamMask = 1U; StreamIt != StreamEndIt; ++StreamIt, CurrentStreamMask += CurrentStreamMask)
+		{
+			if (DirtyStreamsMask & CurrentStreamMask)
+			{
+				const EDataStreamState State = GetStreamState(StreamIt);
+				// Write state
+				ManagerStream.WriteBits((uint32)State, StreamStateBitCount);
+				UE_LOG_DATASTREAM_CONN(Verbose, TEXT("WriteStreamState for StreamIndex: %u, State: %s"), StreamIt, LexToString(State));
+				switch (State)
+				{
+					case EDataStreamState::PendingCreate:
+					{
+						// If we would like to add more data for create, this would be the spot.
+						SetStreamState(StreamIt, EDataStreamState::WaitOnCreateConfirmation);
+						break;
+					}
+					case EDataStreamState::PendingClose:
+					{
+						// For now, if we have no data to flush we can go directly to WaitOnCloseConfirmation
+						UDataStream* Stream = Streams[StreamIt];
+						if (Stream->HasAcknowledgedAllReliableData())
+						{
+							SetStreamState(StreamIt, EDataStreamState::WaitOnCloseConfirmation);
+						}
+						break;
+					}
+					default:
+					{
+						break;
+					}
+				}
+			}
+		}
+	}
+
 	// If we can't fit our header we can't fit anything else either.
 	if (ManagerStream.IsOverflown())
 	{
@@ -286,13 +406,8 @@ UDataStreamManager::EWriteResult UDataStreamManager::FImpl::WriteData(UE::Net::F
 		return EWriteResult::NoData;
 	}
 
-	TObjectPtr<UDataStream>* StreamData = Streams.GetData();
-	const EDataStreamSendStatus* SendStatusData = StreamSendStatus.GetData();
-	UDataStream::EWriteResult CombinedWriteResult = UDataStream::EWriteResult::NoData;
-
 	uint32 DataStreamMask = 0;
-	uint32 CurrentStreamMask = 1;
-	for (SIZE_T StreamIt = 0, StreamEndIt = StreamCount; StreamIt != StreamEndIt; ++StreamIt, CurrentStreamMask += CurrentStreamMask)
+	for (uint32 StreamIt = 0U, StreamEndIt = StreamCount, CurrentStreamMask = 1U; StreamIt != StreamEndIt; ++StreamIt, CurrentStreamMask += CurrentStreamMask)
 	{
 		UDataStream* Stream = StreamData[StreamIt];
 		const EDataStreamSendStatus SendStatus = SendStatusData[StreamIt];
@@ -301,6 +416,13 @@ UDataStreamManager::EWriteResult UDataStreamManager::FImpl::WriteData(UE::Net::F
 			continue;
 		}
 
+		// We only write stream data of the stream is considered is open
+		const EDataStreamState State = StreamState[StreamIt];
+		if (!(State == EDataStreamState::Open || State == EDataStreamState::PendingClose))
+		{
+			continue;
+		}
+
 		FNetBitStreamWriter SubBitStream = ManagerStream.CreateSubstream();
 		FNetSerializationContext SubContext = Context.MakeSubContext(&SubBitStream);
 
@@ -330,7 +452,7 @@ UDataStreamManager::EWriteResult UDataStreamManager::FImpl::WriteData(UE::Net::F
 		CombinedWriteResult = (CombinedWriteResult == EWriteResult::HasMoreData || WriteResult == EWriteResult::HasMoreData) ? EWriteResult::HasMoreData : WriteResult;
 	}
 
-	if (CombinedWriteResult == EWriteResult::NoData)
+	if ((CombinedWriteResult == EWriteResult::NoData) && !bHasStreamsWithDirtyState)
 	{
 		Context.GetBitStreamWriter()->DiscardSubstream(ManagerStream);
 	}
@@ -348,6 +470,8 @@ UDataStreamManager::EWriteResult UDataStreamManager::FImpl::WriteData(UE::Net::F
 
 		// Fixup and store record
 		TempRecord.DataStreamMask = DataStreamMask;
+		TempRecord.DataStreamStateMask = DirtyStreamsMask;
+		DirtyStreamsMask = 0U;
 
 		FRecord*& Record = Records.Enqueue_GetRef();
 		*Record = MoveTemp(TempRecord);
@@ -362,6 +486,309 @@ UDataStreamManager::EWriteResult UDataStreamManager::FImpl::WriteData(UE::Net::F
 	return CombinedWriteResult;
 }
 
+void UDataStreamManager::FImpl::MarkStreamStateDirty(uint32 StreamIndex)
+{
+	DirtyStreamsMask |= 1U << StreamIndex;
+}
+
+UDataStream::EDataStreamState UDataStreamManager::FImpl::GetStreamState(uint32 StreamIndex) const
+{
+	return StreamState[StreamIndex];
+}
+
+void UDataStreamManager::FImpl::SetStreamState(uint32 StreamIndex, EDataStreamState NewState)
+{
+	const EDataStreamState CurrentState = StreamState[StreamIndex];
+
+	switch (CurrentState)
+	{
+		case EDataStreamState::Invalid:
+		{
+			if (NewState == EDataStreamState::Invalid)
+			{
+				// This means that we have rejected a PendingCreate request
+				goto AcceptStateChange;
+			}
+			else if (NewState == EDataStreamState::PendingCreate)
+			{
+				// Next write will respond with PendingCreate to confirm
+				goto AcceptStateChange;
+			}
+		}
+		break;
+		
+		case EDataStreamState::PendingCreate:
+		{
+			if (NewState == EDataStreamState::Invalid)
+			{
+				// Local request to close stream that has not yet been sent.
+				goto AcceptStateChange;
+			}
+			else if (NewState == EDataStreamState::WaitOnCreateConfirmation)			
+			{
+				// We have sent PendingCreate request
+				goto AcceptStateChange;
+			}
+		}
+		break;
+		
+		case EDataStreamState::WaitOnCreateConfirmation:
+		{
+			if (NewState == EDataStreamState::PendingCreate)
+			{
+				// We have dropped WaitOnCreateConfirmation (or PendingCreate)
+				goto AcceptStateChange;
+			}
+			else if (NewState == EDataStreamState::PendingClose)
+			{
+				// API call to close
+				goto AcceptStateChange;
+			}
+			else if (NewState == EDataStreamState::Open)
+			{
+				// Received PendingCreate/Open from other side to confirm
+				goto AcceptStateChange;
+
+			}
+			else if (NewState == EDataStreamState::Invalid)
+			{
+				// Received reject from other side.
+				goto AcceptStateChange;
+			}
+		}
+		break;
+
+		case EDataStreamState::Open:
+		{
+			if (NewState == EDataStreamState::PendingClose)
+			{
+				// API call to close or request from other side to close
+				goto AcceptStateChange;
+			}
+			else if (NewState == EDataStreamState::Open)
+			{
+				// Received PendingCreate/Open from other side to confirm
+				// nothing should be done.
+				return;
+			}
+		}
+		break;
+
+		case EDataStreamState::PendingClose:
+		{
+			if (NewState == EDataStreamState::PendingClose)
+			{
+				// API call to close or request from other side to close
+				goto AcceptStateChange;
+			}
+			else if (NewState == EDataStreamState::WaitOnCloseConfirmation)
+			{
+				// We are flushed and can cleanup
+				goto AcceptStateChange;
+			}
+			else if (NewState == EDataStreamState::Open)
+			{
+				// Received PendingCreate/Open from other side to confirm
+				// nothing should be done.
+				return;
+			}
+		}
+		break;
+
+		case EDataStreamState::WaitOnCloseConfirmation:
+		{
+			if (NewState == EDataStreamState::PendingClose)
+			{
+				// We dropped PendingClose
+				goto AcceptStateChange;				
+			}
+			else if (NewState == EDataStreamState::WaitOnCloseConfirmation)
+			{
+				// Other side has confirmed close
+				goto AcceptStateChange;
+			}
+			else if (NewState == EDataStreamState::Invalid)
+			{
+				// We can cleanup
+				goto AcceptStateChange;
+			}
+		}
+		break;
+
+		default:
+		break;
+	};
+
+	UE_LOG_DATASTREAM_CONN(Verbose, TEXT("SetDataStreamState Reject: for StreamIndex: %u, CurrentState: %s, NewState: %s"), StreamIndex, LexToString(CurrentState), LexToString(NewState));
+	ensure(false);
+	return;
+
+AcceptStateChange:
+	UE_LOG_DATASTREAM_CONN(Verbose, TEXT("SetDataStreamState Accept: for StreamIndex: %u, CurrentState: %s, NewState: %s"), StreamIndex, LexToString(CurrentState), LexToString(NewState));
+	StreamState[StreamIndex] = NewState;
+	MarkStreamStateDirty(StreamIndex);
+	return;
+}
+
+void UDataStreamManager::FImpl::HandleReceivedStreamState(UE::Net::FNetSerializationContext& Context, uint32 StreamIndex, EDataStreamState RecvdState)
+{
+	UDataStream* DataStream = Streams[StreamIndex];
+	const EDataStreamState CurrentState = GetStreamState(StreamIndex);
+
+	auto HandleUnexpectedState = [&CurrentState, &RecvdState, &StreamIndex, &Context, this]()
+	{
+		UE_LOG_DATASTREAM_CONN(Error, TEXT("Received invalid DataStream State: %s for StreamIndex: %u, while in State: %s"), LexToString(RecvdState), StreamIndex, LexToString(CurrentState));
+		Context.SetError(TEXT("Invalid DataStreamState"));
+		// Just for log attention
+		ensure(false);
+	};
+
+	switch (RecvdState)
+	{
+		case EDataStreamState::PendingCreate:
+		{
+			// PendingCreate is received to request or confirm open/create
+			if (CurrentState == EDataStreamState::Invalid)
+			{
+				// Create stream
+				if (CreateStreamFromIndex(StreamIndex) != ECreateDataStreamResult::Success)
+				{
+					// If we fail, we set state as Invalid and send that to server.
+					SetStreamState(StreamIndex, EDataStreamState::Invalid);
+				}
+			}
+			else if (CurrentState == EDataStreamState::WaitOnCreateConfirmation)
+			{
+				// Other side have now confirmed open
+				SetStreamState(StreamIndex, EDataStreamState::Open);
+			}
+			else
+			{
+				HandleUnexpectedState();
+			}
+		}
+		break;
+
+		case EDataStreamState::Open:
+		{
+			// Open is received when other side has accepted stream
+			if (CurrentState == EDataStreamState::WaitOnCreateConfirmation)
+			{
+				// We have now completed open handshake and can send data.
+				SetStreamState(StreamIndex, EDataStreamState::Open);
+			}
+			else if (CurrentState == EDataStreamState::Open)
+			{
+				// We are already open, nothing to do
+			}
+			else
+			{
+				HandleUnexpectedState();
+			}
+		}
+		break;
+		
+		case EDataStreamState::PendingClose:
+		{						
+			if (CurrentState == EDataStreamState::PendingCreate)
+			{
+				// Received Pending close while we have yet to acknowledge or sent create
+				SetStreamState(StreamIndex, EDataStreamState::WaitOnCreateConfirmation);
+				SetStreamState(StreamIndex, EDataStreamState::PendingClose);
+			} 
+			else if (CurrentState == EDataStreamState::WaitOnCreateConfirmation || CurrentState == EDataStreamState::Open)
+			{
+				// Pending close is received when other side has started to close the connection
+				// there might still be data to be flushed but no new data should be written
+				SetStreamState(StreamIndex, EDataStreamState::PendingClose);
+			}
+			else if (CurrentState == EDataStreamState::PendingClose)
+			{
+				UDataStream* Stream = Streams[StreamIndex];
+				if (Stream->HasAcknowledgedAllReliableData())
+				{
+					SetStreamState(StreamIndex, EDataStreamState::WaitOnCloseConfirmation);
+				}
+				else
+				{
+					UE_LOG_DATASTREAM_CONN(Verbose, TEXT("Flushing DataStream StreamIndex: %u in State: %s"), StreamIndex, LexToString(CurrentState));
+					MarkStreamStateDirty(StreamIndex);
+				}
+			}
+			else if (CurrentState == EDataStreamState::WaitOnCloseConfirmation)
+			{
+				// Trigger update of state machine as other side might still be flushing
+				MarkStreamStateDirty(StreamIndex);				
+			}
+			else
+			{
+				HandleUnexpectedState();
+			}
+		}
+		break;
+
+		case EDataStreamState::WaitOnCloseConfirmation:
+		{
+			if (CurrentState == EDataStreamState::WaitOnCloseConfirmation)
+			{
+				SetStreamState(StreamIndex, EDataStreamState::Invalid);
+				DestroyStream(StreamIndex);
+			}
+			else if (CurrentState == EDataStreamState::PendingClose)
+			{
+				if (DataStream->HasAcknowledgedAllReliableData())
+				{
+					SetStreamState(StreamIndex, EDataStreamState::WaitOnCloseConfirmation);
+				}
+				else
+				{
+					// Trigger update of state machine as other side might still be flushing
+					UE_LOG_DATASTREAM_CONN(Verbose, TEXT("Flushing DataStream StreamIndex: %u in State: %s"), StreamIndex, LexToString(CurrentState));
+					MarkStreamStateDirty(StreamIndex);
+				}
+			}
+			else
+			{
+				HandleUnexpectedState();
+			}
+		}
+		break;
+
+		case EDataStreamState::Invalid:
+		{
+			// Sent when a stream is invalidated
+			if (CurrentState == EDataStreamState::Invalid)
+			{
+				// Do nothing
+			}
+			else if (CurrentState == EDataStreamState::WaitOnCreateConfirmation)
+			{
+				// Report error and close stream
+				SetStreamState(StreamIndex, EDataStreamState::PendingClose);
+				SetStreamState(StreamIndex, EDataStreamState::WaitOnCloseConfirmation);
+
+				SetStreamState(StreamIndex, EDataStreamState::Invalid);
+				DestroyStream(StreamIndex);
+			}
+			else if (CurrentState == EDataStreamState::WaitOnCloseConfirmation)
+			{
+				// Ready to destroy
+				SetStreamState(StreamIndex, EDataStreamState::Invalid);
+				DestroyStream(StreamIndex);
+			}
+			else
+			{
+				HandleUnexpectedState();
+			}
+		}
+		break;
+
+		default:
+			HandleUnexpectedState();
+			break;
+	}
+}
+
 void UDataStreamManager::FImpl::ReadData(UE::Net::FNetSerializationContext& Context)
 {
 	using namespace UE::Net;
@@ -369,20 +796,56 @@ void UDataStreamManager::FImpl::ReadData(UE::Net::FNetSerializationContext& Cont
 	FNetBitStreamReader* Stream = Context.GetBitStreamReader();
 	const uint32 StreamCount = 1U + Stream->ReadBits(StreamCountBitCount);
 	const uint32 DataStreamMask = Stream->ReadBits(StreamCount);
+
+	// Read and apply stream state changes
+	const bool bHasDataStreamStateChanges = Stream->ReadBool();
+	const uint32 DataStreamsWithChangedStateMask = bHasDataStreamStateChanges ? Stream->ReadBits(StreamCount) : 0U;
+
+	if (Context.HasErrorOrOverflow())
+	{
+		return;
+	}
+
 	// Validate the received information
-	if (StreamCount > uint32(Streams.Num()) || DataStreamMask == 0U)
+	if (StreamCount > uint32(Streams.Num()) || (DataStreamMask == 0U && !bHasDataStreamStateChanges))
 	{
 		Context.SetError(GNetError_BitStreamError);
 		return;
 	}
 
 	TObjectPtr<UDataStream>* StreamData = Streams.GetData();
+
+	if (bHasDataStreamStateChanges)
+	{
+		for (uint32 StreamIt = 0, StreamEndIt = StreamCount, Mask = 1U; StreamIt != StreamEndIt; ++StreamIt, Mask += Mask)
+		{
+			if (DataStreamsWithChangedStateMask & Mask)
+			{
+				// Read Stream state
+				EDataStreamState RcvdState = (EDataStreamState)Stream->ReadBits(StreamStateBitCount);
+
+				// If something went wrong we should stop deserializing immediately.
+				if (Context.HasErrorOrOverflow())
+				{
+					return;
+				}
+
+				UE_LOG_DATASTREAM_CONN(Verbose, TEXT("ReadStreamState for StreamIndex: %u, State: %s"), StreamIt, LexToString(RcvdState));
+				HandleReceivedStreamState(Context, StreamIt, RcvdState);
+			}
+		}
+	}
+
 	for (uint32 StreamIt = 0, StreamEndIt = StreamCount, Mask = 1U; StreamIt != StreamEndIt; ++StreamIt, Mask += Mask)
 	{
 		if (DataStreamMask & Mask)
 		{
+			// We should always have a DataStream here.
 			UDataStream* DataStream = StreamData[StreamIt];
-			DataStream->ReadData(Context);
+			if (ensure(DataStream))
+			{
+				DataStream->ReadData(Context);
+			}
 			// If something went wrong we should stop deserializing immediately.
 			if (Context.HasErrorOrOverflow())
 			{
@@ -402,15 +865,40 @@ void UDataStreamManager::FImpl::ProcessPacketDeliveryStatus(UE::Net::EPacketDeli
 
 	// Forward the call to each DataStream that was included in the record.
 	const uint32 DataStreamMask = Record->DataStreamMask;
+	const uint32 DataStreamStateMask = Record->DataStreamStateMask;
 	const uint32 StreamCount = Streams.Num();
 
 	TObjectPtr<UDataStream>* StreamData = Streams.GetData();
 	for (uint32 StreamIt = 0, StreamEndIt = StreamCount, Mask = 1U; StreamIt != StreamEndIt; ++StreamIt, Mask += Mask)
 	{
+		if (DataStreamStateMask & Mask)
+		{
+			// State transitions are driven explicitly, but if we drop a transition we must dirty the StreamState to trigger a write.			
+			if  (Status == UE::Net::EPacketDeliveryStatus::Lost)
+			{
+				UE_LOG_DATASTREAM_CONN(Verbose, TEXT("ProcessPacketDeliveryStatus Handle Lost DataStream State for StreamIndex: %u CurrentState: %s "), StreamIt, LexToString(StreamState[StreamIt]));
+
+				// Note: As we do not store actual lost state in record, we are pessimistic for pendingcreate/close
+				if (StreamState[StreamIt] == EDataStreamState::WaitOnCreateConfirmation)
+				{
+					SetStreamState(StreamIt, EDataStreamState::PendingCreate);
+				}
+				else if (StreamState[StreamIt] == EDataStreamState::WaitOnCloseConfirmation)
+				{
+					SetStreamState(StreamIt, EDataStreamState::PendingClose);
+				}
+				MarkStreamStateDirty(StreamIt);
+			}
+		}
+
 		if (DataStreamMask & Mask)
 		{
 			UDataStream* DataStream = StreamData[StreamIt];
-			DataStream->ProcessPacketDeliveryStatus(Status, Record->DataStreamRecords[StreamIt]);
+			// We should always have a DataStream here.
+			if (ensure(DataStream))
+			{
+				DataStream->ProcessPacketDeliveryStatus(Status, Record->DataStreamRecords[StreamIt]);
+			}
 		}
 	}
 
@@ -430,33 +918,150 @@ bool UDataStreamManager::FImpl::HasAcknowledgedAllReliableData() const
 	return true;
 }
 
-ECreateDataStreamResult UDataStreamManager::FImpl::CreateStream(const FName StreamName)
+void UDataStreamManager::FImpl::CloseStream(const FName StreamName)
 {
+	// Find index
+	const UDataStreamDefinitions* StreamDefinitions = GetDefault<UDataStreamDefinitions>();
+	const FDataStreamDefinition* Definition = StreamDefinitions->FindDefinition(StreamName);
+
+	if (!Definition)
+	{
+		UE_LOG_DATASTREAM_CONN(Warning, TEXT("UDataStreamManager::FImpl::CloseStream No DataStreamDefinition exists for name '%s' exists."), ToCStr(StreamName.GetPlainNameString()));
+		return;
+	}
+
+	if (!Definition->bDynamicCreate)
+	{
+		UE_LOG_DATASTREAM_CONN(Warning, TEXT("UDataStreamManager::FImpl::CloseStream cannot request DataStream'%s' to be closed as it is not marked as bDynamicCreate."), ToCStr(StreamName.GetPlainNameString()));
+		return;
+	}
+
+	if (!Streams.ContainsByPredicate(FindStreamByName(StreamName)))
+	{
+		UE_LOG_DATASTREAM_CONN(Warning, TEXT("UDataStreamManager::FImpl::CloseStream No DataStream with name '%s' exists."), ToCStr(StreamName.GetPlainNameString()));
+		return;
+	}
+
+	const uint32 StreamIndex = StreamDefinitions->GetStreamIndex(*Definition);
+	const EDataStreamState CurrentState = GetStreamState(StreamIndex);
+	switch (CurrentState)
+	{
+		case EDataStreamState::PendingCreate:
+		{
+			// If we are in PendingCreate it means that we either have not yet sent create request and can go back to invalid and release the stream
+			SetStreamState(StreamIndex, EDataStreamState::Invalid);
+			DestroyStream(StreamIndex);
+			break;
+		}
+		case EDataStreamState::WaitOnCreateConfirmation:
+		{
+			// If we are closed while waiting for create confirmation.
+			SetStreamState(StreamIndex, EDataStreamState::PendingClose);
+			break;
+		}
+		case EDataStreamState::Open:
+		{
+			SetStreamState(StreamIndex, EDataStreamState::PendingClose);
+			break;
+		}
+		default:
+		{
+			break;
+		}
+	};
+}
+
+ECreateDataStreamResult UDataStreamManager::FImpl::CreateStreamFromDefinition(const FDataStreamDefinition& Definition, ECreateDataStreamFlags Flags)
+{
+	if (Definition.Class == nullptr)
+	{
+		return ECreateDataStreamResult::Error_InvalidDefinition;
+	}
+
+	const int32 WantedStreamIndex = UDataStreamDefinitions::GetStreamIndex(Definition);
+	if (WantedStreamIndex == -1)
+	{
+		return ECreateDataStreamResult::Error_InvalidDefinition;
+	}
+
 	// Bumping MaxStreamCount may require modifying the FRecord and WriteData/ReadData.
-	if (Streams.Num() >= MaxStreamCount)
+	if (WantedStreamIndex >= MaxStreamCount)
 	{
 		return ECreateDataStreamResult::Error_TooManyStreams;
 	}
 
+	// Make room
+	const int32 RequiredStreamCount = WantedStreamIndex + 1;
+	if (Streams.Num() < RequiredStreamCount)
+	{
+		Streams.SetNum(RequiredStreamCount, EAllowShrinking::No);
+		StreamSendStatus.SetNumZeroed(RequiredStreamCount, EAllowShrinking::No);
+		StreamState.SetNumZeroed(RequiredStreamCount, EAllowShrinking::No);
+	}
+
+	UDataStream* Stream = nullptr;
+
+	bool bIsDynamic = Definition.bDynamicCreate;
+	bool bShouldCreate = !bIsDynamic || !EnumHasAnyFlags(Flags, ECreateDataStreamFlags::RegisterIfStreamIsDynamic);
+	if (bShouldCreate)
+	{
+		Stream = NewObject<UDataStream>(GetTransientPackage(), ToRawPtr(Definition.Class), MakeUniqueObjectName(nullptr, Definition.Class, Definition.DataStreamName));
+
+		Streams[WantedStreamIndex] = Stream;
+		StreamSendStatus[WantedStreamIndex] = Definition.DefaultSendStatus;
+
+		// Auto created streams are always considered to be opened
+		StreamState[WantedStreamIndex] = bIsDynamic ? EDataStreamState::PendingCreate : EDataStreamState::Open;
+		if (bIsDynamic)
+		{
+			MarkStreamStateDirty(WantedStreamIndex);
+		}
+		UE_LOG_DATASTREAM_CONN(Verbose, TEXT("Created DataStream with name '%s' with streamindex: %d State:%s"), ToCStr(Definition.DataStreamName.ToString()), WantedStreamIndex, LexToString(StreamState[WantedStreamIndex]));
+	}
+	else
+	{
+		Streams[WantedStreamIndex] = nullptr;
+		StreamSendStatus[WantedStreamIndex] = EDataStreamSendStatus::Pause;
+		StreamState[WantedStreamIndex] = EDataStreamState::Invalid;
+
+		UE_LOG_DATASTREAM_CONN(Verbose, TEXT("Registered DataStream with name '%s' with streamindex: %d State:%s"), ToCStr(Definition.DataStreamName.ToString()), WantedStreamIndex, LexToString(StreamState[WantedStreamIndex]));
+	}
+
+	// Init stream
+	InitStream(Stream);
+
+	return ECreateDataStreamResult::Success;
+}
+
+ECreateDataStreamResult UDataStreamManager::FImpl::CreateStream(const FName StreamName, ECreateDataStreamFlags Flags)
+{
 	if (Streams.ContainsByPredicate(FindStreamByName(StreamName)))
 	{
-		UE_LOG(LogIris, Warning, TEXT("A DataStream with name '%s' already exists."), ToCStr(StreamName.GetPlainNameString()));
+		UE_LOG_DATASTREAM_CONN(Warning, TEXT("A DataStream with name '%s' already exists."), ToCStr(StreamName.GetPlainNameString()));
 		return ECreateDataStreamResult::Error_Duplicate;
 	}
-
+	
 	const UDataStreamDefinitions* StreamDefinitions = GetDefault<UDataStreamDefinitions>();
 	if (const FDataStreamDefinition* Definition = StreamDefinitions->FindDefinition(StreamName))
 	{
-		if (Definition->Class == nullptr)
-		{
-			return ECreateDataStreamResult::Error_InvalidDefinition;
-		}
+		return CreateStreamFromDefinition(*Definition, Flags);
+	}
+
+	return ECreateDataStreamResult::Error_MissingDefinition;
+}
 
-		UDataStream* Stream = NewObject<UDataStream>(GetTransientPackage(), ToRawPtr(Definition->Class), MakeUniqueObjectName(nullptr, Definition->Class, StreamName));
+ECreateDataStreamResult UDataStreamManager::FImpl::CreateStreamFromIndex(int32 StreamIndex)
+{	
+	const UDataStreamDefinitions* StreamDefinitions = GetDefault<UDataStreamDefinitions>();
+	if (!StreamDefinitions->bFixupComplete)
+	{
+		UE_LOG_DATASTREAM_CONN(Warning, TEXT("Cannot create datastream by index if DataStreamDefinitions are not FixedUp."));
+		return ECreateDataStreamResult::Error_MissingDefinition;
+	}
 
-		Streams.Add(Stream);
-		StreamSendStatus.Add(Definition->DefaultSendStatus);
-		return ECreateDataStreamResult::Success;
+	if (const FDataStreamDefinition* Definition = StreamDefinitions->FindDefinition(StreamIndex))
+	{
+		return CreateStreamFromDefinition(*Definition, ECreateDataStreamFlags::None);
 	}
 
 	return ECreateDataStreamResult::Error_MissingDefinition;
@@ -479,7 +1084,7 @@ void UDataStreamManager::FImpl::SetSendStatus(const FName StreamName, EDataStrea
 	const int32 Index = Streams.IndexOfByPredicate(FindStreamByName(StreamName));
 	if (Index == INDEX_NONE)
 	{
-		UE_LOG(LogIris, Display, TEXT("Cannot set send status for DataStream '%s' that hasn't been created."), ToCStr(StreamName.GetPlainNameString()));
+		UE_LOG_DATASTREAM_CONN(Display, TEXT("Cannot set send status for DataStream '%s' that hasn't been created."), ToCStr(StreamName.GetPlainNameString()));
 		return;
 	}
 
@@ -491,15 +1096,26 @@ EDataStreamSendStatus UDataStreamManager::FImpl::GetSendStatus(const FName Strea
 	const int32 Index = Streams.IndexOfByPredicate(FindStreamByName(StreamName));
 	if (Index == INDEX_NONE)
 	{
-		UE_LOG(LogIris, Display, TEXT("Cannot retrieve send status for DataStream '%s' that hasn't been created. Returnning Pause."), ToCStr(StreamName.GetPlainNameString()));
+		UE_LOG_DATASTREAM_CONN(Display, TEXT("Cannot retrieve send status for DataStream '%s' that hasn't been created. Returning Pause."), ToCStr(StreamName.GetPlainNameString()));
 		return EDataStreamSendStatus::Pause;
 	}
 
 	return StreamSendStatus[Index];
 }
 
+UDataStream::EDataStreamState UDataStreamManager::FImpl::GetStreamState(const FName StreamName) const
+{
+	const int32 Index = Streams.IndexOfByPredicate(FindStreamByName(StreamName));
+	if (Index == INDEX_NONE)
+	{
+		return UDataStream::EDataStreamState::Invalid;
+	}
+	return GetStreamState(Index);
+}
+
 void UDataStreamManager::FImpl::InitRecordStorage()
 {
+	const uint32 PacketWindowSize = InitParameters.PacketWindowSize;
 	RecordStorage.SetNum(PacketWindowSize);
 
 	Records = TResizableCircularQueue<FRecord*>(PacketWindowSize);
@@ -513,17 +1129,26 @@ void UDataStreamManager::FImpl::InitRecordStorage()
 	Records.Reset();
 }
 
+void UDataStreamManager::FImpl::InitStream(UDataStream* Stream)
+{
+	if (IsValid(Stream))
+	{
+		Stream->Init(InitParameters);
+	}
+}
+
 void UDataStreamManager::FImpl::InitStreams()
 {
 	UDataStreamDefinitions* StreamDefinitions = GetMutableDefault<UDataStreamDefinitions>();
 	StreamDefinitions->FixupDefinitions();
 
-	TArray<FName> AutoCreateStreams;
-	AutoCreateStreams.Reserve(MaxStreamCount);
-	StreamDefinitions->GetAutoCreateStreamNames(AutoCreateStreams);
-	for (const FName& StreamName : AutoCreateStreams)
+	TArray<FName> StreamsToAutoCreateOrRegister;
+	StreamsToAutoCreateOrRegister.Reserve(MaxStreamCount);
+	StreamDefinitions->GetStreamNamesToAutoCreateOrRegister(StreamsToAutoCreateOrRegister);
+
+	for (const FName& StreamName : StreamsToAutoCreateOrRegister)
 	{
-		CreateStream(StreamName);
+		CreateStream(StreamName, ECreateDataStreamFlags::RegisterIfStreamIsDynamic);
 	}
 }
 
diff --git a/Engine/Source/Runtime/Experimental/Iris/Core/Private/Iris/ReplicationSystem/NetTokenDataStream.cpp b/Engine/Source/Runtime/Experimental/Iris/Core/Private/Iris/ReplicationSystem/NetTokenDataStream.cpp
--- a/Engine/Source/Runtime/Experimental/Iris/Core/Private/Iris/ReplicationSystem/NetTokenDataStream.cpp
+++ b/Engine/Source/Runtime/Experimental/Iris/Core/Private/Iris/ReplicationSystem/NetTokenDataStream.cpp
@@ -53,17 +53,17 @@ UNetTokenDataStream::~UNetTokenDataStream()
 {
 }
 
-void UNetTokenDataStream::Init(const FInitParameters& Params)
+void UNetTokenDataStream::Init(const UDataStream::FInitParameters& Params)
 {
 	using namespace UE::Net;
 
 	ReplicationSystemId = Params.ReplicationSystemId;
 	ConnectionId = Params.ConnectionId;
-	NetExports = Params.NetExports;
 
 	UReplicationSystem* ReplicationSystem = UE::Net::GetReplicationSystem(ReplicationSystemId);
 	NetTokenStore = ReplicationSystem->GetNetTokenStore();
 	RemoteNetTokenStoreState = NetTokenStore->GetRemoteNetTokenStoreState(Params.ConnectionId);
+	NetExports = Params.NetExports;
 
 	// $IRIS $TODO: if we want to make this into a real feature we need to expose some sort of api to mark tokens for pre-export
 	if (Private::bIrisPreExportExistingNetTokensOnConnect)
diff --git a/Engine/Source/Runtime/Experimental/Iris/Core/Private/Iris/ReplicationSystem/NetTokenDataStream.h b/Engine/Source/Runtime/Experimental/Iris/Core/Private/Iris/ReplicationSystem/NetTokenDataStream.h
--- a/Engine/Source/Runtime/Experimental/Iris/Core/Private/Iris/ReplicationSystem/NetTokenDataStream.h
+++ b/Engine/Source/Runtime/Experimental/Iris/Core/Private/Iris/ReplicationSystem/NetTokenDataStream.h
@@ -28,14 +28,6 @@ class UNetTokenDataStream final : public UDataStream
 
 public:
 
-	struct FInitParameters
-	{
-		uint32 ReplicationSystemId;
-		uint32 ConnectionId;
-		UE::Net::Private::FNetExports* NetExports;
-	};
-
-	void Init(const FInitParameters& Params);
 	const UE::Net::FNetTokenStoreState* GetRemoteNetTokenStoreState() const { return RemoteNetTokenStoreState; }
 	void AddNetTokenForExplicitExport(UE::Net::FNetToken NetToken);
 
@@ -44,6 +36,7 @@ private:
 	virtual ~UNetTokenDataStream();
 
 	// UDataStream interface
+	virtual void Init(const UDataStream::FInitParameters& Params) override;
 	virtual EWriteResult WriteData(UE::Net::FNetSerializationContext& Context, FDataStreamRecord const*& OutRecord) override;
 	virtual void ReadData(UE::Net::FNetSerializationContext& Context) override;
 	virtual void ProcessPacketDeliveryStatus(UE::Net::EPacketDeliveryStatus Status, FDataStreamRecord const* Record) override;
diff --git a/Engine/Source/Runtime/Experimental/Iris/Core/Private/Iris/ReplicationSystem/ReplicationConnections.cpp b/Engine/Source/Runtime/Experimental/Iris/Core/Private/Iris/ReplicationSystem/ReplicationConnections.cpp
--- a/Engine/Source/Runtime/Experimental/Iris/Core/Private/Iris/ReplicationSystem/ReplicationConnections.cpp
+++ b/Engine/Source/Runtime/Experimental/Iris/Core/Private/Iris/ReplicationSystem/ReplicationConnections.cpp
@@ -15,7 +15,7 @@ void FReplicationConnections::Deinit()
 	ValidConnections.ForAllSetBits([this](uint32 ConnectionId) { RemoveConnection(ConnectionId); });
 }
 
-void FReplicationConnections::InitDataStreams(uint32 ReplicationSystemId, uint32 ConnectionId, UDataStreamManager* DataStreamManager)
+void FReplicationConnections::InitDataStreamManager(uint32 ReplicationSystemId, uint32 ConnectionId, UDataStreamManager* DataStreamManager)
 {
 	FReplicationConnection* Connection = GetConnection(ConnectionId);
 	if (Connection == nullptr)
@@ -23,28 +23,16 @@ void FReplicationConnections::InitDataStreams(uint32 ReplicationSystemId, uint32
 		return;
 	}
 
-	// Init NetTokenDataStream if we should have one
-	UNetTokenDataStream* NetTokenDataStream = StaticCast<UNetTokenDataStream*>(DataStreamManager->GetStream(FName("NetToken")));
-	if (NetTokenDataStream != nullptr)
-	{
-		UNetTokenDataStream::FInitParameters Params;
-		Params.ConnectionId = ConnectionId;
-		Params.ReplicationSystemId = ReplicationSystemId;
-		Params.NetExports = &DataStreamManager->GetNetExports();
-
-		NetTokenDataStream->Init(Params);
-	}
+	// Init data stream manager and create all DataStreams
+	UDataStream::FInitParameters InitParams;
+	InitParams.ReplicationSystemId = ReplicationSystemId;
+	InitParams.ConnectionId = ConnectionId;
+	InitParams.PacketWindowSize = 256;
 
-	// Init ReplicationDataStream
-	Connection->ReplicationWriter->SetNetExports(DataStreamManager->GetNetExports());
+	DataStreamManager->Init(InitParams);
 
-	UReplicationDataStream* ReplicationDataStream = StaticCast<UReplicationDataStream*>(DataStreamManager->GetStream(FName("Replication")));
-	if (ReplicationDataStream != nullptr) 
-	{
-		checkf(NetTokenDataStream != nullptr, TEXT("%s"), TEXT("ReplicationDataStream requires a NetTokenDataStream"));
-		ReplicationDataStream->SetReaderAndWriter(Connection->ReplicationReader, Connection->ReplicationWriter);
-		Connection->ReplicationDataStream = ReplicationDataStream;
-	}
+	// Store it.
+	Connection->DataStreamManager = DataStreamManager;
 }
 
 void FReplicationConnections::SetReplicationView(uint32 ConnectionId, const FReplicationView& View)
@@ -56,7 +44,8 @@ void FReplicationConnections::RemoveConnection(uint32 ConnectionId)
 {
 	check(ValidConnections.GetBit(ConnectionId));
 	SetReplicationView(ConnectionId, FReplicationView());
-	DestroyReplicationReaderAndWriter(ConnectionId);
+
+	DeinitDataStreamManager(ConnectionId);
 
 	Connections[ConnectionId] = FReplicationConnection();
 	ValidConnections.ClearBit(ConnectionId);
@@ -77,22 +66,25 @@ FNetBitArray FReplicationConnections::GetOpenConnections() const
 	return OpenConnections;
 }
 
-void FReplicationConnections::DestroyReplicationReaderAndWriter(uint32 ConnectionId)
+void FReplicationConnections::DeinitDataStreamManager(uint32 ConnectionId)
 {
-	FReplicationConnection* Connection = GetConnection(ConnectionId);
-	UReplicationDataStream* ReplicationDataStream = CastChecked<UReplicationDataStream>(Connection->ReplicationDataStream.Get(), ECastCheckedType::NullAllowed);
-	if (ReplicationDataStream)
+	if (FReplicationConnection* Connection = GetConnection(ConnectionId))
 	{
-		ReplicationDataStream->SetReaderAndWriter(nullptr, nullptr);
-	}
+		if (UDataStreamManager* DataStreamManager = Connection->DataStreamManager.Get())
+		{
+			DataStreamManager->Deinit();
+		}
 
-	Connection->ReplicationReader->Deinit();
-	Connection->ReplicationWriter->Deinit();
+		// This is a bit special as these are owned by ReplicationSystem rather than the ReplicationDataStream
+		Connection->ReplicationReader->Deinit();
+		Connection->ReplicationWriter->Deinit();
 
-	delete Connection->ReplicationReader;
-	Connection->ReplicationReader = nullptr;
-	delete Connection->ReplicationWriter;
-	Connection->ReplicationWriter = nullptr;
+		// Delete replication reader / writer
+		delete Connection->ReplicationReader;
+		Connection->ReplicationReader = nullptr;
+		delete Connection->ReplicationWriter;
+		Connection->ReplicationWriter = nullptr;
+	}
 }
 
 }
diff --git a/Engine/Source/Runtime/Experimental/Iris/Core/Private/Iris/ReplicationSystem/ReplicationConnections.h b/Engine/Source/Runtime/Experimental/Iris/Core/Private/Iris/ReplicationSystem/ReplicationConnections.h
--- a/Engine/Source/Runtime/Experimental/Iris/Core/Private/Iris/ReplicationSystem/ReplicationConnections.h
+++ b/Engine/Source/Runtime/Experimental/Iris/Core/Private/Iris/ReplicationSystem/ReplicationConnections.h
@@ -21,9 +21,9 @@ namespace UE::Net::Private
 
 struct FReplicationConnection
 {
-	FWeakObjectPtr ReplicationDataStream;
 	FReplicationWriter* ReplicationWriter = nullptr;
 	FReplicationReader* ReplicationReader = nullptr;
+	TWeakObjectPtr<UDataStreamManager> DataStreamManager;
 	FObjectPtr UserData = nullptr;
 	bool bIsClosing = false; // Should be set when a connection starts the graceful close process to finish flushing reliable data
 };
@@ -85,7 +85,8 @@ public:
 	// Returns connections that are not in the closing state
 	FNetBitArray GetOpenConnections() const;
 
-	void InitDataStreams(uint32 ReplicationSystemId, uint32 ConnectionId, UDataStreamManager* DataStreamManager);
+	void InitDataStreamManager(uint32 ReplicationSystemId, uint32 ConnectionId, UDataStreamManager* DataStreamManager);
+	void DeinitDataStreamManager(uint32 ConnectionId);
 
 	void SetReplicationView(uint32 ConnectionId, const FReplicationView& ViewInfo);
 	const FReplicationView& GetReplicationView(uint32 ConnectionId) const { return ReplicationViews[ConnectionId]; }
@@ -96,10 +97,6 @@ public:
 		check(ValidConnections.GetBit(ConnectionId) == true);
 		Connections[ConnectionId].bIsClosing = true;
 	}
-
-private:
-	void DestroyReplicationReaderAndWriter(uint32 ConnectionId);
-
 private:
 	TArray<FReplicationConnection> Connections;
 	TArray<FReplicationView> ReplicationViews;
diff --git a/Engine/Source/Runtime/Experimental/Iris/Core/Private/Iris/ReplicationSystem/ReplicationDataStream.cpp b/Engine/Source/Runtime/Experimental/Iris/Core/Private/Iris/ReplicationSystem/ReplicationDataStream.cpp
--- a/Engine/Source/Runtime/Experimental/Iris/Core/Private/Iris/ReplicationSystem/ReplicationDataStream.cpp
+++ b/Engine/Source/Runtime/Experimental/Iris/Core/Private/Iris/ReplicationSystem/ReplicationDataStream.cpp
@@ -3,6 +3,9 @@
 #include "ReplicationDataStream.h"
 #include "Iris/ReplicationSystem/ReplicationReader.h"
 #include "Iris/ReplicationSystem/ReplicationWriter.h"
+#include "Iris/ReplicationSystem/ReplicationSystem.h"
+#include "Iris/ReplicationSystem/ReplicationSystemInternal.h"
+#include "Iris/ReplicationSystem/ReplicationConnections.h"
 
 UReplicationDataStream::UReplicationDataStream()
 : ReplicationReader(nullptr)
@@ -14,6 +17,38 @@ UReplicationDataStream::~UReplicationDataStream()
 {
 }
 
+void UReplicationDataStream::Init(const UDataStream::FInitParameters& Params)
+{
+	using namespace UE::Net::Private;
+
+	UReplicationSystem* ReplicationSystem = UE::Net::GetReplicationSystem(Params.ReplicationSystemId);
+	if (ensure(ReplicationSystem))
+	{
+		// Init ReplicationReader and writer.
+		if (FReplicationSystemInternal* ReplicationSystemInternal = ReplicationSystem->GetReplicationSystemInternal())
+		{
+			FReplicationConnections& Connections = ReplicationSystemInternal->GetConnections();
+			FReplicationConnection* Connection = Connections.GetConnection(Params.ConnectionId);
+			if (ensure(Connection))
+			{
+				ReplicationWriter = Connection->ReplicationWriter;
+				ReplicationReader = Connection->ReplicationReader;
+
+				if (ensure(Params.NetExports))
+				{
+					ReplicationWriter->SetNetExports(*Params.NetExports);
+				}
+			}
+		}
+	}
+}
+
+void UReplicationDataStream::Deinit()
+{
+	ReplicationWriter = nullptr;
+	ReplicationReader = nullptr;
+}
+
 UDataStream::EWriteResult UReplicationDataStream::BeginWrite(const FBeginWriteParameters& Params)
 {
 	if (ReplicationWriter != nullptr)
@@ -67,9 +102,3 @@ bool UReplicationDataStream::HasAcknowledgedAllReliableData() const
 
 	return true;
 }
-
-void UReplicationDataStream::SetReaderAndWriter(UE::Net::Private::FReplicationReader* Reader, UE::Net::Private::FReplicationWriter* Writer)
-{
-	ReplicationReader = Reader;
-	ReplicationWriter = Writer;
-}
diff --git a/Engine/Source/Runtime/Experimental/Iris/Core/Private/Iris/ReplicationSystem/ReplicationDataStream.h b/Engine/Source/Runtime/Experimental/Iris/Core/Private/Iris/ReplicationSystem/ReplicationDataStream.h
--- a/Engine/Source/Runtime/Experimental/Iris/Core/Private/Iris/ReplicationSystem/ReplicationDataStream.h
+++ b/Engine/Source/Runtime/Experimental/Iris/Core/Private/Iris/ReplicationSystem/ReplicationDataStream.h
@@ -16,14 +16,14 @@ class UReplicationDataStream final : public UDataStream
 {
 	GENERATED_BODY()
 
-public:
-	void SetReaderAndWriter(UE::Net::Private::FReplicationReader* Reader, UE::Net::Private::FReplicationWriter* Writer);
-
 private:
 	UReplicationDataStream();
 	virtual ~UReplicationDataStream();
 
 	// UDataStream interface
+	virtual void Init(const UDataStream::FInitParameters& Params) override;
+	virtual void Deinit() override;
+
 	virtual EWriteResult BeginWrite(const FBeginWriteParameters& Params) override;
 	virtual EWriteResult WriteData(UE::Net::FNetSerializationContext& Context, FDataStreamRecord const*& OutRecord) override;
 	virtual void EndWrite() override;
diff --git a/Engine/Source/Runtime/Experimental/Iris/Core/Private/Iris/ReplicationSystem/ReplicationSystem.cpp b/Engine/Source/Runtime/Experimental/Iris/Core/Private/Iris/ReplicationSystem/ReplicationSystem.cpp
--- a/Engine/Source/Runtime/Experimental/Iris/Core/Private/Iris/ReplicationSystem/ReplicationSystem.cpp
+++ b/Engine/Source/Runtime/Experimental/Iris/Core/Private/Iris/ReplicationSystem/ReplicationSystem.cpp
@@ -25,6 +25,7 @@
 #include "Iris/ReplicationSystem/ReplicationSystemTypes.h"
 #include "Iris/ReplicationSystem/ReplicationTypes.h"
 #include "Iris/ReplicationSystem/ReplicationWriter.h"
+#include "Iris/DataStream/DataStreamManager.h"
 #include "Iris/Serialization/InternalNetSerializationContext.h"
 #include "Iris/Serialization/NetBitStreamWriter.h"
 #include "Iris/Serialization/NetSerializer.h"
@@ -1085,6 +1086,87 @@ bool UReplicationSystem::RegisterNetBlobHandler(UNetBlobHandler* Handler)
 	return NetBlobManager.RegisterNetBlobHandler(Handler);
 }
 
+UDataStream* UReplicationSystem::OpenDataStream(uint32 ConnectionId, FName Name)
+{
+	if (!IsServer())
+	{
+		UE_LOG(LogIris, Error, TEXT("DataStream %s can only be opened from server"), *Name.ToString());
+		return nullptr;
+	}
+
+	UE::Net::Private::FReplicationConnections& Connections = Impl->ReplicationSystemInternal.GetConnections();
+	if (!ensureMsgf(Connections.IsValidConnection(ConnectionId), TEXT("Invalid ConnectionId %u passed to UReplicationSystem::GetDataStream."), ConnectionId))
+	{
+		return nullptr;
+	}
+
+	UE::Net::Private::FReplicationConnection* Connection = Connections.GetConnection(ConnectionId);
+	UDataStreamManager* DataStreamManager = Connection ? Connection->DataStreamManager.Get() : nullptr;
+	if (ensureMsgf(DataStreamManager, TEXT("UReplicationSystem::OpenDataStream Trying to open datastream for not yet initialized connection %u"), ConnectionId))
+	{
+		ECreateDataStreamResult Result = DataStreamManager->CreateStream(Name);
+		if (Result == ECreateDataStreamResult::Success)
+		{
+			return DataStreamManager->GetStream(Name);
+		}
+	}
+
+	return nullptr;
+}
+
+UDataStream* UReplicationSystem::GetDataStream(uint32 ConnectionId, FName Name)
+{
+	UE::Net::Private::FReplicationConnections& Connections = Impl->ReplicationSystemInternal.GetConnections();
+	if (!ensureMsgf(Connections.IsValidConnection(ConnectionId), TEXT("Invalid ConnectionId %u passed to UReplicationSystem::GetDataStream."), ConnectionId))
+	{
+		return nullptr;
+	}
+
+	UE::Net::Private::FReplicationConnection* Connection = Connections.GetConnection(ConnectionId);
+	UDataStreamManager* DataStreamManager = Connection ? Connection->DataStreamManager.Get() : nullptr;
+	if (ensureMsgf(DataStreamManager, TEXT("UReplicationSystem::GetDataStream Trying to get datastream for not yet initialized connection %u"), ConnectionId))
+	{
+		return DataStreamManager->GetStream(Name);
+	}
+
+	return nullptr;
+}
+
+const UDataStream* UReplicationSystem::GetDataStream(uint32 ConnectionId, FName Name) const
+{
+	const UE::Net::Private::FReplicationConnections& Connections = Impl->ReplicationSystemInternal.GetConnections();
+	if (!ensureMsgf(Connections.IsValidConnection(ConnectionId), TEXT("Invalid ConnectionId %u passed to UReplicationSystem::GetDataStream."), ConnectionId))
+	{
+		return nullptr;
+	}
+
+	const UE::Net::Private::FReplicationConnection* Connection = Connections.GetConnection(ConnectionId);
+	const UDataStreamManager* DataStreamManager = Connection ? Connection->DataStreamManager.Get() : nullptr;
+	if (ensureMsgf(DataStreamManager, TEXT("UReplicationSystem::GetDataStream Trying to get datastream for not yet initialized connection %u"), ConnectionId))
+	{
+		return DataStreamManager->GetStream(Name);
+	}
+
+	return nullptr;
+}
+
+void UReplicationSystem::CloseDataStream(uint32 ConnectionId, FName Name)
+{
+	UE::Net::Private::FReplicationConnections& Connections = Impl->ReplicationSystemInternal.GetConnections();
+	if (!ensureMsgf(Connections.IsValidConnection(ConnectionId), TEXT("Invalid ConnectionId %u passed to UReplicationSystem::CloseDataStream."), ConnectionId))
+	{
+		return;
+	}
+
+	UE::Net::Private::FReplicationConnection* Connection = Connections.GetConnection(ConnectionId);
+	UDataStreamManager* DataStreamManager = Connection ? Connection->DataStreamManager.Get() : nullptr;
+	if (ensureMsgf(DataStreamManager, TEXT("UReplicationSystem::CloseDataStream Trying to get datastream for not yet initialized connection %u"), ConnectionId))
+	{
+		DataStreamManager->CloseStream(Name);
+		return;
+	}
+}
+
 bool UReplicationSystem::QueueNetObjectAttachment(uint32 ConnectionId, const UE::Net::FNetObjectReference& TargetRef, const TRefCountPtr<UE::Net::FNetObjectAttachment>& Attachment)
 {
 	UE::Net::Private::FNetBlobManager& NetBlobManager = Impl->ReplicationSystemInternal.GetNetBlobManager();
@@ -1137,10 +1219,10 @@ void UReplicationSystem::ResetRPCSendPolicyFlags()
 	Impl->AttachmentSendPolicyFlags.Reset();
 }
 
-void UReplicationSystem::InitDataStreams(uint32 ConnectionId, UDataStreamManager* DataStreamManager)
+void UReplicationSystem::InitDataStreamManager(uint32 ConnectionId, UDataStreamManager* DataStreamManager)
 {
 	UE::Net::Private::FReplicationConnections& Connections = Impl->ReplicationSystemInternal.GetConnections();
-	Connections.InitDataStreams(GetId(), ConnectionId, DataStreamManager);
+	Connections.InitDataStreamManager(GetId(), ConnectionId, DataStreamManager);
 }
 
 void UReplicationSystem::SetConnectionUserData(uint32 ConnectionId, UObject* InUserData)
diff --git a/Engine/Source/Runtime/Experimental/Iris/Core/Private/Iris/Serialization/NetErrorContext.cpp b/Engine/Source/Runtime/Experimental/Iris/Core/Private/Iris/Serialization/NetErrorContext.cpp
--- a/Engine/Source/Runtime/Experimental/Iris/Core/Private/Iris/Serialization/NetErrorContext.cpp
+++ b/Engine/Source/Runtime/Experimental/Iris/Core/Private/Iris/Serialization/NetErrorContext.cpp
@@ -12,6 +12,7 @@ const FName GNetError_InvalidNetHandle("Invalid NetHandle");
 const FName GNetError_BrokenNetHandle("Broken NetHandle");
 const FName GNetError_InvalidValue("Invalid value");
 const FName GNetError_InternalError("Internal error");
+const FName GNetError_InvalidDataStream("Invalid DataStream");
 
 void FNetErrorContext::SetError(const FName InError)
 {
diff --git a/Engine/Source/Runtime/Experimental/Iris/Core/Public/Iris/DataStream/DataStream.h b/Engine/Source/Runtime/Experimental/Iris/Core/Public/Iris/DataStream/DataStream.h
--- a/Engine/Source/Runtime/Experimental/Iris/Core/Public/Iris/DataStream/DataStream.h
+++ b/Engine/Source/Runtime/Experimental/Iris/Core/Public/Iris/DataStream/DataStream.h
@@ -10,6 +10,12 @@ namespace UE::Net
 	class FNetSerializationContext;
 	enum class EPacketDeliveryStatus : uint8;
 
+	namespace Private
+	{
+		class FNetExports;
+	}
+
+
 enum class EDataStreamWriteMode : unsigned
 {
 	// Allowed to write all data, this is the default WriteMode
@@ -39,8 +45,8 @@ struct FDataStreamRecord
 UENUM()
 enum class EDataStreamSendStatus : uint8
 {
+	Pause = 0,
 	Send,
-	Pause,
 };
 
 /**
@@ -63,6 +69,32 @@ public:
 		HasMoreData,
 	};
 
+	enum class EDataStreamState : uint8
+	{
+		// Stream is invalid
+		Invalid = 0,
+		// We should send open/init to other side
+		PendingCreate, 
+		// We are waiting for confirmation that remote have accepted the stream
+		WaitOnCreateConfirmation,
+		// Stream is open and we will process incoming data.
+		Open,
+		// We are closing, but still considerd open until flushed
+		PendingClose,
+		// We have send a close request and is waiting for confirmation before invalidating stream
+		WaitOnCloseConfirmation,
+
+		Count
+	};
+
+	struct FInitParameters
+	{
+		uint32 ReplicationSystemId = 0U;
+		uint32 ConnectionId = 0U;
+		uint32 PacketWindowSize = 0U;
+		UE::Net::Private::FNetExports* NetExports = nullptr;
+	};
+
 	struct FBeginWriteParameters
 	{
 		UE::Net::EDataStreamWriteMode WriteMode = UE::Net::EDataStreamWriteMode::Full;
@@ -72,6 +104,12 @@ public:
 public:
 	IRISCORE_API virtual ~UDataStream();
 
+	/** Called before any other calls are made. */
+	IRISCORE_API virtual void Init(const FInitParameters& Params);
+
+	/** Called when a created stream is destroyed. */
+	IRISCORE_API virtual void Deinit();
+
 	/**
 	 * Called before any calls to potential WriteData, if it returns EWriteData::NoData no other calls will be made.
 	 * The purpose of the method is to enable a DataStream to setup data that can persist over multiple calls to WriteData if bandwidth allows.
@@ -116,3 +154,5 @@ public:
 	 */
 	IRISCORE_API virtual bool HasAcknowledgedAllReliableData() const PURE_VIRTUAL(HasAcknowledgedAllReliableData, return true;);
 };
+
+const TCHAR* LexToString(const UDataStream::EDataStreamState State);
diff --git a/Engine/Source/Runtime/Experimental/Iris/Core/Public/Iris/DataStream/DataStreamManager.h b/Engine/Source/Runtime/Experimental/Iris/Core/Public/Iris/DataStream/DataStreamManager.h
--- a/Engine/Source/Runtime/Experimental/Iris/Core/Public/Iris/DataStream/DataStreamManager.h
+++ b/Engine/Source/Runtime/Experimental/Iris/Core/Public/Iris/DataStream/DataStreamManager.h
@@ -29,11 +29,6 @@ namespace UE::Net::Private
 	class FNetExports;
 }
 
-struct FDataStreamManagerInitParams
-{
-	uint32 PacketWindowSize = 0;
-};
-
 /**
  * The DataStreamManager contains all active DataStreams that may serialize data.
  * Calls to the DataStream interface functions will be forwarded to active streams.
@@ -49,7 +44,7 @@ public:
 	IRISCORE_API virtual ~UDataStreamManager();
 
 	/** Initializes the manager. No data stream can be created by the manager before this. */
-	IRISCORE_API void Init(const FDataStreamManagerInitParams& InitParams);
+	virtual IRISCORE_API void Init(const FInitParameters& InitParams) override;
 
 	/** Prepare for destruction. Data streams will get ProcessPacketDeliveryStatus for outstanding packets and then marked as garbage. */
 	IRISCORE_API void Deinit();
@@ -95,12 +90,22 @@ public:
 	IRISCORE_API const UDataStream* GetStream(const FName StreamName) const;
 
 	/**
-	 * Gets the data stream with a given name.
-	 * @param StreamName The name of the data stream as configured in UDataStreamDefinitions.
-	 * @return A pointer to the data stream if it exists, nullptr if not.
-	 */
+	* Gets the data stream with a given name.
+	* @param StreamName The name of the data stream as configured in UDataStreamDefinitions.
+	* @return A pointer to the data stream if it exists, nullptr if not.
+	*/
 	IRISCORE_API UDataStream* GetStream(const FName StreamName);
 
+	/**
+	* Request that the dynamic data stream with the given name will be closed after handshake with remote peer. 
+	* NOTE: This will not do anything for streams not marked as bDynamicCreate
+	* @param StreamName The name of the data stream as configured in UDataStreamDefinitions.
+	*/
+	IRISCORE_API void CloseStream(const FName StreamName);
+
+	/** Get the current state of a DataStream Returns Invalid if the stream isn't created. */
+	IRISCORE_API EDataStreamState GetStreamState(const FName StreamName) const;
+
 	/** Set the send status of an already created data stream. */
 	IRISCORE_API void SetSendStatus(const FName StreamName, EDataStreamSendStatus Status);
 
diff --git a/Engine/Source/Runtime/Experimental/Iris/Core/Public/Iris/ReplicationSystem/ReplicationSystem.h b/Engine/Source/Runtime/Experimental/Iris/Core/Public/Iris/ReplicationSystem/ReplicationSystem.h
--- a/Engine/Source/Runtime/Experimental/Iris/Core/Public/Iris/ReplicationSystem/ReplicationSystem.h
+++ b/Engine/Source/Runtime/Experimental/Iris/Core/Public/Iris/ReplicationSystem/ReplicationSystem.h
@@ -16,6 +16,7 @@
 #include "ReplicationSystem.generated.h"
 
 class UDataStreamManager;
+class UDataStream;
 class UFunction;
 class UNetBlobHandler;
 class UNetDriver;
@@ -274,6 +275,32 @@ public:
 	 * @see UNetObjectPrioritizerDefinitions
 	 */
 	IRISCORE_API UNetObjectPrioritizer* GetPrioritizer(const FName PrioritizerName) const;
+
+	/**
+	 * Get DataStream with the provided Name
+	 * @param Name The name of the DataStream. Names of valid DataStream are configured in UDataStreamDefinitions.
+	 * @return A pointer to the DataStream if it exists
+	 * @see UDataStreamDefinitions
+	 */
+	UDataStream* GetDataStream(uint32 ConnectionId, FName Name);
+
+	/** Const version of GetDataStream */
+	const UDataStream* GetDataStream(uint32 ConnectionId, FName Name) const;
+
+	/**
+	 * Request DataStream to be created/opened for the given ConnectionId. NOTE: This only applies to DataStreams that are flagged with bDynamicCreate in the DataStreamDefinition
+	 * @param Name The name of the DataStream. Names of valid DataStream are configured in UDataStreamDefinitions.
+	 * @return A pointer to the DataStream if successfully opened or already existed.
+	 * @see UDataStreamDefinitions
+	 */
+	UDataStream* OpenDataStream(uint32 ConnectionId, FName Name);
+
+	/**
+	 * Initiate close of DataStream for the given ConnectionId. NOTE: This only applies to DataStreams that are flagged with bDynamicCreate in the DataStreamDefinition
+	 * @param Name The name of the DataStream. Names of valid DataStream are configured in UDataStreamDefinitions.
+	 * @see UDataStreamDefinitions
+	 */
+	void CloseDataStream(uint32 ConnectionId, FName Name);
 	
 	// NetBlob 
 
@@ -282,7 +309,7 @@ public:
 	 * for the lifetime of the ReplicationSystem. The handler class must have been configured in UNetBlobHandlerDefinitions in order
 	 * for the call to succeed. Only one instance per class is allowed to be registered.
 	 * @param Handler The UNetBlobHandler to register.
-	 * @return Whether the handler was succesfully registered or not.
+	 * @return Whether the handler was successfully registered or not.
 	 */
 	IRISCORE_API bool RegisterNetBlobHandler(UNetBlobHandler* Handler);
 
@@ -653,7 +680,7 @@ public:
 	IRISCORE_API const UE::Net::FWorldLocations& GetWorldLocations() const;
 
 	/* Init data streams for a connection. For internal use by UDataStreamChannel. */
-	IRISCORE_API void InitDataStreams(uint32 ConnectionId, UDataStreamManager* DataStreamManager);
+	IRISCORE_API void InitDataStreamManager(uint32 ConnectionId, UDataStreamManager* DataStreamManager);
 
 	/**
 	 * Associate data with a connection. Only a single piece of user data is supported per connection.
