Fixed issue when installing single packages and optional bulk data.

#rnx

diff --git a/Engine/Source/Runtime/Experimental/IoStore/OnDemand/Private/OnDemandIoStore.cpp b/Engine/Source/Runtime/Experimental/IoStore/OnDemand/Private/OnDemandIoStore.cpp
--- a/Engine/Source/Runtime/Experimental/IoStore/OnDemand/Private/OnDemandIoStore.cpp
+++ b/Engine/Source/Runtime/Experimental/IoStore/OnDemand/Private/OnDemandIoStore.cpp
@@ -232,7 +232,7 @@ struct FContainerInstallData
 {
 	FPackageStoreEntryMap			PackageStoreEntries;
 	FSoftPackageReferenceMap		SoftPackageReferences;
-	TSet<FPackageId>				PackageIds;
+	TSet<FPackageId>				ResolvedPackageIds;
 	TMap<uint32, EChunkCacheState>	ResolvedChunks;
 	uint64							TotalSize = 0;
 };
@@ -250,14 +250,14 @@ static FIoStatus BuildInstallData(
 	// Setup the package information for each container
 	for (const FSharedOnDemandContainer& Container : Containers)
 	{
+		FContainerInstallData& Data = OutInstallData.FindOrAdd(Container);
+
 		if (Container->Header.IsValid() == false || Container->Header->PackageIds.IsEmpty())
 		{
 			// The container contains no package data
 			continue;
 		}
 
-		FContainerInstallData& Data = OutInstallData.FindOrAdd(Container);
-
 		const FIoContainerHeader& Header = *Container->Header;
 		TConstArrayView<FFilePackageStoreEntry> Entries(
 			reinterpret_cast<const FFilePackageStoreEntry*>(Header.StoreEntries.GetData()),
@@ -320,7 +320,7 @@ static FIoStatus BuildInstallData(
 			FContainerInstallData& Data = Kv.Value;
 			if (const FFilePackageStoreEntry* Entry = Data.PackageStoreEntries.FindRef(PackageId))
 			{
-				Data.PackageIds.Add(PackageId);
+				Data.ResolvedPackageIds.Add(PackageId);
 
 				// Add hard references 
 				for (const FPackageId& ImportedPackageId : Entry->ImportedPackages)
@@ -356,12 +356,18 @@ static FIoStatus BuildInstallData(
 		}
 	}
 
+	// Resolve all chunk ID's from the package ID's
 	for (TPair<FSharedOnDemandContainer, FContainerInstallData>& Kv : OutInstallData)
 	{
 		const FOnDemandContainer& Container = *Kv.Key;
 		FContainerInstallData& Data			= Kv.Value;
 
-		for (const FPackageId& PackageId : Data.PackageIds)
+		if (Data.ResolvedPackageIds.IsEmpty())
+		{
+			continue;
+		}
+
+		for (const FPackageId& PackageId : Data.ResolvedPackageIds)
 		{
 			const FIoChunkId PackageChunkId	= CreatePackageDataChunkId(PackageId);
 			int32 EntryIndex				= Container.FindChunkEntryIndex(PackageChunkId);
@@ -375,6 +381,7 @@ static FIoStatus BuildInstallData(
 			Data.ResolvedChunks.Add(EntryIndex, EChunkCacheState::None);
 			Data.TotalSize += Container.ChunkEntries[EntryIndex].EncodedSize;
 
+			// TODO: Installing optional bulkdata should probably be an install argument
 			const EIoChunkType AdditionalPackageChunkTypes[] =
 			{
 				EIoChunkType::BulkData,
@@ -384,18 +391,44 @@ static FIoStatus BuildInstallData(
 
 			for (EIoChunkType ChunkType : AdditionalPackageChunkTypes)
 			{
-				const FIoChunkId ChunkId = CreateIoChunkId(PackageId.Value(), 0, ChunkType);
+				const FIoChunkId ChunkId = CreateBulkDataIoChunkId(PackageId.Value(), 0, 0, ChunkType);
 				if (EntryIndex = Container.FindChunkEntryIndex(ChunkId); EntryIndex != INDEX_NONE)
 				{
 					Data.ResolvedChunks.Add(EntryIndex, EChunkCacheState::None);
 					Data.TotalSize += Container.ChunkEntries[EntryIndex].EncodedSize;
 				}
+				else
+				{
+					// Some chunk types resides in other containers
+					for (TPair<FSharedOnDemandContainer, FContainerInstallData>& OtherKv : OutInstallData)
+					{
+						const FOnDemandContainer& OtherContainer = *OtherKv.Key;
+						if (&Container != &OtherContainer)
+						{
+							if (EntryIndex = OtherContainer.FindChunkEntryIndex(ChunkId); EntryIndex != INDEX_NONE)
+							{
+								FContainerInstallData& OtherData = OtherKv.Value;
+								OtherData.ResolvedChunks.Add(EntryIndex, EChunkCacheState::None);
+								OtherData.TotalSize += OtherContainer.ChunkEntries[EntryIndex].EncodedSize;
+							}
+						}
+					}
+				}
 			}
 		}
+	}
+
+	// For now we always download these chunks for containers with resolved chunk ID(s)
+	for (TPair<FSharedOnDemandContainer, FContainerInstallData>& Kv : OutInstallData)
+	{
+		const FOnDemandContainer& Container = *Kv.Key;
+		FContainerInstallData& Data			= Kv.Value;
+
+		if (Data.ResolvedChunks.IsEmpty())
+		{
+			continue;
+		}
 
-		// For now we always download these chunks
-		// TODO: this should only be done for containers that were visited above
-		// or containers with the mount ID of the install request.
 		for (int32 EntryIndex = 0; const FIoChunkId& ChunkId : Container.ChunkIds)
 		{
 			switch(ChunkId.GetChunkType())
@@ -1526,7 +1559,7 @@ FOnDemandInstallResult FOnDemandIoStore::TickInstallRequest(const FInstallReques
 	{
 		FOnDemandContainer& Container				= *Kv.Key;
 		const Private::FContainerInstallData& Data	= Kv.Value;
-		if (!Data.PackageIds.IsEmpty())
+		if (!Data.ResolvedPackageIds.IsEmpty())
 		{
 			check(Container.Header.IsValid());
 			const FIoStatus MountStatus = PackageStoreBackend->Mount(Container.UniqueName(), Container.Header);
@@ -1799,121 +1832,110 @@ void FOnDemandIoStore::CreateContainersFromToc(
 	}
 }
 
-FIoStatus FOnDemandIoStore::GetContainersForInstall(
+FIoStatus FOnDemandIoStore::GetContainersAndPackagesForInstall(
 	FStringView MountId,
-	TSet<FSharedOnDemandContainer>& OutContainersForInstallation,
-	TSet<FSharedOnDemandContainer>& OutContainersWithMountId) const
+	const TArray<FString>& TagSets,
+	const TArray<FPackageId>& PackageIds,
+	TSet<FSharedOnDemandContainer>& OutContainersForInstallation, 
+	TSet<FPackageId>& OutPackageIdsToInstall) const
 {
-	// Only install content from non-streaming container(s)
-	TSet<FSharedOnDemandContainer> ContainersForInstallation;
-	TSet<FSharedOnDemandContainer> ContainersWithMountId;
-
-	UE::TUniqueLock Lock(ContainerMutex);
-
-	for (const FSharedOnDemandContainer& Container : Containers)
+	OutContainersForInstallation.Reset();
 	{
-		if (EnumHasAnyFlags(Container->Flags, EOnDemandContainerFlags::InstallOnDemand) == false)
-		{
-			continue;
-		}
+		UE::TUniqueLock Lock(ContainerMutex);
 
-		// Check that the containers for the request have been mounted
-		if (MountId.IsEmpty() == false && Container->MountId == MountId)
+		OutContainersForInstallation.Reserve(Containers.Num());
+		for (const FSharedOnDemandContainer& Container : Containers)
 		{
-			if (EnumHasAnyFlags(Container->Flags, EOnDemandContainerFlags::PendingEncryptionKey) ||
-				EnumHasAnyFlags(Container->Flags, EOnDemandContainerFlags::Mounted) == false)
+			if (Container->HasAnyFlags(EOnDemandContainerFlags::InstallOnDemand) == false)
 			{
-				return EIoErrorCode::PendingEncryptionKey;
+				continue;
 			}
 
-			ContainersWithMountId.Add(Container);
-		}
+			if (Container->HasAnyFlags(EOnDemandContainerFlags::PendingEncryptionKey))
+			{
+				check(Container->HasAnyFlags(EOnDemandContainerFlags::Mounted) == false);
+				if(MountId == Container->MountId)
+				{
+					check(EnumHasAnyFlags(Container->Flags, EOnDemandContainerFlags::Mounted) == false);
+					FIoStatus Status = FIoStatusBuilder(EIoErrorCode::PendingEncryptionKey)
+						<< TEXT("Trying to install content from encrypted container '")
+						<< Container->Name
+						<< TEXT("'");
+					return Status;
+				}
 
-		// Header will not be valid until the container is fully decrypted/mounted
-		// We must check this under the containter lock so it doesn't race with the 
-		// decryption callback.
-		// After mount, the header can be accessed from multiple threads because it will not be changed again.
-		if (EnumHasAnyFlags(Container->Flags, EOnDemandContainerFlags::Mounted))
-		{
-			ContainersForInstallation.Add(Container);
+				continue;
+			}
+			check(Container->HasAnyFlags(EOnDemandContainerFlags::Mounted));
+			OutContainersForInstallation.Add(Container);
 		}
 	}
 
-	OutContainersForInstallation = MoveTemp(ContainersForInstallation);
-	OutContainersWithMountId = MoveTemp(ContainersWithMountId);
-
-	return EIoErrorCode::Ok;
-}
-
-FIoStatus FOnDemandIoStore::GetContainersAndPackagesForInstall(
-	FStringView MountId,
-	const TArray<FString>& TagSets,
-	const TArray<FPackageId>& PackageIds,
-	TSet<FSharedOnDemandContainer>& OutContainersForInstallation, 
-	TSet<FPackageId>& OutPackageIdsToInstall) const
-{
-	TSet<FSharedOnDemandContainer> ContainersForInstallation;
-	TSet<FSharedOnDemandContainer> ContainersWithMountId;
+	if (OutContainersForInstallation.IsEmpty())
+	{
+		return EIoErrorCode::NotFound;
+	}
 
-	if (FIoStatus Status = GetContainersForInstall(
-		MountId,
-		ContainersForInstallation,
-		ContainersWithMountId); !Status.IsOk())
+	// It's not allowed to install all content
+	if (MountId.IsEmpty() && TagSets.IsEmpty() && PackageIds.IsEmpty())
 	{
+		OutContainersForInstallation.Reset();
+		FIoStatus Status = FIoStatusBuilder(EIoErrorCode::InvalidCode)
+			<< TEXT("Trying to install content from all mounted containers");
 		return Status;
 	}
 
-	TSet<FPackageId> PackageIdsToInstall(PackageIds);
-
-	// Install all packages if no tag set(s) was specified
-	if (TagSets.IsEmpty())
+	const bool bInstallAllPackages = TagSets.IsEmpty() && PackageIds.IsEmpty();
+	if (bInstallAllPackages)
 	{
-		for (const FSharedOnDemandContainer& Container : ContainersWithMountId)
+		check(MountId.IsEmpty() == false);
+		for (const FSharedOnDemandContainer& Container : OutContainersForInstallation)
 		{
-			if (Container->Header.IsValid() == false)
+			if (Container->Header.IsValid() == false || MountId != Container->MountId)
 			{
-				continue;;
+				continue;
 			}
+			UE_LOG(LogIoStoreOnDemand, Log, TEXT("Installing all %d package(s) from contianer '%s'"),
+				Container->ChunkIds.Num(), *Container->UniqueName());
 
-			for (const FPackageId& PackageId : Container->Header->PackageIds)
-			{
-				PackageIdsToInstall.Add(PackageId);
-			}
+			OutPackageIdsToInstall.Append(Container->Header->PackageIds);
 		}
+
+		return FIoStatus::Ok;
 	}
-	else
-	{
-		const TSet<FSharedOnDemandContainer>& SearchContainers =
-			MountId.IsEmpty() ? ContainersForInstallation : ContainersWithMountId;
 
-		for (const FSharedOnDemandContainer& Container : SearchContainers)
+	// Add packages from matching tags
+	for (const FSharedOnDemandContainer& Container : OutContainersForInstallation)
+	{
+		if (Container->Header.IsValid() == false || (!MountId.IsEmpty() && MountId != Container->MountId))
 		{
-			if (Container->Header.IsValid() == false)
-			{
-				continue;;
-			}
+			continue;
+		}
 
-			for (const FString& Tag : TagSets)
+		for (const FString& Tag : TagSets)
+		{
+			for (const FOnDemandTagSet& TagSet : Container->TagSets)
 			{
-				for (const FOnDemandTagSet& TagSet : Container->TagSets)
+				if (TagSet.Tag == Tag)
 				{
-					if (TagSet.Tag == Tag)
+					UE_LOG(LogIoStoreOnDemand, Log, TEXT("Installing %d package(s) with tag '%s' from contianer '%s'"),
+						TagSet.PackageIndicies.Num(), *TagSet.Tag, *Container->UniqueName());
+
+					OutPackageIdsToInstall.Reserve(TagSet.PackageIndicies.Num());
+					for (const uint32 PackageIndex : TagSet.PackageIndicies)
 					{
-						for (const uint32 PackageIndex : TagSet.PackageIndicies)
-						{
-							const FPackageId PackageId = Container->Header->PackageIds[IntCastChecked<int32>(PackageIndex)];
-							PackageIdsToInstall.Add(PackageId);
-						}
+						const FPackageId PackageId = Container->Header->PackageIds[IntCastChecked<int32>(PackageIndex)];
+						OutPackageIdsToInstall.Add(PackageId);
 					}
 				}
 			}
 		}
 	}
 
-	OutContainersForInstallation = MoveTemp(ContainersForInstallation);
-	OutPackageIdsToInstall = MoveTemp(PackageIdsToInstall);
+	// Finally add any specified package ID(s) 
+	OutPackageIdsToInstall.Append(PackageIds);
 
-	return EIoErrorCode::Ok;
+	return FIoStatus::Ok;
 }
 
 void FOnDemandIoStore::ReleaseContent(FOnDemandInternalContentHandle& ContentHandle)
diff --git a/Engine/Source/Runtime/Experimental/IoStore/OnDemand/Private/OnDemandIoStore.h b/Engine/Source/Runtime/Experimental/IoStore/OnDemand/Private/OnDemandIoStore.h
--- a/Engine/Source/Runtime/Experimental/IoStore/OnDemand/Private/OnDemandIoStore.h
+++ b/Engine/Source/Runtime/Experimental/IoStore/OnDemand/Private/OnDemandIoStore.h
@@ -105,6 +105,8 @@ struct FOnDemandContainer
 	inline FOnDemandChunkEntry*				FindChunkEntry(const FIoChunkId& ChunkId);
 	inline FOnDemandChunkEntryReferences&	FindOrAddChunkEntryReferences(const FOnDemandInternalContentHandle& ContentHandle);
 	inline TBitArray<>						GetReferencedChunkEntries() const;
+	bool									HasAnyFlags(EOnDemandContainerFlags Contains) const { return EnumHasAnyFlags(Flags, Contains); }
+	bool									HasAllFlags(EOnDemandContainerFlags Contains) const { return EnumHasAllFlags(Flags, Contains); }
 
 	FAES::FAESKey							EncryptionKey;
 	FSharedContainerHeader					Header;
@@ -305,10 +307,6 @@ public:
 	FOnDemandCacheUsage		GetCacheUsage() const;
 
 private:
-	FIoStatus				GetContainersForInstall(
-								FStringView MountId,
-								TSet<FSharedOnDemandContainer>& OutContainersForInstallation,
-								TSet<FSharedOnDemandContainer>& OutContainersWithMountId) const;
 	FIoStatus				GetContainersAndPackagesForInstall(
 								FStringView MountId,
 								const TArray<FString>& TagSets,
