Initialize IAD cache at module startup.

Since I/O and package store backends needs to be registered on the main thread the I/O store on-demand install cache is
now always initialized and registered at module startup. If the cache fails to initialize the actual error will be reported
if the game tries to download and install packages.

diff --git a/Engine/Source/Runtime/Experimental/IoStore/OnDemand/Private/OnDemandIoStore.cpp b/Engine/Source/Runtime/Experimental/IoStore/OnDemand/Private/OnDemandIoStore.cpp
--- a/Engine/Source/Runtime/Experimental/IoStore/OnDemand/Private/OnDemandIoStore.cpp
+++ b/Engine/Source/Runtime/Experimental/IoStore/OnDemand/Private/OnDemandIoStore.cpp
@@ -607,8 +607,8 @@ FIoStatus FOnDemandIoStore::Initialize()
 		}
 	}
 #endif
-
-	UE_LOG(LogIoStoreOnDemand, Display, TEXT("Initialized"));
+	const FIoStatus CacheStatus = ConditionallyInitializeInstallCache();
+	UE_LOG(LogIoStoreOnDemand, Display, TEXT("Initialized, IAD cache enabled='%s'"), CacheStatus.IsOk() ? TEXT("True") : TEXT("Flase"));
 	return EIoErrorCode::Ok;
 }
 
@@ -910,12 +910,8 @@ FIoStatus FOnDemandIoStore::InitializeInstallCache()
 
 FIoStatus FOnDemandIoStore::ConditionallyInitializeInstallCache()
 {
-	FIoStatus Status = FIoStatus::Ok;
-	if (InstallCache.IsValid() == false)
-	{
-		Status = InitializeInstallCache();
-	}
-	return Status;
+	static FIoStatus CacheStatus = InitializeInstallCache();
+	return CacheStatus;
 }
 
 FOnDemandChunkInfo FOnDemandIoStore::GetChunkInfo(const FIoChunkId& ChunkId, EOnDemandContainerFlags ContainerFlags)
