Constraints: make sure to update the delegates when the component is updated

#rnx

diff --git a/Engine/Source/Runtime/Experimental/Animation/Constraints/Private/TransformableHandle.cpp b/Engine/Source/Runtime/Experimental/Animation/Constraints/Private/TransformableHandle.cpp
--- a/Engine/Source/Runtime/Experimental/Animation/Constraints/Private/TransformableHandle.cpp
+++ b/Engine/Source/Runtime/Experimental/Animation/Constraints/Private/TransformableHandle.cpp
@@ -385,6 +385,7 @@ void UTransformableComponentHandle::OnObjectsReplaced(const TMap<UObject*, UObje
 		if (USceneComponent* NewSceneComponent= Cast<USceneComponent>(NewObject))
 		{
 			Component = NewSceneComponent;
+			RegisterDelegates();
 		}
 	}
 }
@@ -543,6 +544,15 @@ void UTransformableComponentHandle::ResolveBoundObjects(FMovieSceneSequenceID Lo
 	// as this can be called after it has been destroyed, we get the component even if it's pending kill otherwise Get() will return a nullptr.
 	static constexpr bool bEvenIfPendingKill = true;
 	const USceneComponent* ComponentEvenIfPendingKill = Component.Get(bEvenIfPendingKill);
+
+	auto ReplaceComponent = [this](USceneComponent* InNewComponent)
+	{
+		if (InNewComponent != Component)
+		{
+			Component = InNewComponent;
+			RegisterDelegates();
+		}
+	};
 	
 	for (const TWeakObjectPtr<>& ParentObject : ConstraintBindingID.ResolveBoundObjects(LocalSequenceID, SharedPlaybackState))
 	{
@@ -553,11 +563,12 @@ void UTransformableComponentHandle::ResolveBoundObjects(FMovieSceneSequenceID Lo
 			{
 				return ComponentEvenIfPendingKill && (SubComponent->GetFName() == ComponentEvenIfPendingKill->GetFName());
 			});
-			Component = Index != INDEX_NONE ? Components[Index] : Actor->GetRootComponent();
+
+			ReplaceComponent(Index != INDEX_NONE ? Components[Index] : Actor->GetRootComponent());
 		}
 		else if (USceneComponent* Comp = Cast<USceneComponent>(ParentObject.Get()))
 		{
-			Component = Comp;
+			ReplaceComponent(Comp);
 		}
 		break; //just do one
 	}
