Defer initialization of on-demand install cache until server fork.

diff --git a/Engine/Source/Runtime/Experimental/IoStore/OnDemand/Private/OnDemandIoStore.cpp b/Engine/Source/Runtime/Experimental/IoStore/OnDemand/Private/OnDemandIoStore.cpp
--- a/Engine/Source/Runtime/Experimental/IoStore/OnDemand/Private/OnDemandIoStore.cpp
+++ b/Engine/Source/Runtime/Experimental/IoStore/OnDemand/Private/OnDemandIoStore.cpp
@@ -521,6 +521,11 @@ FOnDemandIoStore::~FOnDemandIoStore()
 		FCoreInternalDelegates::GetOnPakMountOperation().Remove(OnMountPakHandle);
 	}
 
+	if (OnServerPostForkHandle.IsValid())
+	{
+		FCoreDelegates::OnPostFork.Remove(OnServerPostForkHandle);
+	}
+
 	if (TickFuture.IsValid())
 	{
 		TickFuture.Wait();
@@ -607,8 +612,31 @@ FIoStatus FOnDemandIoStore::Initialize()
 		}
 	}
 #endif
-	const FIoStatus CacheStatus = ConditionallyInitializeInstallCache();
-	UE_LOG(LogIoStoreOnDemand, Display, TEXT("Initialized, IAD cache enabled='%s'"), CacheStatus.IsOk() ? TEXT("True") : TEXT("Flase"));
+
+	FString CacheStatusText;
+	if (FForkProcessHelper::IsForkRequested() && !FForkProcessHelper::IsForkedChildProcess())
+	{
+		// Defer initialization when forking server process to get a unique install cache location per process
+		CacheStatusText = TEXT("WaitingOnServerFork");
+		OnServerPostForkHandle = FCoreDelegates::OnPostFork.AddLambda(
+			[this](EForkProcessRole ProcessRole)
+			{
+				if (ProcessRole == EForkProcessRole::Child)
+				{
+					if (FIoStatus Status = ConditionallyInitializeInstallCache(); !Status.IsOk())
+					{
+						UE_LOG(LogIoStoreOnDemand, Error, TEXT("Failed to initialize install cache, reason '%s'"), *Status.ToString());
+					}
+				}
+			});
+	}
+	else
+	{
+		const FIoStatus CacheStatus = ConditionallyInitializeInstallCache();
+		CacheStatusText = CacheStatus.IsOk() ? TEXT("Enabled") : TEXT("Disabled");
+	}
+
+	UE_LOG(LogIoStoreOnDemand, Display, TEXT("Initialized, InstallCache='%s'"), *CacheStatusText);
 	return EIoErrorCode::Ok;
 }
 
@@ -801,8 +829,7 @@ FIoStatus FOnDemandIoStore::InitializeInstallCache()
 	{
 		if (!Private::HasInstallCacheDirectory())
 		{
-			UE_LOG(LogIoStoreOnDemand, Error, TEXT("Project does not have persistent storage configured"));
-			return FIoStatusBuilder(EIoErrorCode::InvalidParameter) << TEXT("Failed to initialize install cache");
+			return FIoStatusBuilder(EIoErrorCode::InvalidParameter) << TEXT("Persistent storage not configured");
 		}
 
 		FOnDemandInstallCacheConfig CacheConfig;
diff --git a/Engine/Source/Runtime/Experimental/IoStore/OnDemand/Private/OnDemandIoStore.h b/Engine/Source/Runtime/Experimental/IoStore/OnDemand/Private/OnDemandIoStore.h
--- a/Engine/Source/Runtime/Experimental/IoStore/OnDemand/Private/OnDemandIoStore.h
+++ b/Engine/Source/Runtime/Experimental/IoStore/OnDemand/Private/OnDemandIoStore.h
@@ -337,6 +337,7 @@ private:
 	FSharedInstallCache					InstallCache;
 	FSharedPackageStoreBackend			PackageStoreBackend;
 	FDelegateHandle						OnMountPakHandle;
+	FDelegateHandle						OnServerPostForkHandle;
 	TArray<FSharedOnDemandContainer>	Containers;
 	TMap<FString, FIoBuffer>			PendingContainerHeaders;
 	mutable UE::FMutex					ContainerMutex;
