Dataflow : add image nodes : Splitchannesl / CombineChannel / Image from color

diff --git a/Engine/Source/Runtime/Experimental/Dataflow/Core/Private/Dataflow/DataflowImage.cpp b/Engine/Source/Runtime/Experimental/Dataflow/Core/Private/Dataflow/DataflowImage.cpp
--- a/Engine/Source/Runtime/Experimental/Dataflow/Core/Private/Dataflow/DataflowImage.cpp
+++ b/Engine/Source/Runtime/Experimental/Dataflow/Core/Private/Dataflow/DataflowImage.cpp
@@ -17,16 +17,41 @@ const FImage& FDataflowImage::GetImage() const
 	return Image;
 }
 
+void FDataflowImage::CreateR32F(EDataflowImageResolution Resolution)
+{
+	CreateR32F((int32)Resolution, (int32)Resolution);
+}
+
 void FDataflowImage::CreateR32F(int32 Width, int32 Height)
 {
 	Image.Init(Width, Height, ERawImageFormat::R32F, EGammaSpace::Linear);
 }
 
+void FDataflowImage::CreateRGBA32F(EDataflowImageResolution Resolution)
+{
+	CreateRGBA32F((int32)Resolution, (int32)Resolution);
+}
+
 void FDataflowImage::CreateRGBA32F(int32 Width, int32 Height)
 {
 	Image.Init(Width, Height, ERawImageFormat::RGBA32F, EGammaSpace::Linear);
 }
 
+void FDataflowImage::CreateFromColor(EDataflowImageResolution Resolution, FLinearColor Color)
+{
+	CreateFromColor((int32)Resolution, (int32)Resolution, Color);
+}
+
+void FDataflowImage::CreateFromColor(int32 Width, int32 Height, FLinearColor Color)
+{
+	CreateRGBA32F(Width, Height);
+
+	TArray<FLinearColor> FillData;
+	FillData.Init(Color, Width * Height);
+	const FImageView SourceImage(FillData.GetData(), Width, Height);
+	FImageCore::CopyImage(SourceImage, Image);
+}
+
 bool FDataflowImage::CopyRGBAPixels(TArrayView64<FVector4f> Pixels)
 {
 	if (Image.Format != ERawImageFormat::Type::RGBA32F)
@@ -47,6 +72,98 @@ void FDataflowImage::ConvertToRGBA32F()
 	Image.ChangeFormat(ERawImageFormat::RGBA32F, EGammaSpace::Linear);
 }
 
+void FDataflowImage::ReadChannel(EDataflowImageChannel Channel, FDataflowImage& OutImage) const
+{
+	if (Image.Format == ERawImageFormat::R32F)
+	{
+		OutImage.CreateR32F(GetWidth(), GetHeight());
+		// single channel image, all channels get the same data 
+		FImageCore::CopyImage(this->Image, OutImage.Image);
+	}
+	else if (Image.Format == ERawImageFormat::RGBA32F)
+	{
+		OutImage.CreateR32F(GetWidth(), GetHeight());
+		TArrayView64<float> DstPixels = OutImage.Image.AsR32F();
+
+		const TArrayView64<const FLinearColor> SrcPixels = Image.AsRGBA32F();
+		for (int32 Pixel = 0; Pixel < SrcPixels.Num(); ++Pixel)
+		{
+			const FLinearColor& SrcColor = SrcPixels[Pixel];
+			DstPixels[Pixel] = SrcColor.Component((int32)Channel);
+		}
+	}
+	else 
+	{
+		// unsupported format - image is likely not initialized likely fill with black  
+		OutImage.CreateFromColor(GetWidth(), GetHeight(), FLinearColor::Black);
+	}
+}
+
+void FDataflowImage::WriteChannel(EDataflowImageChannel Channel, const FDataflowImage& SrcImage)
+{
+	FImageView SrcImageView = SrcImage.Image;
+
+	FImage ResizedSrcImage;
+	if (SrcImageView.GetWidth() != GetWidth() && SrcImageView.GetHeight() != GetHeight())
+	{
+		// resize source image into ResizedSrcImage
+		ResizedSrcImage.Init(GetWidth(), GetHeight(), SrcImageView.Format, EGammaSpace::Linear);
+		FImageCore::ResizeImage(SrcImageView, ResizedSrcImage);
+		SrcImageView = ResizedSrcImage;
+	}
+
+	FImage SingleChannelSrcImage;
+	if (SrcImageView.Format != ERawImageFormat::R32F)
+	{
+		SingleChannelSrcImage.Init(GetWidth(), GetHeight(), ERawImageFormat::R32F, EGammaSpace::Linear);
+		if (SrcImageView.Format == ERawImageFormat::RGBA32F)
+		{
+			// convert the image grayscale and single channel
+			TArrayView64<float> DstPixels = SingleChannelSrcImage.AsR32F();
+
+			const TArrayView64<const FLinearColor> SrcPixels = SrcImageView.AsRGBA32F();
+			for (int32 Pixel = 0; Pixel < SrcPixels.Num(); ++Pixel)
+			{
+				const FLinearColor& SrcColor = SrcPixels[Pixel];
+				const float Value = SrcColor.LinearRGBToHSV().B; // blue store the value of the color 
+				DstPixels[Pixel] = Value;
+			}
+		}
+		else
+		{
+			// all zero texture
+			SingleChannelSrcImage.RawData.Init(0, SingleChannelSrcImage.RawData.Num());
+		}
+		SrcImageView = SingleChannelSrcImage;
+	}
+
+	ensure(SrcImageView.GetWidth() == GetWidth());
+	ensure(SrcImageView.GetHeight() == GetHeight());
+	ensure(SrcImageView.Format == ERawImageFormat::R32F);
+
+	// finally write to the channel
+	if (Image.Format == ERawImageFormat::R32F)
+	{
+		FImageCore::CopyImage(SrcImageView, Image);
+	}
+	else if (Image.Format == ERawImageFormat::RGBA32F)
+	{
+		const TArrayView64<const float> SrcPixels = SrcImageView.AsR32F();
+
+		TArrayView64<FLinearColor> DstPixels = Image.AsRGBA32F();
+		for (int32 Pixel = 0; Pixel < DstPixels.Num(); ++Pixel)
+		{
+			const float& SrcValue = SrcPixels[Pixel];
+			DstPixels[Pixel].Component((int32)Channel) = SrcValue;
+		}
+	}
+	else
+	{
+		// trying to write a channel to an probably non initialize image
+		ensure(false);
+	}
+}
+
 bool FDataflowImage::Serialize(FArchive& Ar)
 {
 	FImageInfo Info(Image);
diff --git a/Engine/Source/Runtime/Experimental/Dataflow/Core/Private/Dataflow/DataflowImageNodes.cpp b/Engine/Source/Runtime/Experimental/Dataflow/Core/Private/Dataflow/DataflowImageNodes.cpp
new file mode 100644
--- /dev/null
+++ b/Engine/Source/Runtime/Experimental/Dataflow/Core/Private/Dataflow/DataflowImageNodes.cpp
@@ -0,0 +1,166 @@
+// Copyright Epic Games, Inc. All Rights Reserved.
+
+#include "Dataflow/DataflowImageNodes.h"
+
+#include "Dataflow/DataflowNodeFactory.h"
+
+namespace UE::Dataflow
+{
+	void RegisterDataflowImageNodes()
+	{
+		DATAFLOW_NODE_REGISTER_CREATION_FACTORY(FDataflowImageFromColorNode);
+		DATAFLOW_NODE_REGISTER_CREATION_FACTORY(FDataflowImageSplitChannelsNode);
+		DATAFLOW_NODE_REGISTER_CREATION_FACTORY(FDataflowImageCombineChannelsNode);
+	}
+}
+
+////////////////////////////////////////////////////////////////////////////////////////////////////////////////
+
+FDataflowImageFromColorNode::FDataflowImageFromColorNode(const UE::Dataflow::FNodeParameters& InParam, FGuid InGuid)
+	:FDataflowNode(InParam, InGuid)
+{
+	RegisterInputConnection(&FillColor)
+		.SetCanHidePin(true)
+		.SetPinIsHidden(true);
+	RegisterInputConnection(&Resolution)
+		.SetCanHidePin(true)
+		.SetPinIsHidden(true);
+	RegisterOutputConnection(&Image);
+}
+
+void FDataflowImageFromColorNode::Evaluate(UE::Dataflow::FContext& Context, const FDataflowOutput* Out) const
+{
+	if (Out->IsA(&Image))
+	{
+		const FLinearColor& InColor = GetValue(Context, &FillColor);
+		const EDataflowImageResolution& InResolution = GetValue(Context, &Resolution);
+
+		FDataflowImage OutImage;
+		OutImage.CreateFromColor(InResolution, InColor);
+		SetValue(Context, MoveTemp(OutImage), &Image);
+	}
+}
+
+////////////////////////////////////////////////////////////////////////////////////////////////////////////////
+
+FDataflowImageSplitChannelsNode::FDataflowImageSplitChannelsNode(const UE::Dataflow::FNodeParameters& InParam, FGuid InGuid)
+	:FDataflowNode(InParam, InGuid)
+{
+	RegisterInputConnection(&Image);
+	RegisterOutputConnection(&Red);
+	RegisterOutputConnection(&Green);
+	RegisterOutputConnection(&Blue);
+	RegisterOutputConnection(&Alpha);
+}
+
+void FDataflowImageSplitChannelsNode::Evaluate(UE::Dataflow::FContext& Context, const FDataflowOutput* Out) const
+{
+	auto EvaluateForChannel = [this](UE::Dataflow::FContext& Context, EDataflowImageChannel Channel, const FDataflowImage* OutputRef)
+		{
+			FDataflowImage OutImage;
+			const FDataflowImage& InImage = GetValue(Context, &Image);
+			InImage.ReadChannel(Channel, OutImage);
+			SetValue(Context, MoveTemp(OutImage), OutputRef);
+		};
+
+	if (Out->IsA(&Red))
+	{
+		EvaluateForChannel(Context, EDataflowImageChannel::Red, &Red);
+	}
+	else if (Out->IsA(&Green))
+	{
+		EvaluateForChannel(Context, EDataflowImageChannel::Green, &Green);
+	}
+	else if (Out->IsA(&Blue))
+	{
+		EvaluateForChannel(Context, EDataflowImageChannel::Blue, &Blue);
+	}
+	else if (Out->IsA(&Alpha))
+	{
+		EvaluateForChannel(Context, EDataflowImageChannel::Alpha, &Alpha);
+	}
+}
+
+////////////////////////////////////////////////////////////////////////////////////////////////////////////////
+
+FDataflowImageCombineChannelsNode::FDataflowImageCombineChannelsNode(const UE::Dataflow::FNodeParameters& InParam, FGuid InGuid)
+	:FDataflowNode(InParam, InGuid)
+{
+	RegisterInputConnection(&Red);
+	RegisterInputConnection(&Green);
+	RegisterInputConnection(&Blue);
+	RegisterInputConnection(&Alpha);
+	RegisterOutputConnection(&Image);
+}
+
+void FDataflowImageCombineChannelsNode::GetUserDefinedResolution(const TArray<const FDataflowImage*>& Images, int32& OutWidth, int32& OutHeight)  const
+{
+	// first set safe non zero values
+	OutWidth  = (int32)Resolution;
+	OutHeight = (int32)Resolution;
+
+	if (ResolutionOptions == EDataflowImageCombineResolutionOption::UserDefined 
+		|| Images.Num() == 0)
+	{
+		return;
+	}
+
+	if (Images[0])
+	{
+		int32 PotentialW = Images[0]->GetWidth();
+		int32 PotentialH = Images[0]->GetHeight();
+		for (const FDataflowImage* ImageEntry : Images)
+		{
+			if (ImageEntry)
+			{
+				const int32 ImageW = ImageEntry->GetWidth();
+				const int32 ImageH = ImageEntry->GetHeight();
+				if (ImageW > 0 && ImageH > 0)
+				{
+					if (ResolutionOptions == EDataflowImageCombineResolutionOption::Highest)
+					{
+						PotentialW = FMath::Max(PotentialW, ImageW);
+						PotentialH = FMath::Max(PotentialH, ImageH);
+					}
+					else if (ResolutionOptions == EDataflowImageCombineResolutionOption::Lowest)
+					{
+						PotentialW = FMath::Min(PotentialW, ImageW);
+						PotentialH = FMath::Min(PotentialH, ImageH);
+					}
+				}
+			}
+		}
+
+		if (PotentialW > 0 && PotentialH > 0)
+		{
+			OutWidth = PotentialW;
+			OutHeight = PotentialH;
+		}
+	}
+}
+
+void FDataflowImageCombineChannelsNode::Evaluate(UE::Dataflow::FContext& Context, const FDataflowOutput* Out) const
+{
+	if (Out->IsA(&Image))
+	{
+		const FDataflowImage& InRed   = GetValue(Context, &Red);
+		const FDataflowImage& InGreen = GetValue(Context, &Green);
+		const FDataflowImage& InBlue  = GetValue(Context, &Blue);
+		const FDataflowImage& InAlpha = GetValue(Context, &Alpha);
+		const TArray<const FDataflowImage*> InImages = { &InRed, &InGreen, &InBlue, &InAlpha };
+
+		FDataflowImage OutImage;
+
+		int32 OutWidth = 0;
+		int32 OutHeight= 0;
+		GetUserDefinedResolution(InImages, OutWidth, OutHeight);
+		OutImage.CreateRGBA32F(OutWidth, OutHeight);
+
+		OutImage.WriteChannel(EDataflowImageChannel::Red, InRed);
+		OutImage.WriteChannel(EDataflowImageChannel::Green, InGreen);
+		OutImage.WriteChannel(EDataflowImageChannel::Blue, InBlue);
+		OutImage.WriteChannel(EDataflowImageChannel::Alpha, InAlpha);
+
+		SetValue(Context, MoveTemp(OutImage), &Image);
+	}
+}
diff --git a/Engine/Source/Runtime/Experimental/Dataflow/Core/Private/Dataflow/DataflowNodeFactory.cpp b/Engine/Source/Runtime/Experimental/Dataflow/Core/Private/Dataflow/DataflowNodeFactory.cpp
--- a/Engine/Source/Runtime/Experimental/Dataflow/Core/Private/Dataflow/DataflowNodeFactory.cpp
+++ b/Engine/Source/Runtime/Experimental/Dataflow/Core/Private/Dataflow/DataflowNodeFactory.cpp
@@ -4,6 +4,7 @@
 
 #include "Dataflow/DataflowNode.h"
 #include "Dataflow/DataflowCoreNodes.h"
+#include "Dataflow/DataflowImageNodes.h"
 #include "Dataflow/DataflowMathNodes.h"
 #include "Dataflow/DataflowVectorNodes.h"
 #include "Misc/MessageDialog.h"
@@ -19,7 +20,8 @@ namespace UE::Dataflow
 	{
 		RegisterCoreNodes();
 		RegisterDataflowMathNodes();
-		RegisterDataflowVectorNodes();
+		RegisterDataflowImageNodes();
+		UE::Dataflow::RegisterAnyTypes();
 		UE::Dataflow::RegisterAnyTypes();
 	}
 
diff --git a/Engine/Source/Runtime/Experimental/Dataflow/Core/Public/Dataflow/DataflowImage.h b/Engine/Source/Runtime/Experimental/Dataflow/Core/Public/Dataflow/DataflowImage.h
--- a/Engine/Source/Runtime/Experimental/Dataflow/Core/Public/Dataflow/DataflowImage.h
+++ b/Engine/Source/Runtime/Experimental/Dataflow/Core/Public/Dataflow/DataflowImage.h
@@ -21,6 +21,15 @@ enum class EDataflowImageResolution : int32
 	Resolution8192 = 8192 UMETA(DisplayName = "8192 x 8192")
 };
 
+UENUM()
+enum class EDataflowImageChannel : int32
+{
+	Red = 0 UMETA(DisplayName = "Red Channel"),
+	Green = 1 UMETA(DisplayName = "Green Channel"),
+	Blue = 2 UMETA(DisplayName = "Blue Channel"),
+	Alpha = 3 UMETA(DisplayName = "Alpha Channel"),
+};
+
 /** 
 * Represents image for dataflow
 * type is constrained to Float32 with 1 or 4 channels
@@ -36,12 +45,24 @@ struct FDataflowImage
 	/** get the readonly underlying image object */
 	DATAFLOWCORE_API const FImage& GetImage() const;
 
+	/** Create a single channel float format image */
+	DATAFLOWCORE_API void CreateR32F(EDataflowImageResolution Resolution);
+
 	/** Create a single channel float format image */
 	DATAFLOWCORE_API void CreateR32F(int32 Width, int32 Height);
 
+	/** Create a four channels float format image */
+	DATAFLOWCORE_API void CreateRGBA32F(EDataflowImageResolution Resolution);
+
 	/** Create a four channels float format image */
 	DATAFLOWCORE_API void CreateRGBA32F(int32 Width, int32 Height);
 
+	/** Create a four channels float format image filled with a specific color */
+	DATAFLOWCORE_API void CreateFromColor(EDataflowImageResolution Resolution, FLinearColor Color);
+
+	/** Create a four channels float format image filled with a specific color */
+	DATAFLOWCORE_API void CreateFromColor(int32 Width, int32 Height, FLinearColor Color);
+
 	/**
 	* Copy RGBA32F pixels to the image
 	* Number of pixels must match and format must be RGBA32F
@@ -57,6 +78,19 @@ struct FDataflowImage
 
 	DATAFLOWCORE_API bool Serialize(FArchive& Ar);
 
+	/**
+	* Get a specific color channel and copy it to an image
+	* Warning : the outImage will be resized to the size of the current image and any previously store data will be lost 
+	*/
+	DATAFLOWCORE_API void ReadChannel(EDataflowImageChannel Channel, FDataflowImage& OutImage) const;
+
+	/**
+	* write to a sepcific channel from an existing image
+	* if the source image is not the same size it will be resized to adapt the size of the current image
+	* if the source image is not a greyscale image it will be converted to greyscale before copy the data to the channel
+	*/
+	DATAFLOWCORE_API void WriteChannel(EDataflowImageChannel Channel, const FDataflowImage& SrcImage);
+
 private:
 	FImage Image;
 };
diff --git a/Engine/Source/Runtime/Experimental/Dataflow/Core/Public/Dataflow/DataflowImageNodes.h b/Engine/Source/Runtime/Experimental/Dataflow/Core/Public/Dataflow/DataflowImageNodes.h
new file mode 100644
--- /dev/null
+++ b/Engine/Source/Runtime/Experimental/Dataflow/Core/Public/Dataflow/DataflowImageNodes.h
@@ -0,0 +1,136 @@
+// Copyright Epic Games, Inc. All Rights Reserved.
+
+#pragma once
+
+#include "Dataflow/DataflowNode.h"
+#include "Dataflow/DataflowImage.h"
+
+#include "DataflowImageNodes.generated.h"
+
+#define DATAFLOW_IMAGE_NODES_CATEGORY "Image"
+
+/**
+* Create a RGBA image from a single color at a specific resolution
+* Outputs are single channel images
+*/
+USTRUCT()
+struct FDataflowImageFromColorNode : public FDataflowNode
+{
+	GENERATED_USTRUCT_BODY()
+	DATAFLOW_NODE_DEFINE_INTERNAL(FDataflowImageFromColorNode, "ImageFromColor", DATAFLOW_IMAGE_NODES_CATEGORY, "")
+
+public:
+	FDataflowImageFromColorNode(const UE::Dataflow::FNodeParameters& InParam, FGuid InGuid = FGuid::NewGuid());
+
+private:
+	/** color  to fill the image with */
+	UPROPERTY(EditAnywhere, Category = Options, meta = (DataflowInput));
+	FLinearColor FillColor;
+
+	/** resolution of the output image */
+	UPROPERTY(EditAnywhere, Category = Options, meta = (DataflowInput));
+	EDataflowImageResolution Resolution = EDataflowImageResolution::Resolution16;
+
+	/** Output image  */
+	UPROPERTY(meta = (DataflowOutput));
+	FDataflowImage Image;
+
+	virtual void Evaluate(UE::Dataflow::FContext& Context, const FDataflowOutput* Out) const override;
+};
+
+/** 
+* Split a image in individual channels 
+* Outputs are single channel images 
+*/
+USTRUCT()
+struct FDataflowImageSplitChannelsNode : public FDataflowNode
+{
+	GENERATED_USTRUCT_BODY()
+	DATAFLOW_NODE_DEFINE_INTERNAL(FDataflowImageSplitChannelsNode, "ImageSplitChannels", DATAFLOW_IMAGE_NODES_CATEGORY, "")
+
+public:
+	FDataflowImageSplitChannelsNode(const UE::Dataflow::FNodeParameters& InParam, FGuid InGuid = FGuid::NewGuid());
+
+private:
+	/** Input image to split per channel */
+	UPROPERTY(meta = (DataflowInput));
+	FDataflowImage Image;
+
+	/** Red channel */
+	UPROPERTY(meta = (DataflowOutput));
+	FDataflowImage Red;
+
+	/** Green channel */
+	UPROPERTY(meta = (DataflowOutput));
+	FDataflowImage Green;
+
+	/** Blue channel */
+	UPROPERTY(meta = (DataflowOutput));
+	FDataflowImage Blue;
+
+	/** Alpha channel */
+	UPROPERTY(meta = (DataflowOutput));
+	FDataflowImage Alpha;
+
+	virtual void Evaluate(UE::Dataflow::FContext& Context, const FDataflowOutput* Out) const override;
+};
+
+UENUM()
+enum class EDataflowImageCombineResolutionOption : int32
+{
+	Lowest = 0 UMETA(DisplayName = "Lowest Resolution"),
+	Highest = 1 UMETA(DisplayName = "Highest Resolution"),
+	UserDefined = 2 UMETA(DisplayName = "User Defined Resolution "),
+};
+
+/**
+* Combine channels into a single RGBA image
+* Outputs are single channel images
+*/
+USTRUCT()
+struct FDataflowImageCombineChannelsNode : public FDataflowNode
+{
+	GENERATED_USTRUCT_BODY()
+	DATAFLOW_NODE_DEFINE_INTERNAL(FDataflowImageCombineChannelsNode, "ImageCombineChannels", DATAFLOW_IMAGE_NODES_CATEGORY, "")
+
+public:
+	FDataflowImageCombineChannelsNode(const UE::Dataflow::FNodeParameters& InParam, FGuid InGuid = FGuid::NewGuid());
+
+private:
+	/** Red channel - if not connected, use black color */
+	UPROPERTY(meta = (DataflowInput));
+	FDataflowImage Red;
+
+	/** Green channel - if not connected, use black color */
+	UPROPERTY(meta = (DataflowInput));
+	FDataflowImage Green;
+
+	/** Blue channel - if not connected, use black color */
+	UPROPERTY(meta = (DataflowInput));
+	FDataflowImage Blue;
+
+	/** Alpha channel - if not connected, use black color */
+	UPROPERTY(meta = (DataflowInput));
+	FDataflowImage Alpha;
+
+	/** Output image recombined from input channels */
+	UPROPERTY(meta = (DataflowOutput));
+	FDataflowImage Image;
+
+	/** resolution option */
+	UPROPERTY(EditAnywhere, Category = Options);
+	EDataflowImageCombineResolutionOption ResolutionOptions = EDataflowImageCombineResolutionOption::Highest;
+
+	/** resolution of the output image if the resolution option is set to user defined*/
+	UPROPERTY(EditAnywhere, Category = Options, meta = (EditCondition = "ResolutionOptions == EDataflowImageCombineResolutionOption::UserDefined"))
+	EDataflowImageResolution Resolution = EDataflowImageResolution::Resolution512;
+
+	virtual void Evaluate(UE::Dataflow::FContext& Context, const FDataflowOutput* Out) const override;
+
+	void GetUserDefinedResolution(const TArray<const FDataflowImage*>& Images, int32& OutWidth, int32& OutHeight) const;
+};
+
+namespace UE::Dataflow
+{
+	void RegisterDataflowImageNodes();
+}
