Dataflow image node : fix fill color not being initialized

diff --git a/Engine/Source/Runtime/Experimental/Dataflow/Core/Public/Dataflow/DataflowImageNodes.h b/Engine/Source/Runtime/Experimental/Dataflow/Core/Public/Dataflow/DataflowImageNodes.h
--- a/Engine/Source/Runtime/Experimental/Dataflow/Core/Public/Dataflow/DataflowImageNodes.h
+++ b/Engine/Source/Runtime/Experimental/Dataflow/Core/Public/Dataflow/DataflowImageNodes.h
@@ -25,7 +25,7 @@ public:
 private:
 	/** color  to fill the image with */
 	UPROPERTY(EditAnywhere, Category = Options, meta = (DataflowInput));
-	FLinearColor FillColor;
+	FLinearColor FillColor = FLinearColor::Black;
 
 	/** resolution of the output image */
 	UPROPERTY(EditAnywhere, Category = Options, meta = (DataflowInput));
