IasTool - add option to dump .uondemandtoc file on disk to Json.

diff --git a/Engine/Source/Runtime/Experimental/IoStore/OnDemand/Private/Tool/ListTocs.cpp b/Engine/Source/Runtime/Experimental/IoStore/OnDemand/Private/Tool/ListTocs.cpp
--- a/Engine/Source/Runtime/Experimental/IoStore/OnDemand/Private/Tool/ListTocs.cpp
+++ b/Engine/Source/Runtime/Experimental/IoStore/OnDemand/Private/Tool/ListTocs.cpp
@@ -7,6 +7,7 @@
 #include "Async/Mutex.h"
 #include "Async/ParallelFor.h"
 #include "Async/UniqueLock.h"
+#include "HAL/FileManager.h"
 #include "IO/IoStoreOnDemand.h"
 #include "Misc/EnumClassFlags.h"
 #include "Misc/FileHelper.h"
@@ -68,6 +69,7 @@ struct FListTocsParams
 	FString CredentialsFile;
 	FString CredentialsFileKeyName;
 	FString TocUrl;
+	FString TocPath;
 	FString TocKey;
 	FString BuildVersion;
 	FString TargetPlatform;
@@ -85,6 +87,10 @@ FIoStatus FListTocsParams::Validate() const
 	{
 		return FIoStatus::Ok;
 	}
+	if (!TocPath.IsEmpty())
+	{
+		return FIoStatus::Ok;
+	}
 	if (Bucket.IsEmpty())
 	{
 		return FIoStatus(EIoErrorCode::InvalidParameter, TEXT("Invalid bucket name"));
@@ -146,6 +152,7 @@ static FListTocsParams BuildListTocsParams(const FContext& Context)
 	Ret.CredentialsFile			= Context.Get<FStringView>(TEXT("-CredentialsFile"),		Ret.CredentialsFile);
 	Ret.CredentialsFileKeyName	= Context.Get<FStringView>(TEXT("-CredentialsFileKeyName"),	Ret.CredentialsFileKeyName);
 	Ret.TocUrl					= Context.Get<FStringView>(TEXT("-TocUrl"),					Ret.TocUrl);
+	Ret.TocPath					= Context.Get<FStringView>(TEXT("-TocPath"),				Ret.TocPath);
 	Ret.TocKey					= Context.Get<FStringView>(TEXT("-TocKey"),					Ret.TocKey);
 	Ret.BuildVersion			= Context.Get<FStringView>(TEXT("-BuildVersion"),			Ret.BuildVersion);
 	Ret.TargetPlatform			= Context.Get<FStringView>(TEXT("-TargetPlatform"),			Ret.TargetPlatform);
@@ -342,7 +349,39 @@ FIoStatus ListTocs(const FListTocsParams& Params)
 	FMutex TocMutex;
 	TArray<FTocDescription> Tocs;
 
-	if (Params.TocUrl.IsEmpty())
+	if (!Params.TocPath.IsEmpty())
+	{
+		UE_LOG(LogIas, Display, TEXT("Loading TOC from file '%s'"), *Params.TocPath);
+
+		const int64 FileSize = IFileManager::Get().FileSize(*Params.TocPath);
+		if (FileSize <= 0)
+		{
+			return FIoStatus(EIoErrorCode::NotFound);
+		}
+
+		TIoStatusOr<FOnDemandToc> TocStatus = FOnDemandToc::LoadFromFile(Params.TocPath, false);
+		if (!TocStatus.IsOk())
+		{
+			return TocStatus.Status();
+		}
+
+		FOnDemandToc Toc = TocStatus.ConsumeValueOrDie();
+		if (Toc.Header.Magic == FOnDemandTocHeader::ExpectedMagic)
+		{
+			FDateTime DateTime = FDateTime::FromUnixTimestamp(Toc.Meta.EpochTimestamp);
+			Tocs.Add(FTocDescription 
+			{
+				.Toc = MoveTemp(Toc),
+				.DateTime = DateTime,
+				.Size = static_cast<uint64>(FileSize)
+			});
+		}
+		else
+		{
+			return FIoStatus(EIoErrorCode::FileOpenFailed, TEXT("Header mismatch"));
+		}
+	}
+	else if (Params.TocUrl.IsEmpty())
 	{
 		FStringView PrefixView = Params.BucketPrefix;
 		if (PrefixView.StartsWith(TEXT("/")))
@@ -504,7 +543,14 @@ FIoStatus ListTocs(const FListTocsParams& Params)
 			UE_LOG(LogIas, Display, TEXT("%-20s: %.2lf KiB"), TEXT("TocSize"), double(Desc.Size) / 1024.0);
 			UE_LOG(LogIas, Display, TEXT("%-20s: %llu"), TEXT("ChunkCount"), Desc.ChunkCount);
 			UE_LOG(LogIas, Display, TEXT("%-20s: %.2lf MiB"), TEXT("TotalChunkSize"), double(Desc.TotalChunkSize) / 1024.0 / 1024.0);
-			UE_LOG(LogIas, Display, TEXT("%-20s: %s"), TEXT("Url"), Url.ToString());
+			if (Params.TocPath.IsEmpty())
+			{
+				UE_LOG(LogIas, Display, TEXT("%-20s: %s"), TEXT("Url"), Url.ToString());
+			}
+			else
+			{
+				UE_LOG(LogIas, Display, TEXT("%-20s: %s"), TEXT("FilePath"), *Params.TocPath);
+			}
 		}
 		UE_LOG(LogIas, Display, TEXT(""));
 		
@@ -643,6 +689,7 @@ static FCommand ListTocsCommand(
 	{
 		TArgument<FStringView>(TEXT("-Json"),			TEXT("Optionally write output to the given JSON file")),
 		TArgument<FStringView>(TEXT("-TocUrl"),			TEXT("URL of a TOC file to fetch info of (overrides Bucket/ServiceUrl)")),
+		TArgument<FStringView>(TEXT("-TocPath"),		TEXT("Path to a .uondemandtoc file on disk")),
 		TArgument<FStringView>(TEXT("-TocKey"),			TEXT("Filter output by TOC key")),
 		TArgument<FStringView>(TEXT("-BuildVersion"),	TEXT("Filter results a particular build version")),
 		TArgument<FStringView>(TEXT("-TargetPlatform"),	TEXT("Only show results for a given platform")),
