Tentative fix for the IAS socket crash

diff --git a/Engine/Source/Runtime/Experimental/IoStore/HttpClient/Private/Socket.inl b/Engine/Source/Runtime/Experimental/IoStore/HttpClient/Private/Socket.inl
--- a/Engine/Source/Runtime/Experimental/IoStore/HttpClient/Private/Socket.inl
+++ b/Engine/Source/Runtime/Experimental/IoStore/HttpClient/Private/Socket.inl
@@ -55,6 +55,9 @@
 #	include "CoreHttp/Http.inl"
 #endif
 
+#ifndef IAS_HTTP_SOCKET_SEND_MSG_NOSIGNAL
+#define IAS_HTTP_SOCKET_SEND_MSG_NOSIGNAL 0
+#endif
 // }}}
 
 namespace UE::IoStore::HTTP
@@ -359,7 +362,11 @@ void FSocket::Disconnect()
 FOutcome FSocket::Send(const char* Data, uint32 Size)
 {
 	Trace(Socket, ETrace::Send, -1);
+#if IAS_HTTP_SOCKET_SEND_MSG_NOSIGNAL
+	int32 Result = send(Socket, Data, Size, MsgFlagType(MSG_NOSIGNAL));
+#else
 	int32 Result = send(Socket, Data, Size, MsgFlagType(0));
+#endif
 	Trace(Socket, ETrace::Send, FMath::Max(Result, 0));
 
 	if (Result > 0)						return FOutcome::Ok(Result);
